Add a new interface for recording IAS analytics over a specific period of time.

#rnx

- Add a new public interface IAnalyticsRecording that will be used for recording IAS analytics. It will start recording from the point that it was created until destroyed or until ::StopRecording is called.
-- The old way of recording analytics FIoStoreOnDemandModule::ReportAnalytics has remained untouched for now but I plan to change its functionality so that it always reports the total stats since the backend was created rather than reporting the stats since it was last called, since the new interface covers the functionality of recording analytics over a set period of time.
- Moved FAvailableEndPoints to it's own code file so that it can be used in the statistics system.

diff --git a/Engine/Source/Runtime/Experimental/IoStore/OnDemand/Private/AvailableEndPoints.h b/Engine/Source/Runtime/Experimental/IoStore/OnDemand/Private/AvailableEndPoints.h
new file mode 100644
--- /dev/null
+++ b/Engine/Source/Runtime/Experimental/IoStore/OnDemand/Private/AvailableEndPoints.h
@@ -0,0 +1,27 @@
+// Copyright Epic Games, Inc. All Rights Reserved.
+
+#pragma once
+
+#include "Containers/Array.h"
+#include "Containers/UnrealString.h"
+
+namespace UE::IoStore
+{
+
+struct FAvailableEndPoints
+{
+	bool HasCurrent() const
+	{
+		return Current != INDEX_NONE;
+	}
+
+	const FString& GetCurrent() const
+	{
+		return Urls[Current];
+	}
+
+	int32 Current = INDEX_NONE;
+	TArray<FString> Urls;
+};
+
+} // namespace UE::IoStore
diff --git a/Engine/Source/Runtime/Experimental/IoStore/OnDemand/Private/IoStoreOnDemand.cpp b/Engine/Source/Runtime/Experimental/IoStore/OnDemand/Private/IoStoreOnDemand.cpp
--- a/Engine/Source/Runtime/Experimental/IoStore/OnDemand/Private/IoStoreOnDemand.cpp
+++ b/Engine/Source/Runtime/Experimental/IoStore/OnDemand/Private/IoStoreOnDemand.cpp
@@ -1335,6 +1335,18 @@ void FIoStoreOnDemandModule::ReportAnalytics(TArray<FAnalyticsEventAttribute>& O
 	}
 }
 
+TUniquePtr<IAnalyticsRecording> FIoStoreOnDemandModule::StartAnalyticsRecording() const
+{
+	if (HttpIoDispatcherBackend.IsValid())
+	{
+		return HttpIoDispatcherBackend->StartAnalyticsRecording();
+	}
+	else
+	{
+		return TUniquePtr<IAnalyticsRecording>();
+	}
+}
+
 void FIoStoreOnDemandModule::Mount(FOnDemandMountArgs&& Args, FOnDemandMountCompleted&& OnCompleted)
 {
 	if (IoStore.IsValid() == false)
diff --git a/Engine/Source/Runtime/Experimental/IoStore/OnDemand/Private/OnDemandIoDispatcherBackend.cpp b/Engine/Source/Runtime/Experimental/IoStore/OnDemand/Private/OnDemandIoDispatcherBackend.cpp
--- a/Engine/Source/Runtime/Experimental/IoStore/OnDemand/Private/OnDemandIoDispatcherBackend.cpp
+++ b/Engine/Source/Runtime/Experimental/IoStore/OnDemand/Private/OnDemandIoDispatcherBackend.cpp
@@ -3,6 +3,7 @@
 #include "OnDemandIoDispatcherBackend.h"
 
 #include "AnalyticsEventAttribute.h"
+#include "AvailableEndPoints.h"
 #include "Containers/BitArray.h"
 #include "Containers/StringView.h"
 #include "DistributionEndpoints.h"
@@ -715,15 +716,6 @@ class FOnDemandIoBackend final
 	using FIoRequestQueue = TThreadSafeIntrusiveQueue<FIoRequestImpl>;
 	using FChunkRequestQueue = TThreadSafeIntrusiveQueue<FChunkRequest>;
 
-	struct FAvailableEps
-	{
-		bool HasCurrent() const { return Current != INDEX_NONE; }
-		const FString& GetCurrent() const { return Urls[Current]; }
-
-		int32 Current = INDEX_NONE;
-		TArray<FString> Urls;
-	};
-
 	struct FBackendData
 	{
 		static void Attach(FIoRequestImpl* Request, const FIoHash& ChunkKey)
@@ -904,7 +896,10 @@ public:
 	virtual void SetEnabled(bool bEnabled) override;
 	virtual bool IsEnabled() const override;
 	virtual void AbandonCache() override;
+	
 	virtual void ReportAnalytics(TArray<FAnalyticsEventAttribute>& OutAnalyticsArray) const override;
+	virtual TUniquePtr<IAnalyticsRecording> StartAnalyticsRecording() const override;
+	
 	virtual FOnDemandStreamingCacheUsage GetCacheUsage() const override;
 
 	// Runnable
@@ -937,7 +932,7 @@ private:
 	TArray<FChunkRequest*> PendingCacheGets;
 	FBackendStatus BackendStatus; 
 	FOnDemandIoBackendStats Stats;
-	FAvailableEps AvailableEps;
+	FAvailableEndPoints AvailableEps;
 	FDistributedEndpointUrl DistributionUrl;
 	FEventRef DistributedEndpointEvent;
 
@@ -1496,6 +1491,16 @@ void FOnDemandIoBackend::ReportAnalytics(TArray<FAnalyticsEventAttribute>& OutAn
 	}
 }
 
+TUniquePtr<IAnalyticsRecording> FOnDemandIoBackend::StartAnalyticsRecording() const
+{
+	if (GIasReportAnalyticsEnabled)
+	{
+		return Stats.StartAnalyticsRecording(AvailableEps);
+	}
+
+	return TUniquePtr<IAnalyticsRecording>();
+}
+
 FOnDemandStreamingCacheUsage FOnDemandIoBackend::GetCacheUsage() const
 {
 	FOnDemandStreamingCacheUsage Usage;
@@ -1873,4 +1878,4 @@ struct FTests
 
 } // namespace UE::IoStore
 
-#undef UE_IAS_DEBUG_CONSOLE_CMDS
+#undef UE_IAS_DEBUG_CONSOLE_CMDS
\ No newline at end of file
diff --git a/Engine/Source/Runtime/Experimental/IoStore/OnDemand/Private/OnDemandIoDispatcherBackend.h b/Engine/Source/Runtime/Experimental/IoStore/OnDemand/Private/OnDemandIoDispatcherBackend.h
--- a/Engine/Source/Runtime/Experimental/IoStore/OnDemand/Private/OnDemandIoDispatcherBackend.h
+++ b/Engine/Source/Runtime/Experimental/IoStore/OnDemand/Private/OnDemandIoDispatcherBackend.h
@@ -13,6 +13,7 @@ namespace UE::IoStore
 {
 
 class FOnDemandIoStore;
+class IAnalyticsRecording;
 struct FOnDemandStreamingCacheUsage;
 
 struct FDistributedEndpointUrl
@@ -82,6 +83,7 @@ public:
 	virtual bool IsEnabled() const = 0;
 	virtual void AbandonCache() = 0;
 	virtual void ReportAnalytics(TArray<FAnalyticsEventAttribute>& OutAnalyticsArray) const = 0;
+	virtual TUniquePtr<IAnalyticsRecording> StartAnalyticsRecording() const = 0;
 	virtual FOnDemandStreamingCacheUsage GetCacheUsage() const = 0;
 };
 
diff --git a/Engine/Source/Runtime/Experimental/IoStore/OnDemand/Private/Statistics.cpp b/Engine/Source/Runtime/Experimental/IoStore/OnDemand/Private/Statistics.cpp
--- a/Engine/Source/Runtime/Experimental/IoStore/OnDemand/Private/Statistics.cpp
+++ b/Engine/Source/Runtime/Experimental/IoStore/OnDemand/Private/Statistics.cpp
@@ -4,12 +4,14 @@
 
 #include "Algo/MaxElement.h"
 #include "AnalyticsEventAttribute.h"
+#include "AvailableEndPoints.h"
 #include "HAL/IConsoleManager.h"
 #include "HAL/PlatformTime.h"
 #include "IO/IoStoreOnDemand.h"
 #include "Internationalization/Internationalization.h"
 #include "Math/UnrealMathUtility.h"
 #include "Misc/CoreDelegates.h"
+#include "Templates/Requires.h"
 
 LLM_DEFINE_TAG(Ias);
 
@@ -197,7 +199,7 @@ public:
 		}
 		else
 		{
-			IntTotals.Add(FString(Name), 0);
+			IntTotals.Add(FString(Name), Value);
 			return Value;
 		}
 	}
@@ -297,21 +299,46 @@ FIncrementalVariance	GHttpAvgDuration; // Duration of the http requests, in mill
 
 static const uint64 DurationBuckets[] = { 30, 150, 400, 1000 }; // Time boundaries (in ms) for each bucket
 
-int32 FindDurationBucket(uint64 DurationMs)
+// P1907R1 is not yet implemented in some of our compilers so FTrackedValue will not work on subobjects
+// so we need each bucket to be it's own variable rather than an array.
+uint32 GHttpDurationBucket0 = 0;
+uint32 GHttpDurationBucket1 = 0;
+uint32 GHttpDurationBucket2 = 0;
+uint32 GHttpDurationBucket3 = 0;
+uint32 GHttpDurationBucket4 = 0;
+
+uint32* GHttpDurationBuckets[] =
+{
+	&GHttpDurationBucket0,
+	&GHttpDurationBucket1,
+	&GHttpDurationBucket2,
+	&GHttpDurationBucket3,
+	&GHttpDurationBucket4
+};
+
+void IncrementDurationBucket(uint64 DurationMs)
 {
 	const uint32 NumBoundries = UE_ARRAY_COUNT(DurationBuckets);
 	for (int32 Index = 0; Index < NumBoundries; ++Index)
 	{
 		if (DurationMs <= DurationBuckets[Index])
 		{
-			return Index;
+			(*GHttpDurationBuckets[Index])++;
+			return;
 		}
 	}
 
-	return NumBoundries;
+	(*GHttpDurationBuckets[NumBoundries])++;
 }
 
-uint32 GHttpDurationBuckets[UE_ARRAY_COUNT(DurationBuckets) +1] = { 0 };
+void ResetDurationBuckets()
+{
+	const uint32 NumBuckets = UE_ARRAY_COUNT(DurationBuckets) + 1;
+	for (int32 Index = 0; Index < NumBuckets; ++Index)
+	{
+		(*GHttpDurationBuckets[Index]) = 0;
+	}
+}
 
 ////////////////////////////////////////////////////////////////////////////////
 // CSV STATS
@@ -517,16 +544,16 @@ void FOnDemandIoBackendStats::ReportEndPointAnalytics(TArray<FAnalyticsEventAttr
 
 			,TRACK_DELTA("IasDecodeErrors", GIoDecodeErrorCount.Get())
 
-			,TEXT("IasHttpDuration0"), GHttpDurationBuckets[0]
-			,TEXT("IasHttpDuration1"), GHttpDurationBuckets[1]
-			,TEXT("IasHttpDuration2"), GHttpDurationBuckets[2]
-			,TEXT("IasHttpDuration3"), GHttpDurationBuckets[3]
-			,TEXT("IasHttpDuration4"), GHttpDurationBuckets[4]
+			,TEXT("IasHttpDuration0"), GHttpDurationBucket0
+			,TEXT("IasHttpDuration1"), GHttpDurationBucket1
+			,TEXT("IasHttpDuration2"), GHttpDurationBucket2
+			,TEXT("IasHttpDuration3"), GHttpDurationBucket3
+			,TEXT("IasHttpDuration4"), GHttpDurationBucket4
 		);
 
 		// These values we can just reset as they are only being used with analytics
 		GHttpAvgDuration.Reset();
-		FMemory::Memzero(GHttpDurationBuckets);
+		ResetDurationBuckets();
 	}
 
 	if (GIasReportCacheAnalyticsEnabled)
@@ -556,6 +583,212 @@ void FOnDemandIoBackendStats::ReportEndPointAnalytics(TArray<FAnalyticsEventAttr
 #undef TRACK_DELTA
 }
 
+
+class FAnalyticsRecording : public IAnalyticsRecording
+{
+public:
+	FAnalyticsRecording() = delete;
+	FAnalyticsRecording(const FAvailableEndPoints& InAvailableEndPoints, FBackendStatus& InBackendStatus)
+		: AvailableEndPoints(InAvailableEndPoints)
+		, BackendStatus(InBackendStatus)
+	{
+
+	}
+
+	virtual ~FAnalyticsRecording() = default;
+
+private:
+
+	void StopRecording() override
+	{
+		if (!bRecording)
+		{
+			return;
+		}
+	
+		Http.ErrorCount.Stop();
+		Http.RetryCount.Stop();
+		Http.GetCount.Stop();
+		Http.DownloadedBytes.Stop();
+
+		Http.DurationBuckets0.Stop();
+		Http.DurationBuckets1.Stop();
+		Http.DurationBuckets2.Stop();
+		Http.DurationBuckets3.Stop();
+		Http.DurationBuckets4.Stop();
+
+		Io.DecodeErrors.Stop();
+
+		Cache.ErrorCount.Stop();
+		Cache.GetCount.Stop();
+		Cache.PutCount.Stop();
+
+		Cache.WrittenBytes.Stop();
+		Cache.ReadBytes.Stop();
+		Cache.RejectBytes.Stop();
+
+		bRecording = false;
+	}
+
+	void Report(TArray<FAnalyticsEventAttribute>& OutAnalyticsArray) const override
+	{
+		FString CdnUrl = AvailableEndPoints.GetCurrent();
+
+		// Strip the prefix from the url as some analytics systems may have trouble dealing with it
+		if (!CdnUrl.RemoveFromStart(TEXT("http://")))
+		{
+			CdnUrl.RemoveFromStart(TEXT("https://"));
+		}
+
+		AppendAnalyticsEventAttributeArray(OutAnalyticsArray
+			, TEXT("IasCdnUrl"), MoveTemp(CdnUrl)
+			, TEXT("IasHttpHasEverConnected"), GHttpConnectCount.Get() > 0 // Report if the system has ever actually managed to make a connection
+		);
+
+		if (GIasReportHttpAnalyticsEnabled)
+		{
+			AppendAnalyticsEventAttributeArray(OutAnalyticsArray
+				, TEXT("IasHttpErrorCount"), Http.ErrorCount.GetValue()
+				, TEXT("IasHttpRetryCount"), Http.RetryCount.GetValue()
+				, TEXT("IasHttpGetCount"), Http.GetCount.GetValue()
+				, TEXT("IasHttpDownloadedBytes"), Http.DownloadedBytes.GetValue()
+				, TEXT("IasHttpDurationMeanAvg"), GHttpAvgDuration.GetMean()		// TODO: Reconsider
+				, TEXT("IasHttpDurationStdDev"), GHttpAvgDuration.GetDeviation()	// TODO: Reconsider
+
+				, TEXT("IasHttpDuration0"), Http.DurationBuckets0.GetValue()
+				, TEXT("IasHttpDuration1"), Http.DurationBuckets1.GetValue()
+				, TEXT("IasHttpDuration2"), Http.DurationBuckets2.GetValue()
+				, TEXT("IasHttpDuration3"), Http.DurationBuckets3.GetValue()
+				, TEXT("IasHttpDuration4"), Http.DurationBuckets4.GetValue()
+			);
+		}
+
+		AppendAnalyticsEventAttributeArray(OutAnalyticsArray
+			, TEXT("IasDecodeErrors"), GIoDecodeErrorCount.Get()
+		);
+
+		if (GIasReportCacheAnalyticsEnabled)
+		{
+			const float CacheUsagePercent = GCacheMaxBytes > 0 ? 100.f * (float(GCacheCachedBytes.Get()) / float(GCacheMaxBytes)) : 0.f;
+
+			AppendAnalyticsEventAttributeArray(OutAnalyticsArray
+
+				, TEXT("IasCacheEnabled"), this->BackendStatus.IsCacheEnabled()
+				, TEXT("IasCacheMaxBytes"), GCacheMaxBytes
+				, TEXT("IasCacheCachedBytes"), GCacheCachedBytes.Get()
+				, TEXT("IasCacheUsagePercent"), CacheUsagePercent
+
+				, TEXT("IasCacheErrorCount"), Cache.ErrorCount.GetValue()
+				, TEXT("IasCacheGetCount"), Cache.GetCount.GetValue()
+				, TEXT("IasCachePutCount"), Cache.PutCount.GetValue()
+
+				, TEXT("IasCacheWriteBytes"), Cache.WrittenBytes.GetValue()
+				, TEXT("IasCacheReadBytes"), Cache.ReadBytes.GetValue()
+				, TEXT("IasCacheRejectBytes"), Cache.RejectBytes.GetValue()
+			);
+		}
+	}
+
+private:
+
+	/** This wrapper class allows us to treat raw integer types and FCounter types  as the same thing */
+	//template<typename CounterType, CounterType& Counter, int32 INDEX = -1>
+	template<typename CounterType, const CounterType& Counter, int32 INDEX = -1>
+	struct FTrackedValue
+	{
+		FTrackedValue()
+		{
+			Value = GetCurrentValue();
+		}
+
+		void Stop()
+		{
+			Value = GetCurrentValue() - Value;
+			bRecording = false;
+		}
+
+		int64 GetValue() const
+		{
+			if (bRecording)
+			{
+				return GetCurrentValue() - Value;
+			}
+			else
+			{
+				return Value;
+			}
+		}
+
+	private:
+
+		template <
+			typename U = CounterType
+			UE_REQUIRES(!std::is_integral_v<U>)
+		>
+		int64 GetCurrentValue() const
+		{
+			return Counter.Get();
+		}
+
+		template <
+			typename U = CounterType
+			UE_REQUIRES(std::is_integral_v<U>)
+		>
+		int64 GetCurrentValue() const
+		{
+			return Counter;
+		}
+
+		int64 Value;
+		bool bRecording = true;
+	};
+
+	/** Http Stats */
+	struct FHttpStats
+	{
+		FTrackedValue<FCounterInt, GHttpErrorCount>			ErrorCount;
+		FTrackedValue<FCounterInt, GHttpRetryCount>			RetryCount;
+		FTrackedValue<FCounterInt, GHttpGetCount>			GetCount;
+		FTrackedValue<FCounterInt, GHttpDownloadedBytes>	DownloadedBytes;
+
+		FTrackedValue<uint32, GHttpDurationBucket0>			DurationBuckets0;
+		FTrackedValue<uint32, GHttpDurationBucket1>			DurationBuckets1;
+		FTrackedValue<uint32, GHttpDurationBucket2>			DurationBuckets2;
+		FTrackedValue<uint32, GHttpDurationBucket3>			DurationBuckets3;
+		FTrackedValue<uint32, GHttpDurationBucket4>			DurationBuckets4;
+	} Http;
+
+	/** Io Stats */
+	struct FIoStats
+	{
+		FTrackedValue<FCounterAtomicInt, GIoDecodeErrorCount> DecodeErrors;
+	} Io;
+
+	/** Cache Stats */
+	struct FCacheStats
+	{
+		FTrackedValue<FCounterAtomicInt, GCacheErrorCount>		ErrorCount;
+		FTrackedValue<FCounterAtomicInt, GCacheGetCount>		GetCount;
+		FTrackedValue<FCounterAtomicInt, GCachePutCount>		PutCount;
+
+		FTrackedValue<FCounterAtomicInt, GCacheWrittenBytes>	WrittenBytes;
+		FTrackedValue<FCounterAtomicInt, GCacheReadBytes>		ReadBytes;
+		FTrackedValue<FCounterAtomicInt, GCacheRejectBytes>		RejectBytes;
+	} Cache;
+
+	const FAvailableEndPoints& AvailableEndPoints;
+	const FBackendStatus& BackendStatus;
+
+	bool bRecording = true;
+};
+
+TUniquePtr<IAnalyticsRecording> FOnDemandIoBackendStats::StartAnalyticsRecording(const FAvailableEndPoints& AvailableEndPoints) const
+{
+	TUniquePtr<FAnalyticsRecording> Recording = MakeUnique<FAnalyticsRecording>(AvailableEndPoints, BackendStatus);
+
+	return Recording;
+}
+
 void FOnDemandIoBackendStats::OnIoRequestEnqueue()
 {
 	GIoRequestCount.Add(1);
@@ -691,7 +924,7 @@ void FOnDemandIoBackendStats::OnHttpGet(uint64 SizeBytes, uint64 DurationMs)
 
 	GHttpAvgDuration.Increment(static_cast<double>(DurationMs));
 
-	GHttpDurationBuckets[FindDurationBucket(DurationMs)]++;
+	IncrementDurationBucket(DurationMs);
 }
 
 void FOnDemandIoBackendStats::OnHttpCancel()
diff --git a/Engine/Source/Runtime/Experimental/IoStore/OnDemand/Private/Statistics.h b/Engine/Source/Runtime/Experimental/IoStore/OnDemand/Private/Statistics.h
--- a/Engine/Source/Runtime/Experimental/IoStore/OnDemand/Private/Statistics.h
+++ b/Engine/Source/Runtime/Experimental/IoStore/OnDemand/Private/Statistics.h
@@ -20,6 +20,9 @@ LLM_DECLARE_TAG(Ias);
 namespace UE::IoStore
 {
 
+class IAnalyticsRecording;
+struct FAvailableEndPoints;
+
 #if IAS_WITH_STATISTICS
 #	define IAS_STATISTICS_IMPL(...) ;
 #else
@@ -39,6 +42,8 @@ public:
 	/** Report analytics for the current endpoint */
 	void ReportEndPointAnalytics(TArray<FAnalyticsEventAttribute>& OutAnalyticsArray) const IAS_STATISTICS_IMPL()
 
+	TUniquePtr<IAnalyticsRecording> StartAnalyticsRecording(const FAvailableEndPoints& AvailableEndPoints) const IAS_STATISTICS_IMPL(TUniquePtr<IAnalyticsRecording>())
+
 	void OnIoRequestEnqueue() IAS_STATISTICS_IMPL()
 	void OnIoRequestComplete(uint64 Size, uint64 DurationMs) IAS_STATISTICS_IMPL()
 	void OnIoRequestCancel() IAS_STATISTICS_IMPL()
diff --git a/Engine/Source/Runtime/Experimental/IoStore/OnDemand/Public/IO/IoStoreOnDemand.h b/Engine/Source/Runtime/Experimental/IoStore/OnDemand/Public/IO/IoStoreOnDemand.h
--- a/Engine/Source/Runtime/Experimental/IoStore/OnDemand/Public/IO/IoStoreOnDemand.h
+++ b/Engine/Source/Runtime/Experimental/IoStore/OnDemand/Public/IO/IoStoreOnDemand.h
@@ -476,6 +476,20 @@ private:
 	std::atomic<bool> bCanceled{ false };
 };
 
+/** Interface  for recording analytics over a given time period */
+class IAnalyticsRecording
+{
+public:
+	IAnalyticsRecording() = default;
+	virtual ~IAnalyticsRecording() = default;
+
+	/** Writes the current value of the analytics to the output array */
+	virtual void Report(TArray<FAnalyticsEventAttribute>& OutAnalyticsArray) const = 0;
+	
+	/** Calling this will stop recording. Any calls to ::Report after this point will return the same data */
+	virtual void StopRecording() = 0;
+};
+
 class FIoStoreOnDemandModule
 	: public IModuleInterface
 {
@@ -498,7 +512,17 @@ public:
 	UE_API bool IsEnabled() const;
 	UE_API void AbandonCache();
 
+	/**
+	 * Reports the statistics for the current OnDemandBackend. It will record any events that have occurred from the last time that this was called.
+	 * This is a legacy method, prefer using StartAnalyticsRecording instead.
+	 */
 	UE_API void ReportAnalytics(TArray<FAnalyticsEventAttribute>& OutAnalyticsArray) const;
+	
+	/** 
+	 * Create a new analytics interface for the current OnDemandBackend, can be used to record statistics for a set period of time.
+	 * This can return a nullptr if the analytics system is not enabled.
+	 */
+	UE_API TUniquePtr<IAnalyticsRecording> StartAnalyticsRecording() const;
 
 	UE_API void Mount(FOnDemandMountArgs&& Args, FOnDemandMountCompleted&& OnCompleted);
 	UE_API FIoStatus Unmount(FStringView MountId);
