Remove use of FOnDemandIoBackendStats as a singleton and access the member instead.

#rnx

- As we move towards having multiple connections (and potentially multiple caches) running we need to move away from accessing things as singletons and have a stronger chain of ownership.

diff --git a/Engine/Source/Runtime/Experimental/IoStore/OnDemand/Private/OnDemandIoDispatcherBackend.cpp b/Engine/Source/Runtime/Experimental/IoStore/OnDemand/Private/OnDemandIoDispatcherBackend.cpp
--- a/Engine/Source/Runtime/Experimental/IoStore/OnDemand/Private/OnDemandIoDispatcherBackend.cpp
+++ b/Engine/Source/Runtime/Experimental/IoStore/OnDemand/Private/OnDemandIoDispatcherBackend.cpp
@@ -1173,7 +1173,7 @@ void FOnDemandIoBackend::CompleteMaterialize(FChunkRequest* ChunkRequest)
 		return;
 
 	case EIoErrorCode::ReadError:
-		FOnDemandIoBackendStats::Get()->OnCacheError();
+		Stats.OnCacheError();
 		break;
 
 	case EIoErrorCode::Cancelled:
@@ -1709,7 +1709,7 @@ uint32 FOnDemandIoBackend::Run()
 		HttpClient->SetEndpoint(AvailableEps.Current);
 
 		BackendStatus.SetHttpEnabled(true);
-		FOnDemandIoBackendStats::Get()->OnHttpConnected();
+		Stats.OnHttpConnected();
 	}
 	else
 	{
@@ -1736,7 +1736,7 @@ uint32 FOnDemandIoBackend::Run()
 				WaitTime = GIasHttpHealthCheckWaitTime;
 				if (HttpClient->GetEndpoint() != INDEX_NONE)
 				{
-					FOnDemandIoBackendStats::Get()->OnHttpDisconnected();
+					Stats.OnHttpDisconnected();
 
 					AvailableEps.Current = INDEX_NONE;
 					HttpClient->SetEndpoint(INDEX_NONE);
@@ -1747,7 +1747,7 @@ uint32 FOnDemandIoBackend::Run()
 				const FString TestPath = GetEndpointTestPath(); 
 				if (int32 Idx = LatencyTest(AvailableEps.Urls, TestPath, bStopRequested); Idx != INDEX_NONE)
 				{
-					FOnDemandIoBackendStats::Get()->OnHttpConnected();
+					Stats.OnHttpConnected();
 
 					AvailableEps.Current = Idx;
 					HttpClient->SetEndpoint(Idx);
