Adds some comments to FPBDJointSettings. No functional changes.

diff --git a/Engine/Source/Runtime/Experimental/Chaos/Public/Chaos/PBDJointConstraintTypes.h b/Engine/Source/Runtime/Experimental/Chaos/Public/Chaos/PBDJointConstraintTypes.h
--- a/Engine/Source/Runtime/Experimental/Chaos/Public/Chaos/PBDJointConstraintTypes.h
+++ b/Engine/Source/Runtime/Experimental/Chaos/Public/Chaos/PBDJointConstraintTypes.h
@@ -124,21 +124,30 @@ namespace Chaos
 			return !FMemory::Memcmp(this, &Other, sizeof(*this));
 		}
 
+		// The offsets of the joint frames from the bodies. The ordering is [child, parent]
 		FTransformPair ConnectorTransforms;
 
 		FReal Stiffness;
+		// The amount of linear projection, between 0 and 1. Only used if bProjectionEnabled is set too.
 		FReal LinearProjection;
+		// The amount of angular projection, between 0 and 1. Only used if bProjectionEnabled is set too.
 		FReal AngularProjection;
 		FReal ShockPropagation;
 		FReal TeleportDistance;
 		FReal TeleportAngle;			// Radians
+		// This inverse scales the parent mass (and inertia) during the solve. By setting this to
+		// zero, the constraint will treat the parent body as infinitely massive, and thus the
+		// constraint will only operate on the child body.
 		FReal ParentInvMassScale;
 
 		bool bCollisionEnabled;
-		bool bProjectionEnabled;		// @chaos(todo): remove - implied by alpha and teleport settings
-		bool bShockPropagationEnabled;	// @chaos(todo): remove - implied by alpha
+		bool bProjectionEnabled;		// @chaos(todo): remove - implied by Linear/AngularProjection and teleport settings
+		bool bShockPropagationEnabled;	// @chaos(todo): remove - implied by ShockPropagation
 		bool bMassConditioningEnabled;
 
+		// If enabled, then a linear solver will be used. However, note that the immediate solver
+		// doesn't support mixing both linear and non-linear constraints, so in that case you should
+		// make all constraints be either linear or non-linear.
 		bool bUseLinearSolver;
 
 		TVector<EJointMotionType, 3> LinearMotionTypes;
