Always install required chunks even if no packages were resolved.

#rnx

diff --git a/Engine/Source/Runtime/Experimental/IoStore/OnDemand/Private/OnDemandIoStore.cpp b/Engine/Source/Runtime/Experimental/IoStore/OnDemand/Private/OnDemandIoStore.cpp
--- a/Engine/Source/Runtime/Experimental/IoStore/OnDemand/Private/OnDemandIoStore.cpp
+++ b/Engine/Source/Runtime/Experimental/IoStore/OnDemand/Private/OnDemandIoStore.cpp
@@ -247,6 +247,34 @@ static FIoStatus BuildInstallData(
 {
 	OutInstallData.Reserve(Containers.Num());
 
+	// For now we always download these required chunks
+	for (const FSharedOnDemandContainer& Container : Containers)
+	{
+		FContainerInstallData& Data = OutInstallData.FindOrAdd(Container);
+
+		for (int32 EntryIndex = 0; const FIoChunkId& ChunkId : Container->ChunkIds)
+		{
+			switch(ChunkId.GetChunkType())
+			{
+				case EIoChunkType::ExternalFile:
+				case EIoChunkType::ShaderCodeLibrary:
+				case EIoChunkType::ShaderCode:
+				{
+					Data.ResolvedChunks.Add(EntryIndex, EChunkCacheState::None);
+					Data.TotalSize += Container->ChunkEntries[EntryIndex].EncodedSize;
+				}
+				default:
+					break;
+			}
+			++EntryIndex;
+		}
+	}
+
+	if (PackageIds.IsEmpty())
+	{
+		return EIoErrorCode::Ok;
+	}
+
 	// Setup the package information for each container
 	for (const FSharedOnDemandContainer& Container : Containers)
 	{
@@ -288,7 +316,7 @@ static FIoStatus BuildInstallData(
 		}
 	}
 
-	// Traverse dependencies for each package id
+	// Traverse dependencies for each package ID
 	using FQueue = TRingBuffer<FPackageId>;
 
 	FQueue Queue;
@@ -418,35 +446,6 @@ static FIoStatus BuildInstallData(
 		}
 	}
 
-	// For now we always download these chunks for containers with resolved chunk ID(s)
-	for (TPair<FSharedOnDemandContainer, FContainerInstallData>& Kv : OutInstallData)
-	{
-		const FOnDemandContainer& Container = *Kv.Key;
-		FContainerInstallData& Data			= Kv.Value;
-
-		if (Data.ResolvedChunks.IsEmpty())
-		{
-			continue;
-		}
-
-		for (int32 EntryIndex = 0; const FIoChunkId& ChunkId : Container.ChunkIds)
-		{
-			switch(ChunkId.GetChunkType())
-			{
-				case EIoChunkType::ExternalFile:
-				case EIoChunkType::ShaderCodeLibrary:
-				case EIoChunkType::ShaderCode:
-				{
-					Data.ResolvedChunks.Add(EntryIndex, EChunkCacheState::None);
-					Data.TotalSize += Container.ChunkEntries[EntryIndex].EncodedSize;
-				}
-				default:
-					break;
-			}
-			++EntryIndex;
-		}
-	}
-
 	return EIoErrorCode::Ok;
 }
 
@@ -1937,6 +1936,7 @@ FIoStatus FOnDemandIoStore::GetContainersAndPackagesForInstall(
 
 		for (const FString& Tag : TagSets)
 		{
+			const int32 NumBeforeTag = OutPackageIdsToInstall.Num();
 			for (const FOnDemandTagSet& TagSet : Container->TagSets)
 			{
 				if (TagSet.Tag == Tag)
@@ -1952,6 +1952,7 @@ FIoStatus FOnDemandIoStore::GetContainersAndPackagesForInstall(
 					}
 				}
 			}
+			UE_LOG(LogIoStoreOnDemand, Log, TEXT("Installing total %d package(s) with tag '%s'"), OutPackageIdsToInstall.Num() - NumBeforeTag, *Tag)
 		}
 	}
 
