Chaos: Add EmbeddedSpringFacade.

diff --git a/Engine/Source/Runtime/Experimental/Chaos/Private/Chaos/CollectionEmbeddedSpringConstraintFacade.cpp b/Engine/Source/Runtime/Experimental/Chaos/Private/Chaos/CollectionEmbeddedSpringConstraintFacade.cpp
new file mode 100644
--- /dev/null
+++ b/Engine/Source/Runtime/Experimental/Chaos/Private/Chaos/CollectionEmbeddedSpringConstraintFacade.cpp
@@ -0,0 +1,667 @@
+// Copyright Epic Games, Inc. All Rights Reserved.
+
+#include "Chaos/CollectionEmbeddedSpringConstraintFacade.h"
+
+namespace Chaos::Softs
+{
+	namespace Private
+	{
+		template<typename T>
+		TArrayView<T> GetArrayView(TManagedArrayAccessor<T>& DataArray, const TManagedArrayAccessor<int32>& StartArray, const TManagedArrayAccessor<int32>& EndArray, int32 ConstraintIndex)
+		{
+			// Taken from ClothCollection::GetElements.
+			// If this pattern is used more in the GeometryCollection facades, we could add this to a utility.
+			if (DataArray.IsValid() && StartArray.IsValid() && EndArray.IsValid())
+			{
+				const int32 Start = StartArray[ConstraintIndex];
+				const int32 End = EndArray[ConstraintIndex];
+				if (Start != INDEX_NONE && End != INDEX_NONE)
+				{
+					return TArrayView<T>(DataArray.Modify().GetData() + Start, End - Start + 1);
+				}
+				checkf(Start == End, TEXT("Only one boundary of the range is set to INDEX_NONE, when both should."));
+			}
+			return TArrayView<T>();
+		}
+
+		template<typename T>
+		TConstArrayView<T> GetConstArrayView(const TManagedArrayAccessor<T>& DataArray, const TManagedArrayAccessor<int32>& StartArray, const TManagedArrayAccessor<int32>& EndArray, int32 ConstraintIndex)
+		{
+			// Taken from ClothCollection::GetElements.
+			// If this pattern is used more in the GeometryCollection facades, we could add this to a utility.
+			if (DataArray.IsValid() && StartArray.IsValid() && EndArray.IsValid())
+			{
+				const int32 Start = StartArray[ConstraintIndex];
+				const int32 End = EndArray[ConstraintIndex];
+				if (Start != INDEX_NONE && End != INDEX_NONE)
+				{
+					return TConstArrayView<T>(DataArray.Get().GetData() + Start, End - Start + 1);
+				}
+				checkf(Start == End, TEXT("Only one boundary of the range is set to INDEX_NONE, when both should."));
+			}
+			return TConstArrayView<T>();
+		}
+
+		template<typename T>
+		void CopyArrayViewData(const TArrayView<T>& To, const TConstArrayView<T>& From)
+		{
+			// Taken from ClothCollection::CopyArrayViewData.
+			// If this pattern is used more in the GeometryCollection facades, we could add this to a utility.
+			check(To.Num() == From.Num());
+			for (int32 Index = 0; Index < To.Num(); ++Index)
+			{
+				To[Index] = From[Index];
+			}
+		}
+
+		template<typename T>
+		void CopyArrayViewDataAndApplyOffset(const TArrayView<TArray<T>>& To, const TConstArrayView<TArray<T>>& From, const T Offset)
+		{
+			// Taken from ClothCollection::CopyArrayViewDataAndApplyOffset.
+			// If this pattern is used more in the GeometryCollection facades, we could add this to a utility.
+			check(To.Num() == From.Num());
+			for (int32 Index = 0; Index < To.Num(); ++Index)
+			{
+				To[Index] = From[Index];
+				for (T& Value : To[Index])
+				{
+					Value += Offset;
+				}
+			}
+		}
+
+		void SetIndicesAndWeights(const int32 EndPointNumIndices, TArray<int32>& IndexArray, TArray<float>& WeightArray, const TConstArrayView<int32>& InIndexArray, const TConstArrayView<float>& InWeightArray)
+		{
+			IndexArray.Reset();
+			WeightArray.Reset();
+			IndexArray.SetNumZeroed(EndPointNumIndices);
+			WeightArray.SetNumZeroed(EndPointNumIndices);
+			const int32 InEndPointNumIndices = InIndexArray.Num();
+			check(InWeightArray.Num() == InEndPointNumIndices);
+			const int32 MinNumIndices = FMath::Min(EndPointNumIndices, InEndPointNumIndices);
+			for (int32 Idx = 0; Idx < MinNumIndices; ++Idx)
+			{
+				IndexArray[Idx] = InIndexArray[Idx];
+				WeightArray[Idx] = InWeightArray[Idx];
+			}
+		}
+
+		/* @return should remove */
+		bool CleanIndicesAndWeights(TArray<int32>& IndexArray, TArray<float>& WeightArray)
+		{
+			check(IndexArray.Num() == WeightArray.Num());
+			
+			// Return true if any non-zero weighted indices are invalid.
+			// Also find first non-zero weighted index.
+			bool bAnyZeroWeights = false;
+			int32 NonZeroWeightIndex = INDEX_NONE;
+			for (int32 Idx = 0; Idx < IndexArray.Num(); ++Idx)
+			{
+				if (WeightArray[Idx] < UE_SMALL_NUMBER)
+				{
+					bAnyZeroWeights = true;
+				}
+				else
+				{
+					if (IndexArray[Idx] < 0)
+					{
+						return true;
+					}
+					NonZeroWeightIndex = IndexArray[Idx];
+				}
+			}
+			if (NonZeroWeightIndex == INDEX_NONE)
+			{
+				// All weights are Zero. Remove this spring.
+				return true;
+			}
+
+			if (bAnyZeroWeights)
+			{
+				// Set all Indices with zero weights to one of the non-zero weighted indices in this spring. This will help with graph coloring and cache.
+
+				for (int32 Idx = 0; Idx < IndexArray.Num(); ++Idx)
+				{
+					if (WeightArray[Idx] < UE_SMALL_NUMBER)
+					{
+						IndexArray[Idx] = NonZeroWeightIndex;
+					}
+				}
+			}
+			return false;
+		}
+	}
+
+	// Attributes
+	const FName FEmbeddedSpringBaseFacade::SpringLength(TEXT("SpringLength"));
+	const FName FEmbeddedSpringBaseFacade::CompressionStiffness(TEXT("CompressionStiffness"));
+
+	const FName FEmbeddedSpringBaseFacade::SpringConstraintGroupName(TEXT("SpringConstraint"));
+	const FName FEmbeddedSpringBaseFacade::ConstraintStart(TEXT("ConstraintStart"));
+	const FName FEmbeddedSpringBaseFacade::ConstraintEnd(TEXT("ConstraintEnd"));
+	const FName FEmbeddedSpringBaseFacade::ConstraintEndPointNumIndices(TEXT("ConstraintEndPointNumIndices"));
+
+	FEmbeddedSpringBaseFacade::FEmbeddedSpringBaseFacade(FManagedArrayCollection& InCollection, const FName& InVerticesGroup)
+		: FPositionTargetFacade(InCollection, InVerticesGroup)
+		, SpringLengthAttribute(InCollection, SpringLength, GroupName)
+		, CompressionStiffnessAttribute(InCollection, CompressionStiffness, GroupName)
+		, ConstraintStartAttribute(InCollection, ConstraintStart, SpringConstraintGroupName, GroupName)
+		, ConstraintEndAttribute(InCollection, ConstraintEnd, SpringConstraintGroupName, GroupName)
+		, ConstraintEndPointNumIndicesAttribute(InCollection, ConstraintEndPointNumIndices, SpringConstraintGroupName)
+
+	{
+		DefineSchema();
+	}
+
+	FEmbeddedSpringBaseFacade::FEmbeddedSpringBaseFacade(const FManagedArrayCollection& InCollection, const FName& InVerticesGroup)
+		: FPositionTargetFacade(InCollection, InVerticesGroup)
+		, SpringLengthAttribute(InCollection, SpringLength, GroupName)
+		, CompressionStiffnessAttribute(InCollection, CompressionStiffness, GroupName)
+		, ConstraintStartAttribute(InCollection, ConstraintStart, SpringConstraintGroupName, GroupName)
+		, ConstraintEndAttribute(InCollection, ConstraintEnd, SpringConstraintGroupName, GroupName)
+		, ConstraintEndPointNumIndicesAttribute(InCollection, ConstraintEndPointNumIndices, SpringConstraintGroupName)
+	{
+	}
+
+	bool FEmbeddedSpringBaseFacade::IsValid() const
+	{
+		return FPositionTargetFacade::IsValid()
+			&& SpringLengthAttribute.IsValid() && CompressionStiffnessAttribute.IsValid()
+			&& ConstraintStartAttribute.IsValid() && ConstraintEndAttribute.IsValid()
+			&& ConstraintEndPointNumIndicesAttribute.IsValid();
+	}
+
+	void FEmbeddedSpringBaseFacade::DefineSchema()
+	{
+		check(!IsConst());
+		FPositionTargetFacade::DefineSchema();
+
+		SpringLengthAttribute.Add();
+		CompressionStiffnessAttribute.Add();
+
+		ConstraintStartAttribute.Add();
+		ConstraintEndAttribute.Add();
+		ConstraintEndPointNumIndicesAttribute.Add();
+	}
+
+	void FEmbeddedSpringFacade::Append(const FEmbeddedSpringFacade& Other, int32 VertexOffset)
+	{
+		check(!IsConst());
+		check(IsValid() && Other.IsValid());
+		const int32 StartNumConstraints = GetNumSpringConstraints();
+		const int32 OtherNumConstraints = Other.GetNumSpringConstraints();
+		SetNumSpringConstraints(StartNumConstraints + OtherNumConstraints);
+		for (int32 ConstraintIndex = 0; ConstraintIndex < OtherNumConstraints; ++ConstraintIndex)
+		{
+			GetSpringConstraint(StartNumConstraints + ConstraintIndex).Initialize(Other.GetSpringConstraintConst(ConstraintIndex), VertexOffset);
+		}
+	}
+
+	void FEmbeddedSpringFacade::SetNumSpringConstraints(int32 InNumConstraints)
+	{
+		const int32 NumConstraints = GetNumSpringConstraints();
+		for (int32 ConstraintIndex = InNumConstraints; ConstraintIndex < NumConstraints; ++ConstraintIndex)
+		{
+			GetSpringConstraint(ConstraintIndex).Reset();
+		}
+		ConstraintStartAttribute.SetNumElements(InNumConstraints);
+		for (int32 ConstraintIndex = NumConstraints; ConstraintIndex < InNumConstraints; ++ConstraintIndex)
+		{
+			GetSpringConstraint(ConstraintIndex).SetDefaults();
+		}
+	}
+
+	FEmbeddedSpringConstraintFacade FEmbeddedSpringFacade::GetSpringConstraint(int32 ConstraintIndex)
+	{
+		return IsConst() ? FEmbeddedSpringConstraintFacade(ConstCollection, VerticesGroup, ConstraintIndex) : FEmbeddedSpringConstraintFacade(*Collection, VerticesGroup, ConstraintIndex);
+	}
+
+	FEmbeddedSpringConstraintFacade FEmbeddedSpringFacade::GetSpringConstraintConst(int32 ConstraintIndex) const
+	{
+		return FEmbeddedSpringConstraintFacade(ConstCollection, VerticesGroup, ConstraintIndex);
+	}
+
+	int32 FEmbeddedSpringFacade::AddSpringConstraint()
+	{
+		const int32 ConstraintIndex = GetNumSpringConstraints();
+		SetNumSpringConstraints(ConstraintIndex + 1);
+		return ConstraintIndex;
+	}
+
+	void FEmbeddedSpringFacade::RemoveSpringConstraints(const TArray<int32>& SortedDeletionList)
+	{
+		check(!IsConst());
+		for (const int32 ConstraintToRemove : SortedDeletionList)
+		{
+			GetSpringConstraint(ConstraintToRemove).Reset();
+		}
+		Collection->RemoveElements(SpringConstraintGroupName, SortedDeletionList);
+	}
+
+	uint32 FEmbeddedSpringFacade::CalculateTypeHash(uint32 PreviousHash) const
+	{
+		uint32 Hash = PreviousHash;
+		Hash = HashCombineFast(Hash, GetNumSpringConstraints());
+		for (int32 ConstraintIndex = 0; ConstraintIndex < GetNumSpringConstraints(); ++ConstraintIndex)
+		{
+			const FEmbeddedSpringConstraintFacade Constraint = GetSpringConstraintConst(ConstraintIndex);
+			Hash = Constraint.CalculateTypeHash(Hash);
+		}
+		return Hash;
+	}
+
+	void FEmbeddedSpringFacade::CleanupAndCompactInvalidSprings()
+	{
+		TArray<int32> SortedConstraintsToRemove;
+		for (int32 ConstraintIndex = 0; ConstraintIndex < GetNumSpringConstraints(); ++ConstraintIndex)
+		{
+			FEmbeddedSpringConstraintFacade Constraint = GetSpringConstraint(ConstraintIndex);
+			Constraint.CleanupAndCompactInvalidSprings();
+			if (Constraint.GetNumSprings() == 0)
+			{
+				SortedConstraintsToRemove.Add(ConstraintIndex);
+			}
+		}
+		if (!SortedConstraintsToRemove.IsEmpty())
+		{
+			RemoveSpringConstraints(SortedConstraintsToRemove);
+		}
+	}
+
+	FEmbeddedSpringConstraintFacade::FEmbeddedSpringConstraintFacade(FManagedArrayCollection& InCollection, const FName& InVerticesGroup, int32 InConstraintIndex)
+		:Base(InCollection, InVerticesGroup)
+		, ConstraintIndex(InConstraintIndex)
+	{
+		check(IsValid());
+		check(ConstraintIndex >= 0 && ConstraintIndex < ConstraintStartAttribute.Num());
+	}
+
+	FEmbeddedSpringConstraintFacade::FEmbeddedSpringConstraintFacade(const FManagedArrayCollection& InCollection, const FName& InVerticesGroup, int32 InConstraintIndex)
+		: Base(InCollection, InVerticesGroup)
+		, ConstraintIndex(InConstraintIndex)
+	{
+		check(IsValid());
+		check(ConstraintIndex >= 0 && ConstraintIndex < ConstraintStartAttribute.Num());
+	}
+
+	void FEmbeddedSpringConstraintFacade::SetDefaults()
+	{
+		check(!IsConst());
+		ConstraintStartAttribute.ModifyAt(ConstraintIndex, INDEX_NONE);
+		ConstraintEndAttribute.ModifyAt(ConstraintIndex, INDEX_NONE);
+	}
+
+	void FEmbeddedSpringConstraintFacade::Reset()
+	{
+		SetNumSprings(0);
+		SetDefaults();
+	}
+
+	int32 FEmbeddedSpringConstraintFacade::GetNumSprings() const
+	{
+		// Taken from ClothCollection::GetNumElements.
+		// If this pattern is used more in the GeometryCollection facades, we could add this to a utility.
+		if (ConstraintStartAttribute.IsValid() && ConstraintEndAttribute.IsValid())
+		{
+			const int32 Start = ConstraintStartAttribute[ConstraintIndex];
+			const int32 End = ConstraintEndAttribute[ConstraintIndex];
+			if (Start != INDEX_NONE && End != INDEX_NONE)
+			{
+				return End - Start + 1;
+			}
+			checkf(Start == End, TEXT("Only one boundary of the range is set to INDEX_NONE, when both should."));
+		}
+		return 0;
+	}
+
+	void FEmbeddedSpringConstraintFacade::RemoveSprings(const TArray<int32>& SortedDeletionList)
+	{
+		check(!IsConst());
+		if (SortedDeletionList.IsEmpty())
+		{
+			return;
+		}
+		TArray<int32> GlobalIndexSortedDeletionList;
+		const int32 Offset = ConstraintStartAttribute[ConstraintIndex];
+		GlobalIndexSortedDeletionList.SetNumUninitialized(SortedDeletionList.Num());
+		for (int32 Idx = 0; Idx < SortedDeletionList.Num(); ++Idx)
+		{
+			GlobalIndexSortedDeletionList[Idx] = SortedDeletionList[Idx] + Offset;
+		}
+
+		// Taken from ClothCollection::RemoveElements.
+		// If this pattern is used more in the GeometryCollection facades, we could add this to a utility.
+		check(ConstraintStartAttribute.IsValid() && ConstraintEndAttribute.IsValid());
+		TManagedArray<int32>& ConstraintStartArray = ConstraintStartAttribute.Modify();
+		TManagedArray<int32>& ConstraintEndArray = ConstraintEndAttribute.Modify();
+
+		int32& Start = ConstraintStartArray[ConstraintIndex];
+		int32& End = ConstraintEndArray[ConstraintIndex];
+		check(Start != INDEX_NONE && End != INDEX_NONE);
+
+		const int32 OrigStart = Start;
+		const int32 OrigNumElements = End - Start + 1;
+
+		check(GlobalIndexSortedDeletionList[0] >= Start);
+		check(GlobalIndexSortedDeletionList.Last() <= End);
+		check(OrigNumElements >= GlobalIndexSortedDeletionList.Num());
+
+		Collection->RemoveElements(GroupName, GlobalIndexSortedDeletionList);
+
+		if (GlobalIndexSortedDeletionList.Num() == OrigNumElements)
+		{
+			Start = End = INDEX_NONE;
+		}
+		else
+		{
+			const int32 NewNumElements = OrigNumElements - GlobalIndexSortedDeletionList.Num();
+			const int32 NewEnd = OrigStart + NewNumElements - 1;
+			check(Start == OrigStart || Start == INDEX_NONE);
+			check(End == NewEnd || End == INDEX_NONE);
+			Start = OrigStart;
+			End = NewEnd;
+		}
+	}
+
+	void FEmbeddedSpringConstraintFacade::SetNumSprings(int32 InNumSprings)
+	{
+		// Taken from ClothCollection::SetNumElements.
+		// If this pattern is used more in the GeometryCollection facades, we could add this to a utility.
+		check(!IsConst());
+		check(InNumSprings >= 0);
+		check(ConstraintStartAttribute.IsValid() && ConstraintEndAttribute.IsValid());
+		TManagedArray<int32>& ConstraintStartArray = ConstraintStartAttribute.Modify();
+		TManagedArray<int32>& ConstraintEndArray = ConstraintEndAttribute.Modify();
+		
+		int32& Start = ConstraintStartArray[ConstraintIndex];
+		int32& End = ConstraintEndArray[ConstraintIndex];
+		check(Start != INDEX_NONE || End == INDEX_NONE);  // Best to avoid situations where only one boundary of the range is set to INDEX_NONE
+
+		const int32 NumSprings = (Start == INDEX_NONE) ? 0 : End - Start + 1;
+		if (const int32 Delta = InNumSprings - NumSprings)
+		{
+			if (Delta > 0)
+			{
+				// Find a previous valid index range to insert after when the range is empty
+				auto ComputeEnd = [this, &ConstraintEndArray]()->int32
+					{
+						for (int32 Index = ConstraintIndex; Index >= 0; --Index)
+						{
+							if (ConstraintEndArray[Index] != INDEX_NONE)
+							{
+								return ConstraintEndArray[Index];
+							}
+						}
+						return INDEX_NONE;
+					};
+
+				// Grow the array
+				const int32 Position = ComputeEnd() + 1;
+				Collection->InsertElements(Delta, Position, GroupName);
+
+				// Update Start/End
+				if (!NumSprings)
+				{
+					Start = Position;
+				}
+				End = Start + InNumSprings - 1;
+			}
+			else
+			{
+				// Shrink the array
+				const int32 Position = Start + InNumSprings;
+				Collection->RemoveElements(GroupName, -Delta, Position);
+
+				// Update Start/End
+				if (InNumSprings)
+				{
+					End = Position - 1;
+				}
+				else
+				{
+					End = Start = INDEX_NONE;  // It is important to set the start & end to INDEX_NONE so that they never get automatically re-indexed by the managed array collection
+				}
+			}
+		}
+	}
+
+	void FEmbeddedSpringConstraintFacade::Initialize(const TConstArrayView<FIntVector2>& EndPoints, const TConstArrayView<float>& InSpringLength,
+		const TConstArrayView<float>& InExtensionStiffnessWeight, const TConstArrayView<float>& InCompressionStiffnessWeight, const TConstArrayView<float>& InDampingWeight)
+	{
+		check(!IsConst());
+		Reset();
+
+		//~ Spring Constraint Group
+		SetConstraintEndPointNumIndices(FUintVector2(1, 1));
+
+		//~ Group PositionTargets
+		const int32 InNumSprings = EndPoints.Num();
+		check(InNumSprings == InSpringLength.Num());
+		SetNumSprings(InNumSprings);
+		TArrayView<TArray<int32>> SourceIndexArray = GetSourceIndex();
+		TArrayView<TArray<int32>> TargetIndexArray = GetTargetIndex();
+		TArrayView<TArray<float>> SourceWeightArray = GetSourceWeights();
+		TArrayView<TArray<float>> TargetWeightArray = GetTargetWeights();
+		TArrayView<float> SpringLengthArray = GetSpringLength();
+		TArrayView<float> ExtensionStiffnessArray = GetExtensionStiffness();
+		TArrayView<float> CompressionStiffnessArray = GetCompressionStiffness();
+		TArrayView<float> DampingArray = GetDamping();
+		for (int32 SpringIdx = 0; SpringIdx < InNumSprings; ++SpringIdx)
+		{
+			SourceIndexArray[SpringIdx] = { EndPoints[SpringIdx][0] };
+			TargetIndexArray[SpringIdx] = { EndPoints[SpringIdx][1] };
+			SourceWeightArray[SpringIdx] = { 1.f };
+			TargetWeightArray[SpringIdx] = { 1.f };
+			SpringLengthArray[SpringIdx] = InSpringLength[SpringIdx];
+			ExtensionStiffnessArray[SpringIdx] = InExtensionStiffnessWeight.IsValidIndex(SpringIdx) ? InExtensionStiffnessWeight[SpringIdx] : 0.f;
+			CompressionStiffnessArray[SpringIdx] = InCompressionStiffnessWeight.IsValidIndex(SpringIdx) ? InCompressionStiffnessWeight[SpringIdx] : 0.f;
+			DampingArray[SpringIdx] = InDampingWeight.IsValidIndex(SpringIdx) ? InDampingWeight[SpringIdx] : 0.f;
+		}
+	}
+
+	void FEmbeddedSpringConstraintFacade::Initialize(const FUintVector2& EndPointNumIndices, const TConstArrayView<TArray<int32>>& InSourceIndices, const TConstArrayView<TArray<float>>& InSourceWeights,
+		const TConstArrayView<TArray<int32>>& InTargetIndices, const TConstArrayView<TArray<float>>& InTargetWeights, 
+		const TConstArrayView<float>& InSpringLength, 
+		const TConstArrayView<float>& InExtensionStiffnessWeight, const TConstArrayView<float>& InCompressionStiffnessWeight, const TConstArrayView<float>& InDampingWeight)
+	{
+		check(!IsConst());
+		Reset();
+
+		//~ Spring Constraint Group
+		SetConstraintEndPointNumIndices(EndPointNumIndices);
+
+		//~ Group PositionTargets
+		const int32 InNumSprings = InSourceIndices.Num();
+		check(InNumSprings == InSourceWeights.Num());
+		check(InNumSprings == InTargetIndices.Num());
+		check(InNumSprings == InTargetWeights.Num());
+		check(InNumSprings == InSpringLength.Num());
+		SetNumSprings(InNumSprings);
+		TArrayView<float> SpringLengthArray = GetSpringLength();
+		TArrayView<float> ExtensionStiffnessArray = GetExtensionStiffness();
+		TArrayView<float> CompressionStiffnessArray = GetCompressionStiffness();
+		TArrayView<float> DampingArray = GetDamping();
+		for (int32 SpringIdx = 0; SpringIdx < InNumSprings; ++SpringIdx)
+		{
+			SetIndicesAndWeights(SpringIdx, TConstArrayView<int32>(InSourceIndices[SpringIdx]), TConstArrayView<float>(InSourceWeights[SpringIdx]),
+				TConstArrayView<int32>(InTargetIndices[SpringIdx]), TConstArrayView<float>(InTargetWeights[SpringIdx]));
+			SpringLengthArray[SpringIdx] = InSpringLength[SpringIdx];
+			ExtensionStiffnessArray[SpringIdx] = InExtensionStiffnessWeight.IsValidIndex(SpringIdx) ? InExtensionStiffnessWeight[SpringIdx] : 0.f;
+			CompressionStiffnessArray[SpringIdx] = InCompressionStiffnessWeight.IsValidIndex(SpringIdx) ? InCompressionStiffnessWeight[SpringIdx] : 0.f;
+			DampingArray[SpringIdx] = InDampingWeight.IsValidIndex(SpringIdx) ? InDampingWeight[SpringIdx] : 0.f;
+		}
+	}
+
+	void FEmbeddedSpringConstraintFacade::Initialize(const FEmbeddedSpringConstraintFacade& Other, const int32 VertexOffset)
+	{
+		check(!IsConst());
+		check(VerticesGroup == Other.VerticesGroup);
+		Reset();
+
+		//~ Spring Constraint Group
+		SetConstraintEndPointNumIndices(Other.GetConstraintEndPointNumIndices());
+
+		//~ Group PositionTargets
+		SetNumSprings(Other.GetNumSprings());
+		Private::CopyArrayViewDataAndApplyOffset(GetSourceIndex(), Other.GetSourceIndexConst(), VertexOffset);
+		Private::CopyArrayViewDataAndApplyOffset(GetTargetIndex(), Other.GetTargetIndexConst(), VertexOffset);
+		Private::CopyArrayViewData(GetSourceWeights(), Other.GetSourceWeightsConst());
+		Private::CopyArrayViewData(GetTargetWeights(), Other.GetTargetWeightsConst());
+		Private::CopyArrayViewData(GetSpringLength(), Other.GetSpringLengthConst());
+		Private::CopyArrayViewData(GetExtensionStiffness(), Other.GetExtensionStiffnessConst());
+		Private::CopyArrayViewData(GetCompressionStiffness(), Other.GetCompressionStiffnessConst());
+		Private::CopyArrayViewData(GetDamping(), Other.GetDampingConst());
+	}
+
+	void FEmbeddedSpringConstraintFacade::Append(const FEmbeddedSpringConstraintFacade& Other, const int32 VertexOffset)
+	{
+		if (Other.GetNumSprings() == 0)
+		{
+			return;
+		}
+
+		check(!IsConst());
+		check(VerticesGroup == Other.VerticesGroup);
+		check(GetConstraintEndPointNumIndices() == Other.GetConstraintEndPointNumIndices());
+
+		//~ Group PositionTargets
+		const int32 OrigNumSprings = GetNumSprings();
+		const int32 OtherNumSprings = Other.GetNumSprings();
+		SetNumSprings(OrigNumSprings + OtherNumSprings);
+		Private::CopyArrayViewDataAndApplyOffset(TArrayView<TArray<int32>>(GetSourceIndex().GetData() + OrigNumSprings, OtherNumSprings), Other.GetSourceIndexConst(), VertexOffset);
+		Private::CopyArrayViewDataAndApplyOffset(TArrayView<TArray<int32>>(GetTargetIndex().GetData() + OrigNumSprings, OtherNumSprings), Other.GetTargetIndexConst(), VertexOffset);
+		Private::CopyArrayViewData(TArrayView<TArray<float>>(GetSourceWeights().GetData() + OrigNumSprings, OtherNumSprings), Other.GetSourceWeightsConst());
+		Private::CopyArrayViewData(TArrayView<TArray<float>>(GetTargetWeights().GetData() + OrigNumSprings, OtherNumSprings), Other.GetTargetWeightsConst());
+		Private::CopyArrayViewData(TArrayView<float>(GetSpringLength().GetData() + OrigNumSprings, OtherNumSprings), Other.GetSpringLengthConst());
+		Private::CopyArrayViewData(TArrayView<float>(GetExtensionStiffness().GetData() + OrigNumSprings, OtherNumSprings), Other.GetExtensionStiffnessConst());
+		Private::CopyArrayViewData(TArrayView<float>(GetCompressionStiffness().GetData() + OrigNumSprings, OtherNumSprings), Other.GetCompressionStiffnessConst());
+		Private::CopyArrayViewData(TArrayView<float>(GetDamping().GetData() + OrigNumSprings, OtherNumSprings), Other.GetDampingConst());
+	}
+
+	TArrayView<float> FEmbeddedSpringConstraintFacade::GetSpringLength()
+	{
+		return Private::GetArrayView(SpringLengthAttribute, ConstraintStartAttribute, ConstraintEndAttribute, ConstraintIndex);
+	}
+
+	TConstArrayView<float> FEmbeddedSpringConstraintFacade::GetSpringLengthConst() const
+	{
+		return Private::GetConstArrayView(SpringLengthAttribute, ConstraintStartAttribute, ConstraintEndAttribute, ConstraintIndex);
+	}
+
+	TArrayView<float> FEmbeddedSpringConstraintFacade::GetExtensionStiffness()
+	{
+		return Private::GetArrayView(StiffnessAttribute, ConstraintStartAttribute, ConstraintEndAttribute, ConstraintIndex);
+	}
+
+	TConstArrayView<float> FEmbeddedSpringConstraintFacade::GetExtensionStiffnessConst() const
+	{
+		return Private::GetConstArrayView(StiffnessAttribute, ConstraintStartAttribute, ConstraintEndAttribute, ConstraintIndex);
+	}
+
+	TArrayView<float> FEmbeddedSpringConstraintFacade::GetCompressionStiffness()
+	{
+		return Private::GetArrayView(CompressionStiffnessAttribute, ConstraintStartAttribute, ConstraintEndAttribute, ConstraintIndex);
+	}
+
+	TConstArrayView<float> FEmbeddedSpringConstraintFacade::GetCompressionStiffnessConst() const
+	{
+		return Private::GetConstArrayView(CompressionStiffnessAttribute, ConstraintStartAttribute, ConstraintEndAttribute, ConstraintIndex);
+	}
+
+	TArrayView<float> FEmbeddedSpringConstraintFacade::GetDamping()
+	{
+		return Private::GetArrayView(DampingAttribute, ConstraintStartAttribute, ConstraintEndAttribute, ConstraintIndex);
+	}
+
+	TConstArrayView<float> FEmbeddedSpringConstraintFacade::GetDampingConst() const
+	{
+		return Private::GetConstArrayView(DampingAttribute, ConstraintStartAttribute, ConstraintEndAttribute, ConstraintIndex);
+	}
+
+	TArrayView<TArray<int32>> FEmbeddedSpringConstraintFacade::GetSourceIndex()
+	{
+		return Private::GetArrayView(SourceIndexAttribute, ConstraintStartAttribute, ConstraintEndAttribute, ConstraintIndex);
+	}
+
+	TConstArrayView<TArray<int32>> FEmbeddedSpringConstraintFacade::GetSourceIndexConst() const
+	{
+		return Private::GetConstArrayView(SourceIndexAttribute, ConstraintStartAttribute, ConstraintEndAttribute, ConstraintIndex);
+	}
+
+	TArrayView<TArray<float>> FEmbeddedSpringConstraintFacade::GetSourceWeights()
+	{
+		return Private::GetArrayView(SourceWeightsAttribute, ConstraintStartAttribute, ConstraintEndAttribute, ConstraintIndex);
+	}
+
+	TConstArrayView<TArray<float>> FEmbeddedSpringConstraintFacade::GetSourceWeightsConst() const
+	{
+		return Private::GetConstArrayView(SourceWeightsAttribute, ConstraintStartAttribute, ConstraintEndAttribute, ConstraintIndex);
+	}
+
+	TArrayView<TArray<int32>> FEmbeddedSpringConstraintFacade::GetTargetIndex()
+	{
+		return Private::GetArrayView(TargetIndexAttribute, ConstraintStartAttribute, ConstraintEndAttribute, ConstraintIndex);
+	}
+
+	TConstArrayView<TArray<int32>> FEmbeddedSpringConstraintFacade::GetTargetIndexConst() const
+	{
+		return Private::GetConstArrayView(TargetIndexAttribute, ConstraintStartAttribute, ConstraintEndAttribute, ConstraintIndex);
+	}
+
+	TArrayView<TArray<float>> FEmbeddedSpringConstraintFacade::GetTargetWeights()
+	{
+		return Private::GetArrayView(TargetWeightsAttribute, ConstraintStartAttribute, ConstraintEndAttribute, ConstraintIndex);
+	}
+
+	TConstArrayView<TArray<float>> FEmbeddedSpringConstraintFacade::GetTargetWeightsConst() const
+	{
+		return Private::GetConstArrayView(TargetWeightsAttribute, ConstraintStartAttribute, ConstraintEndAttribute, ConstraintIndex);
+	}
+
+	void FEmbeddedSpringConstraintFacade::SetIndicesAndWeights(const int32 SpringIndex, const TConstArrayView<int32>& InSourceIndices, const TConstArrayView<float>& InSourceWeights,
+		const TConstArrayView<int32>& InTargetIndices, const TConstArrayView<float>& InTargetWeights)
+	{
+		const FUintVector2 EndPointNumIndices = GetConstraintEndPointNumIndices();
+		Private::SetIndicesAndWeights(EndPointNumIndices[0], GetSourceIndex()[SpringIndex], GetSourceWeights()[SpringIndex], InSourceIndices, InSourceWeights);
+		Private::SetIndicesAndWeights(EndPointNumIndices[1], GetTargetIndex()[SpringIndex], GetTargetWeights()[SpringIndex], InTargetIndices, InTargetWeights);
+	}
+
+	uint32 FEmbeddedSpringConstraintFacade::CalculateTypeHash(uint32 PreviousHash) const
+	{
+		uint32 Hash = PreviousHash;
+		Hash = HashCombineFast(Hash, GetTypeHash(GetConstraintEndPointNumIndices()));
+		Hash = GetArrayHash(GetSourceIndexConst().GetData(), GetSourceIndexConst().Num(), Hash);
+		Hash = GetArrayHash(GetSourceWeightsConst().GetData(), GetSourceWeightsConst().Num(), Hash);
+		Hash = GetArrayHash(GetTargetIndexConst().GetData(), GetTargetIndexConst().Num(), Hash);
+		Hash = GetArrayHash(GetTargetWeightsConst().GetData(), GetTargetWeightsConst().Num(), Hash);
+		Hash = GetArrayHash(GetSpringLengthConst().GetData(), GetSpringLengthConst().Num(), Hash);
+		Hash = GetArrayHash(GetExtensionStiffnessConst().GetData(), GetExtensionStiffnessConst().Num(), Hash);
+		Hash = GetArrayHash(GetCompressionStiffnessConst().GetData(), GetCompressionStiffnessConst().Num(), Hash);
+		Hash = GetArrayHash(GetDampingConst().GetData(), GetDampingConst().Num(), Hash);
+		return Hash;
+	}
+
+	void FEmbeddedSpringConstraintFacade::CleanupAndCompactInvalidSprings()
+	{
+		TArray<int32> SortedSpringsToRemove;
+		const int32 NumSprings = GetNumSprings();
+
+		for (int32 Index = 0; Index < NumSprings; ++Index)
+		{
+			if (Private::CleanIndicesAndWeights(GetSourceIndex()[Index], GetSourceWeights()[Index]))
+			{
+				SortedSpringsToRemove.Add(Index);
+				continue;
+			}
+			if (Private::CleanIndicesAndWeights(GetTargetIndex()[Index], GetTargetWeights()[Index]))
+			{
+				SortedSpringsToRemove.Add(Index);
+			}
+		}
+
+		if (!SortedSpringsToRemove.IsEmpty())
+		{
+			RemoveSprings(SortedSpringsToRemove);
+		}
+	}
+}
diff --git a/Engine/Source/Runtime/Experimental/Chaos/Private/GeometryCollection/Facades/CollectionPositionTargetFacade.cpp b/Engine/Source/Runtime/Experimental/Chaos/Private/GeometryCollection/Facades/CollectionPositionTargetFacade.cpp
--- a/Engine/Source/Runtime/Experimental/Chaos/Private/GeometryCollection/Facades/CollectionPositionTargetFacade.cpp
+++ b/Engine/Source/Runtime/Experimental/Chaos/Private/GeometryCollection/Facades/CollectionPositionTargetFacade.cpp
@@ -5,7 +5,6 @@
 =============================================================================*/
 
 #include "GeometryCollection/Facades/CollectionPositionTargetFacade.h"
-#include "GeometryCollection/GeometryCollection.h"
 
 namespace GeometryCollection::Facades
 {
@@ -20,13 +19,12 @@ namespace GeometryCollection::Facades
 	const FName FPositionTargetFacade::TargetWeights("TargetWeights");
 	const FName FPositionTargetFacade::SourceWeights("SourceWeights");
 
-	FPositionTargetFacade::FPositionTargetFacade(FManagedArrayCollection& InCollection)
+	FPositionTargetFacade::FPositionTargetFacade(FManagedArrayCollection& InCollection, const FName& InVerticesGroup)
 		: ConstCollection(InCollection)
 		, Collection(&InCollection)
-		, TargetIndexAttribute(InCollection, TargetIndex, GroupName)
-		, SourceIndexAttribute(InCollection, SourceIndex, GroupName)
-		, TargetNameAttribute(InCollection, TargetName, GroupName)
-		, SourceNameAttribute(InCollection, SourceName, GroupName)
+		, VerticesGroup(InVerticesGroup)
+		, TargetIndexAttribute(InCollection, TargetIndex, GroupName, InVerticesGroup)
+		, SourceIndexAttribute(InCollection, SourceIndex, GroupName, InVerticesGroup)
 		, StiffnessAttribute(InCollection, Stiffness, GroupName)
 		, DampingAttribute(InCollection, Damping, GroupName)
 		, TargetWeightsAttribute(InCollection, TargetWeights, GroupName)
@@ -35,13 +33,12 @@ namespace GeometryCollection::Facades
 		DefineSchema();
 	}
 
-	FPositionTargetFacade::FPositionTargetFacade(const FManagedArrayCollection& InCollection)
+	FPositionTargetFacade::FPositionTargetFacade(const FManagedArrayCollection& InCollection, const FName& InVerticesGroup)
 		: ConstCollection(InCollection)
 		, Collection(nullptr)
-		, TargetIndexAttribute(InCollection, TargetIndex, GroupName)
-		, SourceIndexAttribute(InCollection, SourceIndex, GroupName)
-		, TargetNameAttribute(InCollection, TargetName, GroupName)
-		, SourceNameAttribute(InCollection, SourceName, GroupName)
+		, VerticesGroup(InVerticesGroup)
+		, TargetIndexAttribute(InCollection, TargetIndex, GroupName, InVerticesGroup)
+		, SourceIndexAttribute(InCollection, SourceIndex, GroupName, InVerticesGroup)
 		, StiffnessAttribute(InCollection, Stiffness, GroupName)
 		, DampingAttribute(InCollection, Damping, GroupName)
 		, TargetWeightsAttribute(InCollection, TargetWeights, GroupName)
@@ -52,18 +49,16 @@ namespace GeometryCollection::Facades
 
 	bool FPositionTargetFacade::IsValid() const
 	{
-		return TargetIndexAttribute.IsValid() && SourceIndexAttribute.IsValid()
-			&& TargetNameAttribute.IsValid() && SourceNameAttribute.IsValid() && StiffnessAttribute.IsValid()
-			&& DampingAttribute.IsValid() && TargetWeightsAttribute.IsValid() && SourceWeightsAttribute.IsValid();
+		return TargetIndexAttribute.IsValid() && TargetIndexAttribute.GetGroupDependency() ==  VerticesGroup 
+			&& SourceIndexAttribute.IsValid() && SourceIndexAttribute.GetGroupDependency() == VerticesGroup 
+			&& StiffnessAttribute.IsValid() && DampingAttribute.IsValid() && TargetWeightsAttribute.IsValid() && SourceWeightsAttribute.IsValid();
 	}
 
 	void FPositionTargetFacade::DefineSchema()
 	{
 		check(!IsConst());
-		TargetIndexAttribute.Add(ManageArrayAccessor::EPersistencePolicy::MakePersistent, FGeometryCollection::VerticesGroup);
-		SourceIndexAttribute.Add(ManageArrayAccessor::EPersistencePolicy::MakePersistent, FGeometryCollection::VerticesGroup);
-		TargetNameAttribute.Add();
-		SourceNameAttribute.Add();
+		TargetIndexAttribute.Add();
+		SourceIndexAttribute.Add();
 		StiffnessAttribute.Add();
 		DampingAttribute.Add();
 		TargetWeightsAttribute.Add();
@@ -77,28 +72,12 @@ namespace GeometryCollection::Facades
 		if (IsValid())
 		{
 			int32 NewIndex = TargetIndexAttribute.AddElements(1);
-			//TManagedArray<TArray<int32>>& TargetIndex = TargetIndexAttribute.Modify();
 			TargetIndexAttribute.Modify()[NewIndex] = InputData.TargetIndex;
 			SourceIndexAttribute.Modify()[NewIndex] = InputData.SourceIndex;
-			TargetNameAttribute.Modify()[NewIndex] = InputData.TargetName;
-			SourceNameAttribute.Modify()[NewIndex] = InputData.SourceName;
 			StiffnessAttribute.Modify()[NewIndex] = InputData.Stiffness;
 			DampingAttribute.Modify()[NewIndex] = InputData.Damping;
 			TargetWeightsAttribute.Modify()[NewIndex] = InputData.TargetWeights;
 			SourceWeightsAttribute.Modify()[NewIndex] = InputData.SourceWeights;
-			// 
-			// 
-			//TManagedArray<FString>& PositionTargetToGroupLocal = PositionTargetToGroupAttribute.Modify();
-
-			//PositionTarget[NewIndex] = Key.Index;
-			//PositionTargetToGroupLocal[NewIndex] = Key.GroupName.ToString();
-
-			//TManagedArrayAccessor<TArray<int32>> TargetIndexAttribute;
-			//TManagedArrayAccessor<TArray<int32>> SourceIndexAttribute;
-			//TManagedArrayAccessor<FString> TargetNameAttribute;
-			//TManagedArrayAccessor<FString> SourceNameAttribute;
-			//TManagedArrayAccessor<float> StiffnessAttribute;
-			//TManagedArrayAccessor<float> DampingAttribute;
 
 			return NewIndex;
 		}
@@ -120,16 +99,6 @@ namespace GeometryCollection::Facades
 				ReturnData.Damping = DampingAttribute.Get()[DataIndex];
 			}
 
-			if (SourceNameAttribute.Num() > DataIndex && DataIndex > -1)
-			{
-				ReturnData.SourceName = SourceNameAttribute.Get()[DataIndex];
-			}
-
-			if (TargetNameAttribute.Num() > DataIndex && DataIndex > -1)
-			{
-				ReturnData.TargetName = TargetNameAttribute.Get()[DataIndex];
-			}
-
 			if (SourceIndexAttribute.Num() > DataIndex && DataIndex > -1)
 			{
 				ReturnData.SourceIndex = SourceIndexAttribute.Get()[DataIndex];
diff --git a/Engine/Source/Runtime/Experimental/Chaos/Public/Chaos/CollectionEmbeddedSpringConstraintFacade.h b/Engine/Source/Runtime/Experimental/Chaos/Public/Chaos/CollectionEmbeddedSpringConstraintFacade.h
new file mode 100644
--- /dev/null
+++ b/Engine/Source/Runtime/Experimental/Chaos/Public/Chaos/CollectionEmbeddedSpringConstraintFacade.h
@@ -0,0 +1,177 @@
+// Copyright Epic Games, Inc. All Rights Reserved.
+
+#pragma once
+
+#include "GeometryCollection/Facades/CollectionPositionTargetFacade.h"
+
+namespace Chaos::Softs
+{
+	/** Base class for non-zero length springs.*/
+	class FEmbeddedSpringBaseFacade : public GeometryCollection::Facades::FPositionTargetFacade
+	{
+	public:
+		// Group PositionTargets
+		static CHAOS_API const FName SpringLength;
+		static CHAOS_API const FName CompressionStiffness; // Base Stiffness is ExtensionStiffness
+
+		// Collections of springs
+		static CHAOS_API const FName SpringConstraintGroupName;
+		static CHAOS_API const FName ConstraintStart;
+		static CHAOS_API const FName ConstraintEnd;
+		static CHAOS_API const FName ConstraintEndPointNumIndices;
+
+		CHAOS_API FEmbeddedSpringBaseFacade(FManagedArrayCollection& InCollection, const FName& InVerticesGroup = FGeometryCollection::VerticesGroup);
+		CHAOS_API FEmbeddedSpringBaseFacade(const FManagedArrayCollection& InCollection, const FName& InVerticesGroup = FGeometryCollection::VerticesGroup);
+
+		CHAOS_API bool IsValid() const;
+		CHAOS_API void DefineSchema();
+
+	protected:
+		//~ Group PositionTargets
+		TManagedArrayAccessor<float> SpringLengthAttribute;
+		TManagedArrayAccessor<float> CompressionStiffnessAttribute;
+
+		//~ Spring Constraint Group
+		TManagedArrayAccessor<int32> ConstraintStartAttribute;
+		TManagedArrayAccessor<int32> ConstraintEndAttribute;
+		TManagedArrayAccessor<FUintVector2> ConstraintEndPointNumIndicesAttribute;
+	};
+
+	/** Facade for managing groups of springs with the same embedding type (e.g., vertex-vertex or barycentric point-vertex) */
+	class FEmbeddedSpringConstraintFacade final : public FEmbeddedSpringBaseFacade
+	{
+		typedef FEmbeddedSpringBaseFacade Base;
+	public:
+		FEmbeddedSpringConstraintFacade() = delete;
+		FEmbeddedSpringConstraintFacade(const FEmbeddedSpringConstraintFacade&) = delete;
+		FEmbeddedSpringConstraintFacade(FEmbeddedSpringConstraintFacade&&) = default;
+		FEmbeddedSpringConstraintFacade& operator=(const FEmbeddedSpringConstraintFacade&) = delete;
+		FEmbeddedSpringConstraintFacade& operator=(FEmbeddedSpringConstraintFacade&&) = delete;
+
+		CHAOS_API void Reset();
+
+		int32 GetConstraintIndex() const
+		{
+			return ConstraintIndex;
+		}
+
+		CHAOS_API int32 GetNumSprings() const;
+		CHAOS_API void RemoveSprings(const TArray<int32>& SortedSpringsToRemove);
+
+		/** Initialize from a list of Vertex-Vertex springs.
+		 *  Spring Indices are VerticesGroup indices.
+		 *  Stiffness and Damping weights are expected have a value in [0.f, 1.f]. They will be set to constant 0.f if not supplied.
+		*/
+		CHAOS_API void Initialize(const TConstArrayView<FIntVector2>& EndPoints, const TConstArrayView<float>& InSpringLength,
+			const TConstArrayView<float>& InExtensionStiffnesWeight = TConstArrayView<float>(), const TConstArrayView<float>& InCompressionStiffnessWeight = TConstArrayView<float>(), const TConstArrayView<float>& InDampingWeight = TConstArrayView<float>());
+
+		/** Initialize from a list of springs. This is generic for any number of end point indices. If you have vertex-vertex constraints, use the Initialize method above which skips setting weights.
+		 *  Source Indices and Weights arrays are expected to be exactly EndPointNumIndices[0] long.
+		 *  Target Indices and Weights arrays are expected to be exactly EndPointNumIndices[1] long.
+		 *  Indices and Weights will be truncated or filled in with zeros if they aren't the expected length.
+		 *  Spring Indices are VerticesGroup indices.
+		 *  Stiffness and Damping weights are expected have a value in [0.f, 1.f]. They will be set to constant 0.f if not supplied.
+		*/
+		CHAOS_API void Initialize(const FUintVector2& EndPointNumIndices, const TConstArrayView<TArray<int32>>& InSourceIndices, const TConstArrayView<TArray<float>>& InSourceWeights,
+			const TConstArrayView<TArray<int32>>& InTargetIndices, const TConstArrayView<TArray<float>>& InTargetWeights,
+			const TConstArrayView<float>& InSpringLength,
+			const TConstArrayView<float>& InExtensionStiffnesWeight = TConstArrayView<float>(), const TConstArrayView<float>& InCompressionStiffnessWeight = TConstArrayView<float>(), const TConstArrayView<float>& InDampingWeight = TConstArrayView<float>());
+
+		//~ Spring Constraint Group
+		FUintVector2 GetConstraintEndPointNumIndices() const
+		{
+			return ConstraintEndPointNumIndicesAttribute[ConstraintIndex];
+		}
+
+		//~ Group PositionTargets
+		CHAOS_API TArrayView<float> GetSpringLength();
+		CHAOS_API TConstArrayView<float> GetSpringLengthConst() const;
+		CHAOS_API TArrayView<float> GetExtensionStiffness();
+		CHAOS_API TConstArrayView<float> GetExtensionStiffnessConst() const;
+		CHAOS_API TArrayView<float> GetCompressionStiffness();
+		CHAOS_API TConstArrayView<float> GetCompressionStiffnessConst() const;
+		CHAOS_API TArrayView<float> GetDamping();
+		CHAOS_API TConstArrayView<float> GetDampingConst() const;
+		CHAOS_API TConstArrayView<TArray<int32>> GetSourceIndexConst() const;
+		CHAOS_API TConstArrayView<TArray<float>> GetSourceWeightsConst() const;
+		CHAOS_API TConstArrayView<TArray<int32>> GetTargetIndexConst() const;
+		CHAOS_API TConstArrayView<TArray<float>> GetTargetWeightsConst() const;
+		/* Set Indices and Weights for a single spring.
+		 *  Source Indices and Weights arrays are expected to be exactly EndPointNumIndices[0] long.
+		 *  Target Indices and Weights arrays are expected to be exactly EndPointNumIndices[1] long.
+		 *  Indices and Weights will be truncated or filled in with zeros if they aren't the expected length.
+		 *  Spring Indices are VerticesGroup indices.*/
+		CHAOS_API void SetIndicesAndWeights(const int32 SpringIndex, const TConstArrayView<int32>& InSourceIndices, const TConstArrayView<float>& InSourceWeights,
+			const TConstArrayView<int32>& InTargetIndices, const TConstArrayView<float>& InTargetWeights);
+
+		/** Initialize from another constraint. */
+		CHAOS_API void Initialize(const FEmbeddedSpringConstraintFacade& Other, const int32 VertexOffset);
+		/** Append from another constraint. Must have same ConstraintEndPointNumIndices*/
+		CHAOS_API void Append(const FEmbeddedSpringConstraintFacade& Other, const int32 VertexOffset);
+
+		CHAOS_API uint32 CalculateTypeHash(uint32 PreviousHash = 0) const;
+
+		/** Remove springs with invalid vertices. */
+		CHAOS_API void CleanupAndCompactInvalidSprings();
+
+	private:
+		friend class FEmbeddedSpringFacade;
+		FEmbeddedSpringConstraintFacade(FManagedArrayCollection& InCollection, const FName& InVerticesGroup, int32 InConstraintIndex);
+		FEmbeddedSpringConstraintFacade(const FManagedArrayCollection& InCollection, const FName& InVerticesGroup, int32 InConstraintIndex);
+
+		void SetDefaults();
+		void SetNumSprings(int32 NumSprings);
+
+		//~ Spring Constraint Group
+		void SetConstraintEndPointNumIndices(const FUintVector2& NumIndices)
+		{
+			ConstraintEndPointNumIndicesAttribute.ModifyAt(ConstraintIndex, NumIndices);
+		}
+
+		//~ Group PositionTargets
+		TArrayView<TArray<int32>> GetSourceIndex();
+		TArrayView<TArray<float>> GetSourceWeights();
+		TArrayView<TArray<int32>> GetTargetIndex();
+		TArrayView<TArray<float>> GetTargetWeights();
+		int32 ConstraintIndex;
+	};
+
+	/** Facade for managing all FEmbeddedSpringConstraintFacade within a collection */
+	class FEmbeddedSpringFacade final : public FEmbeddedSpringBaseFacade
+	{
+		typedef FEmbeddedSpringBaseFacade Base;
+	public:
+
+		FEmbeddedSpringFacade(FManagedArrayCollection& InCollection, const FName& InVerticesGroup = FGeometryCollection::VerticesGroup)
+			: Base(InCollection, InVerticesGroup)
+		{
+		}
+
+		FEmbeddedSpringFacade(const FManagedArrayCollection& InCollection, const FName& InVerticesGroup = FGeometryCollection::VerticesGroup)
+			: Base(InCollection, InVerticesGroup)
+		{}
+
+		/** Append contents of another facade to this one. */
+		CHAOS_API void Append(const FEmbeddedSpringFacade& Other, int32 VertexOffset);
+
+		/** Constraints (lists of springs) management */
+		int32 GetNumSpringConstraints() const
+		{
+			return ConstraintStartAttribute.Num();
+		}
+		CHAOS_API void SetNumSpringConstraints(int32 Num);
+		CHAOS_API FEmbeddedSpringConstraintFacade GetSpringConstraint(int32 ConstraintIndex);
+		CHAOS_API FEmbeddedSpringConstraintFacade GetSpringConstraintConst(int32 ConstraintIndex) const;
+		CHAOS_API int32 AddSpringConstraint();
+		FEmbeddedSpringConstraintFacade AddGetSpringConstraint()
+		{
+			return GetSpringConstraint(AddSpringConstraint());
+		}
+		CHAOS_API void RemoveSpringConstraints(const TArray<int32>& SortedDeletionList);
+
+		CHAOS_API uint32 CalculateTypeHash(uint32 PreviousHash = 0) const;
+
+		/** Remove springs with invalid vertices.*/
+		CHAOS_API void CleanupAndCompactInvalidSprings();
+	};
+}
diff --git a/Engine/Source/Runtime/Experimental/Chaos/Public/GeometryCollection/Facades/CollectionPositionTargetFacade.h b/Engine/Source/Runtime/Experimental/Chaos/Public/GeometryCollection/Facades/CollectionPositionTargetFacade.h
--- a/Engine/Source/Runtime/Experimental/Chaos/Public/GeometryCollection/Facades/CollectionPositionTargetFacade.h
+++ b/Engine/Source/Runtime/Experimental/Chaos/Public/GeometryCollection/Facades/CollectionPositionTargetFacade.h
@@ -5,6 +5,7 @@
 #include "CoreMinimal.h"
 #include "GeometryCollection/ManagedArrayAccessor.h"
 #include "GeometryCollection/ManagedArrayCollection.h"
+#include "GeometryCollection/GeometryCollection.h"
 #include "GeometryCollection/Facades/CollectionSelectionFacade.h"
 
 namespace GeometryCollection::Facades
@@ -15,11 +16,20 @@ namespace GeometryCollection::Facades
 		TArray<int32> SourceIndex;
 		TArray<float> TargetWeights;
 		TArray<float> SourceWeights;
+		UE_DEPRECATED(5.6, "TargetName will be removed and is no longer populated.")
 		FString TargetName;
+		UE_DEPRECATED(5.6, "SourceName will be removed and is no longer populated.")
 		FString SourceName;
 		float Stiffness;
 		float Damping;
 
+		PRAGMA_DISABLE_DEPRECATION_WARNINGS
+		FPositionTargetsData() = default;
+		FPositionTargetsData(const FPositionTargetsData&) = default;
+		FPositionTargetsData(FPositionTargetsData&&) = default;
+		FPositionTargetsData& operator=(const FPositionTargetsData&) = default;
+		FPositionTargetsData& operator=(FPositionTargetsData&&) = default;
+		PRAGMA_ENABLE_DEPRECATION_WARNINGS
 	};
 
 	/** Kinematic Facade */
@@ -37,13 +47,19 @@ namespace GeometryCollection::Facades
 		static CHAOS_API const FName SourceIndex;
 		static CHAOS_API const FName Stiffness;
 		static CHAOS_API const FName Damping;
+		UE_DEPRECATED(5.6, "SourceName will be removed.")
 		static CHAOS_API const FName SourceName;
+		UE_DEPRECATED(5.6, "TargetName will be removed.")
 		static CHAOS_API const FName TargetName;
 		static CHAOS_API const FName TargetWeights;
 		static CHAOS_API const FName SourceWeights;
 
-		CHAOS_API FPositionTargetFacade(FManagedArrayCollection& InCollection);
-		CHAOS_API FPositionTargetFacade(const FManagedArrayCollection& InCollection);
+		CHAOS_API FPositionTargetFacade(FManagedArrayCollection& InCollection, const FName& InVerticesGroup = FGeometryCollection::VerticesGroup);
+		CHAOS_API FPositionTargetFacade(const FManagedArrayCollection& InCollection, const FName& InVerticesGroup = FGeometryCollection::VerticesGroup);
+
+		PRAGMA_DISABLE_DEPRECATION_WARNINGS
+		FPositionTargetFacade(FPositionTargetFacade&&) = default;
+		PRAGMA_ENABLE_DEPRECATION_WARNINGS
 
 		/** Create the facade attributes. */
 		CHAOS_API void DefineSchema();
@@ -61,14 +77,13 @@ namespace GeometryCollection::Facades
 		CHAOS_API FPositionTargetsData GetPositionTarget(const int32 DataIndex) const;
 		int32 NumPositionTargets() const { return TargetIndexAttribute.Num(); }
 
-	private:
+	protected:
 		const FManagedArrayCollection& ConstCollection;
 		FManagedArrayCollection* Collection = nullptr;
+		const FName VerticesGroup;
 
 		TManagedArrayAccessor<TArray<int32>> TargetIndexAttribute;
 		TManagedArrayAccessor<TArray<int32>> SourceIndexAttribute;
-		TManagedArrayAccessor<FString> TargetNameAttribute;
-		TManagedArrayAccessor<FString> SourceNameAttribute;
 		TManagedArrayAccessor<float> StiffnessAttribute;
 		TManagedArrayAccessor<float> DampingAttribute;
 		TManagedArrayAccessor<TArray<float>> TargetWeightsAttribute;
diff --git a/Engine/Source/Runtime/Experimental/Chaos/Public/GeometryCollection/ManagedArray.h b/Engine/Source/Runtime/Experimental/Chaos/Public/GeometryCollection/ManagedArray.h
--- a/Engine/Source/Runtime/Experimental/Chaos/Public/GeometryCollection/ManagedArray.h
+++ b/Engine/Source/Runtime/Experimental/Chaos/Public/GeometryCollection/ManagedArray.h
@@ -69,6 +69,11 @@ inline void TryBulkSerializeManagedArray(Chaos::FChaosArchive& Ar, TArray<uint8>
 	Array.BulkSerialize(Ar);
 }
 
+inline void TryBulkSerializeManagedArray(Chaos::FChaosArchive& Ar, TArray<FUintVector2>& Array)
+{
+	Array.BulkSerialize(Ar);
+}
+
 inline void TryBulkSerializeManagedArray(Chaos::FChaosArchive& Ar, TArray<FIntVector2>& Array)
 {
 	Array.BulkSerialize(Ar);
diff --git a/Engine/Source/Runtime/Experimental/Chaos/Public/GeometryCollection/ManagedArrayAccessor.h b/Engine/Source/Runtime/Experimental/Chaos/Public/GeometryCollection/ManagedArrayAccessor.h
--- a/Engine/Source/Runtime/Experimental/Chaos/Public/GeometryCollection/ManagedArrayAccessor.h
+++ b/Engine/Source/Runtime/Experimental/Chaos/Public/GeometryCollection/ManagedArrayAccessor.h
@@ -19,22 +19,22 @@ template <typename T>
 struct TManagedArrayAccessor
 {
 public:
-	TManagedArrayAccessor(FManagedArrayCollection& InCollection, const FName& InAttributeName, const FName& InAttributeGroup,  const FName& InGroupDependency = NAME_None)
+	TManagedArrayAccessor(FManagedArrayCollection& InCollection, const FName& InAttributeName, const FName& InAttributeGroup, const FName& InDefaultGroupDependency = NAME_None)
 		: ConstCollection(InCollection)
 		, Collection(&InCollection)
 		, Name(InAttributeName)
 		, Group(InAttributeGroup)
-		, GroupDependency(InGroupDependency)
+		, DefaultGroupDependency(InDefaultGroupDependency)
 		, AttributeArray(InCollection.FindAttributeTyped<T>(InAttributeName, InAttributeGroup))
 		, ConstAttributeArray(AttributeArray)
 	{}
 
-	TManagedArrayAccessor(const FManagedArrayCollection& InCollection, const FName& InAttributeName, const FName& InAttributeGroup, const FName& InGroupDependency = NAME_None)
+	TManagedArrayAccessor(const FManagedArrayCollection& InCollection, const FName& InAttributeName, const FName& InAttributeGroup, const FName& InDefaultGroupDependency = NAME_None)
 		: ConstCollection(InCollection)
 		, Collection(nullptr)
 		, Name(InAttributeName)
 		, Group(InAttributeGroup)
-		, GroupDependency(InGroupDependency)
+		, DefaultGroupDependency(InDefaultGroupDependency)
 		, AttributeArray(nullptr)
 		, ConstAttributeArray(InCollection.FindAttributeTyped<T>(InAttributeName, InAttributeGroup))
 	{}
@@ -46,7 +46,7 @@ public:
 
 	FName GetGroup() const { return Group; }
 
-	FName GetGroupDependency() const { return GroupDependency; }
+	FName GetGroupDependency() const { return IsValid() ? ConstCollection.GetDependency(Name, Group) : DefaultGroupDependency; }
 
 	bool IsConst() const { return Collection == nullptr; }
 
@@ -66,6 +66,30 @@ public:
 		return Collection->AddElements(NumElements, Group);
 	}
 
+	void RemoveElements(int32 NumElements, int32 Position)
+	{
+		check(!IsConst());
+		return Collection->RemoveElements(Group, NumElements, Position);
+	}
+
+	void SetNumElements(int32 InNumElements)
+	{
+		check(!IsConst());
+		check(InNumElements >= 0);
+		const int32 NumElements = Num();
+		if (const int32 Delta = InNumElements - NumElements)
+		{
+			if (Delta > 0)
+			{
+				Collection->AddElements(Delta, Group);
+			}
+			else
+			{
+				Collection->RemoveElements(Group, -Delta, InNumElements);
+			}
+		}
+	}
+
 	/** array style accessor */
 	const T& operator[](int32 Index) const
 	{
@@ -117,7 +141,7 @@ public:
 			Collection->AddGroup(Group);
 		}
 
-		FName LocalGroupDependency = GroupDependency;
+		FName LocalGroupDependency = DefaultGroupDependency;
 		if (!InGroupDependency.IsNone())
 		{
 			LocalGroupDependency = InGroupDependency;
@@ -184,7 +208,7 @@ private:
 
 	FName Name;
 	FName Group;
-	FName GroupDependency;
+	const FName DefaultGroupDependency; // Will be used by Add if a non-NAME_None group dependency is specified by that method.
 	TManagedArray<T>* AttributeArray;
 	const TManagedArray<T>* ConstAttributeArray;
 };
diff --git a/Engine/Source/Runtime/Experimental/Chaos/Public/GeometryCollection/ManagedArrayTypeValues.inl b/Engine/Source/Runtime/Experimental/Chaos/Public/GeometryCollection/ManagedArrayTypeValues.inl
--- a/Engine/Source/Runtime/Experimental/Chaos/Public/GeometryCollection/ManagedArrayTypeValues.inl
+++ b/Engine/Source/Runtime/Experimental/Chaos/Public/GeometryCollection/ManagedArrayTypeValues.inl
@@ -59,6 +59,7 @@ MANAGED_ARRAY_TYPE(TArray<FVector4f>, Vector4fArray)
 MANAGED_ARRAY_TYPE(Chaos::PMatrix33d, PMatrix33d)
 MANAGED_ARRAY_TYPE(TArray<Chaos::PMatrix33d>, PMatrix33dArray)
 MANAGED_ARRAY_TYPE(TArray<TArray<FVector3f>>, FVector3fNestedArray)
+MANAGED_ARRAY_TYPE(FUintVector2, UintVector2)
 
 // NOTE: new types must be added at the bottom to keep serialization from breaking
 
