Dataflow asset : remove UI access to variable from the asset itself ( will be accessed throug another part of the UI )

diff --git a/Engine/Source/Runtime/Experimental/Dataflow/Engine/Public/Dataflow/DataflowObject.h b/Engine/Source/Runtime/Experimental/Dataflow/Engine/Public/Dataflow/DataflowObject.h
--- a/Engine/Source/Runtime/Experimental/Dataflow/Engine/Public/Dataflow/DataflowObject.h
+++ b/Engine/Source/Runtime/Experimental/Dataflow/Engine/Public/Dataflow/DataflowObject.h
@@ -123,8 +123,7 @@ public:
 	UPROPERTY(EditAnywhere, Category = "Evaluation",meta=(EditConditionHides))
 	EDataflowType Type = EDataflowType::Construction;
 
-	// keep hidden for now until we have a proper UI for it 
-	UPROPERTY(EditAnywhere, Category=Variables)
+	UPROPERTY()
 	FInstancedPropertyBag Variables;
 
 public:
