[Physics] Delete some unused fwd declares. Chaos::FSingleParticleProxy is a class that doesn't actually exist, but intellisense doesn't know that, and it's a little annoying.

diff --git a/Engine/Source/Runtime/Experimental/Chaos/Public/PBDRigidsSolver.h b/Engine/Source/Runtime/Experimental/Chaos/Public/PBDRigidsSolver.h
--- a/Engine/Source/Runtime/Experimental/Chaos/Public/PBDRigidsSolver.h
+++ b/Engine/Source/Runtime/Experimental/Chaos/Public/PBDRigidsSolver.h
@@ -51,12 +51,8 @@ namespace ChaosTest
 */
 namespace Chaos
 {
-	class FPersistentPhysicsTask;
-	class FChaosArchive;
 	class FCharacterGroundConstraint;
 	class FCharacterGroundConstraintProxy;
-	class FSingleParticleProxy;
-	class FGeometryParticleBuffer;
 	class FClusterUnionPhysicsProxy;
 
 	template <typename T,typename R,int d>
