Dataflow : fix vertex selection to face selection conversion crash

diff --git a/Engine/Source/Runtime/Experimental/Chaos/Private/GeometryCollection/Facades/CollectionTransformSelectionFacade.cpp b/Engine/Source/Runtime/Experimental/Chaos/Private/GeometryCollection/Facades/CollectionTransformSelectionFacade.cpp
--- a/Engine/Source/Runtime/Experimental/Chaos/Private/GeometryCollection/Facades/CollectionTransformSelectionFacade.cpp
+++ b/Engine/Source/Runtime/Experimental/Chaos/Private/GeometryCollection/Facades/CollectionTransformSelectionFacade.cpp
@@ -1206,7 +1206,7 @@ namespace GeometryCollection::Facades
 			const int32 NumVertices = ConstCollection.NumElements(FGeometryCollection::VerticesGroup);
 
 			TArray<TArray<int32>> VertexToFaceIndices;
-			VertexToFaceIndices.Reset(NumVertices);
+			VertexToFaceIndices.SetNum(NumVertices);
 
 			const TManagedArray<FIntVector>& Indices = IndicesAttribute.Get();
 
