Dataflow:
- Fixed missing initialization

diff --git a/Engine/Source/Runtime/Experimental/Dataflow/Core/Public/Dataflow/DataflowSettings.h b/Engine/Source/Runtime/Experimental/Dataflow/Core/Public/Dataflow/DataflowSettings.h
--- a/Engine/Source/Runtime/Experimental/Dataflow/Core/Public/Dataflow/DataflowSettings.h
+++ b/Engine/Source/Runtime/Experimental/Dataflow/Core/Public/Dataflow/DataflowSettings.h
@@ -31,7 +31,7 @@ struct FTransformLevelColors
 	TArray<FLinearColor> LevelColors;
 
 	UPROPERTY(EditAnywhere, Category = Colors)
-	FLinearColor BlankColor;
+	FLinearColor BlankColor = FLinearColor(255, 255, 255, 255);
 };
 
 UCLASS(config = Editor, defaultconfig, meta = (DisplayName = "Dataflow"), MinimalAPI)
