Iris - Fix inverrted StreamDebugFeatures ensure.

#rnx

diff --git a/Engine/Source/Runtime/Experimental/Iris/Core/Private/Iris/ReplicationSystem/ReplicationReader.cpp b/Engine/Source/Runtime/Experimental/Iris/Core/Private/Iris/ReplicationSystem/ReplicationReader.cpp
--- a/Engine/Source/Runtime/Experimental/Iris/Core/Private/Iris/ReplicationSystem/ReplicationReader.cpp
+++ b/Engine/Source/Runtime/Experimental/Iris/Core/Private/Iris/ReplicationSystem/ReplicationReader.cpp
@@ -383,7 +383,7 @@ void FReplicationReader::ReadStreamDebugFeatures(FNetSerializationContext& Conte
 
 #if !UE_NET_REPLICATIONDATASTREAM_DEBUG
 	// If datastream debug features isn't enabled in this build the debug functions won't do any useful work and will cause bitstream errors.
-	if (!ensure(StreamDebugFeatures != EReplicationDataStreamDebugFeatures::None))
+	if (!ensure(StreamDebugFeatures == EReplicationDataStreamDebugFeatures::None))
 	{
 		UE_LOG(LogIris, Error, TEXT("StreamDebugFeatures enabled on sending side despite receiving side not being able to use them."));
 		Context.SetError(GNetError_InvalidValue);
