Fix IOStore JournaledCache crash

diff --git a/Engine/Source/Runtime/Experimental/IoStore/OnDemand/Private/JournaledCache.cpp b/Engine/Source/Runtime/Experimental/IoStore/OnDemand/Private/JournaledCache.cpp
--- a/Engine/Source/Runtime/Experimental/IoStore/OnDemand/Private/JournaledCache.cpp
+++ b/Engine/Source/Runtime/Experimental/IoStore/OnDemand/Private/JournaledCache.cpp
@@ -976,8 +976,14 @@ static int32 LoadCache(FDiskCache& DiskCache)
 			return nullptr;
 		}
 
+		Cursor += sizeof(FPhraseDesc);
+
 		const auto* FirstEntry = Header->Entries;
-		const uint32 BytesToConsume = sizeof(FDataEntry) * (FirstEntry->EntryCount + 1); //Entries + PhraseDesc
+		if (IsOob(Cursor + sizeof(FDataEntry)))
+		{
+			return nullptr;
+		}
+		const uint32 BytesToConsume = sizeof(FDataEntry) * (FirstEntry->EntryCount);
 		if (IsOob(Cursor + BytesToConsume))
 		{
 			return nullptr;
