Iris - Implemented IrisLogOnceTracker class to facilitate only logging something once per some arbitrary hash, for example the hash of a class name.

#rnx

diff --git a/Engine/Source/Runtime/Experimental/Iris/Core/Private/Iris/Core/IrisLogUtils.cpp b/Engine/Source/Runtime/Experimental/Iris/Core/Private/Iris/Core/IrisLogUtils.cpp
new file mode 100644
--- /dev/null
+++ b/Engine/Source/Runtime/Experimental/Iris/Core/Private/Iris/Core/IrisLogUtils.cpp
@@ -0,0 +1,22 @@
+// Copyright Epic Games, Inc. All Rights Reserved.
+
+#include "Iris/Core/IrisLogUtils.h"
+#include "UObject/NameTypes.h"
+#include "Templates/TypeHash.h"
+
+namespace UE::Net
+{
+
+bool FIrisLogOnceTracker::ShouldLog(uint32 Hash) const
+{
+	bool bAlreadySet = false;
+	LoggedHashes.Emplace(Hash, &bAlreadySet);
+	return !bAlreadySet;
+}
+
+bool FIrisLogOnceTracker::ShouldLog(const FName& Name) const
+{
+	return ShouldLog(::GetTypeHashHelper(Name));
+}
+
+}
diff --git a/Engine/Source/Runtime/Experimental/Iris/Core/Public/Iris/Core/IrisLogUtils.h b/Engine/Source/Runtime/Experimental/Iris/Core/Public/Iris/Core/IrisLogUtils.h
new file mode 100644
--- /dev/null
+++ b/Engine/Source/Runtime/Experimental/Iris/Core/Public/Iris/Core/IrisLogUtils.h
@@ -0,0 +1,30 @@
+// Copyright Epic Games, Inc. All Rights Reserved.
+
+#pragma once
+
+#include "HAL/Platform.h"
+#include "Containers/Set.h"
+
+class FName;
+
+namespace UE::Net
+{
+
+/**
+ * Utility class to determine whether some message coupled to a hash hasn't been logged yet. 
+ * Example usage would be for exmaple if you only want to log something once per class. 
+ */
+class FIrisLogOnceTracker
+{
+public:
+	/** Returns true if the Hash has not been encountered in a call to ShouldLog before. */
+	bool ShouldLog(uint32 Hash) const;
+
+	/** Returns true if the hash of the Name has not been encountered in a call to ShouldLog before. */
+	bool ShouldLog(const FName& Name) const;
+
+private:
+	mutable TSet<uint32> LoggedHashes;
+};
+
+}
