Iris - Log when someone adds a dependency on an object that isn't dynamically filtered.

diff --git a/Engine/Source/Runtime/Experimental/Iris/Core/Private/Iris/ReplicationSystem/Filtering/ReplicationFiltering.cpp b/Engine/Source/Runtime/Experimental/Iris/Core/Private/Iris/ReplicationSystem/Filtering/ReplicationFiltering.cpp
--- a/Engine/Source/Runtime/Experimental/Iris/Core/Private/Iris/ReplicationSystem/Filtering/ReplicationFiltering.cpp
+++ b/Engine/Source/Runtime/Experimental/Iris/Core/Private/Iris/ReplicationSystem/Filtering/ReplicationFiltering.cpp
@@ -9,6 +9,7 @@
 #include "Iris/Core/BitTwiddling.h"
 #include "Iris/Core/IrisCsv.h"
 #include "Iris/Core/IrisLog.h"
+#include "Iris/Core/IrisLogUtils.h"
 #include "Iris/Core/IrisProfiler.h"
 #include "Iris/ReplicationSystem/NetRefHandleManager.h"
 #include "Iris/ReplicationSystem/ReplicationConnections.h"
@@ -159,6 +160,11 @@ private:
 	const FNetRefHandleManager* NetRefHandleManager = nullptr;
 };
 
+// Helper class instance to avoid logging redundant depewndent objects more than once per object class.
+#if !UE_BUILD_SHIPPING
+static FIrisLogOnceTracker ReplicationFiltering_MootDependentObjectTracker;
+#endif
+
 //*************************************************************************************************
 // FReplicationFiltering
 //*************************************************************************************************
@@ -2542,6 +2548,16 @@ void FReplicationFiltering::NotifyObjectRemovedFromGroup(FNetObjectGroupHandle G
 
 void FReplicationFiltering::NotifyAddedDependentObject(FInternalNetRefIndex ObjectIndex)
 {
+#if !UE_BUILD_SHIPPING
+	if (UE_LOG_ACTIVE(LogIrisFiltering, Display))
+	{
+		if (!this->DynamicFilterEnabledObjects.GetBit(ObjectIndex))
+		{
+			const UObject* Object = NetRefHandleManager->GetReplicatedObjectInstance(ObjectIndex);
+			UE_CLOG(Object && ReplicationFiltering_MootDependentObjectTracker.ShouldLog(Object->GetClass()->GetFName()), LogIrisFiltering, Display, TEXT("FReplicationFiltering::NotifyAddedDependentObject: Object doesn't have a dynamic filter set so having a dependency on it won't change when it's replicated. %s"), ToCStr(NetRefHandleManager->PrintObjectFromIndex(ObjectIndex)));
+		}
+	}
+#endif
 	ObjectsRequiringDynamicFilterUpdate.SetBit(ObjectIndex);
 }
 
