Small Capsule Capsule optimization
Follow up the last 15% improvemnent, fixed some inconsistences left

diff --git a/Engine/Source/Runtime/Experimental/Chaos/Private/Chaos/CollisionOneShotManifoldsMiscShapes.cpp b/Engine/Source/Runtime/Experimental/Chaos/Private/Chaos/CollisionOneShotManifoldsMiscShapes.cpp
--- a/Engine/Source/Runtime/Experimental/Chaos/Private/Chaos/CollisionOneShotManifoldsMiscShapes.cpp
+++ b/Engine/Source/Runtime/Experimental/Chaos/Private/Chaos/CollisionOneShotManifoldsMiscShapes.cpp
@@ -325,7 +325,6 @@ namespace Chaos
 			}
 		}
 
-		// @todo(chaos): this will be faster if we transform into the space of one of the capsules
 		void ConstructCapsuleCapsuleOneShotManifold(const FCapsule& CapsuleA, const FRigidTransform3& CapsuleATransform, const FCapsule& CapsuleB, const FRigidTransform3& CapsuleBTransform, const FReal Dt, FPBDCollisionConstraint& Constraint)
 		{
 			// We only build one shot manifolds once
@@ -396,7 +395,7 @@ namespace Chaos
 
 			// We always add the closest point to the manifold
 			// We may also add 2 more points generated from the end cap positions of the smaller capsule
-			FContactPoint ClosestContactPoint;
+			FContactPointf ClosestContactPoint;
 			ClosestContactPoint.ShapeContactPoints[0] = CapsuleAToCapsuleBf.InverseTransformPosition(ClosestLocationA);
 			ClosestContactPoint.ShapeContactPoints[1] = ClosestLocationB;
 			ClosestContactPoint.ShapeContactNormal = ClosestNormal;
@@ -459,9 +458,9 @@ namespace Chaos
 				const FRealSingle SecondRadius,
 				const FVec3f& Orthogonal,
 				const FRealSingle CullDistance,
-				const bool bSwap) -> FContactPoint
+				const bool bSwap) -> FContactPointf
 			{
-				FContactPoint ContactPoint;
+				FContactPointf ContactPoint;
 
 				const FVec3f FirstContactPos = FirstCenter + (FirstT * FirstHalfLen) * FirstAxis + Orthogonal * FirstRadius;
 				const FVec3f SecondSegmentPos = FMath::ClosestPointOnLine(SecondCenter - SecondHalfLen * SecondAxis, SecondCenter + SecondHalfLen * SecondAxis, FirstContactPos);
@@ -538,7 +537,7 @@ namespace Chaos
 
 					if (TMin < FirstT - TDeltaThreshold)
 					{
-						FContactPoint EndContact0 = MakeCapsuleSegmentContact(TMin, FirstCenter, FirstAxis, FirstHalfLen, FirstRadius, FirstTransformToSecond, SecondCenter, SecondAxis, SecondHalfLen, SecondRadius, Orthogonal, Constraint.GetCullDistancef(), bSwap);
+						FContactPointf EndContact0 = MakeCapsuleSegmentContact(TMin, FirstCenter, FirstAxis, FirstHalfLen, FirstRadius, FirstTransformToSecond, SecondCenter, SecondAxis, SecondHalfLen, SecondRadius, Orthogonal, Constraint.GetCullDistancef(), bSwap);
 						if (EndContact0.Phi < Constraint.GetCullDistance())
 						{
 							Constraint.AddOneshotManifoldContact(EndContact0);
@@ -546,7 +545,7 @@ namespace Chaos
 					}
 					if (TMax > FirstT + TDeltaThreshold)
 					{
-						FContactPoint EndContact1 = MakeCapsuleSegmentContact(TMax, FirstCenter, FirstAxis, FirstHalfLen, FirstRadius, FirstTransformToSecond, SecondCenter, SecondAxis, SecondHalfLen, SecondRadius, Orthogonal, Constraint.GetCullDistancef(), bSwap);
+						FContactPointf EndContact1 = MakeCapsuleSegmentContact(TMax, FirstCenter, FirstAxis, FirstHalfLen, FirstRadius, FirstTransformToSecond, SecondCenter, SecondAxis, SecondHalfLen, SecondRadius, Orthogonal, Constraint.GetCullDistancef(), bSwap);
 						if (EndContact1.Phi < Constraint.GetCullDistance())
 						{
 							Constraint.AddOneshotManifoldContact(EndContact1);
