Support handling position (and tangent, normal) transform in vertex shader of GeometryCollection rendering in mobile renderer.

This is now accessing the transform matrix via texel buffer but the vertices itself use fixed-function vertex input instead of MANUAL_VERTEX_FETCH

diff --git a/Engine/Shaders/Private/GeometryCollectionVertexFactory.ush b/Engine/Shaders/Private/GeometryCollectionVertexFactory.ush
--- a/Engine/Shaders/Private/GeometryCollectionVertexFactory.ush
+++ b/Engine/Shaders/Private/GeometryCollectionVertexFactory.ush
@@ -15,7 +15,11 @@
 	#define MANUAL_VERTEX_FETCH 0
 #endif
 
-#if MANUAL_VERTEX_FETCH
+#ifndef USE_SHADER_BONE_TRANSFORM
+	#define USE_SHADER_BONE_TRANSFORM (MANUAL_VERTEX_FETCH)
+#endif
+
+#if MANUAL_VERTEX_FETCH || USE_SHADER_BONE_TRANSFORM
 	#define VF_ColorIndexMask_Index 0
 	#define VF_NumTexcoords_Index 1
 	#define FV_LightMapIndex_Index 2
@@ -32,7 +36,7 @@
 #endif
 
 // Since we don't have template functions it will have to be this way.
-#if MANUAL_VERTEX_FETCH
+#if MANUAL_VERTEX_FETCH || USE_SHADER_BONE_TRANSFORM
 	#define GCVF_GET_INPUT_VERTEX_ID(InputStruct) ((InputStruct).VertexId)
 #else // !MANUAL_VERTEX_FETCH
 	#define GCVF_GET_INPUT_VERTEX_ID(InputStruct) (0U)
@@ -100,7 +104,7 @@ struct FVertexFactoryInput
 	float2	LightMapCoordinate : ATTRIBUTE15;
 #endif
 
-#if MANUAL_VERTEX_FETCH
+#if MANUAL_VERTEX_FETCH || USE_SHADER_BONE_TRANSFORM
 	uint VertexId : SV_VertexID;
 #endif
 };
@@ -119,7 +123,7 @@ struct FPositionOnlyVertexFactoryInput
 	VF_GPUSCENE_DECLARE_INPUT_BLOCK(1)
 	VF_INSTANCED_STEREO_DECLARE_INPUT_BLOCK()
 
-#if MANUAL_VERTEX_FETCH
+#if MANUAL_VERTEX_FETCH || USE_SHADER_BONE_TRANSFORM
 	uint VertexId : SV_VertexID;
 #endif
 };
@@ -136,7 +140,7 @@ struct FPositionAndNormalOnlyVertexFactoryInput
 	VF_GPUSCENE_DECLARE_INPUT_BLOCK(1)
 	VF_INSTANCED_STEREO_DECLARE_INPUT_BLOCK()
 
-#if MANUAL_VERTEX_FETCH
+#if MANUAL_VERTEX_FETCH || USE_SHADER_BONE_TRANSFORM
 	uint VertexId : SV_VertexID;
 #endif
 };
@@ -363,7 +367,7 @@ FTransformData LoadTransformData(FSceneDataIntermediates SceneData, uint VertexI
 	FDFMatrix InstanceLocalToWorld = TransformData.SceneData.InstanceData.LocalToWorld;
 	FDFMatrix PrevInstanceLocalToWorld = TransformData.SceneData.InstanceData.PrevLocalToWorld;
 
-#if MANUAL_VERTEX_FETCH
+#if MANUAL_VERTEX_FETCH || USE_SHADER_BONE_TRANSFORM
 	uint BoneIndex = GCBoneLooseParameters.VertexFetch_BoneMapBuffer[GeometryCollectionVF.VertexFetch_Parameters[VF_VertexOffset] + VertexId];
 	float4x4 LocalToInstance = float4x4(
 		float4(GCBoneLooseParameters.VertexFetch_BoneTransformBuffer[4 * BoneIndex + 0].xyz, 0.0f),
diff --git a/Engine/Source/Runtime/Engine/Private/GeometryCollectionRendering.cpp b/Engine/Source/Runtime/Engine/Private/GeometryCollectionRendering.cpp
--- a/Engine/Source/Runtime/Engine/Private/GeometryCollectionRendering.cpp
+++ b/Engine/Source/Runtime/Engine/Private/GeometryCollectionRendering.cpp
@@ -10,6 +10,7 @@
 #include "ProfilingDebugging/LoadTimeTracker.h"
 #include "PrimitiveUniformShaderParameters.h"
 #include "RenderUtils.h"
+#include "SceneInterface.h"
 
 IMPLEMENT_TYPE_LAYOUT(FGeometryCollectionVertexFactoryShaderParameters);
 
@@ -31,6 +32,11 @@ IMPLEMENT_VERTEX_FACTORY_TYPE(FGeometryCollectionVertexFactory, "/Engine/Private
 	| EVertexFactoryFlags::SupportsPSOPrecaching
 );
 
+bool FGeometryCollectionVertexFactory::UseShaderBoneTransform(EShaderPlatform ShaderPlatform)
+{
+	return FDataDrivenShaderPlatformInfo::GetSupportsTypedBufferSRVs(ShaderPlatform);
+}
+
 bool FGeometryCollectionVertexFactory::ShouldCompilePermutation(const FVertexFactoryShaderPermutationParameters& Parameters)
 {
 	if (!Parameters.MaterialParameters.bIsUsedWithGeometryCollections && !Parameters.MaterialParameters.bIsSpecialEngineMaterial)
@@ -64,6 +70,10 @@ void FGeometryCollectionVertexFactory::ModifyCompilationEnvironment(const FVerte
 	{
 		OutEnvironment.SetDefineIfUnset(TEXT("MANUAL_VERTEX_FETCH"), TEXT("1"));
 	}
+	else if (UseShaderBoneTransform(Parameters.Platform))
+	{
+		OutEnvironment.SetDefineIfUnset(TEXT("USE_SHADER_BONE_TRANSFORM"), TEXT("1"));
+	}
 
 	// Geometry collections use a custom hit proxy per bone
 	OutEnvironment.SetDefine(TEXT("USE_PER_VERTEX_HITPROXY_ID"), 1);
@@ -301,18 +311,18 @@ void FGeometryCollectionVertexFactoryShaderParameters::GetElementShaderBindings(
 	check(VertexFactory->GetType() == &FGeometryCollectionVertexFactory::StaticType);
 	const auto* TypedVertexFactory = static_cast<const FGeometryCollectionVertexFactory*>(VertexFactory);
 
-	const bool bSupportsManualFetch = TypedVertexFactory->SupportsManualVertexFetch(FeatureLevel);
+	const bool bUseShaderBoneTransform = TypedVertexFactory->SupportsManualVertexFetch(FeatureLevel) || TypedVertexFactory->UseShaderBoneTransform(Scene->GetShaderPlatform());
 
 	FRHIUniformBuffer* VertexFactoryUniformBuffer = TypedVertexFactory->GetUniformBuffer();
 
-	if (bSupportsManualFetch || UseGPUScene(GMaxRHIShaderPlatform, FeatureLevel))
+	if (bUseShaderBoneTransform || UseGPUScene(GMaxRHIShaderPlatform, FeatureLevel))
 	{
 		check(VertexFactoryUniformBuffer != nullptr);
 		ShaderBindings.Add(Shader->GetUniformBufferParameter<FGeometryCollectionVertexFactoryUniformShaderParameters>(), VertexFactoryUniformBuffer);
 	}
 
 	// We only want to set the SRV parameters if we support manual vertex fetch.
-	if (bSupportsManualFetch)
+	if (bUseShaderBoneTransform)
 	{
 		FUniformBufferRHIRef LooseParameterBuffer = TypedVertexFactory->GetLooseParameterBuffer();
 		check(LooseParameterBuffer != nullptr);
diff --git a/Engine/Source/Runtime/Engine/Public/GeometryCollectionRendering.h b/Engine/Source/Runtime/Engine/Public/GeometryCollectionRendering.h
--- a/Engine/Source/Runtime/Engine/Public/GeometryCollectionRendering.h
+++ b/Engine/Source/Runtime/Engine/Public/GeometryCollectionRendering.h
@@ -195,6 +195,8 @@ public:
 	FUniformBufferRHIRef LooseParameterUniformBuffer;
 	bool EnableLooseParameter;
 
+	static ENGINE_API bool UseShaderBoneTransform(EShaderPlatform ShaderPlatform);
+
 private:
 	FDataType Data;
 	TUniformBufferRef<FGeometryCollectionVertexFactoryUniformShaderParameters> UniformBuffer;
diff --git a/Engine/Source/Runtime/Experimental/GeometryCollectionEngine/Private/GeometryCollection/GeometryCollectionSceneProxy.cpp b/Engine/Source/Runtime/Experimental/GeometryCollectionEngine/Private/GeometryCollection/GeometryCollectionSceneProxy.cpp
--- a/Engine/Source/Runtime/Experimental/GeometryCollectionEngine/Private/GeometryCollection/GeometryCollectionSceneProxy.cpp
+++ b/Engine/Source/Runtime/Experimental/GeometryCollectionEngine/Private/GeometryCollection/GeometryCollectionSceneProxy.cpp
@@ -54,7 +54,6 @@ static_assert(sizeof(ispc::FMatrix44f) == sizeof(FMatrix44f), "sizeof(ispc::FMat
 static_assert(sizeof(ispc::FVector3f) == sizeof(FVector3f), "sizeof(ispc::FVector3f) != sizeof(FVector3f)");
 #endif
 
-
 static int32 GParallelGeometryCollectionBatchSize = 1024;
 static TAutoConsoleVariable<int32> CVarParallelGeometryCollectionBatchSize(
 	TEXT("r.ParallelGeometryCollectionBatchSize"),
@@ -144,7 +143,7 @@ FGeometryCollectionSceneProxyBase::FGeometryCollectionSceneProxyBase(UGeometryCo
 	, MeshResource(Component->GetRestCollection()->RenderData->MeshResource)
 	, MaterialRelevance(Component->GetMaterialRelevance(FeatureLevel))
 	, VertexFactory(FeatureLevel)
-	, bSupportsManualVertexFetch(VertexFactory.SupportsManualVertexFetch(FeatureLevel))
+	, bUseShaderBoneTransform(VertexFactory.UseShaderBoneTransform(Component->GetScene()->GetShaderPlatform()))
 	, bSupportsTripleBufferVertexUpload(GRHISupportsMapWriteNoOverwrite)
 #if RHI_RAYTRACING
 	, bHasRayTracingRepresentation(IsRayTracingEnabled() && Component->GetRestCollection()->RenderData->bHasMeshData && Component->GetRestCollection()->RenderData->MeshDescription.NumVertices != 0)
@@ -187,7 +186,7 @@ void FGeometryCollectionSceneProxyBase::CreateRenderThreadResources(FRHICommandL
 {
 	if (!bIsNanite || bHasRayTracingRepresentation)
 	{
-		if (bSupportsManualVertexFetch)
+		if (bUseShaderBoneTransform)
 		{
 			// Initialize transform buffers and upload rest transforms.
 			TransformBuffers.AddDefaulted(1);
@@ -238,7 +237,7 @@ void FGeometryCollectionSceneProxyBase::DestroyRenderThreadResources()
 
 	if (!bIsNanite || bHasRayTracingRepresentation)
 	{
-		if (bSupportsManualVertexFetch)
+		if (bUseShaderBoneTransform)
 		{
 			for (int32 i = 0; i < TransformBuffers.Num(); i++)
 			{
@@ -265,11 +264,11 @@ void FGeometryCollectionSceneProxyBase::DestroyRenderThreadResources()
 
 void FGeometryCollectionSceneProxyBase::SetupVertexFactory(FRHICommandListBase& RHICmdList, FGeometryCollectionVertexFactory& GeometryCollectionVertexFactory, FColorVertexBuffer* ColorOverride) const
 {
-	checkf(GeometryCollectionVertexFactory.SupportsManualVertexFetch(FeatureLevel) == bSupportsManualVertexFetch, TEXT("Setting up vertex factory for manual vertex fetch but provided type doesn't support it."));
+	checkf(GeometryCollectionVertexFactory.SupportsManualVertexFetch(FeatureLevel) == VertexFactory.SupportsManualVertexFetch(FeatureLevel), TEXT("Setting up vertex factory for manual vertex fetch but provided type doesn't support it."));
 
 	FGeometryCollectionVertexFactory::FDataType Data;
 
-	FPositionVertexBuffer const& PositionVB = bSupportsManualVertexFetch ? MeshResource.PositionVertexBuffer : SkinnedPositionVertexBuffer;
+	FPositionVertexBuffer const& PositionVB = bUseShaderBoneTransform ? MeshResource.PositionVertexBuffer : SkinnedPositionVertexBuffer;
 	PositionVB.BindPositionVertexBuffer(&GeometryCollectionVertexFactory, Data);
 
 	MeshResource.StaticMeshVertexBuffer.BindTangentVertexBuffer(&GeometryCollectionVertexFactory, Data);
@@ -279,7 +278,7 @@ void FGeometryCollectionSceneProxyBase::SetupVertexFactory(FRHICommandListBase&
 	FColorVertexBuffer const& ColorVB = ColorOverride ? *ColorOverride : MeshResource.ColorVertexBuffer;
 	ColorVB.BindColorVertexBuffer(&GeometryCollectionVertexFactory, Data);
 
-	if (bSupportsManualVertexFetch)
+	if (bUseShaderBoneTransform)
 	{
 		Data.BoneMapSRV = MeshResource.BoneMapVertexBuffer.GetSRV();
 		Data.BoneTransformSRV = GetCurrentTransformBuffer().VertexBufferSRV;
@@ -325,7 +324,7 @@ void FGeometryCollectionSceneProxyBase::SetDynamicData_RenderThread(FRHICommandL
 
 	if (!bIsNanite || bHasRayTracingRepresentation)
 	{
-		if (bSupportsManualVertexFetch)
+		if (bUseShaderBoneTransform)
 		{
 			const bool bLocalGeometryCollectionTripleBufferUploads = (GGeometryCollectionTripleBufferUploads != 0) && bSupportsTripleBufferVertexUpload;
 
@@ -756,7 +755,7 @@ void FGeometryCollectionSceneProxy::SetDynamicData_RenderThread(FRHICommandListB
 		return;
 	}
 		
-	if (bSupportsManualVertexFetch)
+	if (bUseShaderBoneTransform)
 	{
 		const FGeometryCollectionTransformBuffer& TransformBuffer = GetCurrentTransformBuffer();
 		const FGeometryCollectionTransformBuffer& PrevTransformBuffer = GetCurrentPrevTransformBuffer();
diff --git a/Engine/Source/Runtime/Experimental/GeometryCollectionEngine/Private/GeometryCollection/GeometryCollectionSceneProxy.h b/Engine/Source/Runtime/Experimental/GeometryCollectionEngine/Private/GeometryCollection/GeometryCollectionSceneProxy.h
--- a/Engine/Source/Runtime/Experimental/GeometryCollectionEngine/Private/GeometryCollection/GeometryCollectionSceneProxy.h
+++ b/Engine/Source/Runtime/Experimental/GeometryCollectionEngine/Private/GeometryCollection/GeometryCollectionSceneProxy.h
@@ -186,7 +186,7 @@ protected:
 	FMaterialRelevance MaterialRelevance;
 
 	FGeometryCollectionVertexFactory VertexFactory;
-	const bool bSupportsManualVertexFetch;
+	const bool bUseShaderBoneTransform;
 	const bool bSupportsTripleBufferVertexUpload;
 
 	TArray<UMaterialInterface*> Materials;
diff --git a/Engine/Source/Runtime/RHI/Private/DataDrivenShaderPlatformInfo.cpp b/Engine/Source/Runtime/RHI/Private/DataDrivenShaderPlatformInfo.cpp
--- a/Engine/Source/Runtime/RHI/Private/DataDrivenShaderPlatformInfo.cpp
+++ b/Engine/Source/Runtime/RHI/Private/DataDrivenShaderPlatformInfo.cpp
@@ -180,6 +180,7 @@ void FGenericDataDrivenShaderPlatformInfo::SetDefaultValues()
 	bSupportsAsyncPipelineCompilation = true;
 	bSupportsVertexShaderSRVs = true; // Explicitly overriden to false for ES 3.1 platforms via DDPI ini
 	bSupportsVertexShaderUAVs = uint32(ERHIFeatureSupport::Unsupported);
+	bSupportsTypedBufferSRVs = true;
 	bSupportsManualVertexFetch = true;
 	bSupportsVolumeTextureAtomics = true;
 	bSupportsClipDistance = true;
@@ -320,6 +321,7 @@ void FGenericDataDrivenShaderPlatformInfo::ParseDataDrivenShaderInfo(const FConf
 	GET_SECTION_BOOL_HELPER(bSupportsAsyncPipelineCompilation);
 	GET_SECTION_BOOL_HELPER(bSupportsVertexShaderSRVs);
 	GET_SECTION_SUPPORT_HELPER(bSupportsVertexShaderUAVs);
+	GET_SECTION_BOOL_HELPER(bSupportsTypedBufferSRVs);
 	GET_SECTION_BOOL_HELPER(bSupportsManualVertexFetch);
 	GET_SECTION_BOOL_HELPER(bRequiresReverseCullingOnMobile);
 	GET_SECTION_BOOL_HELPER(bOverrideFMaterial_NeedsGBufferEnabled);
@@ -564,6 +566,7 @@ void FGenericDataDrivenShaderPlatformInfo::UpdatePreviewPlatforms()
 				PREVIEW_USE_RUNTIME_VALUE(bSupportsSceneDataCompressedTransforms);
 				PREVIEW_USE_RUNTIME_VALUE(bSupportsVertexShaderSRVs);
 				PREVIEW_USE_RUNTIME_VALUE(bSupportsVertexShaderUAVs);
+				PREVIEW_USE_RUNTIME_VALUE(bSupportsTypedBufferSRVs);
 				PREVIEW_USE_RUNTIME_VALUE(bSupportsManualVertexFetch);
 				PREVIEW_USE_RUNTIME_VALUE(bSupportsRealTypes);
 				PREVIEW_USE_RUNTIME_VALUE(bSupportsUniformBufferObjects);
diff --git a/Engine/Source/Runtime/RHI/Public/DataDrivenShaderPlatformInfo.h b/Engine/Source/Runtime/RHI/Public/DataDrivenShaderPlatformInfo.h
--- a/Engine/Source/Runtime/RHI/Public/DataDrivenShaderPlatformInfo.h
+++ b/Engine/Source/Runtime/RHI/Public/DataDrivenShaderPlatformInfo.h
@@ -100,6 +100,7 @@ class FGenericDataDrivenShaderPlatformInfo
 	uint32 bSupportsAsyncPipelineCompilation : 1;
 	uint32 bSupportsVertexShaderSRVs : 1; // Whether SRVs can be bound to vertex shaders (may be independent from ManualVertexFetch)
 	uint32 bSupportsVertexShaderUAVs : int32(ERHIFeatureSupport::NumBits); // Whether UAVs can be bound to vertex shaders. Requires run-time check of GRHIGlobals.SupportsVertexShaderUAVs.
+	uint32 bSupportsTypedBufferSRVs : 1; // Buffer<>, texelbuffer/texture buffer, SRV with Format
 	uint32 bSupportsManualVertexFetch : 1;
 	uint32 bRequiresReverseCullingOnMobile : 1;
 	uint32 bOverrideFMaterial_NeedsGBufferEnabled : 1;
@@ -697,6 +698,12 @@ public:
 		return ERHIFeatureSupport(Infos[Platform].bSupportsVertexShaderUAVs);
 	}
 
+	static FORCEINLINE_DEBUGGABLE const bool GetSupportsTypedBufferSRVs(const FStaticShaderPlatform Platform)
+	{
+		check(IsValid(Platform));
+		return Infos[Platform].bSupportsTypedBufferSRVs;
+	}
+
 	static FORCEINLINE_DEBUGGABLE const bool GetSupportsManualVertexFetch(const FStaticShaderPlatform Platform)
 	{
 		check(IsValid(Platform));
