Iris
* Move responsability of finding the location and net cull distance of root objects from the EngineBridge to the NetObjectFactories
* Now setting the WorldLocation cache when objects start replicating and register into the cache.

#rnx

diff --git a/Engine/Config/BaseEngine.ini b/Engine/Config/BaseEngine.ini
--- a/Engine/Config/BaseEngine.ini
+++ b/Engine/Config/BaseEngine.ini
@@ -1412,7 +1412,7 @@ RequiredNetDriverChannelClassName=/Script/Engine.DataStreamChannel
 ; Filters
 DefaultSpatialFilterName=Spatial
 ; Filter configs should typically be configured per game
-; Clear all filters
+; Set the class-based dynamic filters configs here. Using None means the objects of this class are always relevant.
 !FilterConfigs=ClearArray
 ;+FilterConfigs=(ClassName=/Script/Engine.LevelScriptActor, DynamicFilterName=NotRouted) ; NotRouted means that this type will not be replicated
 +FilterConfigs=(ClassName=/Script/Engine.Actor, DynamicFilterName=None))
diff --git a/Engine/Plugins/Runtime/ReplicationSystemTestPlugin/Source/Private/Tests/ReplicationSystem/Filtering/TestGridFilter.cpp b/Engine/Plugins/Runtime/ReplicationSystemTestPlugin/Source/Private/Tests/ReplicationSystem/Filtering/TestGridFilter.cpp
--- a/Engine/Plugins/Runtime/ReplicationSystemTestPlugin/Source/Private/Tests/ReplicationSystem/Filtering/TestGridFilter.cpp
+++ b/Engine/Plugins/Runtime/ReplicationSystemTestPlugin/Source/Private/Tests/ReplicationSystem/Filtering/TestGridFilter.cpp
@@ -168,9 +168,11 @@ UE_NET_TEST_FIXTURE(FTestGridFilterFixture, TestWorldLocationIsFrequentlyUpdated
 
 	Server->GetReplicationBridge()->SetExternalWorldLocationUpdateFunctor([&](FNetRefHandle NetHandle, const UObject* ReplicatedObject, FVector& OutLocation, float& OutCullDistance)
 		{
-			const FObjectLoc& ObjectLoc = ObjectLocs[ReplicatedObject];
-			OutLocation = ObjectLoc.Loc;
-			OutCullDistance = ObjectLoc.CullDistance;
+			if (const FObjectLoc* ObjectLoc = ObjectLocs.Find(ReplicatedObject))
+			{
+				OutLocation = ObjectLoc->Loc;
+				OutCullDistance = ObjectLoc->CullDistance;
+			}
 		});
 
 	// Add client
diff --git a/Engine/Plugins/Runtime/ReplicationSystemTestPlugin/Source/Private/Tests/ReplicationSystem/Prioritization/TestNetObjectPrioritizerFixture.cpp b/Engine/Plugins/Runtime/ReplicationSystemTestPlugin/Source/Private/Tests/ReplicationSystem/Prioritization/TestNetObjectPrioritizerFixture.cpp
--- a/Engine/Plugins/Runtime/ReplicationSystemTestPlugin/Source/Private/Tests/ReplicationSystem/Prioritization/TestNetObjectPrioritizerFixture.cpp
+++ b/Engine/Plugins/Runtime/ReplicationSystemTestPlugin/Source/Private/Tests/ReplicationSystem/Prioritization/TestNetObjectPrioritizerFixture.cpp
@@ -63,13 +63,11 @@ FReplicationView FTestNetObjectPrioritizerFixture::MakeReplicationView(const FVe
 
 FTestNetObjectPrioritizerFixture::FPrioritizationResult FTestNetObjectPrioritizerFixture::PrioritizeWorldLocations(const FReplicationView& View, FNetObjectPrioritizerHandle PrioritizerHandle, TConstArrayView<FVector> WorldLocations)
 {
-	// Need to set the world location update function to be able to set world locations.
-	check(Server->GetReplicationBridge()->GetExternalWorldLocationUpdateFunctor() == nullptr);
-
 	FReplicationSystemTestClient* Client = CreateClient();
 
 	Server->GetReplicationSystem()->SetReplicationView(Client->ConnectionIdOnServer, View);
 
+	// Need to set the world location update function to be able to set world locations.
 	Server->GetReplicationBridge()->SetExternalWorldLocationUpdateFunctor([&](FNetRefHandle NetHandle, const UObject* Object, FVector& OutLocation, float& OutCullDistance)
 		{
 			if (const UTestPrioritizationWithWorldLocationObject* WorldLocObject = Cast<UTestPrioritizationWithWorldLocationObject>(Object))
@@ -190,13 +188,11 @@ FNetObjectPrioritizerImage FTestNetObjectPrioritizerFixture::Visualize(FNetObjec
 
 	const double LocationDelta = Params.UnitsPerPixel;
 
-	// Need to set the world location update function to be able to set world locations.
-	check(Server->GetReplicationBridge()->GetExternalWorldLocationUpdateFunctor() == nullptr);
-
 	FReplicationSystemTestClient* Client = CreateClient();
 
 	Server->GetReplicationSystem()->SetReplicationView(Client->ConnectionIdOnServer, Params.View);
 
+	// Need to set the world location update function to be able to set world locations.
 	Server->GetReplicationBridge()->SetExternalWorldLocationUpdateFunctor([&](FNetRefHandle NetHandle, const UObject* Object, FVector& OutLocation, float& OutCullDistance)
 		{
 			if (const UTestPrioritizationWithWorldLocationObject* WorldLocObject = Cast<UTestPrioritizationWithWorldLocationObject>(Object))
diff --git a/Engine/Plugins/Runtime/ReplicationSystemTestPlugin/Source/Private/Tests/ReplicationSystem/ReplicatedTestObject.cpp b/Engine/Plugins/Runtime/ReplicationSystemTestPlugin/Source/Private/Tests/ReplicationSystem/ReplicatedTestObject.cpp
--- a/Engine/Plugins/Runtime/ReplicationSystemTestPlugin/Source/Private/Tests/ReplicationSystem/ReplicatedTestObject.cpp
+++ b/Engine/Plugins/Runtime/ReplicationSystemTestPlugin/Source/Private/Tests/ReplicationSystem/ReplicatedTestObject.cpp
@@ -116,12 +116,8 @@ bool UReplicatedTestObjectBridge::IsAllowedToDestroyInstance(const UObject* Inst
 
 void UReplicatedTestObjectBridge::SetExternalWorldLocationUpdateFunctor(TFunction<void(FNetRefHandle NetHandle, const UObject* ReplicatedObject, FVector& OutLocation, float& OutCullDistance)> LocUpdateFunctor)
 {
-	SetInstanceGetWorldObjectInfoFunction(LocUpdateFunctor);
-}
-
-TFunction<void(UE::Net::FNetRefHandle, const UObject*, FVector&, float&)> UReplicatedTestObjectBridge::GetExternalWorldLocationUpdateFunctor() const
-{
-	return GetInstanceGetWorldObjectInfoFunction();
+	UReplicatedTestObjectFactory* TestFactory = CastChecked<UReplicatedTestObjectFactory>(GetNetFactory(ReplicatedObjectFactoryId));
+	TestFactory->SetWorldUpdateFunctor(LocUpdateFunctor);
 }
 
 void UReplicatedTestObjectBridge::SetExternalPreUpdateFunctor(TFunction<void(TArrayView<UObject*>, const UReplicationBridge*)> PreUpdateFunctor)
diff --git a/Engine/Plugins/Runtime/ReplicationSystemTestPlugin/Source/Private/Tests/ReplicationSystem/ReplicatedTestObject.h b/Engine/Plugins/Runtime/ReplicationSystemTestPlugin/Source/Private/Tests/ReplicationSystem/ReplicatedTestObject.h
--- a/Engine/Plugins/Runtime/ReplicationSystemTestPlugin/Source/Private/Tests/ReplicationSystem/ReplicatedTestObject.h
+++ b/Engine/Plugins/Runtime/ReplicationSystemTestPlugin/Source/Private/Tests/ReplicationSystem/ReplicatedTestObject.h
@@ -552,7 +552,6 @@ public:
 	const UE::Net::FReplicationInstanceProtocol* GetReplicationInstanceProtocol(FNetRefHandle Handle) const;
 
 	void SetExternalWorldLocationUpdateFunctor(TFunction<void(FNetRefHandle NetHandle, const UObject* ReplicatedObject, FVector& OutLocation, float& OutCullDistance)> LocUpdateFunctor);
-	TFunction<void(FNetRefHandle, const UObject*, FVector&, float&)> GetExternalWorldLocationUpdateFunctor() const;
 
 	void SetExternalPreUpdateFunctor(TFunction<void(TArrayView<UObject*>, const UReplicationBridge*)> PreUpdateFunctor);
 
diff --git a/Engine/Plugins/Runtime/ReplicationSystemTestPlugin/Source/Private/Tests/ReplicationSystem/ReplicatedTestObjectFactory.cpp b/Engine/Plugins/Runtime/ReplicationSystemTestPlugin/Source/Private/Tests/ReplicationSystem/ReplicatedTestObjectFactory.cpp
--- a/Engine/Plugins/Runtime/ReplicationSystemTestPlugin/Source/Private/Tests/ReplicationSystem/ReplicatedTestObjectFactory.cpp
+++ b/Engine/Plugins/Runtime/ReplicationSystemTestPlugin/Source/Private/Tests/ReplicationSystem/ReplicatedTestObjectFactory.cpp
@@ -207,4 +207,12 @@ void UReplicatedTestObjectFactory::DestroyReplicatedObject(const FDestroyedConte
 			Context.DestroyedInstance->MarkAsGarbage();
 		}
 	}
+}
+
+void UReplicatedTestObjectFactory::GetWorldInfo(const FWorldInfoContext& Context, FWorldInfoData& OutData)
+{
+	if (GetInstanceWorldObjectInfoFunction)
+	{
+		GetInstanceWorldObjectInfoFunction(Context.Handle, Context.Instance, OutData.WorldLocation, OutData.CullDistance);
+	}
 }
\ No newline at end of file
diff --git a/Engine/Plugins/Runtime/ReplicationSystemTestPlugin/Source/Private/Tests/ReplicationSystem/ReplicatedTestObjectFactory.h b/Engine/Plugins/Runtime/ReplicationSystemTestPlugin/Source/Private/Tests/ReplicationSystem/ReplicatedTestObjectFactory.h
--- a/Engine/Plugins/Runtime/ReplicationSystemTestPlugin/Source/Private/Tests/ReplicationSystem/ReplicatedTestObjectFactory.h
+++ b/Engine/Plugins/Runtime/ReplicationSystemTestPlugin/Source/Private/Tests/ReplicationSystem/ReplicatedTestObjectFactory.h
@@ -39,14 +39,25 @@ public:
 
 	virtual FInstantiateResult InstantiateReplicatedObjectFromHeader(const FInstantiateContext& Context, const UE::Net::FNetObjectCreationHeader* InHeader) override;
 
+	// UNetObjectFactory interface
 	virtual void PostInit(const FPostInitContext& Context) override;
 	virtual void DestroyReplicatedObject(const FDestroyedContext& Context) override;
+	virtual void GetWorldInfo(const FWorldInfoContext& Context, FWorldInfoData& OutData) override;
+
+	//------------------------------------------------------------------------
 
 	void SetCreatedObjectsOnNode(TArray<TStrongObjectPtr<UObject>>* InCreatedObjectsOnNode)
 	{
 		CreatedObjectsOnNode = InCreatedObjectsOnNode;
 	}
 
+	//------------------------------------------------------------------------
+
+	void SetWorldUpdateFunctor(TFunction<void(UE::Net::FNetRefHandle, const UObject*, FVector&, float&)> LocUpdateFunctor)
+	{
+		GetInstanceWorldObjectInfoFunction = LocUpdateFunctor;
+	}
+
 public:
 
 	TArray<TStrongObjectPtr<UObject>>* CreatedObjectsOnNode = nullptr;
@@ -59,6 +70,9 @@ protected:
 
 	virtual TUniquePtr<UE::Net::FNetObjectCreationHeader> CreateAndDeserializeHeader(const UE::Net::FCreationHeaderContext& Context) override;
 
+private:
+
+	TFunction<void(UE::Net::FNetRefHandle, const UObject*, FVector&, float&)> GetInstanceWorldObjectInfoFunction;
 };
 
 	
diff --git a/Engine/Source/Runtime/Engine/Private/Net/Iris/ReplicationSystem/EngineReplicationBridge.cpp b/Engine/Source/Runtime/Engine/Private/Net/Iris/ReplicationSystem/EngineReplicationBridge.cpp
--- a/Engine/Source/Runtime/Engine/Private/Net/Iris/ReplicationSystem/EngineReplicationBridge.cpp
+++ b/Engine/Source/Runtime/Engine/Private/Net/Iris/ReplicationSystem/EngineReplicationBridge.cpp
@@ -101,15 +101,6 @@ void ActorReplicationBridgePreUpdateFunction(TArrayView<UObject*> Instances, con
 	}
 }
 
-void ActorReplicationBridgeGetActorWorldObjectInfo(FNetRefHandle /*Handle*/, const UObject* Instance, FVector& OutWorldLocation, float& OutCullDistance)
-{
-	if (const AActor* Actor = Cast<AActor>(Instance))
-	{
-		OutWorldLocation = Actor->GetActorLocation();
-		OutCullDistance = Actor->GetNetCullDistanceSquared() > 0.0f ? FMath::Sqrt(Actor->GetNetCullDistanceSquared()) : 0.0f;
-	}
-}
-
 bool ShouldIncludeActorInLevelGroups(const AActor* Actor)
 {
 	// Never filter out PlayerControllers based on level as they are required for travel.
@@ -125,7 +116,6 @@ UEngineReplicationBridge::UEngineReplicationBridge()
 , SubObjectFactoryId(UE::Net::InvalidNetObjectFactoryId)
 {
 	SetInstancePreUpdateFunction(UE::Net::Private::ActorReplicationBridgePreUpdateFunction);
-	SetInstanceGetWorldObjectInfoFunction(UE::Net::Private::ActorReplicationBridgeGetActorWorldObjectInfo);
 }
 
 void UEngineReplicationBridge::Initialize(UReplicationSystem* InReplicationSystem)
@@ -287,7 +277,8 @@ UE::Net::FNetRefHandle UEngineReplicationBridge::StartReplicatingActor(AActor* A
 	}
 
 	// Create handles for the registered fragments
-	Super::FRootObjectReplicationParams RootObjectParams = {
+	Super::FRootObjectReplicationParams RootObjectParams = 
+	{
 		.bNeedsPreUpdate = true,
 		.bNeedsWorldLocationUpdate = true,
 		.bIsDormant = Actor->NetDormancy > DORM_Awake,
diff --git a/Engine/Source/Runtime/Engine/Private/Net/Iris/ReplicationSystem/NetActorFactory.cpp b/Engine/Source/Runtime/Engine/Private/Net/Iris/ReplicationSystem/NetActorFactory.cpp
--- a/Engine/Source/Runtime/Engine/Private/Net/Iris/ReplicationSystem/NetActorFactory.cpp
+++ b/Engine/Source/Runtime/Engine/Private/Net/Iris/ReplicationSystem/NetActorFactory.cpp
@@ -516,6 +516,27 @@ void UNetActorFactory::DestroyReplicatedObject(const FDestroyedContext& Context)
 	}	
 }
 
+void UNetActorFactory::GetWorldInfo(const FWorldInfoContext& Context, FWorldInfoData& OutData)
+{
+	using namespace UE::Net;
+
+	const AActor* Actor = Cast<AActor>(Context.Instance);
+	if (!Actor)
+	{
+		ensureMsgf(Actor, TEXT("UNetActorFactory::GetWorldInfo received invalid replicated instance: %s | NetRefHandle: %s"), *GetNameSafe(Context.Instance), *Context.Handle.ToString());
+		return;
+	}
+	
+	if (EnumHasAnyFlags(Context.InfoRequested, EWorldInfoRequested::Location))
+	{
+		OutData.WorldLocation = Actor->GetActorLocation();
+	}
+	if (EnumHasAnyFlags(Context.InfoRequested, EWorldInfoRequested::CullDistance))
+	{
+		OutData.CullDistance = Actor->GetNetCullDistanceSquared() > 0.0f ? FMath::Sqrt(Actor->GetNetCullDistanceSquared()) : 0.0f;
+	}
+}
+
  namespace UE::Net
  {
 
diff --git a/Engine/Source/Runtime/Engine/Private/Net/Iris/ReplicationSystem/NetSubObjectFactory.cpp b/Engine/Source/Runtime/Engine/Private/Net/Iris/ReplicationSystem/NetSubObjectFactory.cpp
--- a/Engine/Source/Runtime/Engine/Private/Net/Iris/ReplicationSystem/NetSubObjectFactory.cpp
+++ b/Engine/Source/Runtime/Engine/Private/Net/Iris/ReplicationSystem/NetSubObjectFactory.cpp
@@ -233,6 +233,11 @@ void UNetSubObjectFactory::DestroyReplicatedObject(const FDestroyedContext& Cont
 	Context.DestroyedInstance->MarkAsGarbage();
 }
 
+void UNetSubObjectFactory::GetWorldInfo(const FWorldInfoContext& Context, FWorldInfoData& OutData)
+{
+	ensureMsgf(false, TEXT("UNetSubObjectFactory::GetWorldInfo called but subobjects should never support this. Instance: %s, NetRefHandle: %s"), *GetNameSafe(Context.Instance), *Bridge->PrintObjectFromNetRefHandle(Context.Handle));
+}
+
 //------------------------------------------------------------------------
 
 namespace UE::Net
diff --git a/Engine/Source/Runtime/Engine/Public/Net/Iris/ReplicationSystem/NetActorFactory.h b/Engine/Source/Runtime/Engine/Public/Net/Iris/ReplicationSystem/NetActorFactory.h
--- a/Engine/Source/Runtime/Engine/Public/Net/Iris/ReplicationSystem/NetActorFactory.h
+++ b/Engine/Source/Runtime/Engine/Public/Net/Iris/ReplicationSystem/NetActorFactory.h
@@ -127,6 +127,7 @@ public:
 
 	virtual void DestroyReplicatedObject(const FDestroyedContext& Context) override;
 
+	virtual void GetWorldInfo(const FWorldInfoContext& Context, FWorldInfoData& OutData) override;
 
 protected:
 
diff --git a/Engine/Source/Runtime/Engine/Public/Net/Iris/ReplicationSystem/NetSubObjectFactory.h b/Engine/Source/Runtime/Engine/Public/Net/Iris/ReplicationSystem/NetSubObjectFactory.h
--- a/Engine/Source/Runtime/Engine/Public/Net/Iris/ReplicationSystem/NetSubObjectFactory.h
+++ b/Engine/Source/Runtime/Engine/Public/Net/Iris/ReplicationSystem/NetSubObjectFactory.h
@@ -33,6 +33,8 @@ public:
 
 	virtual void DestroyReplicatedObject(const FDestroyedContext& Context) override;
 
+	virtual void GetWorldInfo(const FWorldInfoContext& Context, FWorldInfoData& OutData) override;
+
 protected:
 
 	virtual TUniquePtr<UE::Net::FNetObjectCreationHeader> CreateAndFillHeader(UE::Net::FNetRefHandle Handle) override;
diff --git a/Engine/Source/Runtime/Experimental/Iris/Core/Private/Iris/ReplicationSystem/ObjectReplicationBridge.cpp b/Engine/Source/Runtime/Experimental/Iris/Core/Private/Iris/ReplicationSystem/ObjectReplicationBridge.cpp
--- a/Engine/Source/Runtime/Experimental/Iris/Core/Private/Iris/ReplicationSystem/ObjectReplicationBridge.cpp
+++ b/Engine/Source/Runtime/Experimental/Iris/Core/Private/Iris/ReplicationSystem/ObjectReplicationBridge.cpp
@@ -457,6 +457,8 @@ UE::Net::FNetRefHandle UObjectReplicationBridge::StartReplicatingRootObject(UObj
 	{
 		FWorldLocations& WorldLocations = ReplicationSystem->GetReplicationSystemInternal()->GetWorldLocations();
 		WorldLocations.InitObjectInfoCache(ObjectIndex);
+		ForceUpdateWorldLocation(RefHandle, ObjectIndex);
+
 	}
 
 	// Set prioritizer
@@ -1107,16 +1109,6 @@ UObjectReplicationBridge::FInstancePreUpdateFunction UObjectReplicationBridge::G
 	return PreUpdateInstanceFunction;
 }
 
-void UObjectReplicationBridge::SetInstanceGetWorldObjectInfoFunction(FInstanceGetWorldObjectInfoFunction InGetWorldObjectInfoFunction)
-{
-	GetInstanceWorldObjectInfoFunction = InGetWorldObjectInfoFunction;
-}
-
-UObjectReplicationBridge::FInstanceGetWorldObjectInfoFunction UObjectReplicationBridge::GetInstanceGetWorldObjectInfoFunction() const
-{
-	return GetInstanceWorldObjectInfoFunction;
-}
-
 void UObjectReplicationBridge::ForcePollObject(FNetRefHandle Handle)
 {
 	using namespace UE::Net::Private;
@@ -1433,16 +1425,49 @@ void UObjectReplicationBridge::ReconcileNewSubObjects(UE::Net::FNetBitArrayView
 	FNetBitArrayView::ForAllSetBits(GlobalScopeList, CurrentFrameScopeList, FNetBitArrayView::AndNotOp, HandleNewSubObject);
 }
 
-void UObjectReplicationBridge::UpdateInstancesWorldLocation()
+void UObjectReplicationBridge::ForceUpdateWorldLocation(FNetRefHandle NetRefHandle, UE::Net::Private::FInternalNetRefIndex InternalObjectIndex)
 {
 	using namespace UE::Net;
-	using namespace UE::Net::Private;
 
-	if (!GetInstanceWorldObjectInfoFunction)
+	FWorldLocations& WorldLocations = GetReplicationSystem()->GetReplicationSystemInternal()->GetWorldLocations();
+
+	if (!WorldLocations.HasInfoForObject(InternalObjectIndex))
 	{
+		ensureMsgf(false, TEXT("UObjectReplicationBridge::ForceUpdateWorldLocation was called on object that is not registered in the WorldLocations cache: %s"), *PrintObjectFromNetRefHandle(NetRefHandle));
 		return;
 	}
 
+	UNetObjectFactory* NetFactory = GetNetFactory(NetRefHandleManager->GetReplicatedObjectData(InternalObjectIndex).NetFactoryId);
+
+	const UNetObjectFactory::FWorldInfoContext Context
+	{
+		.Instance = NetRefHandleManager->GetReplicatedObjectInstance(InternalObjectIndex),
+		.Handle = NetRefHandle,
+		.InfoRequested = UNetObjectFactory::EWorldInfoRequested::All,
+	};
+
+	UNetObjectFactory::FWorldInfoData WorldInfo;
+	NetFactory->GetWorldInfo(Context, WorldInfo);
+
+	// Check the position lies within world boundaries.
+	ensureMsgf(	WorldInfo.WorldLocation.X >= WorldLocations.GetWorldMinPos().X && WorldInfo.WorldLocation.X <= WorldLocations.GetWorldMaxPos().X &&
+				WorldInfo.WorldLocation.Y >= WorldLocations.GetWorldMinPos().Y && WorldInfo.WorldLocation.Y <= WorldLocations.GetWorldMaxPos().Y,
+				TEXT("Object %s with position %s lies outside configured world boundary."), ToCStr(NetRefHandleManager->PrintObjectFromIndex(InternalObjectIndex)), *WorldInfo.WorldLocation.ToString());
+
+	const FWorldLocations::FObjectInfo WorldObjectInfo
+	{
+		.WorldLocation = WorldInfo.WorldLocation,
+		.CullDistance = WorldInfo.CullDistance
+	};
+
+	WorldLocations.SetObjectInfo(InternalObjectIndex, WorldObjectInfo);
+}
+
+void UObjectReplicationBridge::UpdateInstancesWorldLocation()
+{
+	using namespace UE::Net;
+	using namespace UE::Net::Private;
+
 	FReplicationSystemInternal* ReplicationSystemInternal = GetReplicationSystem()->GetReplicationSystemInternal();
 	const FNetRefHandleManager& LocalNetRefHandleManager = ReplicationSystemInternal->GetNetRefHandleManager();
 	FWorldLocations& WorldLocations = ReplicationSystemInternal->GetWorldLocations();
@@ -1454,15 +1479,30 @@ void UObjectReplicationBridge::UpdateInstancesWorldLocation()
 		const FNetRefHandleManager::FReplicatedObjectData& ObjectData = LocalNetRefHandleManager.GetReplicatedObjectDataNoCheck(InternalObjectIndex);
 		if (ObjectData.InstanceProtocol && EnumHasAnyFlags(ObjectData.InstanceProtocol->InstanceTraits, EReplicationInstanceProtocolTraits::NeedsWorldLocationUpdate))
 		{
-			FWorldLocations::FObjectInfo CachedObjectInfo;
-			
-			GetInstanceWorldObjectInfoFunction(ObjectData.RefHandle, ReplicatedInstances[InternalObjectIndex], CachedObjectInfo.WorldLocation, CachedObjectInfo.CullDistance);
+            const UNetObjectFactory::FWorldInfoContext Context
+            {
+                .Instance = ReplicatedInstances[InternalObjectIndex],
+                .Handle = ObjectData.RefHandle,
+                //TODO: Set InfoRequested based on the what was set dirty by the user
+            };
+
+            UNetObjectFactory::FWorldInfoData WorldInfo;
+
+            UNetObjectFactory* NetFactory = GetNetFactory(ObjectData.NetFactoryId);
+            NetFactory->GetWorldInfo(Context, WorldInfo);
 
 			// Check the position lies within world boundaries.
-			ensureMsgf(CachedObjectInfo.WorldLocation.X >= WorldLocations.GetWorldMinPos().X && CachedObjectInfo.WorldLocation.X <= WorldLocations.GetWorldMaxPos().X && 
-					   CachedObjectInfo.WorldLocation.Y >= WorldLocations.GetWorldMinPos().Y && CachedObjectInfo.WorldLocation.Y <= WorldLocations.GetWorldMaxPos().Y, TEXT("Object %s with position %s lies outside configured world boundary."), ToCStr(NetRefHandleManager->PrintObjectFromIndex((InternalObjectIndex))), *CachedObjectInfo.WorldLocation.ToString());
+			ensureMsgf(WorldInfo.WorldLocation.X >= WorldLocations.GetWorldMinPos().X && WorldInfo.WorldLocation.X <= WorldLocations.GetWorldMaxPos().X && 
+					   WorldInfo.WorldLocation.Y >= WorldLocations.GetWorldMinPos().Y && WorldInfo.WorldLocation.Y <= WorldLocations.GetWorldMaxPos().Y,
+                       TEXT("Object %s with position %s lies outside configured world boundary."), ToCStr(NetRefHandleManager->PrintObjectFromIndex((InternalObjectIndex))), *WorldInfo.WorldLocation.ToString());
 			
-			WorldLocations.SetObjectInfo(InternalObjectIndex, CachedObjectInfo);
+            const FWorldLocations::FObjectInfo WorldObjectInfo
+            { 
+                .WorldLocation = WorldInfo.WorldLocation, 
+                .CullDistance = WorldInfo.CullDistance 
+            };
+
+            WorldLocations.SetObjectInfo(InternalObjectIndex, WorldObjectInfo);
 		}
 	};
 
@@ -2066,7 +2106,8 @@ void UObjectReplicationBridge::LoadConfig()
 	}
 
 	// Filter mappings.
-	DefaultSpatialFilterHandle = ReplicationSystem->GetFilterHandle(BridgeConfig->GetDefaultSpatialFilterName());
+	DefaultSpatialFilterName = BridgeConfig->GetDefaultSpatialFilterName();
+	DefaultSpatialFilterHandle = ReplicationSystem->GetFilterHandle(DefaultSpatialFilterName);
 
 	{
 		for (const FObjectReplicationBridgeFilterConfig& FilterConfig : BridgeConfig->GetFilterConfigs())
diff --git a/Engine/Source/Runtime/Experimental/Iris/Core/Private/Iris/ReplicationSystem/ObjectReplicationBridgeDebugging.cpp b/Engine/Source/Runtime/Experimental/Iris/Core/Private/Iris/ReplicationSystem/ObjectReplicationBridgeDebugging.cpp
--- a/Engine/Source/Runtime/Experimental/Iris/Core/Private/Iris/ReplicationSystem/ObjectReplicationBridgeDebugging.cpp
+++ b/Engine/Source/Runtime/Experimental/Iris/Core/Private/Iris/ReplicationSystem/ObjectReplicationBridgeDebugging.cpp
@@ -656,20 +656,23 @@ void UObjectReplicationBridge::PrintNetCullDistances(const TArray<FString>& Args
 			}
 
 			FCullDistanceInfo& Info = ClassCullDistanceMap.FindOrAdd(RepObjClass);
+
+			// Find the CullDistance of the CDO.
 			if (Info.Class == nullptr)
 			{
-				UObject* RepClassCDO = RepObjClass->GetDefaultObject();
-
-				// Find the CullDistance of the CDO.
-				float CDOCullDistance = 0.0;
-				if (GetInstanceWorldObjectInfoFunction)
+				UNetObjectFactory* NetFactory = GetNetFactory(NetRefHandleManager->GetReplicatedObjectData(RootObjectIndex).NetFactoryId);
+				const UNetObjectFactory::FWorldInfoContext Context
 				{
-					FVector Loc;
-					GetInstanceWorldObjectInfoFunction(NetRefHandleManager->GetNetRefHandleFromInternalIndex(RootObjectIndex), RepClassCDO, Loc, CDOCullDistance);
-				}
+					.Instance = RepObjClass->GetDefaultObject(), // Cheat by passing the CDO as the instance here.
+					.Handle = FNetRefHandle::GetInvalid(), // Set an invalid handle since the CDO instance is not replicated.
+					.InfoRequested = UNetObjectFactory::EWorldInfoRequested::CullDistance
+				};
+
+				UNetObjectFactory::FWorldInfoData WorldData;
+				NetFactory->GetWorldInfo(Context, WorldData);
 
 				Info.Class = RepObjClass;
-				Info.CDOCullDistance = CDOCullDistance;
+				Info.CDOCullDistance = WorldData.CullDistance;
 			}
 
 			Info.NumTotal++;
diff --git a/Engine/Source/Runtime/Experimental/Iris/Core/Public/Iris/ReplicationSystem/NetObjectFactory.h b/Engine/Source/Runtime/Experimental/Iris/Core/Public/Iris/ReplicationSystem/NetObjectFactory.h
--- a/Engine/Source/Runtime/Experimental/Iris/Core/Public/Iris/ReplicationSystem/NetObjectFactory.h
+++ b/Engine/Source/Runtime/Experimental/Iris/Core/Public/Iris/ReplicationSystem/NetObjectFactory.h
@@ -8,6 +8,8 @@
 #include "Iris/ReplicationSystem/ReplicationBridgeTypes.h"
 #include "Iris/ReplicationSystem/NetRefHandle.h"
 
+#include "Misc/EnumClassFlags.h"
+
 #include "NetObjectFactory.generated.h"
 
 class UObjectReplicationBridge;
@@ -86,6 +88,8 @@ public:
 	struct FPostInstantiationContext;
 	struct FPostInitContext;
 	struct FDestroyedContext;
+	struct FWorldInfoContext;
+	struct FWorldInfoData;
 
 	/**
 	 * Creates the header containing all information required to instantiate a remote version of the object represented by the handle.
@@ -157,6 +161,14 @@ public:
 	 */
 	virtual void SubObjectDestroyedFromReplication(const FDestroyedContext& Context) {}
 
+	/**
+	 * Fetch world information about a replicated object so it can be updated in the network engine.
+	 * Only called for root objects.
+	 * @param Context Gives access to the replicated object and which specific world information needs to be updated.
+	 * @param OutData Set in this the requested information.
+	 */
+	virtual void GetWorldInfo(const FWorldInfoContext& Context, FWorldInfoData& OutData) PURE_VIRTUAL(UNetObjectFactory::GetWorldInfo,);
+
 public:
 
 	/** Result of the instantiate request */
@@ -217,6 +229,38 @@ public:
 		EReplicationBridgeDestroyInstanceFlags DestroyFlags = EReplicationBridgeDestroyInstanceFlags::None;
 	};
 
+	/** Details which info needs to be updated in GetWorldInfo */
+	enum class EWorldInfoRequested : uint32
+	{
+		None = 0x0000,
+		Location = 0x0001,
+		CullDistance = 0x0002,
+		All = Location | CullDistance,
+	};
+
+	/** Context when asking the factory for information on a specific object */
+	struct FWorldInfoContext
+	{
+		/** The object instance we are requesting information about */
+		UObject* Instance = nullptr;
+
+		/** The handle of the object */
+		UE::Net::FNetRefHandle Handle;
+
+		/** Specify which info is requested to be updated. */
+		EWorldInfoRequested InfoRequested = EWorldInfoRequested::All;
+	};
+
+	/** The world data the factory needs to fill about a given object. */
+	struct FWorldInfoData
+	{
+		/** The current location of the object in the world. */
+		FVector WorldLocation = FVector::ZeroVector;
+
+		/** The network cull distance of the object. */
+		float CullDistance = 0.0f;
+	};
+
 protected:
 
 	virtual void OnInit() {}
@@ -249,4 +293,6 @@ protected:
 	UObjectReplicationBridge* Bridge = nullptr;
 	UE::Net::FNetObjectFactoryId FactoryId = UE::Net::InvalidNetObjectFactoryId;
 
-};
\ No newline at end of file
+};
+
+ENUM_CLASS_FLAGS(UNetObjectFactory::EWorldInfoRequested);
\ No newline at end of file
diff --git a/Engine/Source/Runtime/Experimental/Iris/Core/Public/Iris/ReplicationSystem/ObjectReplicationBridge.h b/Engine/Source/Runtime/Experimental/Iris/Core/Public/Iris/ReplicationSystem/ObjectReplicationBridge.h
--- a/Engine/Source/Runtime/Experimental/Iris/Core/Public/Iris/ReplicationSystem/ObjectReplicationBridge.h
+++ b/Engine/Source/Runtime/Experimental/Iris/Core/Public/Iris/ReplicationSystem/ObjectReplicationBridge.h
@@ -164,6 +164,9 @@ public:
 	/** Getter for CVar net.Iris.UseVerboseIrisCsvStats which is true if we want per-class iris stats output to CSV */
 	IRISCORE_API bool ShouldUseVerboseCsvStats() const;
 
+	/** Return the name of the default spatial filter config. @see UObjectReplicationBridgeConfig */
+	FName GetDefaultSpatialFilterName() const { return DefaultSpatialFilterName; }
+
 public:
 
 	struct FRootObjectReplicationParams
@@ -217,7 +220,7 @@ public:
 
 public:
 
-	// Debug functions that are triggered via console commands
+	// Debug functions that are triggered via console commands declared in ObjectReplicationBridgeDebugging.cpp
 	void PrintDynamicFilterClassConfig() const;
 	void PrintReplicatedObjects(uint32 ArgTraits) const;
 	void PrintAlwaysRelevantObjects(uint32 ArgTraits) const;
@@ -268,15 +271,6 @@ protected:
 	/** Get the function that we should call before copying state data. */
 	IRISCORE_API FInstancePreUpdateFunction GetInstancePreUpdateFunction() const;
 
-	/** Helper method to get the world location & cull distance for replicated instances with the HasWorldLocation trait. */
-	using FInstanceGetWorldObjectInfoFunction = TFunction<void(UE::Net::FNetRefHandle, const UObject*, FVector&, float&)>;
-
-	/** Set the function that we should call to get the world location of an object. */
-	IRISCORE_API void SetInstanceGetWorldObjectInfoFunction(FInstanceGetWorldObjectInfoFunction InGetWorldObjectInfoFunction);
-	
-	/** Get the function that we should call to get the world location of an object. */
-	IRISCORE_API FInstanceGetWorldObjectInfoFunction GetInstanceGetWorldObjectInfoFunction() const;
-
 	// Poll frequency support
 
 	/** Force polling of Object when ObjectToPollWith is polled. */
@@ -294,6 +288,9 @@ protected:
 	/** Re-initialize config-driven parameters found in ObjectReplicationBridgeConfig. */
 	IRISCORE_API void LoadConfig();
 
+	/** Force the specified object to update it's cached WorldLocation data. */
+	IRISCORE_API void ForceUpdateWorldLocation(FNetRefHandle NetRefHandle, UE::Net::Private::FInternalNetRefIndex InternalObjectIndex);
+
 	/**
 	 * Set the function used to determine whether an object should use the default spatial filter
 	 * unless the class is configured to use some other filter.
@@ -388,7 +385,6 @@ private:
 	int32 GetTypeStatsIndex(const UClass* Class);
 
 	FInstancePreUpdateFunction PreUpdateInstanceFunction;
-	FInstanceGetWorldObjectInfoFunction GetInstanceWorldObjectInfoFunction;
 	
 	FName GetConfigClassPathName(const UClass* Class);
 
@@ -466,6 +462,7 @@ private:
 	// Objects that needs to update their object references due to garbage collection.
 	UE::Net::FNetBitArray GarbageCollectionAffectedObjects;
 
+	FName DefaultSpatialFilterName;
 	UE::Net::FNetObjectFilterHandle DefaultSpatialFilterHandle;
 
 	FDelegateHandle OnCustomConditionChangedHandle;
diff --git a/Engine/Source/Runtime/Experimental/Iris/Core/Public/Iris/ReplicationSystem/WorldLocations.h b/Engine/Source/Runtime/Experimental/Iris/Core/Public/Iris/ReplicationSystem/WorldLocations.h
--- a/Engine/Source/Runtime/Experimental/Iris/Core/Public/Iris/ReplicationSystem/WorldLocations.h
+++ b/Engine/Source/Runtime/Experimental/Iris/Core/Public/Iris/ReplicationSystem/WorldLocations.h
@@ -46,23 +46,21 @@ public:
 	void Deinit();
 
 	/** Returns whether the object has a valid cached data or not. */
-	bool HasInfoForObject(uint32 ObjectIndex) const;
+	inline bool HasInfoForObject(uint32 ObjectIndex) const;
 
 	/** Returns the object's world location if it's valid or a zero vector if it's not. */
-	FVector GetWorldLocation(uint32 ObjectIndex) const;
+	inline FVector GetWorldLocation(uint32 ObjectIndex) const;
 
 	/** Get the object's last cull distance we stored. */
-	float GetCullDistance(uint32 ObjectIndex) const;
+	inline float GetCullDistance(uint32 ObjectIndex) const;
 
 	/** Contains the cached object data we are storing. */
 	struct FObjectInfo
 	{
 		/** Absolute coordinate of the object */
-		FVector WorldLocation;
+		FVector WorldLocation = FVector::ZeroVector;
 		/** Network cull distance of the object */
-		float CullDistance;
-
-		FObjectInfo() : WorldLocation(EForceInit::ForceInitToZero), CullDistance(0.0f) {}
+		float CullDistance = 0.0f;
 	};
 
 public:
@@ -79,15 +77,15 @@ public:
 	 * Objects are not necessarily marked as dirty just because they're moving, such as objects attached to other objects. If such objects are spatially filtered they need to update their world locations in order for replication to work as expected.
 	 * Use SetObjectRequiresFrequentWorldLocationUpdate to force frequent world location update on an object.
 	 */
-	void SetObjectRequiresFrequentWorldLocationUpdate(uint32 ObjectIndex, bool bRequiresFrequentUpdate);
+	inline void SetObjectRequiresFrequentWorldLocationUpdate(uint32 ObjectIndex, bool bRequiresFrequentUpdate);
 
 	/** Returns whether an object requires frequent world location updates. */
-	bool GetObjectRequiresFrequentWorldLocationUpdate(uint32 ObjectIndex) const;
+	inline bool GetObjectRequiresFrequentWorldLocationUpdate(uint32 ObjectIndex) const;
 
-	FNetBitArrayView GetObjectsRequiringFrequentWorldLocationUpdate() const;
+	inline FNetBitArrayView GetObjectsRequiringFrequentWorldLocationUpdate() const;
 
 	void ResetObjectsWithDirtyInfo();
-	FNetBitArrayView GetObjectsWithDirtyInfo() const;
+	inline FNetBitArrayView GetObjectsWithDirtyInfo() const;
 
 	/** Returns the list of objects that registered world location information */
 	const FNetBitArrayView GetObjectsWithWorldInfo() const { return MakeNetBitArrayView(ValidInfoIndexes); }
