Dataflow : fix missing registration of vecxtor types and make sure anytypes are registered first

diff --git a/Engine/Source/Runtime/Experimental/Dataflow/Core/Private/Dataflow/DataflowNodeFactory.cpp b/Engine/Source/Runtime/Experimental/Dataflow/Core/Private/Dataflow/DataflowNodeFactory.cpp
--- a/Engine/Source/Runtime/Experimental/Dataflow/Core/Private/Dataflow/DataflowNodeFactory.cpp
+++ b/Engine/Source/Runtime/Experimental/Dataflow/Core/Private/Dataflow/DataflowNodeFactory.cpp
@@ -18,11 +18,13 @@ namespace UE::Dataflow
 
 	void FNodeFactory::RegisterDefaultNodes()
 	{
+		// anytypes need to be at the top as the connection registration in the nodes are using it
+		UE::Dataflow::RegisterAnyTypes();
+
 		RegisterCoreNodes();
 		RegisterDataflowMathNodes();
+		RegisterDataflowVectorNodes();
 		RegisterDataflowImageNodes();
-		UE::Dataflow::RegisterAnyTypes();
-		UE::Dataflow::RegisterAnyTypes();
 	}
 
 	TSharedPtr<FDataflowNode> FNodeFactory::NewNodeFromRegisteredType(FGraph& Graph, const FNewNodeParameters& Param)
