Dataflow : Anytype registry can now return the default type

diff --git a/Engine/Source/Runtime/Experimental/Dataflow/Core/Private/Dataflow/DataflowAnyTypeRegistry.cpp b/Engine/Source/Runtime/Experimental/Dataflow/Core/Private/Dataflow/DataflowAnyTypeRegistry.cpp
--- a/Engine/Source/Runtime/Experimental/Dataflow/Core/Private/Dataflow/DataflowAnyTypeRegistry.cpp
+++ b/Engine/Source/Runtime/Experimental/Dataflow/Core/Private/Dataflow/DataflowAnyTypeRegistry.cpp
@@ -9,6 +9,16 @@ namespace UE::Dataflow
 		return GetInstance().AreTypesCompatible(TypeA, TypeB);
 	}
 	
+	FName FAnyTypesRegistry::GetStorageTypeStatic(FName Type)
+	{
+		return GetInstance().GetStorageType(Type);
+	}
+
+	bool FAnyTypesRegistry::IsAnyTypeStatic(FName Type)
+	{
+		return GetInstance().IsAnyType(Type);
+	}
+
 	FAnyTypesRegistry& FAnyTypesRegistry::GetInstance()
 	{
 		static FAnyTypesRegistry Instance;
@@ -21,20 +31,34 @@ namespace UE::Dataflow
 		{
 			return true;
 		}
-		if (const FSupportTypeFunction* SupportFuncA = SupportTypeFunctionByName.Find(TypeA))
+		if (const FTypeInfo* TypeInfoA = TypeInfosByName.Find(TypeA))
 		{
-			if ((*SupportFuncA)(TypeB))
+			if (TypeInfoA->SupportTypeFunction && (*TypeInfoA->SupportTypeFunction)(TypeB))
 			{
 				return true;
 			}
 		}
-		if (const FSupportTypeFunction* SupportFuncB = SupportTypeFunctionByName.Find(TypeB))
+		if (const FTypeInfo* TypeInfoB = TypeInfosByName.Find(TypeB))
 		{
-			if ((*SupportFuncB)(TypeA))
+			if (TypeInfoB->SupportTypeFunction && (*TypeInfoB->SupportTypeFunction)(TypeA))
 			{
 				return true;
 			}
 		}
 		return false;
 	}
+
+	FName FAnyTypesRegistry::GetStorageType(FName Type) const
+	{
+		if (const FTypeInfo* TypeInfo = TypeInfosByName.Find(Type))
+		{
+			return TypeInfo->StorageType;
+		}
+		return NAME_None;
+	}
+
+	bool FAnyTypesRegistry::IsAnyType(FName Type) const
+	{
+		return TypeInfosByName.Contains(Type);
+	}
 };
diff --git a/Engine/Source/Runtime/Experimental/Dataflow/Core/Public/Dataflow/DataflowAnyTypeRegistry.h b/Engine/Source/Runtime/Experimental/Dataflow/Core/Public/Dataflow/DataflowAnyTypeRegistry.h
--- a/Engine/Source/Runtime/Experimental/Dataflow/Core/Public/Dataflow/DataflowAnyTypeRegistry.h
+++ b/Engine/Source/Runtime/Experimental/Dataflow/Core/Public/Dataflow/DataflowAnyTypeRegistry.h
@@ -4,6 +4,7 @@
 
 #include "UObject/NameTypes.h"
 #include "Containers/Map.h"
+#include "Dataflow/DataflowTypePolicy.h"
 
 namespace UE::Dataflow
 {
@@ -16,22 +17,37 @@ namespace UE::Dataflow
 			GetInstance().RegisterType<T>(TypeName);
 		}
 		static bool AreTypesCompatibleStatic(FName TypeA, FName TypeB);
+		static FName GetStorageTypeStatic(FName Type);
+		static bool IsAnyTypeStatic(FName Type);
 		static FAnyTypesRegistry& GetInstance();
 
 	private:
 		typedef bool(*FSupportTypeFunction)(FName);
 
+		struct FTypeInfo
+		{
+			FSupportTypeFunction SupportTypeFunction = nullptr;
+			FName StorageType;
+		};
+
 		FAnyTypesRegistry() {};
 
 		template<typename T>
 		void RegisterType(FName TypeName)
 		{
-			SupportTypeFunctionByName.Add(TypeName, &T::FPolicyType::SupportsTypeStatic);
+			FTypeInfo TypeInfo
+			{
+				.SupportTypeFunction = &T::FPolicyType::SupportsTypeStatic,
+				.StorageType = FName(TDataflowPolicyTypeName<typename T::FStorageType>::GetName()),
+			};
+			TypeInfosByName.Emplace(TypeName, TypeInfo);
 		}
 
 		bool AreTypesCompatible(FName TypeA, FName TypeB) const;
+		FName GetStorageType(FName Type) const;
+		bool IsAnyType(FName Type) const;
 
-		TMap<FName,FSupportTypeFunction> SupportTypeFunctionByName;
+		TMap<FName, FTypeInfo> TypeInfosByName;
 	};
 };
 
diff --git a/Engine/Source/Runtime/Experimental/Dataflow/Core/Public/Dataflow/DataflowTypePolicy.h b/Engine/Source/Runtime/Experimental/Dataflow/Core/Public/Dataflow/DataflowTypePolicy.h
--- a/Engine/Source/Runtime/Experimental/Dataflow/Core/Public/Dataflow/DataflowTypePolicy.h
+++ b/Engine/Source/Runtime/Experimental/Dataflow/Core/Public/Dataflow/DataflowTypePolicy.h
@@ -16,6 +16,16 @@ struct TDataflowPolicyTypeName
 {
 };
 
+// special version for void 
+template<>
+struct TDataflowPolicyTypeName<void>
+{
+	FORCEINLINE static const TCHAR* GetName()
+	{
+		return TEXT("");
+	}
+};
+
 #define UE_DATAFLOW_MAKE_STRING(x) #x
 #define UE_DATAFLOW_POLICY_DECLARE_TYPENAME(TType) \
 template<> \
@@ -65,6 +75,7 @@ UE_DATAFLOW_POLICY_DECLARE_TYPENAME(FIntVector4)
 UE_DATAFLOW_POLICY_DECLARE_TYPENAME(FRotator)
 UE_DATAFLOW_POLICY_DECLARE_TYPENAME(FTransform)
 
+UE_DATAFLOW_POLICY_DECLARE_TYPENAME(FDataflowSelection)
 UE_DATAFLOW_POLICY_DECLARE_TYPENAME(FDataflowTransformSelection)
 UE_DATAFLOW_POLICY_DECLARE_TYPENAME(FDataflowVertexSelection)
 UE_DATAFLOW_POLICY_DECLARE_TYPENAME(FDataflowFaceSelection)
