Dataflow : Add specific method to check if type dependent connections are connected

diff --git a/Engine/Source/Runtime/Experimental/Dataflow/Core/Private/Dataflow/DataflowNode.cpp b/Engine/Source/Runtime/Experimental/Dataflow/Core/Private/Dataflow/DataflowNode.cpp
--- a/Engine/Source/Runtime/Experimental/Dataflow/Core/Private/Dataflow/DataflowNode.cpp
+++ b/Engine/Source/Runtime/Experimental/Dataflow/Core/Private/Dataflow/DataflowNode.cpp
@@ -1346,23 +1346,8 @@ bool FDataflowNode::OnConnectionTypeChanged(const FDataflowConnection& Connectio
 		return false; // no changes
 	}
 
-	// first check if any connection is connected 
-	bool bIsAnyDependenciesConnected = false;
-	ForEachConnection(
-		[&bIsAnyDependenciesConnected, &Connection](const FDataflowConnection* OtherConnection)
-		{
-			// ignore the connection that is actively being notified
-			if (OtherConnection 
-				&& (OtherConnection != &Connection)
-				&& OtherConnection->IsConnected()
-				)
-			{
-				bIsAnyDependenciesConnected = true;
-				return false; // stop the iteration
-			}
-			return true;
-		});
-	if (bIsAnyDependenciesConnected)
+	// checks if any connection is connected 
+	if (IsAnytypeDependencyConnected(DependencyGroup, &Connection))
 	{
 		return false; // no changes
 	}
@@ -1446,4 +1431,30 @@ bool FDataflowNode::SetAllConnectionConcreteType(FName NewType)
 	}
 
 	return bChanged;
+}
+
+bool FDataflowNode::IsAnytypeDependencyConnected(FName DependencyGroup, const FDataflowConnection* IgnoreConnection) const
+{
+	if (DependencyGroup.IsNone())
+	{
+		return false; 
+	}
+
+	// first check if any connection is connected 
+	bool bIsAnyDependencyConnected = false;
+	ForEachConnection(
+		[&bIsAnyDependencyConnected, &IgnoreConnection](const FDataflowConnection* Connection)
+		{
+			// ignore the connection that is actively being notified
+			if (Connection
+				&& (IgnoreConnection == nullptr || Connection != IgnoreConnection)
+				&& Connection->IsConnected()
+				)
+			{
+				bIsAnyDependencyConnected = true;
+				return false; // stop the iteration
+			}
+			return true;
+		});
+	return bIsAnyDependencyConnected;
 }
\ No newline at end of file
diff --git a/Engine/Source/Runtime/Experimental/Dataflow/Core/Public/Dataflow/DataflowNode.h b/Engine/Source/Runtime/Experimental/Dataflow/Core/Public/Dataflow/DataflowNode.h
--- a/Engine/Source/Runtime/Experimental/Dataflow/Core/Public/Dataflow/DataflowNode.h
+++ b/Engine/Source/Runtime/Experimental/Dataflow/Core/Public/Dataflow/DataflowNode.h
@@ -518,9 +518,9 @@ PRAGMA_ENABLE_DEPRECATION_WARNINGS
 	* Signature of the function is bool Function(FDataflowConnection* Connection)
 	*/
 	template <typename TFunction>
-	void ForEachConnection(TFunction Function)
+	void ForEachConnection(TFunction Function) const
 	{
-		for (TPair<UE::Dataflow::FConnectionKey, FDataflowInput*>& Elem : ExpandedInputs)
+		for (const TPair<UE::Dataflow::FConnectionKey, FDataflowInput*>& Elem : ExpandedInputs)
 		{
 			if (Function(Elem.Value) == false)
 			{
@@ -631,6 +631,9 @@ PRAGMA_ENABLE_DEPRECATION_WARNINGS
 	// returns true if the output type was changed successfully
 	DATAFLOWCORE_API bool SetOutputConcreteType(const UE::Dataflow::FConnectionReference& OutputReference, FName NewType);
 
+	/** Check if any connection matching a specific type dependency group are connected */
+	DATAFLOWCORE_API bool IsAnytypeDependencyConnected(FName DependencyGroup, const FDataflowConnection* IgnoreConnection = nullptr) const;
+
 protected:
 	bool OnConnectionTypeChanged(const FDataflowConnection& Connection);
 
