Dataflow:
- Added ability to specify WireThickness for PinType
- Added custom color override for Input/Output pins
- Bug fixes

diff --git a/Engine/Plugins/Experimental/Dataflow/Source/DataflowEditor/Private/Dataflow/DataflowSNode.cpp b/Engine/Plugins/Experimental/Dataflow/Source/DataflowEditor/Private/Dataflow/DataflowSNode.cpp
--- a/Engine/Plugins/Experimental/Dataflow/Source/DataflowEditor/Private/Dataflow/DataflowSNode.cpp
+++ b/Engine/Plugins/Experimental/Dataflow/Source/DataflowEditor/Private/Dataflow/DataflowSNode.cpp
@@ -20,19 +20,23 @@
 #include "Widgets/Images/SImage.h"
 #include "Editor/Transactor.h"
 #include "GraphEditorSettings.h"
+#include "Dataflow/DataflowSchema.h"
 
 #include UE_INLINE_GENERATED_CPP_BY_NAME(DataflowSNode)
 
 #define LOCTEXT_NAMESPACE "SDataflowEdNode"
 
 //
-// SDataflowOutputPin
+// SDataflowPin
 //
 
-void SDataflowOutputPin::Construct(const FArguments& InArgs, UEdGraphPin* InPin)
+void SDataflowPin::Construct(const FArguments& InArgs, UEdGraphPin* InPin)
 {
-	const bool bIsPinInvalid = InArgs._IsPinInvalid.Get();
-	const FText InvalidPinDisplayText = bIsPinInvalid ? NSLOCTEXT("DataflowGraph", "DataflowOutputPinInvalidText", "*") : NSLOCTEXT("DataflowGraph", "DataflowOutputPinValidText", " ");
+	bIsPinInvalid = InArgs._IsPinInvalid.Get();
+	PinColorOverride = InArgs._PinColorOverride.Get();
+	bIsPinColorOverriden = InArgs._bIsPinColorOverriden.Get();
+
+	const FText InvalidPinDisplayText = bIsPinInvalid.Get() ? NSLOCTEXT("DataflowGraph", "DataflowOutputPinInvalidText", "*") : NSLOCTEXT("DataflowGraph", "DataflowOutputPinValidText", " ");
 
 	SGraphPin::Construct(SGraphPin::FArguments(), InPin);
 
@@ -48,6 +52,46 @@ void SDataflowOutputPin::Construct(const FArguments& InArgs, UEdGraphPin* InPin)
 		];
 }
 
+FSlateColor SDataflowPin::GetPinColor() const
+{
+	const bool IsColorOverriden = bIsPinColorOverriden.Get();
+	const FLinearColor ColorOverride = PinColorOverride.Get();
+
+	UEdGraphPin* GraphPin = GetPinObj();
+	if (GraphPin && !GraphPin->IsPendingKill())
+	{
+		if (bIsDiffHighlighted)
+		{
+			return FSlateColor(FLinearColor(0.9f, 0.2f, 0.15f));
+		}
+		if (GraphPin->bOrphanedPin)
+		{
+			return FSlateColor(FLinearColor::Red);
+		}
+		if (const UEdGraphSchema* Schema = GraphPin->GetSchema())
+		{
+			if (!GetPinObj()->GetOwningNode()->IsNodeEnabled() || GetPinObj()->GetOwningNode()->IsDisplayAsDisabledForced() || !IsEditingEnabled() || GetPinObj()->GetOwningNode()->IsNodeUnrelated())
+			{
+				if (IsColorOverriden)
+				{
+					return FSlateColor(ColorOverride);
+				}
+
+				return Schema->GetPinTypeColor(GraphPin->PinType) * FLinearColor(1.0f, 1.0f, 1.0f, 0.5f);
+			}
+
+			if (IsColorOverriden)
+			{
+				return FSlateColor(ColorOverride);
+			}
+
+			return Schema->GetPinTypeColor(GraphPin->PinType) * PinColorModifier;
+		}
+	}
+
+	return FLinearColor::White;
+}
+
 //
 // SDataflowEdNode
 //
@@ -123,12 +167,16 @@ void SDataflowEdNode::Construct(const FArguments& InArgs, UDataflowEdNode* InNod
 
 TSharedPtr<SGraphPin> SDataflowEdNode::CreatePinWidget(UEdGraphPin* Pin) const
 {
-	if (Pin->Direction == EEdGraphPinDirection::EGPD_Output)
+	if (DataflowGraphNode)
 	{
-		if (DataflowGraphNode)
+		if (TSharedPtr<FDataflowNode> DataflowNode = DataflowGraphNode->GetDataflowNode())
 		{
-			if (TSharedPtr<FDataflowNode> DataflowNode = DataflowGraphNode->GetDataflowNode())
+			bool bIsOutputInvalid = false;
+
+			if (Pin->Direction == EEdGraphPinDirection::EGPD_Output)
 			{
+
+
 				if (FDataflowOutput* Output = DataflowNode->FindOutput(Pin->GetFName()))
 				{
 					if (const TSharedPtr<UE::Dataflow::FContext> DataflowContext = DataflowInterface->GetDataflowContext())
@@ -142,18 +190,34 @@ TSharedPtr<SGraphPin> SDataflowEdNode::CreatePinWidget(UEdGraphPin* Pin) const
 						// Node's Timestamp is invalid or
 						// Node's Timestamp is greater than CacheKey's Timestamp -> Pin is invalid
 						//
-						const bool bIsOutputInvalid = !NumKeys ||
+						bIsOutputInvalid = !NumKeys ||
 							!CacheKeys.Contains(Output->CacheKey()) ||
-PRAGMA_DISABLE_DEPRECATION_WARNINGS  // Until LastModifiedTimestamp becomes private
-							DataflowNode->LastModifiedTimestamp.IsInvalid() || 
+						PRAGMA_DISABLE_DEPRECATION_WARNINGS  // Until LastModifiedTimestamp becomes private
+							DataflowNode->LastModifiedTimestamp.IsInvalid() ||
 							!DataflowContext->IsCacheEntryAfterTimestamp(Output->CacheKey(), DataflowNode->LastModifiedTimestamp);
-PRAGMA_ENABLE_DEPRECATION_WARNINGS
+						PRAGMA_ENABLE_DEPRECATION_WARNINGS
 
-						return SNew(SDataflowOutputPin, Pin)
-								.IsPinInvalid(bIsOutputInvalid);
 					}
 				}
 			}
+				
+			bool bColorOverriden = false;
+			FLinearColor OverrideColor = FLinearColor::Black;
+
+			if (const UDataflowSchema* DataflowSchema = Cast<UDataflowSchema>(Pin->GetSchema()))
+			{
+				TOptional<FLinearColor> OverrideColorOpt = DataflowSchema->GetPinColorOverride(DataflowNode, Pin);
+				if (OverrideColorOpt.IsSet())
+				{
+					OverrideColor = OverrideColorOpt.GetValue();
+					bColorOverriden = true;
+				}
+			}
+
+			return SNew(SDataflowPin, Pin)
+				.IsPinInvalid(bIsOutputInvalid)
+				.bIsPinColorOverriden(bColorOverriden)
+				.PinColorOverride(OverrideColor);
 		}
 	}
 
@@ -266,7 +330,6 @@ void SDataflowEdNode::UpdateErrorInfo()
 			}
 		}
 	}
-	
 }
 
 FReply SDataflowEdNode::OnMouseButtonDoubleClick(const FGeometry& InMyGeometry, const FPointerEvent& InMouseEvent)
diff --git a/Engine/Plugins/Experimental/Dataflow/Source/DataflowEditor/Private/Dataflow/DataflowSchema.cpp b/Engine/Plugins/Experimental/Dataflow/Source/DataflowEditor/Private/Dataflow/DataflowSchema.cpp
--- a/Engine/Plugins/Experimental/Dataflow/Source/DataflowEditor/Private/Dataflow/DataflowSchema.cpp
+++ b/Engine/Plugins/Experimental/Dataflow/Source/DataflowEditor/Private/Dataflow/DataflowSchema.cpp
@@ -22,6 +22,8 @@
 
 #define LOCTEXT_NAMESPACE "DataflowNode"
 
+static const float CDefaultWireThickness = 1.5f;
+
 UDataflowSchema::UDataflowSchema()
 {
 }
@@ -286,14 +288,27 @@ FLinearColor UDataflowSchema::GetTypeColor(const FName& Type)
 	const UGraphEditorSettings* Settings = GetDefault<UGraphEditorSettings>();
 	const UDataflowSettings* DataflowSettings = GetDefault<UDataflowSettings>();
 
-	if (::UE::Dataflow::FPinColorRegistry::Get().IsPinTypeRegistered(Type))
+	if (::UE::Dataflow::FPinSettingsRegistry::Get().IsPinTypeRegistered(Type))
 	{
-		return ::UE::Dataflow::FPinColorRegistry::Get().GetPinColor(Type);
+		return ::UE::Dataflow::FPinSettingsRegistry::Get().GetPinColor(Type);
 	}
 
 	return Settings->DefaultPinTypeColor;
 }
 
+float UDataflowSchema::GetPinTypeWireThickness(const FName& Type) const
+{
+	const UGraphEditorSettings* Settings = GetDefault<UGraphEditorSettings>();
+	const UDataflowSettings* DataflowSettings = GetDefault<UDataflowSettings>();
+
+	if (::UE::Dataflow::FPinSettingsRegistry::Get().IsPinTypeRegistered(Type))
+	{
+		return ::UE::Dataflow::FPinSettingsRegistry::Get().GetPinWireThickness(Type);
+	}
+
+	return CDefaultWireThickness;
+}
+
 static void CreateAndConnectNewReRouteNode(UEdGraphPin* FromPin, UEdGraphPin* ToPin, const FVector2D& GraphPosition)
 {
 	const UEdGraphNode* FromNode = FromPin->GetOwningNode();
@@ -354,6 +369,31 @@ bool UDataflowSchema::TryCreateConnection(UEdGraphPin* PinA, UEdGraphPin* PinB)
 	return Super::TryCreateConnection(PinA, PinB);
 }
 
+TOptional<FLinearColor> UDataflowSchema::GetPinColorOverride(TSharedPtr<FDataflowNode> DataflowNode, UEdGraphPin* Pin) const
+{
+	TOptional<FLinearColor> OutColor;
+
+	if (DataflowNode)
+	{
+		if (const UScriptStruct* ScriptStruct = DataflowNode->TypedScriptStruct())
+		{
+			FProperty* PinProperty = ScriptStruct->FindPropertyByName(Pin->GetFName());
+			if (PinProperty && PinProperty->HasMetaData("PinColor"))
+			{
+				FLinearColor WireColor;
+
+				const FString& ColorString = PinProperty->GetMetaData("PinColor");
+				if (WireColor.InitFromString(ColorString))
+				{
+					OutColor = WireColor;
+				}
+			}
+		}
+	}
+
+	return OutColor;
+}
+
 FConnectionDrawingPolicy* UDataflowSchema::CreateConnectionDrawingPolicy(int32 InBackLayerID, int32 InFrontLayerID, float InZoomFactor, const FSlateRect& InClippingRect, class FSlateWindowElementList& InDrawElements, class UEdGraph* InGraphObj) const
 {
 	return new FDataflowConnectionDrawingPolicy(InBackLayerID, InFrontLayerID, InZoomFactor, InClippingRect, InDrawElements, InGraphObj);
@@ -379,6 +419,25 @@ void FDataflowConnectionDrawingPolicy::DetermineWiringStyle(UEdGraphPin* OutputP
 	if (DataflowSchema && OutputPin)
 	{
 		Params.WireColor = DataflowSchema->GetPinTypeColor(OutputPin->PinType);
+
+		// Check if there is color override specified on the UProperty
+		UDataflowEdNode* EdNode = Cast<UDataflowEdNode>(OutputPin->GetOwningNode());
+		if (EdNode)
+		{
+			TSharedPtr<FDataflowNode> DataflowNode = EdNode->GetDataflowNode();
+			if (DataflowNode)
+			{
+				TOptional<FLinearColor> OverrideColor = DataflowSchema->GetPinColorOverride(DataflowNode, OutputPin);
+				if (OverrideColor.IsSet())
+				{
+					Params.WireColor = OverrideColor.GetValue();
+				}
+			}
+		}
+
+		float WireThickness = DataflowSchema->GetPinTypeWireThickness(OutputPin->PinType.PinCategory);
+		WireThickness = WireThickness < CDefaultWireThickness ? CDefaultWireThickness : WireThickness;
+		Params.WireThickness = WireThickness;
 	}
 
 	if (OutputPin && InputPin)
diff --git a/Engine/Plugins/Experimental/Dataflow/Source/DataflowEditor/Public/Dataflow/DataflowSNode.h b/Engine/Plugins/Experimental/Dataflow/Source/DataflowEditor/Public/Dataflow/DataflowSNode.h
--- a/Engine/Plugins/Experimental/Dataflow/Source/DataflowEditor/Public/Dataflow/DataflowSNode.h
+++ b/Engine/Plugins/Experimental/Dataflow/Source/DataflowEditor/Public/Dataflow/DataflowSNode.h
@@ -15,22 +15,30 @@ class UDataflowEdNode;
 class SCheckBox;
 
 //
-// SDataflowOutputPin
+// SDataflowPin
 //
 
-class SDataflowOutputPin : public SGraphPin
+class SDataflowPin : public SGraphPin
 {
 public:
-	SLATE_BEGIN_ARGS(SDataflowOutputPin)
+	SLATE_BEGIN_ARGS(SDataflowPin)
 		: _IsPinInvalid(false)
+		, _PinColorOverride(FLinearColor::White)
+		, _bIsPinColorOverriden(false)
 	{}
 		SLATE_ATTRIBUTE(bool, IsPinInvalid)
+		SLATE_ATTRIBUTE(FLinearColor, PinColorOverride)
+		SLATE_ATTRIBUTE(bool, bIsPinColorOverriden)
 	SLATE_END_ARGS()
 
 	void Construct(const FArguments& InArgs, UEdGraphPin* InPin);
 
+	virtual FSlateColor GetPinColor() const override;
+
 private:
-	TAttribute<bool> IsPinInvalid;
+	TAttribute<bool> bIsPinInvalid;
+	TAttribute<FLinearColor> PinColorOverride;
+	TAttribute<bool> bIsPinColorOverriden;
 };
 
 //
diff --git a/Engine/Plugins/Experimental/Dataflow/Source/DataflowEditor/Public/Dataflow/DataflowSchema.h b/Engine/Plugins/Experimental/Dataflow/Source/DataflowEditor/Public/Dataflow/DataflowSchema.h
--- a/Engine/Plugins/Experimental/Dataflow/Source/DataflowEditor/Public/Dataflow/DataflowSchema.h
+++ b/Engine/Plugins/Experimental/Dataflow/Source/DataflowEditor/Public/Dataflow/DataflowSchema.h
@@ -35,6 +35,9 @@ public:
 	//~ End UEdGraphSchema Interface
 
 	static FLinearColor GetTypeColor(const FName& Type);
+	float GetPinTypeWireThickness(const FName& Type) const;
+	TOptional<FLinearColor> GetPinColorOverride(TSharedPtr<FDataflowNode> DataflowNode, UEdGraphPin* Pin) const;
+
 private:
 	static bool CanPinBeConnectedToNode(const UEdGraphPin* Pin, const UE::Dataflow::FFactoryParameters& NodeParameters);
 };
diff --git a/Engine/Plugins/Experimental/GeometryCollectionPlugin/Source/GeometryCollectionNodes/Private/Dataflow/DataflowNodeAndPinTypeColors.cpp b/Engine/Plugins/Experimental/GeometryCollectionPlugin/Source/GeometryCollectionNodes/Private/Dataflow/DataflowNodeAndPinTypeColors.cpp
--- a/Engine/Plugins/Experimental/GeometryCollectionPlugin/Source/GeometryCollectionNodes/Private/Dataflow/DataflowNodeAndPinTypeColors.cpp
+++ b/Engine/Plugins/Experimental/GeometryCollectionPlugin/Source/GeometryCollectionNodes/Private/Dataflow/DataflowNodeAndPinTypeColors.cpp
@@ -10,6 +10,7 @@ namespace UE::Dataflow
 	void RegisterGeometryCollectionNodesColors()
 	{
 		static const FLinearColor CDefaultNodeBodyTintColor = FLinearColor(0.f, 0.f, 0.f, 0.5f);
+		static const float CDefaultWireThickness = 1.5f;
 
 		// Node colors
 		DATAFLOW_NODE_REGISTER_CREATION_FACTORY_NODE_COLORS_BY_CATEGORY("GeometryCollection|Fracture", FLinearColor(1.f, 1.f, 0.8f), CDefaultNodeBodyTintColor);
@@ -23,7 +24,6 @@ namespace UE::Dataflow
 		DATAFLOW_NODE_REGISTER_CREATION_FACTORY_NODE_COLORS_BY_CATEGORY("Fields", FLinearColor(.0f, 0.8f, 1.f), CDefaultNodeBodyTintColor);
 		DATAFLOW_NODE_REGISTER_CREATION_FACTORY_NODE_COLORS_BY_CATEGORY("Generators", FLinearColor(.4f, 0.8f, 0.f), CDefaultNodeBodyTintColor);
 		DATAFLOW_NODE_REGISTER_CREATION_FACTORY_NODE_COLORS_BY_CATEGORY("Materials", FLinearColor(FColor::Purple), CDefaultNodeBodyTintColor);
-		DATAFLOW_NODE_REGISTER_CREATION_FACTORY_NODE_COLORS_BY_CATEGORY("Math", FLinearColor(0.f, 0.4f, 0.8f), CDefaultNodeBodyTintColor);
 		DATAFLOW_NODE_REGISTER_CREATION_FACTORY_NODE_COLORS_BY_CATEGORY("Mesh", FLinearColor(1.f, 0.16f, 0.05f), CDefaultNodeBodyTintColor);
 		DATAFLOW_NODE_REGISTER_CREATION_FACTORY_NODE_COLORS_BY_CATEGORY("GeometryCollection", FLinearColor(0.55f, 0.45f, 1.0f), CDefaultNodeBodyTintColor);
 		DATAFLOW_NODE_REGISTER_CREATION_FACTORY_NODE_COLORS_BY_CATEGORY("Development", FLinearColor(1.f, 0.f, 0.f), CDefaultNodeBodyTintColor);
@@ -37,26 +37,13 @@ namespace UE::Dataflow
 		DATAFLOW_NODE_REGISTER_CREATION_FACTORY_NODE_COLORS_BY_CATEGORY("GeometryCollection|UV", FLinearColor(.1f, 1.f, 0.6f), CDefaultNodeBodyTintColor);
 
 		// PinType colors
-		DATAFLOW_NODE_REGISTER_CREATION_FACTORY_PIN_COLORS_BY_TYPE("FManagedArrayCollection", FLinearColor(0.353393f, 0.454175f, 1.0f, 1.0f));
-		DATAFLOW_NODE_REGISTER_CREATION_FACTORY_PIN_COLORS_BY_TYPE("float", FLinearColor(0.357667f, 1.f, 0.06f));
-		DATAFLOW_NODE_REGISTER_CREATION_FACTORY_PIN_COLORS_BY_TYPE("double", FLinearColor(0.039216f, 0.666667f, 0.f));
-		DATAFLOW_NODE_REGISTER_CREATION_FACTORY_PIN_COLORS_BY_TYPE("byte", FLinearColor(0.f, 0.16f, 0.13127f));
-		DATAFLOW_NODE_REGISTER_CREATION_FACTORY_PIN_COLORS_BY_TYPE("int32", FLinearColor(0.013575f, 0.77f, 0.429609f));
-		DATAFLOW_NODE_REGISTER_CREATION_FACTORY_PIN_COLORS_BY_TYPE("bool", FLinearColor(0.3f, 0.f, 0.f, 1.f));
-		DATAFLOW_NODE_REGISTER_CREATION_FACTORY_PIN_COLORS_BY_TYPE("FString", FLinearColor(1.f, 0.f, 0.660537f));
-		DATAFLOW_NODE_REGISTER_CREATION_FACTORY_PIN_COLORS_BY_TYPE("FName", FLinearColor(0.607717f, 0.224984f, 1.f));
-		DATAFLOW_NODE_REGISTER_CREATION_FACTORY_PIN_COLORS_BY_TYPE("FText", FLinearColor(0.8f, 0.2f, 0.4f));
-		DATAFLOW_NODE_REGISTER_CREATION_FACTORY_PIN_COLORS_BY_TYPE("FVector", FLinearColor(1.f, 0.591255f, 0.016512f));
-		DATAFLOW_NODE_REGISTER_CREATION_FACTORY_PIN_COLORS_BY_TYPE("FTransform", FLinearColor(1.f, 0.172585f, 0.f));
-		DATAFLOW_NODE_REGISTER_CREATION_FACTORY_PIN_COLORS_BY_TYPE("FRotator", FLinearColor(0.353393f, 0.454175f, 1.f));
-		DATAFLOW_NODE_REGISTER_CREATION_FACTORY_PIN_COLORS_BY_TYPE("TArray<FString>", FLinearColor(1.0f, 0.172585f, 0.0f, 1.0f));
-		DATAFLOW_NODE_REGISTER_CREATION_FACTORY_PIN_COLORS_BY_TYPE("FBox", FLinearColor(0.013575f, 0.770000f, 0.429609f, 1.0f));
-		DATAFLOW_NODE_REGISTER_CREATION_FACTORY_PIN_COLORS_BY_TYPE("FSphere", FLinearColor(0.2f, 0.6f, 1.f, 1.0f));
-		DATAFLOW_NODE_REGISTER_CREATION_FACTORY_PIN_COLORS_BY_TYPE("FDataflowAnyType", FLinearColor(0.3f, 0.3f, 0.3f, 1.0f));
-		DATAFLOW_NODE_REGISTER_CREATION_FACTORY_PIN_COLORS_BY_TYPE("FDataflowTransformSelection", FLinearColor(0.f, 0.9f, 0.0f, 1.0f));
-		DATAFLOW_NODE_REGISTER_CREATION_FACTORY_PIN_COLORS_BY_TYPE("TArray<TObjectPtr<UMaterial>>", FLinearColor(0.f, 0.2f, 0.0f, 1.0f));
-		DATAFLOW_NODE_REGISTER_CREATION_FACTORY_PIN_COLORS_BY_TYPE("TArray<TObjectPtr<UMaterialInterface>>", FLinearColor(0.f, 0.4f, 0.0f, 1.0f));
-		DATAFLOW_NODE_REGISTER_CREATION_FACTORY_PIN_COLORS_BY_TYPE("TArray<FGeometryCollectionAutoInstanceMesh>", FLinearColor(0.f, 0.6f, 0.0f, 1.0f));
-		DATAFLOW_NODE_REGISTER_CREATION_FACTORY_PIN_COLORS_BY_TYPE("TObjectPtr<UStaticMesh>", FLinearColor(0.f, 0.8f, 0.0f, 1.0f));
+		DATAFLOW_NODE_REGISTER_CREATION_FACTORY_PIN_SETTINGS_BY_TYPE("FManagedArrayCollection", FLinearColor(0.353393f, 0.454175f, 1.0f, 1.0f), CDefaultWireThickness);
+		DATAFLOW_NODE_REGISTER_CREATION_FACTORY_PIN_SETTINGS_BY_TYPE("TArray<FString>", FLinearColor(1.0f, 0.172585f, 0.0f, 1.0f), CDefaultWireThickness);
+		DATAFLOW_NODE_REGISTER_CREATION_FACTORY_PIN_SETTINGS_BY_TYPE("FBox", FLinearColor(0.013575f, 0.770000f, 0.429609f, 1.0f), CDefaultWireThickness);
+		DATAFLOW_NODE_REGISTER_CREATION_FACTORY_PIN_SETTINGS_BY_TYPE("FSphere", FLinearColor(0.2f, 0.6f, 1.f, 1.0f), CDefaultWireThickness);
+		DATAFLOW_NODE_REGISTER_CREATION_FACTORY_PIN_SETTINGS_BY_TYPE("TArray<TObjectPtr<UMaterial>>", FLinearColor(0.f, 0.2f, 0.0f, 1.0f), CDefaultWireThickness);
+		DATAFLOW_NODE_REGISTER_CREATION_FACTORY_PIN_SETTINGS_BY_TYPE("TArray<TObjectPtr<UMaterialInterface>>", FLinearColor(0.f, 0.4f, 0.0f, 1.0f), CDefaultWireThickness);
+		DATAFLOW_NODE_REGISTER_CREATION_FACTORY_PIN_SETTINGS_BY_TYPE("TArray<FGeometryCollectionAutoInstanceMesh>", FLinearColor(0.f, 0.6f, 0.0f, 1.0f), CDefaultWireThickness);
+		DATAFLOW_NODE_REGISTER_CREATION_FACTORY_PIN_SETTINGS_BY_TYPE("TObjectPtr<UStaticMesh>", FLinearColor(0.f, 0.8f, 0.0f, 1.0f), CDefaultWireThickness);
 	}
 }
diff --git a/Engine/Source/Runtime/Experimental/Dataflow/Core/Private/Dataflow/DataflowCoreNodeAndPinTypeColors.cpp b/Engine/Source/Runtime/Experimental/Dataflow/Core/Private/Dataflow/DataflowCoreNodeAndPinTypeColors.cpp
new file mode 100644
--- /dev/null
+++ b/Engine/Source/Runtime/Experimental/Dataflow/Core/Private/Dataflow/DataflowCoreNodeAndPinTypeColors.cpp
@@ -0,0 +1,80 @@
+// Copyright Epic Games, Inc. All Rights Reserved.
+
+#include "Dataflow/DataflowCoreNodeAndPinTypeColors.h"
+#include "Dataflow/DataflowNode.h"
+#include "Math/Color.h" 
+#include "Dataflow/DataflowNodeColorsRegistry.h"
+
+namespace UE::Dataflow
+{
+	void RegisterDataflowCoreColors()
+	{
+		static const FLinearColor CDefaultNodeBodyTintColor = FLinearColor(0.f, 0.f, 0.f, 0.5f);
+		static const float CDefaultWireThickness = 1.5f;
+
+		// Node colors
+		DATAFLOW_NODE_REGISTER_CREATION_FACTORY_NODE_COLORS_BY_CATEGORY("Math", FLinearColor(0.f, 0.4f, 0.8f), CDefaultNodeBodyTintColor);
+		DATAFLOW_NODE_REGISTER_CREATION_FACTORY_NODE_COLORS_BY_CATEGORY("Image", FLinearColor(0.f, 0.4f, 0.8f), CDefaultNodeBodyTintColor);
+
+		// PinType colors
+		DATAFLOW_NODE_REGISTER_CREATION_FACTORY_PIN_SETTINGS_BY_TYPE("bool", FLinearColor(0.300000f, 0.0f, 0.0f, 1.0f), CDefaultWireThickness);
+		DATAFLOW_NODE_REGISTER_CREATION_FACTORY_PIN_SETTINGS_BY_TYPE("uint8", FLinearColor(0.013575f, 0.770000f, 0.429609f, 1.0f), CDefaultWireThickness);
+		DATAFLOW_NODE_REGISTER_CREATION_FACTORY_PIN_SETTINGS_BY_TYPE("uint16", FLinearColor(0.013575f, 0.770000f, 0.429609f, 1.0f), CDefaultWireThickness);
+		DATAFLOW_NODE_REGISTER_CREATION_FACTORY_PIN_SETTINGS_BY_TYPE("uint32", FLinearColor(0.013575f, 0.770000f, 0.429609f, 1.0f), CDefaultWireThickness);
+		DATAFLOW_NODE_REGISTER_CREATION_FACTORY_PIN_SETTINGS_BY_TYPE("uint64", FLinearColor(0.013575f, 0.770000f, 0.429609f, 1.0f), CDefaultWireThickness);
+		DATAFLOW_NODE_REGISTER_CREATION_FACTORY_PIN_SETTINGS_BY_TYPE("int8", FLinearColor(0.013575f, 0.770000f, 0.429609f, 1.0f), CDefaultWireThickness);
+		DATAFLOW_NODE_REGISTER_CREATION_FACTORY_PIN_SETTINGS_BY_TYPE("int16", FLinearColor(0.013575f, 0.770000f, 0.429609f, 1.0f), CDefaultWireThickness);
+		DATAFLOW_NODE_REGISTER_CREATION_FACTORY_PIN_SETTINGS_BY_TYPE("int32", FLinearColor(0.013575f, 0.770000f, 0.429609f, 1.0f), CDefaultWireThickness);
+		DATAFLOW_NODE_REGISTER_CREATION_FACTORY_PIN_SETTINGS_BY_TYPE("int64", FLinearColor(0.013575f, 0.770000f, 0.429609f, 1.0f), CDefaultWireThickness);
+		DATAFLOW_NODE_REGISTER_CREATION_FACTORY_PIN_SETTINGS_BY_TYPE("float", FLinearColor(0.357667f, 1.0f, 0.060000f, 1.0f), CDefaultWireThickness);
+		DATAFLOW_NODE_REGISTER_CREATION_FACTORY_PIN_SETTINGS_BY_TYPE("double", FLinearColor(0.039216f, 0.666667f, 0.0f, 1.0f), CDefaultWireThickness);
+
+		DATAFLOW_NODE_REGISTER_CREATION_FACTORY_PIN_SETTINGS_BY_TYPE("FName", FLinearColor(0.607717f, 0.224984f, 1.0f, 1.0f), CDefaultWireThickness);
+		DATAFLOW_NODE_REGISTER_CREATION_FACTORY_PIN_SETTINGS_BY_TYPE("FString", FLinearColor(1.0f, 0.0f, 0.660537f, 1.0f), CDefaultWireThickness);
+		DATAFLOW_NODE_REGISTER_CREATION_FACTORY_PIN_SETTINGS_BY_TYPE("FText", FLinearColor(0.8f, 0.2f, 0.4f, 1.0f), CDefaultWireThickness);
+
+		DATAFLOW_NODE_REGISTER_CREATION_FACTORY_PIN_SETTINGS_BY_TYPE("FVector2D", FLinearColor(1.0f, 0.591255f, 0.016512f, 1.0f), CDefaultWireThickness);
+		DATAFLOW_NODE_REGISTER_CREATION_FACTORY_PIN_SETTINGS_BY_TYPE("FVector", FLinearColor(1.0f, 0.591255f, 0.016512f, 1.0f), CDefaultWireThickness);
+		DATAFLOW_NODE_REGISTER_CREATION_FACTORY_PIN_SETTINGS_BY_TYPE("FVector4", FLinearColor(1.0f, 0.591255f, 0.016512f, 1.0f), CDefaultWireThickness);
+		DATAFLOW_NODE_REGISTER_CREATION_FACTORY_PIN_SETTINGS_BY_TYPE("FVector2f", FLinearColor(1.0f, 0.591255f, 0.016512f, 1.0f), CDefaultWireThickness);
+		DATAFLOW_NODE_REGISTER_CREATION_FACTORY_PIN_SETTINGS_BY_TYPE("FVector3f", FLinearColor(1.0f, 0.591255f, 0.016512f, 1.0f), CDefaultWireThickness);
+		DATAFLOW_NODE_REGISTER_CREATION_FACTORY_PIN_SETTINGS_BY_TYPE("FVector4f", FLinearColor(1.0f, 0.591255f, 0.016512f, 1.0f), CDefaultWireThickness);
+
+		DATAFLOW_NODE_REGISTER_CREATION_FACTORY_PIN_SETTINGS_BY_TYPE("FQuat", FLinearColor(1.f, 0.591255f, 0.016512f), CDefaultWireThickness);
+		DATAFLOW_NODE_REGISTER_CREATION_FACTORY_PIN_SETTINGS_BY_TYPE("FQuat4f", FLinearColor(1.f, 0.591255f, 0.016512f), CDefaultWireThickness);
+
+		DATAFLOW_NODE_REGISTER_CREATION_FACTORY_PIN_SETTINGS_BY_TYPE("FLinearColor", FLinearColor(1.f, 0.591255f, 0.016512f), CDefaultWireThickness);
+
+		DATAFLOW_NODE_REGISTER_CREATION_FACTORY_PIN_SETTINGS_BY_TYPE("FIntPoint", FLinearColor(1.f, 0.591255f, 0.016512f), CDefaultWireThickness);
+		DATAFLOW_NODE_REGISTER_CREATION_FACTORY_PIN_SETTINGS_BY_TYPE("FIntVector3", FLinearColor(1.f, 0.591255f, 0.016512f), CDefaultWireThickness);
+		DATAFLOW_NODE_REGISTER_CREATION_FACTORY_PIN_SETTINGS_BY_TYPE("FIntVector4", FLinearColor(1.f, 0.591255f, 0.016512f), CDefaultWireThickness);
+
+		DATAFLOW_NODE_REGISTER_CREATION_FACTORY_PIN_SETTINGS_BY_TYPE("FRotator", FLinearColor(0.353393f, 0.454175f, 1.0f, 1.0f), CDefaultWireThickness);
+		DATAFLOW_NODE_REGISTER_CREATION_FACTORY_PIN_SETTINGS_BY_TYPE("FTransform", FLinearColor(1.0f, 0.172585f, 0.0f, 1.0f), CDefaultWireThickness);
+
+		DATAFLOW_NODE_REGISTER_CREATION_FACTORY_PIN_SETTINGS_BY_TYPE("FDataflowTransformSelection", FLinearColor(0.3f, 0.3f, 0.3f, 1.0f), CDefaultWireThickness);
+		DATAFLOW_NODE_REGISTER_CREATION_FACTORY_PIN_SETTINGS_BY_TYPE("FDataflowVertexSelection", FLinearColor(0.3f, 0.3f, 0.3f, 1.0f), CDefaultWireThickness);
+		DATAFLOW_NODE_REGISTER_CREATION_FACTORY_PIN_SETTINGS_BY_TYPE("FDataflowFaceSelection", FLinearColor(0.3f, 0.3f, 0.3f, 1.0f), CDefaultWireThickness);
+		DATAFLOW_NODE_REGISTER_CREATION_FACTORY_PIN_SETTINGS_BY_TYPE("FDataflowGeometrySelection", FLinearColor(0.3f, 0.3f, 0.3f, 1.0f), CDefaultWireThickness);
+		DATAFLOW_NODE_REGISTER_CREATION_FACTORY_PIN_SETTINGS_BY_TYPE("FDataflowMaterialSelection", FLinearColor(0.3f, 0.3f, 0.3f, 1.0f), CDefaultWireThickness);
+
+		DATAFLOW_NODE_REGISTER_CREATION_FACTORY_PIN_SETTINGS_BY_TYPE("FDataflowAnyType", FLinearColor(0.220000f, 0.195800f, 0.195800f, 1.0f), CDefaultWireThickness);
+		DATAFLOW_NODE_REGISTER_CREATION_FACTORY_PIN_SETTINGS_BY_TYPE("FDataflowAllTypes", FLinearColor(0.220000f, 0.195800f, 0.195800f, 1.0f), CDefaultWireThickness);
+		DATAFLOW_NODE_REGISTER_CREATION_FACTORY_PIN_SETTINGS_BY_TYPE("FDataflowNumericTypes", FLinearColor(0.220000f, 0.195800f, 0.195800f, 1.0f), CDefaultWireThickness);
+		DATAFLOW_NODE_REGISTER_CREATION_FACTORY_PIN_SETTINGS_BY_TYPE("FDataflowVectorTypes", FLinearColor(0.220000f, 0.195800f, 0.195800f, 1.0f), CDefaultWireThickness);
+		DATAFLOW_NODE_REGISTER_CREATION_FACTORY_PIN_SETTINGS_BY_TYPE("FDataflowStringTypes", FLinearColor(0.220000f, 0.195800f, 0.195800f, 1.0f), CDefaultWireThickness);
+		DATAFLOW_NODE_REGISTER_CREATION_FACTORY_PIN_SETTINGS_BY_TYPE("FDataflowBoolTypes", FLinearColor(0.220000f, 0.195800f, 0.195800f, 1.0f), CDefaultWireThickness);
+		DATAFLOW_NODE_REGISTER_CREATION_FACTORY_PIN_SETTINGS_BY_TYPE("FDataflowTransformTypes", FLinearColor(0.220000f, 0.195800f, 0.195800f, 1.0f), CDefaultWireThickness);
+		DATAFLOW_NODE_REGISTER_CREATION_FACTORY_PIN_SETTINGS_BY_TYPE("FDataflowStringConvertibleTypes", FLinearColor(0.220000f, 0.195800f, 0.195800f, 1.0f), CDefaultWireThickness);
+		DATAFLOW_NODE_REGISTER_CREATION_FACTORY_PIN_SETTINGS_BY_TYPE("FDataflowUObjectConvertibleTypes", FLinearColor(0.220000f, 0.195800f, 0.195800f, 1.0f), CDefaultWireThickness);
+		DATAFLOW_NODE_REGISTER_CREATION_FACTORY_PIN_SETTINGS_BY_TYPE("FDataflowSelectionTypes", FLinearColor(0.220000f, 0.195800f, 0.195800f, 1.0f), CDefaultWireThickness);
+
+		DATAFLOW_NODE_REGISTER_CREATION_FACTORY_PIN_SETTINGS_BY_TYPE("FDataflowVectorArrayTypes", FLinearColor(0.220000f, 0.195800f, 0.195800f, 1.0f), CDefaultWireThickness);
+		DATAFLOW_NODE_REGISTER_CREATION_FACTORY_PIN_SETTINGS_BY_TYPE("FDataflowNumericArrayTypes", FLinearColor(0.220000f, 0.195800f, 0.195800f, 1.0f), CDefaultWireThickness);
+		DATAFLOW_NODE_REGISTER_CREATION_FACTORY_PIN_SETTINGS_BY_TYPE("FDataflowStringArrayTypes", FLinearColor(0.220000f, 0.195800f, 0.195800f, 1.0f), CDefaultWireThickness);
+		DATAFLOW_NODE_REGISTER_CREATION_FACTORY_PIN_SETTINGS_BY_TYPE("FDataflowBoolArrayTypes", FLinearColor(0.220000f, 0.195800f, 0.195800f, 1.0f), CDefaultWireThickness);
+		DATAFLOW_NODE_REGISTER_CREATION_FACTORY_PIN_SETTINGS_BY_TYPE("FDataflowTransformArrayTypes", FLinearColor(0.220000f, 0.195800f, 0.195800f, 1.0f), CDefaultWireThickness);
+
+		DATAFLOW_NODE_REGISTER_CREATION_FACTORY_PIN_SETTINGS_BY_TYPE("FDataflowImage", FLinearColor(0.220000f, 0.1f, 0.1f, 1.0f), CDefaultWireThickness);
+	}
+}
diff --git a/Engine/Source/Runtime/Experimental/Dataflow/Core/Private/Dataflow/DataflowNodeColorsRegistry.cpp b/Engine/Source/Runtime/Experimental/Dataflow/Core/Private/Dataflow/DataflowNodeColorsRegistry.cpp
--- a/Engine/Source/Runtime/Experimental/Dataflow/Core/Private/Dataflow/DataflowNodeColorsRegistry.cpp
+++ b/Engine/Source/Runtime/Experimental/Dataflow/Core/Private/Dataflow/DataflowNodeColorsRegistry.cpp
@@ -53,7 +53,7 @@ namespace UE::Dataflow
 		GetMutableDefault<UDataflowSettings>()->RegisterColors(Category, NodeColors);
 	}
 
-	FLinearColor FNodeColorsRegistry::GetNodeTitleColor(const FName& Category)
+	FLinearColor FNodeColorsRegistry::GetNodeTitleColor(const FName& Category) const
 	{
 		if (ColorsMap.Contains(Category))
 		{
@@ -82,7 +82,7 @@ namespace UE::Dataflow
 		return FNodeColors().NodeTitleColor;
 	}
 
-	FLinearColor FNodeColorsRegistry::GetNodeBodyTintColor(const FName& Category)
+	FLinearColor FNodeColorsRegistry::GetNodeBodyTintColor(const FName& Category) const
 	{
 		if (ColorsMap.Contains(Category))
 		{
@@ -121,20 +121,20 @@ namespace UE::Dataflow
 
 	/* --------------------------------------------------------------------------------------------------------------- */
 
-	FPinColorRegistry::FPinColorRegistry()
+	FPinSettingsRegistry::FPinSettingsRegistry()
 	{
 		UDataflowSettings* DataflowSettings = GetMutableDefault<UDataflowSettings>();
-		DataflowSettingsChangedDelegateHandle = DataflowSettings->GetOnDataflowSettingsChangedPinColorDelegate().AddRaw(this, &FPinColorRegistry::PinColorChangedInSettings);
+		DataflowSettingsChangedDelegateHandle = DataflowSettings->GetOnDataflowSettingsChangedPinSettingsDelegate().AddRaw(this, &FPinSettingsRegistry::PinSettingsChangedInSettings);
 
-		const FPinColorMap PinColorMap = DataflowSettings->GetPinColorMap();
+		const FPinSettingsMap PinColorMap = DataflowSettings->GetPinSettingsMap();
 
 		for (auto& Elem : PinColorMap)
 		{
-			ColorMap.FindOrAdd(Elem.Key) = Elem.Value;
+			SettingsMap.FindOrAdd(Elem.Key) = Elem.Value;
 		}
 	}
 
-	FPinColorRegistry::~FPinColorRegistry()
+	FPinSettingsRegistry::~FPinSettingsRegistry()
 	{
 		if (UObjectInitialized())
 		{
@@ -146,53 +146,62 @@ namespace UE::Dataflow
 		}
 	}
 
-	FPinColorRegistry& FPinColorRegistry::Get()
+	FPinSettingsRegistry& FPinSettingsRegistry::Get()
 	{
-		return TLazySingleton<FPinColorRegistry>::Get();
+		return TLazySingleton<FPinSettingsRegistry>::Get();
 	}
 
-	void FPinColorRegistry::TearDown()
+	void FPinSettingsRegistry::TearDown()
 	{
-		return TLazySingleton<FPinColorRegistry>::TearDown();
+		return TLazySingleton<FPinSettingsRegistry>::TearDown();
 	}
 
-	void FPinColorRegistry::RegisterPinColor(const FName& PinType, const FLinearColor& PinColor)
+	void FPinSettingsRegistry::RegisterPinSettings(const FName& PinType, const FPinSettings& InSettings)
 	{
-		if (!ColorMap.Contains(PinType))
+		if (!SettingsMap.Contains(PinType))
 		{
-			ColorMap.Add(PinType, PinColor);
+			SettingsMap.Add(PinType, InSettings);
 		}
 
 		// Register colors in DataflowSettings
-		GetMutableDefault<UDataflowSettings>()->RegisterPinColor(PinType, PinColor);
+		GetMutableDefault<UDataflowSettings>()->RegisterPinSettings(PinType, InSettings);
 	}
 
-	FLinearColor FPinColorRegistry::GetPinColor(const FName& PinType)
+	FLinearColor FPinSettingsRegistry::GetPinColor(const FName& PinType) const
 	{
-		if (ColorMap.Contains(PinType))
+		if (SettingsMap.Contains(PinType))
 		{
-			return ColorMap[PinType];
+			return SettingsMap[PinType].PinColor;
 		}
 
 		return FLinearColor(0.f, 0.f, 0.f);
 	}
 
-	void FPinColorRegistry::PinColorChangedInSettings(const FPinColorMap& PinColorMap)
+	float FPinSettingsRegistry::GetPinWireThickness(const FName& PinType) const
 	{
-		for (auto& Elem : PinColorMap)
+		if (SettingsMap.Contains(PinType))
+		{
+			return SettingsMap[PinType].WireThickness;
+		}
+
+		return 1.f;
+	}
+
+	void FPinSettingsRegistry::PinSettingsChangedInSettings(const FPinSettingsMap& PinSettingsrMap)
+	{
+		for (auto& Elem : PinSettingsrMap)
 		{
-			ColorMap.FindOrAdd(Elem.Key) = Elem.Value;
+			SettingsMap.FindOrAdd(Elem.Key) = Elem.Value;
 		}
 	}
 
-	bool FPinColorRegistry::IsPinTypeRegistered(const FName& PinType)
+	bool FPinSettingsRegistry::IsPinTypeRegistered(const FName& PinType) const
 	{
-		if (ColorMap.Contains(PinType))
+		if (SettingsMap.Contains(PinType))
 		{
 			return true;
 		}
 
 		return false;
 	}
-
 }
\ No newline at end of file
diff --git a/Engine/Source/Runtime/Experimental/Dataflow/Core/Private/Dataflow/DataflowNodeFactory.cpp b/Engine/Source/Runtime/Experimental/Dataflow/Core/Private/Dataflow/DataflowNodeFactory.cpp
--- a/Engine/Source/Runtime/Experimental/Dataflow/Core/Private/Dataflow/DataflowNodeFactory.cpp
+++ b/Engine/Source/Runtime/Experimental/Dataflow/Core/Private/Dataflow/DataflowNodeFactory.cpp
@@ -9,6 +9,7 @@
 #include "Dataflow/DataflowVectorNodes.h"
 #include "Misc/MessageDialog.h"
 #include "Logging/LogMacros.h"
+#include "Dataflow/DataflowCoreNodeAndPinTypeColors.h"
 
 DEFINE_LOG_CATEGORY_STATIC(LogDataflowFactory, Warning, All);
 
@@ -21,6 +22,8 @@ namespace UE::Dataflow
 		// anytypes need to be at the top as the connection registration in the nodes are using it
 		UE::Dataflow::RegisterAnyTypes();
 
+		UE::Dataflow::RegisterDataflowCoreColors();
+
 		RegisterCoreNodes();
 		RegisterDataflowMathNodes();
 		RegisterDataflowVectorNodes();
diff --git a/Engine/Source/Runtime/Experimental/Dataflow/Core/Private/Dataflow/DataflowSettings.cpp b/Engine/Source/Runtime/Experimental/Dataflow/Core/Private/Dataflow/DataflowSettings.cpp
--- a/Engine/Source/Runtime/Experimental/Dataflow/Core/Private/Dataflow/DataflowSettings.cpp
+++ b/Engine/Source/Runtime/Experimental/Dataflow/Core/Private/Dataflow/DataflowSettings.cpp
@@ -39,7 +39,7 @@ void UDataflowSettings::PostEditChangeProperty(FPropertyChangedEvent& PropertyCh
 	if (PropertyChangedEvent.Property != nullptr)
 	{
 		OnDataflowSettingsChangedDelegate.Broadcast(NodeColorsMap);
-		OnDataflowSettingsChangedPinColorDelegate.Broadcast(PinColorMap);
+		OnDataflowSettingsChangedPinSettingsDelegate.Broadcast(PinSettingsMap);
 	}
 
 	Super::PostEditChangeProperty(PropertyChangedEvent);
@@ -56,13 +56,13 @@ FNodeColors UDataflowSettings::RegisterColors(const FName& Category, const FNode
 	return NodeColorsMap[Category];
 }
 
-FLinearColor UDataflowSettings::RegisterPinColor(const FName& InPinType, const FLinearColor& InColor)
+FPinSettings UDataflowSettings::RegisterPinSettings(const FName& InPinType, const FPinSettings& InSettings)
 {
-	if (!PinColorMap.Contains(InPinType))
+	if (!PinSettingsMap.Contains(InPinType))
 	{
-		PinColorMap.Add(InPinType, InColor);
+		PinSettingsMap.Add(InPinType, InSettings);
 	}
-	return PinColorMap[InPinType];
+	return PinSettingsMap[InPinType];
 }
 
 
diff --git a/Engine/Source/Runtime/Experimental/Dataflow/Core/Public/Dataflow/DataflowCoreNodeAndPinTypeColors.h b/Engine/Source/Runtime/Experimental/Dataflow/Core/Public/Dataflow/DataflowCoreNodeAndPinTypeColors.h
new file mode 100644
--- /dev/null
+++ b/Engine/Source/Runtime/Experimental/Dataflow/Core/Public/Dataflow/DataflowCoreNodeAndPinTypeColors.h
@@ -0,0 +1,9 @@
+// Copyright Epic Games, Inc. All Rights Reserved.
+
+#pragma once 
+
+namespace UE::Dataflow
+{
+	void RegisterDataflowCoreColors();
+}
+
diff --git a/Engine/Source/Runtime/Experimental/Dataflow/Core/Public/Dataflow/DataflowImageNodes.h b/Engine/Source/Runtime/Experimental/Dataflow/Core/Public/Dataflow/DataflowImageNodes.h
--- a/Engine/Source/Runtime/Experimental/Dataflow/Core/Public/Dataflow/DataflowImageNodes.h
+++ b/Engine/Source/Runtime/Experimental/Dataflow/Core/Public/Dataflow/DataflowImageNodes.h
@@ -57,19 +57,19 @@ private:
 	FDataflowImage Image;
 
 	/** Red channel */
-	UPROPERTY(meta = (DataflowOutput));
+	UPROPERTY(meta = (DataflowOutput, PinColor = "R=1, G=0, B=0"));
 	FDataflowImage Red;
 
 	/** Green channel */
-	UPROPERTY(meta = (DataflowOutput));
+	UPROPERTY(meta = (DataflowOutput, PinColor = "R=0, G=1, B=0"));
 	FDataflowImage Green;
 
 	/** Blue channel */
-	UPROPERTY(meta = (DataflowOutput));
+	UPROPERTY(meta = (DataflowOutput, PinColor = "R=0, G=0.2, B=1"));
 	FDataflowImage Blue;
 
 	/** Alpha channel */
-	UPROPERTY(meta = (DataflowOutput));
+	UPROPERTY(meta = (DataflowOutput, PinColor = "R=0.7, G=0.7, B=0.7"));
 	FDataflowImage Alpha;
 
 	virtual void Evaluate(UE::Dataflow::FContext& Context, const FDataflowOutput* Out) const override;
@@ -98,19 +98,19 @@ public:
 
 private:
 	/** Red channel - if not connected, use black color */
-	UPROPERTY(meta = (DataflowInput));
+	UPROPERTY(meta = (DataflowInput, PinColor = "R=1, G=0, B=0"));
 	FDataflowImage Red;
 
 	/** Green channel - if not connected, use black color */
-	UPROPERTY(meta = (DataflowInput));
+	UPROPERTY(meta = (DataflowInput, PinColor = "R=0, G=1, B=0"));
 	FDataflowImage Green;
 
 	/** Blue channel - if not connected, use black color */
-	UPROPERTY(meta = (DataflowInput));
+	UPROPERTY(meta = (DataflowInput, PinColor = "R=0, G=0.2, B=1"));
 	FDataflowImage Blue;
 
 	/** Alpha channel - if not connected, use black color */
-	UPROPERTY(meta = (DataflowInput));
+	UPROPERTY(meta = (DataflowInput, PinColor = "R=0.7, G=0.7, B=0.7"));
 	FDataflowImage Alpha;
 
 	/** Output image recombined from input channels */
diff --git a/Engine/Source/Runtime/Experimental/Dataflow/Core/Public/Dataflow/DataflowNode.h b/Engine/Source/Runtime/Experimental/Dataflow/Core/Public/Dataflow/DataflowNode.h
--- a/Engine/Source/Runtime/Experimental/Dataflow/Core/Public/Dataflow/DataflowNode.h
+++ b/Engine/Source/Runtime/Experimental/Dataflow/Core/Public/Dataflow/DataflowNode.h
@@ -795,9 +795,9 @@ private:
 	::UE::Dataflow::FNodeColorsRegistry::Get().RegisterNodeColors(A, {C1, C2});		\
 }																					\
 
-#define DATAFLOW_NODE_REGISTER_CREATION_FACTORY_PIN_COLORS_BY_TYPE(A, C)			\
+#define DATAFLOW_NODE_REGISTER_CREATION_FACTORY_PIN_SETTINGS_BY_TYPE(A, C, T)		\
 {																					\
-	::UE::Dataflow::FPinColorRegistry::Get().RegisterPinColor(A, C);				\
+	::UE::Dataflow::FPinSettingsRegistry::Get().RegisterPinSettings(A, {C, T});		\
 }																					\
 
 }
diff --git a/Engine/Source/Runtime/Experimental/Dataflow/Core/Public/Dataflow/DataflowNodeColorsRegistry.h b/Engine/Source/Runtime/Experimental/Dataflow/Core/Public/Dataflow/DataflowNodeColorsRegistry.h
--- a/Engine/Source/Runtime/Experimental/Dataflow/Core/Public/Dataflow/DataflowNodeColorsRegistry.h
+++ b/Engine/Source/Runtime/Experimental/Dataflow/Core/Public/Dataflow/DataflowNodeColorsRegistry.h
@@ -24,8 +24,8 @@ namespace UE::Dataflow
 		static DATAFLOWCORE_API void TearDown();
 
 		DATAFLOWCORE_API void RegisterNodeColors(const FName& Category, const FNodeColors& NodeColors);
-		DATAFLOWCORE_API FLinearColor GetNodeTitleColor(const FName& Category);
-		DATAFLOWCORE_API FLinearColor GetNodeBodyTintColor(const FName& Category);
+		DATAFLOWCORE_API FLinearColor GetNodeTitleColor(const FName& Category) const;
+		DATAFLOWCORE_API FLinearColor GetNodeBodyTintColor(const FName& Category) const;
 		DATAFLOWCORE_API void NodeColorsChangedInSettings(const FNodeColorsMap& NodeColorsMap);
 
 	private:
@@ -41,22 +41,23 @@ namespace UE::Dataflow
 	//
 	// Registry for custom Pin colors
 	//
-	class FPinColorRegistry
+	class FPinSettingsRegistry
 	{
 	public:
-		static DATAFLOWCORE_API FPinColorRegistry& Get();
+		static DATAFLOWCORE_API FPinSettingsRegistry& Get();
 		static DATAFLOWCORE_API void TearDown();
 
-		DATAFLOWCORE_API void RegisterPinColor(const FName& PinType, const FLinearColor& PinColor);
-		DATAFLOWCORE_API FLinearColor GetPinColor(const FName& PinType);
-		DATAFLOWCORE_API void PinColorChangedInSettings(const FPinColorMap& PinColorMap);
-		DATAFLOWCORE_API bool IsPinTypeRegistered(const FName& PinType);
+		DATAFLOWCORE_API void RegisterPinSettings(const FName& PinType, const FPinSettings& InSettings);
+		DATAFLOWCORE_API FLinearColor GetPinColor(const FName& PinType) const;
+		DATAFLOWCORE_API float GetPinWireThickness(const FName& PinType) const;
+		DATAFLOWCORE_API void PinSettingsChangedInSettings(const FPinSettingsMap& PinSettingsrMap);
+		DATAFLOWCORE_API bool IsPinTypeRegistered(const FName& PinType) const;
 
 	private:
-		DATAFLOWCORE_API FPinColorRegistry();
-		DATAFLOWCORE_API ~FPinColorRegistry();
+		DATAFLOWCORE_API FPinSettingsRegistry();
+		DATAFLOWCORE_API ~FPinSettingsRegistry();
 
-		FPinColorMap ColorMap;					// [PinType] -> Color
+		FPinSettingsMap SettingsMap;					// [PinType] -> {Color, WireThickness}
 		FDelegateHandle DataflowSettingsChangedDelegateHandle;
 
 		friend FLazySingleton;
diff --git a/Engine/Source/Runtime/Experimental/Dataflow/Core/Public/Dataflow/DataflowSettings.h b/Engine/Source/Runtime/Experimental/Dataflow/Core/Public/Dataflow/DataflowSettings.h
--- a/Engine/Source/Runtime/Experimental/Dataflow/Core/Public/Dataflow/DataflowSettings.h
+++ b/Engine/Source/Runtime/Experimental/Dataflow/Core/Public/Dataflow/DataflowSettings.h
@@ -8,7 +8,7 @@
 #include "DataflowSettings.generated.h"
 
 typedef TMap<FName, FNodeColors> FNodeColorsMap;
-typedef TMap<FName, FLinearColor> FPinColorMap;
+typedef TMap<FName, FPinSettings> FPinSettingsMap;
 
 USTRUCT()
 struct FNodeColors
@@ -22,6 +22,18 @@ struct FNodeColors
 	FLinearColor NodeBodyTintColor = FLinearColor(0.f, 0.f, 0.f);
 };
 
+USTRUCT()
+struct FPinSettings
+{
+	GENERATED_USTRUCT_BODY()
+
+	UPROPERTY(EditAnywhere, Category = Settings)
+	FLinearColor PinColor = FLinearColor(0.f, 0.f, 0.f);
+
+	UPROPERTY(EditAnywhere, Category = Settings)
+	float WireThickness = 1.f;
+};
+
 USTRUCT()
 struct FTransformLevelColors
 {
@@ -42,8 +54,8 @@ class UDataflowSettings : public UDeveloperSettings
 	UPROPERTY(config, EditAnywhere, Category = NodeColors)
 	TMap<FName, FNodeColors> NodeColorsMap;
 
-	UPROPERTY(config, EditAnywhere, Category = PinColors)
-	TMap<FName, FLinearColor> PinColorMap;
+	UPROPERTY(config, EditAnywhere, Category = PinSettings)
+	TMap<FName, FPinSettings> PinSettingsMap;
 
 	UPROPERTY(config, EditAnywhere, Category = TransformLevelColors)
 	FTransformLevelColors TransformLevelColors;
@@ -52,10 +64,10 @@ class UDataflowSettings : public UDeveloperSettings
 	DATAFLOWCORE_API virtual FName GetCategoryName() const override;
 
 	DATAFLOWCORE_API FNodeColors RegisterColors(const FName& Category, const FNodeColors& Colors);
-	DATAFLOWCORE_API FLinearColor RegisterPinColor(const FName& PinType, const FLinearColor& Color);
+	DATAFLOWCORE_API FPinSettings RegisterPinSettings(const FName& PinType, const FPinSettings& InSettings);
 
 	const TMap<FName, FNodeColors>& GetNodeColorsMap() { return NodeColorsMap; }
-	const TMap<FName, FLinearColor>& GetPinColorMap() { return PinColorMap; }
+	const TMap<FName, FPinSettings>& GetPinSettingsMap() { return PinSettingsMap; }
 
 #if WITH_EDITOR
 	DATAFLOWCORE_API virtual FText GetSectionText() const override;
@@ -68,15 +80,15 @@ class UDataflowSettings : public UDeveloperSettings
 
 public:
 	DECLARE_MULTICAST_DELEGATE_OneParam(FOnDataflowSettingsChanged, const FNodeColorsMap&);
-	DECLARE_MULTICAST_DELEGATE_OneParam(FOnDataflowSettingsChangedPinColor, const FPinColorMap&);
+	DECLARE_MULTICAST_DELEGATE_OneParam(FOnDataflowSettingsChangedPinSettings, const FPinSettingsMap&);
 
 	/** Gets a multicast delegate which is called whenever one of the parameters in this settings object changes. */
 	FOnDataflowSettingsChanged& GetOnDataflowSettingsChangedDelegate() { return OnDataflowSettingsChangedDelegate; }
-	FOnDataflowSettingsChangedPinColor& GetOnDataflowSettingsChangedPinColorDelegate() { return OnDataflowSettingsChangedPinColorDelegate; }
+	FOnDataflowSettingsChangedPinSettings& GetOnDataflowSettingsChangedPinSettingsDelegate() { return OnDataflowSettingsChangedPinSettingsDelegate; }
 
 protected:
 	FOnDataflowSettingsChanged OnDataflowSettingsChangedDelegate;
-	FOnDataflowSettingsChangedPinColor OnDataflowSettingsChangedPinColorDelegate;
+	FOnDataflowSettingsChangedPinSettings OnDataflowSettingsChangedPinSettingsDelegate;
 
 };
 
