I/O store on-demand cleanup.

- extracted interface from module functions
- moved all config handling to OnDemandConfig.h
- I/O store and install cache are initialized at module startup and IAS backend after hotfix if specified

#rnx

diff --git a/Engine/Plugins/Experimental/DefaultInstallBundleManager/Source/Private/DefaultInstallBundleManager.cpp b/Engine/Plugins/Experimental/DefaultInstallBundleManager/Source/Private/DefaultInstallBundleManager.cpp
--- a/Engine/Plugins/Experimental/DefaultInstallBundleManager/Source/Private/DefaultInstallBundleManager.cpp
+++ b/Engine/Plugins/Experimental/DefaultInstallBundleManager/Source/Private/DefaultInstallBundleManager.cpp
@@ -1372,7 +1372,7 @@ void FDefaultInstallBundleManager::UpdateBundleSources(FContentReleaseRequestRef
 	if (BundleInfo.bMountedOnDemandTocs)
 	{
 		check(IOStoreOnDemandModule);
-		FIoStatus UnmountStatus = IOStoreOnDemandModule->Unmount(Request->BundleName.ToString());
+		const FIoStatus UnmountStatus = UE::IoStore::GetOnDemandIoStore().Unmount(Request->BundleName.ToString());
 		if (!UnmountStatus.IsOk())
 		{
 			LOG_INSTALL_BUNDLE_MAN_OVERRIDE(Request->LogVerbosityOverride, Error, TEXT("Error unmounting ondemand toc for bundle '%s' : %s"), *BundleInfo.BundleNameString, *UnmountStatus.ToString());
@@ -1725,7 +1725,7 @@ void FDefaultInstallBundleManager::MountPaks(FContentRequestRef Request)
 
 		BundleMountArgs->Options &= ~UE::IoStore::EOnDemandMountOptions::CallbackOnGameThread;
 
-		IOStoreOnDemandModule->Mount(MoveTemp(*BundleMountArgs), [OnDemandMountContext](UE::IoStore::FOnDemandMountResult OnDemandMountResult)
+		UE::IoStore::GetOnDemandIoStore().Mount(MoveTemp(*BundleMountArgs), [OnDemandMountContext](UE::IoStore::FOnDemandMountResult OnDemandMountResult)
 		{
 			OnDemandMountContext->Results.Add(MoveTemp(OnDemandMountResult));
 		});
diff --git a/Engine/Plugins/Experimental/RuntimeTelemetry/Source/Private/RuntimeTelemetry.cpp b/Engine/Plugins/Experimental/RuntimeTelemetry/Source/Private/RuntimeTelemetry.cpp
--- a/Engine/Plugins/Experimental/RuntimeTelemetry/Source/Private/RuntimeTelemetry.cpp
+++ b/Engine/Plugins/Experimental/RuntimeTelemetry/Source/Private/RuntimeTelemetry.cpp
@@ -24,16 +24,16 @@ void FRuntimeTelemetry::RecordEvent_IoStoreOnDemand(const FString& Context, TArr
 
 	if (FStudioTelemetry::Get().IsSessionRunning())
 	{
-		if (FIoStoreOnDemandModule* OnDemandModule = FModuleManager::Get().GetModulePtr<FIoStoreOnDemandModule>("IoStoreOnDemand"))
+		if (UE::IoStore::IOnDemandIoStore* IoStore = UE::IoStore::TryGetOnDemandIoStore())
 		{
 			// Only record event data if IoStoreOnDemand is enabled
-			if (OnDemandModule->IsEnabled())
+			if (EnumHasAnyFlags(IoStore->GetStreamingOptions(), EOnDemandStreamingOptions::Enabled))
 			{
 				const int SchemaVersion = 1;
 
 				TArray<FAnalyticsEventAttribute> EventAttributes = Attributes;
 
-				OnDemandModule->ReportAnalytics(EventAttributes);
+				IoStore->ReportAnalytics(EventAttributes);
 
 				EventAttributes.Emplace(TEXT("SchemaVersion"), SchemaVersion);
 				EventAttributes.Emplace(TEXT("Context"), Context);
diff --git a/Engine/Plugins/Runtime/GameFeatures/Source/GameFeatures/Private/GameFeaturePluginStateMachine.cpp b/Engine/Plugins/Runtime/GameFeatures/Source/GameFeatures/Private/GameFeaturePluginStateMachine.cpp
--- a/Engine/Plugins/Runtime/GameFeatures/Source/GameFeatures/Private/GameFeaturePluginStateMachine.cpp
+++ b/Engine/Plugins/Runtime/GameFeatures/Source/GameFeatures/Private/GameFeaturePluginStateMachine.cpp
@@ -1581,7 +1581,7 @@ struct FGameFeaturePluginState_Downloading : public FGameFeaturePluginState
 			check(InstallArgs.ContentHandle.IsValid());
 
 			// This should be pretty small, so not going to worry about progress here.
-			IoStoreOnDemandModule->Install(MoveTemp(InstallArgs),
+			UE::IoStore::GetOnDemandIoStore().Install(MoveTemp(InstallArgs),
 				[this, LambdaOnDemandContext = IoStoreOnDemandContext](const UE::IoStore::FOnDemandInstallResult& OnDemandInstallResult)
 				{
 					if (!LambdaOnDemandContext->bStateValid)
@@ -2758,7 +2758,7 @@ struct FGameFeaturePluginState_AssetDependencyStreaming : public FGameFeaturePlu
 			InstallArgs.ContentHandle = UE::GameFeatures::GFPSharedInstallTracker.AddOnDemandContentHandle(
 				InstallBundle, UE::GameFeatures::EGFPInstallLevel::AssetStream);
 
-			IoStoreOnDemandModule->Install(MoveTemp(InstallArgs),
+			UE::IoStore::GetOnDemandIoStore().Install(MoveTemp(InstallArgs),
 				// On Complete
 				[this, LambdaOnDemandContext = IoStoreOnDemandContext](const UE::IoStore::FOnDemandInstallResult& OnDemandInstallResult)
 				{
diff --git a/Engine/Plugins/Runtime/GameFeatures/Source/GameFeatures/Private/GameFeaturesSubsystem.cpp b/Engine/Plugins/Runtime/GameFeatures/Source/GameFeatures/Private/GameFeaturesSubsystem.cpp
--- a/Engine/Plugins/Runtime/GameFeatures/Source/GameFeatures/Private/GameFeaturesSubsystem.cpp
+++ b/Engine/Plugins/Runtime/GameFeatures/Source/GameFeatures/Private/GameFeaturesSubsystem.cpp
@@ -3085,7 +3085,7 @@ struct FGameFeatureAssetStreamContext : public FGameFeaturePluginAssetStreamHand
 
 			// TODO: convert Priority to IOStore priority here
 
-			IoStoreOnDemandModule->Install(MoveTemp(InstallArgs),
+			UE::IoStore::GetOnDemandIoStore().Install(MoveTemp(InstallArgs),
 				[Context = SharedThis(this), RequestIndex](const UE::IoStore::FOnDemandInstallResult& OnDemandInstallResult)
 				{
 					Context->OnInstallComplete(OnDemandInstallResult, RequestIndex);
diff --git a/Engine/Source/Runtime/Core/Private/IO/IoDispatcher.cpp b/Engine/Source/Runtime/Core/Private/IO/IoDispatcher.cpp
--- a/Engine/Source/Runtime/Core/Private/IO/IoDispatcher.cpp
+++ b/Engine/Source/Runtime/Core/Private/IO/IoDispatcher.cpp
@@ -50,7 +50,8 @@ static const TCHAR* const GetIoErrorText_ErrorCodeTextArray[] =
 	TEXT("Invalid Encryption Key"),
 	TEXT("Compression Error"),
 	TEXT("Pending Fork"),
-	TEXT("Pending Encryption Key")
+	TEXT("Pending Encryption Key"),
+	TEXT("Disabled")
 };
 
 CORE_API const TCHAR* const* GetIoErrorText_ErrorCodeText = GetIoErrorText_ErrorCodeTextArray;
diff --git a/Engine/Source/Runtime/Core/Public/IO/IoStatus.h b/Engine/Source/Runtime/Core/Public/IO/IoStatus.h
--- a/Engine/Source/Runtime/Core/Public/IO/IoStatus.h
+++ b/Engine/Source/Runtime/Core/Public/IO/IoStatus.h
@@ -31,6 +31,7 @@ enum class EIoErrorCode
 	CompressionError,
 	PendingFork,
 	PendingEncryptionKey,
+	Disabled
 };
 
 /** Get I/O error code description. */
diff --git a/Engine/Source/Runtime/Experimental/IoStore/OnDemand/Private/IoStoreOnDemand.cpp b/Engine/Source/Runtime/Experimental/IoStore/OnDemand/Private/IoStoreOnDemand.cpp
--- a/Engine/Source/Runtime/Experimental/IoStore/OnDemand/Private/IoStoreOnDemand.cpp
+++ b/Engine/Source/Runtime/Experimental/IoStore/OnDemand/Private/IoStoreOnDemand.cpp
@@ -5,20 +5,18 @@
 #include "Containers/StringView.h"
 #include "HAL/FileManager.h"
 #include "HAL/LowLevelMemTracker.h"
-#include "IO/IasCache.h"
-#include "Misc/Base64.h"
+#include "IO/IoStatus.h"
 #include "Misc/CommandLine.h"
 #include "Misc/ConfigCacheIni.h"
 #include "Misc/CoreDelegates.h"
-#include "Misc/EncryptionKeyManager.h"
 #include "Misc/Parse.h"
 #include "Misc/PathViews.h"
 #include "Modules/ModuleManager.h"
+#include "OnDemandConfig.h"
 #include "OnDemandIoStore.h"
 #include "OnDemandHttpClient.h"
 #include "OnDemandIoDispatcherBackend.h"
 #include "Statistics.h"
-#include "String/LexFromString.h"
 #if !UE_BUILD_SHIPPING
 #include "Misc/PackageName.h"
 #endif //!UE_BUILD_SHIPPING
@@ -35,6 +33,8 @@ DEFINE_LOG_CATEGORY(LogIas);
 namespace UE::IoStore
 {
 
+TSharedPtr<FOnDemandIoStore, ESPMode::ThreadSafe> GOnDemandIoStore;
+
 ///////////////////////////////////////////////////////////////////////////////
 #if !UE_BUILD_SHIPPING
 namespace Commands
@@ -69,49 +69,49 @@ static TIoStatusOr<FOnDemandHostGroup> SplitHostAndTocUrl(FStringView Url, FStri
 ////////////////////////////////////////////////////////////////////////////////
 static void PurgeInstallCache()
 {
-	FIoStoreOnDemandModule* IOStoreOnDemandModule = FModuleManager::Get().GetModulePtr<FIoStoreOnDemandModule>(TEXT("IoStoreOnDemand"));
-	if(!IOStoreOnDemandModule)
+	if (UE::IoStore::IOnDemandIoStore* IoStore = UE::IoStore::TryGetOnDemandIoStore())
 	{
-		UE_LOG(LogIoStoreOnDemand, Error, TEXT("Could not find IoStoreOnDemand module"));
-		return;
+		UE_LOG(LogIoStoreOnDemand, Display, TEXT("Purging on demand install cache"));
+		IoStore->Purge(FOnDemandPurgeArgs(), [](const FOnDemandPurgeResult& Result)
+		{
+			if (Result.Status.IsOk())
+			{
+				UE_LOG(LogIoStoreOnDemand, Display, TEXT("Purged on demand install cache"));
+			}
+			else
+			{
+				UE_LOG(LogIoStoreOnDemand, Error, TEXT("Failed Purged on demand install cache: %s"), *Result.Status.ToString());
+			}
+		});
 	}
-
-	UE_LOG(LogIoStoreOnDemand, Display, TEXT("Purging on demand install cache"));
-	IOStoreOnDemandModule->Purge(FOnDemandPurgeArgs(), [](const FOnDemandPurgeResult& Result)
+	else
 	{
-		if (Result.Status.IsOk())
-		{
-			UE_LOG(LogIoStoreOnDemand, Display, TEXT("Purged on demand install cache"));
-		}
-		else
-		{
-			UE_LOG(LogIoStoreOnDemand, Error, TEXT("Failed Purged on demand install cache: %s"), *Result.Status.ToString());
-		}
-	});
+		UE_LOG(LogIoStoreOnDemand, Error, TEXT("Purge install cache failed, reason 'I/O store on-demand module not initialized'"));
+	}
 }
 
 ////////////////////////////////////////////////////////////////////////////////
 static void PrintCacheUsage()
 {
-	FIoStoreOnDemandModule* IOStoreOnDemandModule = FModuleManager::Get().GetModulePtr<FIoStoreOnDemandModule>(TEXT("IoStoreOnDemand"));
-	if (!IOStoreOnDemandModule)
+	if (UE::IoStore::IOnDemandIoStore* IoStore = UE::IoStore::TryGetOnDemandIoStore())
 	{
-		UE_LOG(LogIoStoreOnDemand, Error, TEXT("Could not find IoStoreOnDemand module"));
-		return;
-	}
+		TIoStatusOr<FOnDemandCacheUsage> MaybeUsage = IoStore->GetCacheUsage();
+		if (!MaybeUsage.IsOk())
+		{
+			UE_LOG(LogIoStoreOnDemand, Error, TEXT("iostore.CacheUsage failed: %s"), *MaybeUsage.Status().ToString());
+			return;
+		}
 
-	TIoStatusOr<FOnDemandCacheUsage> MaybeUsage = IOStoreOnDemandModule->GetCacheUsage();
-	if (!MaybeUsage.IsOk())
+		const FOnDemandCacheUsage CacheUsage = MaybeUsage.ConsumeValueOrDie();
+		UE_LOG(LogIoStoreOnDemand, Display, TEXT("iostore.CacheUsage"));
+		TStringBuilder<512> Sb;
+		Sb << TEXTVIEW("InstallCache: ") << CacheUsage.InstallCache << TEXTVIEW(", StreamingCache: ") << CacheUsage.StreamingCache;
+		UE_LOG(LogIoStoreOnDemand, Display, TEXT("%s"), Sb.ToString());
+	}
+	else
 	{
-		UE_LOG(LogIoStoreOnDemand, Error, TEXT("iostore.CacheUsage failed: %s"), *MaybeUsage.Status().ToString());
-		return;
+		UE_LOG(LogIoStoreOnDemand, Error, TEXT("Print cache usage failed, reason 'I/O store on-demand module not initialized'"));
 	}
-
-	const FOnDemandCacheUsage CacheUsage = MaybeUsage.ConsumeValueOrDie();
-	UE_LOG(LogIoStoreOnDemand, Display, TEXT("iostore.CacheUsage"));
-	TStringBuilder<512> Sb;
-	Sb << TEXTVIEW("InstallCache: ") << CacheUsage.InstallCache << TEXTVIEW(", StreamingCache: ") << CacheUsage.StreamingCache;
-	UE_LOG(LogIoStoreOnDemand, Display, TEXT("%s"), Sb.ToString());
 }
 
 ////////////////////////////////////////////////////////////////////////////////
@@ -149,7 +149,7 @@ static void MountUrl(const TArray<FString>& Args)
 	}
 	const FString& MountId = Args.Num() > 2 ? Args[2] : Url; 
 
-	IoStoreOnDemand->Mount(
+	UE::IoStore::GetOnDemandIoStore().Mount(
 		FOnDemandMountArgs
 		{
 			.MountId = MountId,
@@ -247,7 +247,7 @@ static void MountFile(const TArray<FString>& Args)
 	for (FOnDemandMountArgs& MountArgs : AllMountArgs)
 	{
 		MountArgs.Options = Options;
-		IoStoreOnDemand->Mount(
+		UE::IoStore::GetOnDemandIoStore().Mount(
 			MoveTemp(MountArgs),
 			[](FOnDemandMountResult MountResult)
 			{
@@ -293,7 +293,7 @@ static void InstallPackage(const TArray<FString>& Args)
 	InstallArgs.PackageIds.Add(PackageId);
 	InstallArgs.ContentHandle = DefaultContentHandle;
 
-	IoStoreOnDemand->Install(
+	UE::IoStore::GetOnDemandIoStore().Install(
 		MoveTemp(InstallArgs),
 		[PackageName = MoveTemp(PackageName)](FOnDemandInstallResult InstallResult)
 		{
@@ -345,247 +345,6 @@ FString GIasOnDemandTocExt = TEXT(".uondemandtoc");
 
 static const TCHAR* NotInitializedError = TEXT("I/O store on-demand not initialized");
 
-/** Temp cvar to allow the fallback url to be hotfixed in case of problems */
-static FString GDistributedEndpointFallbackUrl;
-static FAutoConsoleVariableRef CVar_DistributedEndpointFallbackUrl(
-	TEXT("ias.DistributedEndpointFallbackUrl"),
-	GDistributedEndpointFallbackUrl,
-	TEXT("CDN url to be used if a distributed endpoint cannot be reached (overrides IoStoreOnDemand.ini)")
-);
-
-////////////////////////////////////////////////////////////////////////////////
-int64 ParseSizeParam(FStringView Value)
-{
-	Value = Value.TrimStartAndEnd();
-
-	int64 Size = -1;
-	LexFromString(Size, Value);
-	if (Size >= 0)
-	{
-		if (Value.EndsWith(TEXT("GB"))) return Size << 30;
-		if (Value.EndsWith(TEXT("MB"))) return Size << 20;
-		if (Value.EndsWith(TEXT("KB"))) return Size << 10;
-	}
-	return Size;
-}
-
-////////////////////////////////////////////////////////////////////////////////
-static int64 ParseSizeParam(const TCHAR* CommandLine, const TCHAR* Param)
-{
-	FString ParamValue;
-	if (!FParse::Value(CommandLine, Param, ParamValue))
-	{
-		return -1;
-	}
-
-	return ParseSizeParam(ParamValue);
-}
-
-////////////////////////////////////////////////////////////////////////////////
-static bool ParseEncryptionKeyParam(const FString& Param, FGuid& OutKeyGuid, FAES::FAESKey& OutKey)
-{
-	TArray<FString> Tokens;
-	Param.ParseIntoArray(Tokens, TEXT(":"), true);
-
-	if (Tokens.Num() == 2)
-	{
-		TArray<uint8> KeyBytes;
-		if (FGuid::Parse(Tokens[0], OutKeyGuid) && FBase64::Decode(Tokens[1], KeyBytes))
-		{
-			if (OutKeyGuid != FGuid() && KeyBytes.Num() == FAES::FAESKey::KeySize)
-			{
-				FMemory::Memcpy(OutKey.Key, KeyBytes.GetData(), FAES::FAESKey::KeySize);
-				return true;
-			}
-		}
-	}
-	
-	return false;
-}
-
-////////////////////////////////////////////////////////////////////////////////
-static bool ApplyEncryptionKeyFromString(const FString& GuidKeyPair)
-{
-	FGuid KeyGuid;
-	FAES::FAESKey Key;
-
-	if (ParseEncryptionKeyParam(GuidKeyPair, KeyGuid, Key))
-	{
-		// TODO: PAK and I/O store should share key manager
-		FEncryptionKeyManager::Get().AddKey(KeyGuid, Key);
-		FCoreDelegates::GetRegisterEncryptionKeyMulticastDelegate().Broadcast(KeyGuid, Key);
-
-		return true;
-	}
-	else
-	{
-		return false;
-	}
-}
-
-////////////////////////////////////////////////////////////////////////////////
-static bool TryParseConfigContent(const FString& ConfigContent, const FString& ConfigFileName, FOnDemandEndpointConfig& OutEndpoint)
-{
-	if (ConfigContent.IsEmpty())
-	{
-		return false;
-	}
-
-	FConfigFile Config;
-	Config.ProcessInputFileContents(ConfigContent, ConfigFileName);
-
-	Config.GetString(TEXT("Endpoint"), TEXT("DistributionUrl"), OutEndpoint.DistributionUrl);
-	if (!OutEndpoint.DistributionUrl.IsEmpty())
-	{
-		Config.GetString(TEXT("Endpoint"), TEXT("FallbackUrl"), OutEndpoint.FallbackUrl);
-
-		if (!GDistributedEndpointFallbackUrl.IsEmpty())
-		{
-			OutEndpoint.FallbackUrl = GDistributedEndpointFallbackUrl;
-		}
-	}
-	
-	Config.GetArray(TEXT("Endpoint"), TEXT("ServiceUrl"), OutEndpoint.ServiceUrls);
-	Config.GetString(TEXT("Endpoint"), TEXT("TocPath"), OutEndpoint.TocPath);
-
-	if (OutEndpoint.DistributionUrl.EndsWith(TEXT("/")))
-	{
-		OutEndpoint.DistributionUrl = OutEndpoint.DistributionUrl.Left(OutEndpoint.DistributionUrl.Len() - 1);
-	}
-
-	for (FString& ServiceUrl : OutEndpoint.ServiceUrls)
-	{
-		if (ServiceUrl.EndsWith(TEXT("/")))
-		{
-			ServiceUrl.LeftInline(ServiceUrl.Len() - 1);
-		}
-	}
-
-	if (OutEndpoint.TocPath.StartsWith(TEXT("/")))
-	{
-		OutEndpoint.TocPath.RightChopInline(1);
-	}
-
-	FString ContentKey;
-	if (Config.GetString(TEXT("Endpoint"), TEXT("ContentKey"), ContentKey))
-	{
-		ApplyEncryptionKeyFromString(ContentKey);
-	}
-
-	return OutEndpoint.IsValid();
-}
-
-////////////////////////////////////////////////////////////////////////////////
-static bool TryParseConfigFileFromPlatformPackage(FOnDemandEndpointConfig& OutConfig)
-{
-	const FString ConfigFileName = TEXT("IoStoreOnDemand.ini");
-	const FString ConfigPath = FPaths::Combine(TEXT("Cloud"), ConfigFileName);
-	
-	if (FPlatformMisc::FileExistsInPlatformPackage(ConfigPath))
-	{
-		const FString ConfigContent = FPlatformMisc::LoadTextFileFromPlatformPackage(ConfigPath);
-		return TryParseConfigContent(ConfigContent, ConfigFileName, OutConfig);
-	}
-	else
-	{
-		return false;
-	}
-}
-
-////////////////////////////////////////////////////////////////////////////////
-bool TryParseEndpointConfig(const TCHAR* CommandLine, FOnDemandEndpointConfig& OutConfig)
-{
-	OutConfig = FOnDemandEndpointConfig();
-#if !UE_BUILD_SHIPPING
-	FString UrlParam;
-	if (FParse::Value(CommandLine, TEXT("Ias.TocUrl="), UrlParam))
-	{
-		FStringView UrlView(UrlParam);
-		if (UrlView.StartsWith(TEXTVIEW("http://")) && UrlView.EndsWith(TEXTVIEW(".iochunktoc")))
-		{
-			int32 Delim = INDEX_NONE;
-			if (UrlView.RightChop(7).FindChar(TEXT('/'), Delim))
-			{
-				OutConfig.ServiceUrls.Add(FString(UrlView.Left(7 +  Delim)));
-				OutConfig.TocPath = UrlView.RightChop(OutConfig.ServiceUrls[0].Len() + 1);
-			}
-		}
-	}
-	else
-#endif
-	{
-		if (TryParseConfigFileFromPlatformPackage(OutConfig))
-		{
-			TStringBuilder<256> TocFilePath;
-			FPathViews::Append(TocFilePath, TEXT("Cloud"), FPaths::GetBaseFilename(OutConfig.TocPath));
-			TocFilePath.Append(TEXT(".iochunktoc"));
-
-			if (FPlatformMisc::FileExistsInPlatformPackage(*TocFilePath))
-			{
-				OutConfig.TocFilePath = TocFilePath;
-			}
-		}
-	}
-
-	return OutConfig.IsValid();
-}
-
-////////////////////////////////////////////////////////////////////////////////
-static FIasCacheConfig GetIasCacheConfig(const TCHAR* CommandLine)
-{
-	FIasCacheConfig Ret;
-
-	// Fetch values from .ini files
-	auto GetConfigIntImpl = [CommandLine] (const TCHAR* ConfigKey, const TCHAR* ParamName, auto& Out)
-	{
-		int64 Value = -1;
-		if (FString Temp; GConfig->GetString(TEXT("Ias"), ConfigKey, Temp, GEngineIni))
-		{
-			Value = ParseSizeParam(Temp);
-		}
-#if !UE_BUILD_SHIPPING
-		if (int64 Override = ParseSizeParam(CommandLine, ParamName); Override >= 0)
-		{
-			Value = Override;
-		}
-#endif
-
-		if (Value >= 0)
-		{
-			Out = decltype(Out)(Value);
-		}
-
-		return true;
-	};
-
-#define GetConfigInt(Name, Dest) \
-	do { GetConfigIntImpl(TEXT("FileCache.") Name, TEXT("Ias.FileCache.") Name TEXT("="), Dest); } while (false)
-	GetConfigInt(TEXT("WritePeriodSeconds"),	Ret.WriteRate.Seconds);
-	GetConfigInt(TEXT("WriteOpsPerPeriod"),		Ret.WriteRate.Ops);
-	GetConfigInt(TEXT("WriteBytesPerPeriod"),	Ret.WriteRate.Allowance);
-	GetConfigInt(TEXT("DiskQuota"),				Ret.DiskQuota);
-	GetConfigInt(TEXT("MemoryQuota"),			Ret.MemoryQuota);
-	GetConfigInt(TEXT("JournalQuota"),			Ret.JournalQuota);
-	GetConfigInt(TEXT("JournalMagic"),			Ret.JournalMagic);
-	GetConfigInt(TEXT("DemandThreshold"),		Ret.Demand.Threshold);
-	GetConfigInt(TEXT("DemandBoost"),			Ret.Demand.Boost);
-	GetConfigInt(TEXT("DemandSuperBoost"),		Ret.Demand.SuperBoost);
-#undef GetConfigInt
-
-#if !UE_BUILD_SHIPPING
-	if (FParse::Param(CommandLine, TEXT("Ias.DropCache")))
-	{
-		Ret.DropCache = true;
-	}
-	if (FParse::Param(CommandLine, TEXT("Ias.NoCache")))
-	{
-		Ret.DiskQuota = 0;
-	}
-#endif
-
-	return Ret;
-}
-
 ////////////////////////////////////////////////////////////////////////////////
 FStringBuilderBase& operator<<(FStringBuilderBase& Sb, const FOnDemandInstallCacheUsage& CacheUsage)
 {
@@ -606,205 +365,19 @@ FStringBuilderBase& operator<<(FStringBuilderBase& Sb, const FOnDemandStreamingC
 }
 
 ////////////////////////////////////////////////////////////////////////////////
-void FIoStoreOnDemandModule::SetBulkOptionalEnabled(bool bInEnabled)
-{
-	if (HttpIoDispatcherBackend.IsValid())
-	{
-		HttpIoDispatcherBackend->SetBulkOptionalEnabled(bInEnabled);
-	}
-	else
-	{
-		UE_LOG(LogIas, Log, TEXT("Deferring call to FIoStoreOnDemandModule::SetBulkOptionalEnabled(%s)"), bInEnabled ? TEXT("true") : TEXT("false"));
-		DeferredBulkOptionalEnabled = bInEnabled;
-	}
-}
-
-void FIoStoreOnDemandModule::SetEnabled(bool bInEnabled)
-{
-	if (HttpIoDispatcherBackend.IsValid())
-	{
-		HttpIoDispatcherBackend->SetEnabled(bInEnabled);
-	}
-	else
-	{
-		UE_LOG(LogIas, Log, TEXT("Deferring call to FIoStoreOnDemandModule::SetEnabled(%s)"), bInEnabled ? TEXT("true") : TEXT("false"));
-		DeferredEnabled = bInEnabled;
-	}
-}
-
-void FIoStoreOnDemandModule::AbandonCache()
-{
-	if (HttpIoDispatcherBackend.IsValid())
-	{
-		HttpIoDispatcherBackend->AbandonCache();
-	}
-	else
-	{
-		UE_LOG(LogIas, Log, TEXT("Deferring call to FIoStoreOnDemandModule::AbandonCache"));
-		DeferredAbandonCache = true;
-	}
-}
-
-bool FIoStoreOnDemandModule::IsEnabled() const
-{
-	return HttpIoDispatcherBackend.IsValid()? HttpIoDispatcherBackend->IsEnabled():DeferredAbandonCache.IsSet();
-}
-
-void FIoStoreOnDemandModule::ReportAnalytics(TArray<FAnalyticsEventAttribute>& OutAnalyticsArray) const
-{
-	if (HttpIoDispatcherBackend.IsValid())
-	{
-		HttpIoDispatcherBackend->ReportAnalytics(OutAnalyticsArray);
-	}
-}
-
-TUniquePtr<IAnalyticsRecording> FIoStoreOnDemandModule::StartAnalyticsRecording() const
-{
-	if (HttpIoDispatcherBackend.IsValid())
-	{
-		return HttpIoDispatcherBackend->StartAnalyticsRecording();
-	}
-	else
-	{
-		return TUniquePtr<IAnalyticsRecording>();
-	}
-}
-
-void FIoStoreOnDemandModule::Mount(FOnDemandMountArgs&& Args, FOnDemandMountCompleted&& OnCompleted)
-{
-	if (IoStore.IsValid() == false)
-	{
-		IoStore = MakeShared<FOnDemandIoStore>();
-		if (FIoStatus Status = IoStore->Initialize(); !Status.IsOk())
-		{
-			UE_LOG(LogIas, Error, TEXT("Failed to initialize I/O store on-demand, reason '%s'"), *Status.ToString());
-			IoStore.Reset();
-			return OnCompleted(FOnDemandMountResult{ Args.MountId, Status });
-		}
-	}
-
-	IoStore->Mount(MoveTemp(Args), MoveTemp(OnCompleted));
-}
-
-void FIoStoreOnDemandModule::Install(
-	FOnDemandInstallArgs&& Args,
-	FOnDemandInstallCompleted&& OnCompleted,
-	FOnDemandInstallProgressed&& OnProgress /*= nullptr*/,
-	const FOnDemandCancellationToken* CancellationToken)
-{
-	if (IoStore.IsValid() == false)
-	{
-		IoStore = MakeShared<FOnDemandIoStore>();
-		if (FIoStatus Status = IoStore->Initialize(); !Status.IsOk())
-		{
-			UE_LOG(LogIas, Error, TEXT("Failed to initialize I/O store on-demand, reason '%s'"), *Status.ToString());
-			IoStore.Reset();
-			return OnCompleted(FOnDemandInstallResult{ .Status = Status });
-		}
-	}
-
-	IoStore->Install(MoveTemp(Args), MoveTemp(OnCompleted), MoveTemp(OnProgress), CancellationToken);
-}
-
-void FIoStoreOnDemandModule::Purge(FOnDemandPurgeArgs&& Args, FOnDemandPurgeCompleted&& OnCompleted)
-{
-	if (IoStore.IsValid() == false)
-	{
-		IoStore = MakeShared<FOnDemandIoStore>();
-		if (FIoStatus Status = IoStore->Initialize(); !Status.IsOk())
-		{
-			UE_LOG(LogIas, Error, TEXT("Failed to initialize I/O store on-demand, reason '%s'"), *Status.ToString());
-			IoStore.Reset();
-			return OnCompleted(FOnDemandPurgeResult{ .Status = Status });
-		}
-	}
-
-	IoStore->Purge(MoveTemp(Args), MoveTemp(OnCompleted));
-}
-
-FIoStatus FIoStoreOnDemandModule::Unmount(FStringView MountId)
-{
-	if (IoStore.IsValid())
-	{
-		return IoStore->Unmount(MountId);
-	}
-	return FIoStatus(EIoErrorCode::InvalidCode, NotInitializedError);
-}
-
-TIoStatusOr<uint64> FIoStoreOnDemandModule::GetInstallSize(const FOnDemandGetInstallSizeArgs& Args) const
-{
-	if (IoStore)
-	{
-		return IoStore->GetInstallSize(Args);
-	}
-
-	return FIoStatus(EIoErrorCode::InvalidCode, NotInitializedError);
-}
-
-FIoStatus FIoStoreOnDemandModule::GetInstallSizesByMountId(const FOnDemandGetInstallSizeArgs& Args, TMap<FString, uint64>& OutSizesByMountId) const
-{
-	if (IoStore)
-	{
-		return IoStore->GetInstallSizesByMountId(Args, OutSizesByMountId);
-	}
-
-	return FIoStatus(EIoErrorCode::InvalidCode, NotInitializedError);
-}
-
-TIoStatusOr<FOnDemandCacheUsage> FIoStoreOnDemandModule::GetCacheUsage() const
-{
-	if (IoStore.IsValid() == false)
-	{
-		return FIoStatus(EIoErrorCode::InvalidCode, NotInitializedError);
-	}
-
-	FOnDemandCacheUsage CacheUsage = IoStore->GetCacheUsage();
-	if (HttpIoDispatcherBackend.IsValid())
-	{
-		CacheUsage.StreamingCache = HttpIoDispatcherBackend->GetCacheUsage();
-	}
-
-	return CacheUsage;
-}
-
-void FIoStoreOnDemandModule::InitializeInternal()
+FIoStatus FIoStoreOnDemandModule::InitializeStreamingBackend()
 {
 	LLM_SCOPE_BYTAG(Ias);
 
-#if WITH_EDITOR
-	bool bEnabledInEditor = false;
-	GConfig->GetBool(TEXT("Ias"), TEXT("EnableInEditor"), bEnabledInEditor, GEngineIni);
-
-	if (!bEnabledInEditor)
+	if (GOnDemandIoStore.IsValid() == false)
 	{
-		return;
+		return FIoStatus(EIoErrorCode::InvalidCode, TEXT("I/O store on-demand module has not been initialized"));
 	}
-#endif //WITH_EDITOR
 
 	const TCHAR* CommandLine = FCommandLine::Get();
-	
-#if !UE_BUILD_SHIPPING
-	if (FParse::Param(CommandLine, TEXT("NoIas")))
-	{
-		return;
-	}
-#endif
-
-	if (IoStore.IsValid() == false)
-	{
-		IoStore = MakeShared<FOnDemandIoStore>();
-		if (FIoStatus Status = IoStore->Initialize(); !Status.IsOk())
-		{
-			UE_LOG(LogIas, Warning, TEXT("Failed to initialize I/O store on demand, reason '%s'"), *Status.ToString());
-			return;
-		}
-	}
-
-	// Make sure we haven't called initialize before
-	check(!HttpIoDispatcherBackend.IsValid());
 
-	FOnDemandEndpointConfig EndpointConfig;
-	if (TryParseEndpointConfig(CommandLine, EndpointConfig) == false)
+	TIoStatusOr<FOnDemandEndpointConfig> EndpointConfigStatus = Config::TryParseEndpointConfig(CommandLine);
+	if (EndpointConfigStatus.IsOk() == false)
 	{
 		// Temporary solution: Delete the ias cache directory if the game hasn't been configured to use IAS
 		if (FPaths::HasProjectPersistentDownloadDir())
@@ -819,72 +392,29 @@ void FIoStoreOnDemandModule::InitializeInternal()
 				Ifm.DeleteDirectory(CachePath.ToString(), false, true);
 			}
 		}
-		return;
-	}
 
-	{
-		FString EncryptionKey;
-		if (FParse::Value(CommandLine, TEXT("Ias.EncryptionKey="), EncryptionKey))
-		{
-			ApplyEncryptionKeyFromString(EncryptionKey);
-		}
+		return EndpointConfigStatus.Status();
 	}
 
-	TUniquePtr<IIasCache> Cache;
-	FIasCacheConfig CacheConfig = GetIasCacheConfig(CommandLine);
-
-	CacheConfig.DropCache = DeferredAbandonCache.Get(CacheConfig.DropCache);
-	if (CacheConfig.DiskQuota > 0)
-	{
-		if (FPaths::HasProjectPersistentDownloadDir())
-		{
-			FString CacheDir = FPaths::ProjectPersistentDownloadDir();
-			Cache = MakeIasCache(*CacheDir, CacheConfig);
-
-			UE_CLOG(!Cache.IsValid(), LogIas, Warning, TEXT("File cache disabled - streaming only (init-fail)"));
-		}
-		else
-		{
-			UE_LOG(LogIas, Warning, TEXT("File cache disabled - streaming only (project has no persistent download dir enabled for this platform)"));
-		}
-	}
-	else
-	{
-		UE_LOG(LogIas, Log, TEXT("File cache disabled - streaming only (zero-quota)"));
-	}
+	FOnDemandEndpointConfig EndpointConfig = EndpointConfigStatus.ConsumeValueOrDie();
 
 	bool bFirstAttempt = false;
-	if (FIoStatus Status = LoadDefaultHttpCertificates(bFirstAttempt); !Status.IsOk())
+	if (FIoStatus CertStatus = LoadDefaultHttpCertificates(bFirstAttempt); !CertStatus.IsOk())
 	{
 		if (EndpointConfig.RequiresTls())
 		{
-			UE_LOG(LogIas, Error, TEXT("Endpoint configuration requires TLS but certificates was not loaded, reason '%s'"),
-				*Status.ToString());
-			return;
+			FIoStatus Status = FIoStatusBuilder(EIoErrorCode::ReadError)
+				<< TEXT("Endpoint requires TLS but certificates failed to load due to '")
+				<< CertStatus.ToString()
+				<< TEXT("'");
+			return Status;
 		}
 	}
-	HttpIoDispatcherBackend = MakeOnDemandIoDispatcherBackend(EndpointConfig, *IoStore, MoveTemp(Cache));
 
-	int32 BackendPriority = -10;
-#if !UE_BUILD_SHIPPING
-	if (FParse::Param(CommandLine, TEXT("Ias")))
+	if (FIoStatus Status = GOnDemandIoStore->InitializeStreamingBackend(EndpointConfig); Status.IsOk() == false)
 	{
-		// Bump the priority to be higher then the file system backend
-		BackendPriority = 10;
+		return Status;
 	}
-#endif
-
-	// Setup any states changes issued before initialization
-	if (DeferredEnabled.IsSet())
-	{
-		HttpIoDispatcherBackend->SetEnabled(*DeferredEnabled);
-	}
-	if (DeferredBulkOptionalEnabled.IsSet())
-	{
-		HttpIoDispatcherBackend->SetBulkOptionalEnabled(*DeferredBulkOptionalEnabled);
-	}
-	
-	FIoDispatcher::Get().Mount(HttpIoDispatcherBackend.ToSharedRef(), BackendPriority);
 
 	bool bUsePerContainerTocsConfigValue = false;
 	if (GConfig)
@@ -915,20 +445,17 @@ void FIoStoreOnDemandModule::InitializeInternal()
 	else if (!EndpointConfig.ServiceUrls.IsEmpty() && !EndpointConfig.TocPath.IsEmpty())
 	{
 		TIoStatusOr<FOnDemandHostGroup> HostGroup = FOnDemandHostGroup::Create(EndpointConfig.ServiceUrls);
-		if (HostGroup.IsOk())
+		if (HostGroup.IsOk() == false)
 		{
-			MountArgs.Emplace(FOnDemandMountArgs
-			{
-				.MountId = EndpointConfig.ServiceUrls[0] / EndpointConfig.TocPath,
-				.TocRelativeUrl = EndpointConfig.TocPath,
-				.HostGroup = HostGroup.ConsumeValueOrDie(),
-				.Options = EOnDemandMountOptions::StreamOnDemand
-			});
+			return HostGroup.Status();
 		}
-		else
+		MountArgs.Emplace(FOnDemandMountArgs
 		{
-			UE_LOG(LogIas, Error, TEXT("Failed to create host group, reason: '%s'"), *HostGroup.Status().ToString());
-		}
+			.MountId = EndpointConfig.ServiceUrls[0] / EndpointConfig.TocPath,
+			.TocRelativeUrl = EndpointConfig.TocPath,
+			.HostGroup = HostGroup.ConsumeValueOrDie(),
+			.Options = EOnDemandMountOptions::StreamOnDemand
+		});
 	}
 
 #if !UE_BUILD_SHIPPING
@@ -954,7 +481,7 @@ void FIoStoreOnDemandModule::InitializeInternal()
 
 	if (MountArgs)
 	{
-		IoStore->Mount(
+		GOnDemandIoStore->Mount(
 			MoveTemp(MountArgs.GetValue()),
 			[](FOnDemandMountResult MountResult)
 			{
@@ -964,7 +491,7 @@ void FIoStoreOnDemandModule::InitializeInternal()
 #if !UE_BUILD_SHIPPING
 		if (InstallArgs)
 		{
-			IoStore->Install(
+			GOnDemandIoStore->Install(
 				MoveTemp(InstallArgs.GetValue()),
 				[](FOnDemandInstallResult InstallResult)
 				{
@@ -974,8 +501,10 @@ void FIoStoreOnDemandModule::InitializeInternal()
 		}
 #endif
 	}
+
+	return FIoStatus::Ok;
 }
-	
+
 void FIoStoreOnDemandModule::StartupModule()
 {
 #if PLATFORM_WINDOWS
@@ -991,18 +520,42 @@ void FIoStoreOnDemandModule::StartupModule()
 			TCHAR SystemErrorMsg[MAX_SPRINTF] = { 0 };
 			FPlatformMisc::GetSystemErrorMessage(SystemErrorMsg, sizeof(SystemErrorMsg), Result);
 
-			UE_LOG(LogIas, Error, TEXT("WSAStartup failed due to: %s (%d)"), SystemErrorMsg, Result);
+			UE_LOG(LogIoStoreOnDemand, Error, TEXT("WSAStartup failed due to: %s (%d)"), SystemErrorMsg, Result);
 		}
 	}
 #endif //PLATFORM_WINDOWS
 
+	GOnDemandIoStore = MakeShared<FOnDemandIoStore>();
+	if (FIoStatus Status = GOnDemandIoStore->Initialize(); !Status.IsOk())
+	{
+		UE_LOG(LogIoStoreOnDemand, Error, TEXT("Failed to initialize I/O store on demand, reason '%s'"), *Status.ToString());
+		GOnDemandIoStore.Reset();
+		return;
+	}
+
 #if !UE_IAS_CUSTOM_INITIALIZATION
-	InitializeInternal();
+	if (FIoStatus Status = InitializeStreamingBackend(); !Status.IsOk())
+	{
+		if (Status.GetErrorCode() == EIoErrorCode::Disabled || Status.GetErrorCode() == EIoErrorCode::NotFound)
+		{
+			UE_LOG(LogIoStoreOnDemand, Log, TEXT("Individual asset streaming disabled, reason '%s'"), *Status.ToString());
+		}
+		else 
+		{
+			UE_LOG(LogIoStoreOnDemand, Error, TEXT("Failed to initialize individual asset streaming, reason '%s'"), *Status.ToString());
+		}
+		return;
+	}
+	else
+	{
+		UE_LOG(LogIoStoreOnDemand, Log, TEXT("Individual asset streaming enabled"));
+	}
 #endif // !UE_IAS_CUSTOM_INITIALIZATION
 }
 
 void FIoStoreOnDemandModule::ShutdownModule()
 {
+	GOnDemandIoStore.Reset();
 #if PLATFORM_WINDOWS
 	if (bPlatformSpecificSetup)
 	{
@@ -1013,7 +566,7 @@ void FIoStoreOnDemandModule::ShutdownModule()
 			TCHAR SystemErrorMsg[MAX_SPRINTF] = { 0 };
 			FPlatformMisc::GetSystemErrorMessage(SystemErrorMsg, sizeof(SystemErrorMsg), SystemError);
 
-			UE_LOG(LogIas, Error, TEXT("WSACleanup failed due to: %s (%u)"), SystemErrorMsg, SystemError);
+			UE_LOG(LogIoStoreOnDemand, Error, TEXT("WSACleanup failed due to: %s (%u)"), SystemErrorMsg, SystemError);
 		}
 
 		bPlatformSpecificSetup = false;
@@ -1022,18 +575,26 @@ void FIoStoreOnDemandModule::ShutdownModule()
 }
 
 #if UE_IAS_CUSTOM_INITIALIZATION
-
-EOnDemandInitResult FIoStoreOnDemandModule::Initialize()
+FIoStatus FIoStoreOnDemandModule::InitializePostHotfix()
 {
-	InitializeInternal();
+	return InitializeStreamingBackend();
+}
+#endif 
 
-	return HttpIoDispatcherBackend.IsValid() ? EOnDemandInitResult::Success : EOnDemandInitResult::Disabled;
-};
+////////////////////////////////////////////////////////////////////////////////
+IOnDemandIoStore* TryGetOnDemandIoStore()
+{
+	return GOnDemandIoStore.Get();
+}
 
-#endif // UE_IAS_CUSTOM_INITIALIZATION
+////////////////////////////////////////////////////////////////////////////////
+IOnDemandIoStore& GetOnDemandIoStore()
+{
+	check(GOnDemandIoStore.IsValid());
+	return *GOnDemandIoStore.Get();
+}
 
 } // namespace UE::IoStore
 
 ////////////////////////////////////////////////////////////////////////////////
-
 IMPLEMENT_MODULE(UE::IoStore::FIoStoreOnDemandModule, IoStoreOnDemand);
diff --git a/Engine/Source/Runtime/Experimental/IoStore/OnDemand/Private/OnDemandConfig.cpp b/Engine/Source/Runtime/Experimental/IoStore/OnDemand/Private/OnDemandConfig.cpp
new file mode 100644
--- /dev/null
+++ b/Engine/Source/Runtime/Experimental/IoStore/OnDemand/Private/OnDemandConfig.cpp
@@ -0,0 +1,404 @@
+// Copyright Epic Games, Inc. All Rights Reserved.
+
+#include "OnDemandConfig.h"
+#include "IO/IoStatus.h"
+#include "OnDemandIoDispatcherBackend.h"
+#include "OnDemandInstallCache.h"
+
+#include "HAL/PlatformProcess.h"
+#include "IO/IasCache.h"
+#include "IO/IoStoreOnDemand.h"
+#include "Misc/Base64.h"
+#include "Misc/ConfigCacheIni.h"
+#include "Misc/CoreDelegates.h"
+#include "Misc/CoreMisc.h"
+#include "Misc/EncryptionKeyManager.h"
+#include "Misc/Fork.h"
+#include "Misc/Parse.h"
+#include "Misc/Paths.h"
+#include "Misc/PathViews.h"
+#include "String/LexFromString.h"
+
+namespace UE::IoStore::Config
+{
+
+////////////////////////////////////////////////////////////////////////////////
+bool GIoStoreOnDemandInstallCacheEnabled = true;
+static FAutoConsoleVariableRef CVar_IoStoreOnDemandInstallCacheEnabled(
+	TEXT("iostore.OnDemandInstallCacheEnabled"),
+	GIoStoreOnDemandInstallCacheEnabled,
+	TEXT("Whether the on-demand install cache is enabled."),
+	ECVF_ReadOnly
+);
+
+/** Temp cvar to allow the fallback url to be hotfixed in case of problems */
+static FString GDistributedEndpointFallbackUrl;
+static FAutoConsoleVariableRef CVar_DistributedEndpointFallbackUrl(
+	TEXT("ias.DistributedEndpointFallbackUrl"),
+	GDistributedEndpointFallbackUrl,
+	TEXT("CDN url to be used if a distributed endpoint cannot be reached (overrides IoStoreOnDemand.ini)")
+);
+
+////////////////////////////////////////////////////////////////////////////////
+static bool ParseEncryptionKeyParam(const FString& Param, FGuid& OutKeyGuid, FAES::FAESKey& OutKey)
+{
+	TArray<FString> Tokens;
+	Param.ParseIntoArray(Tokens, TEXT(":"), true);
+
+	if (Tokens.Num() == 2)
+	{
+		TArray<uint8> KeyBytes;
+		if (FGuid::Parse(Tokens[0], OutKeyGuid) && FBase64::Decode(Tokens[1], KeyBytes))
+		{
+			if (OutKeyGuid != FGuid() && KeyBytes.Num() == FAES::FAESKey::KeySize)
+			{
+				FMemory::Memcpy(OutKey.Key, KeyBytes.GetData(), FAES::FAESKey::KeySize);
+				return true;
+			}
+		}
+	}
+	
+	return false;
+}
+
+////////////////////////////////////////////////////////////////////////////////
+static bool ApplyEncryptionKeyFromString(const FString& GuidKeyPair)
+{
+	FGuid KeyGuid;
+	FAES::FAESKey Key;
+
+	if (ParseEncryptionKeyParam(GuidKeyPair, KeyGuid, Key))
+	{
+		// TODO: PAK and I/O store should share key manager
+		FEncryptionKeyManager::Get().AddKey(KeyGuid, Key);
+		FCoreDelegates::GetRegisterEncryptionKeyMulticastDelegate().Broadcast(KeyGuid, Key);
+
+		return true;
+	}
+	else
+	{
+		return false;
+	}
+}
+
+////////////////////////////////////////////////////////////////////////////////
+static bool TryParseConfigContent(const FString& ConfigContent, const FString& ConfigFileName, FOnDemandEndpointConfig& OutEndpoint)
+{
+	if (ConfigContent.IsEmpty())
+	{
+		return false;
+	}
+
+	FConfigFile ConfigFile;
+	ConfigFile.ProcessInputFileContents(ConfigContent, ConfigFileName);
+
+	ConfigFile.GetString(TEXT("Endpoint"), TEXT("DistributionUrl"), OutEndpoint.DistributionUrl);
+	if (!OutEndpoint.DistributionUrl.IsEmpty())
+	{
+		ConfigFile.GetString(TEXT("Endpoint"), TEXT("FallbackUrl"), OutEndpoint.FallbackUrl);
+
+		if (!GDistributedEndpointFallbackUrl.IsEmpty())
+		{
+			OutEndpoint.FallbackUrl = GDistributedEndpointFallbackUrl;
+		}
+	}
+	
+	ConfigFile.GetArray(TEXT("Endpoint"), TEXT("ServiceUrl"), OutEndpoint.ServiceUrls);
+	ConfigFile.GetString(TEXT("Endpoint"), TEXT("TocPath"), OutEndpoint.TocPath);
+
+	if (OutEndpoint.DistributionUrl.EndsWith(TEXT("/")))
+	{
+		OutEndpoint.DistributionUrl = OutEndpoint.DistributionUrl.Left(OutEndpoint.DistributionUrl.Len() - 1);
+	}
+
+	for (FString& ServiceUrl : OutEndpoint.ServiceUrls)
+	{
+		if (ServiceUrl.EndsWith(TEXT("/")))
+		{
+			ServiceUrl.LeftInline(ServiceUrl.Len() - 1);
+		}
+	}
+
+	if (OutEndpoint.TocPath.StartsWith(TEXT("/")))
+	{
+		OutEndpoint.TocPath.RightChopInline(1);
+	}
+
+	FString ContentKey;
+	if (ConfigFile.GetString(TEXT("Endpoint"), TEXT("ContentKey"), ContentKey))
+	{
+		ApplyEncryptionKeyFromString(ContentKey);
+	}
+
+	return OutEndpoint.IsValid();
+}
+
+////////////////////////////////////////////////////////////////////////////////
+int64 ParseSizeParam(FStringView Value)
+{
+	Value = Value.TrimStartAndEnd();
+
+	int64 Size = -1;
+	LexFromString(Size, Value);
+	if (Size >= 0)
+	{
+		if (Value.EndsWith(TEXT("GB"))) return Size << 30;
+		if (Value.EndsWith(TEXT("MB"))) return Size << 20;
+		if (Value.EndsWith(TEXT("KB"))) return Size << 10;
+	}
+	return Size;
+}
+
+////////////////////////////////////////////////////////////////////////////////
+int64 ParseSizeParam(const TCHAR* CommandLine, const TCHAR* Param)
+{
+	FString ParamValue;
+	if (!FParse::Value(CommandLine, Param, ParamValue))
+	{
+		return -1;
+	}
+
+	return ParseSizeParam(ParamValue);
+}
+
+////////////////////////////////////////////////////////////////////////////////
+FIasCacheConfig GetStreamingCacheConfig(const TCHAR* CommandLine)
+{
+	FIasCacheConfig Ret;
+
+	// Fetch values from .ini files
+	auto GetConfigIntImpl = [CommandLine] (const TCHAR* ConfigKey, const TCHAR* ParamName, auto& Out)
+	{
+		int64 Value = -1;
+		if (FString Temp; GConfig->GetString(TEXT("Ias"), ConfigKey, Temp, GEngineIni))
+		{
+			Value = ParseSizeParam(Temp);
+		}
+#if !UE_BUILD_SHIPPING
+		if (int64 Override = ParseSizeParam(CommandLine, ParamName); Override >= 0)
+		{
+			Value = Override;
+		}
+#endif
+
+		if (Value >= 0)
+		{
+			Out = decltype(Out)(Value);
+		}
+
+		return true;
+	};
+
+#define GetConfigInt(Name, Dest) \
+	do { GetConfigIntImpl(TEXT("FileCache.") Name, TEXT("Ias.FileCache.") Name TEXT("="), Dest); } while (false)
+	GetConfigInt(TEXT("WritePeriodSeconds"),	Ret.WriteRate.Seconds);
+	GetConfigInt(TEXT("WriteOpsPerPeriod"),		Ret.WriteRate.Ops);
+	GetConfigInt(TEXT("WriteBytesPerPeriod"),	Ret.WriteRate.Allowance);
+	GetConfigInt(TEXT("DiskQuota"),				Ret.DiskQuota);
+	GetConfigInt(TEXT("MemoryQuota"),			Ret.MemoryQuota);
+	GetConfigInt(TEXT("JournalQuota"),			Ret.JournalQuota);
+	GetConfigInt(TEXT("JournalMagic"),			Ret.JournalMagic);
+	GetConfigInt(TEXT("DemandThreshold"),		Ret.Demand.Threshold);
+	GetConfigInt(TEXT("DemandBoost"),			Ret.Demand.Boost);
+	GetConfigInt(TEXT("DemandSuperBoost"),		Ret.Demand.SuperBoost);
+#undef GetConfigInt
+
+#if !UE_BUILD_SHIPPING
+	if (FParse::Param(CommandLine, TEXT("Ias.DropCache")))
+	{
+		Ret.DropCache = true;
+	}
+	if (FParse::Param(CommandLine, TEXT("Ias.NoCache")))
+	{
+		Ret.DiskQuota = 0;
+	}
+#endif
+
+	return Ret;
+}
+
+////////////////////////////////////////////////////////////////////////////////
+TIoStatusOr<FOnDemandEndpointConfig> TryParseEndpointConfig(const TCHAR* CommandLine)
+{
+	{
+		FString EncryptionKey;
+		if (FParse::Value(CommandLine, TEXT("Ias.EncryptionKey="), EncryptionKey))
+		{
+			ApplyEncryptionKeyFromString(EncryptionKey);
+		}
+	}
+
+#if WITH_EDITOR
+	bool bEnabledInEditor = false;
+	GConfig->GetBool(TEXT("Ias"), TEXT("EnableInEditor"), bEnabledInEditor, GEngineIni);
+
+	if (!bEnabledInEditor)
+	{
+		FIoStatus Status = FIoStatusBuilder(EIoErrorCode::Disabled) << TEXT("Disabled in Editor");
+		return Status;
+	}
+#endif //WITH_EDITOR
+
+#if !UE_BUILD_SHIPPING
+	if (FParse::Param(CommandLine, TEXT("NoIas")))
+	{
+		FIoStatus Status = FIoStatusBuilder(EIoErrorCode::Disabled) << TEXT("Disabled by '-NoIas'");
+		return Status;
+	}
+#endif
+
+	FOnDemandEndpointConfig OutConfig;
+#if !UE_BUILD_SHIPPING
+	FString UrlParam;
+	if (FParse::Value(CommandLine, TEXT("Ias.TocUrl="), UrlParam))
+	{
+		FStringView UrlView(UrlParam);
+		if (UrlView.StartsWith(TEXTVIEW("http://")) && UrlView.EndsWith(TEXTVIEW(".iochunktoc")))
+		{
+			int32 Delim = INDEX_NONE;
+			if (UrlView.RightChop(7).FindChar(TEXT('/'), Delim))
+			{
+				OutConfig.ServiceUrls.Add(FString(UrlView.Left(7 +  Delim)));
+				OutConfig.TocPath = UrlView.RightChop(OutConfig.ServiceUrls[0].Len() + 1);
+			}
+		}
+
+		if (OutConfig.IsValid() == false)
+		{
+			FIoStatus Status = FIoStatusBuilder(EIoErrorCode::InvalidParameter) << TEXT("Failed to parse '-Ias.TocUrl='");
+			return Status;
+		}
+
+		return OutConfig;
+	}
+#endif
+
+	const FString ConfigFileName = TEXT("IoStoreOnDemand.ini");
+	const FString ConfigPath = FPaths::Combine(TEXT("Cloud"), ConfigFileName);
+
+	if (FPlatformMisc::FileExistsInPlatformPackage(ConfigPath) == false)
+	{
+		FIoStatus Status = FIoStatusBuilder(EIoErrorCode::NotFound)
+			<< TEXT("Failed to find config file '")
+			<< ConfigPath
+			<< TEXT("'");
+		return Status;
+	}
+
+	const FString ConfigContent = FPlatformMisc::LoadTextFileFromPlatformPackage(ConfigPath);
+	if (TryParseConfigContent(ConfigContent, ConfigFileName, OutConfig) == false || OutConfig.IsValid() == false)
+	{
+		FIoStatus Status = FIoStatusBuilder(EIoErrorCode::ReadError)
+			<< TEXT("Failed to read config file '")
+			<< ConfigPath
+			<< TEXT("'");
+		return Status;
+	}
+
+	return OutConfig;
+}
+
+///////////////////////////////////////////////////////////////////////////////
+static FString GetInstallCacheDirectory(const TCHAR* CommandLine)
+{
+	FString DirName;
+
+	if (IsRunningDedicatedServer())
+	{
+		if (!FForkProcessHelper::IsForkRequested())
+		{
+			DirName = TEXT("InstallCacheServer");
+		}
+		else
+		{
+			if (!FForkProcessHelper::IsForkedChildProcess())
+			{
+				UE_LOG(LogIoStoreOnDemand, Fatal, TEXT("Attempting to create install cache before forking!"));
+			}
+
+			FString CommandLineDir;
+			bool bUsePathFromCommandLine = FParse::Value(CommandLine, TEXT("ServerIOInstallCacheDir="), CommandLineDir);
+			if (bUsePathFromCommandLine)
+			{
+				if (!FPaths::ValidatePath(CommandLineDir))
+				{
+					bUsePathFromCommandLine = false;
+					UE_LOG(LogIoStoreOnDemand, Error, TEXT("Invalid ServerIOInstallCacheDir from command line: %s"), *CommandLineDir);
+				}
+				else if (!FPaths::IsRelative(CommandLineDir))
+				{
+					bUsePathFromCommandLine = false;
+					UE_LOG(LogIoStoreOnDemand, Error, TEXT("ServerIOInstallCacheDir from command line is not relative: %s"), *CommandLineDir);
+				}
+
+				if (bUsePathFromCommandLine)
+				{
+					return FPaths::ProjectPersistentDownloadDir() / CommandLineDir;
+				}
+			}
+
+			DirName = FString::Printf(TEXT("InstallCacheServer-%u"), FPlatformProcess::GetCurrentProcessId());
+		}
+	}
+#if WITH_EDITOR
+	else if (GIsEditor)
+	{
+		DirName = TEXT("InstallCacheEditor");
+	}
+#endif //if WITH_EDITOR
+	else
+	{
+		DirName = TEXT("InstallCache");
+	}
+
+	return FPaths::ProjectPersistentDownloadDir() / TEXT("IoStore") / DirName;
+}
+
+///////////////////////////////////////////////////////////////////////////////
+TIoStatusOr<FOnDemandInstallCacheConfig> TryParseInstallCacheConfig(const TCHAR* CommandLine)
+{
+	bool bUseInstallCache = GIoStoreOnDemandInstallCacheEnabled;
+#if !UE_BUILD_SHIPPING
+	bUseInstallCache = FParse::Param(CommandLine, TEXT("NoIAD")) == false;
+#endif
+	if (bUseInstallCache == false)
+	{
+		FIoStatus Status = FIoStatusBuilder(EIoErrorCode::InvalidParameter) << TEXT("Disabled");
+		return Status;
+	}
+
+	if (FPaths::HasProjectPersistentDownloadDir() == false)
+	{
+		FIoStatus Status = FIoStatusBuilder(EIoErrorCode::InvalidParameter) << TEXT("Persistent storage not configured");
+		return Status;
+	}
+
+	FOnDemandInstallCacheConfig OutConfig;
+
+	const int64 DiskQuota = [CommandLine]
+	{
+		if (FString ParamValue; FParse::Value(CommandLine, TEXT("-Iad.FileCache.DiskQuota="), ParamValue))
+		{
+			return Config::ParseSizeParam(ParamValue);
+		}
+		else if (FString ValueStr; GConfig->GetString(TEXT("OnDemandInstall"), TEXT("FileCache.DiskQuota"), ValueStr, GEngineIni))
+		{
+			return Config::ParseSizeParam(ValueStr);
+		}
+		return int64(-1);
+	}();
+
+	if (DiskQuota > 0)
+	{
+		OutConfig.DiskQuota = DiskQuota;
+	}
+
+	OutConfig.RootDirectory = Config::GetInstallCacheDirectory(CommandLine);
+#if !UE_BUILD_SHIPPING
+	OutConfig.bDropCache = FParse::Param(CommandLine, TEXT("Iad.DropCache"));
+#endif
+
+	return OutConfig;
+}
+
+} // namespace UE::IoStore
diff --git a/Engine/Source/Runtime/Experimental/IoStore/OnDemand/Private/OnDemandConfig.h b/Engine/Source/Runtime/Experimental/IoStore/OnDemand/Private/OnDemandConfig.h
new file mode 100644
--- /dev/null
+++ b/Engine/Source/Runtime/Experimental/IoStore/OnDemand/Private/OnDemandConfig.h
@@ -0,0 +1,26 @@
+// Copyright Epic Games, Inc. All Rights Reserved.
+
+#pragma once
+
+#include "CoreTypes.h"
+#include "Containers/StringView.h"
+#include "IO/IoStatus.h"
+
+namespace UE::IoStore { struct FOnDemandEndpointConfig; }
+namespace UE::IoStore { struct FIasCacheConfig; }
+namespace UE::IoStore { struct FOnDemandInstallCacheConfig; }
+
+namespace UE::IoStore::Config
+{
+
+int64 ParseSizeParam(FStringView Value);
+
+int64 ParseSizeParam(const TCHAR* CommandLine, const TCHAR* Param);
+
+TIoStatusOr<FOnDemandEndpointConfig> TryParseEndpointConfig(const TCHAR* CommandLine);
+
+FIasCacheConfig GetStreamingCacheConfig(const TCHAR* CommandLine);
+
+TIoStatusOr<FOnDemandInstallCacheConfig> TryParseInstallCacheConfig(const TCHAR* CommandLine);
+
+} // namespace UE::IoStore
diff --git a/Engine/Source/Runtime/Experimental/IoStore/OnDemand/Private/OnDemandIoDispatcherBackend.cpp b/Engine/Source/Runtime/Experimental/IoStore/OnDemand/Private/OnDemandIoDispatcherBackend.cpp
--- a/Engine/Source/Runtime/Experimental/IoStore/OnDemand/Private/OnDemandIoDispatcherBackend.cpp
+++ b/Engine/Source/Runtime/Experimental/IoStore/OnDemand/Private/OnDemandIoDispatcherBackend.cpp
@@ -251,19 +251,7 @@ FAutoConsoleVariableRef CVarOnDemandBackendThreadPriority(
 	TEXT("Note that this is switchable at runtime"),
 	ECVF_Default);
 
-#if !UE_BUILD_SHIPPING
-static FAutoConsoleCommand CVar_IasAbandonCache(
-	TEXT("Ias.AbandonCache"),
-	TEXT("Abandon the local file cache"),
-	FConsoleCommandDelegate::CreateLambda([]()
-	{
-		FIoStoreOnDemandModule& Module = FModuleManager::Get().GetModuleChecked<FIoStoreOnDemandModule>("IoStoreOnDemand");
-		Module.AbandonCache();
-	})
-);
-#endif //!UE_BUILD_SHIPPING
 ///////////////////////////////////////////////////////////////////////////////
-
 static bool LatencyTest(FStringView Url, FStringView Path)
 {
 	TRACE_CPUPROFILER_EVENT_SCOPE(IasBackend::LatencyTest);
@@ -1878,4 +1866,4 @@ struct FTests
 
 } // namespace UE::IoStore
 
-#undef UE_IAS_DEBUG_CONSOLE_CMDS
\ No newline at end of file
+#undef UE_IAS_DEBUG_CONSOLE_CMDS
diff --git a/Engine/Source/Runtime/Experimental/IoStore/OnDemand/Private/OnDemandIoStore.cpp b/Engine/Source/Runtime/Experimental/IoStore/OnDemand/Private/OnDemandIoStore.cpp
--- a/Engine/Source/Runtime/Experimental/IoStore/OnDemand/Private/OnDemandIoStore.cpp
+++ b/Engine/Source/Runtime/Experimental/IoStore/OnDemand/Private/OnDemandIoStore.cpp
@@ -1,13 +1,16 @@
 // Copyright Epic Games, Inc. All Rights Reserved.
 
 #include "OnDemandIoStore.h"
+#include "IO/IoStatus.h"
+#include "OnDemandConfig.h"
 #include "OnDemandHttpClient.h"
 #include "OnDemandInstallCache.h"
 #include "OnDemandPackageStoreBackend.h"
+#include "OnDemandIoDispatcherBackend.h"
+#include "IO/IasCache.h"
 
 #include "Algo/Accumulate.h"
 #include "Algo/AnyOf.h"
-#include "Algo/Copy.h"
 #include "Algo/RemoveIf.h"
 #include "Algo/Transform.h"
 #include "Async/ManualResetEvent.h"
@@ -30,7 +33,6 @@
 #include "Misc/ConfigCacheIni.h"
 #include "Misc/CoreDelegates.h"
 #include "Misc/CoreDelegatesInternal.h"
-#include "Misc/CoreMisc.h"
 #include "Misc/EncryptionKeyManager.h"
 #include "Misc/Guid.h"
 #include "Misc/Paths.h"
@@ -42,14 +44,6 @@
 #endif
 
 ///////////////////////////////////////////////////////////////////////////////
-bool GIoStoreOnDemandInstallCacheEnabled = true;
-static FAutoConsoleVariableRef CVar_IoStoreOnDemandInstallCacheEnabled(
-	TEXT("iostore.OnDemandInstallCacheEnabled"),
-	GIoStoreOnDemandInstallCacheEnabled,
-	TEXT("Whether the on-demand install cache is enabled."),
-	ECVF_ReadOnly
-);
-
 bool GIoStoreOnDemandTreatMissingChunksAsError = true;
 static FAutoConsoleVariableRef CVar_IoStoreOnDemandTreatMissingChunksAsError(
 	TEXT("iostore.TreatMissingOnDemandChunksAsError"),
@@ -68,67 +62,6 @@ extern FString GIasOnDemandTocExt;
 namespace Private
 {
 
-///////////////////////////////////////////////////////////////////////////////
-static FString GetInstallCacheDirectory()
-{
-	FString DirName;
-
-	if (IsRunningDedicatedServer())
-	{
-		if (!FForkProcessHelper::IsForkRequested())
-		{
-			DirName = TEXT("InstallCacheServer");
-		}
-		else
-		{
-			if (!FForkProcessHelper::IsForkedChildProcess())
-			{
-				UE_LOG(LogIoStoreOnDemand, Fatal, TEXT("Attempting to create IOStore cache before forking!"));
-			}
-
-			FString CommandLineDir;
-			bool bUsePathFromCommandLine = FParse::Value(FCommandLine::Get(), TEXT("ServerIOInstallCacheDir="), CommandLineDir);
-			if (bUsePathFromCommandLine)
-			{
-				if (!FPaths::ValidatePath(CommandLineDir))
-				{
-					bUsePathFromCommandLine = false;
-					UE_LOG(LogIoStoreOnDemand, Error, TEXT("Invalid ServerIOInstallCacheDir from command line: %s"), *CommandLineDir);
-				}
-				else if (!FPaths::IsRelative(CommandLineDir))
-				{
-					bUsePathFromCommandLine = false;
-					UE_LOG(LogIoStoreOnDemand, Error, TEXT("ServerIOInstallCacheDir from command line is not relative: %s"), *CommandLineDir);
-				}
-
-				if (bUsePathFromCommandLine)
-				{
-					return FPaths::ProjectPersistentDownloadDir() / CommandLineDir;
-				}
-			}
-
-			DirName = FString::Printf(TEXT("InstallCacheServer-%u"), FPlatformProcess::GetCurrentProcessId());
-		}
-	}
-#if WITH_EDITOR
-	else if (GIsEditor)
-	{
-		DirName = TEXT("InstallCacheEditor");
-	}
-#endif //if WITH_EDITOR
-	else
-	{
-		DirName = TEXT("InstallCache");
-	}
-
-	return FPaths::ProjectPersistentDownloadDir() / TEXT("IoStore") / DirName;
-}
-///////////////////////////////////////////////////////////////////////////////
-static bool HasInstallCacheDirectory()
-{
-	return FPaths::HasProjectPersistentDownloadDir();
-}
-
 ///////////////////////////////////////////////////////////////////////////////
 static FAnsiStringBuilderBase& GetChunkUrl(
 	const FStringView& Host,
@@ -612,31 +545,83 @@ FIoStatus FOnDemandIoStore::Initialize()
 	}
 #endif
 
-	FString CacheStatusText;
-	if (FForkProcessHelper::IsForkRequested() && !FForkProcessHelper::IsForkedChildProcess())
+	const FIoStatus InstallCacheStatus = InitializeInstallCache();
+	if (InstallCacheStatus.GetErrorCode() == EIoErrorCode::PendingFork)
 	{
-		// Defer initialization when forking server process to get a unique install cache location per process
-		CacheStatusText = TEXT("WaitingOnServerFork");
+		UE_LOG(LogIoStoreOnDemand, Log, TEXT("Deferring initialization of install cache until post server fork"))
 		OnServerPostForkHandle = FCoreDelegates::OnPostFork.AddLambda(
 			[this](EForkProcessRole ProcessRole)
 			{
 				if (ProcessRole == EForkProcessRole::Child)
 				{
-					if (FIoStatus Status = ConditionallyInitializeInstallCache(); !Status.IsOk())
+					if (FIoStatus Status = InitializeInstallCache(); !Status.IsOk())
 					{
-						UE_LOG(LogIoStoreOnDemand, Error, TEXT("Failed to initialize install cache, reason '%s'"), *Status.ToString());
+						if (Status.GetErrorCode() == EIoErrorCode::Disabled)
+						{
+							UE_LOG(LogIoStoreOnDemand, Log, TEXT("Install cache disabled"));
+						}
+						else
+						{
+							UE_LOG(LogIoStoreOnDemand, Error, TEXT("Failed to initialize install cache, reason '%s'"), *Status.ToString());
+						}
 					}
 				}
 			});
+
+		return FIoStatus::Ok;
+	}
+	else if (InstallCacheStatus.GetErrorCode() == EIoErrorCode::Disabled)
+	{
+		UE_LOG(LogIoStoreOnDemand, Log, TEXT("Install cache disabled"));
+		return FIoStatus::Ok;
+	}
+
+	return InstallCacheStatus;
+}
+
+FIoStatus FOnDemandIoStore::InitializeStreamingBackend(const FOnDemandEndpointConfig& EndpointConfig)
+{
+	TUniquePtr<IIasCache> StreamingCache;
+	FIasCacheConfig CacheConfig = Config::GetStreamingCacheConfig(FCommandLine::Get());
+
+	if (CacheConfig.DiskQuota > 0)
+	{
+		if (FPaths::HasProjectPersistentDownloadDir())
+		{
+			FString CacheDir = FPaths::ProjectPersistentDownloadDir();
+			StreamingCache = MakeIasCache(*CacheDir, CacheConfig);
+
+			UE_CLOG(!StreamingCache.IsValid(), LogIas, Warning, TEXT("File cache disabled - streaming only (init-fail)"));
+		}
+		else
+		{
+			UE_LOG(LogIas, Warning, TEXT("File cache disabled - streaming only (project has no persistent download dir enabled for this platform)"));
+		}
 	}
 	else
 	{
-		const FIoStatus CacheStatus = ConditionallyInitializeInstallCache();
-		CacheStatusText = CacheStatus.IsOk() ? TEXT("Enabled") : TEXT("Disabled");
+		UE_LOG(LogIas, Log, TEXT("File cache disabled - streaming only (zero-quota)"));
 	}
 
-	UE_LOG(LogIoStoreOnDemand, Display, TEXT("Initialized, InstallCache='%s'"), *CacheStatusText);
-	return EIoErrorCode::Ok;
+	StreamingBackend = MakeOnDemandIoDispatcherBackend(EndpointConfig, *this, MoveTemp(StreamingCache));
+	if (StreamingBackend.IsValid() == false)
+	{
+		return FIoStatus(EIoErrorCode::InvalidParameter, TEXT("Failed to create streaming backend"));
+	}
+	StreamingBackend->SetEnabled(EnumHasAnyFlags(StreamingOptions, EOnDemandStreamingOptions::Enabled));
+	StreamingBackend->SetBulkOptionalEnabled(EnumHasAnyFlags(StreamingOptions, EOnDemandStreamingOptions::OptionalBulkDataDisabled) == false);
+
+	int32 BackendPriority = -10;
+#if !UE_BUILD_SHIPPING
+	if (FParse::Param(FCommandLine::Get(), TEXT("Ias")))
+	{
+		UE_LOG(LogIas, Log, TEXT("Setting streaming backend priority higher than file system backend"));
+		BackendPriority = 10;
+	}
+#endif
+	FIoDispatcher::Get().Mount(StreamingBackend.ToSharedRef(), BackendPriority);
+
+	return FIoStatus::Ok;
 }
 
 void FOnDemandIoStore::Mount(FOnDemandMountArgs&& Args, FOnDemandMountCompleted&& OnCompleted)
@@ -820,65 +805,31 @@ FIoStatus FOnDemandIoStore::InitializeInstallCache()
 		return FIoStatusBuilder(EIoErrorCode::PendingFork) << TEXT("Install cache waiting for fork");
 	}
 
-	bool bUseInstallCache = GIoStoreOnDemandInstallCacheEnabled;
-#if !UE_BUILD_SHIPPING
-	bUseInstallCache = FParse::Param(FCommandLine::Get(), TEXT("NoIAD")) == false;
-#endif
-	if (bUseInstallCache)
+	TIoStatusOr<FOnDemandInstallCacheConfig> CacheConfig = Config::TryParseInstallCacheConfig(FCommandLine::Get());
+	if (CacheConfig.IsOk() == false)
 	{
-		if (!Private::HasInstallCacheDirectory())
-		{
-			return FIoStatusBuilder(EIoErrorCode::InvalidParameter) << TEXT("Persistent storage not configured");
-		}
-
-		FOnDemandInstallCacheConfig CacheConfig;
-
-		const int64 DiskQuota = []
-		{
-			if (FString ParamValue; FParse::Value(FCommandLine::Get(), TEXT("-Iad.FileCache.DiskQuota="), ParamValue))
-			{
-				return ParseSizeParam(ParamValue);
-			}
-			else if (FString ValueStr; GConfig->GetString(TEXT("OnDemandInstall"), TEXT("FileCache.DiskQuota"), ValueStr, GEngineIni))
-			{
-				return ParseSizeParam(ValueStr);
-			}
-			return int64(-1);
-		}();
-
-		if (DiskQuota > 0)
-		{
-			CacheConfig.DiskQuota = DiskQuota;
-		}
+		return CacheConfig.Status();
+	}
 
-		CacheConfig.RootDirectory = Private::GetInstallCacheDirectory();
+	InstallCache = MakeOnDemandInstallCache(*this, CacheConfig.ConsumeValueOrDie());
+	if (InstallCache.IsValid())
+	{
+		int32 BackendPriority = -5; // Lower than file (zero) but higher than streaming backend (-10)
 #if !UE_BUILD_SHIPPING
-		CacheConfig.bDropCache = FParse::Param(FCommandLine::Get(), TEXT("Iad.DropCache"));
-#endif
-		InstallCache = MakeOnDemandInstallCache(*this, CacheConfig);
-		if (InstallCache.IsValid())
+		if (FParse::Param(FCommandLine::Get(), TEXT("Iad")))
 		{
-			int32 BackendPriority = -5; // Lower than file (zero) but higher than streaming backend (-10)
-#if !UE_BUILD_SHIPPING
-			if (FParse::Param(FCommandLine::Get(), TEXT("Iad")))
-			{
-				// Bump the priority to be higher then the file system backend
-				BackendPriority = 5;
-			}
-#endif
-			FIoDispatcher::Get().Mount(InstallCache.ToSharedRef(), BackendPriority);
-			PackageStoreBackend = MakeOnDemandPackageStoreBackend();
-			FPackageStore::Get().Mount(PackageStoreBackend.ToSharedRef());
-		}
-		else
-		{
-			UE_LOG(LogIoStoreOnDemand, Error, TEXT("Failed to initialize install cache"));
-			return FIoStatusBuilder(EIoErrorCode::InvalidParameter) << TEXT("Failed to initialize install cache");
+			// Bump the priority to be higher then the file system backend
+			BackendPriority = 5;
 		}
+#endif
+		FIoDispatcher::Get().Mount(InstallCache.ToSharedRef(), BackendPriority);
+		PackageStoreBackend = MakeOnDemandPackageStoreBackend();
+		FPackageStore::Get().Mount(PackageStoreBackend.ToSharedRef());
 	}
 	else
 	{
-		UE_LOG(LogIoStoreOnDemand, Warning, TEXT("Install cache disabled"));
+		UE_LOG(LogIoStoreOnDemand, Error, TEXT("Failed to initialize install cache"));
+		return FIoStatusBuilder(EIoErrorCode::InvalidParameter) << TEXT("Failed to initialize install cache");
 	}
 
 #if !(UE_BUILD_SHIPPING|UE_BUILD_TEST)
@@ -934,12 +885,6 @@ FIoStatus FOnDemandIoStore::InitializeInstallCache()
 	return EIoErrorCode::Ok;
 }
 
-FIoStatus FOnDemandIoStore::ConditionallyInitializeInstallCache()
-{
-	static FIoStatus CacheStatus = InitializeInstallCache();
-	return CacheStatus;
-}
-
 FOnDemandChunkInfo FOnDemandIoStore::GetChunkInfo(const FIoChunkId& ChunkId, EOnDemandContainerFlags ContainerFlags)
 {
 	TUniqueLock Lock(ContainerMutex);
@@ -1329,9 +1274,9 @@ FOnDemandInstallResult FOnDemandIoStore::TickInstallRequest(const FInstallReques
 		return OutResult;
 	}
 
-	Status = ConditionallyInitializeInstallCache();
-	if (Status.IsOk() == false)
+	if (InstallCache.IsValid() == false)
 	{
+		Status = FIoStatusBuilder(EIoErrorCode::InvalidParameter) << TEXT("Install cache not initialized");
 		return OutResult;
 	}
 
@@ -2012,9 +1957,42 @@ FOnDemandCacheUsage FOnDemandIoStore::GetCacheUsage() const
 		CacheUsage.InstallCache = InstallCache->GetCacheUsage();
 	}
 
+	if (StreamingBackend.IsValid())
+	{
+		CacheUsage.StreamingCache = StreamingBackend->GetCacheUsage();
+	}
+
 	return CacheUsage;
 }
 
+void FOnDemandIoStore::SetStreamingOptions(EOnDemandStreamingOptions Options) 
+{
+	StreamingOptions = Options;
+	if (StreamingBackend.IsValid())
+	{
+		StreamingBackend->SetEnabled(EnumHasAnyFlags(Options, EOnDemandStreamingOptions::Enabled));
+		StreamingBackend->SetBulkOptionalEnabled(EnumHasAnyFlags(Options, EOnDemandStreamingOptions::OptionalBulkDataDisabled) == false);
+	}
+}
+
+void FOnDemandIoStore::ReportAnalytics(TArray<FAnalyticsEventAttribute>& OutAnalyticsArray) const
+{
+	if (StreamingBackend.IsValid())
+	{
+		return StreamingBackend->ReportAnalytics(OutAnalyticsArray);
+	}
+}
+
+TUniquePtr<IAnalyticsRecording> FOnDemandIoStore::StartAnalyticsRecording() const
+{
+	if (StreamingBackend.IsValid())
+	{
+		return StreamingBackend->StartAnalyticsRecording();
+	}
+
+	return TUniquePtr<IAnalyticsRecording>();
+}
+
 ///////////////////////////////////////////////////////////////////////////////
 FOnDemandInternalContentHandle::~FOnDemandInternalContentHandle()
 {
diff --git a/Engine/Source/Runtime/Experimental/IoStore/OnDemand/Private/OnDemandIoStore.h b/Engine/Source/Runtime/Experimental/IoStore/OnDemand/Private/OnDemandIoStore.h
--- a/Engine/Source/Runtime/Experimental/IoStore/OnDemand/Private/OnDemandIoStore.h
+++ b/Engine/Source/Runtime/Experimental/IoStore/OnDemand/Private/OnDemandIoStore.h
@@ -25,12 +25,13 @@ using FSharedContainerHeader		= TSharedPtr<FIoContainerHeader>;
 namespace UE::IoStore
 {
 
-class IOnDemandPackageStoreBackend;
-class IOnDemandInstallCache;
+struct FOnDemandEndpointConfig;
 enum class ELogOnDemandCacheUsage : uint8;
-using FSharedPackageStoreBackend	= TSharedPtr<IOnDemandPackageStoreBackend>;
-using FSharedInstallCache			= TSharedPtr<IOnDemandInstallCache>;
-using FWeakOnDemandIoStore			= TWeakPtr<FOnDemandIoStore, ESPMode::ThreadSafe>;
+
+using FSharedStreamingBackend		= TSharedPtr<class IOnDemandIoDispatcherBackend>;
+using FSharedPackageStoreBackend	= TSharedPtr<class IOnDemandPackageStoreBackend>;
+using FSharedInstallCache			= TSharedPtr<class IOnDemandInstallCache>;
+using FWeakOnDemandIoStore			= TWeakPtr<class FOnDemandIoStore, ESPMode::ThreadSafe>;
 
 ///////////////////////////////////////////////////////////////////////////////
 class FOnDemandInternalContentHandle
@@ -253,6 +254,7 @@ FOnDemandChunkInfo FOnDemandChunkInfo::Find(FSharedOnDemandContainer Container,
 ///////////////////////////////////////////////////////////////////////////////
 class FOnDemandIoStore
 	: public TSharedFromThis<FOnDemandIoStore, ESPMode::ThreadSafe>
+	, public IOnDemandIoStore
 {
 	struct FMountRequest
 	{
@@ -282,29 +284,34 @@ class FOnDemandIoStore
 	using FSharedPurgeRequest = TSharedRef<FPurgeRequest>;
 
 public:
-	FOnDemandIoStore();
-	~FOnDemandIoStore();
-	FOnDemandIoStore(const FOnDemandIoStore&) = delete;
-	FOnDemandIoStore(FOnDemandIoStore&&) = delete;
-	FOnDemandIoStore& operator=(const FOnDemandIoStore&) = delete;
-	FOnDemandIoStore& operator=(FOnDemandIoStore&&) = delete;
-
-	FIoStatus				Initialize();
-	void					Mount(FOnDemandMountArgs&& Args, FOnDemandMountCompleted&& OnCompleted);
-	void					Install(
-								FOnDemandInstallArgs&& Args,
-								FOnDemandInstallCompleted&& OnCompleted,
-								FOnDemandInstallProgressed&& OnProgress = nullptr,
-								const FOnDemandCancellationToken* CancellationToken = nullptr);
-	void					Purge(FOnDemandPurgeArgs&& Args, FOnDemandPurgeCompleted&& OnCompleted);
-	FIoStatus				Unmount(FStringView MountId);
-	TIoStatusOr<uint64>		GetInstallSize(const FOnDemandGetInstallSizeArgs& Args) const;
-	FIoStatus				GetInstallSizesByMountId(const FOnDemandGetInstallSizeArgs& Args, TMap<FString, uint64>& OutSizesByMountId) const;
-	FOnDemandChunkInfo		GetStreamingChunkInfo(const FIoChunkId& ChunkId);
-	FOnDemandChunkInfo		GetInstalledChunkInfo(const FIoChunkId& ChunkId);
-	void					ReleaseContent(FOnDemandInternalContentHandle& ContentHandle);
-	void					GetReferencedContent(TArray<FSharedOnDemandContainer>& OutContainers, TArray<TBitArray<>>& OutChunkEntryIndices);
-	FOnDemandCacheUsage		GetCacheUsage() const;
+	UE_NONCOPYABLE(FOnDemandIoStore);
+
+										FOnDemandIoStore();
+										~FOnDemandIoStore();
+
+	FIoStatus							Initialize();
+	FIoStatus							InitializeStreamingBackend(const FOnDemandEndpointConfig& EndpointConfig);
+	FOnDemandChunkInfo					GetStreamingChunkInfo(const FIoChunkId& ChunkId);
+	FOnDemandChunkInfo					GetInstalledChunkInfo(const FIoChunkId& ChunkId);
+	void								ReleaseContent(FOnDemandInternalContentHandle& ContentHandle);
+	void								GetReferencedContent(TArray<FSharedOnDemandContainer>& OutContainers, TArray<TBitArray<>>& OutChunkEntryIndices);
+
+	// IOnDemandIoStore interface
+	virtual void						Mount(FOnDemandMountArgs&& Args, FOnDemandMountCompleted&& OnCompleted) override;
+	virtual void						Install(
+											FOnDemandInstallArgs&& Args,
+											FOnDemandInstallCompleted&& OnCompleted,
+											FOnDemandInstallProgressed&& OnProgress = nullptr,
+											const FOnDemandCancellationToken* CancellationToken = nullptr) override;
+	virtual void						Purge(FOnDemandPurgeArgs&& Args, FOnDemandPurgeCompleted&& OnCompleted) override;
+	virtual FIoStatus					Unmount(FStringView MountId) override;
+	virtual TIoStatusOr<uint64> 		GetInstallSize(const FOnDemandGetInstallSizeArgs& Args) const override;
+	virtual FIoStatus					GetInstallSizesByMountId(const FOnDemandGetInstallSizeArgs& Args, TMap<FString, uint64>& OutSizesByMountId) const override;
+	virtual FOnDemandCacheUsage			GetCacheUsage() const override;
+	virtual EOnDemandStreamingOptions	GetStreamingOptions() const override { return StreamingOptions; }
+	virtual void						SetStreamingOptions(EOnDemandStreamingOptions Options) override;
+	virtual void						ReportAnalytics(TArray<FAnalyticsEventAttribute>& OutAnalyticsArray) const override;
+	virtual TUniquePtr<IAnalyticsRecording> StartAnalyticsRecording() const override;
 
 private:
 	FIoStatus				GetContainersAndPackagesForInstall(
@@ -314,7 +321,6 @@ private:
 								TSet<FSharedOnDemandContainer>& OutContainersForInstallation,
 								TSet<FPackageId>& OutPackageIdsToInstall) const;
 	FIoStatus				InitializeInstallCache();
-	FIoStatus				ConditionallyInitializeInstallCache();
 	FOnDemandChunkInfo		GetChunkInfo(const FIoChunkId& ChunkId, EOnDemandContainerFlags ContainerFlags);
 	void					TryEnterTickLoop();
 	void					TickLoop();
@@ -336,6 +342,7 @@ private:
 
 	FSharedInstallCache					InstallCache;
 	FSharedPackageStoreBackend			PackageStoreBackend;
+	FSharedStreamingBackend				StreamingBackend;
 	FDelegateHandle						OnMountPakHandle;
 	FDelegateHandle						OnServerPostForkHandle;
 	TArray<FSharedOnDemandContainer>	Containers;
@@ -347,6 +354,7 @@ private:
 	TArray<FSharedPurgeRequest>			PurgeRequests;
 	UE::FMutex							RequestMutex;
 
+	EOnDemandStreamingOptions			StreamingOptions = EOnDemandStreamingOptions::Enabled;
 	bool								bTicking = false;
 	bool								bTickRequested = false;
 	TFuture<void>						TickFuture;
diff --git a/Engine/Source/Runtime/Experimental/IoStore/OnDemand/Public/IO/IoStoreOnDemand.h b/Engine/Source/Runtime/Experimental/IoStore/OnDemand/Public/IO/IoStoreOnDemand.h
--- a/Engine/Source/Runtime/Experimental/IoStore/OnDemand/Public/IO/IoStoreOnDemand.h
+++ b/Engine/Source/Runtime/Experimental/IoStore/OnDemand/Public/IO/IoStoreOnDemand.h
@@ -20,8 +20,6 @@
 
 struct FAnalyticsEventAttribute;
 namespace UE::IoStore { class FOnDemandInternalContentHandle; }
-namespace UE::IoStore { class FOnDemandIoStore; }
-namespace UE::IoStore { class IOnDemandIoDispatcherBackend; }
 
 // Custom initialization allows users to control when
 // the system should be initialized.
@@ -35,26 +33,6 @@ UE_API DECLARE_LOG_CATEGORY_EXTERN(LogIas, Log, All);
 namespace UE::IoStore
 {
 
-////////////////////////////////////////////////////////////////////////////////
-
-int64 ParseSizeParam(FStringView Value);
-
-////////////////////////////////////////////////////////////////////////////////
-#if UE_IAS_CUSTOM_INITIALIZATION
-
-/** Result of calling FIoStoreOnDemandModule::Initialize */
-enum class EOnDemandInitResult
-{
-	/** The module initialized correctly and can be used */
-	Success = 0,
-	/** The module is disabled as OnDemand data is not required for the current process*/
-	Disabled,
-	/** The module was unable to start up correctly due to an unexpected error */
-	Error,
-};
-
-#endif // UE_IAS_CUSTOM_INITIALIZATION
-
 /**
  * Keeps referenced data pinned in the cache until released.
  */
@@ -318,7 +296,7 @@ private:
 	std::atomic<bool> bCanceled{ false };
 };
 
-/** Interface  for recording analytics over a given time period */
+/** Interface for recording analytics over a given time period */
 class IAnalyticsRecording
 {
 public:
@@ -332,69 +310,105 @@ public:
 	virtual void StopRecording() = 0;
 };
 
-class FIoStoreOnDemandModule
-	: public IModuleInterface
+/** Streaming options (IAS). */
+enum class EOnDemandStreamingOptions
 {
-private:
-	void InitializeInternal();
-	TSharedPtr<IOnDemandIoDispatcherBackend> HttpIoDispatcherBackend;
-	// Deferred state requests if called before backend
-	// is initialized
-	TOptional<bool> DeferredEnabled;
-	TOptional<bool> DeferredAbandonCache;
-	TOptional<bool> DeferredBulkOptionalEnabled;
-	TSharedPtr<FOnDemandIoStore, ESPMode::ThreadSafe> IoStore;
-
-	/** Indicates that platform specific setup was invoked when the module was started and needs to be cleaned up on shutdown */
-	bool bPlatformSpecificSetup = false;
+	/** On-demand streaming disabled. */
+	Disabled					= 0,
+	/** On-demand streaming enabled. */
+	Enabled						= (1 << 0),
+	/** Optional bulk data disabled. */
+	OptionalBulkDataDisabled	= (1 << 1)
+};
+ENUM_CLASS_FLAGS(EOnDemandStreamingOptions);
 
+/** Interface for installing and streaming content on-demand. */
+class IOnDemandIoStore
+{
 public:
-	UE_API void SetBulkOptionalEnabled(bool bInEnabled);
-	UE_API void SetEnabled(bool bInEnabled);
-	UE_API bool IsEnabled() const;
-	UE_API void AbandonCache();
-
+	/** Destroy the I/O store. */
+	virtual ~IOnDemandIoStore() = default;
+	/**
+	 * Mount an on-demand container.
+	 * @param	Args		Mount arguments.
+	 * @param	OnCompleted	Completion callback.
+	 **/
+	virtual void Mount(FOnDemandMountArgs&& Args, FOnDemandMountCompleted&& OnCompleted) = 0;
+	/**
+	 * Unmount all container(s) associated with the specified mount ID.
+	 * @param	MountId	The Mount ID passed to the mount call.
+	 **/
+	virtual FIoStatus Unmount(FStringView MountId) = 0;
+	/**
+	 * Install content.
+	 * @param	Args				The install arguments.
+	 * @param	OnCompleted			Completion callback.
+	 * @param	OnProgress			Optional progress callback.
+	 * @param	CancellationToken	Optional cancellation token.
+	 **/
+	virtual void Install(
+			FOnDemandInstallArgs&& Args,
+			FOnDemandInstallCompleted&& OnCompleted,
+			FOnDemandInstallProgressed&& OnProgress = nullptr,
+			const FOnDemandCancellationToken* CancellationToken = nullptr) = 0;
+	/**
+	 * Purge the cache.
+	 * @param	Args		The purge arguments.
+	 * @param	OnComplete	Completion callback.
+	 **/
+	virtual void Purge(FOnDemandPurgeArgs&& Args, FOnDemandPurgeCompleted&& OnCompleted) = 0;
+	/**
+	 * Get the total space in bytes needed to install the specified content.
+	 * @param Args	Arguments for computing required install size.
+	 **/
+	virtual TIoStatusOr<uint64> GetInstallSize(const FOnDemandGetInstallSizeArgs& Args) const = 0;
+	/**
+	 * Get the total space in bytes needed to install the specified content.
+	 * @param Args	Arguments for computing required install size.
+	 * @return		The total size by mount ID.
+	 **/
+	virtual FIoStatus GetInstallSizesByMountId(const FOnDemandGetInstallSizeArgs& Args, TMap<FString, uint64>& OutSizesByMountId) const = 0;
+	/** Returns the total cache size in bytes. */
+	virtual FOnDemandCacheUsage GetCacheUsage() const = 0;
+	/** Return the curren streaming optinos. */
+	virtual EOnDemandStreamingOptions GetStreamingOptions() const = 0;
+	/** Set streaming options.*/
+	virtual void SetStreamingOptions(EOnDemandStreamingOptions Options) = 0;
 	/**
 	 * Reports the statistics for the current OnDemandBackend. It will record any events that have occurred from the last time that this was called.
 	 * This is a legacy method, prefer using StartAnalyticsRecording instead.
 	 */
-	UE_API void ReportAnalytics(TArray<FAnalyticsEventAttribute>& OutAnalyticsArray) const;
-	
+	virtual void ReportAnalytics(TArray<FAnalyticsEventAttribute>& OutAnalyticsArray) const = 0;
 	/** 
 	 * Create a new analytics interface for the current OnDemandBackend, can be used to record statistics for a set period of time.
 	 * This can return a nullptr if the analytics system is not enabled.
 	 */
-	UE_API TUniquePtr<IAnalyticsRecording> StartAnalyticsRecording() const;
-
-	UE_API void Mount(FOnDemandMountArgs&& Args, FOnDemandMountCompleted&& OnCompleted);
-	UE_API FIoStatus Unmount(FStringView MountId);
-
-	UE_API void Install(
-		FOnDemandInstallArgs&& Args,
-		FOnDemandInstallCompleted&& OnCompleted,
-		FOnDemandInstallProgressed&& OnProgress = nullptr,
-		const FOnDemandCancellationToken* CancellationToken = nullptr);
-
-	UE_API void Purge(FOnDemandPurgeArgs&& Args, FOnDemandPurgeCompleted&& OnCompleted);
-
-	UE_API TIoStatusOr<uint64> GetInstallSize(const FOnDemandGetInstallSizeArgs& Args) const;
-
-	UE_API FIoStatus GetInstallSizesByMountId(const FOnDemandGetInstallSizeArgs& Args, TMap<FString, uint64>& OutSizesByMountId) const;
-
-	/** This does not queue a request so the returned state may not be completely consistent
-	  * if there are requests currently being processed. 
-	  * This should only be used for purposes such as debugging telemetry. 
-	  */
-	UE_API TIoStatusOr<FOnDemandCacheUsage> GetCacheUsage() const;
+	virtual TUniquePtr<IAnalyticsRecording> StartAnalyticsRecording() const = 0;
+};
 
+class FIoStoreOnDemandModule
+	: public IModuleInterface
+{
+public:
+	/** Startup module. */
 	virtual void StartupModule() override;
+	/** Shutdown module. */
 	virtual void ShutdownModule() override;
-	
 #if UE_IAS_CUSTOM_INITIALIZATION
-	UE_API EOnDemandInitResult Initialize();
-#endif //UE_IAS_CUSTOM_INITIALIZATION
+	UE_API FIoStatus InitializePostHotfix();
+#endif 
+
+private:
+	FIoStatus InitializeStreamingBackend();
+	/** Indicates that platform specific setup was invoked when the module was started and needs to be cleaned up on shutdown */
+	bool bPlatformSpecificSetup = false;
 };
 
+/** Returns the on-demand I/O store if available. */
+UE_API IOnDemandIoStore* TryGetOnDemandIoStore();
+/** Returns the on-demand I/O store. */
+UE_API IOnDemandIoStore& GetOnDemandIoStore();
+
 } // namespace UE::IoStore
 
 #undef UE_API
