Fixed IAS bug where the endpoint config TOC path not being set.

#rnx

diff --git a/Engine/Source/Runtime/Experimental/IoStore/OnDemand/Private/OnDemandConfig.cpp b/Engine/Source/Runtime/Experimental/IoStore/OnDemand/Private/OnDemandConfig.cpp
--- a/Engine/Source/Runtime/Experimental/IoStore/OnDemand/Private/OnDemandConfig.cpp
+++ b/Engine/Source/Runtime/Experimental/IoStore/OnDemand/Private/OnDemandConfig.cpp
@@ -295,6 +295,17 @@ TIoStatusOr<FOnDemandEndpointConfig> TryParseEndpointConfig(const TCHAR* Command
 		return Status;
 	}
 
+	{
+		TStringBuilder<256> TocFilePath;
+		FPathViews::Append(TocFilePath, TEXT("Cloud"), FPaths::GetBaseFilename(OutConfig.TocPath));
+		TocFilePath.Append(TEXT(".iochunktoc"));
+
+		if (FPlatformMisc::FileExistsInPlatformPackage(*TocFilePath))
+		{
+			OutConfig.TocFilePath = TocFilePath;
+		}
+	}
+
 	return OutConfig;
 }
 
