Prevent manifold reuse when contact modifiers have been applied

Previously, a contact could be reused even if the user applied a contact modifier. This meant that the modified contact could be reused on the next step. This then meant that a user's contact modification code could operate on the modified contact as opposed to the original unmodified contact. This behavior can be unintuitive to users who might write their code under the assumption they are modifying the original contact state. As an example, if a user wants to invert the contact normal, their code could be inverting the actual contact normal, or it could be inverting their own inverted normal from the previous step. This behavior is made more unintuitive because contact reuse is currently frame rate dependent. User code could behave very differently at higher frame rates compared to lower ones depending on whether the contact was reused or not. Preventing this kind of reuse makes the behavior more consistent for users.

diff --git a/Engine/Source/Runtime/Experimental/Chaos/Private/Chaos/Collision/PBDCollisionConstraint.cpp b/Engine/Source/Runtime/Experimental/Chaos/Private/Chaos/Collision/PBDCollisionConstraint.cpp
--- a/Engine/Source/Runtime/Experimental/Chaos/Private/Chaos/Collision/PBDCollisionConstraint.cpp
+++ b/Engine/Source/Runtime/Experimental/Chaos/Private/Chaos/Collision/PBDCollisionConstraint.cpp
@@ -803,7 +803,7 @@ namespace Chaos
 		ClosestManifoldPointIndex = INDEX_NONE;
 		Flags.bWasManifoldRestored = false;
 
-		if (!Flags.bCanRestoreManifold)
+		if (!Flags.bCanRestoreManifold || Flags.bModifierApplied)
 		{
 			return false;
 		}
