Adding some Iris huge object and pending object metrics

diff --git a/Engine/Source/Runtime/Engine/Private/Net/Iris/ReplicationSystem/EngineReplicationBridge.cpp b/Engine/Source/Runtime/Engine/Private/Net/Iris/ReplicationSystem/EngineReplicationBridge.cpp
--- a/Engine/Source/Runtime/Engine/Private/Net/Iris/ReplicationSystem/EngineReplicationBridge.cpp
+++ b/Engine/Source/Runtime/Engine/Private/Net/Iris/ReplicationSystem/EngineReplicationBridge.cpp
@@ -880,8 +880,6 @@ void UEngineReplicationBridge::ConsumeNetMetrics(TArray<FAnalyticsEventAttribute
 
 	GetReplicationSystem()->CollectNetMetrics(MetricsCollector);
 
-	//$IRIS TODO: Add ResetNetMetrics call if the RepSystem starts caching cumulative stats
-
 	const TMap<FName, FNetMetric>& Metrics = MetricsCollector.GetMetrics();
 
 	for (auto& It : Metrics)
@@ -916,6 +914,7 @@ void UEngineReplicationBridge::ConsumeNetMetrics(TArray<FAnalyticsEventAttribute
 		}
 	}
 
+	GetReplicationSystem()->ResetNetMetrics();
 }
 
 #else //!UE_WITH_IRIS
diff --git a/Engine/Source/Runtime/Experimental/Iris/Core/Private/Iris/ReplicationSystem/ReplicationSystem.cpp b/Engine/Source/Runtime/Experimental/Iris/Core/Private/Iris/ReplicationSystem/ReplicationSystem.cpp
--- a/Engine/Source/Runtime/Experimental/Iris/Core/Private/Iris/ReplicationSystem/ReplicationSystem.cpp
+++ b/Engine/Source/Runtime/Experimental/Iris/Core/Private/Iris/ReplicationSystem/ReplicationSystem.cpp
@@ -71,6 +71,12 @@ public:
 
 	FName MetricNameTotalRootObjects;
 	FName MetricNameTotalSubObjects;
+	FName MetricNameAvgPendingObjectCount;
+	FName MetricNameMaxPendingObjectCount;
+	FName MetricNameAvgPendingDependentObjectCount;
+	FName MetricNameMaxPendingDependentObjectCount;
+	FName MetricNameAvgHugeObjectSendQueue;
+	FName MetricNameMaxHugeObjectSendQueue;
 
 	explicit FReplicationSystemImpl(UReplicationSystem* InReplicationSystem, const UReplicationSystem::FReplicationSystemParams& Params)
 	: ReplicationSystem(InReplicationSystem)
@@ -86,6 +92,12 @@ public:
 	{
 		MetricNameTotalRootObjects = TEXT("TotalSubObjects");
 		MetricNameTotalSubObjects = TEXT("TotalRootObjects");
+		MetricNameAvgPendingObjectCount = TEXT("AvgPendingObjectCount");
+		MetricNameMaxPendingObjectCount = TEXT("MaxPendingObjectCount");
+		MetricNameAvgPendingDependentObjectCount = TEXT("AvgPendingDependentObjectCount");
+		MetricNameMaxPendingDependentObjectCount = TEXT("MaxPendingDependentObjectCount");
+		MetricNameAvgHugeObjectSendQueue = TEXT("AvgHugeObjectSendQueue");
+		MetricNameMaxHugeObjectSendQueue = TEXT("MaxHugeObjectSendQueue");
 	}
 
 	~FReplicationSystemImpl()
@@ -706,6 +718,36 @@ public:
 		// Collect stats on total replicated objects
 		OutNetMetrics.EmplaceMetric(MetricNameTotalRootObjects, FNetMetric(TotalNetObjects-TotalSubObjects));
 		OutNetMetrics.EmplaceMetric(MetricNameTotalSubObjects, FNetMetric(TotalSubObjects));
+
+		// Pending object and huge object stats
+		const FReplicationStats& ReplicationStats = ReplicationSystemInternal.GetAccumulatedReplicationStats();
+		double AvgPendingObjectCount;
+		double AvgPendingDependentObjectCount;
+		double AvgHugeObjectSendQueue;
+		if (ReplicationStats.SampleCount)
+		{
+			double SampleCount = (double)ReplicationStats.SampleCount;
+			AvgPendingObjectCount = (double)ReplicationStats.PendingObjectCount / SampleCount;
+			AvgPendingDependentObjectCount = (double)ReplicationStats.PendingDependentObjectCount / SampleCount;
+			AvgHugeObjectSendQueue = (double)ReplicationStats.HugeObjectSendQueue / SampleCount;
+		}
+		else
+		{
+			AvgPendingObjectCount = 0.0f;
+			AvgPendingDependentObjectCount = 0.0f;
+			AvgHugeObjectSendQueue = 0.0f;
+		}
+		OutNetMetrics.EmplaceMetric(MetricNameAvgPendingObjectCount, AvgPendingObjectCount);
+		OutNetMetrics.EmplaceMetric(MetricNameMaxPendingObjectCount, ReplicationStats.MaxPendingObjectCount);
+		OutNetMetrics.EmplaceMetric(MetricNameAvgPendingDependentObjectCount, AvgPendingDependentObjectCount);
+		OutNetMetrics.EmplaceMetric(MetricNameMaxPendingDependentObjectCount, ReplicationStats.MaxPendingDependentObjectCount);
+		OutNetMetrics.EmplaceMetric(MetricNameAvgHugeObjectSendQueue, AvgHugeObjectSendQueue);
+		OutNetMetrics.EmplaceMetric(MetricNameMaxHugeObjectSendQueue, ReplicationStats.MaxHugeObjectSendQueue);
+	}
+
+	void ResetNetMetrics()
+	{
+		ReplicationSystemInternal.GetAccumulatedReplicationStats() = {};
 	}
 }; // end class FReplicationSystemImpl
 
@@ -790,6 +832,8 @@ void UReplicationSystem::PreSendUpdate(const FSendUpdateParams& Params)
 			SendStats.SetNumberOfReplicatingConnections(ReplicatingConnections.CountSetBits());
 		}
 #endif
+		InternalSys.GetAccumulatedReplicationStats().Accumulate(InternalSys.GetTickReplicationStats());
+		InternalSys.GetTickReplicationStats() = {};
 	
 		// Force a integrity check of all replicated instances
 		if (bDoCollectGarbage || ReplicationSystemCVars::bForcePruneBeforeUpdate)
@@ -939,6 +983,9 @@ void UReplicationSystem::PostSendUpdate()
 			FNetTypeStats& TypeStats = Impl->ReplicationSystemInternal.GetNetTypeStats();
 			TypeStats.ReportCSVStats();
 
+			FReplicationStats& ReplicationStats = Impl->ReplicationSystemInternal.GetTickReplicationStats();
+			ReplicationStats.ReportCSVStats();
+
 			if (Impl->ReplicationSystemInternal.IsDirtyNetObjectTrackerInitialized())
 			{
 				Impl->ReplicationSystemInternal.GetDirtyNetObjectTracker().ReportCSVStats();
@@ -1920,6 +1967,11 @@ void UReplicationSystem::CollectNetMetrics(UE::Net::FNetMetrics& OutNetMetrics)
 	Impl->CollectNetMetrics(OutNetMetrics);
 }
 
+void UReplicationSystem::ResetNetMetrics()
+{
+	Impl->ResetNetMetrics();
+}
+
 UE::Net::FReplicationSystemDelegates& UReplicationSystem::GetDelegates()
 {
 	return Impl->Delegates;
diff --git a/Engine/Source/Runtime/Experimental/Iris/Core/Private/Iris/ReplicationSystem/ReplicationSystemInternal.h b/Engine/Source/Runtime/Experimental/Iris/Core/Private/Iris/ReplicationSystem/ReplicationSystemInternal.h
--- a/Engine/Source/Runtime/Experimental/Iris/Core/Private/Iris/ReplicationSystem/ReplicationSystemInternal.h
+++ b/Engine/Source/Runtime/Experimental/Iris/Core/Private/Iris/ReplicationSystem/ReplicationSystemInternal.h
@@ -107,6 +107,26 @@ public:
 		return SendStats;
 	}
 
+	FReplicationStats& GetTickReplicationStats()
+	{
+		return TickReplicationStats;
+	}
+
+	const FReplicationStats& GetTickReplicationStats() const
+	{
+		return TickReplicationStats;
+	}
+
+	FReplicationStats& GetAccumulatedReplicationStats()
+	{
+		return AccumulatedReplicationStats;
+	}
+
+	const FReplicationStats& GetAccumulatedReplicationStats() const
+	{
+		return AccumulatedReplicationStats;
+	}
+
 	FForwardNetRPCCallMulticastDelegate& GetForwardNetRPCCallMulticastDelegate()
 	{
 		return ForwardNetRPCCallMulticastDelegate;
@@ -139,6 +159,8 @@ private:
 	FDeltaCompressionBaselineInvalidationTracker DeltaCompressionBaselineInvalidationTracker;
 	FNetSendStats SendStats;
 	FNetTypeStats TypeStats;
+	FReplicationStats TickReplicationStats = {};
+	FReplicationStats AccumulatedReplicationStats = {};
 	FForwardNetRPCCallMulticastDelegate ForwardNetRPCCallMulticastDelegate;
 	uint32 Id;
 
diff --git a/Engine/Source/Runtime/Experimental/Iris/Core/Private/Iris/ReplicationSystem/ReplicationWriter.cpp b/Engine/Source/Runtime/Experimental/Iris/Core/Private/Iris/ReplicationSystem/ReplicationWriter.cpp
--- a/Engine/Source/Runtime/Experimental/Iris/Core/Private/Iris/ReplicationSystem/ReplicationWriter.cpp
+++ b/Engine/Source/Runtime/Experimental/Iris/Core/Private/Iris/ReplicationSystem/ReplicationWriter.cpp
@@ -3721,16 +3721,28 @@ void FReplicationWriter::EndWrite()
 		}
 #endif
 
+		const uint32 NumPendingDependentObjects = WriteContext.DependentObjectsPendingSend.Num();
+		const uint32 NumPendingObjectsToWrite = WriteContext.ScheduledObjectCount - WriteContext.CurrentIndex;
+
+		FReplicationStats& ReplicationStats = Parameters.ReplicationSystem->GetReplicationSystemInternal()->GetTickReplicationStats();
+		ReplicationStats.PendingObjectCount += NumPendingObjectsToWrite;
+		ReplicationStats.PendingDependentObjectCount += NumPendingDependentObjects;
+		ReplicationStats.HugeObjectSendQueue += HugeObjectSendQueue.NumRootObjectsInTransit();
+		ReplicationStats.MaxPendingObjectCount = FMath::Max(ReplicationStats.MaxPendingObjectCount, NumPendingObjectsToWrite);
+		ReplicationStats.MaxPendingDependentObjectCount = FMath::Max(ReplicationStats.MaxPendingDependentObjectCount, NumPendingDependentObjects);
+		ReplicationStats.MaxHugeObjectSendQueue = FMath::Max(ReplicationStats.MaxHugeObjectSendQueue, HugeObjectSendQueue.NumRootObjectsInTransit());
+		++ReplicationStats.SampleCount;
+
  #if UE_NET_ENABLE_REPLICATIONWRITER_LOG
 		// See if we failed to write any objects
-		if (const int32 NumPendingDependentObjects = WriteContext.DependentObjectsPendingSend.Num())
+		if (NumPendingDependentObjects)
 		{
-			UE_LOG_REPLICATIONWRITER_WARNING(TEXT("FReplicationWriter::EndWrite() Has %d more dependent objects to write"), NumPendingDependentObjects);
+			UE_LOG_REPLICATIONWRITER_WARNING(TEXT("FReplicationWriter::EndWrite() Has %u more dependent objects to write"), NumPendingDependentObjects);
 		}
 		
-		if (const int32 NumPendingObjectsToWrite = WriteContext.ScheduledObjectCount - WriteContext.CurrentIndex)
+		if (NumPendingObjectsToWrite)
 		{
-			UE_LOG_REPLICATIONWRITER_WARNING(TEXT("FReplicationWriter::EndWrite() Has %d more objects with dirty data"), NumPendingObjectsToWrite);
+			UE_LOG_REPLICATIONWRITER_WARNING(TEXT("FReplicationWriter::EndWrite() Has %u more objects with dirty data"), NumPendingObjectsToWrite);
 		}
 #endif
 
diff --git a/Engine/Source/Runtime/Experimental/Iris/Core/Private/Iris/Stats/NetStats.cpp b/Engine/Source/Runtime/Experimental/Iris/Core/Private/Iris/Stats/NetStats.cpp
--- a/Engine/Source/Runtime/Experimental/Iris/Core/Private/Iris/Stats/NetStats.cpp
+++ b/Engine/Source/Runtime/Experimental/Iris/Core/Private/Iris/Stats/NetStats.cpp
@@ -245,4 +245,29 @@ void FNetTypeStats::ReportCSVStats()
 #undef UE_NET_STATS_RECORD_TYPESTATS_COUNT
 #undef UE_NET_STATS_RECORD_TYPESTATS_BITS
 
+void FReplicationStats::ReportCSVStats()
+{
+#if UE_NET_IRIS_CSV_STATS
+	if (SampleCount)
+	{
+		const double SampleCountDbl = float(SampleCount);
+		CSV_CUSTOM_STAT(Iris, AvgPendingObjectCount, double(PendingObjectCount) / SampleCountDbl, ECsvCustomStatOp::Set);
+		CSV_CUSTOM_STAT(Iris, AvgPendingDependentObjectCount, double(PendingDependentObjectCount) / SampleCountDbl, ECsvCustomStatOp::Set);
+		CSV_CUSTOM_STAT(Iris, AvgHugeObjectSendQueue, double(HugeObjectSendQueue) / SampleCountDbl, ECsvCustomStatOp::Set);
+		CSV_CUSTOM_STAT(Iris, MaxPendingObjectCount, int32(MaxPendingObjectCount), ECsvCustomStatOp::Set);
+		CSV_CUSTOM_STAT(Iris, MaxPendingDependentObjectCount, int32(MaxPendingDependentObjectCount), ECsvCustomStatOp::Set);
+		CSV_CUSTOM_STAT(Iris, MaxHugeObjectSendQueue, int32(MaxHugeObjectSendQueue), ECsvCustomStatOp::Set);
+	}
+	else
+	{
+		CSV_CUSTOM_STAT(Iris, AvgPendingObjectCount, 0, ECsvCustomStatOp::Set);
+		CSV_CUSTOM_STAT(Iris, AvgPendingDependentObjectCount, 0, ECsvCustomStatOp::Set);
+		CSV_CUSTOM_STAT(Iris, AvgHugeObjectSendQueue, 0, ECsvCustomStatOp::Set);
+		CSV_CUSTOM_STAT(Iris, MaxPendingObjectCount, 0, ECsvCustomStatOp::Set);
+		CSV_CUSTOM_STAT(Iris, MaxPendingDependentObjectCount, 0, ECsvCustomStatOp::Set);
+		CSV_CUSTOM_STAT(Iris, MaxHugeObjectSendQueue, 0, ECsvCustomStatOp::Set);
+	}
+#endif
+}
+
 } // end namespace UE::Net::Private
diff --git a/Engine/Source/Runtime/Experimental/Iris/Core/Private/Iris/Stats/NetStats.h b/Engine/Source/Runtime/Experimental/Iris/Core/Private/Iris/Stats/NetStats.h
--- a/Engine/Source/Runtime/Experimental/Iris/Core/Private/Iris/Stats/NetStats.h
+++ b/Engine/Source/Runtime/Experimental/Iris/Core/Private/Iris/Stats/NetStats.h
@@ -210,4 +210,29 @@ private:
 	bool bIsEnabled = false;
 };
 
+struct FReplicationStats
+{
+	/** Report the stats to the CSV profiler. Does nothing if CSV profiler support is compiled out. */
+	void ReportCSVStats();
+
+	void Accumulate(const FReplicationStats& Stats)
+	{
+		PendingObjectCount += Stats.PendingObjectCount;
+		PendingDependentObjectCount += Stats.PendingDependentObjectCount;
+		HugeObjectSendQueue += Stats.HugeObjectSendQueue;
+		MaxPendingObjectCount = FMath::Max(MaxPendingObjectCount, Stats.MaxPendingObjectCount);
+		MaxPendingDependentObjectCount = FMath::Max(MaxPendingDependentObjectCount, Stats.MaxPendingDependentObjectCount);
+		MaxHugeObjectSendQueue = FMath::Max(MaxHugeObjectSendQueue, Stats.MaxHugeObjectSendQueue);
+		SampleCount += Stats.SampleCount;
+	}
+
+	uint64 PendingObjectCount;
+	uint64 PendingDependentObjectCount;
+	uint64 HugeObjectSendQueue;
+	uint32 MaxPendingObjectCount;
+	uint32 MaxPendingDependentObjectCount;
+	uint32 MaxHugeObjectSendQueue;
+	uint32 SampleCount;
+};
+
 } // end namespace UE::Net::Private
diff --git a/Engine/Source/Runtime/Experimental/Iris/Core/Public/Iris/ReplicationSystem/ReplicationSystem.h b/Engine/Source/Runtime/Experimental/Iris/Core/Public/Iris/ReplicationSystem/ReplicationSystem.h
--- a/Engine/Source/Runtime/Experimental/Iris/Core/Public/Iris/ReplicationSystem/ReplicationSystem.h
+++ b/Engine/Source/Runtime/Experimental/Iris/Core/Public/Iris/ReplicationSystem/ReplicationSystem.h
@@ -722,6 +722,9 @@ public:
 	/** Collect relevant metrics that could be sent to analytics */
 	IRISCORE_API void CollectNetMetrics(UE::Net::FNetMetrics& OutNetMetrics) const;
 
+	/** Reset cached accumulated metrics */
+	IRISCORE_API void ResetNetMetrics();
+
 	IRISCORE_API UE::Net::FReplicationSystemDelegates& GetDelegates();
 
 public:
