Fix CIS issue from BendingConstraints change.

diff --git a/Engine/Source/Runtime/Experimental/Chaos/Private/Chaos/PBDBendingConstraintsBase.cpp b/Engine/Source/Runtime/Experimental/Chaos/Private/Chaos/PBDBendingConstraintsBase.cpp
--- a/Engine/Source/Runtime/Experimental/Chaos/Private/Chaos/PBDBendingConstraintsBase.cpp
+++ b/Engine/Source/Runtime/Experimental/Chaos/Private/Chaos/PBDBendingConstraintsBase.cpp
@@ -10,7 +10,7 @@ bool bChaos_Bending_ISPC_Enabled = CHAOS_BENDING_ISPC_ENABLED_DEFAULT;
 FAutoConsoleVariableRef CVarChaosBendingISPCEnabled(TEXT("p.Chaos.Bending.ISPC"), bChaos_Bending_ISPC_Enabled, TEXT("Whether to use ISPC optimizations in Bending constraints"));
 
 bool bChaos_Bending_ISPCUseHWGather = false;
-FAutoConsoleVariableRef CVarChaosTriangleMeshISPCUseHWGather(TEXT("p.Chaos.Bending.ISPCUseHWGather"), bChaos_Bending_ISPCUseHWGather, TEXT("Use (slower) HW gather when computing Bending constraints to see if this fixes the ISPC crash."));
+FAutoConsoleVariableRef CVarChaosBendingISPCUseHWGather(TEXT("p.Chaos.Bending.ISPCUseHWGather"), bChaos_Bending_ISPCUseHWGather, TEXT("Use (slower) HW gather when computing Bending constraints to see if this fixes the ISPC crash."));
 #endif
 
 static_assert(sizeof(ispc::FVector4f) == sizeof(Chaos::Softs::FPAndInvM), "sizeof(ispc::FVector4f) != sizeof(Chaos::Softs::FPAndInvM");
