Geometry collection : Fix bug where transforms are not updated properly if one of the two removal feature is disabled

diff --git a/Engine/Source/Runtime/Experimental/GeometryCollectionEngine/Private/GeometryCollection/GeometryCollectionComponent.cpp b/Engine/Source/Runtime/Experimental/GeometryCollectionEngine/Private/GeometryCollection/GeometryCollectionComponent.cpp
--- a/Engine/Source/Runtime/Experimental/GeometryCollectionEngine/Private/GeometryCollection/GeometryCollectionComponent.cpp
+++ b/Engine/Source/Runtime/Experimental/GeometryCollectionEngine/Private/GeometryCollection/GeometryCollectionComponent.cpp
@@ -4412,7 +4412,8 @@ void UGeometryCollectionComponent::UpdateRemovalIfNeeded()
 
 	// if removal is enabled, update the dynamic collection transform based on the decay 
 	// todo: we could optimize this using a list of transform to update from when we update the decay values
-	if (DynamicCollection && bAllowRemovalOnBreak && bAllowRemovalOnSleep && PhysicsProxy)
+	const bool bAllowAnyRemoval = (bAllowRemovalOnBreak || bAllowRemovalOnSleep);
+	if (DynamicCollection && bAllowAnyRemoval && PhysicsProxy)
 	{
 		FGeometryCollectionDecayDynamicFacade DecayFacade(*DynamicCollection);
 		if (DecayFacade.IsValid())
