Geometry collection : fix internal clusters created on first break to fail to crumble properly using remove on break
- This is caused by the fact that we create the runtime remove on break data only when we need it
- But this happens after the root is broken and HasChilden called on the dynamic collection will show that the root has no children , failing to aknowledge that it was cluster
- The solution is to read the cluster information from the rest collection

diff --git a/Engine/Source/Runtime/Experimental/Chaos/Private/GeometryCollectionProxyData.cpp b/Engine/Source/Runtime/Experimental/Chaos/Private/GeometryCollectionProxyData.cpp
--- a/Engine/Source/Runtime/Experimental/Chaos/Private/GeometryCollectionProxyData.cpp
+++ b/Engine/Source/Runtime/Experimental/Chaos/Private/GeometryCollectionProxyData.cpp
@@ -123,6 +123,14 @@ bool FTransformDynamicCollection::HasChildren(int32 Index) const
 	return false;
 }
 
+bool FTransformDynamicCollection::IsCluster(int32 Index) const
+{
+	if (RestCollectionShared && RestCollectionShared->Children.IsValidIndex(Index))
+	{
+		return (RestCollectionShared->Children[Index].Num() > 0);
+	}
+	return false;
+}
 
 /*
 * FGeometryDynamicCollection (FTransformDynamicCollection)
diff --git a/Engine/Source/Runtime/Experimental/Chaos/Public/GeometryCollectionProxyData.h b/Engine/Source/Runtime/Experimental/Chaos/Public/GeometryCollectionProxyData.h
--- a/Engine/Source/Runtime/Experimental/Chaos/Public/GeometryCollectionProxyData.h
+++ b/Engine/Source/Runtime/Experimental/Chaos/Public/GeometryCollectionProxyData.h
@@ -42,6 +42,9 @@ public:
 
 	CHAOS_API bool HasChildren(int32 Index) const;
 
+	/** check if the transform is a cluster - ie has originally children as per rest collection */
+	CHAOS_API bool IsCluster(int32 Index) const;
+
 	template<typename Lambda>
 	void IterateThroughChildren(int32 Index, Lambda&& LambdaIt) const
 	{
diff --git a/Engine/Source/Runtime/Experimental/GeometryCollectionEngine/Private/GeometryCollection/GeometryCollectionEngineRemoval.cpp b/Engine/Source/Runtime/Experimental/GeometryCollectionEngine/Private/GeometryCollection/GeometryCollectionEngineRemoval.cpp
--- a/Engine/Source/Runtime/Experimental/GeometryCollectionEngine/Private/GeometryCollection/GeometryCollectionEngineRemoval.cpp
+++ b/Engine/Source/Runtime/Experimental/GeometryCollectionEngine/Private/GeometryCollection/GeometryCollectionEngineRemoval.cpp
@@ -78,7 +78,7 @@ void FGeometryCollectionRemoveOnBreakDynamicFacade::SetAttributeValues(const Geo
 			const FVector2f RemovalTimer = RemoveOnBreakData.GetRemovalTimer();
 			const float MinRemovalTime = FMath::Max(0.0f, RemovalTimer.X);
 			const float MaxRemovalTime = FMath::Max(MinRemovalTime, RemovalTimer.Y);
-			const bool bIsCluster = DynamicCollection.HasChildren(Idx);
+			const bool bIsCluster = DynamicCollection.IsCluster(Idx);
 			const bool bUseClusterCrumbling = (bIsCluster && RemoveOnBreakData.GetClusterCrumbling());
 			BreakRemovalDuration[Idx] = bUseClusterCrumbling ? CrumblingRemovalTimer : Random.FRandRange(MinRemovalTime, MaxRemovalTime);
 		}
