Imporve the way we deal with type compatibility, especially around using anytypes

Anytypes with concrete types can now change under safe condition ( not connected and no connected type dependent connection )
Fix some bugs when connecting to anytype output where the UI would not reflect the change
node list menu now properly use the concrete type of the pulled pin to filter the nodes

diff --git a/Engine/Plugins/Experimental/Dataflow/Source/DataflowEditor/Private/Dataflow/DataflowSchema.cpp b/Engine/Plugins/Experimental/Dataflow/Source/DataflowEditor/Private/Dataflow/DataflowSchema.cpp
--- a/Engine/Plugins/Experimental/Dataflow/Source/DataflowEditor/Private/Dataflow/DataflowSchema.cpp
+++ b/Engine/Plugins/Experimental/Dataflow/Source/DataflowEditor/Private/Dataflow/DataflowSchema.cpp
@@ -101,11 +101,14 @@ bool UDataflowSchema::CanPinBeConnectedToNode(const UEdGraphPin* Pin, const UE::
 	{
 		// get the type from the dataflow input/output as the pin type may not be precise enough when using anytypes
 		TSharedPtr<const FDataflowNode> DataflowNode;
+		TSharedPtr<const UE::Dataflow::FGraph> DataflowGraph;
 		if (UDataflowEdNode* EdNode = Cast<UDataflowEdNode>(Pin->GetOwningNode()))
 		{
 			DataflowNode = EdNode->GetDataflowNode();
+			DataflowGraph = EdNode->GetDataflowGraph();
 		}
-		if (DataflowNode)
+
+		if (DataflowNode && DataflowGraph)
 		{
 			if (Pin->Direction == EEdGraphPinDirection::EGPD_Output)
 			{
@@ -114,10 +117,7 @@ bool UDataflowSchema::CanPinBeConnectedToNode(const UEdGraphPin* Pin, const UE::
 					const TArray<FDataflowInput*> Inputs = NodeParameters.DefaultNodeObject->GetInputs();
 					for (const FDataflowInput* Input : Inputs)
 					{
-						if (Input && (
-							Input->SupportsType(PinOutput->GetType()) || 
-							PinOutput->SupportsType(Input->GetType())
-							))
+						if (Input && DataflowGraph->CanConnect(*PinOutput, *Input))
 						{
 							return true;
 						}
@@ -131,10 +131,7 @@ bool UDataflowSchema::CanPinBeConnectedToNode(const UEdGraphPin* Pin, const UE::
 					const TArray<FDataflowOutput*> Outputs = NodeParameters.DefaultNodeObject->GetOutputs();
 					for (const FDataflowOutput* Output : Outputs)
 					{
-						if (Output && (
-							Output->SupportsType(PinInput->GetType()) ||
-							PinInput->SupportsType(Output->GetType())
-							))
+						if (Output && DataflowGraph->CanConnect(*Output, *PinInput))
 						{
 							return true;
 						}
@@ -216,6 +213,36 @@ bool HasLoopIfConnected(const UEdGraphNode* FromNode, const UEdGraphNode* ToNode
 	return false;
 }
 
+bool UDataflowSchema::CanConnectPins(const UEdGraphPin& OutputPin, const UEdGraphPin& InputPin)
+{
+	if (ensure(
+		OutputPin.Direction == EEdGraphPinDirection::EGPD_Output
+		&& InputPin.Direction == EEdGraphPinDirection::EGPD_Input
+		))
+	{
+		UDataflowEdNode* EdOutputNode = Cast<UDataflowEdNode>(OutputPin.GetOwningNode());
+		UDataflowEdNode* EdInputNode = Cast<UDataflowEdNode>(InputPin.GetOwningNode());
+		if (EdOutputNode && EdInputNode)
+		{
+			if (EdOutputNode->DataflowGraph && EdInputNode->DataflowGraph)
+			{
+				TSharedPtr<const FDataflowNode> OutputNode = EdOutputNode->GetDataflowNode();
+				TSharedPtr<const FDataflowNode> InputNode = EdInputNode->GetDataflowNode();
+				if (OutputNode && InputNode)
+				{
+					const FDataflowOutput* Output = OutputNode->FindOutput(OutputPin.PinName);
+					const FDataflowInput* Input = InputNode->FindInput(InputPin.PinName);
+					if (Output && Input)
+					{
+						return EdOutputNode->DataflowGraph->CanConnect(*Output, *Input);
+					}
+				}
+			}
+		}
+	}
+	return false;
+}
+
 const FPinConnectionResponse UDataflowSchema::CanCreateConnection(const UEdGraphPin* InPinA, const UEdGraphPin* InPinB) const
 {
 	bool bSwapped = false;
@@ -238,9 +265,7 @@ const FPinConnectionResponse UDataflowSchema::CanCreateConnection(const UEdGraph
 
 			if (EdNodeA && EdNodeB && (EdNodeA != EdNodeB))
 			{
-				const bool AIsCompatibleWithB = EdNodeA->PinIsCompatibleWithType(*PinA, PinB->PinType);
-				const bool BIsCompatibleWithA = EdNodeB->PinIsCompatibleWithType(*PinB, PinA->PinType);
-				if (AIsCompatibleWithB || BIsCompatibleWithA)
+				if (CanConnectPins(*PinA, *PinB))
 				{
 					// cycle checking on connect
 					if (!HasLoopIfConnected(EdNodeA, EdNodeB))
diff --git a/Engine/Plugins/Experimental/Dataflow/Source/DataflowEditor/Public/Dataflow/DataflowSchema.h b/Engine/Plugins/Experimental/Dataflow/Source/DataflowEditor/Public/Dataflow/DataflowSchema.h
--- a/Engine/Plugins/Experimental/Dataflow/Source/DataflowEditor/Public/Dataflow/DataflowSchema.h
+++ b/Engine/Plugins/Experimental/Dataflow/Source/DataflowEditor/Public/Dataflow/DataflowSchema.h
@@ -40,6 +40,7 @@ public:
 
 private:
 	static bool CanPinBeConnectedToNode(const UEdGraphPin* Pin, const UE::Dataflow::FFactoryParameters& NodeParameters);
+	static bool CanConnectPins(const UEdGraphPin& OutputPin, const UEdGraphPin& InputPin);
 };
 
 class FDataflowConnectionDrawingPolicy : public FConnectionDrawingPolicy, public FGCObject
diff --git a/Engine/Source/Runtime/Experimental/Dataflow/Core/Private/Dataflow/DataflowConnection.cpp b/Engine/Source/Runtime/Experimental/Dataflow/Core/Private/Dataflow/DataflowConnection.cpp
--- a/Engine/Source/Runtime/Experimental/Dataflow/Core/Private/Dataflow/DataflowConnection.cpp
+++ b/Engine/Source/Runtime/Experimental/Dataflow/Core/Private/Dataflow/DataflowConnection.cpp
@@ -4,6 +4,7 @@
 
 #include "Dataflow/DataflowNodeParameters.h"
 #include "Dataflow/DataflowNode.h"
+#include "Dataflow/DataflowAnyTypeRegistry.h"
 
 #include UE_INLINE_GENERATED_CPP_BY_NAME(DataflowConnection)
 
@@ -12,6 +13,7 @@ const UE::Dataflow::FPin UE::Dataflow::FPin::InvalidPin = { UE::Dataflow::FPin::
 
 FDataflowConnection::FDataflowConnection(UE::Dataflow::FPin::EDirection InDirection, FName InType, FName InName, FDataflowNode* InOwningNode, const FProperty* InProperty, FGuid InGuid)
 	: Type(InType)
+	, OriginalType(InType)
 	, Name(InName)
 	, OwningNode(InOwningNode)
 	, Property(InProperty)
@@ -24,6 +26,7 @@ FDataflowConnection::FDataflowConnection(UE::Dataflow::FPin::EDirection InDirect
 
 FDataflowConnection::FDataflowConnection(UE::Dataflow::FPin::EDirection InDirection, const UE::Dataflow::FConnectionParameters& Params)
 	: Type(Params.Type)
+	, OriginalType(Params.Type)
 	, Name(Params.Name)
 	, OwningNode(Params.Owner)
 	, Property(Params.Property)
@@ -35,21 +38,23 @@ FDataflowConnection::FDataflowConnection(UE::Dataflow::FPin::EDirection InDirect
 }
 
 void FDataflowConnection::InitFromType()
-{
-	bIsAnyType = false;
-	bHasConcreteType = true;
+{	
+	bIsAnyType = FDataflowConnection::IsAnyType(Type);
+	bHasConcreteType = !bIsAnyType;
+#if WITH_EDITOR
 	if (Property && Property->GetClass()->IsChildOf(FStructProperty::StaticClass()))
 	{
 		if (const FStructProperty* StructProperty = CastField<FStructProperty>(Property))
 		{
-			if (StructProperty->Struct->IsChildOf<FDataflowAnyType>())
-			{
-				Type = FDataflowAnyType::TypeName;
-				bIsAnyType = true;
-				bHasConcreteType = false;
-			}
+			const bool bInheritsFromAnyType = StructProperty->Struct->IsChildOf<FDataflowAnyType>();
+			ensure(bIsAnyType == bInheritsFromAnyType);
 		}
 	}
+	else
+	{
+		ensure(bIsAnyType == false);
+	}
+#endif
 }
 
 bool FDataflowConnection::IsOwningNodeEnabled() const
@@ -74,7 +79,7 @@ uint32 FDataflowConnection::GetOwningNodeValueHash() const
 
 bool FDataflowConnection::IsAnyType(const FName& InType)
 {
-	return (InType == FDataflowAnyType::TypeName);
+	return UE::Dataflow::FAnyTypesRegistry::IsAnyTypeStatic(InType);
 }
 
 void FDataflowConnection::SetAsAnyType(bool bAnyType, const FName& ConcreteType)
@@ -89,12 +94,13 @@ void FDataflowConnection::SetAsAnyType(bool bAnyType, const FName& ConcreteType)
 
 bool FDataflowConnection::SupportsType(FName InType) const
 {
+	// Incoming anytypes are not supported ( they need to be concrete types )
 	if (IsAnyType(InType))
 	{
 		return false;
 	}
 	// resort to policy only if the concrete type is not defined ( case of anytype connection )
-	if (bIsAnyType && !bHasConcreteType)
+	if (bIsAnyType)
 	{
 		return TypePolicy ? TypePolicy->SupportsType(InType) : true;
 	}
@@ -116,15 +122,67 @@ bool FDataflowConnection::IsExtendedType(FName InType) const
 	return (InType.ToString().StartsWith(Type.ToString() + "<"));
 }
 
+bool FDataflowConnection::IsSafeToTryChangingType() const
+{
+	if (!IsAnyType())
+	{
+		return false;
+	}
+	if (IsConnected())
+	{
+		return false;
+	}
+	if (IsAnytypeDependencyConnected())
+	{
+		return false;
+	}
+	return true;
+}
+
+bool FDataflowConnection::ResetToOriginalType()
+{
+	if (!IsAnyType())
+	{
+		// non any type simply have Type == OriginalType
+		return true;
+	}
+
+	if (IsSafeToTryChangingType())
+	{
+		Type = OriginalType;
+		bHasConcreteType = false;
+		return true;
+	}
+
+	return false;
+}
+
+bool FDataflowConnection::IsAnytypeDependencyConnected() const
+{
+	if (OwningNode && !TypeDependencyGroup.IsNone())
+	{
+		return OwningNode->IsAnytypeDependencyConnected(TypeDependencyGroup);
+	}
+	return false;
+}
+
 bool FDataflowConnection::SetConcreteType(FName InType)
 {
 	// Can only change from AnyType to a concrete type
 	if (Type != InType)
 	{
+		// special case when fixing types from Array to Array<
 		const bool bExtendedType = IsExtendedType(InType);
-		if (ensure(!bHasConcreteType || bExtendedType))
+		if (bExtendedType)
+		{
+			Type = InType;
+			bHasConcreteType = true;
+			return true;
+		}
+		// standard case: make sure we are safe to change and that this type is supported
+		if (ensure(IsSafeToTryChangingType()))
 		{
-			if (ensure(bExtendedType || SupportsType(InType)))
+			if (ensure(SupportsType(InType)))
 			{
 				Type = InType;
 				bHasConcreteType = true;
diff --git a/Engine/Source/Runtime/Experimental/Dataflow/Core/Private/Dataflow/DataflowGraph.cpp b/Engine/Source/Runtime/Experimental/Dataflow/Core/Private/Dataflow/DataflowGraph.cpp
--- a/Engine/Source/Runtime/Experimental/Dataflow/Core/Private/Dataflow/DataflowGraph.cpp
+++ b/Engine/Source/Runtime/Experimental/Dataflow/Core/Private/Dataflow/DataflowGraph.cpp
@@ -7,6 +7,7 @@
 #include "Dataflow/DataflowNodeFactory.h"
 #include "Logging/LogMacros.h"
 #include "Dataflow/DataflowArchive.h"
+#include "Dataflow/DataflowAnyTypeRegistry.h"
 #include "UObject/UE5MainStreamObjectVersion.h"
 #include "UObject/FortniteMainBranchObjectVersion.h"
 
@@ -111,40 +112,173 @@ namespace UE::Dataflow
 		}
 	}
 
+	namespace Private
+	{
+		FString GetConnectionFullName(const FDataflowConnection& Connection)
+		{
+			static const FName InvalidName("Invalid");
+
+			const FName NodeName = Connection.GetOwningNode()
+				? Connection.GetOwningNode()->GetName()
+				: InvalidName;
+			const FName ConnectionName = Connection.GetName();
+
+			return FString::Format(TEXT("{0}:{1}"), { NodeName.ToString(), ConnectionName.ToString() });
+		}
+
+		FLink MakeConnectionLink(FDataflowOutput& Output, FDataflowInput& Input)
+		{
+			return FLink(
+				Output.GetOwningNode()->GetGuid(),
+				Output.GetGuid(),
+				Input.GetOwningNode()->GetGuid(),
+				Input.GetGuid()
+			);
+		}
+	}
+
+	bool FGraph::CanConnect(const FDataflowOutput& Output, const FDataflowInput& Input) const
+	{
+		return (GetConnectType(Output, Input) != EConnectType::REJECTED);
+	}
+
+	FGraph::EConnectType FGraph::GetConnectType(const FDataflowOutput& Output, const FDataflowInput& Input) const
+	{
+		// Direct concrete type match 
+		const bool bAreStrictlyTheSameType = (Output.GetType() == Input.GetType());
+		const bool bBothHaveConcreteType = (Output.HasConcreteType() && Input.HasConcreteType());
+		if (bAreStrictlyTheSameType && bBothHaveConcreteType)
+		{
+			return EConnectType::DIRECT;
+		}
+
+		// both are unassigned anytypes 
+		const bool bIsInputUnassignedAnyType = Input.IsAnyType() && !Input.HasConcreteType();
+		const bool bIsOutputUnassignedAnyType = Output.IsAnyType() && !Output.HasConcreteType();
+		if (bIsInputUnassignedAnyType && bIsOutputUnassignedAnyType)
+		{
+			// todo(dataflow) currently unsupported by could be in the future by promoting both input and output 
+			//				  using their default values if available
+			return EConnectType::REJECTED;
+		}
+
+		// in case where both are anytypes but only one has concrete type 
+		// we want to favor the unassigned one to be promoted
+		const bool bFavorInputPromotion = 
+			(bIsInputUnassignedAnyType && Output.IsAnyType() && Output.HasConcreteType())
+			|| (Input.IsAnyType() && !Output.IsAnyType() && Output.HasConcreteType())
+			;
+
+		const bool bCanChangeInput = Input.IsAnyType() && Output.HasConcreteType();
+		if (bCanChangeInput)
+		{
+			// Already connected input or type dependencies must be rejected
+			// todo(dataflow) : the input being connected rule could be relaxed in the future 
+			//					because inputs are always connected to a single output
+			if (!Input.IsSafeToTryChangingType())
+			{
+				return EConnectType::REJECTED;
+			}
+
+			const bool bWouldBeCompatible = UE::Dataflow::FAnyTypesRegistry::AreTypesCompatibleStatic(Output.GetType(), Input.GetOriginalType());
+			if (bWouldBeCompatible)
+			{
+				return EConnectType::INPUT_PROMOTION;
+			}
+		}
+
+		const bool bCanChangeOutput = Output.IsAnyType() && Input.HasConcreteType();
+		if (Output.IsAnyType())
+		{
+			// Already connected output or type dependencies must be rejected
+			// because a output can be connected to multiple output 
+			if (!Output.IsSafeToTryChangingType())
+			{
+				return EConnectType::REJECTED;
+			}
+
+			const bool bWouldBeCompatible = UE::Dataflow::FAnyTypesRegistry::AreTypesCompatibleStatic(Input.GetType(), Output.GetOriginalType());
+			if (bWouldBeCompatible)
+			{
+				return EConnectType::OUTPUT_PROMOTION;
+			}
+		}
+
+		return EConnectType::REJECTED;
+
+	}
 
 	void FGraph::Connect(FDataflowOutput* OutputConnection, FDataflowInput* InputConnection)
 	{
 		if (ensure(OutputConnection && InputConnection))
 		{
-			FDataflowOutput* const OldOutputConnection = InputConnection->GetConnection();
-			if (OldOutputConnection != OutputConnection)
+			Connect(*OutputConnection, *InputConnection);
+		}
+	}
+
+	bool FGraph::Connect(FDataflowOutput& Output, FDataflowInput& Input)
+	{
+		const EConnectType ConnectType = GetConnectType(Output, Input);
+		if (ConnectType == EConnectType::REJECTED)
+		{
+			UE_LOG(LogChaosDataflow, Error, TEXT("FGraph::Connect(): failed to connect output [%s] from input [%s] - incompatible types"),
+				*Private::GetConnectionFullName(Output),
+				*Private::GetConnectionFullName(Input)
+			);
+			return false;
+		}
+
+		if (ConnectType == EConnectType::INPUT_PROMOTION)
+		{
+			UE_LOG(LogChaosDataflow, Verbose, TEXT("FGraph::Connect(): updating input [%s] type to match output [%s] type (%s)"),
+				*Private::GetConnectionFullName(Output),
+				*Private::GetConnectionFullName(Input),
+				*Output.GetType().ToString()
+			);
+			Input.ResetToOriginalType();
+			//Input.OwningNode->NotifyConnectionTypeChanged(&Input);
+			Input.SetConcreteType(Output.GetType());
+			Input.OwningNode->NotifyConnectionTypeChanged(&Input);
+		}
+		else if (ConnectType == EConnectType::OUTPUT_PROMOTION)
+		{
+			UE_LOG(LogChaosDataflow, Verbose, TEXT("FGraph::Connect(): updating output [%s] type to match input [%s] type (%s)"),
+				*Private::GetConnectionFullName(Output),
+				*Private::GetConnectionFullName(Input),
+				*Input.GetType().ToString()
+			);
+			Output.ResetToOriginalType();
+			//Output.OwningNode->NotifyConnectionTypeChanged(&Output);
+			Output.SetConcreteType(Input.GetType());
+			Output.OwningNode->NotifyConnectionTypeChanged(&Output);
+		}
+
+		FDataflowOutput* const OldOutput = Input.GetConnection();
+		if (OldOutput != &Output)
+		{
+			if (OldOutput)
 			{
-				if (OldOutputConnection)
-				{
-					UE_LOG(LogChaosDataflow, Verbose, TEXT("FGraph::Connect(): Disconnecting output [%s:%s] from input [%s:%s]"),
-						OldOutputConnection && OldOutputConnection->GetOwningNode() ? *OldOutputConnection->GetOwningNode()->GetName().ToString() : TEXT("Invalid"),
-						OldOutputConnection ? *OldOutputConnection->GetName().ToString() : TEXT("Invalid"),
-						InputConnection->GetOwningNode() ? *InputConnection->GetOwningNode()->GetName().ToString() : TEXT("Invalid"),
-						*InputConnection->GetName().ToString());
-					// Note: Do not remove the expired connection from the input to avoid an unnecessary invalidation.
-					//       Simply clobber it with calling AddConnection() on the input instead.
-					OldOutputConnection->RemoveConnection(InputConnection);
-					Connections.RemoveSwap(FLink(
-						OldOutputConnection->GetOwningNode()->GetGuid(), OldOutputConnection->GetGuid(),
-						InputConnection->GetOwningNode()->GetGuid(), InputConnection->GetGuid()));
-				}
-				UE_LOG(LogChaosDataflow, Verbose, TEXT("FGraph::Connect(): Connecting output [%s:%s] to input [%s:%s]"),
-					OutputConnection->GetOwningNode() ? *OutputConnection->GetOwningNode()->GetName().ToString() : TEXT("Invalid"),
-					*OutputConnection->GetName().ToString(),
-					InputConnection->GetOwningNode() ? *InputConnection->GetOwningNode()->GetName().ToString() : TEXT("Invalid"),
-					*InputConnection->GetName().ToString());
-				OutputConnection->AddConnection(InputConnection);
-				InputConnection->AddConnection(OutputConnection);
-				Connections.Add(FLink(
-					OutputConnection->GetOwningNode()->GetGuid(), OutputConnection->GetGuid(),
-					InputConnection->GetOwningNode()->GetGuid(), InputConnection->GetGuid()));
+				UE_LOG(LogChaosDataflow, Verbose, TEXT("FGraph::Connect(): Disconnecting output [%s] from input [%s]"), 
+					*Private::GetConnectionFullName(*OldOutput), 
+					*Private::GetConnectionFullName(Input)
+				);
+
+				// Note: Do not remove the expired connection from the input to avoid an unnecessary invalidation.
+				//       Simply clobber it with calling AddConnection() on the input instead.
+				OldOutput->RemoveConnection(&Input);
+				Connections.RemoveSwap(Private::MakeConnectionLink(*OldOutput, Input));
 			}
+
+			UE_LOG(LogChaosDataflow, Verbose, TEXT("FGraph::Connect(): Connecting output [%s] to input [%s]"),
+				*Private::GetConnectionFullName(Output),
+				*Private::GetConnectionFullName(Input)
+			);
+
+			Output.AddConnection(&Input);
+			Input.AddConnection(&Output);
+			Connections.Add(Private::MakeConnectionLink(Output, Input));
 		}
+		return true;
 	}
 
 	void FGraph::Disconnect(FDataflowOutput* OutputConnection, FDataflowInput* InputConnection)
diff --git a/Engine/Source/Runtime/Experimental/Dataflow/Core/Public/Dataflow/DataflowConnection.h b/Engine/Source/Runtime/Experimental/Dataflow/Core/Public/Dataflow/DataflowConnection.h
--- a/Engine/Source/Runtime/Experimental/Dataflow/Core/Public/Dataflow/DataflowConnection.h
+++ b/Engine/Source/Runtime/Experimental/Dataflow/Core/Public/Dataflow/DataflowConnection.h
@@ -133,6 +133,7 @@ struct FDataflowConnection
 
 protected:
 	FName Type;
+	FName OriginalType;
 	FName Name;
 	FDataflowNode* OwningNode = nullptr;
 	const FProperty* Property = nullptr;
@@ -172,10 +173,18 @@ protected:
 
 	/** 
 	* Set the concrete type of an anytype connection 
-	* @return true if the typewas effectivelly changed
+	* @return true if the type was effectivelly changed
 	*/
 	DATAFLOWCORE_API bool SetConcreteType(FName InType);
 
+	/**
+	* reset the output to its original type 
+	* no-op for non enytypes types
+	* Anytypes connection will fail if it is already connected or has connected type dependencies
+	* @return true if the change was made 
+	*/
+	DATAFLOWCORE_API bool ResetToOriginalType();
+
 public:
 	FDataflowConnection() {};
 	UE_DEPRECATED(5.5, "Deprecated constructor : use FConnectionParameters to pass parameters")
@@ -205,6 +214,7 @@ public:
 	}
 
 	FName GetType() const { return Type; }
+	FName GetOriginalType() const { return OriginalType; }
 
 	FGuid GetGuid() const { return Guid; }
 	void SetGuid(FGuid InGuid) { Guid = InGuid; }
@@ -215,6 +225,14 @@ public:
 	virtual void* RealAddress() const { ensure(OwningNode);  return (void*)((size_t)OwningNode + (size_t)GetOffset()); };
 	UE::Dataflow::FContextCacheKey CacheKey() const { return GetTypeHash(Guid); };
 
+	/* 
+	* check is this connection can chage type 
+	* case where it cannot are : 
+	*	it is already connected
+	*	it has copnnected type dependencies
+	*/
+	DATAFLOWCORE_API bool IsSafeToTryChangingType() const;
+
 	virtual bool IsConnected() const { return false; }
 	virtual bool AddConnection(FDataflowConnection* In) { return false; };
 	virtual bool RemoveConnection(FDataflowConnection* In) { return false; }
@@ -228,6 +246,8 @@ public:
 	DATAFLOWCORE_API FDataflowConnection& SetTypeDependencyGroup(FName DependencyGroupName);
 	DATAFLOWCORE_API FName GetTypeDependencyGroup() const { return TypeDependencyGroup;	}
 
+	DATAFLOWCORE_API bool IsAnytypeDependencyConnected() const;
+
 	template<class T>
 	bool IsA(const T* InVar) const
 	{
diff --git a/Engine/Source/Runtime/Experimental/Dataflow/Core/Public/Dataflow/DataflowGraph.h b/Engine/Source/Runtime/Experimental/Dataflow/Core/Public/Dataflow/DataflowGraph.h
--- a/Engine/Source/Runtime/Experimental/Dataflow/Core/Public/Dataflow/DataflowGraph.h
+++ b/Engine/Source/Runtime/Experimental/Dataflow/Core/Public/Dataflow/DataflowGraph.h
@@ -172,6 +172,18 @@ namespace UE::Dataflow
 		DATAFLOWCORE_API void ClearConnections(FDataflowInput* Input);
 		DATAFLOWCORE_API void ClearConnections(FDataflowOutput* Output);
 
+		enum class EConnectType : uint8
+		{
+			REJECTED = 0,
+			DIRECT, // both are already compatible
+			INPUT_PROMOTION, // input can be changed to adapt the output type
+			OUTPUT_PROMOTION, // output can be changed to adapt the input type
+		};
+
+		DATAFLOWCORE_API bool CanConnect(const FDataflowOutput& Output, const FDataflowInput& Input) const;
+		DATAFLOWCORE_API EConnectType GetConnectType(const FDataflowOutput& Output, const FDataflowInput& Input) const;
+		DATAFLOWCORE_API bool Connect(FDataflowOutput& Output, FDataflowInput& Input);
+
 		DATAFLOWCORE_API void Connect(FDataflowOutput* OutputConnection, FDataflowInput* InputConnection);
 		DATAFLOWCORE_API void Disconnect(FDataflowOutput* OutputConnection, FDataflowInput* InputConnection);
 
diff --git a/Engine/Source/Runtime/Experimental/Dataflow/Engine/Private/Dataflow/DataflowEdNode.cpp b/Engine/Source/Runtime/Experimental/Dataflow/Engine/Private/Dataflow/DataflowEdNode.cpp
--- a/Engine/Source/Runtime/Experimental/Dataflow/Engine/Private/Dataflow/DataflowEdNode.cpp
+++ b/Engine/Source/Runtime/Experimental/Dataflow/Engine/Private/Dataflow/DataflowEdNode.cpp
@@ -368,7 +368,7 @@ void UDataflowEdNode::PinConnectionListChanged(UEdGraphPin* Pin)
 					// Add any newly added connections
 					for (const UEdGraphPin* const LinkedCon : Pin->LinkedTo)
 					{
-						const UDataflowEdNode* const LinkedNode = Cast<UDataflowEdNode>(LinkedCon->GetOwningNode());
+						UDataflowEdNode* const LinkedNode = Cast<UDataflowEdNode>(LinkedCon->GetOwningNode());
 						if (ensure(LinkedNode && LinkedNode->IsBound()))
 						{
 							if (const TSharedPtr<FDataflowNode> LinkedDataflowNode = DataflowGraph->FindBaseNode(LinkedNode->GetDataflowNodeGuid()))
@@ -377,13 +377,9 @@ void UDataflowEdNode::PinConnectionListChanged(UEdGraphPin* Pin)
 								{
 									if (ConnectionInput->GetConnectedOutputs().Find(LinkedConOutput) == INDEX_NONE)
 									{
-										const bool bTypeChanged = DataflowNode->TrySetConnectionType(ConnectionInput, LinkedConOutput->GetType());
 										DataflowGraph->Connect(LinkedConOutput, ConnectionInput);
-
-										if (bTypeChanged)
-										{
-											UpdatePinsFromDataflowNode();
-										}
+										UpdatePinsFromDataflowNode();
+										LinkedNode->UpdatePinsFromDataflowNode();
 									}
 								}
 							}
