Fix DataflowSelection::Initialize in the case where the selection is not empty

diff --git a/Engine/Source/Runtime/Experimental/Dataflow/Core/Private/Dataflow/DataflowSelection.cpp b/Engine/Source/Runtime/Experimental/Dataflow/Core/Private/Dataflow/DataflowSelection.cpp
--- a/Engine/Source/Runtime/Experimental/Dataflow/Core/Private/Dataflow/DataflowSelection.cpp
+++ b/Engine/Source/Runtime/Experimental/Dataflow/Core/Private/Dataflow/DataflowSelection.cpp
@@ -6,8 +6,7 @@
 
 void FDataflowSelection::Initialize(int32 NumBits, bool Value) 
 { 
-	SelectionArray.AddUninitialized(NumBits); 
-	SelectionArray.SetRange(0, NumBits, Value); 
+	SelectionArray.Init(Value, NumBits); 
 }
 
 void FDataflowSelection::Initialize(const FDataflowSelection& Other)
