[AutoRTFM] Remove no-memory-validation for hot opens

The memory validator now automatically throttles when the validation is taking an excessive time, on a per-open basis.
This means we can now enable memory validation for hot code, and ensure that this is correct too.

diff --git a/Engine/Source/Runtime/Core/Private/Delegates/DelegateHandle.cpp b/Engine/Source/Runtime/Core/Private/Delegates/DelegateHandle.cpp
--- a/Engine/Source/Runtime/Core/Private/Delegates/DelegateHandle.cpp
+++ b/Engine/Source/Runtime/Core/Private/Delegates/DelegateHandle.cpp
@@ -15,8 +15,7 @@ uint64 FDelegateHandle::GenerateNewID()
 	// Just increment a counter to generate an ID.
 	uint64 Result = 0; // Initialize just to silence static analysis.
 	
-	// Memory validation disabled as this is hot code.
-	UE_AUTORTFM_OPEN_NO_VALIDATION
+	UE_AUTORTFM_OPEN
 	{
 		Result = ++UE::Delegates::Private::GNextID;
 
diff --git a/Engine/Source/Runtime/Core/Private/HAL/ThreadingBase.cpp b/Engine/Source/Runtime/Core/Private/HAL/ThreadingBase.cpp
--- a/Engine/Source/Runtime/Core/Private/HAL/ThreadingBase.cpp
+++ b/Engine/Source/Runtime/Core/Private/HAL/ThreadingBase.cpp
@@ -1458,8 +1458,7 @@ FTlsAutoCleanup* FThreadSingletonInitializer::Get( TFunctionRef<FTlsAutoCleanup*
 {
 	uint32 TlsSlot;
 	FTlsAutoCleanup* ThreadSingleton = nullptr;
-	// Memory validation disabled as this is hot code.
-	UE_AUTORTFM_OPEN_NO_VALIDATION
+	UE_AUTORTFM_OPEN
 	{
 		TlsSlot = (uint32)FPlatformAtomics::AtomicRead_Relaxed((int32*)&InOutTlsSlot);
 		if (TlsSlot == FPlatformTLS::InvalidTlsSlot)
diff --git a/Engine/Source/Runtime/Core/Private/UObject/UnrealNames.cpp b/Engine/Source/Runtime/Core/Private/UObject/UnrealNames.cpp
--- a/Engine/Source/Runtime/Core/Private/UObject/UnrealNames.cpp
+++ b/Engine/Source/Runtime/Core/Private/UObject/UnrealNames.cpp
@@ -3128,8 +3128,7 @@ private:
 	static FName MakeInternal(FNameStringView View, EFindName FindType, int32 InternalNumber)
 	{
 		FNameEntryIds Ids;
-		// Memory validation disabled as this is hot code.
-		UE_AUTORTFM_OPEN_NO_VALIDATION
+		UE_AUTORTFM_OPEN
 		{
 			Ids = FindOrStoreString(View, FindType);
 		};
diff --git a/Engine/Source/Runtime/Core/Private/Unix/UnixPlatformMisc.cpp b/Engine/Source/Runtime/Core/Private/Unix/UnixPlatformMisc.cpp
--- a/Engine/Source/Runtime/Core/Private/Unix/UnixPlatformMisc.cpp
+++ b/Engine/Source/Runtime/Core/Private/Unix/UnixPlatformMisc.cpp
@@ -1697,8 +1697,7 @@ void FUnixPlatformMisc::CreateGuid(FGuid& Result)
 {
 	if (AutoRTFM::IsTransactional())
 	{
-		// Memory validation disabled as this is hot code.
-		Result = AutoRTFM::Open<AutoRTFM::EMemoryValidationLevel::Disabled>([]
+		Result = AutoRTFM::Open([]
 		{
 			FGuid Guid;
 			CreateGuidImpl(Guid);
diff --git a/Engine/Source/Runtime/Core/Private/Windows/WindowsPlatformMisc.cpp b/Engine/Source/Runtime/Core/Private/Windows/WindowsPlatformMisc.cpp
--- a/Engine/Source/Runtime/Core/Private/Windows/WindowsPlatformMisc.cpp
+++ b/Engine/Source/Runtime/Core/Private/Windows/WindowsPlatformMisc.cpp
@@ -1422,8 +1422,7 @@ void FWindowsPlatformMisc::CreateGuid(FGuid& Result)
 {
 	if (AutoRTFM::IsTransactional())
 	{
-		// Memory validation disabled as this is hot code.
-		Result = AutoRTFM::Open<AutoRTFM::EMemoryValidationLevel::Disabled>([]
+		Result = AutoRTFM::Open([]
 		{
 			FGuid Guid;
 			verify( CoCreateGuid( (GUID*)&Guid )==S_OK );
diff --git a/Engine/Source/Runtime/Core/Public/HAL/FMemory.inl b/Engine/Source/Runtime/Core/Public/HAL/FMemory.inl
--- a/Engine/Source/Runtime/Core/Public/HAL/FMemory.inl
+++ b/Engine/Source/Runtime/Core/Public/HAL/FMemory.inl
@@ -22,8 +22,7 @@ FMEMORY_INLINE_FUNCTION_DECORATOR void* FMemory::Malloc(SIZE_T Count, uint32 Ali
 	// we don't end up keeping track of the writes to the allocator's internal data structures.
 	// This is because allocators are already transactional - malloc can be rolled back by
 	// calling free.
-	// Memory validation disabled as this is hot code.
-	void* Ptr = AutoRTFM::Open<AutoRTFM::EMemoryValidationLevel::Disabled>([Count, Alignment]
+	void* Ptr = AutoRTFM::Open([Count, Alignment]
 	{
 		void* Alloc = nullptr;
 		if (!FMEMORY_INLINE_GMalloc)
@@ -169,8 +168,7 @@ FMEMORY_INLINE_FUNCTION_DECORATOR void FMemory::Free(void* Original)
 FMEMORY_INLINE_FUNCTION_DECORATOR SIZE_T FMemory::GetAllocSize(void* Original)
 {
 	SIZE_T Result;
-	// Memory validation disabled as this is hot code.
-	UE_AUTORTFM_OPEN_NO_VALIDATION
+	UE_AUTORTFM_OPEN
 	{
 		if (!FMEMORY_INLINE_GMalloc)
 		{
@@ -237,8 +235,7 @@ FMEMORY_INLINE_FUNCTION_DECORATOR void* FMemory::MallocZeroed(SIZE_T Count, uint
 FMEMORY_INLINE_FUNCTION_DECORATOR SIZE_T FMemory::QuantizeSize(SIZE_T Count, uint32 Alignment)
 {
 	SIZE_T Result;
-	// Memory validation disabled as this is hot code.
-	UE_AUTORTFM_OPEN_NO_VALIDATION
+	UE_AUTORTFM_OPEN
 	{
 		if (!FMEMORY_INLINE_GMalloc)
 		{
diff --git a/Engine/Source/Runtime/Core/Public/Misc/MTTransactionallySafeAccessDetector.h b/Engine/Source/Runtime/Core/Public/Misc/MTTransactionallySafeAccessDetector.h
--- a/Engine/Source/Runtime/Core/Public/Misc/MTTransactionallySafeAccessDetector.h
+++ b/Engine/Source/Runtime/Core/Public/Misc/MTTransactionallySafeAccessDetector.h
@@ -85,8 +85,7 @@ private:
 	{
 		if (AutoRTFM::IsTransactional() || AutoRTFM::IsCommittingOrAborting())
 		{
-			// Memory validation disabled as this is hot code.
-			const bool bResult = AutoRTFM::Open<AutoRTFM::EMemoryValidationLevel::Disabled>([&]
+			const bool bResult = AutoRTFM::Open([&]
 				{
 					// The transactional system which can increment the lock counts
 					// is always single-threaded, thus this is safe to check without atomicity.
@@ -151,8 +150,7 @@ private:
 	{
 		if (AutoRTFM::IsTransactional() || AutoRTFM::IsCommittingOrAborting())
 		{
-			// Memory validation disabled as this is hot code.
-			const bool bResult = AutoRTFM::Open<AutoRTFM::EMemoryValidationLevel::Disabled>([&]
+			const bool bResult = AutoRTFM::Open([&]
 				{
 					if ((0 == State->ReadLockCount) && (0 == State->WriteLockCount))
 					{
diff --git a/Engine/Source/Runtime/Core/Public/Misc/TransactionallySafeCriticalSection.h b/Engine/Source/Runtime/Core/Public/Misc/TransactionallySafeCriticalSection.h
--- a/Engine/Source/Runtime/Core/Public/Misc/TransactionallySafeCriticalSection.h
+++ b/Engine/Source/Runtime/Core/Public/Misc/TransactionallySafeCriticalSection.h
@@ -72,8 +72,7 @@ struct FTransactionallySafeCriticalSectionDefinition final
 	{
 		if (AutoRTFM::IsTransactional() || AutoRTFM::IsCommittingOrAborting())
 		{
-			// Memory validation disabled as this is hot code.
-			UE_AUTORTFM_OPEN_NO_VALIDATION
+			UE_AUTORTFM_OPEN
 				{
 					// The transactional system which can increment TransactionalLockCount
 					// is always single-threaded, thus this is safe to check without atomicity.
diff --git a/Engine/Source/Runtime/Core/Public/Misc/TransactionallySafeRWLock.h b/Engine/Source/Runtime/Core/Public/Misc/TransactionallySafeRWLock.h
--- a/Engine/Source/Runtime/Core/Public/Misc/TransactionallySafeRWLock.h
+++ b/Engine/Source/Runtime/Core/Public/Misc/TransactionallySafeRWLock.h
@@ -76,8 +76,7 @@ struct FTransactionallySafeRWLockDefinition
 		}
 
 		// As the State was constructed in the open, it must be released in the open.
-		// Memory validation disabled as this is hot code.
-		UE_AUTORTFM_OPEN_NO_VALIDATION{ State = nullptr; };
+		UE_AUTORTFM_OPEN{ State = nullptr; };
 	}
 
 	void ReadLock()
@@ -112,8 +111,7 @@ struct FTransactionallySafeRWLockDefinition
 	{
 		if (AutoRTFM::IsTransactional() || AutoRTFM::IsCommittingOrAborting())
 		{
-			// Memory validation disabled as this is hot code.
-			UE_AUTORTFM_OPEN_NO_VALIDATION
+			UE_AUTORTFM_OPEN
 				{
 					// The transactional system which can increment TransactionalLockCount
 					// is always single-threaded, thus this is safe to check without atomicity.
@@ -163,8 +161,7 @@ struct FTransactionallySafeRWLockDefinition
 		{
 			bool Ret = false;
 
-			// Memory validation disabled as this is hot code.
-			UE_AUTORTFM_OPEN_NO_VALIDATION
+			UE_AUTORTFM_OPEN
 				{
 					// The transactional system which can increment TransactionalLockCount
 					// is always single-threaded, thus this is safe to check without atomicity.
diff --git a/Engine/Source/Runtime/Core/Public/ProfilingDebugging/CsvProfiler.h b/Engine/Source/Runtime/Core/Public/ProfilingDebugging/CsvProfiler.h
--- a/Engine/Source/Runtime/Core/Public/ProfilingDebugging/CsvProfiler.h
+++ b/Engine/Source/Runtime/Core/Public/ProfilingDebugging/CsvProfiler.h
@@ -565,16 +565,14 @@ public:
 		: StatName(InStatName)
 		, CategoryIndex(InCategoryIndex)
 	{
-		// Memory validation disabled as this is hot code.
-		UE_AUTORTFM_OPEN_NO_VALIDATION{ FCsvProfiler::BeginStat(StatName, CategoryIndex, InNamedEventName); };
+		UE_AUTORTFM_OPEN{ FCsvProfiler::BeginStat(StatName, CategoryIndex, InNamedEventName); };
 		AutoRTFM::PushOnAbortHandler(this, [SN = StatName, CI = CategoryIndex](){ FCsvProfiler::EndStat(SN, CI); });
 	}
 
 	~FScopedCsvStat()
 	{
 		AutoRTFM::PopOnAbortHandler(this);
-		// Memory validation disabled as this is hot code.
-		UE_AUTORTFM_OPEN_NO_VALIDATION{ FCsvProfiler::EndStat(StatName, CategoryIndex); };
+		UE_AUTORTFM_OPEN{ FCsvProfiler::EndStat(StatName, CategoryIndex); };
 	}
 	const char * StatName;
 	uint32 CategoryIndex;
@@ -613,16 +611,14 @@ public:
 	FScopedCsvStatExclusive(const char * InStatName, const char* InNamedEventName = nullptr)
 		: StatName(InStatName)
 	{
-		// Memory validation disabled as this is hot code.
-		UE_AUTORTFM_OPEN_NO_VALIDATION{ FCsvProfiler::BeginExclusiveStat(StatName, InNamedEventName); };
+		UE_AUTORTFM_OPEN{ FCsvProfiler::BeginExclusiveStat(StatName, InNamedEventName); };
 		AutoRTFM::PushOnAbortHandler(this, [SN = StatName](){ FCsvProfiler::EndExclusiveStat(SN); });
 	}
 
 	~FScopedCsvStatExclusive()
 	{
 		AutoRTFM::PopOnAbortHandler(this);
-		// Memory validation disabled as this is hot code.
-		UE_AUTORTFM_OPEN_NO_VALIDATION{ FCsvProfiler::EndExclusiveStat(StatName); };
+		UE_AUTORTFM_OPEN{ FCsvProfiler::EndExclusiveStat(StatName); };
 	}
 	const char * StatName;
 };
diff --git a/Engine/Source/Runtime/Core/Public/Stats/LightweightStats.h b/Engine/Source/Runtime/Core/Public/Stats/LightweightStats.h
--- a/Engine/Source/Runtime/Core/Public/Stats/LightweightStats.h
+++ b/Engine/Source/Runtime/Core/Public/Stats/LightweightStats.h
@@ -69,7 +69,7 @@ namespace UE::Stats::Private
 		{
 			if (DerivedScopeCounterType::CanEmitStat(InStatId))
 			{
-				UE_AUTORTFM_OPEN_NO_VALIDATION
+				UE_AUTORTFM_OPEN
 				{
 					DerivedScopeCounterType::BeginNamedEvent(InStatId.GetName());
 					bPop = true;
@@ -82,8 +82,7 @@ namespace UE::Stats::Private
 		{
 			AutoRTFM::PopOnAbortHandler(this);
 
-			// Memory validation disabled as this is hot code.
-			UE_AUTORTFM_OPEN_NO_VALIDATION
+			UE_AUTORTFM_OPEN
 			{
 				if (bPop)
 				{
diff --git a/Engine/Source/Runtime/Core/Public/Stats/StatsSystemTypes.h b/Engine/Source/Runtime/Core/Public/Stats/StatsSystemTypes.h
--- a/Engine/Source/Runtime/Core/Public/Stats/StatsSystemTypes.h
+++ b/Engine/Source/Runtime/Core/Public/Stats/StatsSystemTypes.h
@@ -1626,8 +1626,7 @@ public:
 	 */
 	FORCEINLINE_STATS FScopeCycleCounter(TStatId StatId, EStatFlags StatFlags, bool bAlways = false, UE::FSourceLocation SourceLocation = UE::FSourceLocation::Current())
 	{
-		// Memory validation disabled as this is hot code.
-		UE_AUTORTFM_OPEN_NO_VALIDATION
+		UE_AUTORTFM_OPEN
 		{
 			Start(StatId, StatFlags, bAlways, SourceLocation);
 		};
@@ -1645,8 +1644,7 @@ public:
 	FORCEINLINE_STATS ~FScopeCycleCounter()
 	{
 		AutoRTFM::PopOnAbortHandler(this);
-		// Memory validation disabled as this is hot code.
-		UE_AUTORTFM_OPEN_NO_VALIDATION
+		UE_AUTORTFM_OPEN
 		{
 			Stop();
 		};
diff --git a/Engine/Source/Runtime/Core/Public/Templates/RefCounting.h b/Engine/Source/Runtime/Core/Public/Templates/RefCounting.h
--- a/Engine/Source/Runtime/Core/Public/Templates/RefCounting.h
+++ b/Engine/Source/Runtime/Core/Public/Templates/RefCounting.h
@@ -137,8 +137,7 @@ public:
 			// though we are deferring reference count changes until commit time.
 			// Be advised that GetRefCount() would reveal our trickery, since it
 			// always returns the true refcount.
-			// Memory validation disabled as this is hot code.
-			UE_AUTORTFM_OPEN_NO_VALIDATION
+			UE_AUTORTFM_OPEN
 			{
 				RefsToReturn = uint32(NumRefs.load(std::memory_order_relaxed)) - 1;
 			};
@@ -162,8 +161,7 @@ public:
 		// Inside of an AutoRTFM transaction, the returned refcount value may be higher than you'd
 		// expect, because all Release calls are deferred until the transaction commit time.
 		uint32 Refs;
-		// Memory validation disabled as this is hot code.
-		UE_AUTORTFM_OPEN_NO_VALIDATION
+		UE_AUTORTFM_OPEN
 		{
 			Refs = uint32(NumRefs.load(std::memory_order_relaxed));
 		};
@@ -313,8 +311,7 @@ public:
 
 		if constexpr (Mode == ERefCountingMode::ThreadSafe)
 		{
-			// Memory validation disabled as this is hot code.
-			UE_AUTORTFM_OPEN_NO_VALIDATION
+			UE_AUTORTFM_OPEN
 			{
 				Refs = RefCount++;
 			};
@@ -394,8 +391,7 @@ public:
 
 			uint32 Count = 0;
 
-			// Memory validation disabled as this is hot code.
-			UE_AUTORTFM_OPEN_NO_VALIDATION
+			UE_AUTORTFM_OPEN
 			{
 				// This reference count may be accessed by multiple threads
 				Count = RefCount.load(std::memory_order_relaxed);
diff --git a/Engine/Source/Runtime/Core/Public/Templates/SharedPointerInternals.h b/Engine/Source/Runtime/Core/Public/Templates/SharedPointerInternals.h
--- a/Engine/Source/Runtime/Core/Public/Templates/SharedPointerInternals.h
+++ b/Engine/Source/Runtime/Core/Public/Templates/SharedPointerInternals.h
@@ -81,8 +81,7 @@ namespace SharedPointerInternals
 
 				int32 Count = 0;
 
-				// Memory validation disabled as this is hot code.
-				UE_AUTORTFM_OPEN_NO_VALIDATION
+				UE_AUTORTFM_OPEN
 					{
 						// This reference count may be accessed by multiple threads
 						Count = SharedReferenceCount.load(std::memory_order_relaxed);
@@ -112,16 +111,14 @@ namespace SharedPointerInternals
 				// in response to the increment, so there's nothing to order with.
 
 #if defined(_MSC_VER) && (defined(_M_X64) || defined(_M_IX86))
-				// Memory validation disabled as this is hot code.
-				UE_AUTORTFM_OPEN_NO_VALIDATION
+				UE_AUTORTFM_OPEN
 				{
 					// We do a regular SC increment here because it maps to an _InterlockedIncrement (lock inc).
 					// The codegen for a relaxed fetch_add is actually much worse under MSVC (lock xadd).
 					++SharedReferenceCount;
 				};
 #else
-				// Memory validation disabled as this is hot code.
-				UE_AUTORTFM_OPEN_NO_VALIDATION
+				UE_AUTORTFM_OPEN
 				{
 					SharedReferenceCount.fetch_add(1, std::memory_order_relaxed);
 				};
@@ -150,8 +147,7 @@ namespace SharedPointerInternals
 			{
 				bool bSucceeded = false;
 
-				// Memory validation disabled as this is hot code.
-				UE_AUTORTFM_OPEN_NO_VALIDATION
+				UE_AUTORTFM_OPEN
 				{
 					// See AddSharedReference for the same reasons that std::memory_order_relaxed is used in this function.
 
@@ -260,16 +256,14 @@ namespace SharedPointerInternals
 				// See AddSharedReference for the same reasons that std::memory_order_relaxed is used in this function.
 
 #if defined(_MSC_VER) && (defined(_M_X64) || defined(_M_IX86))
-				// Memory validation disabled as this is hot code.
-				UE_AUTORTFM_OPEN_NO_VALIDATION
+				UE_AUTORTFM_OPEN
 					{
 						// We do a regular SC increment here because it maps to an _InterlockedIncrement (lock inc).
 						// The codegen for a relaxed fetch_add is actually much worse under MSVC (lock xadd).
 						++WeakReferenceCount;
 					};
 #else
-				// Memory validation disabled as this is hot code.
-				UE_AUTORTFM_OPEN_NO_VALIDATION
+				UE_AUTORTFM_OPEN
 					{
 						WeakReferenceCount.fetch_add(1, std::memory_order_relaxed);
 					};
diff --git a/Engine/Source/Runtime/CoreUObject/Private/UObject/Class.cpp b/Engine/Source/Runtime/CoreUObject/Private/UObject/Class.cpp
--- a/Engine/Source/Runtime/CoreUObject/Private/UObject/Class.cpp
+++ b/Engine/Source/Runtime/CoreUObject/Private/UObject/Class.cpp
@@ -6801,8 +6801,7 @@ UFunction* UClass::FindFunctionByName(FName InName, EIncludeSuperFlag::Type Incl
 
 	UFunction* Result = nullptr;
 
-	// Memory validation disabled as this is hot code.
-	UE_AUTORTFM_OPEN_NO_VALIDATION
+	UE_AUTORTFM_OPEN
 	{
 		UClass* SuperClass = GetSuperClass();
 		if (IncludeSuper == EIncludeSuperFlag::ExcludeSuper || ( Interfaces.Num() == 0 && SuperClass == nullptr ) )
diff --git a/Engine/Source/Runtime/CoreUObject/Private/UObject/CoreRedirectsContext.cpp b/Engine/Source/Runtime/CoreUObject/Private/UObject/CoreRedirectsContext.cpp
--- a/Engine/Source/Runtime/CoreUObject/Private/UObject/CoreRedirectsContext.cpp
+++ b/Engine/Source/Runtime/CoreUObject/Private/UObject/CoreRedirectsContext.cpp
@@ -130,8 +130,7 @@ FScopeCoreRedirectsContext::~FScopeCoreRedirectsContext()
 uint32 FCoreRedirectsContext::FRWLockWithExclusiveRecursion::LoadWriteLockOwnerThreadIdRelaxed() const
 {
 	uint32 ThreadId;
-	// Memory validation is disabled, as this is a hot code path.
-	UE_AUTORTFM_OPEN_NO_VALIDATION
+	UE_AUTORTFM_OPEN
 	{
 		ThreadId = WriteLockOwnerThreadId.load(std::memory_order_relaxed);
 	};
@@ -142,11 +141,10 @@ void FCoreRedirectsContext::FRWLockWithExclusiveRecursion::StoreWriteLockOwnerTh
 {
 	if (AutoRTFM::IsClosed())
 	{
-		// Memory validation is disabled, as this is a hot code path.
 		// If an abort occurs, we want to restore the variable's previous value, so we use
 		// `RecordOpenWrite` here so that AutoRTFM can undo our work. (We aren't overly
 		// concerned with atomicity as all loads are done in `memory_order_relaxed`.)
-		UE_AUTORTFM_OPEN_NO_VALIDATION
+		UE_AUTORTFM_OPEN
 		{
 			AutoRTFM::RecordOpenWrite(&ThreadId);
 			WriteLockOwnerThreadId.store(ThreadId, Order);
diff --git a/Engine/Source/Runtime/CoreUObject/Private/UObject/ScriptCore.cpp b/Engine/Source/Runtime/CoreUObject/Private/UObject/ScriptCore.cpp
--- a/Engine/Source/Runtime/CoreUObject/Private/UObject/ScriptCore.cpp
+++ b/Engine/Source/Runtime/CoreUObject/Private/UObject/ScriptCore.cpp
@@ -193,8 +193,7 @@ FBlueprintContext* FBlueprintContextGetThreadSingletonImpl()
 FBlueprintContext* FBlueprintContext::GetThreadSingleton()
 {
 	FBlueprintContext* Result;
-	// Memory validation disabled as this is hot code.
-	UE_AUTORTFM_OPEN_NO_VALIDATION{ Result = FBlueprintContextGetThreadSingletonImpl(); };
+	UE_AUTORTFM_OPEN{ Result = FBlueprintContextGetThreadSingletonImpl(); };
 	return Result;
 }
 
diff --git a/Engine/Source/Runtime/CoreUObject/Private/UObject/UObjectHash.cpp b/Engine/Source/Runtime/CoreUObject/Private/UObject/UObjectHash.cpp
--- a/Engine/Source/Runtime/CoreUObject/Private/UObject/UObjectHash.cpp
+++ b/Engine/Source/Runtime/CoreUObject/Private/UObject/UObjectHash.cpp
@@ -536,8 +536,7 @@ public:
 		if (!(IsGarbageCollectingAndLockingUObjectHashTables() && IsInGameThread()))
 		{
 			Tables = &InTables;
-			// Memory validation disabled as this is hot code.
-			UE_AUTORTFM_OPEN_NO_VALIDATION{ InTables.Lock(); };
+			UE_AUTORTFM_OPEN{ InTables.Lock(); };
 			AutoRTFM::PushOnAbortHandler(this, [this](){ if (this->Tables){ this->Tables->Unlock(); }});
 		}
 		else
@@ -553,8 +552,7 @@ public:
 #if THREADSAFE_UOBJECTS
 		if (Tables)
 		{
-			// Memory validation disabled as this is hot code.
-			UE_AUTORTFM_OPEN_NO_VALIDATION{ Tables->Unlock(); };
+			UE_AUTORTFM_OPEN{ Tables->Unlock(); };
 			AutoRTFM::PopOnAbortHandler(this);
 		}
 #endif
@@ -884,8 +882,7 @@ UObject* StaticFindObjectFastInternal(const UClass* ObjectClass, const UObject*
 	UObject* Result = nullptr;
 
 	// Transactionally, a static find operation has to occur in the open as it touches shared state.
-	// Memory validation disabled as this is hot code.
-	UE_AUTORTFM_OPEN_NO_VALIDATION
+	UE_AUTORTFM_OPEN
 	{
 		INC_DWORD_STAT(STAT_FindObjectFast);
 
diff --git a/Engine/Source/Runtime/CoreUObject/Public/UObject/Stack.h b/Engine/Source/Runtime/CoreUObject/Public/UObject/Stack.h
--- a/Engine/Source/Runtime/CoreUObject/Public/UObject/Stack.h
+++ b/Engine/Source/Runtime/CoreUObject/Public/UObject/Stack.h
@@ -35,8 +35,7 @@ FORCEINLINE void* UeVstackAllocHelper(FVirtualStackAllocator* Allocator, size_t
 	{
 		// in transactional code we redirect the virtual stack allocator to FMemory::Malloc
 		// until we properly implement a transaction-friendly stack allocator
-		// Memory validation disabled as this is hot code.
-		Result = AutoRTFM::Open<AutoRTFM::EMemoryValidationLevel::Disabled>([Size, Align]
+		Result = AutoRTFM::Open([Size, Align]
 		{
 			return FMemory::Malloc(Size, Align);
 		});
diff --git a/Engine/Source/Runtime/CoreUObject/Public/UObject/UObjectArray.h b/Engine/Source/Runtime/CoreUObject/Public/UObject/UObjectArray.h
--- a/Engine/Source/Runtime/CoreUObject/Public/UObject/UObjectArray.h
+++ b/Engine/Source/Runtime/CoreUObject/Public/UObject/UObjectArray.h
@@ -172,8 +172,7 @@ public:
 	{
 		checkf((int32(FlagToClear) & ~int32(EInternalObjectFlags_AllFlags)) == 0, TEXT("%d is not a valid internal flag value"), int32(FlagToClear));
 		bool Result = false;
-		// Memory validation disabled as this is hot code.
-		UE_AUTORTFM_OPEN_NO_VALIDATION
+		UE_AUTORTFM_OPEN
 		{
 			FlagToClear &= ~EInternalObjectFlags_ReachabilityFlags; // reachability flags can only be cleared by GC through *_ForGC functions
 			FlagToClear &= ~EInternalObjectFlags::RefCounted; // refcounted flag is internal and must only be cleared internally by AddRef/ReleaseRef.
@@ -210,8 +209,7 @@ public:
 	{
 		checkf((int32(FlagToSet) & ~int32(EInternalObjectFlags_AllFlags)) == 0, TEXT("%d is not a valid internal flag value"), int32(FlagToSet));
 		bool Result = false;
-		// Memory validation disabled as this is hot code.
-		UE_AUTORTFM_OPEN_NO_VALIDATION
+		UE_AUTORTFM_OPEN
 		{
 			FlagToSet &= ~EInternalObjectFlags_ReachabilityFlags; // reachability flags can only be cleared by GC through *_ForGC functions
 			FlagToSet &= ~EInternalObjectFlags::RefCounted; // refcounted flag is internal and must only be set by AddRef/ReleaseRef.
@@ -312,8 +310,7 @@ public:
 
 	void AddRef()
 	{
-		// Memory validation disabled as this is hot code.
-		UE_AUTORTFM_OPEN_NO_VALIDATION
+		UE_AUTORTFM_OPEN
 		{
 			FPlatformAtomics::InterlockedIncrement(&RefCount);
 			if ((GetFlags() & EInternalObjectFlags::RefCounted) != EInternalObjectFlags::RefCounted)
@@ -325,8 +322,7 @@ public:
 
 	void ReleaseRef()
 	{
-		// Memory validation disabled as this is hot code.
-		UE_AUTORTFM_OPEN_NO_VALIDATION
+		UE_AUTORTFM_OPEN
 		{
 			// This alone is not thread-safe as we may race with AddRef and in that case we don't want ClearRootFlags to apply.
 			// We fix this by validating that the refcount is still 0 while inside the root locks in ClearRootFlags.
diff --git a/Engine/Source/Runtime/CoreUObject/Public/UObject/UObjectBase.h b/Engine/Source/Runtime/CoreUObject/Public/UObject/UObjectBase.h
--- a/Engine/Source/Runtime/CoreUObject/Public/UObject/UObjectBase.h
+++ b/Engine/Source/Runtime/CoreUObject/Public/UObject/UObjectBase.h
@@ -320,8 +320,7 @@ private:
 
 		int32 Result = 0;
 
-		// Memory validation disabled as this is hot code.
-		UE_AUTORTFM_OPEN_NO_VALIDATION
+		UE_AUTORTFM_OPEN
 		{
 			Result = FPlatformAtomics::AtomicRead_Relaxed((int32*)&ObjectFlags);
 		};
diff --git a/Engine/Source/Runtime/CoreUObject/Public/UObject/UObjectBaseUtility.h b/Engine/Source/Runtime/CoreUObject/Public/UObject/UObjectBaseUtility.h
--- a/Engine/Source/Runtime/CoreUObject/Public/UObject/UObjectBaseUtility.h
+++ b/Engine/Source/Runtime/CoreUObject/Public/UObject/UObjectBaseUtility.h
@@ -776,8 +776,7 @@ public:
 	FORCEINLINE TStatId GetStatID(bool bForDeferredUse = false) const
 	{
 		TStatId Result{};
-		// Memory validation disabled as this is hot code.
-		UE_AUTORTFM_OPEN_NO_VALIDATION
+		UE_AUTORTFM_OPEN
 		{
 #if STATS
 			const TStatId& StatID = GUObjectArray.IndexToObject(InternalIndex)->StatID;
@@ -943,8 +942,7 @@ public:
 	 */
 	FORCEINLINE_STATS FScopeCycleCounterUObject(const UObjectBaseUtility *Object)
 	{
-		// Memory validation disabled as this is hot code.
-		UE_AUTORTFM_OPEN_NO_VALIDATION
+		UE_AUTORTFM_OPEN
 		{
 			if (Object)
 			{
@@ -972,8 +970,7 @@ public:
 	 */
 	FORCEINLINE_STATS FScopeCycleCounterUObject(const UObjectBaseUtility *Object, TStatId OtherStat)
 	{
-		// Memory validation disabled as this is hot code.
-		UE_AUTORTFM_OPEN_NO_VALIDATION
+		UE_AUTORTFM_OPEN
 		{
 			if (Object)
 			{
@@ -1010,8 +1007,7 @@ public:
 	FORCEINLINE_STATS ~FScopeCycleCounterUObject()
 	{
 		AutoRTFM::PopOnAbortHandler(this);
-		// Memory validation disabled as this is hot code.
-		UE_AUTORTFM_OPEN_NO_VALIDATION
+		UE_AUTORTFM_OPEN
 		{
 			Stop();
 		};
@@ -1039,8 +1035,7 @@ public:
 #endif
 	{
 #if ENABLE_STATNAMEDEVENTS_UOBJECT && CPUPROFILERTRACE_ENABLED
-		// Memory validation disabled as this is hot code.
-		UE_AUTORTFM_OPEN_NO_VALIDATION
+		UE_AUTORTFM_OPEN
 		{
 			if (GCycleStatsShouldEmitNamedEvents && UE_TRACE_CHANNELEXPR_IS_ENABLED(CpuChannel) && Object)
 			{
@@ -1070,8 +1065,7 @@ public:
 #if ENABLE_STATNAMEDEVENTS_UOBJECT && CPUPROFILERTRACE_ENABLED
 		AutoRTFM::PopOnAbortHandler(this);
 
-		// Memory validation disabled as this is hot code.
-		UE_AUTORTFM_OPEN_NO_VALIDATION
+		UE_AUTORTFM_OPEN
 		{
 			if (bPop)
 			{
@@ -1101,8 +1095,7 @@ public:
 	{
 		if (GHitchDetected &&  StatObject)
 		{
-			// Memory validation disabled as this is hot code.
-			UE_AUTORTFM_OPEN_NO_VALIDATION
+			UE_AUTORTFM_OPEN
 			{
 				ReportHitch();
 			};
diff --git a/Engine/Source/Runtime/Engine/Private/DataReplication.cpp b/Engine/Source/Runtime/Engine/Private/DataReplication.cpp
--- a/Engine/Source/Runtime/Engine/Private/DataReplication.cpp
+++ b/Engine/Source/Runtime/Engine/Private/DataReplication.cpp
@@ -403,8 +403,7 @@ void FObjectReplicator::InitRecentProperties(uint8* Source)
 	// In that case, we don't need to create any receiving state, as no one will be sending data to us.
 	ECreateRepStateFlags Flags = (Connection->IsInternalAck() && bIsServer) ? ECreateRepStateFlags::SkipCreateReceivingState : ECreateRepStateFlags::None;
 
-	// Memory validation disabled as this is hot code.
-	UE_AUTORTFM_OPEN_NO_VALIDATION
+	UE_AUTORTFM_OPEN
 	{
 		RepState = LocalRepLayout.CreateRepState(Source, RepChangedPropertyTracker, Flags);
 	};
diff --git a/Engine/Source/Runtime/Experimental/Chaos/Private/Chaos/AABB.cpp b/Engine/Source/Runtime/Experimental/Chaos/Private/Chaos/AABB.cpp
--- a/Engine/Source/Runtime/Experimental/Chaos/Private/Chaos/AABB.cpp
+++ b/Engine/Source/Runtime/Experimental/Chaos/Private/Chaos/AABB.cpp
@@ -414,8 +414,7 @@ inline TAABB<FReal, 3> TransformedAABBHelperISPC(const TAABB<FReal, 3>& AABB, co
 	check(bRealTypeCompatibleWithISPC);
 #if INTEL_ISPC
 	TVector<Chaos::FReal, 3> NewMin, NewMax;
-	// Memory validation disabled as this is hot code.
-	UE_AUTORTFM_OPEN_NO_VALIDATION
+	UE_AUTORTFM_OPEN
 	{
 		ispc::TransformedAABB((const ispc::FTransform&)SpaceTransform, (const ispc::FVector&)AABB.Min(), (const ispc::FVector&)AABB.Max(), (ispc::FVector&)NewMin, (ispc::FVector&)NewMax);
 	};
@@ -441,8 +440,7 @@ inline TAABB<Chaos::FRealSingle, 3> TransformedAABBHelperISPC(const TAABB<Chaos:
 	static_assert(sizeof(ispc::FVector3f) == sizeof(Chaos::TVector<Chaos::FRealSingle, 3>), "sizeof(ispc::FVector3f) != sizeof(Chaos::TVector<Chaos::FRealSingle, 3>)");
 
 	TVector<Chaos::FRealSingle, 3> NewMin, NewMax;
-	// Memory validation disabled as this is hot code.
-	UE_AUTORTFM_OPEN_NO_VALIDATION
+	UE_AUTORTFM_OPEN
 	{
 		ispc::TransformedAABBMixed((const ispc::FTransform&)SpaceTransform, (const ispc::FVector3f&)AABB.Min(), (const ispc::FVector3f&)AABB.Max(), (ispc::FVector3f&)NewMin, (ispc::FVector3f&)NewMax);
 	};
@@ -466,8 +464,7 @@ inline TAABB<FReal, 3> TransformedAABBHelperISPC2(const TAABB<FReal, 3>& AABB, c
 	check(bRealTypeCompatibleWithISPC);
 #if INTEL_ISPC
 	TVector<Chaos::FReal, 3> NewMin, NewMax;
-	// Memory validation disabled as this is hot code.
-	UE_AUTORTFM_OPEN_NO_VALIDATION
+	UE_AUTORTFM_OPEN
 	{
 		ispc::TransformedAABB2((const ispc::FTransform&)SpaceTransform, (const ispc::FVector&)AABB.Min(), (const ispc::FVector&)AABB.Max(), (ispc::FVector&)NewMin, (ispc::FVector&)NewMax);
 	};
@@ -493,8 +490,7 @@ inline TAABB<Chaos::FRealSingle, 3> TransformedAABBHelperISPC2(const TAABB<Chaos
 	static_assert(sizeof(ispc::FVector3f) == sizeof(Chaos::TVector<Chaos::FRealSingle, 3>), "sizeof(ispc::FVector3f) != sizeof(Chaos::TVector<Chaos::FRealSingle, 3>)");
 
 	TVector<Chaos::FRealSingle, 3> NewMin, NewMax;
-	// Memory validation disabled as this is hot code.
-	UE_AUTORTFM_OPEN_NO_VALIDATION
+	UE_AUTORTFM_OPEN
 	{
 		ispc::TransformedAABBMixed((const ispc::FTransform&)SpaceTransform, (const ispc::FVector3f&)AABB.Min(), (const ispc::FVector3f&)AABB.Max(), (ispc::FVector3f&)NewMin, (ispc::FVector3f&)NewMax);
 	};
diff --git a/Engine/Source/Runtime/Experimental/Chaos/Public/Chaos/RefCountedObject.h b/Engine/Source/Runtime/Experimental/Chaos/Public/Chaos/RefCountedObject.h
--- a/Engine/Source/Runtime/Experimental/Chaos/Public/Chaos/RefCountedObject.h
+++ b/Engine/Source/Runtime/Experimental/Chaos/Public/Chaos/RefCountedObject.h
@@ -37,8 +37,7 @@ public:
 	{
 		bool bIsFirstReference;
 
-		// Memory validation disabled as this is hot code.
-		UE_AUTORTFM_OPEN_NO_VALIDATION
+		UE_AUTORTFM_OPEN
 		{
 			bIsFirstReference = (++NumRefs == 1);
 		};
@@ -77,8 +76,7 @@ public:
 	uint32 GetRefCount() const
 	{
 		uint32 Ret;
-		// Memory validation disabled as this is hot code.
-		UE_AUTORTFM_OPEN_NO_VALIDATION
+		UE_AUTORTFM_OPEN
 		{
 			Ret = uint32(NumRefs.load());
 		};
diff --git a/Engine/Source/Runtime/Experimental/Chaos/Public/Framework/Threading.h b/Engine/Source/Runtime/Experimental/Chaos/Public/Framework/Threading.h
--- a/Engine/Source/Runtime/Experimental/Chaos/Public/Framework/Threading.h
+++ b/Engine/Source/Runtime/Experimental/Chaos/Public/Framework/Threading.h
@@ -1047,8 +1047,7 @@ FORCEINLINE void EnsureIsInGameThreadContext()
 		{
 			if (AutoRTFM::IsTransactional() || AutoRTFM::IsCommittingOrAborting())
 			{
-				// Memory validation disabled as this is hot code.
-				UE_AUTORTFM_OPEN_NO_VALIDATION
+				UE_AUTORTFM_OPEN
 					{
 						// The transactional system which can increment TransactionalLockCount
 						// is always single-threaded, thus this is safe to check without atomicity.
