Prevent IAD from logging errors when failing to start up on projects that don't require it.

#rnx

- Since the module is now included in the engine all the time, it will attempt to start up on projects not actively using it. This will cause problems on platforms that require persistent storage to be explicitly set up as the lack of persistent storage will cause some of our start up paths to return errors.
- For now instead of returning errors we will return the system as disabled, which will only log a notification and not an error.
- Note that some systems will assume that GOnDemandIoStore is valid if the initialization didn't error so we need to make sure that the pointer can still be returned.
- This is a temp fix top unblock CIS, a better solution will be added shortly.
- The code was originally by [REDACTED_PERSON], I am submitting on his behalf.

diff --git a/Engine/Source/Runtime/Experimental/IoStore/OnDemand/Private/IoStoreOnDemand.cpp b/Engine/Source/Runtime/Experimental/IoStore/OnDemand/Private/IoStoreOnDemand.cpp
--- a/Engine/Source/Runtime/Experimental/IoStore/OnDemand/Private/IoStoreOnDemand.cpp
+++ b/Engine/Source/Runtime/Experimental/IoStore/OnDemand/Private/IoStoreOnDemand.cpp
@@ -529,7 +529,6 @@ void FIoStoreOnDemandModule::StartupModule()
 	if (FIoStatus Status = GOnDemandIoStore->Initialize(); !Status.IsOk())
 	{
 		UE_LOG(LogIoStoreOnDemand, Error, TEXT("Failed to initialize I/O store on demand, reason '%s'"), *Status.ToString());
-		GOnDemandIoStore.Reset();
 		return;
 	}
 
diff --git a/Engine/Source/Runtime/Experimental/IoStore/OnDemand/Private/OnDemandConfig.cpp b/Engine/Source/Runtime/Experimental/IoStore/OnDemand/Private/OnDemandConfig.cpp
--- a/Engine/Source/Runtime/Experimental/IoStore/OnDemand/Private/OnDemandConfig.cpp
+++ b/Engine/Source/Runtime/Experimental/IoStore/OnDemand/Private/OnDemandConfig.cpp
@@ -374,13 +374,13 @@ TIoStatusOr<FOnDemandInstallCacheConfig> TryParseInstallCacheConfig(const TCHAR*
 #endif
 	if (bUseInstallCache == false)
 	{
-		FIoStatus Status = FIoStatusBuilder(EIoErrorCode::InvalidParameter) << TEXT("Disabled");
+		FIoStatus Status = FIoStatusBuilder(EIoErrorCode::Disabled) << TEXT("Disabled");
 		return Status;
 	}
 
 	if (FPaths::HasProjectPersistentDownloadDir() == false)
 	{
-		FIoStatus Status = FIoStatusBuilder(EIoErrorCode::InvalidParameter) << TEXT("Persistent storage not configured");
+		FIoStatus Status = FIoStatusBuilder(EIoErrorCode::Disabled) << TEXT("Persistent storage not configured");
 		return Status;
 	}
 
