Iris - Added command Net.Iris.PrintDynamicFilterClassConfigIssues with options SkipBPIssueDetection and SkipActorChildIssueDetection. It will report filter configured classes that cannot be found and potential issues with unconfigured subclasses that aren't being filtered as their parent. SkipBPIssueDetection will not report filtering issues for blueprint classes. SkipActorChildIssueDetection will ignore reporting potential issues for classes whose nearest configured super class is Actor.

diff --git a/Engine/Source/Runtime/Engine/Private/Net/Iris/ReplicationSystem/EngineReplicationBridge.cpp b/Engine/Source/Runtime/Engine/Private/Net/Iris/ReplicationSystem/EngineReplicationBridge.cpp
--- a/Engine/Source/Runtime/Engine/Private/Net/Iris/ReplicationSystem/EngineReplicationBridge.cpp
+++ b/Engine/Source/Runtime/Engine/Private/Net/Iris/ReplicationSystem/EngineReplicationBridge.cpp
@@ -677,6 +677,17 @@ float UEngineReplicationBridge::GetPollFrequencyOfRootObject(const UObject* Repl
 	return PollFrequency;
 }
 
+bool UEngineReplicationBridge::IsClassReplicatedByDefault(const UClass* Class) const
+{
+	if (Class && Class->IsChildOf(AActor::StaticClass()))
+	{
+		const AActor* CDO = Cast<AActor>(Class->GetDefaultObject());
+		return CDO && CDO->GetIsReplicated();
+	}
+
+	return Super::IsClassReplicatedByDefault(Class);
+}
+
 void UEngineReplicationBridge::WakeUpObjectInstantiatedFromRemote(AActor* Actor) const
 {
 	// If the actor is already awake or can't be woken up then return immediately.
diff --git a/Engine/Source/Runtime/Engine/Public/Net/Iris/ReplicationSystem/EngineReplicationBridge.h b/Engine/Source/Runtime/Engine/Public/Net/Iris/ReplicationSystem/EngineReplicationBridge.h
--- a/Engine/Source/Runtime/Engine/Public/Net/Iris/ReplicationSystem/EngineReplicationBridge.h
+++ b/Engine/Source/Runtime/Engine/Public/Net/Iris/ReplicationSystem/EngineReplicationBridge.h
@@ -118,6 +118,9 @@ protected:
 
 	virtual float GetPollFrequencyOfRootObject(const UObject* ReplicatedObject) const override;
 
+	/** Returns true if the class is derived from Actor and its CDO has set bReplicates. */
+	virtual bool IsClassReplicatedByDefault(const UClass* Class) const;
+
 	[[nodiscard]] virtual FString PrintConnectionInfo(uint32 ConnectionId) const override;
 
 private:
diff --git a/Engine/Source/Runtime/Experimental/Iris/Core/Private/Iris/ReplicationSystem/ObjectReplicationBridge.cpp b/Engine/Source/Runtime/Experimental/Iris/Core/Private/Iris/ReplicationSystem/ObjectReplicationBridge.cpp
--- a/Engine/Source/Runtime/Experimental/Iris/Core/Private/Iris/ReplicationSystem/ObjectReplicationBridge.cpp
+++ b/Engine/Source/Runtime/Experimental/Iris/Core/Private/Iris/ReplicationSystem/ObjectReplicationBridge.cpp
@@ -1823,6 +1823,11 @@ bool UObjectReplicationBridge::IsClassCritical(const UClass* Class)
 	return false;
 }
 
+bool UObjectReplicationBridge::IsClassReplicatedByDefault(const UClass* Class) const
+{
+	return Class != nullptr;
+}
+
 FString UObjectReplicationBridge::PrintConnectionInfo(uint32 ConnectionId) const
 {
 	return FString::Printf(TEXT("ConnectionId:%u"), ConnectionId);
diff --git a/Engine/Source/Runtime/Experimental/Iris/Core/Private/Iris/ReplicationSystem/ObjectReplicationBridgeDebugging.cpp b/Engine/Source/Runtime/Experimental/Iris/Core/Private/Iris/ReplicationSystem/ObjectReplicationBridgeDebugging.cpp
--- a/Engine/Source/Runtime/Experimental/Iris/Core/Private/Iris/ReplicationSystem/ObjectReplicationBridgeDebugging.cpp
+++ b/Engine/Source/Runtime/Experimental/Iris/Core/Private/Iris/ReplicationSystem/ObjectReplicationBridgeDebugging.cpp
@@ -15,10 +15,13 @@
 #include "Iris/ReplicationSystem/ReplicationOperations.h"
 
 #include "Iris/ReplicationSystem/Filtering/ReplicationFiltering.h"
+#include "Iris/ReplicationSystem/ObjectReplicationBridgeConfig.h"
 
 #include "Net/Core/NetBitArrayPrinter.h"
 #include "Net/Core/Trace/NetDebugName.h"
 
+#include "Templates/UnrealTemplate.h"
+#include "UObject/ClassTree.h"
 #include "UObject/CoreNet.h"
 
 /**
@@ -67,6 +70,48 @@ EPrintTraits FindPrintTraitsFromArgs(const TArray<FString>& Args)
 	return Traits;
 }
 
+enum class EDynamicFilterPrintTraits : uint32
+{
+	Default							= 0x0000,
+	// Print classes and what filter they will set. Includes configured classes as well as all replicated classes that have been automatically assigned a filter so far. 
+	Config							= 0x0001,
+	// Detect issues like non-existing classes being configured or child classes not being configured and ending up with a different filter
+	IssueDetection					= 0x0002,
+	// Ignore blueprint classes if filter config IssueDetection is enabled.
+	SkipBPIssueDetection			= 0x0004,
+	// Ignore issues related to classes whose nearest configured super class is Actor
+	SkipActorChildIssueDetection	= 0x0008,
+
+};
+ENUM_CLASS_FLAGS(EDynamicFilterPrintTraits);
+
+EDynamicFilterPrintTraits FindDynamicFilterPrintTraitsFromArgs(const TArray<FString>& Args)
+{
+	EDynamicFilterPrintTraits Traits = EDynamicFilterPrintTraits::Default;
+
+	if (Args.FindByPredicate([](const FString& Str) { return Str.Contains(TEXT("Config")); }))
+	{
+		Traits = Traits | EDynamicFilterPrintTraits::Config;
+	}
+
+	if (Args.FindByPredicate([](const FString& Str) { return Str.Contains(TEXT("IssueDetection")); }))
+	{
+		Traits = Traits | EDynamicFilterPrintTraits::IssueDetection;
+	}
+
+	if (Args.FindByPredicate([](const FString& Str) { return Str.Contains(TEXT("SkipBPIssueDetection")); }) )
+	{
+		Traits = Traits | EDynamicFilterPrintTraits::SkipBPIssueDetection;
+	}
+
+	if (Args.FindByPredicate([](const FString& Str) { return Str.Contains(TEXT("SkipActorChildIssueDetection")); }) )
+	{
+		Traits = Traits | EDynamicFilterPrintTraits::SkipActorChildIssueDetection;
+	}
+
+	return Traits;
+}
+
 /** Holds information about root objects sortable by class name */
 struct FRootObjectData
 {
@@ -257,7 +302,7 @@ void LogViaTrait(FLogContext& LogContext, EPrintTraits ArgTraits, EPrintTraits D
 // --------------------------------------------------------------------------------------------------------------------------------------------
 
 //-----------------------------------------------
-FAutoConsoleCommand ObjectBridgePrintDynamicFilter(
+FAutoConsoleCommand ObjectBridgePrintDynamicFilterConfig(
 	TEXT("Net.Iris.PrintDynamicFilterClassConfig"), 
 	TEXT("Prints the dynamic filter configured to be assigned to specific classes."), 
 	FConsoleCommandWithArgsDelegate::CreateLambda([](const TArray< FString >& Args)
@@ -278,26 +323,236 @@ FAutoConsoleCommand ObjectBridgePrintDynamicFilter(
 		return;
 	}
 
-	ObjectBridge->PrintDynamicFilterClassConfig();
+	const EDynamicFilterPrintTraits PrintTraits = EDynamicFilterPrintTraits::Config | FindDynamicFilterPrintTraitsFromArgs(Args);
+	ObjectBridge->PrintDynamicFilterClassConfig(static_cast<uint32>(PrintTraits));
 }));
 
-void UObjectReplicationBridge::PrintDynamicFilterClassConfig() const
+FAutoConsoleCommand ObjectBridgePrintDynamicFilterConfigIssues(
+	TEXT("Net.Iris.PrintDynamicFilterClassConfigIssues"), 
+	TEXT("Prints potential issues with the class filter config. Optional argument SkipBPIssueDetection will ignore blueprint classes. Optional argument SkipActorChildIssueDetection will ignore issues related to classes whose nearest configured super class is Actor."), 
+	FConsoleCommandWithArgsDelegate::CreateLambda([](const TArray< FString >& Args)
 {
+	using namespace UE::Net::Private::ObjectBridgeDebugging;
+
+	UReplicationSystem* RepSystem = FindReplicationSystemFromArg(Args);
+	if (!RepSystem)
+	{
+		UE_LOG(LogIrisBridge, Error, TEXT("Could not find ReplicationSystem."));
+		return;
+	}
+
+	UObjectReplicationBridge* ObjectBridge = CastChecked<UObjectReplicationBridge>(RepSystem->GetReplicationBridge());
+	if (!ObjectBridge)
+	{
+		UE_LOG(LogIrisBridge, Error, TEXT("Could not find ObjectReplicationBridge."));
+		return;
+	}
+
+	const EDynamicFilterPrintTraits PrintTraits = EDynamicFilterPrintTraits::IssueDetection | FindDynamicFilterPrintTraitsFromArgs(Args);
+	ObjectBridge->PrintDynamicFilterClassConfig(static_cast<uint32>(PrintTraits));
+}));
+
+void UObjectReplicationBridge::PrintDynamicFilterClassConfig(uint32 ArgTraits)
+{
+	using namespace UE::Net;
+	using namespace UE::Net::Private::ObjectBridgeDebugging;
+
 	const UReplicationSystem* RepSystem = GetReplicationSystem();
 
 	UE_LOG(LogIrisFilterConfig, Display, TEXT(""));
-	UE_LOG(LogIrisFilterConfig, Display, TEXT("Default Dynamic Filter Class Config:"));
+
+	const EDynamicFilterPrintTraits PrintTraits = static_cast<EDynamicFilterPrintTraits>(ArgTraits);
+	if (EnumHasAnyFlags(PrintTraits, EDynamicFilterPrintTraits::Config))
 	{
-		TMap<FName, FClassFilterInfo> SortedClassConfig = ClassesWithDynamicFilter;
+		UE_LOG(LogIrisFilterConfig, Display, TEXT("### Begin Default Dynamic Filter Class Config ###"));
+		{
+			TMap<FName, FClassFilterInfo> SortedClassConfig = ClassesWithDynamicFilter;
+
+			SortedClassConfig.KeyStableSort([](FName lhs, FName rhs) { return lhs.Compare(rhs) < 0; });
+			for (auto MapIt = SortedClassConfig.CreateConstIterator(); MapIt; ++MapIt)
+			{
+				const FName ClassName = MapIt.Key();
+				const FClassFilterInfo FilterInfo = MapIt.Value();
 
-		SortedClassConfig.KeyStableSort([](FName lhs, FName rhs){return lhs.Compare(rhs) < 0;});
-		for (auto MapIt = SortedClassConfig.CreateConstIterator(); MapIt; ++MapIt)
+				UE_LOG(LogIrisFilterConfig, Display, TEXT("\t%s -> %s"), *ClassName.ToString(), *RepSystem->GetFilterName(FilterInfo.FilterHandle).ToString());
+			}
+		}
+		UE_LOG(LogIrisFilterConfig, Display, TEXT("### End Default Dynamic Filter Class Config ###"));
+	}
+
+	if (EnumHasAnyFlags(PrintTraits, EDynamicFilterPrintTraits::IssueDetection))
+	{
+		UE_LOG(LogIrisFilterConfig, Display, TEXT("### Begin Dynamic Filter Class Config Issue Detection ###"));
 		{
-			const FName ClassName = MapIt.Key();
-			const FClassFilterInfo FilterInfo = MapIt.Value();
+			const UObjectReplicationBridgeConfig* BridgeConfig = UObjectReplicationBridgeConfig::GetConfig();
+
+			// Guard ClassesWithDynamicFilter. We want to use regular functions to validate the setup
+			TGuardValue<TMap<FName, FClassFilterInfo>> GuardClassesWithDynamicFilter(ClassesWithDynamicFilter, TMap<FName, FClassFilterInfo>());
+
+			// Populate ClassesWithDynamicFilter with the user config
+			for (const FObjectReplicationBridgeFilterConfig& FilterConfig : BridgeConfig->GetFilterConfigs())
+			{
+				FClassFilterInfo FilterInfo;
+				FilterInfo.FilterHandle = ReplicationSystem->GetFilterHandle(FilterConfig.DynamicFilterName);
+				FilterInfo.FilterProfile = FilterConfig.FilterProfile;
+				FilterInfo.bForceEnable = FilterConfig.bForceEnableOnAllInstances;
+				ClassesWithDynamicFilter.Add(FilterConfig.ClassName, FilterInfo);
+			}
+
+			// Populate a class tree in order to be able to more easily traverse classes top down and stop validating on subclasses with specific filter configs.
+			FClassTree FilterConfigClassTree(UObject::StaticClass());
+			if (!ClassesWithDynamicFilter.IsEmpty())
+			{
+				TArray<UObject*> Classes;
+				Classes.Reserve(128*1024);
+				constexpr bool bIncludeDerivedClasses = true;
+				GetObjectsOfClass(UClass::StaticClass(), Classes, bIncludeDerivedClasses, RF_ClassDefaultObject, GetObjectIteratorDefaultInternalExclusionFlags(EInternalObjectFlags::None));
+				if (EnumHasAnyFlags(PrintTraits, EDynamicFilterPrintTraits::SkipBPIssueDetection))
+				{
+					for (UObject* Object : Classes)
+					{
+						UClass* Class = static_cast<UClass*>(Object);
+						if (Class->HasAnyClassFlags(CLASS_Native))
+						{
+							FilterConfigClassTree.AddClass(Class);
+						}
+					}
+				}
+				else
+				{
+					for (UObject* Object : Classes)
+					{
+						UClass* Class = static_cast<UClass*>(Object);
+						FilterConfigClassTree.AddClass(Class);
+					}
+				}
+			}
+
+			// Iterate over class filter configs and see if we can find issues with it.
+			TSet<UClass*> ClassesWithConfig;
+			TSet<UClass*> ValidatedClasses;
+			TArray<const FClassTree*> ClassesToValidate;
+			TArray<const FClassTree*> ClassesToValidateNext;
+			const bool bSkipActorChildIssues = EnumHasAnyFlags(PrintTraits, EDynamicFilterPrintTraits::SkipActorChildIssueDetection);
+
+			// Store all classes with config for easy validation termination of subclass trees
+			for (const FObjectReplicationBridgeFilterConfig& FilterConfig : BridgeConfig->GetFilterConfigs())
+			{
+				const FString& ClassName = FilterConfig.ClassName.ToString();
 
-			UE_LOG(LogIrisFilterConfig, Display, TEXT("\t%s -> %s"), *ClassName.ToString(), *RepSystem->GetFilterName(FilterInfo.FilterHandle).ToString());
+				constexpr UObject* ClassOuter = nullptr;
+				constexpr bool bExactClass = true;
+				UClass* Class = Cast<UClass>(StaticFindObject(UClass::StaticClass(), ClassOuter, ToCStr(ClassName), bExactClass));
+				if (Class)
+				{
+					ClassesWithConfig.Add(Class);
+				}
+				else
+				{
+					UE_LOG(LogIrisFilterConfig, Warning, TEXT("\tCan't find class %s. Check spelling."), ToCStr(ClassName));
+				}
+			}
+
+			// Validate classes in top down order
+			TArray<UClass*> RootClassProspects;
+			bool bAlreadyValidated = false;
+			for (UClass* Class : ClassesWithConfig)
+			{
+				if (ValidatedClasses.Find(Class))
+				{
+					continue;
+				}
+
+				// Iterate over super classes and see if any has an explicit filter config
+				RootClassProspects.Reset();
+				RootClassProspects.Add(Class);
+				for (UClass* SuperClass = Class->GetSuperClass(); SuperClass; SuperClass = SuperClass->GetSuperClass())
+				{
+					if (ClassesWithConfig.Find(SuperClass))
+					{
+						RootClassProspects.Add(SuperClass);
+					}
+				}
+
+				UClass* ClassToValidate = RootClassProspects.Last();
+				ValidatedClasses.FindOrAdd(ClassToValidate, &bAlreadyValidated);
+				if (bAlreadyValidated)
+				{
+					continue;
+				}
+
+				const FClassTree* ClassTreeToValidate = FilterConfigClassTree.FindNode(ClassToValidate);
+				if (ClassTreeToValidate)
+				{
+					constexpr bool bRecurse = false;
+					ClassTreeToValidate->GetChildClasses(ClassesToValidate, bRecurse);
+				}
+
+				while (!ClassesToValidate.IsEmpty())
+				{
+					for (const FClassTree* ChildClassTree : ClassesToValidate)
+					{
+						if (!ChildClassTree)
+						{
+							continue;
+						}
+
+						UClass* ChildClass = ChildClassTree->GetClass();
+
+						ValidatedClasses.FindOrAdd(ChildClass, &bAlreadyValidated);
+						if (bAlreadyValidated)
+						{
+							continue;
+						}
+
+						// Validate classes that are replicated by default
+						bool bShouldValidateChildClasses = true;
+						if (IsClassReplicatedByDefault(ChildClass))
+						{
+							constexpr bool bRequireForceEnabled = false;
+							FName FilterProfile;
+							const FNetObjectFilterHandle ChildClassFilter = GetDynamicFilter(ChildClass, bRequireForceEnabled, FilterProfile);
+
+							// If the child class doesn't have an explicit config make sure its closest parent class with a config uses the same filter
+							bShouldValidateChildClasses = ClassesWithConfig.Contains(ChildClass);
+							if (!bShouldValidateChildClasses)
+							{
+								for (UClass* SuperClass = ChildClass->GetSuperClass(); SuperClass; SuperClass = SuperClass->GetSuperClass())
+								{
+									const FName SuperClassPath = GetConfigClassPathName(SuperClass);
+									if (const FClassFilterInfo* SuperClassFilterInfo = ClassesWithDynamicFilter.Find(SuperClassPath))
+									{
+										if (SuperClassFilterInfo->FilterHandle == ChildClassFilter || (bSkipActorChildIssues && SuperClass->GetFName() == NAME_Actor))
+										{
+											bShouldValidateChildClasses = true;
+										}
+										else
+										{
+											// Log a message and don't recurse into child classes
+											UE_LOG(LogIrisFilterConfig, Display, TEXT("Child class %s uses a different filter than its closest configured super class %s. Child class filter: %s Parent class filter: %s. Recommend adding an explicit filter config."), ToCStr(GetConfigClassPathName(ChildClass).ToString()), ToCStr(SuperClassPath.ToString()), ToCStr(RepSystem->GetFilterName(ChildClassFilter).ToString()), ToCStr(RepSystem->GetFilterName(SuperClassFilterInfo->FilterHandle).ToString()));
+										}
+										break;
+									}
+								}
+							}
+						}
+
+						// Don't traverse subclasses to classes with potential issues.
+						if (bShouldValidateChildClasses)
+						{
+							TArray<const FClassTree*> ChildClasses;
+							constexpr bool bRecurse = false;
+							ChildClassTree->GetChildClasses(ChildClasses, bRecurse);
+							ClassesToValidateNext.Append(ChildClasses);
+						}
+					}
+
+					Swap(ClassesToValidate, ClassesToValidateNext);
+					ClassesToValidateNext.Reset();
+				}
+			}
 		}
+		UE_LOG(LogIrisFilterConfig, Display, TEXT("### End Dynamic Filter Class Config Issue Dection ###"));
 	}
 }
 
diff --git a/Engine/Source/Runtime/Experimental/Iris/Core/Public/Iris/ReplicationSystem/ObjectReplicationBridge.h b/Engine/Source/Runtime/Experimental/Iris/Core/Public/Iris/ReplicationSystem/ObjectReplicationBridge.h
--- a/Engine/Source/Runtime/Experimental/Iris/Core/Public/Iris/ReplicationSystem/ObjectReplicationBridge.h
+++ b/Engine/Source/Runtime/Experimental/Iris/Core/Public/Iris/ReplicationSystem/ObjectReplicationBridge.h
@@ -221,7 +221,7 @@ public:
 public:
 
 	// Debug functions that are triggered via console commands declared in ObjectReplicationBridgeDebugging.cpp
-	void PrintDynamicFilterClassConfig() const;
+	void PrintDynamicFilterClassConfig(uint32 ArgTraits);
 	void PrintReplicatedObjects(uint32 ArgTraits) const;
 	void PrintAlwaysRelevantObjects(uint32 ArgTraits) const;
 	void PrintRelevantObjects(uint32 ArgTraits) const;
@@ -313,6 +313,9 @@ protected:
 	/** Returns true if the class is considered critical and we force a disconnection if a protocol mismatch prevents instances of this class from replicating. */
 	IRISCORE_API bool IsClassCritical(const UClass* Class);
 
+	/** Returns true if the class is replicated by default. This bridge will assume so. */
+	IRISCORE_API virtual bool IsClassReplicatedByDefault(const UClass* Class) const;
+
 	/** Returns the most relevant description of the client tied to this connection id. */
 	[[nodiscard]] IRISCORE_API virtual FString PrintConnectionInfo(uint32 ConnectionId) const;
 
