Iris - Fix spelling of AlwaysRelevantNetObjectFilterConfig

diff --git a/Engine/Source/Runtime/Experimental/Iris/Core/Public/Iris/ReplicationSystem/Filtering/AlwaysRelevantNetObjectFilter.h b/Engine/Source/Runtime/Experimental/Iris/Core/Public/Iris/ReplicationSystem/Filtering/AlwaysRelevantNetObjectFilter.h
--- a/Engine/Source/Runtime/Experimental/Iris/Core/Public/Iris/ReplicationSystem/Filtering/AlwaysRelevantNetObjectFilter.h
+++ b/Engine/Source/Runtime/Experimental/Iris/Core/Public/Iris/ReplicationSystem/Filtering/AlwaysRelevantNetObjectFilter.h
@@ -6,7 +6,7 @@
 #include "AlwaysRelevantNetObjectFilter.generated.h"
 
 UCLASS(transient, MinimalAPI)
-class UAlwaysReleantNetObjectFilterConfig final : public UNetObjectFilterConfig
+class UAlwaysRelevantNetObjectFilterConfig final : public UNetObjectFilterConfig
 {
 	GENERATED_BODY()
 };
