Iris - The Net.Iris.PrintAlwaysRelevantObjects console command now includes objects filtered through AlwaysRelevantNetObjectFilter.

diff --git a/Engine/Source/Runtime/Experimental/Iris/Core/Private/Iris/ReplicationSystem/ObjectReplicationBridgeDebugging.cpp b/Engine/Source/Runtime/Experimental/Iris/Core/Private/Iris/ReplicationSystem/ObjectReplicationBridgeDebugging.cpp
--- a/Engine/Source/Runtime/Experimental/Iris/Core/Private/Iris/ReplicationSystem/ObjectReplicationBridgeDebugging.cpp
+++ b/Engine/Source/Runtime/Experimental/Iris/Core/Private/Iris/ReplicationSystem/ObjectReplicationBridgeDebugging.cpp
@@ -14,6 +14,7 @@
 #include "Iris/ReplicationSystem/ReplicationSystemInternal.h"
 #include "Iris/ReplicationSystem/ReplicationOperations.h"
 
+#include "Iris/ReplicationSystem/Filtering/NetObjectFilterDefinitions.h"
 #include "Iris/ReplicationSystem/Filtering/ReplicationFiltering.h"
 #include "Iris/ReplicationSystem/ObjectReplicationBridgeConfig.h"
 
@@ -694,6 +695,23 @@ void UObjectReplicationBridge::PrintAlwaysRelevantObjects(uint32 ArgTraits) cons
 	
 	ReplicationSystemInternal->GetFiltering().BuildAlwaysRelevantList(MakeNetBitArrayView(AlwaysRelevantList), ReplicationSystemInternal->GetNetRefHandleManager().GetGlobalScopableInternalIndices());
 
+	// Include objects configured with an AlwaysRelveantNetObjectFilter
+	{
+		const FName AlwaysRelevantFilterName("/Script/IrisCore.AlwaysRelevantNetObjectFilter");
+		const UNetObjectFilterDefinitions* FilterDefinitions = GetDefault<UNetObjectFilterDefinitions>();
+		for (const FNetObjectFilterDefinition& FilterConfig : FilterDefinitions->GetFilterDefinitions())
+		{
+			// In theory there can be multiple filters using the AlwaysRelevantNetObjectFilter. While unlikely we do support it.
+			if (FilterConfig.ClassName == AlwaysRelevantFilterName)
+			{
+				FNetBitArray DynamicFilterAlwaysRelevantRelevant;
+				DynamicFilterAlwaysRelevantRelevant.Init(NetRefHandleManager->GetCurrentMaxInternalNetRefIndex());
+				ReplicationSystemInternal->GetFiltering().BuildObjectsInFilterList(MakeNetBitArrayView(DynamicFilterAlwaysRelevantRelevant), FilterConfig.FilterName);
+				MakeNetBitArrayView(AlwaysRelevantList).Combine(MakeNetBitArrayView(DynamicFilterAlwaysRelevantRelevant), FNetBitArrayView::OrOp);
+			}
+		}
+	}
+
 	// Remove subobjects from the list.
 	MakeNetBitArrayView(AlwaysRelevantList).Combine(NetRefHandleManager->GetSubObjectInternalIndicesView(), FNetBitArrayView::AndNotOp);
 
