Add dataflow caching support for TArray<TObjectPtr> arrays, so the underlying UObjects will not be garbage collected while the cache is still referencing them

diff --git a/Engine/Source/Runtime/Experimental/Dataflow/Core/Public/Dataflow/DataflowNodeParameters.h b/Engine/Source/Runtime/Experimental/Dataflow/Core/Public/Dataflow/DataflowNodeParameters.h
--- a/Engine/Source/Runtime/Experimental/Dataflow/Core/Public/Dataflow/DataflowNodeParameters.h
+++ b/Engine/Source/Runtime/Experimental/Dataflow/Core/Public/Dataflow/DataflowNodeParameters.h
@@ -69,7 +69,8 @@ namespace UE::Dataflow
 			CacheElementTyped,
 			CacheElementReference,
 			CacheElementNull,
-			CacheElementUObject
+			CacheElementUObject, // UObjectPtr
+			CacheElementUObjectArray // TArray of UObjectPtr
 		};
 
 		FContextCacheElementBase(EType CacheElementType, FGuid InNodeGuid = FGuid(), const FProperty* InProperty = nullptr, uint32 InNodeHash = 0, FTimestamp InTimestamp = FTimestamp::Invalid)
@@ -209,6 +210,41 @@ namespace UE::Dataflow
 		FDataType ObjectPtr;                         // Decaying T removes any reference and gets the correct underlying storage data type
 	};
 
+	//--------------------------------------------------------------------
+	// TArray<UObjectPtr> cache element, prevents the object from being garbage collected while in the cache
+	//--------------------------------------------------------------------
+	template<class T>
+	struct TContextCacheElementUObjectArray : public FContextCacheElementBase, public FGCObject
+	{
+		template<typename ArrayT>
+		TContextCacheElementUObjectArray(FGuid InNodeGuid, const FProperty* InProperty, ArrayT&& InObjectPtrArray, uint32 InNodeHash, FTimestamp Timestamp)
+			: FContextCacheElementBase(EType::CacheElementUObjectArray, InNodeGuid, InProperty, InNodeHash, Timestamp)
+			, ObjectPtrArray(Forward<ArrayT>(InObjectPtrArray))
+		{}
+
+		const TArray<T>& GetData(const FContext& /*Context*/, const FProperty* /*PropertyIn*/, const TArray<T>& /*Default*/) const
+		{
+			return ObjectPtrArray;
+		}
+
+		virtual const void* GetUntypedData(const FContext& /*Context*/, const FProperty* /*PropertyIn*/) const
+		{
+			return (void*)&ObjectPtrArray;
+		}
+
+		inline virtual TUniquePtr<FContextCacheElementBase> CreateReference(FContextCacheKey InReferenceDataKey) const override;
+
+		//~ Begin FGCObject interface
+		virtual void AddReferencedObjects(FReferenceCollector& Collector) override { Collector.AddReferencedObjects(ObjectPtrArray); }
+		virtual FString GetReferencerName() const override { return TEXT("TContextCacheElementUObjectArray"); }
+		//~ End FGCObject interface
+
+	private:
+		typedef typename TDecay<T>::Type FDataType;     // Using universal references here means T could be either const& or an rvalue reference
+		TArray<FDataType> ObjectPtrArray;               // Decaying T removes any reference and gets the correct underlying storage data type
+	};
+
+
 	// cache element method implementation 
 	template<class T>
 	const T& FContextCacheElementBase::GetTypedData(const FContext& Context, const FProperty* PropertyIn, const T& Default) const
@@ -229,6 +265,16 @@ namespace UE::Dataflow
 		{
 			return Default; 
 		}
+		if (Type == EType::CacheElementUObjectArray)
+		{
+			if constexpr (TIsTArray<T>::Value)
+			{
+				if constexpr (TIsUObjectPtrElement<typename T::ElementType>::Value)
+				{
+					return static_cast<const TContextCacheElementUObjectArray<typename T::ElementType>&>(*this).GetData(Context, PropertyIn, Default);
+				}
+			}
+		}
 		if (Type == EType::CacheElementUObject)
 		{
 			if constexpr (TIsUObjectPtrElement<T>::Value)
@@ -302,13 +348,26 @@ namespace UE::Dataflow
 		void SetData(FContextCacheKey InKey, const FProperty* InProperty, T&& InValue, const FGuid& InNodeGuid, uint32 InNodeHash, const FTimestamp& InTimestamp)
 		{
 			TUniquePtr<FContextCacheElementBase> DataStoreEntry;
-			if constexpr (TIsUObjectPtrElement<T>::Value)
+			bool bMadeEntry = false;
+			if constexpr (TIsTArray<typename TDecay<T>::Type>::Value)
 			{
-				DataStoreEntry = MakeUnique<TContextCacheElementUObject<T>>(InNodeGuid, InProperty, Forward<T>(InValue), InNodeHash, InTimestamp);
+				typedef typename TDecay<T>::Type ArrayType;
+				if constexpr (TIsUObjectPtrElement<typename ArrayType::ElementType>::Value)
+				{
+					bMadeEntry = true;
+					DataStoreEntry = MakeUnique<TContextCacheElementUObjectArray<typename ArrayType::ElementType>>(InNodeGuid, InProperty, Forward<T>(InValue), InNodeHash, InTimestamp);
+				}
 			}
-			else
+			if (!bMadeEntry)
 			{
-				DataStoreEntry = MakeUnique<TContextCacheElement<T>>(InNodeGuid, InProperty, Forward<T>(InValue), InNodeHash, InTimestamp);
+				if constexpr (TIsUObjectPtrElement<T>::Value)
+				{
+					DataStoreEntry = MakeUnique<TContextCacheElementUObject<T>>(InNodeGuid, InProperty, Forward<T>(InValue), InNodeHash, InTimestamp);
+				}
+				else
+				{
+					DataStoreEntry = MakeUnique<TContextCacheElement<T>>(InNodeGuid, InProperty, Forward<T>(InValue), InNodeHash, InTimestamp);
+				}
 			}
 
 			SetDataImpl(InKey, MoveTemp(DataStoreEntry));
@@ -546,6 +605,12 @@ namespace UE::Dataflow
 	{
 		return MakeUnique<TContextCacheElementReference<T>>(GetNodeGuid(), GetProperty(), InReferenceDataKey, GetNodeHash(), GetTimestamp());
 	}
+
+	template<class T>
+	TUniquePtr<FContextCacheElementBase> TContextCacheElementUObjectArray<T>::CreateReference(FContextCacheKey InReferenceDataKey) const
+	{
+		return MakeUnique<TContextCacheElementReference<T>>(GetNodeGuid(), GetProperty(), InReferenceDataKey, GetNodeHash(), GetTimestamp());
+	}
 }
 
 
