Chaos - add string conversions for some core Chaos types

diff --git a/Engine/Source/Runtime/Experimental/Chaos/Private/Chaos/ChaosLogUtil.cpp b/Engine/Source/Runtime/Experimental/Chaos/Private/Chaos/ChaosLogUtil.cpp
--- a/Engine/Source/Runtime/Experimental/Chaos/Private/Chaos/ChaosLogUtil.cpp
+++ b/Engine/Source/Runtime/Experimental/Chaos/Private/Chaos/ChaosLogUtil.cpp
@@ -24,6 +24,26 @@ namespace Chaos
 		return FString::Printf(TEXT("Invalid(%d)"), (int32)ObjectState);
 	}
 
+	FString ToString(const FRotation3& V)
+	{
+		return FString::Format(TEXT("{0}, {1}, {2}, {3}"), { V.X, V.Y, V.Z, V.W });
+	}
+
+	FString ToString(const FRotation3f& V)
+	{
+		return FString::Format(TEXT("{0}, {1}, {2}, {3}"), { V.X, V.Y, V.Z, V.W });
+	}
+
+	FString ToString(const FVec3& V)
+	{
+		return FString::Format(TEXT("{0}, {1}, {2}"), { V.X, V.Y, V.Z });
+	}
+
+	FString ToString(const FVec3f& V)
+	{
+		return FString::Format(TEXT("{0}, {1}, {2}"), { V.X, V.Y, V.Z });
+	}
+
 	FString ToString(const TBitArray<>& BitArray)
 	{
 		FString S;
diff --git a/Engine/Source/Runtime/Experimental/Chaos/Public/Chaos/ChaosLogUtil.h b/Engine/Source/Runtime/Experimental/Chaos/Public/Chaos/ChaosLogUtil.h
--- a/Engine/Source/Runtime/Experimental/Chaos/Public/Chaos/ChaosLogUtil.h
+++ b/Engine/Source/Runtime/Experimental/Chaos/Public/Chaos/ChaosLogUtil.h
@@ -1,12 +1,18 @@
 // Copyright Epic Games, Inc. All Rights Reserved.
 #pragma once
 
+#include "Chaos/Core.h"
 #include "Chaos/Particle/ObjectState.h"
+#include "Containers/UnrealString.h"
 
 class FString;
 
 namespace Chaos
 {
 	CHAOS_API FString ToString(const EObjectStateType ObjectState);
+	CHAOS_API FString ToString(const FRotation3& V);
+	CHAOS_API FString ToString(const FRotation3f& V);
+	CHAOS_API FString ToString(const FVec3& V);
+	CHAOS_API FString ToString(const FVec3f& V);
 	CHAOS_API FString ToString(const TBitArray<>& BitArray);
 }
\ No newline at end of file
