Initial change to support on-demand install priorities.

- updated install API to return a request handle. The handle can be used to cancel and or update the priority

#rnx

diff --git a/Engine/Plugins/Runtime/GameFeatures/Source/GameFeatures/Private/GameFeaturePluginStateMachine.cpp b/Engine/Plugins/Runtime/GameFeatures/Source/GameFeatures/Private/GameFeaturePluginStateMachine.cpp
--- a/Engine/Plugins/Runtime/GameFeatures/Source/GameFeatures/Private/GameFeaturePluginStateMachine.cpp
+++ b/Engine/Plugins/Runtime/GameFeatures/Source/GameFeatures/Private/GameFeaturePluginStateMachine.cpp
@@ -1419,9 +1419,17 @@ struct FGameFeaturePluginState_Downloading : public FGameFeaturePluginState
 	// Required for callback lifetime safety
 	struct FIoStoreOnDemandContext
 	{
+		TArray<UE::IoStore::FOnDemandInstallRequest> InstallRequests;
 		int32 PendingInstalls = 0;
-		UE::IoStore::FOnDemandCancellationToken CancelToken;
 		bool bStateValid = true;
+
+		void Cancel()
+		{
+			for (UE::IoStore::FOnDemandInstallRequest& Request : InstallRequests)
+			{
+				Request.Cancel();
+			}
+		}
 	};
 	TSharedPtr<FIoStoreOnDemandContext> IoStoreOnDemandContext;
 
@@ -1454,7 +1462,7 @@ struct FGameFeaturePluginState_Downloading : public FGameFeaturePluginState
 
 		if (IoStoreOnDemandContext)
 		{
-			IoStoreOnDemandContext->CancelToken.Cancel();
+			IoStoreOnDemandContext->Cancel();
 			IoStoreOnDemandContext->bStateValid = false;
 			IoStoreOnDemandContext = nullptr;
 		}
@@ -1591,6 +1599,7 @@ struct FGameFeaturePluginState_Downloading : public FGameFeaturePluginState
 		IoStoreOnDemandContext = MakeShared<FIoStoreOnDemandContext>();
 		IoStoreOnDemandContext->PendingInstalls = MetaData.InstallBundlesWithAssetDependencies.Num();
 
+		UE::IoStore::IOnDemandIoStore& IoStore = UE::IoStore::GetOnDemandIoStore();
 		for (const FName InstallBundle : MetaData.InstallBundlesWithAssetDependencies)
 		{
 			UE::IoStore::FOnDemandInstallArgs InstallArgs;
@@ -1603,7 +1612,8 @@ struct FGameFeaturePluginState_Downloading : public FGameFeaturePluginState
 			check(InstallArgs.ContentHandle.IsValid());
 
 			// This should be pretty small, so not going to worry about progress here.
-			UE::IoStore::GetOnDemandIoStore().Install(MoveTemp(InstallArgs),
+			IoStoreOnDemandContext->InstallRequests.Add(
+				IoStore.Install(MoveTemp(InstallArgs),
 				[this, LambdaOnDemandContext = IoStoreOnDemandContext](const UE::IoStore::FOnDemandInstallResult& OnDemandInstallResult)
 				{
 					if (!LambdaOnDemandContext->bStateValid)
@@ -1620,10 +1630,7 @@ struct FGameFeaturePluginState_Downloading : public FGameFeaturePluginState
 
 						// IOStore errors aren't localized 
 						Result = GetErrorResult(TEXT("IoStoreOnDemand.OnComplete."), ErrorCode, FText::AsCultureInvariant(StatusString));
-						if (!LambdaOnDemandContext->CancelToken.IsCanceled())
-						{
-							TryCancelState();
-						}
+						TryCancelState();
 					}
 
 					--IoStoreOnDemandContext->PendingInstalls;
@@ -1632,7 +1639,7 @@ struct FGameFeaturePluginState_Downloading : public FGameFeaturePluginState
 						bPluginDownloaded = Result.HasValue();
 						UpdateStateMachineImmediate();
 					}
-				}, nullptr, &IoStoreOnDemandContext->CancelToken);
+				}));
 		}
 	}
 
@@ -1763,7 +1770,7 @@ struct FGameFeaturePluginState_Downloading : public FGameFeaturePluginState
 			BundleManager->CancelUpdateContent(PendingBundleDownloads);
 			if (IoStoreOnDemandContext)
 			{
-				IoStoreOnDemandContext->CancelToken.Cancel();
+				IoStoreOnDemandContext->Cancel();
 			}
 		}
 	}
@@ -2693,9 +2700,9 @@ struct FGameFeaturePluginState_AssetDependencyStreaming : public FGameFeaturePlu
 	// Required for callback lifetime safety
 	struct FIoStoreOnDemandContext
 	{
+		TArray<UE::IoStore::FOnDemandInstallRequest> InstallRequests;
 		TArray<FIoStoreOnDemandProgress> Progress;
 		int32 PendingInstalls = 0;
-		UE::IoStore::FOnDemandCancellationToken CancelToken;
 		bool bStateValid = true;
 	};
 
@@ -2714,7 +2721,10 @@ struct FGameFeaturePluginState_AssetDependencyStreaming : public FGameFeaturePlu
 
 		if (IoStoreOnDemandContext)
 		{
-			IoStoreOnDemandContext->CancelToken.Cancel();
+			for (UE::IoStore::FOnDemandInstallRequest& Request : IoStoreOnDemandContext->InstallRequests)
+			{
+				Request.Cancel();
+			}
 			IoStoreOnDemandContext->bStateValid = false;
 			IoStoreOnDemandContext = nullptr;
 		}
@@ -2763,6 +2773,7 @@ struct FGameFeaturePluginState_AssetDependencyStreaming : public FGameFeaturePlu
 		IoStoreOnDemandContext->Progress.Reserve(MetaData.InstallBundlesWithAssetDependencies.Num());
 
 		const TArray<EStreamingAssetInstallMode>& InstallModes = MaybeInstallModes.GetValue();
+		UE::IoStore::IOnDemandIoStore& IoStore = UE::IoStore::GetOnDemandIoStore();
 		for (int i = 0; const FName InstallBundle : MetaData.InstallBundlesWithAssetDependencies)
 		{
 			IoStoreOnDemandContext->Progress.Emplace(FIoStoreOnDemandProgress{InstallBundle});
@@ -2780,7 +2791,7 @@ struct FGameFeaturePluginState_AssetDependencyStreaming : public FGameFeaturePlu
 			InstallArgs.ContentHandle = UE::GameFeatures::GFPSharedInstallTracker.AddOnDemandContentHandle(
 				InstallBundle, UE::GameFeatures::EGFPInstallLevel::AssetStream);
 
-			UE::IoStore::GetOnDemandIoStore().Install(MoveTemp(InstallArgs),
+			IoStoreOnDemandContext->InstallRequests.Add(IoStore.Install(MoveTemp(InstallArgs),
 				// On Complete
 				[this, LambdaOnDemandContext = IoStoreOnDemandContext](const UE::IoStore::FOnDemandInstallResult& OnDemandInstallResult)
 				{
@@ -2798,10 +2809,7 @@ struct FGameFeaturePluginState_AssetDependencyStreaming : public FGameFeaturePlu
 
 						// IOStore errors aren't localized 
 						Result = GetErrorResult(TEXT("IoStoreOnDemand.OnComplete."), ErrorCode, FText::AsCultureInvariant(StatusString));
-						if (!LambdaOnDemandContext->CancelToken.IsCanceled())
-						{
-							TryCancelState();
-						}
+						TryCancelState();
 					}
 
 					--IoStoreOnDemandContext->PendingInstalls;
@@ -2834,8 +2842,7 @@ struct FGameFeaturePluginState_AssetDependencyStreaming : public FGameFeaturePlu
 					const float OverallProgress = SumProgress.GetRelativeProgress();
 
 					StateProperties.OnFeatureStateProgressUpdate.ExecuteIfBound(OverallProgress);
-				},
-				&IoStoreOnDemandContext->CancelToken);
+				}));
 		}
 	}
 
@@ -2864,7 +2871,10 @@ struct FGameFeaturePluginState_AssetDependencyStreaming : public FGameFeaturePlu
 	{
 		if (IoStoreOnDemandContext)
 		{
-			IoStoreOnDemandContext->CancelToken.Cancel();
+			for (UE::IoStore::FOnDemandInstallRequest& InstallRequest : IoStoreOnDemandContext->InstallRequests)
+			{
+				InstallRequest.Cancel();
+			}
 		}
 	}
 };
diff --git a/Engine/Plugins/Runtime/GameFeatures/Source/GameFeatures/Private/GameFeaturesSubsystem.cpp b/Engine/Plugins/Runtime/GameFeatures/Source/GameFeatures/Private/GameFeaturesSubsystem.cpp
--- a/Engine/Plugins/Runtime/GameFeatures/Source/GameFeatures/Private/GameFeaturesSubsystem.cpp
+++ b/Engine/Plugins/Runtime/GameFeatures/Source/GameFeatures/Private/GameFeaturesSubsystem.cpp
@@ -2,6 +2,7 @@
 
 #include "GameFeaturesSubsystem.h"
 #include "Algo/Accumulate.h"
+#include "Algo/AnyOf.h"
 #include "Algo/TopologicalSort.h"
 #include "Algo/Unique.h"
 #include "AssetRegistry/AssetData.h"
@@ -2832,7 +2833,7 @@ struct FGameFeatureAssetStreamContext : public FGameFeaturePluginAssetStreamHand
 	TArray<UE::IoStore::FOnDemandInstallProgress> ProgressList;
 	int32 PendingInstalls = 0;
 	bool bIsComplete = false;
-	UE::IoStore::FOnDemandCancellationToken CancellationToken;
+	TArray<UE::IoStore::FOnDemandInstallRequest> InstallRequests;
 
 	FGameFeatureAssetStreamContext(FGameFeaturePluginStreamedAssetsRef&& InAssetRef)
 		: AssetRef(MoveTemp(InAssetRef))
@@ -2857,7 +2858,10 @@ struct FGameFeatureAssetStreamContext : public FGameFeaturePluginAssetStreamHand
 
 	virtual void Cancel() override
 	{
-		CancellationToken.Cancel();
+		for (UE::IoStore::FOnDemandInstallRequest& Request : InstallRequests)
+		{
+			Request.Cancel();
+		}
 	}
 
 	virtual FGameFeaturePluginStreamedAssetsRef GetAssetsRef() const override
@@ -2865,6 +2869,11 @@ struct FGameFeatureAssetStreamContext : public FGameFeaturePluginAssetStreamHand
 		return AssetRef;
 	}
 
+	bool IsCanceled()
+	{
+		return Algo::AnyOf(InstallRequests, [](const UE::IoStore::FOnDemandInstallRequest& R) { return R.IsCancelled(); });
+	}
+
 	void SetComplete()
 	{
 		bIsComplete = true;
@@ -2890,7 +2899,7 @@ struct FGameFeatureAssetStreamContext : public FGameFeaturePluginAssetStreamHand
 
 	void Start(TArrayView<FName> AssetPackageNames, EGameFeatureAssetStreamPriority Priority)
 	{
-		if (CancellationToken.IsCanceled())
+		if (IsCanceled())
 		{
 			SetCompleteCanceled();
 			return;
@@ -2963,7 +2972,7 @@ struct FGameFeatureAssetStreamContext : public FGameFeaturePluginAssetStreamHand
 
 	void StartByTag(TArrayView<FString> Tags, EGameFeatureAssetStreamPriority Priority)
 	{
-		if (CancellationToken.IsCanceled())
+		if (IsCanceled())
 		{
 			SetCompleteCanceled();
 			return;
@@ -2980,7 +2989,7 @@ struct FGameFeatureAssetStreamContext : public FGameFeaturePluginAssetStreamHand
 
 	void StartByTag(TArrayView<FString> PluginURLs, TArrayView<FString> Tags, EGameFeatureAssetStreamPriority Priority)
 	{
-		if (CancellationToken.IsCanceled())
+		if (IsCanceled())
 		{
 			SetCompleteCanceled();
 			return;
@@ -3077,6 +3086,7 @@ struct FGameFeatureAssetStreamContext : public FGameFeaturePluginAssetStreamHand
 		PendingInstalls = InstallArgList.Num();
 		ProgressList.Empty(PendingInstalls);
 
+		UE::IoStore::IOnDemandIoStore& IoStore = UE::IoStore::GetOnDemandIoStore();
 		for (UE::IoStore::FOnDemandInstallArgs& InstallArgs : InstallArgList)
 		{
 			const int32 RequestIndex = ProgressList.Emplace();
@@ -3086,7 +3096,7 @@ struct FGameFeatureAssetStreamContext : public FGameFeaturePluginAssetStreamHand
 
 			// TODO: convert Priority to IOStore priority here
 
-			UE::IoStore::GetOnDemandIoStore().Install(MoveTemp(InstallArgs),
+			InstallRequests.Add(IoStore.Install(MoveTemp(InstallArgs),
 				[Context = SharedThis(this), RequestIndex](const UE::IoStore::FOnDemandInstallResult& OnDemandInstallResult)
 				{
 					Context->OnInstallComplete(OnDemandInstallResult, RequestIndex);
@@ -3094,8 +3104,7 @@ struct FGameFeatureAssetStreamContext : public FGameFeaturePluginAssetStreamHand
 				[Context = SharedThis(this), RequestIndex](const UE::IoStore::FOnDemandInstallProgress& InProgress)
 				{
 					Context->OnInstallProgress(InProgress, RequestIndex);
-				},
-				&CancellationToken);
+				}));
 		}
 	}
 
diff --git a/Engine/Source/Runtime/Experimental/IoStore/OnDemand/Private/OnDemandIoStore.cpp b/Engine/Source/Runtime/Experimental/IoStore/OnDemand/Private/OnDemandIoStore.cpp
--- a/Engine/Source/Runtime/Experimental/IoStore/OnDemand/Private/OnDemandIoStore.cpp
+++ b/Engine/Source/Runtime/Experimental/IoStore/OnDemand/Private/OnDemandIoStore.cpp
@@ -639,25 +639,26 @@ void FOnDemandIoStore::Mount(FOnDemandMountArgs&& Args, FOnDemandMountCompleted&
 	TryEnterTickLoop();
 }
 
-void FOnDemandIoStore::Install(
+FOnDemandInstallRequest FOnDemandIoStore::Install(
 	FOnDemandInstallArgs&& Args,
 	FOnDemandInstallCompleted&& OnCompleted,
-	FOnDemandInstallProgressed&& OnProgress /*= nullptr*/,
-	const FOnDemandCancellationToken* CancellationToken)
+	FOnDemandInstallProgressed&& OnProgress)
 {
 	LLM_SCOPE_BYNAME(TEXT("FileSystem/OnDemandIoStore"));
-	FSharedInstallRequest InstallRequest = MakeShared<FInstallRequest>();
-	InstallRequest->Args				= MoveTemp(Args);
-	InstallRequest->OnCompleted			= MoveTemp(OnCompleted);
-	InstallRequest->OnProgressed		= MoveTemp(OnProgress);
-	InstallRequest->CancellationToken	= CancellationToken;
+	
+	FSharedInternalInstallRequest InstallRequest = MakeShared<FOnDemandInternalInstallRequest, ESPMode::ThreadSafe>(
+		MoveTemp(Args),
+		MoveTemp(OnCompleted),
+		MoveTemp(OnProgress));
 
 	{
 		UE::TUniqueLock Lock(RequestMutex);
-		InstallRequests.Add(MoveTemp(InstallRequest));
+		InstallRequests.Add(InstallRequest);
 	}
 
 	TryEnterTickLoop();
+
+	return FOnDemandInstallRequest(AsWeak(), InstallRequest);
 }
 
 void FOnDemandIoStore::Purge(FOnDemandPurgeArgs&& Args, FOnDemandPurgeCompleted&& OnCompleted)
@@ -1004,24 +1005,27 @@ bool FOnDemandIoStore::Tick()
 		}
 	}
 
-	TArray<FSharedInstallRequest> LocalInstallRequests;
+	FSharedInternalInstallRequest InstallRequest;
 	{
 		UE::TUniqueLock Lock(RequestMutex);
-		LocalInstallRequests = InstallRequests;
+		if (InstallRequests.IsEmpty() == false)
+		{
+			InstallRequests.StableSort(FOnDemandInternalInstallRequest::SortByPriority);
+			InstallRequest = InstallRequests.Pop(EAllowShrinking::No);
+		}
 	}
 
 	// Tick install request(s)
-	for (FSharedInstallRequest& Request : LocalInstallRequests)
+	if (InstallRequest.IsValid())
 	{
-		FOnDemandInstallResult Result = TickInstallRequest(*Request);
+		FOnDemandInstallResult Result = TickInstallRequest(*InstallRequest);
 
+		// If the error code is not unknown the the request has already been completed when cancelled 
+		EIoErrorCode Expected = EIoErrorCode::Unknown;
+		if (InstallRequest->ErrorCode.compare_exchange_strong(Expected, Result.Status.GetErrorCode()))
 		{
-			UE::TUniqueLock Lock(RequestMutex);
-			InstallRequests.Remove(Request);
+			CompleteInstallRequest(InstallRequest, MoveTemp(Result));
 		}
-
-		CompleteInstallRequest(*Request, MoveTemp(Result));
-
 		bTicked = true;
 	}
 
@@ -1240,14 +1244,14 @@ void FOnDemandIoStore::CompleteMountRequest(FMountRequest& Request, FOnDemandMou
 	}
 }
 
-FOnDemandInstallResult FOnDemandIoStore::TickInstallRequest(const FInstallRequest& InstallRequest)
+FOnDemandInstallResult FOnDemandIoStore::TickInstallRequest(FOnDemandInternalInstallRequest& InstallRequest)
 {
 	UE_LOG(LogIoStoreOnDemand, Log, TEXT("Ticking install request, ContentHandle='%s'"),
 		*LexToString(InstallRequest.Args.ContentHandle));
 
 	FOnDemandInstallResult OutResult;
 	FIoStatus& Status					= OutResult.Status;
-	FOnDemandInstallProgress& Progress	= OutResult.Progress;
+	FOnDemandInstallProgress& Progress	= InstallRequest.Progress;
 
 	const double StartTime = FPlatformTime::Seconds();
 	ON_SCOPE_EXIT
@@ -1257,10 +1261,8 @@ FOnDemandInstallResult FOnDemandIoStore::TickInstallRequest(const FInstallReques
 
 	auto CheckAndSetCancelled = [&InstallRequest, &Status]() -> bool
 	{
-		if (InstallRequest.CancellationToken && InstallRequest.CancellationToken->IsCanceled())
+		if (InstallRequest.IsCancelRequested())
 		{
-			UE_LOG(LogIoStoreOnDemand, Log, TEXT("Cancelling install request, ContentHandle='%s'"),
-				*LexToString(InstallRequest.Args.ContentHandle));
 			Status = FIoStatus(EIoErrorCode::Cancelled);
 			return true;
 		}
@@ -1550,32 +1552,49 @@ FOnDemandInstallResult FOnDemandIoStore::TickInstallRequest(const FInstallReques
 	return OutResult;
 }
 
-void FOnDemandIoStore::CompleteInstallRequest(FInstallRequest& Request, FOnDemandInstallResult&& InstallResult)
+void FOnDemandIoStore::CompleteInstallRequest(FSharedInternalInstallRequest Request, FOnDemandInstallResult&& InstallResult)
 {
-	if (!Request.OnCompleted)
+	FOnDemandInstallRequest::EStatus Status = FOnDemandInstallRequest::EStatus::None;
+	switch (InstallResult.Status.GetErrorCode())
 	{
-		return;
+		case EIoErrorCode::Ok:
+			Status = FOnDemandInstallRequest::Ok;
+			break;
+		case EIoErrorCode::Cancelled:
+			Status = FOnDemandInstallRequest::Cancelled;
+			break;
+		default:
+			Status = FOnDemandInstallRequest::Error;
+			break;
 	}
 
-	if (EnumHasAnyFlags(Request.Args.Options, EOnDemandInstallOptions::CallbackOnGameThread))
+	if (!Request->OnCompleted)
+	{
+		return Request->Status.store(Status);
+	}
+
+	if (EnumHasAnyFlags(Request->Args.Options, EOnDemandInstallOptions::CallbackOnGameThread))
 	{
 		ExecuteOnGameThread(
 			UE_SOURCE_LOCATION,
-			[OnCompleted = MoveTemp(Request.OnCompleted), InstallResult = MoveTemp(InstallResult)]() mutable
+			[Request, InstallResult = MoveTemp(InstallResult), Status]() mutable
 			{
+				FOnDemandInstallCompleted OnCompleted = MoveTemp(Request->OnCompleted);
 				OnCompleted(MoveTemp(InstallResult));
+				Request->Status.store(Status);
 			});
 	}
 	else
 	{
-		FOnDemandInstallCompleted OnCompleted = MoveTemp(Request.OnCompleted);
+		FOnDemandInstallCompleted OnCompleted = MoveTemp(Request->OnCompleted);
 		OnCompleted(MoveTemp(InstallResult));
+		Request->Status.store(Status);
 	}
 }
 
-void FOnDemandIoStore::ProgressInstallRequest(const FInstallRequest& Request, const FOnDemandInstallProgress& Progress)
+void FOnDemandIoStore::ProgressInstallRequest(const FOnDemandInternalInstallRequest& Request, const FOnDemandInstallProgress& Progress)
 {
-	if (!Request.OnProgressed)
+	if (!Request.OnProgress)
 	{
 		return;
 	}
@@ -1584,17 +1603,14 @@ void FOnDemandIoStore::ProgressInstallRequest(const FInstallRequest& Request, co
 	{
 		ExecuteOnGameThread(
 			UE_SOURCE_LOCATION,
-			[OnProgressed = Request.OnProgressed, Progress]()
+			[&Request, Progress]()
 			{
-				OnProgressed(Progress);
+				Request.OnProgress(Progress);
 			});
 	}
 	else
 	{
-		if (Request.OnProgressed)
-		{
-			Request.OnProgressed(Progress);
-		}
+		Request.OnProgress(Progress);
 	}
 }
 
@@ -1910,7 +1926,7 @@ FIoStatus FOnDemandIoStore::GetContainersAndPackagesForInstall(
 void FOnDemandIoStore::ReleaseContent(FOnDemandInternalContentHandle& ContentHandle)
 {
 	LLM_SCOPE_BYNAME(TEXT("FileSystem/OnDemandIoStore"));
-	UE_LOG(LogIoStoreOnDemand, Log, TEXT("Releasing content handle '%s'"), *LexToString(ContentHandle));
+	UE_LOG(LogIoStoreOnDemand, Log, TEXT("Releasing content handle, %s"), *LexToString(ContentHandle));
 
 	UE::TUniqueLock Lock(ContainerMutex);
 	for (FSharedOnDemandContainer& Container : Containers)
@@ -1949,6 +1965,36 @@ void FOnDemandIoStore::GetReferencedContent(TArray<FSharedOnDemandContainer>& Ou
 	}
 }
 
+void FOnDemandIoStore::CancelInstallRequest(FSharedInternalInstallRequest InstallRequest)
+{
+	check(InstallRequest.IsValid());
+
+	EIoErrorCode Expected = EIoErrorCode::Unknown;
+	if (InstallRequest->ErrorCode.compare_exchange_strong(Expected, EIoErrorCode::Cancelled) == false) 
+	{
+		// If the error code is not unknown the request will be completed in the tick loop 
+		return;
+	}
+
+	UE_LOG(LogIoStoreOnDemand, Log, TEXT("Cancelling install request, %s"), *LexToString(*InstallRequest));
+
+	FOnDemandInstallResult InstallResult
+	{
+		.Status				= FIoStatus(EIoErrorCode::Cancelled),
+		.DurationInSeconds 	= FPlatformTime::Seconds() - InstallRequest->StartTime,
+		.Progress			= MoveTemp(InstallRequest->Progress)
+	};
+	CompleteInstallRequest(InstallRequest, MoveTemp(InstallResult));
+}
+
+void FOnDemandIoStore::UpdateInstallRequestPriority(FSharedInternalInstallRequest InstallRequest, EOnDemandInstallPriority NewPriority)
+{
+	{
+		UE::TUniqueLock Lock(RequestMutex);
+		InstallRequest->Args.Priority = NewPriority;
+	}
+}
+
 FOnDemandCacheUsage FOnDemandIoStore::GetCacheUsage() const
 {
 	FOnDemandCacheUsage CacheUsage;
@@ -2004,7 +2050,7 @@ FOnDemandInternalContentHandle::~FOnDemandInternalContentHandle()
 
 FString LexToString(const FOnDemandInternalContentHandle& Handle)
 {
-	return FString::Printf(TEXT("0x%llX (%s)"), Handle.HandleId(), *Handle.DebugName);
+	return FString::Printf(TEXT("Id=0x%llX DebugName='%s'"), Handle.HandleId(), *Handle.DebugName);
 }
 
 ///////////////////////////////////////////////////////////////////////////////
@@ -2040,4 +2086,71 @@ FString LexToString(const FOnDemandContentHandle& Handle)
 	return Handle.IsValid() ? LexToString(*Handle.Handle) : TEXT("Invalid");
 }
 
+////////////////////////////////////////////////////////////////////////////////
+FString LexToString(const FOnDemandInternalInstallRequest& Request)
+{
+	return FString::Printf(TEXT("Id=0x%llX, DebugName='%s', ContentHandle='%s'"),
+		UPTRINT(&Request),
+		Request.Args.DebugName.IsEmpty()
+			? TEXT("NoName")
+			: *Request.Args.DebugName,
+		*LexToString(Request.Args.ContentHandle));
+}
+
+uint32 FOnDemandInternalInstallRequest::NextSeqNo = 0;
+
+///////////////////////////////////////////////////////////////////////////////
+FOnDemandInstallRequest::FOnDemandInstallRequest()
+{
+}
+
+FOnDemandInstallRequest::FOnDemandInstallRequest(FWeakOnDemandIoStore InIoStore, FSharedInternalInstallRequest InternalRequest)
+	: IoStore(InIoStore)
+	, Request(InternalRequest)
+{
+}
+
+FOnDemandInstallRequest::~FOnDemandInstallRequest()
+{
+	Cancel();
+}
+
+FOnDemandInstallRequest::FOnDemandInstallRequest(FOnDemandInstallRequest&& Other)
+	: Request(MoveTemp(Other.Request))
+{
+}
+
+FOnDemandInstallRequest& FOnDemandInstallRequest::operator=(FOnDemandInstallRequest&& Other)
+{
+	Request = MoveTemp(Other.Request);
+	return *this;
+}
+
+FOnDemandInstallRequest::EStatus FOnDemandInstallRequest::GetStatus() const
+{
+	return Request.IsValid() ? Request->Status.load(std::memory_order_relaxed) : EStatus::None; 
+}
+
+void FOnDemandInstallRequest::Cancel()
+{
+	if (IsPending())
+	{
+		if (TSharedPtr<FOnDemandIoStore, ESPMode::ThreadSafe> Pinned = IoStore.Pin(); Pinned.IsValid())
+		{
+			Pinned->CancelInstallRequest(Request);
+		}
+	}
+}
+
+void FOnDemandInstallRequest::UpdatePriority(EPriority NewPriority)
+{
+	if (IsPending())
+	{
+		if (TSharedPtr<FOnDemandIoStore, ESPMode::ThreadSafe> Pinned = IoStore.Pin(); Pinned.IsValid())
+		{
+			Pinned->UpdateInstallRequestPriority(Request, NewPriority);
+		}
+	}
+}
+
 } // namespace UE::IoStore
diff --git a/Engine/Source/Runtime/Experimental/IoStore/OnDemand/Private/OnDemandIoStore.h b/Engine/Source/Runtime/Experimental/IoStore/OnDemand/Private/OnDemandIoStore.h
--- a/Engine/Source/Runtime/Experimental/IoStore/OnDemand/Private/OnDemandIoStore.h
+++ b/Engine/Source/Runtime/Experimental/IoStore/OnDemand/Private/OnDemandIoStore.h
@@ -31,7 +31,6 @@ enum class ELogOnDemandCacheUsage : uint8;
 using FSharedStreamingBackend		= TSharedPtr<class IOnDemandIoDispatcherBackend>;
 using FSharedPackageStoreBackend	= TSharedPtr<class IOnDemandPackageStoreBackend>;
 using FSharedInstallCache			= TSharedPtr<class IOnDemandInstallCache>;
-using FWeakOnDemandIoStore			= TWeakPtr<class FOnDemandIoStore, ESPMode::ThreadSafe>;
 
 ///////////////////////////////////////////////////////////////////////////////
 class FOnDemandInternalContentHandle
@@ -53,6 +52,46 @@ public:
 
 FString LexToString(const FOnDemandInternalContentHandle& Handle);
 
+////////////////////////////////////////////////////////////////////////////////
+class FOnDemandInternalInstallRequest
+{
+public:
+	using EStatus = FOnDemandRequest::EStatus;
+
+	static uint32 NextSeqNo;
+
+	static bool SortByPriority(FSharedInternalInstallRequest LHS, FSharedInternalInstallRequest RHS)
+	{
+		if (LHS->Args.Priority == RHS->Args.Priority)
+		{
+			return LHS->SeqNo > RHS->SeqNo;
+		}
+		return static_cast<uint8>(LHS->Args.Priority) < static_cast<uint8>(RHS->Args.Priority);
+	};
+
+	FOnDemandInternalInstallRequest(
+		FOnDemandInstallArgs&& InArgs,
+		FOnDemandInstallCompleted&& InOnCompleted,
+		FOnDemandInstallProgressed&& InOnProgress)
+			: Args(MoveTemp(InArgs))
+			, OnCompleted(MoveTemp(InOnCompleted))
+			, OnProgress(MoveTemp(InOnProgress))
+	{ }
+
+	bool IsCancelRequested() const { return EIoErrorCode::Cancelled == ErrorCode.load(std::memory_order_relaxed); }
+
+	FOnDemandInstallArgs		Args;
+	FOnDemandInstallCompleted	OnCompleted;
+	FOnDemandInstallProgressed	OnProgress;
+	FOnDemandInstallProgress	Progress;
+	double						StartTime = FPlatformTime::Seconds();
+	uint32						SeqNo = NextSeqNo++;
+	std::atomic<EStatus>		Status{EStatus::Pending};
+	std::atomic<EIoErrorCode>	ErrorCode{EIoErrorCode::Unknown};
+};
+
+FString LexToString(const FOnDemandInternalInstallRequest& Request);
+
 ///////////////////////////////////////////////////////////////////////////////
 enum class EOnDemandContainerFlags : uint8
 {
@@ -265,16 +304,6 @@ class FOnDemandIoStore
 
 	using FSharedMountRequest	= TSharedRef<FMountRequest>;
 
-	struct FInstallRequest
-	{
-		FOnDemandInstallArgs				Args;
-		FOnDemandInstallCompleted			OnCompleted;
-		FOnDemandInstallProgressed			OnProgressed;
-		const FOnDemandCancellationToken*	CancellationToken = nullptr;
-	};
-
-	using FSharedInstallRequest	= TSharedRef<FInstallRequest>;
-
 	struct FPurgeRequest
 	{
 		FOnDemandPurgeArgs					Args;
@@ -295,14 +324,15 @@ public:
 	FOnDemandChunkInfo					GetInstalledChunkInfo(const FIoChunkId& ChunkId);
 	void								ReleaseContent(FOnDemandInternalContentHandle& ContentHandle);
 	void								GetReferencedContent(TArray<FSharedOnDemandContainer>& OutContainers, TArray<TBitArray<>>& OutChunkEntryIndices);
+	void								CancelInstallRequest(FSharedInternalInstallRequest InstallRequest);
+	void								UpdateInstallRequestPriority(FSharedInternalInstallRequest InstallRequest, EOnDemandInstallPriority NewPriority);
 
 	// IOnDemandIoStore interface
 	virtual void						Mount(FOnDemandMountArgs&& Args, FOnDemandMountCompleted&& OnCompleted) override;
-	virtual void						Install(
+	virtual FOnDemandInstallRequest 	Install(
 											FOnDemandInstallArgs&& Args,
 											FOnDemandInstallCompleted&& OnCompleted,
-											FOnDemandInstallProgressed&& OnProgress = nullptr,
-											const FOnDemandCancellationToken* CancellationToken = nullptr) override;
+											FOnDemandInstallProgressed&& OnProgress = nullptr) override;
 	virtual void						Purge(FOnDemandPurgeArgs&& Args, FOnDemandPurgeCompleted&& OnCompleted) override;
 	virtual FIoStatus					Unmount(FStringView MountId) override;
 	virtual TIoStatusOr<uint64> 		GetInstallSize(const FOnDemandGetInstallSizeArgs& Args) const override;
@@ -327,9 +357,9 @@ private:
 	bool					Tick();
 	FIoStatus				TickMountRequest(FMountRequest& MountRequest);
 	void					CompleteMountRequest(FMountRequest& MountRequest, FOnDemandMountResult&& MountResult);
-	FOnDemandInstallResult	TickInstallRequest(const FInstallRequest& InstallRequest);
-	void					CompleteInstallRequest(FInstallRequest& InstallRequest, FOnDemandInstallResult&& InstallResult);
-	void					ProgressInstallRequest(const FInstallRequest& InstallRequest, const FOnDemandInstallProgress& Progress);
+	FOnDemandInstallResult	TickInstallRequest(FOnDemandInternalInstallRequest& InstallRequest);
+	void					CompleteInstallRequest(FSharedInternalInstallRequest InstallRequest, FOnDemandInstallResult&& InstallResult);
+	void					ProgressInstallRequest(const FOnDemandInternalInstallRequest& InstallRequest, const FOnDemandInstallProgress& Progress);
 	void					CompletePurgeRequest(FPurgeRequest& PurgeRequest, FOnDemandPurgeResult&& Result);
 	void					OnEncryptionKeyAdded(const FGuid& Id, const FAES::FAESKey& Key);
 	static void				CreateContainersFromToc(
@@ -340,24 +370,24 @@ private:
 								const TArray<FString>& ExistingContainers,
 								TArray<FSharedOnDemandContainer>& Out);
 
-	FSharedInstallCache					InstallCache;
-	FSharedPackageStoreBackend			PackageStoreBackend;
-	FSharedStreamingBackend				StreamingBackend;
-	FDelegateHandle						OnMountPakHandle;
-	FDelegateHandle						OnServerPostForkHandle;
-	TArray<FSharedOnDemandContainer>	Containers;
-	TMap<FString, FIoBuffer>			PendingContainerHeaders;
-	mutable UE::FMutex					ContainerMutex;
-
-	TArray<FSharedMountRequest>			MountRequests;
-	TArray<FSharedInstallRequest>		InstallRequests;
-	TArray<FSharedPurgeRequest>			PurgeRequests;
-	UE::FMutex							RequestMutex;
-
-	EOnDemandStreamingOptions			StreamingOptions = EOnDemandStreamingOptions::Enabled;
-	bool								bTicking = false;
-	bool								bTickRequested = false;
-	TFuture<void>						TickFuture;
+	FSharedInstallCache						InstallCache;
+	FSharedPackageStoreBackend				PackageStoreBackend;
+	FSharedStreamingBackend					StreamingBackend;
+	FDelegateHandle							OnMountPakHandle;
+	FDelegateHandle							OnServerPostForkHandle;
+	TArray<FSharedOnDemandContainer>		Containers;
+	TMap<FString, FIoBuffer>				PendingContainerHeaders;
+	mutable UE::FMutex						ContainerMutex;
+
+	TArray<FSharedMountRequest>				MountRequests;
+	TArray<FSharedInternalInstallRequest>	InstallRequests;
+	TArray<FSharedPurgeRequest>				PurgeRequests;
+	UE::FMutex								RequestMutex;
+
+	EOnDemandStreamingOptions				StreamingOptions = EOnDemandStreamingOptions::Enabled;
+	bool									bTicking = false;
+	bool									bTickRequested = false;
+	TFuture<void>							TickFuture;
 };
 
 } // namespace UE::IoStore
diff --git a/Engine/Source/Runtime/Experimental/IoStore/OnDemand/Public/IO/IoStoreOnDemand.h b/Engine/Source/Runtime/Experimental/IoStore/OnDemand/Public/IO/IoStoreOnDemand.h
--- a/Engine/Source/Runtime/Experimental/IoStore/OnDemand/Public/IO/IoStoreOnDemand.h
+++ b/Engine/Source/Runtime/Experimental/IoStore/OnDemand/Public/IO/IoStoreOnDemand.h
@@ -19,7 +19,6 @@
 #define UE_API IOSTOREONDEMAND_API
 
 struct FAnalyticsEventAttribute;
-namespace UE::IoStore { class FOnDemandInternalContentHandle; }
 
 // Custom initialization allows users to control when
 // the system should be initialized.
@@ -33,6 +32,59 @@ UE_API DECLARE_LOG_CATEGORY_EXTERN(LogIas, Log, All);
 namespace UE::IoStore
 {
 
+using FWeakOnDemandIoStore			= TWeakPtr<class FOnDemandIoStore, ESPMode::ThreadSafe>;
+using FSharedInternalContentHandle	= TSharedPtr<class FOnDemandInternalContentHandle, ESPMode::ThreadSafe>;
+using FSharedInternalInstallRequest	= TSharedPtr<class FOnDemandInternalInstallRequest, ESPMode::ThreadSafe>;
+
+/** Interface for on-demand request types. */
+class FOnDemandRequest
+{
+public:
+	/** Request status. */
+	enum EStatus : uint8
+	{
+		/** The request has not been issued. */
+		None,
+		/** The request is pending. */
+		Pending,
+		/** The request completed successfully. */
+		Ok,
+		/** The request was cancelled. */
+		Cancelled,
+		/** The request failed. */
+		Error
+	};
+	/** Request priority. */
+	enum EPriority : uint8
+	{
+		/** Below normal priority. */
+		Low,
+		/** Normal priority. */
+		Normal,
+		/** Above normal priority. */
+		High
+	};
+	/** Destructor. */
+	virtual ~FOnDemandRequest() = default;
+	/** Returns the current status. */
+	virtual EStatus GetStatus() const = 0;
+	/** Returns true if the request is invalid. */
+	bool IsNone() const { return EStatus::None == GetStatus(); }
+	/** Returns true if the request is pending. */
+	bool IsPending() const { return EStatus::Pending == GetStatus(); }
+	/** Returns true if the request was successful. */
+	bool IsOk() const { return EStatus::Ok == GetStatus(); }
+	/** Returns true if the request was cancelled. */
+	bool IsCancelled() const { return EStatus::Cancelled == GetStatus(); }
+	/** Returns true if the request was unsuccessful. */
+	bool IsError() const { return EStatus::Error == GetStatus(); }
+	/** Returns true if the request is completed. */
+	bool IsCompleted() const { return static_cast<uint8>(EStatus::Pending) < static_cast<uint8>(GetStatus()); }
+
+protected:
+	FOnDemandRequest() = default;
+};
+
 /**
  * Keeps referenced data pinned in the cache until released.
  */
@@ -58,7 +110,7 @@ public:
 
 private:
 	friend class FOnDemandIoStore;
-	TSharedPtr<FOnDemandInternalContentHandle, ESPMode::ThreadSafe> Handle;
+	FSharedInternalContentHandle Handle;
 };
 
 /** Options for controlling the behavior of mounted container(s). */
@@ -128,6 +180,9 @@ enum class EOnDemandInstallOptions
 };
 ENUM_CLASS_FLAGS(EOnDemandInstallOptions);
 
+/** Install priority. */
+using EOnDemandInstallPriority = FOnDemandRequest::EPriority;
+
 /** Arguments for installing/downloading on-demand content. */
 struct FOnDemandInstallArgs 
 {
@@ -141,6 +196,10 @@ struct FOnDemandInstallArgs
 	FOnDemandContentHandle ContentHandle;
 	/** Install options. */
 	EOnDemandInstallOptions Options = EOnDemandInstallOptions::None;
+	/** Priority. */
+	EOnDemandInstallPriority Priority = EOnDemandInstallPriority::Low;
+	/** Optional debug name . */
+	FSharedString DebugName;
 };
 
 /** Holds information about progress for an install request. */
@@ -217,6 +276,36 @@ struct FOnDemandInstallResult
 /** Install completion callback. */
 using FOnDemandInstallCompleted = TUniqueFunction<void(FOnDemandInstallResult)>;
 
+/** A single-ownership handle to an install request. */
+class FOnDemandInstallRequest final
+	: public FOnDemandRequest
+{
+	friend class FOnDemandIoStore;
+public:
+	/** Creates an invalid install request. */
+	UE_API FOnDemandInstallRequest();
+	/** Move constructable. */;
+	UE_API FOnDemandInstallRequest(FOnDemandInstallRequest&&);
+	/** Move assignable. */
+	UE_API FOnDemandInstallRequest& operator=(FOnDemandInstallRequest&&);
+	/** Destructor. */
+	UE_API ~FOnDemandInstallRequest();
+	/** Returns the current request status. */
+	UE_API virtual EStatus GetStatus() const override;
+	/** Cancel the install request. */
+	UE_API void Cancel();
+	/** Update priority of the instlal request. */
+	UE_API void UpdatePriority(EPriority NewPriority);
+
+private:
+	FOnDemandInstallRequest(FWeakOnDemandIoStore IoStore, FSharedInternalInstallRequest InternalRequest);
+	FOnDemandInstallRequest(const FOnDemandInstallRequest&) = delete;
+	FOnDemandInstallRequest& operator=(const FOnDemandInstallRequest&) = delete;
+
+	FWeakOnDemandIoStore IoStore;	
+	FSharedInternalInstallRequest Request;
+};
+
 /** Arguments for purging on-demand content. */
 struct FOnDemandPurgeArgs
 {
@@ -273,29 +362,6 @@ struct FOnDemandCacheUsage
 	FOnDemandStreamingCacheUsage StreamingCache;
 };
 
-/** Token used for signalling an operation to be cancelled. */
-class FOnDemandCancellationToken
-{
-public:
-	UE_NONCOPYABLE(FOnDemandCancellationToken);
-	FOnDemandCancellationToken() = default;
-
-	/** Signal the operation to be cancelled. */
-	void Cancel()
-	{
-		bCanceled = true;
-	}
-
-	/** Returns whether an operation should be cancelled. */
-	bool IsCanceled() const
-	{
-		return bCanceled;
-	}
-
-private:
-	std::atomic<bool> bCanceled{ false };
-};
-
 /** Interface for recording analytics over a given time period */
 class IAnalyticsRecording
 {
@@ -346,11 +412,10 @@ public:
 	 * @param	OnProgress			Optional progress callback.
 	 * @param	CancellationToken	Optional cancellation token.
 	 **/
-	virtual void Install(
+	virtual FOnDemandInstallRequest Install(
 			FOnDemandInstallArgs&& Args,
 			FOnDemandInstallCompleted&& OnCompleted,
-			FOnDemandInstallProgressed&& OnProgress = nullptr,
-			const FOnDemandCancellationToken* CancellationToken = nullptr) = 0;
+			FOnDemandInstallProgressed&& OnProgress = nullptr) = 0;
 	/**
 	 * Purge the cache.
 	 * @param	Args		The purge arguments.
