Removing experimental IAS-based Zen cache - the bespoke Zen cache is now fully functional.

NB. IAS cache header moved back to Private.

#rnx

diff --git a/Engine/Source/Runtime/Experimental/IoStore/OnDemand/Internal/IO/IasCache.h b/Engine/Source/Runtime/Experimental/IoStore/OnDemand/Internal/IO/IasCache.h
deleted file mode 100644
diff --git a/Engine/Source/Runtime/Experimental/IoStore/OnDemand/Private/IasCache.h b/Engine/Source/Runtime/Experimental/IoStore/OnDemand/Private/IasCache.h
new file mode 100644
--- /dev/null
+++ b/Engine/Source/Runtime/Experimental/IoStore/OnDemand/Private/IasCache.h
@@ -0,0 +1,84 @@
+// Copyright Epic Games, Inc. All Rights Reserved.
+
+#pragma once
+
+#include "Containers/StringView.h"
+#include "IO/IoStatus.h"
+#include "Memory/MemoryFwd.h"
+#include "Tasks/Task.h"
+#include "Templates/UniquePtr.h"
+
+class FIoBuffer;
+class FIoReadOptions;
+struct FIoHash;
+
+namespace UE::IoStore
+{
+
+/** Cache for binary blobs with a 20 byte cache key. */
+class IIasCache
+{
+public:
+	virtual ~IIasCache() = default;
+
+	/** Deletes the IAS object, dropping all data persisted to disk and releasing
+	OS resources. As this also deletes the object thus any unique pointers should
+	be released prior to abandonment; TUniquePtr->Release()->Abandon(). Be sure
+	to cancel and collect any Get() tasks beforehand. */
+	virtual void Abandon() = 0;
+
+	/** Returns whether the specified cache key is present in the cache. */
+	virtual bool ContainsChunk(const FIoHash& Key) const = 0;
+
+	/** Get the chunk associated with the specified cache key. If the data is
+	 already in memory it is return in OutData. Otherwise the returned status
+	 indicates if the key can be materialized or if it does not exist */
+	virtual EIoErrorCode Get(const FIoHash& Key, FIoBuffer& OutData) = 0;
+
+	/** Materializes a cached items data from disk. The data is read into Dest
+	 so Dest must remain valid throughout. The result of the disk read is returned
+	 in Status (same lifetime needs as Dest). DoneEvent is triggered when the read
+	 succeeds, is not found, or if an IO error occurred. */
+	virtual void Materialize(const FIoHash& Key, FIoBuffer& Dest, EIoErrorCode& Status, UE::Tasks::FTaskEvent DoneEvent) = 0;
+
+	/** Cancels a previously request to Materialize(). Note that the materialize
+	 already can still complete as the read operation may already be in flight. If
+	 the cancel succeeds, the materialize's DoneEvent is not triggered. */
+	virtual void Cancel(FIoBuffer& GivenDest) = 0;
+
+	/** Insert a new chunk into the cache. */
+	virtual FIoStatus Put(const FIoHash& Key, FIoBuffer& Data) = 0;
+
+	/** Get cache usage information. */
+	virtual void GetCacheUsage(uint64& OutDiskSize, uint64& OutMaxDiskSize) const = 0;
+};
+
+struct FIasCacheConfig
+{
+	struct FRate
+	{
+		uint32	Allowance = 16 << 20;
+		uint32	Ops = 32;
+		uint32	Seconds = 60;
+	};
+
+	struct FDemand
+	{
+		uint8	Threshold = 30;
+		uint8	Boost = 60;
+		uint8	SuperBoost = 87;
+	};
+
+	FStringView Name = TEXT("ias");
+	uint64		DiskQuota = 512ull << 20;
+	uint32		MemoryQuota = 2 << 20;
+	uint32		JournalQuota = 4 << 20; // description in JournalCache.cpp
+	uint32		JournalMagic = 0; // can be used to invalidate cache
+	FRate		WriteRate;
+	FDemand		Demand;
+	bool		DropCache = false;
+};
+
+TUniquePtr<IIasCache> MakeIasCache(const TCHAR* RootDir, const FIasCacheConfig& Config);
+
+} // namespace UE::IoStore
diff --git a/Engine/Source/Runtime/Experimental/IoStore/OnDemand/Private/JournaledCache.cpp b/Engine/Source/Runtime/Experimental/IoStore/OnDemand/Private/JournaledCache.cpp
--- a/Engine/Source/Runtime/Experimental/IoStore/OnDemand/Private/JournaledCache.cpp
+++ b/Engine/Source/Runtime/Experimental/IoStore/OnDemand/Private/JournaledCache.cpp
@@ -1,6 +1,6 @@
 // Copyright Epic Games, Inc. All Rights Reserved.
 
-#include "IO/IasCache.h"
+#include "IasCache.h"
 #include "IO/IoStoreOnDemand.h"
 #include "Statistics.h"
 
diff --git a/Engine/Source/Runtime/Experimental/IoStore/OnDemand/Private/OnDemandConfig.cpp b/Engine/Source/Runtime/Experimental/IoStore/OnDemand/Private/OnDemandConfig.cpp
--- a/Engine/Source/Runtime/Experimental/IoStore/OnDemand/Private/OnDemandConfig.cpp
+++ b/Engine/Source/Runtime/Experimental/IoStore/OnDemand/Private/OnDemandConfig.cpp
@@ -6,7 +6,7 @@
 #include "OnDemandInstallCache.h"
 
 #include "HAL/PlatformProcess.h"
-#include "IO/IasCache.h"
+#include "IasCache.h"
 #include "IO/IoStoreOnDemand.h"
 #include "Misc/Base64.h"
 #include "Misc/ConfigCacheIni.h"
diff --git a/Engine/Source/Runtime/Experimental/IoStore/OnDemand/Private/OnDemandIoDispatcherBackend.cpp b/Engine/Source/Runtime/Experimental/IoStore/OnDemand/Private/OnDemandIoDispatcherBackend.cpp
--- a/Engine/Source/Runtime/Experimental/IoStore/OnDemand/Private/OnDemandIoDispatcherBackend.cpp
+++ b/Engine/Source/Runtime/Experimental/IoStore/OnDemand/Private/OnDemandIoDispatcherBackend.cpp
@@ -14,7 +14,7 @@
 #include "HAL/PreprocessorHelpers.h"
 #include "HAL/Runnable.h"
 #include "HAL/RunnableThread.h"
-#include "IO/IasCache.h"
+#include "IasCache.h"
 #include "IO/IoAllocators.h"
 #include "IO/IoChunkEncoding.h"
 #include "IO/IoDispatcher.h"
diff --git a/Engine/Source/Runtime/Experimental/IoStore/OnDemand/Private/OnDemandIoStore.cpp b/Engine/Source/Runtime/Experimental/IoStore/OnDemand/Private/OnDemandIoStore.cpp
--- a/Engine/Source/Runtime/Experimental/IoStore/OnDemand/Private/OnDemandIoStore.cpp
+++ b/Engine/Source/Runtime/Experimental/IoStore/OnDemand/Private/OnDemandIoStore.cpp
@@ -7,7 +7,7 @@
 #include "OnDemandInstallCache.h"
 #include "OnDemandPackageStoreBackend.h"
 #include "OnDemandIoDispatcherBackend.h"
-#include "IO/IasCache.h"
+#include "IasCache.h"
 
 #include "Algo/Accumulate.h"
 #include "Algo/AnyOf.h"
diff --git a/Engine/Source/Runtime/StorageServerClient/Private/StorageServerIASCache.cpp b/Engine/Source/Runtime/StorageServerClient/Private/StorageServerIASCache.cpp
deleted file mode 100644
diff --git a/Engine/Source/Runtime/StorageServerClient/Private/StorageServerIASCache.h b/Engine/Source/Runtime/StorageServerClient/Private/StorageServerIASCache.h
deleted file mode 100644
diff --git a/Engine/Source/Runtime/StorageServerClient/Private/StorageServerIoDispatcherBackend.cpp b/Engine/Source/Runtime/StorageServerClient/Private/StorageServerIoDispatcherBackend.cpp
--- a/Engine/Source/Runtime/StorageServerClient/Private/StorageServerIoDispatcherBackend.cpp
+++ b/Engine/Source/Runtime/StorageServerClient/Private/StorageServerIoDispatcherBackend.cpp
@@ -29,13 +29,6 @@ FStorageServerIoDispatcherBackend::FStorageServerIoDispatcherBackend(FStorageSer
 	, NewRequestEvent(FPlatformProcess::GetSynchEventFromPool(false))
 	, BatchCompletedEvent(FPlatformProcess::GetSynchEventFromPool(false))
 {
-#if WITH_STORAGESERVER_IAS_CACHE
-	bool bUseIASCache = FParse::Param(FCommandLine::Get(), TEXT("ZenIASCache"));
-	if (bUseIASCache)
-	{
-		Cache = MakeShared<FStorageServerIASCache>(Connection);
-	}
-#endif
 }
 
 FStorageServerIoDispatcherBackend::~FStorageServerIoDispatcherBackend()
@@ -164,12 +157,6 @@ bool FStorageServerIoDispatcherBackend::Resolve(FIoRequestImpl* Request)
 		{
 			Request->SetResult(Result.ValueOrDie());
 			TRACE_IOSTORE_BACKEND_REQUEST_COMPLETED(Request, Result.ValueOrDie().GetSize());
-#if WITH_STORAGESERVER_IAS_CACHE
-			if (Cache.IsValid())
-			{
-				Cache->OptionalCacheWrite(Request, Result.ValueOrDie());
-			}
-#endif
 		}
 		else
 		{
@@ -366,13 +353,6 @@ void FStorageServerIoDispatcherBackend::FBatch::DoThreadedWork()
 		{
 			Request->SetResult(Result.ValueOrDie());
 			TRACE_IOSTORE_BACKEND_REQUEST_COMPLETED(Request, Result.ValueOrDie().GetSize());
-
-#if WITH_STORAGESERVER_IAS_CACHE
-			if (Owner.Cache.IsValid())
-			{
-				Owner.Cache->OptionalCacheWrite(Request, Result.ValueOrDie());
-			}
-#endif
 		}
 		else
 		{
diff --git a/Engine/Source/Runtime/StorageServerClient/Private/StorageServerIoDispatcherBackend.h b/Engine/Source/Runtime/StorageServerClient/Private/StorageServerIoDispatcherBackend.h
--- a/Engine/Source/Runtime/StorageServerClient/Private/StorageServerIoDispatcherBackend.h
+++ b/Engine/Source/Runtime/StorageServerClient/Private/StorageServerIoDispatcherBackend.h
@@ -3,7 +3,6 @@
 #pragma once
 
 #include "IO/IoDispatcherBackend.h"
-#include "StorageServerIASCache.h"
 #include "HAL/Runnable.h"
 #include "Misc/IQueuedWork.h"
 
@@ -99,9 +98,6 @@ private:
 	FIoRequestImpl* CompletedRequestsHead = nullptr;
 	FIoRequestImpl* CompletedRequestsTail = nullptr;
 	TAtomic<bool> bStopRequested{ false };
-#if WITH_STORAGESERVER_IAS_CACHE
-	TSharedPtr<FStorageServerIASCache> Cache;
-#endif
 };
 
 #endif
diff --git a/Engine/Source/Runtime/StorageServerClient/StorageServerClient.Build.cs b/Engine/Source/Runtime/StorageServerClient/StorageServerClient.Build.cs
--- a/Engine/Source/Runtime/StorageServerClient/StorageServerClient.Build.cs
+++ b/Engine/Source/Runtime/StorageServerClient/StorageServerClient.Build.cs
@@ -13,7 +13,6 @@ public class StorageServerClient : ModuleRules
 				"Sockets",
 				"CookOnTheFly",
 				"Json",
-				"IoStoreOnDemand" // for IAS cache
 			}
 		);
 	}
