Chaos Softs::FEvolution: make enabling/disabling particle ranges in different groups thread-safe. The cloth solver processes cloths (groups) in parallel.

diff --git a/Engine/Plugins/ChaosCloth/Source/ChaosCloth/Private/ChaosCloth/ChaosClothingSimulationSolver.cpp b/Engine/Plugins/ChaosCloth/Source/ChaosCloth/Private/ChaosCloth/ChaosClothingSimulationSolver.cpp
--- a/Engine/Plugins/ChaosCloth/Source/ChaosCloth/Private/ChaosCloth/ChaosClothingSimulationSolver.cpp
+++ b/Engine/Plugins/ChaosCloth/Source/ChaosCloth/Private/ChaosCloth/ChaosClothingSimulationSolver.cpp
@@ -1625,7 +1625,7 @@ void FClothingSimulationSolver::ApplyPreSimulationTransforms()
 
 	if (Evolution)
 	{
-		const TArray<uint32> ActiveGroups = Evolution->GetActiveGroups().Array();
+		const TArray<uint32> ActiveGroups = Evolution->GetActiveGroupsArray();
 		PhysicsParallelFor(ActiveGroups.Num(), [this, &ActiveGroups, &DeltaLocalSpaceLocation, MaxVelocitySquared](int32 ActiveGroupIndex)
 		{
 			const uint32 GroupId = ActiveGroups[ActiveGroupIndex];
@@ -1846,7 +1846,7 @@ void FClothingSimulationSolver::PreSubstep( const Softs::FSolverReal Interpolati
 
 	if (Evolution)
 	{
-		const TArray<uint32> ActiveGroups = Evolution->GetActiveGroups().Array();
+		const TArray<uint32> ActiveGroups = Evolution->GetActiveGroupsArray();
 		PhysicsParallelFor(ActiveGroups.Num(), [this, &ActiveGroups, InterpolationAlpha, bDetectSelfCollisions](int32 ActiveGroupIndex)
 		{
 			const uint32 GroupId = ActiveGroups[ActiveGroupIndex];
@@ -1979,7 +1979,7 @@ void FClothingSimulationSolver::UpdateSolverField()
 			SampleIndices.SetNum(NumActiveParticles, EAllowShrinking::No);
 
 			int32 SampleIndex = 0;
-			for (const uint32 GroupId : Evolution->GetActiveGroups())
+			for (const uint32 GroupId : Evolution->GetActiveGroupsArray())
 			{
 				for (const int32 SoftBodyId : Evolution->GetGroupActiveSoftBodies(GroupId))
 				{
@@ -2247,7 +2247,8 @@ FBoxSphereBounds FClothingSimulationSolver::CalculateBounds() const
 	{
 		// Calculate bounding box
 		Softs::FSolverAABB3 BoundingBox = Softs::FSolverAABB3::EmptyAABB();
-		for (const uint32 GroupId : Evolution->GetActiveGroups())
+		const TArray<uint32> ActiveGroupsArray = Evolution->GetActiveGroupsArray();
+		for (const uint32 GroupId : ActiveGroupsArray)
 		{
 			for (const int32 SoftBodyId : Evolution->GetGroupActiveSoftBodies(GroupId))
 			{
@@ -2287,7 +2288,7 @@ FBoxSphereBounds FClothingSimulationSolver::CalculateBounds() const
 		const Softs::FSolverVec3 Center = BoundingBox.Center();
 		Softs::FSolverReal SquaredRadius = (Softs::FSolverReal)0.;
 
-		for (const uint32 GroupId : Evolution->GetActiveGroups())
+		for (const uint32 GroupId : ActiveGroupsArray)
 		{
 			for (const int32 SoftBodyId : Evolution->GetGroupActiveSoftBodies(GroupId))
 			{
diff --git a/Engine/Source/Runtime/Experimental/Chaos/Private/Chaos/SoftsEvolution.cpp b/Engine/Source/Runtime/Experimental/Chaos/Private/Chaos/SoftsEvolution.cpp
--- a/Engine/Source/Runtime/Experimental/Chaos/Private/Chaos/SoftsEvolution.cpp
+++ b/Engine/Source/Runtime/Experimental/Chaos/Private/Chaos/SoftsEvolution.cpp
@@ -189,10 +189,24 @@ void FEvolution::Reset()
 	CollisionKinematicUpdate.Reset();
 }
 
+TArray<uint32> FEvolution::GetActiveGroupsArray() const
+{
+	TArray<uint32> ActiveGroups;
+	ActiveGroups.Reserve(Groups.Size());
+	for (uint32 GroupId = 0; GroupId < Groups.Size(); ++GroupId)
+	{
+		if (!Groups.ActiveSoftBodies.IsEmpty())
+		{
+			ActiveGroups.Add(GroupId);
+		}
+	}
+	return ActiveGroups;
+}
+
 int32 FEvolution::NumActiveParticles() const
 {
 	int32 Result = 0;
-	for (const uint32 GroupId : GetActiveGroups())
+	for (const uint32 GroupId : GetActiveGroupsArray())
 	{
 		for (const int32 SoftBodyId : GetGroupActiveSoftBodies(GroupId))
 		{
@@ -293,15 +307,10 @@ void FEvolution::ActivateSoftBody(int32 SoftBodyId, bool bActivate)
 	if (bActivate)
 	{
 		Groups.ActiveSoftBodies[SoftBodies.GroupId[SoftBodyId]].Add(SoftBodyId);
-		ActiveGroups.Add(SoftBodies.GroupId[SoftBodyId]);
 	}
 	else
 	{
 		Groups.ActiveSoftBodies[SoftBodies.GroupId[SoftBodyId]].Remove(SoftBodyId);
-		if (Groups.ActiveSoftBodies[SoftBodies.GroupId[SoftBodyId]].IsEmpty())
-		{
-			ActiveGroups.Remove(SoftBodies.GroupId[SoftBodyId]);
-		}
 	}
 }
 
@@ -401,7 +410,7 @@ void FEvolution::AdvanceOneTimeStep(const FSolverReal Dt, const FSolverReal Time
 			Private::EvolutionSolverDefault::MinNumIterations, MaxNumIterations);
 	NumUsedIterations = TimeDependentNumIterations;
 
-	const TArray<uint32> ActiveGroupsArray = ActiveGroups.Array();
+	const TArray<uint32> ActiveGroupsArray = GetActiveGroupsArray();
 
 	PhysicsParallelFor(ActiveGroupsArray.Num(),
 		[this, &ActiveGroupsArray, Dt, TimeDependentNumIterations](int32 ActiveGroupIndex)
diff --git a/Engine/Source/Runtime/Experimental/Chaos/Public/Chaos/SoftsEvolution.h b/Engine/Source/Runtime/Experimental/Chaos/Public/Chaos/SoftsEvolution.h
--- a/Engine/Source/Runtime/Experimental/Chaos/Public/Chaos/SoftsEvolution.h
+++ b/Engine/Source/Runtime/Experimental/Chaos/Public/Chaos/SoftsEvolution.h
@@ -93,7 +93,12 @@ public:
 	// Giving non-const access so data can be set freely, but do not add or remove particles here. Use AddSoftBody
 	FSolverParticles& GetParticles() { return Particles; }
 
-	const TSet<uint32>& GetActiveGroups() const { return ActiveGroups; }
+	UE_DEPRECATED(5.6, "Use GetActiveGroupsArray instead")
+	TSet<uint32> GetActiveGroups() const 
+	{
+		return TSet<uint32>(GetActiveGroupsArray());
+	}
+	CHAOS_API TArray<uint32> GetActiveGroupsArray() const;
 	CHAOS_API int32 NumActiveParticles() const;
 
 	const FEvolutionGroupContext& GetGroupContext(uint32 GroupId) const
@@ -108,7 +113,7 @@ public:
 	}
 
 	/** 
-	 * Add a SoftBody to a Group.
+	 * Add a SoftBody to a Group. Adding and removing soft bodies is not threadsafe.
 	 * 
 	 * @return SoftBodyId
 	 */
@@ -121,6 +126,7 @@ public:
 	int32 GetSoftBodyGroupId(int32 SoftBodyId) const { return SoftBodies.GroupId[SoftBodyId]; }
 	CHAOS_API void SetSoftBodyProperties(int32 SoftBodyId, const FCollectionPropertyConstFacade& PropertyCollection,
 		const TMap<FString, TConstArrayView<FRealSingle>>& WeightMaps);
+	// Activating/deactivating Soft Bodies in different groups is threadsafe. Activations within a group is not threadsafe.
 	CHAOS_API void ActivateSoftBody(int32 SoftBodyId, bool bActivate);
 	bool IsSoftBodyActive(int32 SoftBodyId) const { return SoftBodies.Status[SoftBodyId] == FSoftBodies::EStatus::Active; }
 	FSolverParticlesRange& GetSoftBodyParticles(int32 SoftBodyId) { return SoftBodies.ParticleRanges[SoftBodyId]; }
@@ -131,12 +137,13 @@ public:
 	FSolverReal GetLastLinearSolveError(int32 SoftBodyId) const { return SoftBodies.LinearSystems[SoftBodyId].GetLastSolveError(); }
 	
 	/**
-	 * Add Collision particle range to a group.
+	 * Add Collision particle range to a group. Adding and removing collision particle ranges is not threadsafe.
 	 * 
 	 * @return Particle range offset (unique id for this range)
 	 */
 	CHAOS_API int32 AddCollisionParticleRange(uint32 GroupId, int32 NumParticles, bool bEnable);
 	CHAOS_API void RemoveCollisionParticleRange(int32 CollisionRangeId);
+	// Activating/deactivating Collision Particle ranges in different groups is threadsafe. Activations within a group is not threadsafe.
 	CHAOS_API void ActivateCollisionParticleRange(int32 CollisionRangeId, bool bEnable);
 	const TSet<int32>& GetGroupActiveCollisionParticleRanges(uint32 GroupId) const { return Groups.ActiveCollisionParticleRanges[GroupId]; }
 	CHAOS_API TArray<FSolverCollisionParticlesRange> GetActiveCollisionParticles(uint32 GroupId) const;
@@ -521,7 +528,6 @@ private:
 
 	// Per-Group data
 	FGroups Groups;
-	TSet<uint32> ActiveGroups; // Groups with at least one active softbody
 
 	KinematicUpdateFunc KinematicUpdate;
 	CollisionKinematicUpdateFunc CollisionKinematicUpdate;
