Commit 3 of 3 - Decoupling ISMPool from GeometryCollection

Moving GC to use ISMPool instead of GC_ISMPool
Deprecated existing GC_ISM files for licensee compatibility.

diff --git a/Engine/Source/Runtime/Experimental/GeometryCollectionEngine/GeometryCollectionEngine.Build.cs b/Engine/Source/Runtime/Experimental/GeometryCollectionEngine/GeometryCollectionEngine.Build.cs
--- a/Engine/Source/Runtime/Experimental/GeometryCollectionEngine/GeometryCollectionEngine.Build.cs
+++ b/Engine/Source/Runtime/Experimental/GeometryCollectionEngine/GeometryCollectionEngine.Build.cs
@@ -12,21 +12,22 @@ namespace UnrealBuildTool.Rules
 			PublicDependencyModuleNames.AddRange(
 				new string[]
 				{
+					"ChaosSolverEngine",
 					"Core",
 					"CoreUObject",
-                    "Engine",
-                    "RenderCore",
-                    "RHI",
-					"Renderer",
-                    "FieldSystemEngine",
-	                "ChaosSolverEngine",
-					"NetCore",
-					"IntelISPC",
 					"DataflowCore",
 					"DataflowEngine",
+					"Engine",
+					"FieldSystemEngine",
+					"IntelISPC",
+					"ISMPool",
 					"MeshDescription",
-					"StaticMeshDescription",
+					"NetCore",
+					"Renderer",
+					"RenderCore",
+					"RHI",
 					"SkeletalMeshDescription",
+					"StaticMeshDescription",
 				}
 				);
 			PrivateDependencyModuleNames.AddRange(
diff --git a/Engine/Source/Runtime/Experimental/GeometryCollectionEngine/Private/GeometryCollection/GeometryCollectionComponent.cpp b/Engine/Source/Runtime/Experimental/GeometryCollectionEngine/Private/GeometryCollection/GeometryCollectionComponent.cpp
--- a/Engine/Source/Runtime/Experimental/GeometryCollectionEngine/Private/GeometryCollection/GeometryCollectionComponent.cpp
+++ b/Engine/Source/Runtime/Experimental/GeometryCollectionEngine/Private/GeometryCollection/GeometryCollectionComponent.cpp
@@ -24,14 +24,14 @@
 #include "GeometryCollection/GeometryCollectionClusteringUtility.h"
 #include "GeometryCollection/GeometryCollectionComponentPluginPrivate.h"
 #include "GeometryCollection/GeometryCollectionDebugDrawComponent.h"
+#include "GeometryCollection/GeometryCollectionISMPoolRenderer.h"
 #include "GeometryCollection/GeometryCollectionObject.h"
 #include "GeometryCollection/GeometryCollectionProximityUtility.h"
 #include "GeometryCollection/GeometryCollectionSQAccelerator.h"
 #include "GeometryCollection/GeometryCollectionSceneProxy.h"
 #include "GeometryCollection/GeometryCollectionUtility.h"
-#include "GeometryCollection/GeometryCollectionISMPoolActor.h"
-#include "GeometryCollection/GeometryCollectionISMPoolComponent.h"
-#include "GeometryCollection/GeometryCollectionISMPoolRenderer.h"
+#include "ISMPool/ISMPoolActor.h"
+#include "ISMPool/ISMPoolComponent.h"
 #include "Math/Sphere.h"
 #include "Modules/ModuleManager.h"
 #include "Net/Core/PushModel/PushModel.h"
diff --git a/Engine/Source/Runtime/Experimental/GeometryCollectionEngine/Private/GeometryCollection/GeometryCollectionISMPoolActor.cpp b/Engine/Source/Runtime/Experimental/GeometryCollectionEngine/Private/GeometryCollection/GeometryCollectionISMPoolActor.cpp
--- a/Engine/Source/Runtime/Experimental/GeometryCollectionEngine/Private/GeometryCollection/GeometryCollectionISMPoolActor.cpp
+++ b/Engine/Source/Runtime/Experimental/GeometryCollectionEngine/Private/GeometryCollection/GeometryCollectionISMPoolActor.cpp
@@ -9,11 +9,15 @@
 AGeometryCollectionISMPoolActor::AGeometryCollectionISMPoolActor(const FObjectInitializer& ObjectInitializer)
 	: Super(ObjectInitializer)
 {
+PRAGMA_DISABLE_DEPRECATION_WARNINGS
 	ISMPoolComp = CreateDefaultSubobject<UGeometryCollectionISMPoolComponent>(TEXT("ISMPoolComp"));
+PRAGMA_ENABLE_DEPRECATION_WARNINGS
 	RootComponent = ISMPoolComp;
 
 #if UE_ENABLE_DEBUG_DRAWING
+PRAGMA_DISABLE_DEPRECATION_WARNINGS
 	ISMPoolDebugDrawComp = CreateDefaultSubobject<UGeometryCollectionISMPoolDebugDrawComponent>(TEXT("ISMPoolDebug"));
+PRAGMA_ENABLE_DEPRECATION_WARNINGS
 	ISMPoolDebugDrawComp->SetCollisionEnabled(ECollisionEnabled::NoCollision);
 	ISMPoolDebugDrawComp->SetCanEverAffectNavigation(false);
 	ISMPoolDebugDrawComp->CanCharacterStepUpOn = ECanBeCharacterBase::ECB_No;
diff --git a/Engine/Source/Runtime/Experimental/GeometryCollectionEngine/Private/GeometryCollection/GeometryCollectionISMPoolComponent.cpp b/Engine/Source/Runtime/Experimental/GeometryCollectionEngine/Private/GeometryCollection/GeometryCollectionISMPoolComponent.cpp
--- a/Engine/Source/Runtime/Experimental/GeometryCollectionEngine/Private/GeometryCollection/GeometryCollectionISMPoolComponent.cpp
+++ b/Engine/Source/Runtime/Experimental/GeometryCollectionEngine/Private/GeometryCollection/GeometryCollectionISMPoolComponent.cpp
@@ -269,7 +269,9 @@ FGeometryCollectionISMPool::FGeometryCollectionISMPool()
 {
 }
 
+PRAGMA_DISABLE_DEPRECATION_WARNINGS
 FGeometryCollectionISMPool::FISMIndex FGeometryCollectionISMPool::GetOrAddISM(UGeometryCollectionISMPoolComponent* OwningComponent, const FGeometryCollectionStaticMeshInstance& MeshInstance, bool& bOutISMCreated)
+PRAGMA_ENABLE_DEPRECATION_WARNINGS
 {
 	FISMIndex* ISMIndexPtr = MeshToISMIndex.Find(MeshInstance);
 	if (ISMIndexPtr != nullptr)
@@ -304,7 +306,9 @@ FGeometryCollectionISMPool::FISMIndex FGeometryCollectionISMPool::GetOrAddISM(UG
 	return ISMIndex;
 }
 
+PRAGMA_DISABLE_DEPRECATION_WARNINGS
 FGeometryCollectionMeshInfo FGeometryCollectionISMPool::AddInstancesToISM(UGeometryCollectionISMPoolComponent* OwningComponent, const FGeometryCollectionStaticMeshInstance& MeshInstance, int32 InstanceCount, TArrayView<const float> CustomDataFloats)
+PRAGMA_ENABLE_DEPRECATION_WARNINGS
 {
 	bool bISMCreated = false;
 	FGeometryCollectionMeshInfo Info;
@@ -517,7 +521,9 @@ static bool AreWeakPointersValid(FGeometryCollectionStaticMeshInstance& InMeshIn
 	return true;
 }
 
+PRAGMA_DISABLE_DEPRECATION_WARNINGS
 void FGeometryCollectionISMPool::ProcessPreallocationRequests(UGeometryCollectionISMPoolComponent* OwningComponent, int32 MaxPreallocations)
+PRAGMA_ENABLE_DEPRECATION_WARNINGS
 {
 	int32 NumAdded = 0;
 	for (TSet<FGeometryCollectionStaticMeshInstance>::TIterator It(PrellocationQueue); It; ++It)
@@ -573,7 +579,9 @@ void FGeometryCollectionISMPool::UpdateAbsoluteTransforms(const FTransform& Base
 	}
 }
 
+PRAGMA_DISABLE_DEPRECATION_WARNINGS
 void FGeometryCollectionISMPool::Tick(UGeometryCollectionISMPoolComponent* OwningComponent)
+PRAGMA_ENABLE_DEPRECATION_WARNINGS
 {
 	// Recache component lifecycle state from cvar.
 	const bool bRemovedKeepAlive = bCachedKeepAlive && !GComponentKeepAlive;
@@ -637,7 +645,9 @@ void UGeometryCollectionISMPoolComponent::TickComponent(float DeltaTime, enum EL
 	Pool.Tick(this);
 }
 
+PRAGMA_DISABLE_DEPRECATION_WARNINGS
 UGeometryCollectionISMPoolComponent::FMeshGroupId  UGeometryCollectionISMPoolComponent::CreateMeshGroup(bool bAllowPerInstanceRemoval)
+PRAGMA_ENABLE_DEPRECATION_WARNINGS
 {
 	FGeometryCollectionMeshGroup Group;
 	Group.bAllowPerInstanceRemoval = bAllowPerInstanceRemoval;
@@ -645,7 +655,9 @@ UGeometryCollectionISMPoolComponent::FMeshGroupId  UGeometryCollectionISMPoolCom
 	return NextMeshGroupId++;
 }
 
+PRAGMA_DISABLE_DEPRECATION_WARNINGS
 void UGeometryCollectionISMPoolComponent::DestroyMeshGroup(FMeshGroupId MeshGroupId)
+PRAGMA_ENABLE_DEPRECATION_WARNINGS
 {
 	if (FGeometryCollectionMeshGroup* MeshGroup = MeshGroups.Find(MeshGroupId))
 	{
@@ -654,7 +666,9 @@ void UGeometryCollectionISMPoolComponent::DestroyMeshGroup(FMeshGroupId MeshGrou
 	}
 }
 
+PRAGMA_DISABLE_DEPRECATION_WARNINGS
 UGeometryCollectionISMPoolComponent::FMeshId UGeometryCollectionISMPoolComponent::AddMeshToGroup(FMeshGroupId MeshGroupId, const FGeometryCollectionStaticMeshInstance& MeshInstance, int32 InstanceCount, TArrayView<const float> CustomDataFloats)
+PRAGMA_ENABLE_DEPRECATION_WARNINGS
 {
 	if (FGeometryCollectionMeshGroup* MeshGroup = MeshGroups.Find(MeshGroupId))
 	{
@@ -665,12 +679,16 @@ UGeometryCollectionISMPoolComponent::FMeshId UGeometryCollectionISMPoolComponent
 	return INDEX_NONE;
 }
 
+PRAGMA_DISABLE_DEPRECATION_WARNINGS
 bool UGeometryCollectionISMPoolComponent::BatchUpdateInstancesTransforms(FMeshGroupId MeshGroupId, FMeshId MeshId, int32 StartInstanceIndex, const TArray<FTransform>& NewInstancesTransforms, bool bWorldSpace, bool bMarkRenderStateDirty, bool bTeleport)
+PRAGMA_ENABLE_DEPRECATION_WARNINGS
 {
 	return BatchUpdateInstancesTransforms(MeshGroupId, MeshId, StartInstanceIndex, MakeArrayView(NewInstancesTransforms), bWorldSpace, bMarkRenderStateDirty, bTeleport);
 }
 
+PRAGMA_DISABLE_DEPRECATION_WARNINGS
 bool UGeometryCollectionISMPoolComponent::BatchUpdateInstancesTransforms(FMeshGroupId MeshGroupId, FMeshId MeshId, int32 StartInstanceIndex, TArrayView<const FTransform> NewInstancesTransforms, bool bWorldSpace, bool bMarkRenderStateDirty, bool bTeleport)
+PRAGMA_ENABLE_DEPRECATION_WARNINGS
 {
 	if (FGeometryCollectionMeshGroup* MeshGroup = MeshGroups.Find(MeshGroupId))
 	{
@@ -680,7 +698,9 @@ bool UGeometryCollectionISMPoolComponent::BatchUpdateInstancesTransforms(FMeshGr
 	return false;
 }
 
+PRAGMA_DISABLE_DEPRECATION_WARNINGS
 bool UGeometryCollectionISMPoolComponent::BatchUpdateInstanceCustomData(FMeshGroupId MeshGroupId, int32 CustomFloatIndex, float CustomFloatValue)
+PRAGMA_ENABLE_DEPRECATION_WARNINGS
 {
 	if (FGeometryCollectionMeshGroup* MeshGroup = MeshGroups.Find(MeshGroupId))
 	{
diff --git a/Engine/Source/Runtime/Experimental/GeometryCollectionEngine/Private/GeometryCollection/GeometryCollectionISMPoolRenderer.cpp b/Engine/Source/Runtime/Experimental/GeometryCollectionEngine/Private/GeometryCollection/GeometryCollectionISMPoolRenderer.cpp
--- a/Engine/Source/Runtime/Experimental/GeometryCollectionEngine/Private/GeometryCollection/GeometryCollectionISMPoolRenderer.cpp
+++ b/Engine/Source/Runtime/Experimental/GeometryCollectionEngine/Private/GeometryCollection/GeometryCollectionISMPoolRenderer.cpp
@@ -7,10 +7,10 @@
 #include "GeometryCollection/Facades/CollectionInstancedMeshFacade.h"
 #include "GeometryCollection/GeometryCollection.h"
 #include "GeometryCollection/GeometryCollectionComponent.h"
-#include "GeometryCollection/GeometryCollectionISMPoolActor.h"
-#include "GeometryCollection/GeometryCollectionISMPoolComponent.h"
-#include "GeometryCollection/GeometryCollectionISMPoolSubSystem.h"
 #include "GeometryCollection/GeometryCollectionObject.h"
+#include "ISMPool/ISMPoolActor.h"
+#include "ISMPool/ISMPoolComponent.h"
+#include "ISMPool/ISMPoolSubSystem.h"
 
 #include UE_INLINE_GENERATED_CPP_BY_NAME(GeometryCollectionISMPoolRenderer)
 
@@ -29,7 +29,7 @@ void UGeometryCollectionISMPoolRenderer::OnRegisterGeometryCollection(UGeometryC
 				LocalISMPoolComponent->DestroyComponent();
 			}
 
-			LocalISMPoolComponent = NewObject<UGeometryCollectionISMPoolComponent>(this, NAME_None, RF_Transient | RF_DuplicateTransient);
+			LocalISMPoolComponent = NewObject<UISMPoolComponent>(this, NAME_None, RF_Transient | RF_DuplicateTransient);
 			LocalISMPoolComponent->SetTickablePoolManagement(false);
 			LocalISMPoolComponent->SetupAttachment(&InComponent);
 			LocalISMPoolComponent->RegisterComponent();
@@ -104,7 +104,7 @@ void UGeometryCollectionISMPoolRenderer::UpdateTransforms(UGeometryCollection co
 	UpdateInstanceTransforms(InGeometryCollection, ComponentTransform, InTransforms);
 }
 
-UGeometryCollectionISMPoolComponent* UGeometryCollectionISMPoolRenderer::GetISMPoolComponent() const
+UISMPoolComponent* UGeometryCollectionISMPoolRenderer::GetISMPoolComponent() const
 {
 	if (!bIsRegistered)
 	{
@@ -113,7 +113,7 @@ UGeometryCollectionISMPoolComponent* UGeometryCollectionISMPoolRenderer::GetISMP
 	return LocalISMPoolComponent ? LocalISMPoolComponent : CachedISMPoolComponent;
 }
 
-UGeometryCollectionISMPoolComponent* UGeometryCollectionISMPoolRenderer::GetOrCreateISMPoolComponent()
+UISMPoolComponent* UGeometryCollectionISMPoolRenderer::GetOrCreateISMPoolComponent()
 {
 	if (!bIsRegistered)
 	{
@@ -125,11 +125,11 @@ UGeometryCollectionISMPoolComponent* UGeometryCollectionISMPoolRenderer::GetOrCr
 	}
 	if (!CachedISMPoolComponent)
 	{
-		if (UGeometryCollectionISMPoolSubSystem* ISMPoolSubSystem = UWorld::GetSubsystem<UGeometryCollectionISMPoolSubSystem>(GetWorld()))
+		if (UISMPoolSubSystem* ISMPoolSubSystem = UWorld::GetSubsystem<UISMPoolSubSystem>(GetWorld()))
 		{
 			if (ULevel* Level = OwningLevel.Get())
 			{
-				if (AGeometryCollectionISMPoolActor* ISMPoolActor = ISMPoolSubSystem->FindISMPoolActor(Level))
+				if (AISMPoolActor* ISMPoolActor = ISMPoolSubSystem->FindISMPoolActor(Level))
 				{
 					CachedISMPoolComponent = ISMPoolActor->GetISMPoolComp();
 				}
@@ -147,7 +147,7 @@ void UGeometryCollectionISMPoolRenderer::InitMergedMeshFromGeometryCollection(UG
 		return;
 	}
 
-	UGeometryCollectionISMPoolComponent* ISMPoolComponent = GetOrCreateISMPoolComponent();
+	UISMPoolComponent* ISMPoolComponent = GetOrCreateISMPoolComponent();
 	MergedMeshGroup.GroupIndex = ISMPoolComponent != nullptr ? ISMPoolComponent->CreateMeshGroup() : INDEX_NONE;
 
 	if (MergedMeshGroup.GroupIndex == INDEX_NONE)
@@ -162,7 +162,7 @@ void UGeometryCollectionISMPoolRenderer::InitMergedMeshFromGeometryCollection(UG
 			continue;
 		}
 
-		FGeometryCollectionStaticMeshInstance StaticMeshInstance;
+		FISMPoolStaticMeshInstance StaticMeshInstance;
 		StaticMeshInstance.StaticMesh = StaticMesh;
 
 		TArray<float> DummyCustomData;
@@ -178,7 +178,7 @@ void UGeometryCollectionISMPoolRenderer::InitInstancesFromGeometryCollection(UGe
 		return;
 	}
 
-	UGeometryCollectionISMPoolComponent* ISMPoolComponent = GetOrCreateISMPoolComponent();
+	UISMPoolComponent* ISMPoolComponent = GetOrCreateISMPoolComponent();
 	InstancesGroup.GroupIndex = ISMPoolComponent != nullptr ? ISMPoolComponent->CreateMeshGroup() : INDEX_NONE;
 
 	if (InstancesGroup.GroupIndex == INDEX_NONE)
@@ -202,7 +202,7 @@ void UGeometryCollectionISMPoolRenderer::InitInstancesFromGeometryCollection(UGe
 					break;
 				}
 			}
-			FGeometryCollectionStaticMeshInstance StaticMeshInstance;
+			FISMPoolStaticMeshInstance StaticMeshInstance;
 			StaticMeshInstance.StaticMesh = const_cast<UStaticMesh*>(StaticMesh);
 			StaticMeshInstance.Desc.NumCustomDataFloats = AutoInstanceMesh.GetNumDataPerInstance();
 			if (bMaterialOverride)
@@ -223,7 +223,7 @@ void UGeometryCollectionISMPoolRenderer::UpdateMergedMeshTransforms(FTransform c
 		return;
 	}
 
-	UGeometryCollectionISMPoolComponent* ISMPoolComponent = GetOrCreateISMPoolComponent();
+	UISMPoolComponent* ISMPoolComponent = GetOrCreateISMPoolComponent();
 	if (ISMPoolComponent == nullptr)
 	{
 		return;
@@ -251,7 +251,7 @@ void UGeometryCollectionISMPoolRenderer::UpdateInstanceTransforms(UGeometryColle
 		return;
 	}
 
-	UGeometryCollectionISMPoolComponent* ISMPoolComponent = GetOrCreateISMPoolComponent();
+	UISMPoolComponent* ISMPoolComponent = GetOrCreateISMPoolComponent();
 	if (ISMPoolComponent == nullptr)
 	{
 		return;
@@ -292,7 +292,7 @@ void UGeometryCollectionISMPoolRenderer::ReleaseGroup(FISMPoolGroup& InOutGroup)
 
 	// Component and owning actor may already be released safely by a level unload.
 	// Don't want to create a new component here, so don't use GetOrCreateISMPoolComponent().
-	UGeometryCollectionISMPoolComponent* ISMPoolComponent = GetISMPoolComponent();
+	UISMPoolComponent* ISMPoolComponent = GetISMPoolComponent();
 	if (ISMPoolComponent != nullptr)
 	{
 		ISMPoolComponent->DestroyMeshGroup(InOutGroup.GroupIndex);
diff --git a/Engine/Source/Runtime/Experimental/GeometryCollectionEngine/Private/GeometryCollection/GeometryCollectionISMPoolRenderer.h b/Engine/Source/Runtime/Experimental/GeometryCollectionEngine/Private/GeometryCollection/GeometryCollectionISMPoolRenderer.h
--- a/Engine/Source/Runtime/Experimental/GeometryCollectionEngine/Private/GeometryCollection/GeometryCollectionISMPoolRenderer.h
+++ b/Engine/Source/Runtime/Experimental/GeometryCollectionEngine/Private/GeometryCollection/GeometryCollectionISMPoolRenderer.h
@@ -6,8 +6,8 @@
 
 #include "GeometryCollectionISMPoolRenderer.generated.h"
 
-class UGeometryCollectionISMPoolComponent;
 class UGeometryCollectionComponent;
+class UISMPoolComponent;
 class ULevel;
 
 /** Implementation of a geometry collection custom renderer that pushes AutoInstanceMeshes to an ISMPool. */
@@ -18,11 +18,11 @@ class UGeometryCollectionISMPoolRenderer : public UObject, public IGeometryColle
 
 	/** Instanced Static Mesh Pool component that is used to render our meshes. */
 	UPROPERTY(Transient)
-	TObjectPtr<UGeometryCollectionISMPoolComponent> CachedISMPoolComponent;
+	TObjectPtr<UISMPoolComponent> CachedISMPoolComponent;
 
 	/** Set if we have an Instanced Static Mesh Pool component owned by this renderer (ie when in Editor mode). Non-transient to behave correctly under actor duplication. */
 	UPROPERTY()
-	TObjectPtr<UGeometryCollectionISMPoolComponent> LocalISMPoolComponent;
+	TObjectPtr<UISMPoolComponent> LocalISMPoolComponent;
 
 public:
 	//~ Begin IGeometryCollectionExternalRenderInterface Interface.
@@ -56,8 +56,8 @@ protected:
 	bool bIsRegistered = false;
 
 private:
-	UGeometryCollectionISMPoolComponent* GetISMPoolComponent() const;
-	UGeometryCollectionISMPoolComponent* GetOrCreateISMPoolComponent();
+	UISMPoolComponent* GetISMPoolComponent() const;
+	UISMPoolComponent* GetOrCreateISMPoolComponent();
 	void InitMergedMeshFromGeometryCollection(UGeometryCollection const& InGeometryCollection);
 	void InitInstancesFromGeometryCollection(UGeometryCollection const& InGeometryCollection);
 	void UpdateMergedMeshTransforms(FTransform const& InBaseTransform, TArrayView<const FTransform3f> InLocalTransforms);
diff --git a/Engine/Source/Runtime/Experimental/GeometryCollectionEngine/Private/GeometryCollection/GeometryCollectionISMPoolSubSystem.cpp b/Engine/Source/Runtime/Experimental/GeometryCollectionEngine/Private/GeometryCollection/GeometryCollectionISMPoolSubSystem.cpp
--- a/Engine/Source/Runtime/Experimental/GeometryCollectionEngine/Private/GeometryCollection/GeometryCollectionISMPoolSubSystem.cpp
+++ b/Engine/Source/Runtime/Experimental/GeometryCollectionEngine/Private/GeometryCollection/GeometryCollectionISMPoolSubSystem.cpp
@@ -14,7 +14,9 @@ UGeometryCollectionISMPoolSubSystem::UGeometryCollectionISMPoolSubSystem()
 void UGeometryCollectionISMPoolSubSystem::Initialize(FSubsystemCollectionBase& Collection)
 {
 	Super::Initialize(Collection);
+PRAGMA_DISABLE_DEPRECATION_WARNINGS
 	Collection.InitializeDependency<UGeometryCollectionISMPoolSubSystem>();
+PRAGMA_ENABLE_DEPRECATION_WARNINGS
 }
 
 void UGeometryCollectionISMPoolSubSystem::Deinitialize()
@@ -23,6 +25,7 @@ void UGeometryCollectionISMPoolSubSystem::Deinitialize()
 	Super::Deinitialize();
 }
 
+PRAGMA_DISABLE_DEPRECATION_WARNINGS
 AGeometryCollectionISMPoolActor* UGeometryCollectionISMPoolSubSystem::FindISMPoolActor(ULevel* Level)
 {
 	AGeometryCollectionISMPoolActor* ISMPoolActor = nullptr;
@@ -50,16 +53,21 @@ AGeometryCollectionISMPoolActor* UGeometryCollectionISMPoolSubSystem::FindISMPoo
 	}
 	return ISMPoolActor;
 }
+PRAGMA_ENABLE_DEPRECATION_WARNINGS
 
 void UGeometryCollectionISMPoolSubSystem::OnActorEndPlay(AActor* InSource, EEndPlayReason::Type Reason)
 {
 	if (ULevel* ActorLevel = InSource->GetLevel())
 	{
+PRAGMA_DISABLE_DEPRECATION_WARNINGS
 		PerLevelISMPoolActors.Remove(ActorLevel);
+PRAGMA_ENABLE_DEPRECATION_WARNINGS
 	}
 }
 
+PRAGMA_DISABLE_DEPRECATION_WARNINGS
 void UGeometryCollectionISMPoolSubSystem::GetISMPoolActors(TArray<AGeometryCollectionISMPoolActor*>& OutActors) const
+PRAGMA_ENABLE_DEPRECATION_WARNINGS
 {
 	for (const auto& MapEntry : PerLevelISMPoolActors)
 	{
diff --git a/Engine/Source/Runtime/Experimental/GeometryCollectionEngine/Public/GeometryCollection/GeometryCollectionComponent.h b/Engine/Source/Runtime/Experimental/GeometryCollectionEngine/Public/GeometryCollection/GeometryCollectionComponent.h
--- a/Engine/Source/Runtime/Experimental/GeometryCollectionEngine/Public/GeometryCollection/GeometryCollectionComponent.h
+++ b/Engine/Source/Runtime/Experimental/GeometryCollectionEngine/Public/GeometryCollection/GeometryCollectionComponent.h
@@ -41,7 +41,7 @@ struct FGeometryCollectionDecayContext;
 struct FGeometryCollectionSection;
 struct FDamageCollector;
 class FPhysScene_Chaos;
-class AGeometryCollectionISMPoolActor;
+class AISMPoolActor;
 class IGeometryCollectionExternalRenderInterface;
 enum ESimulationInitializationState : uint8;
 enum class EClusterConnectionTypeEnum : uint8;
@@ -1497,7 +1497,7 @@ protected:
 
 	/** Deprecated for CustomRendererType. */
 	UPROPERTY()
-	TObjectPtr<AGeometryCollectionISMPoolActor> ISMPool_DEPRECATED;
+	TObjectPtr<AISMPoolActor> ISMPool_DEPRECATED;
 
 	/** Custom class type that will be used to render the geometry collection instead of using the native rendering. */
 	UPROPERTY(EditAnywhere, BlueprintReadWrite, Category = "ChaosPhysics|Rendering", meta = (editcondition = "bOverrideCustomRenderer", MustImplement = "/Script/GeometryCollectionEngine.GeometryCollectionExternalRenderInterface"))
diff --git a/Engine/Source/Runtime/Experimental/GeometryCollectionEngine/Public/GeometryCollection/GeometryCollectionISMPoolActor.h b/Engine/Source/Runtime/Experimental/GeometryCollectionEngine/Public/GeometryCollection/GeometryCollectionISMPoolActor.h
--- a/Engine/Source/Runtime/Experimental/GeometryCollectionEngine/Public/GeometryCollection/GeometryCollectionISMPoolActor.h
+++ b/Engine/Source/Runtime/Experimental/GeometryCollectionEngine/Public/GeometryCollection/GeometryCollectionISMPoolActor.h
@@ -8,7 +8,7 @@ class UGeometryCollectionISMPoolComponent;
 class UGeometryCollectionISMPoolDebugDrawComponent;
 
 UCLASS(ConversionRoot, ComponentWrapperClass, MinimalAPI)
-class AGeometryCollectionISMPoolActor : public AActor
+class UE_DEPRECATED(5.6, "AGeometryCollectionISMPoolActor is deprecated, please use AISMPoolActor instead.") AGeometryCollectionISMPoolActor : public AActor
 {
 	GENERATED_UCLASS_BODY()
 
diff --git a/Engine/Source/Runtime/Experimental/GeometryCollectionEngine/Public/GeometryCollection/GeometryCollectionISMPoolComponent.h b/Engine/Source/Runtime/Experimental/GeometryCollectionEngine/Public/GeometryCollection/GeometryCollectionISMPoolComponent.h
--- a/Engine/Source/Runtime/Experimental/GeometryCollectionEngine/Public/GeometryCollection/GeometryCollectionISMPoolComponent.h
+++ b/Engine/Source/Runtime/Experimental/GeometryCollectionEngine/Public/GeometryCollection/GeometryCollectionISMPoolComponent.h
@@ -334,7 +334,7 @@ struct FGeometryCollectionISMPool
 * Component that manages a pool of ISM components in order to allow multiple client components that use the same meshes to the share ISMs.
 */
 UCLASS(meta = (BlueprintSpawnableComponent), MinimalAPI)
-class UGeometryCollectionISMPoolComponent: public USceneComponent
+class UE_DEPRECATED(5.6, "UGeometryCollectionISMPoolComponent is deprecated, please use UISMPoolComponent instead.") UGeometryCollectionISMPoolComponent: public USceneComponent
 {
 	GENERATED_UCLASS_BODY()
 
diff --git a/Engine/Source/Runtime/Experimental/GeometryCollectionEngine/Public/GeometryCollection/GeometryCollectionISMPoolDebugDrawComponent.cpp b/Engine/Source/Runtime/Experimental/GeometryCollectionEngine/Public/GeometryCollection/GeometryCollectionISMPoolDebugDrawComponent.cpp
--- a/Engine/Source/Runtime/Experimental/GeometryCollectionEngine/Public/GeometryCollection/GeometryCollectionISMPoolDebugDrawComponent.cpp
+++ b/Engine/Source/Runtime/Experimental/GeometryCollectionEngine/Public/GeometryCollection/GeometryCollectionISMPoolDebugDrawComponent.cpp
@@ -26,17 +26,20 @@ static TAutoConsoleVariable<int32> CVarISMPoolStats(
 	TEXT("p.Chaos.GC.ISMPoolDebugStats"),
 	0,
 	TEXT("Show stats for the ISM pools"),
+PRAGMA_DISABLE_DEPRECATION_WARNINGS
 	FConsoleVariableDelegate::CreateLambda([](IConsoleVariable* InVariable) { UGeometryCollectionISMPoolDebugDrawComponent::UpdateAllTickEnabled(); }),
+PRAGMA_ENABLE_DEPRECATION_WARNINGS
 	ECVF_Default);
 
 static TAutoConsoleVariable<int32> CVarISMPoolDebugDraw(
 	TEXT("p.Chaos.GC.ISMPoolDebugDraw"),
 	0,
 	TEXT("Show debug drawing for the ISM pools"),
+PRAGMA_DISABLE_DEPRECATION_WARNINGS
 	FConsoleVariableDelegate::CreateLambda([](IConsoleVariable* InVariable) { UGeometryCollectionISMPoolDebugDrawComponent::UpdateAllTickEnabled(); FGlobalComponentRecreateRenderStateContext Context; }),
+PRAGMA_ENABLE_DEPRECATION_WARNINGS
 	ECVF_Default);
 
-
 class FGeometryCollectionISMPoolDebugDrawSceneProxy final : public FDebugRenderSceneProxy
 {
 public:
@@ -46,11 +49,13 @@ public:
 		return reinterpret_cast<size_t>(&UniquePointer);
 	}
 
+PRAGMA_DISABLE_DEPRECATION_WARNINGS
 	FGeometryCollectionISMPoolDebugDrawSceneProxy(UGeometryCollectionISMPoolDebugDrawComponent const* InComponent)
 		: FDebugRenderSceneProxy(InComponent)
 	{
 		DrawType = FDebugRenderSceneProxy::SolidAndWireMeshes;
 	}
+PRAGMA_ENABLE_DEPRECATION_WARNINGS
 
 	FPrimitiveViewRelevance GetViewRelevance(const FSceneView* View) const override
 	{
@@ -62,7 +67,6 @@ public:
 	}
 };
 
-
 UGeometryCollectionISMPoolDebugDrawComponent::UGeometryCollectionISMPoolDebugDrawComponent(const FObjectInitializer& ObjectInitializer)
 	: Super(ObjectInitializer)
 {
@@ -83,10 +87,12 @@ void UGeometryCollectionISMPoolDebugDrawComponent::UpdateTickEnabled()
 
 void UGeometryCollectionISMPoolDebugDrawComponent::UpdateAllTickEnabled()
 {
+PRAGMA_DISABLE_DEPRECATION_WARNINGS
 	for (TObjectIterator<UGeometryCollectionISMPoolDebugDrawComponent> It; It; ++It)
 	{
 		It->UpdateTickEnabled();
 	}
+PRAGMA_ENABLE_DEPRECATION_WARNINGS
 }
 
 void UGeometryCollectionISMPoolDebugDrawComponent::BeginPlay()
@@ -243,7 +249,9 @@ static void DrawBoxAndChildren(FDebugRenderSceneProxy* DebugProxy, FTransform co
 	}
 }
 
+PRAGMA_DISABLE_DEPRECATION_WARNINGS
 FDebugRenderSceneProxy* UGeometryCollectionISMPoolDebugDrawComponent::CreateDebugSceneProxy()
+PRAGMA_ENABLE_DEPRECATION_WARNINGS
 {
 	const bool bForceShowBounds = CVarISMPoolDebugDraw.GetValueOnAnyThread() != 0;
 	if (!(SelectedComponent && (bShowBounds || bForceShowBounds)))
@@ -280,8 +288,10 @@ void UGeometryCollectionISMPoolDebugDrawComponent::GetOnScreenMessages(TMultiMap
 
 	if (bShowGlobalStats || bForceShowStats)
 	{
+PRAGMA_DISABLE_DEPRECATION_WARNINGS
 		AGeometryCollectionISMPoolActor const* ISMPoolActor = Cast<AGeometryCollectionISMPoolActor>(GetOuter());
 		UGeometryCollectionISMPoolComponent const* ISMPoolComponent = ISMPoolActor ? ISMPoolActor->GetISMPoolComp() : nullptr;
+PRAGMA_ENABLE_DEPRECATION_WARNINGS
 
 		if (ISMPoolComponent)
 		{
@@ -289,7 +299,9 @@ void UGeometryCollectionISMPoolDebugDrawComponent::GetOnScreenMessages(TMultiMap
 			int32 NumHISMs = 0;
 			int32 NumInstancesTotal = 0;
 			TArray<int32> NumInstancesPerISM;
+PRAGMA_DISABLE_DEPRECATION_WARNINGS
 			NumInstancesPerISM.Reserve(ISMPoolComponent->Pool.ISMs.Num());
+PRAGMA_ENABLE_DEPRECATION_WARNINGS
 
 			for (FGeometryCollectionISM const& ISM : ISMPoolComponent->Pool.ISMs)
 			{
@@ -357,4 +369,4 @@ void UGeometryCollectionISMPoolDebugDrawComponent::GetOnScreenMessages(TMultiMap
 	}
 }
 
-#endif
+#endif
\ No newline at end of file
diff --git a/Engine/Source/Runtime/Experimental/GeometryCollectionEngine/Public/GeometryCollection/GeometryCollectionISMPoolDebugDrawComponent.h b/Engine/Source/Runtime/Experimental/GeometryCollectionEngine/Public/GeometryCollection/GeometryCollectionISMPoolDebugDrawComponent.h
--- a/Engine/Source/Runtime/Experimental/GeometryCollectionEngine/Public/GeometryCollection/GeometryCollectionISMPoolDebugDrawComponent.h
+++ b/Engine/Source/Runtime/Experimental/GeometryCollectionEngine/Public/GeometryCollection/GeometryCollectionISMPoolDebugDrawComponent.h
@@ -10,7 +10,7 @@ class FGeometryCollectionISMPoolDebugDrawDelegateHelper;
 class UInstancedStaticMeshComponent;
 
 UCLASS(ClassGroup = Debug)
-class UGeometryCollectionISMPoolDebugDrawComponent : public UDebugDrawComponent
+class UE_DEPRECATED(5.6, "UGeometryCollectionISMPoolDebugDrawComponent is deprecated, please use UISMPoolDebugDrawComponent instead.") UGeometryCollectionISMPoolDebugDrawComponent : public UDebugDrawComponent
 {
 	GENERATED_UCLASS_BODY()
 
diff --git a/Engine/Source/Runtime/Experimental/GeometryCollectionEngine/Public/GeometryCollection/GeometryCollectionISMPoolSubSystem.h b/Engine/Source/Runtime/Experimental/GeometryCollectionEngine/Public/GeometryCollection/GeometryCollectionISMPoolSubSystem.h
--- a/Engine/Source/Runtime/Experimental/GeometryCollectionEngine/Public/GeometryCollection/GeometryCollectionISMPoolSubSystem.h
+++ b/Engine/Source/Runtime/Experimental/GeometryCollectionEngine/Public/GeometryCollection/GeometryCollectionISMPoolSubSystem.h
@@ -12,7 +12,7 @@ class ULevel;
  * Used by geometry collection now but repurposed for more general use.
  */
 UCLASS(MinimalAPI)
-class UGeometryCollectionISMPoolSubSystem : public UWorldSubsystem
+class UE_DEPRECATED(5.6, "UGeometryCollectionISMPoolSubSystem is deprecated, please use UISMPoolSubSystem instead.") UGeometryCollectionISMPoolSubSystem : public UWorldSubsystem
 {
 	GENERATED_BODY()
 
@@ -25,15 +25,21 @@ public:
 	// USubsystem END
 
 	/** Finds or creates an actor. */
+PRAGMA_DISABLE_DEPRECATION_WARNINGS
 	GEOMETRYCOLLECTIONENGINE_API AGeometryCollectionISMPoolActor* FindISMPoolActor(ULevel* Level);
+PRAGMA_ENABLE_DEPRECATION_WARNINGS
 	
 	/** Get all actors managed by the subsystem. */
+PRAGMA_DISABLE_DEPRECATION_WARNINGS
 	GEOMETRYCOLLECTIONENGINE_API void GetISMPoolActors(TArray<AGeometryCollectionISMPoolActor*>& OutActors) const;
+PRAGMA_ENABLE_DEPRECATION_WARNINGS
 
 protected:
 	UFUNCTION()
 	void OnActorEndPlay(AActor* InSource, EEndPlayReason::Type Reason);
 
 	/** ISMPool are per level **/
+PRAGMA_DISABLE_DEPRECATION_WARNINGS
 	TMap<TObjectPtr<ULevel>, TObjectPtr<AGeometryCollectionISMPoolActor> > PerLevelISMPoolActors;
+PRAGMA_ENABLE_DEPRECATION_WARNINGS
 };
