Make bChaos_Bending_ISPCUseHWGather and bChaos_TriangleMesh_ISPCUseHWGather keyed off of new define USE_ISPC_UNOPTIMIZED_VECTORGATHER_DEFAULT.

Likely will make this affect all VectorGathers and remove the cvars in a future release.

diff --git a/Engine/Source/Runtime/Experimental/Chaos/Private/Chaos/PBDBendingConstraintsBase.cpp b/Engine/Source/Runtime/Experimental/Chaos/Private/Chaos/PBDBendingConstraintsBase.cpp
--- a/Engine/Source/Runtime/Experimental/Chaos/Private/Chaos/PBDBendingConstraintsBase.cpp
+++ b/Engine/Source/Runtime/Experimental/Chaos/Private/Chaos/PBDBendingConstraintsBase.cpp
@@ -9,7 +9,7 @@
 bool bChaos_Bending_ISPC_Enabled = CHAOS_BENDING_ISPC_ENABLED_DEFAULT;
 FAutoConsoleVariableRef CVarChaosBendingISPCEnabled(TEXT("p.Chaos.Bending.ISPC"), bChaos_Bending_ISPC_Enabled, TEXT("Whether to use ISPC optimizations in Bending constraints"));
 
-bool bChaos_Bending_ISPCUseHWGather = false;
+bool bChaos_Bending_ISPCUseHWGather = USE_ISPC_UNOPTIMIZED_VECTORGATHER_DEFAULT;
 FAutoConsoleVariableRef CVarChaosBendingISPCUseHWGather(TEXT("p.Chaos.Bending.ISPCUseHWGather"), bChaos_Bending_ISPCUseHWGather, TEXT("Use (slower) HW gather when computing Bending constraints to see if this fixes the ISPC crash."));
 #endif
 
diff --git a/Engine/Source/Runtime/Experimental/Chaos/Private/Chaos/TriangleMesh.cpp b/Engine/Source/Runtime/Experimental/Chaos/Private/Chaos/TriangleMesh.cpp
--- a/Engine/Source/Runtime/Experimental/Chaos/Private/Chaos/TriangleMesh.cpp
+++ b/Engine/Source/Runtime/Experimental/Chaos/Private/Chaos/TriangleMesh.cpp
@@ -32,15 +32,19 @@ static_assert(sizeof(ispc::TArrayInt) == sizeof(TArray<int32>), "sizeof(ispc::TA
 #define USE_ISPC_KERNEL_CONSOLE_VARIABLES_IN_SHIPPING 0
 #endif
 
+#if !defined(USE_ISPC_UNOPTIMIZED_VECTORGATHER_DEFAULT)
+#define USE_ISPC_UNOPTIMIZED_VECTORGATHER_DEFAULT 0
+#endif
+
 // Support run-time toggling on supported platforms in non-shipping configurations
 #if !INTEL_ISPC || (UE_BUILD_SHIPPING && !USE_ISPC_KERNEL_CONSOLE_VARIABLES_IN_SHIPPING)
 static constexpr bool bChaos_TriangleMesh_ISPC_Enabled = INTEL_ISPC && CHAOS_TRIANGLE_MESH_ISPC_ENABLED_DEFAULT;
-static constexpr bool bChaos_TriangleMesh_ISPCUseHWGather = false;
+static constexpr bool bChaos_TriangleMesh_ISPCUseHWGather = USE_ISPC_UNOPTIMIZED_VECTORGATHER_DEFAULT;
 #else
 static bool bChaos_TriangleMesh_ISPC_Enabled = CHAOS_TRIANGLE_MESH_ISPC_ENABLED_DEFAULT;
 static FAutoConsoleVariableRef CVarChaosTriangleMeshISPCEnabled(TEXT("p.Chaos.TriangleMesh.ISPC"), bChaos_TriangleMesh_ISPC_Enabled, TEXT("Whether to use ISPC optimizations in triangle mesh calculations"));
 
-static bool bChaos_TriangleMesh_ISPCUseHWGather = false;
+static bool bChaos_TriangleMesh_ISPCUseHWGather = USE_ISPC_UNOPTIMIZED_VECTORGATHER_DEFAULT;
 static FAutoConsoleVariableRef CVarChaosTriangleMeshISPCUseHWGather(TEXT("p.Chaos.TriangleMesh.ISPCUseHWGather"), bChaos_TriangleMesh_ISPCUseHWGather, TEXT("Use (slower) HW gather when computing triangle mesh normals to see if this fixes the ISPC crash."));
 
 #endif
diff --git a/Engine/Source/Runtime/Experimental/Chaos/Public/Chaos/PBDBendingConstraintsBase.h b/Engine/Source/Runtime/Experimental/Chaos/Public/Chaos/PBDBendingConstraintsBase.h
--- a/Engine/Source/Runtime/Experimental/Chaos/Public/Chaos/PBDBendingConstraintsBase.h
+++ b/Engine/Source/Runtime/Experimental/Chaos/Public/Chaos/PBDBendingConstraintsBase.h
@@ -408,10 +408,14 @@ protected:
 #define USE_ISPC_KERNEL_CONSOLE_VARIABLES_IN_SHIPPING 0
 #endif
 
+#if !defined(USE_ISPC_UNOPTIMIZED_VECTORGATHER_DEFAULT)
+#define USE_ISPC_UNOPTIMIZED_VECTORGATHER_DEFAULT 0
+#endif
+
 // Support run-time toggling on supported platforms in non-shipping configurations
 #if !INTEL_ISPC || (UE_BUILD_SHIPPING && !USE_ISPC_KERNEL_CONSOLE_VARIABLES_IN_SHIPPING)
 static constexpr bool bChaos_Bending_ISPC_Enabled = INTEL_ISPC && CHAOS_BENDING_ISPC_ENABLED_DEFAULT;
-static constexpr bool bChaos_Bending_ISPCUseHWGather = false;
+static constexpr bool bChaos_Bending_ISPCUseHWGather = USE_ISPC_UNOPTIMIZED_VECTORGATHER_DEFAULT;
 #else
 extern CHAOS_API bool bChaos_Bending_ISPC_Enabled;
 extern CHAOS_API bool bChaos_Bending_ISPCUseHWGather;
