Dataflow variable node : Fix Float type showing as double
- this is because internally selecting Float will result in a property of type double
the node now properly down cast it to a float

diff --git a/Engine/Source/Runtime/Experimental/Dataflow/Engine/Private/Dataflow/DataflowVariableNodes.cpp b/Engine/Source/Runtime/Experimental/Dataflow/Engine/Private/Dataflow/DataflowVariableNodes.cpp
--- a/Engine/Source/Runtime/Experimental/Dataflow/Engine/Private/Dataflow/DataflowVariableNodes.cpp
+++ b/Engine/Source/Runtime/Experimental/Dataflow/Engine/Private/Dataflow/DataflowVariableNodes.cpp
@@ -137,7 +137,8 @@ void FGetDataflowVariableNode::UpdateOutputTypes(const FPropertyBagPropertyDesc&
 		EvaluateFunction = &EvaluateFloat;
 		break;
 	case EPropertyBagPropertyType::Double:
-		SetOutputConcreteType(&Value, TDataflowSingleTypePolicy<double>::TypeName);
+		// the UI shows float but behind the scene set a double property let only deal with float 
+		SetOutputConcreteType(&Value, TDataflowSingleTypePolicy<float>::TypeName);
 		EvaluateFunction = &EvaluateDouble;
 		break;
 	case EPropertyBagPropertyType::Name:
@@ -292,12 +293,14 @@ void FGetDataflowVariableNode::Evaluate(UE::Dataflow::FContext& Context, const F
 
 /*static*/ void FGetDataflowVariableNode::EvaluateDouble(const FGetDataflowVariableNode& Node, const FInstancedPropertyBag& Variables, UE::Dataflow::FContext& Context, const FDataflowOutput& Out)
 {
-	double OutDouble = 0.0;
+	// UI shows type as float but internally always uses a double and never float 
+	// so let's output as float 
+	float OutFloat = 0.0;
 	if (const double* Value = Variables.GetValueDouble(Node.VariableName).TryGetValue())
 	{
-		OutDouble = *Value;
+		OutFloat = *Value;
 	}
-	Out.SetValue<double>(MoveTemp(OutDouble), Context);
+	Out.SetValue<float>(MoveTemp(OutFloat), Context);
 }
 
 /*static*/ void FGetDataflowVariableNode::EvaluateName(const FGetDataflowVariableNode& Node, const FInstancedPropertyBag& Variables, UE::Dataflow::FContext& Context, const FDataflowOutput& Out)
