Iris
* Added description to public NetRefHandle functions

#rnx

diff --git a/Engine/Source/Runtime/Experimental/Iris/Core/Public/Iris/ReplicationSystem/NetRefHandle.h b/Engine/Source/Runtime/Experimental/Iris/Core/Public/Iris/ReplicationSystem/NetRefHandle.h
--- a/Engine/Source/Runtime/Experimental/Iris/Core/Public/Iris/ReplicationSystem/NetRefHandle.h
+++ b/Engine/Source/Runtime/Experimental/Iris/Core/Public/Iris/ReplicationSystem/NetRefHandle.h
@@ -40,9 +40,13 @@ public:
 	uint32 GetReplicationSystemId() const { check(ReplicationSystemId != 0); return (uint32)(ReplicationSystemId - 1); }
 	bool IsValid() const { return Value != InvalidValue; }
 
+	/** Does the handle know which ReplicationSystem it is related to. */
 	bool IsCompleteHandle() const { return Value != InvalidValue && ReplicationSystemId != 0U; }
 
+	/** Static handles have ODD Id's */
 	bool IsStatic() const { return Id & StaticIdMask; }
+
+	/** Dynamic handles have EVEN Id's */
 	bool IsDynamic() const { return IsValid() && !IsStatic(); }
 
 	bool operator==(const FNetRefHandle& Other)const { return Id == Other.Id; }
