Geometry Collection: Avoid crash on dangling callbacks in post physics sync

#rnx

diff --git a/Engine/Source/Runtime/Experimental/GeometryCollectionEngine/Private/GeometryCollection/GeometryCollectionComponent.cpp b/Engine/Source/Runtime/Experimental/GeometryCollectionEngine/Private/GeometryCollection/GeometryCollectionComponent.cpp
--- a/Engine/Source/Runtime/Experimental/GeometryCollectionEngine/Private/GeometryCollection/GeometryCollectionComponent.cpp
+++ b/Engine/Source/Runtime/Experimental/GeometryCollectionEngine/Private/GeometryCollection/GeometryCollectionComponent.cpp
@@ -716,6 +716,21 @@ UGeometryCollectionComponent::UGeometryCollectionComponent(const FObjectInitiali
 #endif
 }
 
+PRAGMA_DISABLE_DEPRECATION_WARNINGS // Disable warnings for implicit destruction of deprecated members.
+UGeometryCollectionComponent::~UGeometryCollectionComponent()
+{
+	if (!ensureMsgf(
+			!PhysicsProxy,
+			TEXT("Geometry collection %s destroyed while still attached to physics proxy. "
+				"BeginDestroy() might be overridden without a base class call."),
+			*GetName()))
+	{
+		// Clear callback to avoid use-after-destruction via dangling pointer.
+		PhysicsProxy->SetPostPhysicsSyncCallback(nullptr);
+	}
+}
+PRAGMA_ENABLE_DEPRECATION_WARNINGS
+
 Chaos::FPhysicsSolver* UGeometryCollectionComponent::GetSolver(const UGeometryCollectionComponent& GeometryCollectionComponent)
 {
 	if(GeometryCollectionComponent.ChaosSolverActor)
@@ -4807,7 +4822,7 @@ void UGeometryCollectionComponent::OnDestroyPhysicsState()
 	if(PhysicsProxy)
 	{
 		// Clear physics sync callback
-		PhysicsProxy->SetPostPhysicsSyncCallback([]() {}); 
+		PhysicsProxy->SetPostPhysicsSyncCallback(nullptr); 
 
 		if (FPhysScene* PhysScene = GetInnerChaosScene())
 		{
diff --git a/Engine/Source/Runtime/Experimental/GeometryCollectionEngine/Public/GeometryCollection/GeometryCollectionComponent.h b/Engine/Source/Runtime/Experimental/GeometryCollectionEngine/Public/GeometryCollection/GeometryCollectionComponent.h
--- a/Engine/Source/Runtime/Experimental/GeometryCollectionEngine/Public/GeometryCollection/GeometryCollectionComponent.h
+++ b/Engine/Source/Runtime/Experimental/GeometryCollectionEngine/Public/GeometryCollection/GeometryCollectionComponent.h
@@ -586,6 +586,8 @@ class UGeometryCollectionComponent : public UMeshComponent, public IChaosNotifyH
 
 public:
 
+	GEOMETRYCOLLECTIONENGINE_API virtual ~UGeometryCollectionComponent();
+
 	// Collision profile name that indicates we should use the geometry collection's default collision profile.
 	GEOMETRYCOLLECTIONENGINE_API static FName DefaultCollisionProfileName;
 
@@ -1734,7 +1736,7 @@ private:
 
 	TArray<bool> EventsPlayed;
 
-	FGeometryCollectionPhysicsProxy* PhysicsProxy;
+	FGeometryCollectionPhysicsProxy* PhysicsProxy = nullptr;
 	TUniquePtr<FGeometryDynamicCollection> DynamicCollection;
 
 	// Temporary storage for body setup in order to initialise a dummy body instance
