Don't cancel on-demand install requests when request handle goes out of scope.

#rnx

diff --git a/Engine/Source/Runtime/Experimental/IoStore/OnDemand/Private/OnDemandIoStore.cpp b/Engine/Source/Runtime/Experimental/IoStore/OnDemand/Private/OnDemandIoStore.cpp
--- a/Engine/Source/Runtime/Experimental/IoStore/OnDemand/Private/OnDemandIoStore.cpp
+++ b/Engine/Source/Runtime/Experimental/IoStore/OnDemand/Private/OnDemandIoStore.cpp
@@ -2112,7 +2112,6 @@ FOnDemandInstallRequest::FOnDemandInstallRequest(FWeakOnDemandIoStore InIoStore,
 
 FOnDemandInstallRequest::~FOnDemandInstallRequest()
 {
-	Cancel();
 }
 
 FOnDemandInstallRequest::FOnDemandInstallRequest(FOnDemandInstallRequest&& Other)
