IAS timeout was running 1000x faster than expected causing false request failures and idle detections on a small subset of platforms.

#rnx

diff --git a/Engine/Source/Runtime/Experimental/IoStore/HttpClient/Private/Api.inl b/Engine/Source/Runtime/Experimental/IoStore/HttpClient/Private/Api.inl
--- a/Engine/Source/Runtime/Experimental/IoStore/HttpClient/Private/Api.inl
+++ b/Engine/Source/Runtime/Experimental/IoStore/HttpClient/Private/Api.inl
@@ -272,7 +272,7 @@ FTicketPerf::FSample FTicketPerf::GetSample() const
 	static uint64 Freq;
 	if (Freq == 0)
 	{
-		Freq = uint64(1.0 / FPlatformTime::GetSecondsPerCycle());
+		Freq = uint64(1.0 / FPlatformTime::GetSecondsPerCycle64());
 	}
 
 	auto ToMs = [] (uint64 Value) { return uint32((Value * 1000ull) / Freq); };
diff --git a/Engine/Source/Runtime/Experimental/IoStore/HttpClient/Private/Loop.inl b/Engine/Source/Runtime/Experimental/IoStore/HttpClient/Private/Loop.inl
--- a/Engine/Source/Runtime/Experimental/IoStore/HttpClient/Private/Loop.inl
+++ b/Engine/Source/Runtime/Experimental/IoStore/HttpClient/Private/Loop.inl
@@ -1408,7 +1408,7 @@ uint32 FEventLoop::FImpl::Tick(int32 PollTimeoutMs)
 		static uint64 Base = 0;
 		if (Freq == 0)
 		{
-			Freq = uint64(1.0 / FPlatformTime::GetSecondsPerCycle());
+			Freq = uint64(1.0 / FPlatformTime::GetSecondsPerCycle64());
 			Base = FPlatformTime::Cycles64();
 		}
 		uint64 NowBig = ((FPlatformTime::Cycles64() - Base) * 1000) / Freq;
diff --git a/Engine/Source/Runtime/Experimental/IoStore/HttpClient/Private/Misc.inl b/Engine/Source/Runtime/Experimental/IoStore/HttpClient/Private/Misc.inl
--- a/Engine/Source/Runtime/Experimental/IoStore/HttpClient/Private/Misc.inl
+++ b/Engine/Source/Runtime/Experimental/IoStore/HttpClient/Private/Misc.inl
@@ -724,7 +724,7 @@ private:
 ////////////////////////////////////////////////////////////////////////////////
 FThrottler::FThrottler()
 {
-	CycleFreq = uint64(1.0 / FPlatformTime::GetSecondsPerCycle());
+	CycleFreq = uint64(1.0 / FPlatformTime::GetSecondsPerCycle64());
 	CycleLast = FPlatformTime::Cycles64() - CycleFreq;
 	CycleIdle = CycleFreq * 8;
 }
diff --git a/Engine/Source/Runtime/Experimental/IoStore/HttpClient/Private/Socket.inl b/Engine/Source/Runtime/Experimental/IoStore/HttpClient/Private/Socket.inl
--- a/Engine/Source/Runtime/Experimental/IoStore/HttpClient/Private/Socket.inl
+++ b/Engine/Source/Runtime/Experimental/IoStore/HttpClient/Private/Socket.inl
@@ -195,7 +195,8 @@ int32 FWaiter::Wait(TArrayView<FWaiter> Waiters, int32 TimeoutMs)
 	timeval* TimeValPtr = (TimeoutMs >= 0 ) ? &TimeVal : nullptr;
 	if (TimeoutMs > 0)
 	{
-		TimeVal = { TimeoutMs >> 10, TimeoutMs & ((1 << 10) - 1) };
+		int32 TimeoutUs = TimeoutMs << 10;
+		TimeVal = { TimeoutMs >> 10, TimeoutUs & ((1 << 20) - 1) };
 	}
 
 	fd_set FdSetRead;	FD_ZERO(&FdSetRead);
diff --git a/Engine/Source/Runtime/Experimental/IoStore/OnDemand/Private/JournaledCache.cpp b/Engine/Source/Runtime/Experimental/IoStore/OnDemand/Private/JournaledCache.cpp
--- a/Engine/Source/Runtime/Experimental/IoStore/OnDemand/Private/JournaledCache.cpp
+++ b/Engine/Source/Runtime/Experimental/IoStore/OnDemand/Private/JournaledCache.cpp
@@ -1515,7 +1515,7 @@ FGovernor::FGovernor()
 ////////////////////////////////////////////////////////////////////////////////
 void FGovernor::Set(uint32 Allowance, uint32 Ops, uint32 Seconds)
 {
-	int64 CycleFreq	= int64(1.0 / FPlatformTime::GetSecondsPerCycle());
+	int64 CycleFreq	= int64(1.0 / FPlatformTime::GetSecondsPerCycle64());
 	return Set(Allowance, Ops, Seconds, CycleFreq);
 }
 
@@ -1868,7 +1868,7 @@ uint32 FServiceThread::Run()
 {
 	LLM_SCOPE_BYTAG(Ias);
 
-	int64 CycleFreq	= int64(1.0 / FPlatformTime::GetSecondsPerCycle());
+	int64 CycleFreq	= int64(1.0 / FPlatformTime::GetSecondsPerCycle64());
 
 	while (RunCount.load(std::memory_order_relaxed))
 	{
diff --git a/Engine/Source/Runtime/Experimental/IoStore/OnDemand/Private/LatencyTesting.cpp b/Engine/Source/Runtime/Experimental/IoStore/OnDemand/Private/LatencyTesting.cpp
--- a/Engine/Source/Runtime/Experimental/IoStore/OnDemand/Private/LatencyTesting.cpp
+++ b/Engine/Source/Runtime/Experimental/IoStore/OnDemand/Private/LatencyTesting.cpp
@@ -52,7 +52,7 @@ void LatencyTest(FStringView InUrl, FStringView InPath, uint32 InTimeOutMs, TArr
 		Result = Ok ? int32(Cycles) : -1;
 	}
 
-	int64 Freq = int64(1.0 / FPlatformTime::GetSecondsPerCycle());
+	int64 Freq = int64(1.0 / FPlatformTime::GetSecondsPerCycle64());
 	for (int32& Result : OutResults)
 	{
 		if (Result == -1)
