Iris
* Stop initializing the world location of an actor when it starts replicating.
* This exposed a bug where Iris doesn't tell the NetObjectGridFilter to update objects if only their net cull distance override changes.

diff --git a/Engine/Source/Runtime/Experimental/Iris/Core/Private/Iris/ReplicationSystem/ObjectReplicationBridge.cpp b/Engine/Source/Runtime/Experimental/Iris/Core/Private/Iris/ReplicationSystem/ObjectReplicationBridge.cpp
--- a/Engine/Source/Runtime/Experimental/Iris/Core/Private/Iris/ReplicationSystem/ObjectReplicationBridge.cpp
+++ b/Engine/Source/Runtime/Experimental/Iris/Core/Private/Iris/ReplicationSystem/ObjectReplicationBridge.cpp
@@ -457,8 +457,6 @@ UE::Net::FNetRefHandle UObjectReplicationBridge::StartReplicatingRootObject(UObj
 	{
 		FWorldLocations& WorldLocations = ReplicationSystem->GetReplicationSystemInternal()->GetWorldLocations();
 		WorldLocations.InitObjectInfoCache(ObjectIndex);
-		ForceUpdateWorldLocation(RefHandle, ObjectIndex);
-
 	}
 
 	// Set prioritizer
