Dataflow : Add subgraph shell class

diff --git a/Engine/Source/Runtime/Experimental/Dataflow/Engine/Private/Dataflow/DataflowSubGraph.cpp b/Engine/Source/Runtime/Experimental/Dataflow/Engine/Private/Dataflow/DataflowSubGraph.cpp
new file mode 100644
--- /dev/null
+++ b/Engine/Source/Runtime/Experimental/Dataflow/Engine/Private/Dataflow/DataflowSubGraph.cpp
@@ -0,0 +1,7 @@
+// Copyright Epic Games, Inc. All Rights Reserved.
+
+#include "Dataflow/DataflowSubGraph.h"
+
+UDataflowSubgraph::UDataflowSubgraph(class FObjectInitializer const& ObjectInitializer)
+	: UEdGraph(ObjectInitializer)
+{}
diff --git a/Engine/Source/Runtime/Experimental/Dataflow/Engine/Public/Dataflow/DataflowSubGraph.h b/Engine/Source/Runtime/Experimental/Dataflow/Engine/Public/Dataflow/DataflowSubGraph.h
new file mode 100644
--- /dev/null
+++ b/Engine/Source/Runtime/Experimental/Dataflow/Engine/Public/Dataflow/DataflowSubGraph.h
@@ -0,0 +1,22 @@
+// Copyright Epic Games, Inc. All Rights Reserved.
+
+#pragma once
+
+#include "EdGraph/EdGraph.h"
+#include "StructUtils/PropertyBag.h"
+
+#include "DataflowSubGraph.generated.h"
+
+UCLASS(MinimalAPI)
+class UDataflowSubgraph : public UEdGraph
+{
+	GENERATED_UCLASS_BODY()
+
+public:
+	const FInstancedPropertyBag& GetInputs() const { return Inputs; }
+	const FInstancedPropertyBag& GetOutputs() const { return Outputs; }
+
+private: 
+	FInstancedPropertyBag Inputs;
+	FInstancedPropertyBag Outputs;
+};
\ No newline at end of file
