Dataflow : add method to find a dataflow asset from a main or sub graph

diff --git a/Engine/Source/Runtime/Experimental/Dataflow/Engine/Private/Dataflow/DataflowObject.cpp b/Engine/Source/Runtime/Experimental/Dataflow/Engine/Private/Dataflow/DataflowObject.cpp
--- a/Engine/Source/Runtime/Experimental/Dataflow/Engine/Private/Dataflow/DataflowObject.cpp
+++ b/Engine/Source/Runtime/Experimental/Dataflow/Engine/Private/Dataflow/DataflowObject.cpp
@@ -221,6 +221,30 @@ TObjectPtr<UDataflowEdNode> UDataflow::FindEdNodeByDataflowNodeGuid(const FGuid&
 	return TObjectPtr<UDataflowEdNode>(nullptr);
 }
 
+UDataflow* UDataflow::GetDataflowAssetFromEdGraph(UEdGraph* EdGraph)
+{
+	// main graph is the dataflow asset
+	UDataflow* DataflowAsset = Cast<UDataflow>(EdGraph);
+	if (!DataflowAsset && EdGraph)
+	{
+		UEdGraph* EdParentGraph = Cast<UEdGraph>(EdGraph->GetOuter());
+		DataflowAsset = UDataflow::GetDataflowAssetFromEdGraph(EdParentGraph);
+	}
+	return DataflowAsset;
+}
+
+const UDataflow* UDataflow::GetDataflowAssetFromEdGraph(const UEdGraph* EdGraph)
+{
+	// main graph is the dataflow asset
+	const UDataflow* DataflowAsset = Cast<const UDataflow>(EdGraph);
+	if (!DataflowAsset && EdGraph)
+	{
+		const UEdGraph* EdParentGraph = Cast<const UEdGraph>(EdGraph->GetOuter());
+		DataflowAsset = UDataflow::GetDataflowAssetFromEdGraph(EdParentGraph);
+	}
+	return DataflowAsset;
+}
+
 #if WITH_EDITOR
 bool UDataflow::CanEditChange(const FProperty* InProperty) const
 {
diff --git a/Engine/Source/Runtime/Experimental/Dataflow/Engine/Public/Dataflow/DataflowObject.h b/Engine/Source/Runtime/Experimental/Dataflow/Engine/Public/Dataflow/DataflowObject.h
--- a/Engine/Source/Runtime/Experimental/Dataflow/Engine/Public/Dataflow/DataflowObject.h
+++ b/Engine/Source/Runtime/Experimental/Dataflow/Engine/Public/Dataflow/DataflowObject.h
@@ -166,6 +166,13 @@ public:
 
 	const UE::Dataflow::FTimestamp& GetRenderingTimestamp() const { return LastModifiedRenderTarget; }
 
+	/**
+	* Find the Dataflow asset from a sepcific graph/subgraph
+	*/
+	DATAFLOWENGINE_API static UDataflow* GetDataflowAssetFromEdGraph(UEdGraph* EdGraph);
+	DATAFLOWENGINE_API static const UDataflow* GetDataflowAssetFromEdGraph(const UEdGraph* EdGraph);
+
+
 #if WITH_EDITORONLY_DATA
 
 	/*
