Add an overload of FMultiEndpointHttpClient::Get that accepts a FIoOffsetAndLength, allowing the user to only request a partial read.

#rnx

diff --git a/Engine/Source/Runtime/Experimental/IoStore/OnDemand/Private/OnDemandHttpClient.cpp b/Engine/Source/Runtime/Experimental/IoStore/OnDemand/Private/OnDemandHttpClient.cpp
--- a/Engine/Source/Runtime/Experimental/IoStore/OnDemand/Private/OnDemandHttpClient.cpp
+++ b/Engine/Source/Runtime/Experimental/IoStore/OnDemand/Private/OnDemandHttpClient.cpp
@@ -398,6 +398,12 @@ TUniquePtr<FMultiEndpointHttpClient> FMultiEndpointHttpClient::Create(const FMul
 }
 
 void FMultiEndpointHttpClient::Get(const FOnDemandHostGroup& Endpoint, FAnsiStringView RelativeUrl, FOnHttpResponse&& OnResponse)
+{
+	FIoOffsetAndLength DefaultChunkRange(0,0);
+	Get(Endpoint, RelativeUrl, DefaultChunkRange, MoveTemp(OnResponse));
+}
+
+void FMultiEndpointHttpClient::Get(const FOnDemandHostGroup& Endpoint, FAnsiStringView RelativeUrl, const FIoOffsetAndLength& ChunkRange, FOnHttpResponse&& OnResponse)
 {
 	FConnection& Connection = GetConnection(Endpoint);
 
@@ -405,6 +411,7 @@ void FMultiEndpointHttpClient::Get(const FOnDemandHostGroup& Endpoint, FAnsiStri
 	{
 		.OnResponse		= MoveTemp(OnResponse),
 		.RelativeUrl	= FAnsiString(RelativeUrl),
+		.Range			= ChunkRange,
 		.Connection		= Connection,
 		.StartTime		= FPlatformTime::Cycles64(),
 		.Host			= Connection.CurrentHost
@@ -509,9 +516,15 @@ void FMultiEndpointHttpClient::IssueRequest(FRequest&& Request)
 		}
 	};
 
-	EventLoop.Send(
-		EventLoop.Get(Url, ConnectionPool),
-		MoveTemp(HttpSink));
+	UE::IoStore::HTTP::FRequest HttpRequest = EventLoop.Get(Url, ConnectionPool);
+
+	if (Request.Range.GetOffset() > 0 || Request.Range.GetLength() > 0)
+	{
+		HttpRequest.Header(ANSITEXTVIEW("Range"),
+			WriteToAnsiString<64>(ANSITEXTVIEW("bytes="), Request.Range.GetOffset(), ANSITEXTVIEW("-"), Request.Range.GetOffset() + Request.Range.GetLength() - 1));
+	}
+
+	EventLoop.Send(MoveTemp(HttpRequest), MoveTemp(HttpSink));
 }
 
 void FMultiEndpointHttpClient::RetryRequest(FRequest&& Request)
diff --git a/Engine/Source/Runtime/Experimental/IoStore/OnDemand/Private/OnDemandHttpClient.h b/Engine/Source/Runtime/Experimental/IoStore/OnDemand/Private/OnDemandHttpClient.h
--- a/Engine/Source/Runtime/Experimental/IoStore/OnDemand/Private/OnDemandHttpClient.h
+++ b/Engine/Source/Runtime/Experimental/IoStore/OnDemand/Private/OnDemandHttpClient.h
@@ -112,6 +112,8 @@ public:
 	static TUniquePtr<FMultiEndpointHttpClient> Create(const FMultiEndpointHttpClientConfig& Config);
 
 	void Get(const FOnDemandHostGroup& Endpoint, FAnsiStringView RelativeUrl, FOnHttpResponse&& OnResponse);
+	void Get(const FOnDemandHostGroup& Endpoint, FAnsiStringView RelativeUrl, const FIoOffsetAndLength& ChunkRange, FOnHttpResponse&& OnResponse);
+
 	bool Tick(int32 WaitTimeMs, uint32 MaxKiBPerSecond);
 	bool Tick() { return Tick(-1, 0); }
 
