Dataflow : Fix spelling on SubGraph class ( not yet used anywhere )

diff --git a/Engine/Source/Runtime/Experimental/Dataflow/Engine/Private/Dataflow/DataflowSubGraph.cpp b/Engine/Source/Runtime/Experimental/Dataflow/Engine/Private/Dataflow/DataflowSubGraph.cpp
--- a/Engine/Source/Runtime/Experimental/Dataflow/Engine/Private/Dataflow/DataflowSubGraph.cpp
+++ b/Engine/Source/Runtime/Experimental/Dataflow/Engine/Private/Dataflow/DataflowSubGraph.cpp
@@ -2,6 +2,6 @@
 
 #include "Dataflow/DataflowSubGraph.h"
 
-UDataflowSubgraph::UDataflowSubgraph(class FObjectInitializer const& ObjectInitializer)
+UDataflowSubGraph::UDataflowSubGraph(class FObjectInitializer const& ObjectInitializer)
 	: UEdGraph(ObjectInitializer)
 {}
diff --git a/Engine/Source/Runtime/Experimental/Dataflow/Engine/Public/Dataflow/DataflowSubGraph.h b/Engine/Source/Runtime/Experimental/Dataflow/Engine/Public/Dataflow/DataflowSubGraph.h
--- a/Engine/Source/Runtime/Experimental/Dataflow/Engine/Public/Dataflow/DataflowSubGraph.h
+++ b/Engine/Source/Runtime/Experimental/Dataflow/Engine/Public/Dataflow/DataflowSubGraph.h
@@ -8,7 +8,7 @@
 #include "DataflowSubGraph.generated.h"
 
 UCLASS(MinimalAPI)
-class UDataflowSubgraph : public UEdGraph
+class UDataflowSubGraph : public UEdGraph
 {
 	GENERATED_UCLASS_BODY()
 
