Dataflow : add methods to find, add and remove a subgraph

diff --git a/Engine/Source/Runtime/Experimental/Dataflow/Engine/Private/Dataflow/DataflowObject.cpp b/Engine/Source/Runtime/Experimental/Dataflow/Engine/Private/Dataflow/DataflowObject.cpp
--- a/Engine/Source/Runtime/Experimental/Dataflow/Engine/Private/Dataflow/DataflowObject.cpp
+++ b/Engine/Source/Runtime/Experimental/Dataflow/Engine/Private/Dataflow/DataflowObject.cpp
@@ -5,6 +5,8 @@
 #include "Dataflow/DataflowEdNode.h"
 #include "Dataflow/DataflowNodeParameters.h"
 #include "Dataflow/DataflowObjectInterface.h"
+#include "Dataflow/DataflowSubGraph.h"
+
 #if WITH_EDITOR
 #include "EdGraph/EdGraphPin.h"
 #endif
@@ -245,6 +247,57 @@ const UDataflow* UDataflow::GetDataflowAssetFromEdGraph(const UEdGraph* EdGraph)
 	return DataflowAsset;
 }
 
+const UDataflowSubGraph* UDataflow::FindSubGraphByName(FName Name) const
+{
+	for (const UDataflowSubGraph* SubGraph : DataflowSubGraphs)
+	{
+		if (SubGraph->GetFName() == Name)
+		{
+			return SubGraph;
+		}
+	}
+	return nullptr;
+}
+
+UDataflowSubGraph* UDataflow::FindSubGraphByName(FName Name)
+{
+	for (UDataflowSubGraph* SubGraph : DataflowSubGraphs)
+	{
+		if (SubGraph->GetFName() == Name)
+		{
+			return SubGraph;
+		}
+	}
+	return nullptr;
+}
+
+void UDataflow::AddSubGraph(UDataflowSubGraph* SubGraph)
+{
+	if (ensure(SubGraph))
+	{
+		if (ensure(SubGraph->IsInOuter(this)))
+		{
+#if WITH_EDITORONLY_DATA
+			SubGraphs.AddUnique(SubGraph);
+#endif
+			DataflowSubGraphs.AddUnique(SubGraph);
+			Modify();
+		}
+	}	
+}
+
+void UDataflow::RemoveSubGraph(UDataflowSubGraph* SubGraph)
+{
+	if (ensure(SubGraph))
+	{
+#if WITH_EDITORONLY_DATA
+			SubGraphs.Remove(SubGraph);
+#endif
+			DataflowSubGraphs.Remove(SubGraph);
+			Modify();
+	}
+}
+
 #if WITH_EDITOR
 bool UDataflow::CanEditChange(const FProperty* InProperty) const
 {
diff --git a/Engine/Source/Runtime/Experimental/Dataflow/Engine/Public/Dataflow/DataflowObject.h b/Engine/Source/Runtime/Experimental/Dataflow/Engine/Public/Dataflow/DataflowObject.h
--- a/Engine/Source/Runtime/Experimental/Dataflow/Engine/Public/Dataflow/DataflowObject.h
+++ b/Engine/Source/Runtime/Experimental/Dataflow/Engine/Public/Dataflow/DataflowObject.h
@@ -24,6 +24,7 @@ class FArchive;
 class UDataflow;
 class UObject;
 class UDataflowEdNode;
+class UDataflowSubGraph;
 class UMaterial;
 namespace UE::Dataflow { class FGraph; }
 
@@ -172,6 +173,13 @@ public:
 	DATAFLOWENGINE_API static UDataflow* GetDataflowAssetFromEdGraph(UEdGraph* EdGraph);
 	DATAFLOWENGINE_API static const UDataflow* GetDataflowAssetFromEdGraph(const UEdGraph* EdGraph);
 
+	/** Find a SubGraph by name */
+	DATAFLOWENGINE_API const UDataflowSubGraph* FindSubGraphByName(FName Name) const;
+	DATAFLOWENGINE_API UDataflowSubGraph* FindSubGraphByName(FName Name);
+
+	/** Add a SubGraph to the asset */
+	DATAFLOWENGINE_API void AddSubGraph(UDataflowSubGraph* SubGraph);
+	DATAFLOWENGINE_API void RemoveSubGraph(UDataflowSubGraph* SubGraph);
 
 #if WITH_EDITORONLY_DATA
 
@@ -210,6 +218,13 @@ private:
 	bool bEnablePerNodeTransactionSerialization = true;
 #endif
 
+private:
+	/** 
+	* List of Dataflow SubGraphs 
+	* in editor they also exists in the SubGraphs parent class property
+	*/
+	TArray<TObjectPtr<UDataflowSubGraph>> DataflowSubGraphs; 
+
 private:
 	/** Broadcasts a notification just before the dataflow is destroyed. */
 	FMeshDeformerBeginDestroyEvent BeginDestroyEvent;
