[NetPhysics] RewindData caching optimization, only cache data in one "phase" per physics frame; the phase we rewind to and compare server states with.
- Reduced rewind data memory allocation to 1/3
- Fixes issue where rewind could be done to the wrong phase if particles had not cached data in the phase we rewind to, increasing non-determinism during resimulation.

diff --git a/Engine/Source/Programs/HeadlessChaos/Private/HeadlessChaosTestRewind.cpp b/Engine/Source/Programs/HeadlessChaos/Private/HeadlessChaosTestRewind.cpp
--- a/Engine/Source/Programs/HeadlessChaos/Private/HeadlessChaosTestRewind.cpp
+++ b/Engine/Source/Programs/HeadlessChaos/Private/HeadlessChaosTestRewind.cpp
@@ -39,9 +39,18 @@ namespace ChaosTest {
 		// Make a solver
 		FChaosSolversModule* Module = FChaosSolversModule::GetModule();
 		auto* Solver = Module->CreateSolver(nullptr,/*AsyncDt=*/-1);
-                InitSolverSettings(Solver);
+		InitSolverSettings(Solver);
+
+		/** int32 Optimization starts at 0 and iterate up to 2
+		* 0 = No optimizations
+		* 1 = CollisionResimCache
+		* 2 = Optimized RewindData */
 
-		Solver->EnableRewindCapture(RewindHistorySize, !!Optimization);
+		Solver->EnableRewindCapture(RewindHistorySize, Optimization == 1); // Use CollisionResimCache optimization as optimization 1
+		if (Chaos::FRewindData* RewindData = Solver->GetRewindData())
+		{
+			RewindData->SetRewindDataOptimization(Optimization == 2); // Use optimized RewindData in as optimization 2
+		}
 		Solver->SetThreadingMode_External(EThreadingModeTemp::SingleThread);
 		OutSimDt = DtSizes[StepMode];
 		if (StepMode > 0)
@@ -57,7 +66,7 @@ namespace ChaosTest {
 		template <typename TLambda>
 		static void TestEmpty(const TLambda& Lambda, int32 RewindHistorySize = 200)
 		{
-			for (int Optimization = 0; Optimization < 2; ++Optimization)
+			for (int Optimization = 0; Optimization < 3; ++Optimization)
 			{
 				for (int DtMode = 0; DtMode < 4; ++DtMode)
 				{
@@ -1454,7 +1463,7 @@ namespace ChaosTest {
 		TRewindHelper::TestDynamicSphere([](auto* Solver, FReal SimDt, int32 Optimization, auto Proxy, auto Sphere)
 			{
 				FRewindCallbackTestHelper* Helper = RegisterCallbackHelper(Solver);
-				Helper->ProcessInputsFunc = [Proxy, RewindData = Solver->GetRewindData()](const int32 PhysicsStep, bool bIsResimming)
+				Helper->ProcessInputsFunc = [Proxy, RewindData = Solver->GetRewindData(), Optimization](const int32 PhysicsStep, bool bIsResimming)
 					{
 						if (PhysicsStep == 3)
 						{
@@ -1462,8 +1471,24 @@ namespace ChaosTest {
 							Proxy->GetPhysicsThreadAPI()->SetAngularAcceleration(FVec3(0, 0, 10) + Proxy->GetPhysicsThreadAPI()->AngularAcceleration());
 						}
 
-					for (int32 Step = 0; Step < PhysicsStep; ++Step)
+						for (int32 Step = 0; Step < PhysicsStep; ++Step)
 						{
+							if (Optimization != 2) // Optimization 2 does not cache data in PrePushData so no need to check that phase
+							{
+								//Before push no force or torque at all
+								FGeometryParticleState PrePushState = RewindData->GetPastStateAtFrame(*Proxy->GetHandle_LowLevel(), Step, FFrameAndPhase::EParticleHistoryPhase::PrePushData);
+								EXPECT_EQ(PrePushState.Acceleration()[2], 0);
+								EXPECT_EQ(PrePushState.AngularAcceleration()[2], 0);
+							}
+
+							{
+								//After push (but before callback) only see GT force, and no torque
+								FGeometryParticleState PostPushState = RewindData->GetPastStateAtFrame(*Proxy->GetHandle_LowLevel(), Step, FFrameAndPhase::EParticleHistoryPhase::PostPushData);
+								EXPECT_EQ(PostPushState.Acceleration()[2], 1);	//GT always sets force of 1
+								EXPECT_EQ(PostPushState.AngularAcceleration()[2], 0);	//GT never sets torque
+							}
+
+							if (Optimization != 2) // Optimization 2 does not cache data in PostCallbacks so no need to check that phase
 							{
 							FGeometryParticleState State = RewindData->GetPastStateAtFrame(*Proxy->GetHandle_LowLevel(), Step, FFrameAndPhase::EParticleHistoryPhase::PostCallbacks);
 								if (Step == 3)
@@ -1477,24 +1502,9 @@ namespace ChaosTest {
 									EXPECT_EQ(State.AngularAcceleration()[2], 0);	//GT never sets torque
 								}
 							}
+						}
+					};
 
-							{
-								//Before push no force or torque at all
-							FGeometryParticleState PrePushState = RewindData->GetPastStateAtFrame(*Proxy->GetHandle_LowLevel(), Step, FFrameAndPhase::EParticleHistoryPhase::PrePushData);
-								EXPECT_EQ(PrePushState.Acceleration()[2], 0);
-								EXPECT_EQ(PrePushState.AngularAcceleration()[2], 0);
-							}
-
-							{
-								//After push (but before callback) only see GT force, and no torque
-							FGeometryParticleState PostPushState = RewindData->GetPastStateAtFrame(*Proxy->GetHandle_LowLevel(), Step, FFrameAndPhase::EParticleHistoryPhase::PostPushData);
-								EXPECT_EQ(PostPushState.Acceleration()[2], 1);	//GT always sets force of 1
-								EXPECT_EQ(PostPushState.AngularAcceleration()[2], 0);	//GT never sets torque
-							}
-
-					}
-				};
-				
 				auto& Particle = Proxy->GetGameThreadAPI();
 				Particle.SetGravityEnabled(false);
 
diff --git a/Engine/Source/Runtime/Engine/Private/PhysicsEngine/PhysicsReplicationCache.cpp b/Engine/Source/Runtime/Engine/Private/PhysicsEngine/PhysicsReplicationCache.cpp
--- a/Engine/Source/Runtime/Engine/Private/PhysicsEngine/PhysicsReplicationCache.cpp
+++ b/Engine/Source/Runtime/Engine/Private/PhysicsEngine/PhysicsReplicationCache.cpp
@@ -246,8 +246,8 @@ void FPhysicsReplicationCacheAsync::PopulateReplicationCache_Internal()
 			if (Chaos::FPBDRigidParticleHandle* Handle = Interface.GetRigidParticle(PhysicsObject))
 			{
 				FRigidBodyState& ReplicationState = ReplicationCacheData.GetState();
-				ReplicationState.Position = Handle->GetP(); // Note, X is not yet updated in PostSolve_Internal so using P which is up to date
-				ReplicationState.Quaternion = Handle->GetQ(); // Note, R is not yet updated in PostSolve_Internal so using Q which is up to date
+				ReplicationState.Position = Handle->GetX();
+				ReplicationState.Quaternion = Handle->GetR();
 				ReplicationState.LinVel = Handle->GetV();
 				ReplicationState.AngVel = FMath::RadiansToDegrees(Handle->GetW());
 				ReplicationState.Flags = Handle->ObjectState() == Chaos::EObjectStateType::Sleeping ? ERigidBodyFlags::Sleeping : 0;
diff --git a/Engine/Source/Runtime/Experimental/Chaos/Private/PBDRigidsSolver.cpp b/Engine/Source/Runtime/Experimental/Chaos/Private/PBDRigidsSolver.cpp
--- a/Engine/Source/Runtime/Experimental/Chaos/Private/PBDRigidsSolver.cpp
+++ b/Engine/Source/Runtime/Experimental/Chaos/Private/PBDRigidsSolver.cpp
@@ -1884,6 +1884,11 @@ namespace Chaos
 
 		ProcessSinglePushedData_Internal(PushData);
 
+		if (FRewindData* RewindData = GetRewindData())
+		{
+			RewindData->CacheCurrentDirtyData(FFrameAndPhase::PostPushData);
+		}
+
 		//run any commands passed in. These don't generate outputs and are a one off so just do them here
 		//note: commands run before sim callbacks. This is important for sub-stepping since we want each sub-step to have a consistent view
 		//so for example if the user deletes a floor surface, we want all sub-steps to see that in the same way
diff --git a/Engine/Source/Runtime/Experimental/Chaos/Private/RewindData.cpp b/Engine/Source/Runtime/Experimental/Chaos/Private/RewindData.cpp
--- a/Engine/Source/Runtime/Experimental/Chaos/Private/RewindData.cpp
+++ b/Engine/Source/Runtime/Experimental/Chaos/Private/RewindData.cpp
@@ -157,7 +157,7 @@ bool FJointStateBase::IsInSync(const FPBDJointConstraintHandle& Handle, const FF
 
 
 bool bCVarRewindDataOptimization = true;
-FAutoConsoleVariableRef CVarRewindDataOptimization(TEXT("p.Resim.RewindDataOptimization"), bCVarRewindDataOptimization, TEXT("Default value for RewinData optimization, note that this can be overridden at runtime by API calls. Effect: Only alter the minimum required properties during a resim for particles not marked for FullResim"));
+FAutoConsoleVariableRef CVarRewindDataOptimization(TEXT("p.Resim.RewindDataOptimization"), bCVarRewindDataOptimization, TEXT("Default value for RewinData optimization, note that this can be overridden at runtime by API calls. Effect: Only alter the minimum required properties during a resim for particles not marked for FullResim and only cache data during the PostPushData phase and lower memory allocation for the history cache to 1/3 of non-optimized flow."));
 
 FRewindData::FRewindData(FPBDRigidsSolver* InSolver, int32 NumFrames, bool InRewindDataOptimization, int32 InCurrentFrame)
 	: Managers(NumFrames + 1)	//give 1 extra for saving at head
@@ -424,6 +424,10 @@ bool FRewindData::RewindToFrame(int32 Frame)
 
 	BlockResimFrame = bResimAllowRewindToResimulatedFrames ? Frame : CurFrame;
 
+	ResimFrame = Frame;
+	CurFrame = Frame;
+	bNeedsSave = false;
+
 	auto RewindHelper = [RewindFrameAndPhase, CurFrameAndPhase, this](auto Obj, bool bResimAsFollower, auto& Property, const auto& RewindFunc) -> bool
 	{
 		if (bResimAsFollower)
@@ -439,7 +443,11 @@ bool FRewindData::RewindToFrame(int32 Frame)
 			//If we're rewinding an object that needs to save head (during resim when we get back to latest frame and phase we need to check for desync)
 			if (!Property.IsClean(RewindFrameAndPhase))
 			{
-				CopyDataFromObject(Property.WriteAccessMonotonic(CurFrameAndPhase, PropertiesPool), *Obj);
+				if (!bRewindDataOptimization)
+				{
+					// When not using optimized RewindData cache the current state in Phase::PrePushData on rewind.
+					CopyDataFromObject(Property.WriteAccessMonotonic(CurFrameAndPhase, PropertiesPool), *Obj);
+				}
 				RewindFunc(Obj, *Property.Read(RewindFrameAndPhase, PropertiesPool));
 
 				return true;
@@ -461,18 +469,7 @@ bool FRewindData::RewindToFrame(int32 Frame)
 		History.CachePreCorrectionState(*PTParticle);
 
 		// Todo: This should be set by bubble resimulation so that Dynamic and Sleeping particles outside of relevancy of resim doesn't actually resimulate as dynamic particles.
-		switch (PTParticle->ObjectState())
-		{
-		case EObjectStateType::Dynamic:
-		case EObjectStateType::Sleeping:
-			PTParticle->SetResimType(EResimType::FullResim);
-			break;
-		default:
-
-			PTParticle->SetResimType(EResimType::ResimAsFollower);
-			break;
-		}
-
+		DirtyParticleInfo.bResimAsFollower = PTParticle->ObjectState() != EObjectStateType::Dynamic && PTParticle->ObjectState() != EObjectStateType::Sleeping;
 		const bool bResimAsFollower = DirtyParticleInfo.bResimAsFollower;
 
 		bool bAnyChange = RewindHelper(PTParticle, bResimAsFollower, History.ParticlePositionRotation, [](auto Particle, const auto& Data) {Particle->SetXR(Data); });
@@ -487,6 +484,10 @@ bool FRewindData::RewindToFrame(int32 Frame)
 		bAnyChange |= RewindHelper(PTParticle->CastToRigidParticle(), bResimAsFollower, History.DynamicsMisc, [this](auto Particle, const auto& Data) {Solver->SetParticleDynamicMisc(Particle, Data); });
 		bAnyChange |= RewindHelper(PTParticle->CastToRigidParticle(), bResimAsFollower, History.MassProps, [](auto Particle, const auto& Data) {Particle->SetMassProps(Data); });
 
+		// Todo: This should be set by bubble resimulation so that Dynamic and Sleeping particles outside of relevancy of resim doesn't actually resimulate as dynamic particles.
+		// Set this after rewinding, since ResimType gets overwritten if NonFrequentData is cached
+		PTParticle->SetResimType(bResimAsFollower ? EResimType::ResimAsFollower : EResimType::FullResim);
+
 		if (!bResimAsFollower)
 		{
 			if (bAnyChange)
@@ -522,16 +523,12 @@ bool FRewindData::RewindToFrame(int32 Frame)
 	}
 #endif
 
-	ResimFrame = Frame;
-	CurFrame = Frame;
-	bNeedsSave = false;
-
 	return true;
 }
 
 void FRewindData::StepNonResimParticles(const int32 Frame)
 {
-	const FFrameAndPhase FrameAndPhase{ Frame, FFrameAndPhase::PrePushData };
+	const FFrameAndPhase FrameAndPhase{ Frame, FFrameAndPhase::PostPushData };
 	auto RewindHelper = [FrameAndPhase, this](auto Obj, auto& Property, const auto& RewindFunc) -> bool
 	{
 		if (auto Val = Property.Read(FrameAndPhase, PropertiesPool))
@@ -740,7 +737,7 @@ void FRewindData::FinishFrame()
 						//Anything saved on upcoming frame (was done during rewind) can be removed since we are now at head
 						Info.ClearPhaseAndFuture(FutureFrame);
 					}
-					else
+					else if (!bRewindDataOptimization)
 					{
 						//solver doesn't affect dynamics, so no reason to test if they desynced from original sim
 						//question: should we skip all other properties? dynamics is a commonly changed one but might be worth skipping everything solver skips
@@ -800,6 +797,7 @@ void FRewindData::AdvanceFrameImp(IResimCacheBase* ResimCache)
 {
 	FramesSaved = FMath::Min(FramesSaved + 1, static_cast<int32>(Managers.Capacity()));
 
+	const bool bHasResimCache = ResimCache != nullptr;
 	const int32 EarliestFrame = CurFrame - FramesSaved;
 	const FFrameAndPhase FrameAndPhase{ CurFrame, FFrameAndPhase::PostCallbacks };
 
@@ -822,7 +820,7 @@ void FRewindData::AdvanceFrameImp(IResimCacheBase* ResimCache)
 				auto Handle = Info.GetObjectPtr();
 				Info.bResimAsFollower = Handle->ResimType() == EResimType::ResimAsFollower;
 
-				if (IsResim() && !Info.bResimAsFollower)
+				if (!bRewindDataOptimization && IsResim() && !Info.bResimAsFollower)
 				{
 					DesyncIfNecessary</*bSkipDynamics=*/false>(Info, FrameAndPhase);
 				}
@@ -843,20 +841,23 @@ void FRewindData::AdvanceFrameImp(IResimCacheBase* ResimCache)
 	};
 
 	TArray<FGeometryParticleHandle*> DesyncedParticles;
-	if (IsResim())
+	if (IsResim() && bHasResimCache)
 	{
 		DesyncedParticles.Reserve(DirtyParticles.Num());
 	}
 
 	AdvanceHelper(DirtyParticles,
-		[&DesyncedParticles](FGeometryParticleHandle* DesyncedHandle)
+		[&DesyncedParticles, bHasResimCache](FGeometryParticleHandle* DesyncedHandle)
 		{
-			DesyncedParticles.Add(DesyncedHandle);
+			if (bHasResimCache)
+			{
+				DesyncedParticles.Add(DesyncedHandle);
+			}
 		},
 
 		[this, FrameAndPhase](FDirtyParticleInfo& Info, FGeometryParticleHandle* Handle)
 		{
-			if (Info.DirtyDynamics == CurFrame && !IsResimAndInSync(*Handle))
+			if (!bRewindDataOptimization && Info.DirtyDynamics == CurFrame && !IsResimAndInSync(*Handle))
 			{
 				//we only need to check the cast because right now there's no property system on PT, so any time a sim callback touches a particle we just mark it as dirty dynamics
 				if (auto Rigid = Handle->CastToRigidParticle())
@@ -873,7 +874,7 @@ void FRewindData::AdvanceFrameImp(IResimCacheBase* ResimCache)
 	//TODO: if joint is desynced we should desync particles as well
 	//If particle of joint is desynced, we need to make sure the joint is reconsidered too for optimization, though maybe not "desynced"
 
-	if (IsResim() && ResimCache)
+	if (IsResim() && bHasResimCache)
 	{
 		ResimCache->SetDesyncedParticles(MoveTemp(DesyncedParticles));
 	}
@@ -966,6 +967,14 @@ void FRewindData::PushGTDirtyData(const FDirtyPropertiesManager& SrcManager,cons
 			break;
 		}
 
+		if (bRewindDataOptimization)
+		{
+			/** When using optimization, add and mark particle as dirty and then rely on FRewindData::CacheCurrentDirtyData() to cache PostPushData */
+			FDirtyParticleInfo& Info = FindOrAddDirtyObj(*PTParticle, Proxy->IsInitialized() ? INDEX_NONE : CurFrame);
+			Info.LastDirtyFrame = CurFrame;
+			break;
+		}
+
 		const bool bKeepRecording = CopyHelper(PTParticle, [PTParticle, &DirtyPropHelper](FGeometryParticleStateBase& Latest)
 		{
 			DirtyPropHelper(Latest.ParticlePositionRotation, EChaosPropertyFlags::XR, *PTParticle);
@@ -1013,6 +1022,14 @@ void FRewindData::PushGTDirtyData(const FDirtyPropertiesManager& SrcManager,cons
 			break;
 		}
 
+		if (bRewindDataOptimization)
+		{
+			/** When using optimization, add and mark joint as dirty and then rely on FRewindData::CacheCurrentDirtyData() to cache PostPushData */
+			FDirtyJointInfo& Info = FindOrAddDirtyObj(*Joint, Proxy->IsInitialized() ? INDEX_NONE : CurFrame);
+			Info.LastDirtyFrame = CurFrame;
+			break;
+		}
+
 		CopyHelper(Joint, [Joint, &DirtyPropHelper](FJointStateBase& Latest)
 		{
 			DirtyPropHelper(Latest.JointSettings, EChaosPropertyFlags::JointSettings, *Joint);
@@ -1102,7 +1119,6 @@ void FRewindData::MarkDirtyFromPT(FGeometryParticleHandle& Handle)
 			Data->CopyFrom(Handle);
 		}
 	}
-	
 
 	if (auto Kinematic = Handle.CastToKinematicParticle())
 	{
@@ -1360,6 +1376,19 @@ void FRewindData::PushStateAtFrame(FGeometryParticleHandle& Handle, int32 Frame,
 
 void FRewindData::PushPTDirtyData(TPBDRigidParticleHandle<FReal, 3>& Handle, const int32 SrcDataIdx)
 {
+	if (bRewindDataOptimization)
+	{
+		/** When using optimization, mark particle as dirty and then rely on FRewindData::CacheCurrentDirtyData() to cache PostPushData 
+		* If this is a new entry in the dirty particles collection, cache it here since CacheCurrentDirtyData for PostPushData has already happened this frame */
+
+		if (FDirtyParticleInfo* Info = FindDirtyObj(Handle))
+		{
+			// Mark particle as dirty so it doesn't get cleared from the dirty particles
+			Info->LastDirtyFrame = CurFrame;
+			return;
+		}
+	}
+
 	const bool bRecordingHistory = !IsResimAndInSync(Handle);
 
 	FDirtyParticleInfo& Info = FindOrAddDirtyObj(Handle);
@@ -1405,6 +1434,79 @@ void FRewindData::PushPTDirtyData(TPBDRigidParticleHandle<FReal, 3>& Handle, con
 	}
 }
 
+void FRewindData::CacheCurrentDirtyData(FFrameAndPhase::EParticleHistoryPhase CurrentPhase)
+{
+	if (bRewindDataOptimization && CurrentPhase != FFrameAndPhase::PostPushData)
+	{
+		// Don't cache any other phase than PostPushData when using optimization
+		return;
+	}
+
+	const FFrameAndPhase FrameAndPhase { CurFrame, CurrentPhase };
+
+	auto DirtyPropHelper = [this, FrameAndPhase](auto& Property, const EChaosPropertyFlags PropName, const auto& Object)
+		{
+			const bool bRecordingHistory = !IsResimAndInSync(Object);
+
+			if (bRecordingHistory || Property.IsClean(FrameAndPhase))
+			{
+				if (auto* Data = Property.WriteAccessNonDecreasing(FrameAndPhase, PropertiesPool))
+				{
+					CopyDataFromObject(*Data, Object);
+				}
+			}
+		};
+
+	for (FDirtyParticleInfo& DirtyParticleInfo : DirtyParticles)
+	{
+		if (FGeometryParticleHandle* Geometry = DirtyParticleInfo.GetObjectPtr())
+		{
+			if (IsResim())
+			{
+				if (!DirtyParticleInfo.bResimAsFollower)
+				{
+					// Check if particle differ from the currently cached history for this frame, if so, clear the history and mark particle as desynced so it will cache data during resimulation
+					DesyncIfNecessary</*bSkipDynamics=*/false>(DirtyParticleInfo, FrameAndPhase);
+				}
+
+				if (Geometry->SyncState() == ESyncState::InSync)
+				{
+					/* No need to cache data in history if the particle is still in sync during resimulation */
+					continue;
+				}
+			}
+
+			FGeometryParticleStateBase& Latest = DirtyParticleInfo.GetHistory();
+
+			DirtyPropHelper(Latest.ParticlePositionRotation, EChaosPropertyFlags::XR, *Geometry);
+			DirtyPropHelper(Latest.NonFrequentData, EChaosPropertyFlags::NonFrequentData, *Geometry);
+
+			if (auto Kinematic = Geometry->CastToKinematicParticle())
+			{
+				DirtyPropHelper(Latest.Velocities, EChaosPropertyFlags::Velocities, *Kinematic);
+				DirtyPropHelper(Latest.KinematicTarget, EChaosPropertyFlags::KinematicTarget, *Kinematic);
+
+				if (auto Rigid = Kinematic->CastToRigidParticle())
+				{
+					DirtyPropHelper(Latest.Dynamics, EChaosPropertyFlags::Dynamics, *Rigid);
+					DirtyPropHelper(Latest.DynamicsMisc, EChaosPropertyFlags::DynamicMisc, *Rigid);
+					DirtyPropHelper(Latest.MassProps, EChaosPropertyFlags::MassProps, *Rigid);
+				}
+			}
+		}
+	}
+
+	for (FDirtyJointInfo& DirtyJointInfo : DirtyJoints)
+	{
+		if (const FPBDJointConstraintHandle* Joint = DirtyJointInfo.GetObjectPtr())
+		{
+			FJointStateBase& Latest = DirtyJointInfo.GetHistory();
+
+			DirtyPropHelper(Latest.JointSettings, EChaosPropertyFlags::JointSettings, *Joint);
+		}
+	}
+}
+
 FGeometryParticleState FRewindData::GetPastStateAtFrame(const FGeometryParticleHandle& Handle, int32 Frame, FFrameAndPhase::EParticleHistoryPhase Phase) const
 {
 	return GetPastStateAtFrameImp<FGeometryParticleState>(DirtyParticles, Handle, Frame, Phase);
diff --git a/Engine/Source/Runtime/Experimental/Chaos/Public/Chaos/ParticleDirtyFlags.h b/Engine/Source/Runtime/Experimental/Chaos/Public/Chaos/ParticleDirtyFlags.h
--- a/Engine/Source/Runtime/Experimental/Chaos/Public/Chaos/ParticleDirtyFlags.h
+++ b/Engine/Source/Runtime/Experimental/Chaos/Public/Chaos/ParticleDirtyFlags.h
@@ -627,8 +627,11 @@ public:
 		SetGeometry(Other.GetGeometry());
 		SetUniqueIdx(Other.UniqueIdx());
 		SetSpatialIdx(Other.SpatialIdx());
+
+		// Todo, can these be removed?
 		SetResimType(Other.ResimType());
 		SetEnabledDuringResim(Other.EnabledDuringResim());
+
 #if CHAOS_DEBUG_NAME
 		SetDebugName(Other.DebugName());
 #endif
@@ -639,9 +642,7 @@ public:
 	{
 		return GetGeometry() == Other.GetGeometry()
 			&& UniqueIdx() == Other.UniqueIdx()
-			&& SpatialIdx() == Other.SpatialIdx()
-			&& ResimType() == Other.ResimType()
-			&& EnabledDuringResim() == Other.EnabledDuringResim();
+			&& SpatialIdx() == Other.SpatialIdx();
 	}
 
 	bool operator==(const FParticleNonFrequentData& Other) const
diff --git a/Engine/Source/Runtime/Experimental/Chaos/Public/Chaos/ParticleHandle.h b/Engine/Source/Runtime/Experimental/Chaos/Public/Chaos/ParticleHandle.h
--- a/Engine/Source/Runtime/Experimental/Chaos/Public/Chaos/ParticleHandle.h
+++ b/Engine/Source/Runtime/Experimental/Chaos/Public/Chaos/ParticleHandle.h
@@ -567,7 +567,6 @@ public:
 		SetGeometry(Chaos::FImplicitObjectPtr(InData.GetGeometry()));
 		SetUniqueIdx(InData.UniqueIdx());
 		SetSpatialIdx(InData.SpatialIdx());
-		SetResimType(InData.ResimType());
 		SetParticleID(InData.ParticleID());
 	}
 
diff --git a/Engine/Source/Runtime/Experimental/Chaos/Public/RewindData.h b/Engine/Source/Runtime/Experimental/Chaos/Public/RewindData.h
--- a/Engine/Source/Runtime/Experimental/Chaos/Public/RewindData.h
+++ b/Engine/Source/Runtime/Experimental/Chaos/Public/RewindData.h
@@ -779,11 +779,13 @@ private:
 			const FFrameAndPhase& LatestFrameAndPhase = Buffer[Prev].FrameAndPhase;
 			if (bEnsureMonotonic)
 			{
-				ensure(LatestFrameAndPhase < FrameAndPhase);	//Must write in monotonic growing order so that x_{n+1} > x_n
+				//Must write in monotonic growing order so that x_{n+1} > x_n
+				ensureMsgf(LatestFrameAndPhase < FrameAndPhase, TEXT("WriteAccessImp<bEnsureMonotonic = true> trying to write to already written FrameAndPhase: %d/%d >= %d/%d"), LatestFrameAndPhase.Frame, LatestFrameAndPhase.Phase, FrameAndPhase.Frame, FrameAndPhase.Phase);
 			}
 			else
 			{
-				ensure(LatestFrameAndPhase <= FrameAndPhase);	//Must write in growing order so that x_{n+1} >= x_n
+				//Must write in growing order so that x_{n+1} >= x_n
+				ensureMsgf(LatestFrameAndPhase <= FrameAndPhase, TEXT("WriteAccessImp<bEnsureMonotonic = false> trying to write to already written FrameAndPhase: %d/%d > %d/%d"), LatestFrameAndPhase.Frame, LatestFrameAndPhase.Phase, FrameAndPhase.Frame, FrameAndPhase.Phase);
 				if (LatestFrameAndPhase == FrameAndPhase)
 				{
 					//Already wrote once for this FrameAndPhase so skip
@@ -1067,6 +1069,7 @@ inline int32 ComputeCircularSize(int32 NumFrames) { return NumFrames * FFrameAnd
 
 struct FGeometryParticleStateBase
 {
+	UE_DEPRECATED(5.6, "Use the constructor that takes a @param bCacheOnePhase")
 	explicit FGeometryParticleStateBase(int32 NumFrames)
 	: ParticlePositionRotation(ComputeCircularSize(NumFrames))
 	, NonFrequentData(ComputeCircularSize(NumFrames))
@@ -1079,7 +1082,20 @@ struct FGeometryParticleStateBase
 	, TargetVelocities(ComputeCircularSize(NumFrames))
 	, TargetStates(ComputeCircularSize(NumFrames))
 	{
+	}
 
+	explicit FGeometryParticleStateBase(int32 NumFrames, bool bCacheOnePhase)
+		: ParticlePositionRotation(bCacheOnePhase ? NumFrames : ComputeCircularSize(NumFrames))
+		, NonFrequentData(bCacheOnePhase ? NumFrames : ComputeCircularSize(NumFrames))
+		, Velocities(bCacheOnePhase ? NumFrames : ComputeCircularSize(NumFrames))
+		, Dynamics(bCacheOnePhase ? NumFrames : ComputeCircularSize(NumFrames))
+		, DynamicsMisc(bCacheOnePhase ? NumFrames : ComputeCircularSize(NumFrames))
+		, MassProps(bCacheOnePhase ? NumFrames : ComputeCircularSize(NumFrames))
+		, KinematicTarget(bCacheOnePhase ? NumFrames : ComputeCircularSize(NumFrames))
+		, TargetPositions(NumFrames)
+		, TargetVelocities(NumFrames)
+		, TargetStates(NumFrames)
+	{
 	}
 
 	FGeometryParticleStateBase(const FGeometryParticleStateBase& Other) = delete;
@@ -1123,9 +1139,6 @@ struct FGeometryParticleStateBase
 		DynamicsMisc.ClearEntryAndFuture(FrameAndPhase);
 		MassProps.ClearEntryAndFuture(FrameAndPhase);
 		KinematicTarget.ClearEntryAndFuture(FrameAndPhase);
-		TargetPositions.ClearEntryAndFuture(FrameAndPhase);
-		TargetVelocities.ClearEntryAndFuture(FrameAndPhase);
-		TargetStates.ClearEntryAndFuture(FrameAndPhase);
 	}
 
 	void ExtractHistoryState(int32& PositionValidCount, int32& VelocityValidCount, int32& PositionNextIterator, int32& VelocityNextIterator) const
@@ -1324,12 +1337,19 @@ private:
 
 struct FJointStateBase
 {
+	UE_DEPRECATED(5.6, "Use the constructor that takes a @param bCacheOnePhase")
 	explicit FJointStateBase(int32 NumFrames)
 		: JointSettings(ComputeCircularSize(NumFrames))
 		, JointProxies(ComputeCircularSize(NumFrames))
 	{
 	}
 
+	explicit FJointStateBase(int32 NumFrames, bool bCacheOnePhase)
+		: JointSettings(bCacheOnePhase ? NumFrames : ComputeCircularSize(NumFrames))
+		, JointProxies(bCacheOnePhase ? NumFrames : ComputeCircularSize(NumFrames))
+	{
+	}
+
 	FJointStateBase(const FJointStateBase& Other) = delete;
 	FJointStateBase(FJointStateBase&& Other) = default;
 	~FJointStateBase() = default;
@@ -1622,6 +1642,7 @@ public:
 	/** Check if a resim cache based on IResimCacheBase (FEvolutionResimCache by default) is being used. Read FPhysicsSolverBase.SetUseCollisionResimCache() for more info. */
 	bool GetUseCollisionResimCache() const;
 
+	/** Called just before physics is solved */
 	template <typename CreateCache>
 	void AdvanceFrame(FReal DeltaTime, const CreateCache& CreateCacheFunc)
 	{
@@ -1675,13 +1696,22 @@ public:
 	//Number of particles that we're currently storing history for
 	int32 GetNumDirtyParticles() const { return DirtyParticles.Num(); }
 
+	/** Called just before Proxy::PushToPhysicsState is called */
 	void PushGTDirtyData(const FDirtyPropertiesManager& SrcManager,const int32 SrcDataIdx,const FDirtyProxy& Dirty, const FShapeDirtyData* ShapeDirtyData);
 
+	/** Called post solve but just before PQ are applied to XR */
 	void PushPTDirtyData(TPBDRigidParticleHandle<FReal,3>& Rigid,const int32 SrcDataIdx);
 
+	/** Call this to mark specific particle as dirty and cache its current data, expected to be called at the start of OnPreSimulate_Internal, it caches data in FFrameAndPhase::PostPushData */
 	void CHAOS_API MarkDirtyFromPT(FGeometryParticleHandle& Handle);
+	
+	/** Call this to mark specific joint as dirty and cache its current data, expected to be called at the start of OnPreSimulate_Internal, it caches data in FFrameAndPhase::PostPushData */
 	void CHAOS_API MarkDirtyJointFromPT(FPBDJointConstraintHandle& Handle);
 
+	/** Caches data for all particles and joints marked dirty in RewindData 
+	* @param CurrentPhase Data gets cached under this phase for the current frame, be sure to call this from the corresponding phase during physics */
+	void CacheCurrentDirtyData(FFrameAndPhase::EParticleHistoryPhase CurrentPhase);
+
 	void CHAOS_API SpawnProxyIfNeeded(FSingleParticlePhysicsProxy& Proxy);
 
 	/** Add input history to the rewind data for future use while resimulating */
@@ -1783,8 +1813,8 @@ private:
 		bool bResimAsSlave = true;
 		bool bResimAsFollower = true;	//Indicates the particle will always resim in the exact same way from game thread data
 
-		TDirtyObjectInfo(FDirtyPropertiesPool& InPropertiesPool, TObj& InObj, const int32 CurFrame, const int32 NumFrames)
-			: History(NumFrames)
+		TDirtyObjectInfo(FDirtyPropertiesPool& InPropertiesPool, TObj& InObj, const int32 CurFrame, const int32 NumFrames, const bool bCacheOnePhase)
+			: History(NumFrames, bCacheOnePhase)
 			, ObjPtr(&InObj)
 			, PropertiesPool(&InPropertiesPool)
 			, LastDirtyFrame(CurFrame)
@@ -1872,6 +1902,11 @@ private:
 		return FindDirtyObjImp(DirtyParticles, Handle);
 	}
 
+	FDirtyJointInfo* FindDirtyObj(const FPBDJointConstraintHandle& Handle)
+	{
+		return FindDirtyObjImp(DirtyJoints, Handle);
+	}
+
 	template <typename TDirtyObjs, typename TObj>
 	auto& FindOrAddDirtyObjImp(TDirtyObjs & DirtyObjs, TObj & Handle, const int32 InitializedOnFrame = INDEX_NONE)
 	{
@@ -1881,7 +1916,7 @@ private:
 		}
 
 		using TDirtyObj = decltype(NoRefHelper(DirtyObjs.GetDenseAt(0)));
-		TDirtyObj& Info = DirtyObjs.Add(&Handle, TDirtyObj(PropertiesPool, Handle, CurFrame, Managers.Capacity()));
+		TDirtyObj& Info = DirtyObjs.Add(&Handle, TDirtyObj(PropertiesPool, Handle, CurFrame, Managers.Capacity(), bRewindDataOptimization));
 		Info.InitializedOnStep = InitializedOnFrame;
 		return Info;
 	}
