Add an ansi string overload of the latency test.

#rnx

- This will help in moving the backend from the single endpoint http client to the multipoint. Once that is done we can remove the non-ansi string overload.

diff --git a/Engine/Source/Runtime/Experimental/IoStore/OnDemand/Private/LatencyTesting.cpp b/Engine/Source/Runtime/Experimental/IoStore/OnDemand/Private/LatencyTesting.cpp
--- a/Engine/Source/Runtime/Experimental/IoStore/OnDemand/Private/LatencyTesting.cpp
+++ b/Engine/Source/Runtime/Experimental/IoStore/OnDemand/Private/LatencyTesting.cpp
@@ -11,14 +11,12 @@
 namespace UE::IoStore
 {
 
-void LatencyTest(FStringView InUrl, FStringView InPath, uint32 InTimeOutMs, TArrayView<int32> OutResults)
+void LatencyTest(FAnsiStringView InUrl, FStringView InPath, uint32 InTimeOutMs, TArrayView<int32> OutResults)
 {
 	using namespace HTTP;
 
-	auto AnsiUrl = StringCast<ANSICHAR>(InUrl.GetData(), InUrl.Len());
-
 	FConnectionPool::FParams PoolParams;
-	PoolParams.SetHostFromUrl(AnsiUrl);
+	PoolParams.SetHostFromUrl(InUrl);
 	PoolParams.ConnectionCount = 1;
 	FConnectionPool Pool(PoolParams);
 
@@ -66,4 +64,10 @@ void LatencyTest(FStringView InUrl, FStringView InPath, uint32 InTimeOutMs, TArr
 	UE_LOG(LogIas, Log, TEXT("Testing endpoint %s"), ANSI_TO_TCHAR(ConnectionDesc.ToString()));
 }
 
+void LatencyTest(FStringView InUrl, FStringView InPath, uint32 InTimeOutMs, TArrayView<int32> OutResults)
+{
+	auto AnsiUrl = StringCast<ANSICHAR>(InUrl.GetData(), InUrl.Len());
+	LatencyTest(AnsiUrl, InPath, InTimeOutMs, OutResults);
+}
+
 } // namespace UE::IoStore
diff --git a/Engine/Source/Runtime/Experimental/IoStore/OnDemand/Private/OnDemandIoDispatcherBackend.cpp b/Engine/Source/Runtime/Experimental/IoStore/OnDemand/Private/OnDemandIoDispatcherBackend.cpp
--- a/Engine/Source/Runtime/Experimental/IoStore/OnDemand/Private/OnDemandIoDispatcherBackend.cpp
+++ b/Engine/Source/Runtime/Experimental/IoStore/OnDemand/Private/OnDemandIoDispatcherBackend.cpp
@@ -56,6 +56,7 @@ namespace UE::IoStore
 
 extern FString GIasOnDemandTocExt;
 void LatencyTest(FStringView, FStringView, uint32, TArrayView<int32>);
+void LatencyTest(FAnsiStringView, FStringView, uint32, TArrayView<int32>);
 
 ///////////////////////////////////////////////////////////////////////////////
 /** Note that GIasHttpPrimaryEndpoint has no effect after initial start up */
