ModuleInput: fix for GetSafeNormal2D
(found with [[nodiscard]])
#rnx

diff --git a/Engine/Source/Runtime/Experimental/ChaosVehicles/ChaosVehiclesCore/Private/SimModule/ModuleInput.cpp b/Engine/Source/Runtime/Experimental/ChaosVehicles/ChaosVehiclesCore/Private/SimModule/ModuleInput.cpp
--- a/Engine/Source/Runtime/Experimental/ChaosVehicles/ChaosVehiclesCore/Private/SimModule/ModuleInput.cpp
+++ b/Engine/Source/Runtime/Experimental/ChaosVehicles/ChaosVehiclesCore/Private/SimModule/ModuleInput.cpp
@@ -50,10 +50,10 @@ void FModuleInputValue::SetMagnitude(float NewSize)
 		Value.X = NewSize;
 		break;
 	case EModuleInputValueType::MAxis2D:
-		Value.GetSafeNormal2D()* NewSize;
+		Value = Value.GetSafeNormal2D()* NewSize;
 		break;
 	case EModuleInputValueType::MAxis3D:
-		Value.GetSafeNormal()* NewSize;
+		Value = Value.GetSafeNormal()* NewSize;
 		break;
 	default:
 		checkf(false, TEXT("Unsupported value type for module input value!"));
