Add a new cmdline option '-Ias.SkipDevCDNs' that will remove any internal development CDNs from the list provided by the distributed endpoint.

#rnx

- This option is compiled out in shipping and is intended as a debug tool only.
- We identify the development urls via a new config file value 'ini:Engine:[Ias]:DevelopmentCDNPattern' in the form of a regular expression. Any url that this expression matches is removed from the list.
- The order of the remaining urls is preserved.
- By default no regular expression is set. Any team wanting to make use of this debug feature will need to provide one first.

diff --git a/Engine/Source/Runtime/Experimental/IoStore/OnDemand/Private/DistributionEndpoints.cpp b/Engine/Source/Runtime/Experimental/IoStore/OnDemand/Private/DistributionEndpoints.cpp
--- a/Engine/Source/Runtime/Experimental/IoStore/OnDemand/Private/DistributionEndpoints.cpp
+++ b/Engine/Source/Runtime/Experimental/IoStore/OnDemand/Private/DistributionEndpoints.cpp
@@ -2,6 +2,7 @@
 
 #include "DistributionEndpoints.h"
 
+#include "Algo/RemoveIf.h"
 #include "Containers/StringConv.h"
 #include "Dom/JsonValue.h"
 #include "HAL/IConsoleManager.h"
@@ -9,6 +10,8 @@
 #include "IO/Http/Client.h"
 #include "IO/IoBuffer.h"
 #include "IO/IoStoreOnDemand.h"
+#include "Misc/CommandLine.h"
+#include "Misc/ConfigCacheIni.h"
 #include "Misc/ScopeRWLock.h"
 #include "Serialization/JsonReader.h"
 #include "Serialization/JsonSerializer.h"
@@ -16,6 +19,10 @@
 #include "Statistics.h"
 #include "Tasks/Task.h"
 
+#if !UE_BUILD_SHIPPING
+#include "Internationalization/Regex.h"
+#endif // !UE_BUILD_SHIPPING
+
 namespace UE::IoStore
 {
 
@@ -90,6 +97,28 @@ FDistributionEndpoints::EResult FDistributionEndpoints::ResolveEndpoints(const F
 	while (Loop.Tick(-1))
 		;
 
+#if !UE_BUILD_SHIPPING
+	if (FParse::Param(FCommandLine::Get(), TEXT("Ias.SkipDevCDNs")))
+	{
+		FString ConfigRegex;
+		GConfig->GetString(TEXT("Ias"), TEXT("DevelopmentCDNPattern"), ConfigRegex, GEngineIni);
+
+		if(!ConfigRegex.IsEmpty())
+		{
+			FRegexPattern Pattern(ConfigRegex);
+
+			const int32 OriginaNum = OutServiceUrls.Num();
+			OutServiceUrls.SetNum(Algo::StableRemoveIf(OutServiceUrls, [&Pattern](const FString& Entry)
+			{
+				FRegexMatcher Regex = FRegexMatcher(Pattern, Entry);
+				return Regex.FindNext();
+			}));
+
+			UE_LOG(LogIas, Log, TEXT("Removed %d development CDNs from the distributed endpoint list"), OriginaNum - OutServiceUrls.Num());
+		}
+	}
+#endif //!UE_BUILD_SHIPPING
+
 	UE_CLOG(Result == EResult::Success, LogIas, Log, TEXT("Successfully resolved distributed endpoint '%s' %d urls found"), *DistributionUrl, OutServiceUrls.Num());
 	UE_CLOG(Result != EResult::Success, LogIas, Log, TEXT("Failed to resolve distributed endpoint '%s'"), *DistributionUrl);
 
