Make Windows builds use the unoptimized version of ISPC VectorGather to workaround a crash caused by using the optimized version. Clean up hotfix cvars and speculative fix optional code related to debugging the crash.

* Added a new define USE_ISPC_OPTIMIZED_VECTORGATHER which is used by the Vector.isph code to enable the optimized gather code. The default value is 1.
* Set USE_ISPC_OPTIMIZED_VECTORGATHER to 0 on Windows.
* Removed the BendingConstraints and TriangleMesh ISPC functions which explicitly used the unoptimized gather, their associated cvar switches, and the related define USE_ISPC_UNOPTIMIZED_VECTORGATHER_DEFAULT
* Removed setting USE_ISPC_KERNEL_CONSOLE_VARIABLES_IN_SHIPPING on the PC config. This will make PC ship builds match other configurations and compile out the cvars.

diff --git a/Engine/Platforms/Windows/Source/Runtime/Experimental/Chaos/Chaos_Windows.Build.cs b/Engine/Platforms/Windows/Source/Runtime/Experimental/Chaos/Chaos_Windows.Build.cs
--- a/Engine/Platforms/Windows/Source/Runtime/Experimental/Chaos/Chaos_Windows.Build.cs
+++ b/Engine/Platforms/Windows/Source/Runtime/Experimental/Chaos/Chaos_Windows.Build.cs
@@ -9,7 +9,8 @@ namespace UnrealBuildTool.Rules
     {
         public Chaos_Windows(ReadOnlyTargetRules Target) : base(Target)
         {
-			PublicDefinitions.Add("USE_ISPC_KERNEL_CONSOLE_VARIABLES_IN_SHIPPING=0");
+			// Disabling the optimized VectorGather call because it is causing crashes on PC.
+			PublicDefinitions.Add("USE_ISPC_OPTIMIZED_VECTORGATHER=0");
 		}
     }
 }
diff --git a/Engine/Source/Runtime/Core/Public/Math/Vector.isph b/Engine/Source/Runtime/Core/Public/Math/Vector.isph
--- a/Engine/Source/Runtime/Core/Public/Math/Vector.isph
+++ b/Engine/Source/Runtime/Core/Public/Math/Vector.isph
@@ -12,6 +12,11 @@
 #define EXPLICIT_VECTOR2 0
 #endif
 
+// Opt-out of VectorGather optimization
+#ifndef USE_ISPC_OPTIMIZED_VECTORGATHER
+#define USE_ISPC_OPTIMIZED_VECTORGATHER 1
+#endif
+
 struct FVector2f
 {
 #if EXPLICIT_VECTOR2 == 0
@@ -606,6 +611,7 @@ inline varying FVector3f VectorGather(const uniform FVector3f *varying SrcPtr)
 {
 	varying FVector3f Result;
 
+#if USE_ISPC_OPTIMIZED_VECTORGATHER
 	if (((1 << TARGET_WIDTH) - 1 ^ lanemask()) == 0)
 	{
 #if TARGET_WIDTH == 4
@@ -616,6 +622,7 @@ inline varying FVector3f VectorGather(const uniform FVector3f *varying SrcPtr)
 #endif
 	}
 	else
+#endif
 	{
 		#pragma ignore warning(perf)
 		Result = *SrcPtr;
@@ -671,7 +678,8 @@ inline void VectorScatter(uniform FVector3f *varying DstPtr, const varying FVect
 inline varying FVector4d VectorGather(const uniform FVector4d *varying SrcPtr)
 {
 	varying FVector4d Result;
-
+	
+#if USE_ISPC_OPTIMIZED_VECTORGATHER
 	if(((1<<TARGET_WIDTH)-1 ^ lanemask()) == 0)
 	{
 #if TARGET_WIDTH == 4
@@ -682,6 +690,7 @@ inline varying FVector4d VectorGather(const uniform FVector4d *varying SrcPtr)
 #endif
 	}
 	else
+#endif
 	{
 		#pragma ignore warning(perf)
 		Result = *SrcPtr;
@@ -693,7 +702,8 @@ inline varying FVector4d VectorGather(const uniform FVector4d *varying SrcPtr)
 inline varying FVector4f VectorGather(const uniform FVector4f *varying SrcPtr)
 {
 	varying FVector4f Result;
-
+	
+#if USE_ISPC_OPTIMIZED_VECTORGATHER
 	if(((1<<TARGET_WIDTH)-1 ^ lanemask()) == 0)
 	{
 #if TARGET_WIDTH == 4
@@ -704,6 +714,7 @@ inline varying FVector4f VectorGather(const uniform FVector4f *varying SrcPtr)
 #endif
 	}
 	else
+#endif
 	{
 		#pragma ignore warning(perf)
 		Result = *SrcPtr;
@@ -759,7 +770,8 @@ inline void VectorScatter(uniform FVector4f *varying DstPtr, const varying FVect
 inline varying FIntVector VectorGather(const uniform FIntVector *varying SrcPtr)
 {
 	varying FIntVector Result;
-
+	
+#if USE_ISPC_OPTIMIZED_VECTORGATHER
 	if (((1 << TARGET_WIDTH) - 1 ^ lanemask()) == 0)
 	{
 #if TARGET_WIDTH == 4
@@ -770,6 +782,7 @@ inline varying FIntVector VectorGather(const uniform FIntVector *varying SrcPtr)
 #endif
 	}
 	else
+#endif
 	{
 		#pragma ignore warning(perf)
 		Result = *SrcPtr;
diff --git a/Engine/Source/Runtime/Experimental/Chaos/Private/Chaos/PBDBendingConstraints.cpp b/Engine/Source/Runtime/Experimental/Chaos/Private/Chaos/PBDBendingConstraints.cpp
--- a/Engine/Source/Runtime/Experimental/Chaos/Private/Chaos/PBDBendingConstraints.cpp
+++ b/Engine/Source/Runtime/Experimental/Chaos/Private/Chaos/PBDBendingConstraints.cpp
@@ -169,28 +169,14 @@ void FPBDBendingConstraints::Apply(SolverParticlesOrRange& Particles, const FSol
 				{
 					const int32 ColorStart = ConstraintsPerColorStartIndex[ConstraintColorIndex];
 					const int32 ColorSize = ConstraintsPerColorStartIndex[ConstraintColorIndex + 1] - ColorStart;
-					if (bChaos_Bending_ISPCUseHWGather)
-					{
-						ispc::ApplyBendingConstraintsHWGather(
-							(ispc::FVector4f*)Particles.GetPAndInvM().GetData(),
-							(ispc::FIntVector4*)&Constraints.GetData()[ColorStart],
-							&RestAngles.GetData()[ColorStart],
-							&IsBuckled.GetData()[ColorStart],
-							ExpStiffnessValue,
-							ExpBucklingValue,
-							ColorSize);
-					}
-					else
-					{
-						ispc::ApplyBendingConstraints(
-							(ispc::FVector4f*)Particles.GetPAndInvM().GetData(),
-							(ispc::FIntVector4*)&Constraints.GetData()[ColorStart],
-							&RestAngles.GetData()[ColorStart],
-							&IsBuckled.GetData()[ColorStart],
-							ExpStiffnessValue,
-							ExpBucklingValue,
-							ColorSize);
-					}
+					ispc::ApplyBendingConstraints(
+						(ispc::FVector4f*)Particles.GetPAndInvM().GetData(),
+						(ispc::FIntVector4*)&Constraints.GetData()[ColorStart],
+						&RestAngles.GetData()[ColorStart],
+						&IsBuckled.GetData()[ColorStart],
+						ExpStiffnessValue,
+						ExpBucklingValue,
+						ColorSize);
 				}
 			}
 			else
@@ -217,36 +203,18 @@ void FPBDBendingConstraints::Apply(SolverParticlesOrRange& Particles, const FSol
 				{
 					const int32 ColorStart = ConstraintsPerColorStartIndex[ConstraintColorIndex];
 					const int32 ColorSize = ConstraintsPerColorStartIndex[ConstraintColorIndex + 1] - ColorStart;
-					if (bChaos_Bending_ISPCUseHWGather)
-					{
-						ispc::ApplyBendingConstraintsWithMapsHWGather(
-							(ispc::FVector4f*)Particles.GetPAndInvM().GetData(),
-							(ispc::FIntVector4*)&Constraints.GetData()[ColorStart],
-							&RestAngles.GetData()[ColorStart],
-							&IsBuckled.GetData()[ColorStart],
-							StiffnessHasWeightMap,
-							StiffnessHasWeightMap ? &Stiffness.GetIndices().GetData()[ColorStart] : nullptr,
-							&Stiffness.GetTable().GetData()[0],
-							BucklingStiffnessHasWeightMap,
-							BucklingStiffnessHasWeightMap ? &BucklingStiffness.GetIndices().GetData()[ColorStart] : nullptr,
-							&BucklingStiffness.GetTable().GetData()[0],
-							ColorSize);
-					}
-					else
-					{
-						ispc::ApplyBendingConstraintsWithMaps(
-							(ispc::FVector4f*)Particles.GetPAndInvM().GetData(),
-							(ispc::FIntVector4*)&Constraints.GetData()[ColorStart],
-							&RestAngles.GetData()[ColorStart],
-							&IsBuckled.GetData()[ColorStart],
-							StiffnessHasWeightMap,
-							StiffnessHasWeightMap ? &Stiffness.GetIndices().GetData()[ColorStart] : nullptr,
-							&Stiffness.GetTable().GetData()[0],
-							BucklingStiffnessHasWeightMap,
-							BucklingStiffnessHasWeightMap ? &BucklingStiffness.GetIndices().GetData()[ColorStart] : nullptr,
-							&BucklingStiffness.GetTable().GetData()[0],
-							ColorSize);
-					}
+					ispc::ApplyBendingConstraintsWithMaps(
+						(ispc::FVector4f*)Particles.GetPAndInvM().GetData(),
+						(ispc::FIntVector4*)&Constraints.GetData()[ColorStart],
+						&RestAngles.GetData()[ColorStart],
+						&IsBuckled.GetData()[ColorStart],
+						StiffnessHasWeightMap,
+						StiffnessHasWeightMap ? &Stiffness.GetIndices().GetData()[ColorStart] : nullptr,
+						&Stiffness.GetTable().GetData()[0],
+						BucklingStiffnessHasWeightMap,
+						BucklingStiffnessHasWeightMap ? &BucklingStiffness.GetIndices().GetData()[ColorStart] : nullptr,
+						&BucklingStiffness.GetTable().GetData()[0],
+						ColorSize);
 				}
 			}
 			else
diff --git a/Engine/Source/Runtime/Experimental/Chaos/Private/Chaos/PBDBendingConstraints.ispc b/Engine/Source/Runtime/Experimental/Chaos/Private/Chaos/PBDBendingConstraints.ispc
--- a/Engine/Source/Runtime/Experimental/Chaos/Private/Chaos/PBDBendingConstraints.ispc
+++ b/Engine/Source/Runtime/Experimental/Chaos/Private/Chaos/PBDBendingConstraints.ispc
@@ -95,39 +95,6 @@ export void InitBendingConstraintsIsBuckled(const uniform FVector3f X[],
 	}
 }
 
-export void InitBendingConstraintsIsBuckledHWGather(const uniform FVector3f X[],
-									const uniform FIntVector4 Constraints[],
-									const uniform float RestAngles[],
-									uniform bool IsBuckled[],
-									const uniform float BucklingRatio,
-									const uniform int32 NumConstraints)
-{
-	foreach(i = 0 ... NumConstraints)
-	{
-		const varying FIntVector4 Constraint = VectorLoad(&Constraints[extract(i,0)]);
-
-		const varying int32 Index1 = Constraint.V[0];
-		const varying int32 Index2 = Constraint.V[1];
-		const varying int32 Index3 = Constraint.V[2];
-		const varying int32 Index4 = Constraint.V[3];
-		
-		#pragma ignore warning(perf)
-		const varying FVector3f X1 = X[Index1];
-		#pragma ignore warning(perf)
-		const varying FVector3f X2 = X[Index2];
-		#pragma ignore warning(perf)
-		const varying FVector3f X3 = X[Index3];
-		#pragma ignore warning(perf)
-		const varying FVector3f X4 = X[Index4];
-
-		const varying float RestAngle = RestAngles[i];
-		const varying float BucklingAngle = BucklingRatio * (FLOAT_PI - abs(RestAngle));
-
-		const varying float Angle = CalcAngle(X1, X2, X3, X4);
-		IsBuckled[i] = FLOAT_PI - abs(Angle) < BucklingAngle;
-	}
-}
-
 export void InitBendingConstraintsIsBuckledWithMaps(const uniform FVector3f X[],
 									const uniform FIntVector4 Constraints[],
 									const uniform float RestAngles[],
@@ -161,43 +128,6 @@ export void InitBendingConstraintsIsBuckledWithMaps(const uniform FVector3f X[],
 	}
 }
 
-export void InitBendingConstraintsIsBuckledWithMapsHWGather(const uniform FVector3f X[],
-									const uniform FIntVector4 Constraints[],
-									const uniform float RestAngles[],
-									uniform bool IsBuckled[],
-									const uniform uint8 BucklingRatioIndices[],
-									const uniform float BucklingRatioTable[],
-									const uniform int32 NumConstraints)
-{
-	foreach(i = 0 ... NumConstraints)
-	{
-		const varying FIntVector4 Constraint = VectorLoad(&Constraints[extract(i,0)]);
-
-		const varying int32 Index1 = Constraint.V[0];
-		const varying int32 Index2 = Constraint.V[1];
-		const varying int32 Index3 = Constraint.V[2];
-		const varying int32 Index4 = Constraint.V[3];
-		const varying int8 BucklingRatioIndex = BucklingRatioIndices[i];
-		
-		#pragma ignore warning(perf)
-		const varying FVector3f X1 = X[Index1];
-		#pragma ignore warning(perf)
-		const varying FVector3f X2 = X[Index2];
-		#pragma ignore warning(perf)
-		const varying FVector3f X3 = X[Index3];
-		#pragma ignore warning(perf)
-		const varying FVector3f X4 = X[Index4];
-		
-		#pragma ignore warning(perf)
-		const varying float BucklingRatio = BucklingRatioTable[BucklingRatioIndex];
-		const varying float RestAngle = RestAngles[i];
-		const varying float BucklingAngle = BucklingRatio * (FLOAT_PI - abs(RestAngle));
-
-		const varying float Angle = CalcAngle(X1, X2, X3, X4);
-		IsBuckled[i] = FLOAT_PI - abs(Angle) < BucklingAngle;
-	}
-}
-
 export void ApplyBendingConstraints(uniform FVector4f PandInvM[],
 									const uniform FIntVector4 Constraints[],
 									const uniform float RestAngles[],
@@ -261,73 +191,6 @@ export void ApplyBendingConstraints(uniform FVector4f PandInvM[],
 	}
 }
 
-export void ApplyBendingConstraintsHWGather(uniform FVector4f PandInvM[],
-									const uniform FIntVector4 Constraints[],
-									const uniform float RestAngles[],
-									const uniform bool IsBuckled[],
-									const uniform float Stiffness,
-									const uniform float BucklingStiffness,
-									const uniform int32 NumConstraints)
-{
-
-	foreach(i = 0 ... NumConstraints)
-	{
-		const varying FIntVector4 Constraint = VectorLoad(&Constraints[extract(i,0)]);
-
-		const varying int32 i1 = Constraint.V[0];
-		const varying int32 i2 = Constraint.V[1];
-		const varying int32 i3 = Constraint.V[2];
-		const varying int32 i4 = Constraint.V[3];
-		
-		#pragma ignore warning(perf)
-		const varying FVector4f PandInvM2 = PandInvM[i2];
-		#pragma ignore warning(perf)
-		const varying FVector4f PandInvM3 = PandInvM[i3];
-		#pragma ignore warning(perf)
-		const varying FVector4f PandInvM1 = PandInvM[i1];
-		#pragma ignore warning(perf)
-		const varying FVector4f PandInvM4 = PandInvM[i4];
-
-		varying FVector3f P1, P2, P3, P4;
-		varying float M1, M2, M3, M4;
-
-		UnzipPandInvM(PandInvM1, P1, M1);
-		UnzipPandInvM(PandInvM2, P2, M2);
-		UnzipPandInvM(PandInvM3, P3, M3);
-		UnzipPandInvM(PandInvM4, P4, M4);
-
-
-		varying FVector3f Grad1, Grad2, Grad3, Grad4;
-		const varying float Angle = CalcGradientsAndAngle(P1, P2, P3, P4, Grad1, Grad2, Grad3, Grad4);
-
-		const varying float Denom = M1 * VectorSizeSquared(Grad1) + M2 * VectorSizeSquared(Grad2) + M3 * VectorSizeSquared(Grad3) + M4 * VectorSizeSquared(Grad4);
-		
-		const varying float RestAngle = RestAngles[i];
-		const varying float StiffnessValue = select(IsBuckled[i], BucklingStiffness, Stiffness);
-
-		const varying float Delta = clamp(StiffnessValue * (Angle - RestAngle), -SingleStepAngleLimit, SingleStepAngleLimit);
-
-		const varying float S = SafeDivide(Delta, Denom);
-
-		if (M1 > 0)
-		{
-			VectorScatter(&PandInvM[i1], SetVector4(P1 - S * M1 * Grad1, M1));
-		}
-		if (M2 > 0)
-		{
-			VectorScatter(&PandInvM[i2], SetVector4(P2 - S * M2 * Grad2, M2));
-		}
-		if (M3 > 0)
-		{
-			VectorScatter(&PandInvM[i3], SetVector4(P3 - S * M3 * Grad3, M3));
-		}
-		if (M4 > 0)
-		{
-			VectorScatter(&PandInvM[i4], SetVector4(P4 - S * M4 * Grad4, M4));
-		}
-	}
-}
-
 export void ApplyBendingConstraintsWithMaps(uniform FVector4f PandInvM[],
 											const uniform FIntVector4 Constraints[],
 											const uniform float RestAngles[],
@@ -366,84 +229,6 @@ export void ApplyBendingConstraintsWithMaps(uniform FVector4f PandInvM[],
 		UnzipPandInvM(PandInvM4, P4, M4);
 
 
-		varying FVector3f Grad1, Grad2, Grad3, Grad4;
-		const varying float Angle = CalcGradientsAndAngle(P1, P2, P3, P4, Grad1, Grad2, Grad3, Grad4);
-
-		const varying float Denom = M1 * VectorSizeSquared(Grad1) + M2 * VectorSizeSquared(Grad2) + M3 * VectorSizeSquared(Grad3) + M4 * VectorSizeSquared(Grad4);
-		
-		#pragma ignore warning(perf)
-		const varying float Stiffness = StiffnessTable[StiffnessIndex];
-		#pragma ignore warning(perf)
-		const varying float BucklingStiffness = BucklingStiffnessTable[BucklingStiffnessIndex];
-
-		const varying float RestAngle = RestAngles[i];
-		const varying float StiffnessValue = select(IsBuckled[i], BucklingStiffness, Stiffness);
-
-		const varying float Delta = clamp(StiffnessValue * (Angle - RestAngle), -SingleStepAngleLimit, SingleStepAngleLimit);
-
-		const varying float S = SafeDivide(Delta, Denom);
-
-		if (M1 > 0)
-		{
-			VectorScatter(&PandInvM[i1], SetVector4(P1 - S * M1 * Grad1, M1));
-		}
-		if (M2 > 0)
-		{
-			VectorScatter(&PandInvM[i2], SetVector4(P2 - S * M2 * Grad2, M2));
-		}
-		if (M3 > 0)
-		{
-			VectorScatter(&PandInvM[i3], SetVector4(P3 - S * M3 * Grad3, M3));
-		}
-		if (M4 > 0)
-		{
-			VectorScatter(&PandInvM[i4], SetVector4(P4 - S * M4 * Grad4, M4));
-		}
-	}
-}
-
-export void ApplyBendingConstraintsWithMapsHWGather(uniform FVector4f PandInvM[],
-											const uniform FIntVector4 Constraints[],
-											const uniform float RestAngles[],
-											const uniform bool IsBuckled[],
-											const uniform bool StiffnessHasMap,
-											const uniform uint8 StiffnessIndices[],
-											const uniform float StiffnessTable[],
-											const uniform bool BucklingStiffnessHasMap,
-											const uniform uint8 BucklingStiffnessIndices[],
-											const uniform float BucklingStiffnessTable[],
-											const uniform int32 NumConstraints)
-{
-
-	foreach(i = 0 ... NumConstraints)
-	{
-		const varying FIntVector4 Constraint = VectorLoad(&Constraints[extract(i,0)]);
-
-		const varying int32 i1 = Constraint.V[0];
-		const varying int32 i2 = Constraint.V[1];
-		const varying int32 i3 = Constraint.V[2];
-		const varying int32 i4 = Constraint.V[3];
-		const varying int8 StiffnessIndex = StiffnessHasMap ? StiffnessIndices[i] : 0;
-		const varying int8 BucklingStiffnessIndex = BucklingStiffnessHasMap ? BucklingStiffnessIndices[i] : 0;		
-		
-		#pragma ignore warning(perf)
-		const varying FVector4f PandInvM2 = PandInvM[i2];
-		#pragma ignore warning(perf)
-		const varying FVector4f PandInvM3 = PandInvM[i3];
-		#pragma ignore warning(perf)
-		const varying FVector4f PandInvM1 = PandInvM[i1];
-		#pragma ignore warning(perf)
-		const varying FVector4f PandInvM4 = PandInvM[i4];
-
-		varying FVector3f P1, P2, P3, P4;
-		varying float M1, M2, M3, M4;
-
-		UnzipPandInvM(PandInvM1, P1, M1);
-		UnzipPandInvM(PandInvM2, P2, M2);
-		UnzipPandInvM(PandInvM3, P3, M3);
-		UnzipPandInvM(PandInvM4, P4, M4);
-
-
 		varying FVector3f Grad1, Grad2, Grad3, Grad4;
 		const varying float Angle = CalcGradientsAndAngle(P1, P2, P3, P4, Grad1, Grad2, Grad3, Grad4);
 
diff --git a/Engine/Source/Runtime/Experimental/Chaos/Private/Chaos/PBDBendingConstraintsBase.cpp b/Engine/Source/Runtime/Experimental/Chaos/Private/Chaos/PBDBendingConstraintsBase.cpp
--- a/Engine/Source/Runtime/Experimental/Chaos/Private/Chaos/PBDBendingConstraintsBase.cpp
+++ b/Engine/Source/Runtime/Experimental/Chaos/Private/Chaos/PBDBendingConstraintsBase.cpp
@@ -9,8 +9,6 @@
 bool bChaos_Bending_ISPC_Enabled = CHAOS_BENDING_ISPC_ENABLED_DEFAULT;
 FAutoConsoleVariableRef CVarChaosBendingISPCEnabled(TEXT("p.Chaos.Bending.ISPC"), bChaos_Bending_ISPC_Enabled, TEXT("Whether to use ISPC optimizations in Bending constraints"));
 
-bool bChaos_Bending_ISPCUseHWGather = USE_ISPC_UNOPTIMIZED_VECTORGATHER_DEFAULT;
-FAutoConsoleVariableRef CVarChaosBendingISPCUseHWGather(TEXT("p.Chaos.Bending.ISPCUseHWGather"), bChaos_Bending_ISPCUseHWGather, TEXT("Use (slower) HW gather when computing Bending constraints to see if this fixes the ISPC crash."));
 #endif
 
 static_assert(sizeof(ispc::FVector4f) == sizeof(Chaos::Softs::FPAndInvM), "sizeof(ispc::FVector4f) != sizeof(Chaos::Softs::FPAndInvM");
@@ -163,55 +161,26 @@ namespace Chaos::Softs
 		{
 			if (BucklingRatioWeighted.HasWeightMap())
 			{
-				if (bChaos_Bending_ISPCUseHWGather)
-				{
-					ispc::InitBendingConstraintsIsBuckledWithMapsHWGather(
-						(const ispc::FVector3f*)InParticles.XArray().GetData(),
-						(const ispc::FIntVector4*)Constraints.GetData(),
-						RestAngles.GetData(),
-						IsBuckled.GetData(),
-						BucklingRatioWeighted.GetIndices().GetData(),
-						BucklingRatioWeighted.GetTable().GetData(),
-						Constraints.Num()
-					);
-				}
-				else
-				{
-					ispc::InitBendingConstraintsIsBuckledWithMaps(
-						(const ispc::FVector3f*)InParticles.XArray().GetData(),
-						(const ispc::FIntVector4*)Constraints.GetData(),
-						RestAngles.GetData(),
-						IsBuckled.GetData(),
-						BucklingRatioWeighted.GetIndices().GetData(),
-						BucklingRatioWeighted.GetTable().GetData(),
-						Constraints.Num()
-					);
-				}
+				ispc::InitBendingConstraintsIsBuckledWithMaps(
+					(const ispc::FVector3f*)InParticles.XArray().GetData(),
+					(const ispc::FIntVector4*)Constraints.GetData(),
+					RestAngles.GetData(),
+					IsBuckled.GetData(),
+					BucklingRatioWeighted.GetIndices().GetData(),
+					BucklingRatioWeighted.GetTable().GetData(),
+					Constraints.Num()
+				);
 			}
 			else
 			{
-				if (bChaos_Bending_ISPCUseHWGather)
-				{
-					ispc::InitBendingConstraintsIsBuckledHWGather(
-						(const ispc::FVector3f*)InParticles.XArray().GetData(),
-						(const ispc::FIntVector4*)Constraints.GetData(),
-						RestAngles.GetData(),
-						IsBuckled.GetData(),
-						(FSolverReal)BucklingRatioWeighted,
-						Constraints.Num()
-					);
-				}
-				else
-				{
-					ispc::InitBendingConstraintsIsBuckled(
-						(const ispc::FVector3f*)InParticles.XArray().GetData(),
-						(const ispc::FIntVector4*)Constraints.GetData(),
-						RestAngles.GetData(),
-						IsBuckled.GetData(),
-						(FSolverReal)BucklingRatioWeighted,
-						Constraints.Num()
-					);
-				}
+				ispc::InitBendingConstraintsIsBuckled(
+					(const ispc::FVector3f*)InParticles.XArray().GetData(),
+					(const ispc::FIntVector4*)Constraints.GetData(),
+					RestAngles.GetData(),
+					IsBuckled.GetData(),
+					(FSolverReal)BucklingRatioWeighted,
+					Constraints.Num()
+				);
 			}
 		}
 		else
diff --git a/Engine/Source/Runtime/Experimental/Chaos/Private/Chaos/TriangleMesh.cpp b/Engine/Source/Runtime/Experimental/Chaos/Private/Chaos/TriangleMesh.cpp
--- a/Engine/Source/Runtime/Experimental/Chaos/Private/Chaos/TriangleMesh.cpp
+++ b/Engine/Source/Runtime/Experimental/Chaos/Private/Chaos/TriangleMesh.cpp
@@ -32,21 +32,12 @@ static_assert(sizeof(ispc::TArrayInt) == sizeof(TArray<int32>), "sizeof(ispc::TA
 #define USE_ISPC_KERNEL_CONSOLE_VARIABLES_IN_SHIPPING 0
 #endif
 
-#if !defined(USE_ISPC_UNOPTIMIZED_VECTORGATHER_DEFAULT)
-#define USE_ISPC_UNOPTIMIZED_VECTORGATHER_DEFAULT 0
-#endif
-
 // Support run-time toggling on supported platforms in non-shipping configurations
 #if !INTEL_ISPC || (UE_BUILD_SHIPPING && !USE_ISPC_KERNEL_CONSOLE_VARIABLES_IN_SHIPPING)
 static constexpr bool bChaos_TriangleMesh_ISPC_Enabled = INTEL_ISPC && CHAOS_TRIANGLE_MESH_ISPC_ENABLED_DEFAULT;
-static constexpr bool bChaos_TriangleMesh_ISPCUseHWGather = USE_ISPC_UNOPTIMIZED_VECTORGATHER_DEFAULT;
 #else
 static bool bChaos_TriangleMesh_ISPC_Enabled = CHAOS_TRIANGLE_MESH_ISPC_ENABLED_DEFAULT;
 static FAutoConsoleVariableRef CVarChaosTriangleMeshISPCEnabled(TEXT("p.Chaos.TriangleMesh.ISPC"), bChaos_TriangleMesh_ISPC_Enabled, TEXT("Whether to use ISPC optimizations in triangle mesh calculations"));
-
-static bool bChaos_TriangleMesh_ISPCUseHWGather = USE_ISPC_UNOPTIMIZED_VECTORGATHER_DEFAULT;
-static FAutoConsoleVariableRef CVarChaosTriangleMeshISPCUseHWGather(TEXT("p.Chaos.TriangleMesh.ISPCUseHWGather"), bChaos_TriangleMesh_ISPCUseHWGather, TEXT("Use (slower) HW gather when computing triangle mesh normals to see if this fixes the ISPC crash."));
-
 #endif
 
 namespace Chaos
@@ -438,23 +429,11 @@ void FTriangleMesh::GetFaceNormals(TArray<TVec3<T>>& Normals, const TConstArrayV
 		if (bChaos_TriangleMesh_ISPC_Enabled && std::is_same_v<T, FRealSingle>)
 		{
 			Normals.SetNumUninitialized(MElements.Num());
-			if (bChaos_TriangleMesh_ISPCUseHWGather)
-			{
-				ispc::GetFaceNormalsHWGather(
-					(ispc::FVector3f*)Normals.GetData(),
-					(ispc::FVector3f*)Points.GetData(),
-					(ispc::FIntVector*)MElements.GetData(),
-					MElements.Num());
-
-			}
-			else
-			{
-				ispc::GetFaceNormals(
-					(ispc::FVector3f*)Normals.GetData(),
-					(ispc::FVector3f*)Points.GetData(),
-					(ispc::FIntVector*)MElements.GetData(),
-					MElements.Num());
-			}
+			ispc::GetFaceNormals(
+				(ispc::FVector3f*)Normals.GetData(),
+				(ispc::FVector3f*)Points.GetData(),
+				(ispc::FIntVector*)MElements.GetData(),
+				MElements.Num());
 		}
 		else
 #endif
diff --git a/Engine/Source/Runtime/Experimental/Chaos/Private/Chaos/TriangleMesh.ispc b/Engine/Source/Runtime/Experimental/Chaos/Private/Chaos/TriangleMesh.ispc
--- a/Engine/Source/Runtime/Experimental/Chaos/Private/Chaos/TriangleMesh.ispc
+++ b/Engine/Source/Runtime/Experimental/Chaos/Private/Chaos/TriangleMesh.ispc
@@ -26,31 +26,6 @@ export void GetFaceNormals(uniform FVector3f Normals[],
 	}
 }
 
-export void GetFaceNormalsHWGather(uniform FVector3f Normals[],
-							const uniform FVector3f Points[],
-							const uniform FIntVector Elements[],
-							const uniform int32 NumElements)
-{
-	foreach(i = 0 ... NumElements)
-	{
-		const varying FIntVector Tri = VectorLoad(&Elements[extract(i, 0)]);
-		
-		#pragma ignore warning(perf)
-		const varying FVector3f p0 = Points[Tri.V[0]];
-		#pragma ignore warning(perf)
-		const varying FVector3f p1 = Points[Tri.V[1]];
-		#pragma ignore warning(perf)
-		const varying FVector3f p2 = Points[Tri.V[2]];
-
-		const varying FVector3f p10 = p1 - p0;
-		const varying FVector3f p20 = p2 - p0;
-		const varying FVector3f Cross = VectorCross(p20, p10);
-		const varying FVector3f CrossNormalized = VectorGetSafeNormal(Cross);
-
-		VectorStore(&Normals[extract(i, 0)], CrossNormalized);
-	}
-}
-
 struct TArrayInt
 {
 	int32 *Data;
diff --git a/Engine/Source/Runtime/Experimental/Chaos/Public/Chaos/PBDBendingConstraintsBase.h b/Engine/Source/Runtime/Experimental/Chaos/Public/Chaos/PBDBendingConstraintsBase.h
--- a/Engine/Source/Runtime/Experimental/Chaos/Public/Chaos/PBDBendingConstraintsBase.h
+++ b/Engine/Source/Runtime/Experimental/Chaos/Public/Chaos/PBDBendingConstraintsBase.h
@@ -408,15 +408,9 @@ protected:
 #define USE_ISPC_KERNEL_CONSOLE_VARIABLES_IN_SHIPPING 0
 #endif
 
-#if !defined(USE_ISPC_UNOPTIMIZED_VECTORGATHER_DEFAULT)
-#define USE_ISPC_UNOPTIMIZED_VECTORGATHER_DEFAULT 0
-#endif
-
 // Support run-time toggling on supported platforms in non-shipping configurations
 #if !INTEL_ISPC || (UE_BUILD_SHIPPING && !USE_ISPC_KERNEL_CONSOLE_VARIABLES_IN_SHIPPING)
 static constexpr bool bChaos_Bending_ISPC_Enabled = INTEL_ISPC && CHAOS_BENDING_ISPC_ENABLED_DEFAULT;
-static constexpr bool bChaos_Bending_ISPCUseHWGather = USE_ISPC_UNOPTIMIZED_VECTORGATHER_DEFAULT;
 #else
 extern CHAOS_API bool bChaos_Bending_ISPC_Enabled;
-extern CHAOS_API bool bChaos_Bending_ISPCUseHWGather;
 #endif
