Fix deprecated warning

diff --git a/Engine/Source/Runtime/Experimental/GeometryCollectionEngine/Public/GeometryCollection/GeometryCollectionISMPoolActor.h b/Engine/Source/Runtime/Experimental/GeometryCollectionEngine/Public/GeometryCollection/GeometryCollectionISMPoolActor.h
--- a/Engine/Source/Runtime/Experimental/GeometryCollectionEngine/Public/GeometryCollection/GeometryCollectionISMPoolActor.h
+++ b/Engine/Source/Runtime/Experimental/GeometryCollectionEngine/Public/GeometryCollection/GeometryCollectionISMPoolActor.h
@@ -21,7 +21,9 @@ private:
 
 public:
 	/** Returns ISMPoolComp subobject **/
+PRAGMA_DISABLE_DEPRECATION_WARNINGS
 	UGeometryCollectionISMPoolComponent* GetISMPoolComp() const { return ISMPoolComp; }
+PRAGMA_ENABLE_DEPRECATION_WARNINGS
 };
 
 
