Remove parallel for, from collision detection
Now Task graph is used instead
20% improvement in the entire collision detection including broad phase, middle phase and narrow phase in some scenario
Gain is high when a lot of collision occurs
Small lost in perfromance when there is no collision
Cvar p.Chaos.NumWorkerCollisionFactor is introduced and setup to 2 by default,
Gain can be really higher when set to 8, but decrease more perf when no collision occurs.

diff --git a/Engine/Source/Runtime/Experimental/Chaos/Private/Chaos/Collision/SpatialAccelerationBroadPhase.cpp b/Engine/Source/Runtime/Experimental/Chaos/Private/Chaos/Collision/SpatialAccelerationBroadPhase.cpp
--- a/Engine/Source/Runtime/Experimental/Chaos/Private/Chaos/Collision/SpatialAccelerationBroadPhase.cpp
+++ b/Engine/Source/Runtime/Experimental/Chaos/Private/Chaos/Collision/SpatialAccelerationBroadPhase.cpp
@@ -3,4 +3,33 @@
 
 namespace Chaos
 {
+	namespace CVars
+	{
+		int32 NumWorkerCollisionFactor = 2;
+		FAutoConsoleVariableRef CVarNumWorkerCollisionFactor(TEXT("p.Chaos.NumWorkerCollisionFactor"), NumWorkerCollisionFactor, TEXT("Set the number of tasks created for collision detection per worker."));
+	}
+	namespace Private
+	{
+
+		ENamedThreads::Type FAssignMidPhaseTask::GetDesiredThread()
+		{
+			return ENamedThreads::AnyHiPriThreadHiPriTask;
+		}
+		void FAssignMidPhaseTask::Process()
+		{
+			SpatialAccelerationBroadPhase.AssignMidPhases(BroadPhaseContext);
+		}
+
+		ENamedThreads::Type FProcessMidPhaseTask::GetDesiredThread()
+		{
+			return ENamedThreads::AnyHiPriThreadHiPriTask;
+		}
+		void FProcessMidPhaseTask::Process()
+		{
+			if (!BroadPhaseContext.Overlaps.IsEmpty())
+			{
+				SpatialAccelerationBroadPhase.ProcessMidPhases(Dt, BroadPhaseContext);
+			}
+		}
+	}
 }
\ No newline at end of file
diff --git a/Engine/Source/Runtime/Experimental/Chaos/Private/Chaos/PBDRigidsEvolutionGBF.cpp b/Engine/Source/Runtime/Experimental/Chaos/Private/Chaos/PBDRigidsEvolutionGBF.cpp
--- a/Engine/Source/Runtime/Experimental/Chaos/Private/Chaos/PBDRigidsEvolutionGBF.cpp
+++ b/Engine/Source/Runtime/Experimental/Chaos/Private/Chaos/PBDRigidsEvolutionGBF.cpp
@@ -195,7 +195,7 @@ namespace Chaos
 
 		// Enable improved midphase distribution among worker threads. Without this the midphase and narrowphase run on whatever thread the
 		// broadphas eoverlap was detected on which tends to give poor distribution.
-		bool bChaosMidPhaseRedistributionEnabled = true;
+		bool bChaosMidPhaseRedistributionEnabled = false;
 		FAutoConsoleVariableRef CVarChaosSolverMidPhaseRedistributionEnabled(TEXT("p.Chaos.BroadPhase.MidPhaseRedistributionEnabled"), bChaosMidPhaseRedistributionEnabled, TEXT(""));
 	}
 	using namespace CVars;
diff --git a/Engine/Source/Runtime/Experimental/Chaos/Public/Chaos/Collision/SpatialAccelerationBroadPhase.h b/Engine/Source/Runtime/Experimental/Chaos/Public/Chaos/Collision/SpatialAccelerationBroadPhase.h
--- a/Engine/Source/Runtime/Experimental/Chaos/Public/Chaos/Collision/SpatialAccelerationBroadPhase.h
+++ b/Engine/Source/Runtime/Experimental/Chaos/Public/Chaos/Collision/SpatialAccelerationBroadPhase.h
@@ -24,6 +24,7 @@ namespace Chaos::CVars
 {
 	extern bool bChaosMidPhaseRedistributionEnabled;
 	extern int32 ChaosOneWayInteractionPairCollisionMode;
+	extern int32 NumWorkerCollisionFactor;
 }
 
 namespace Chaos
@@ -293,6 +294,46 @@ namespace Chaos
 			/** Handle to be stored to retrieve the payload on which we are querying the acceleration structure*/
 			FAccelerationStructureHandle AccelerationHandle;
 		};
+
+		template<bool bOnlyRigid, typename THandle, typename LambdaHandle, typename SpatialAccelerationType>
+		class FProduceOverlapTask;
+
+		class FAssignMidPhaseTask
+		{
+		public:
+			FAssignMidPhaseTask(FSpatialAccelerationBroadPhase& InSpatialAccelerationBroadPhase, Private::FBroadPhaseContext& InBroadPhaseContext)
+				: SpatialAccelerationBroadPhase(InSpatialAccelerationBroadPhase)
+				, BroadPhaseContext(InBroadPhaseContext) {}
+
+			TStatId GetStatId() const { RETURN_QUICK_DECLARE_CYCLE_STAT(FAssignMidPhaseTask, STATGROUP_TaskGraphTasks); }
+			static CHAOS_API ENamedThreads::Type GetDesiredThread();
+			static ESubsequentsMode::Type GetSubsequentsMode() { return ESubsequentsMode::TrackSubsequents; }
+			void DoTask(ENamedThreads::Type CurrentThread, const FGraphEventRef& MyCompletionGraphEvent) { Process(); }
+
+		private:
+			void Process();
+			FSpatialAccelerationBroadPhase& SpatialAccelerationBroadPhase;
+			Private::FBroadPhaseContext& BroadPhaseContext;
+		};
+
+		class FProcessMidPhaseTask
+		{
+		public:
+			FProcessMidPhaseTask(FSpatialAccelerationBroadPhase& InSpatialAccelerationBroadPhase, Private::FBroadPhaseContext& InBroadPhaseContext, FReal InDt)
+				: SpatialAccelerationBroadPhase(InSpatialAccelerationBroadPhase)
+				, BroadPhaseContext(InBroadPhaseContext)
+				, Dt(InDt) {}
+
+			TStatId GetStatId() const { RETURN_QUICK_DECLARE_CYCLE_STAT(FProcessMidPhaseTask, STATGROUP_TaskGraphTasks); }
+			static CHAOS_API ENamedThreads::Type GetDesiredThread();
+			static ESubsequentsMode::Type GetSubsequentsMode() { return ESubsequentsMode::TrackSubsequents; }
+			void DoTask(ENamedThreads::Type CurrentThread, const FGraphEventRef& MyCompletionGraphEvent) { Process(); }
+		private:
+			void Process();
+			FSpatialAccelerationBroadPhase& SpatialAccelerationBroadPhase;
+			Private::FBroadPhaseContext& BroadPhaseContext;
+			FReal Dt;
+		};
 	}
 
 	/**
@@ -371,16 +412,18 @@ namespace Chaos
 				FChaosVDContextWrapper CVDContext;
 				CVD_GET_WRAPPED_CURRENT_CONTEXT(CVDContext);
 
-				// Find or assign a midphase to each overlapping particle pair
-				const auto& AssignMidPhasesWorker = [this, Dt, &CVDContext](const int32 ContextIndex)
+				FGraphEventArray NextPendingTasks;
+				const int32 NumContextTask = ProduceOverlapTasks.Num();
+				check(NumContextTask <= NumActiveBroadphaseContexts);
+				for (int32 ContextIndex = 0; ContextIndex < NumContextTask; ContextIndex++)
 				{
-					CVD_SCOPE_CONTEXT(CVDContext.Context);
-					AssignMidPhases(BroadphaseContexts[ContextIndex]);
-				};
-				PhysicsParallelFor(NumActiveBroadphaseContexts, AssignMidPhasesWorker, bDisableCollisionParallelFor);
-
-				// Merge all the midphases from each worker into the primary allocator
-				Allocator->ProcessNewMidPhases();
+					const FGraphEventArray Prereq = { ProduceOverlapTasks[ContextIndex] };
+					const FGraphEventRef PendingTask = TGraphTask<Private::FAssignMidPhaseTask>::CreateTask(&Prereq).ConstructAndDispatchWhenReady(*this, BroadphaseContexts[ContextIndex]);
+					NextPendingTasks.Add(PendingTask);
+				}
+				// It is required to wait for the middle phase modifiers, which is a callback to modify middle phases.
+				// This could be removed by adding dependent tasks for the middle phase modifiers.
+				FTaskGraphInterface::Get().WaitUntilTasksComplete(NextPendingTasks);
 			}
 
 			// Run some error checks in non-shipping builds
@@ -396,7 +439,7 @@ namespace Chaos
 		}
 
 		/**
-		 * Generate all the collision constraints for the set of overlapping objects produced byt he broad phase
+		 * Generate all the collision constraints for the set of overlapping objects produced by the broad phase
 		*/
 		void ProduceCollisions(FReal Dt)
 		{
@@ -412,13 +455,16 @@ namespace Chaos
 			// For that we would need to queue and redistribute the NarrowPhases, but currently each
 			// MidPhase runs the NarrowPhase in ProcessMidPhase
 			RedistributeMidPhasesInContexts();
-
-			const auto& ProcessMidPhasesWorker = [this, Dt, &CVDContext](const int32 ContextIndex)
+			const int32 NumContextTask = ProduceOverlapTasks.Num();
+			FGraphEventArray ProcessMidPhaseTasks;
+			ProcessMidPhaseTasks.Reserve(NumContextTask);
+			check(NumContextTask <= NumActiveBroadphaseContexts);
+			for (int32 ContextIndex = 0; ContextIndex < NumContextTask; ContextIndex++)
 			{
-				CVD_SCOPE_CONTEXT(CVDContext.Context);
-				ProcessMidPhases(Dt, BroadphaseContexts[ContextIndex]);
-			};
-			PhysicsParallelFor(NumActiveBroadphaseContexts, ProcessMidPhasesWorker, bDisableCollisionParallelFor);
+				const FGraphEventRef ProcessMidPhaseTask = TGraphTask<Private::FProcessMidPhaseTask>::CreateTask().ConstructAndDispatchWhenReady(*this, BroadphaseContexts[ContextIndex], Dt);
+				ProcessMidPhaseTasks.Add(ProcessMidPhaseTask);
+			}
+			FTaskGraphInterface::Get().WaitUntilTasksComplete(ProcessMidPhaseTasks);
 		}
 
 		/** @brief This function is the outer loop of collision detection. It loops over the
@@ -443,49 +489,52 @@ namespace Chaos
 			}
 
 			// The number of contexts that may have new overlaps in them
-			NumActiveBroadphaseContexts = 0;
-
-			// A set of contexts, one for each worker thread, containing the allocation buffers etc.
-			// NOTE: Ideally contexts arrays would not resize after the initial creation, but actually it might because
-			// of the way the ParticleView parallel-for is implemented. If the view is a set of sub-views, we
-			// will get a parallel-for per sub-view which will have different numbers of particles and may
-			// therefore require a different number of workers.
-			// Fortunately, we only get calls to ContextCreator for a sub-view (with a possibly different NumWorkers)
-			// after the previous sub-view's parallel-for has completed, so we don't need to worry about context addresses 
-			// changing (we could use TChunkedArray, or allocate contexts on the heap if we ever need to handle that).
-			const auto& ContextCreator = [this, Allocator, &Settings, &OverlapView](const int32 WorkerIndex, const int32 NumWorkers) -> int32
-			{
-				// Make sure we have enough contexts 
-				// (NOTE: ContextCreator gets called for every worker, but in serial and always with the same NumWorkers for any single ParallelFor)
-				if (NumWorkers > BroadphaseContexts.Num())
-				{
-					check(WorkerIndex == 0);	// Should be the same for all subsequent workers
-					BroadphaseContexts.SetNum(NumWorkers);
-				}
+			NumActiveBroadphaseContexts = !!GSingleThreadedPhysics ? 1 : FMath::Max(FMath::Min(Chaos::CVars::NumWorkerCollisionFactor*(FTaskGraphInterface::Get().GetNumWorkerThreads()-1), Chaos::MaxNumWorkers), 1);
+			NumActiveBroadphaseContexts = FMath::Max(1, NumActiveBroadphaseContexts);
 
-				// Remember the number of contexts we may have added overlaps to (if there are nested parallel-fors, we need to make 
-				// sure we keep the largest NumWorkers of all the parallel-fors)
-				NumActiveBroadphaseContexts = FMath::Max(NumWorkers, NumActiveBroadphaseContexts);
-				Allocator->SetMaxContexts(NumActiveBroadphaseContexts);
+			BroadphaseContexts.SetNum(NumActiveBroadphaseContexts);
+			Allocator->SetMaxContexts(NumActiveBroadphaseContexts);
 
-				// Retrieve the midphase/collision allocator for this context
-				Private::FCollisionContextAllocator* ContextAllocator = Allocator->GetContextAllocator(WorkerIndex);
+			for (int32 Index = 0; Index < NumActiveBroadphaseContexts; Index++)
+			{
+				Private::FCollisionContextAllocator* ContextAllocator = Allocator->GetContextAllocator(Index);
+				BroadphaseContexts[Index].CollisionContext.SetSettings(Settings);
+				BroadphaseContexts[Index].CollisionContext.SetAllocator(ContextAllocator);
+			}
 
-				// Build the collision context for the worker
-				BroadphaseContexts[WorkerIndex].CollisionContext.SetSettings(Settings);
-				BroadphaseContexts[WorkerIndex].CollisionContext.SetAllocator(ContextAllocator);
+			using TSOA = typename ViewType::TSOA;
+			using THandle = typename TSOA::THandleType;
+			using THandleBase = typename THandle::THandleBase;
+			ProduceOverlapTasks.Reset();
 
-				return WorkerIndex;
-			};
+			int32 ContextIndex = 0;
 
-			// Generate the set of AABB overlapping particles that are allowed to collide with each other
+			for (int32 ViewIndex = 0; ViewIndex < OverlapView.SOAViews.Num(); ++ViewIndex)
 			{
-				CSV_SCOPED_TIMING_STAT(PhysicsVerbose, DetectCollisions_AABBTree);
-				const auto& ProduceOverlapsWorker = [this, Dt, &InSpatialAcceleration](auto& Particle1, const int32 ActiveIdxIdx, const int32 ContextIndex)
+				const TSOAView<TSOA>& SOAView = OverlapView.SOAViews[ViewIndex];
+				const int32 NumParticles = SOAView.Size();
+				if (NumParticles == 0)
 				{
-					ProduceParticleOverlaps<bOnlyRigid>(Dt, Particle1.Handle(), InSpatialAcceleration, BroadphaseContexts[ContextIndex]);
-				};
-				OverlapView.ParallelFor(ContextCreator, ProduceOverlapsWorker, bDisableCollisionParallelFor);
+					continue;
+				}
+				// Dispatch tasks according if it is an array of handles or a SOA
+				if (const TArray<THandle*>* CurHandlesArray = SOAView.HandlesArray)
+				{
+					auto ReturnHandle = [CurHandlesArray](int32 Index) -> THandleBase
+					{ 
+						THandle* HandlePtr = (*CurHandlesArray)[Index];
+						return THandleBase(HandlePtr->GeometryParticles, HandlePtr->ParticleIdx);
+					};
+					using LambdaType = decltype(ReturnHandle);
+					const int32 NumHandles = CurHandlesArray->Num();
+					DispatchTasks<bOnlyRigid, SpatialAccelerationType, THandle, LambdaType>(ContextIndex, NumHandles, InSpatialAcceleration, Dt, ReturnHandle);
+				}
+				else if (TSOA* Soa = SOAView.SOA)
+				{
+					auto ReturnHandle = [Soa](int32 Index) { return THandleBase(Soa, Index); };
+					using LambdaType = decltype(ReturnHandle);
+					DispatchTasks<bOnlyRigid, SpatialAccelerationType, THandle, LambdaType>(ContextIndex, NumParticles, InSpatialAcceleration, Dt, ReturnHandle);
+				}
 			}
 		}
 
@@ -543,6 +592,41 @@ namespace Chaos
 
 	private:
 
+		template<bool bOnlyRigid, typename SpatialAccelerationType, typename THandle, typename LambdaType>
+		void DispatchTasks(int32& ContextIndex, int32 NumParticles, const SpatialAccelerationType& InSpatialAcceleration, FReal Dt, LambdaType ReturnHandle)
+		{
+			const int32 ParticleByContext = FMath::Max(FMath::DivideAndRoundUp(NumParticles, NumActiveBroadphaseContexts), SmallBatchSize);
+			const int32 NumBatch = FMath::DivideAndRoundUp(NumParticles, ParticleByContext);
+
+			for (int32 BatchIndex = 0; BatchIndex < NumBatch; ContextIndex++, BatchIndex++)
+			{
+				if (ContextIndex >= NumActiveBroadphaseContexts)
+				{
+					ContextIndex = 0;
+				}
+
+				int32 LastIndex = (BatchIndex + 1) * ParticleByContext;
+				LastIndex = BatchIndex == NumBatch - 1 ? FMath::Min(NumParticles, LastIndex) : LastIndex;
+
+				if (ContextIndex >= ProduceOverlapTasks.Num())
+				{
+					FGraphEventRef PendingTask = TGraphTask<Private::FProduceOverlapTask<bOnlyRigid, THandle, LambdaType, SpatialAccelerationType>>::CreateTask().
+						ConstructAndDispatchWhenReady(*this, BroadphaseContexts[ContextIndex],
+							InSpatialAcceleration, Dt, ReturnHandle, BatchIndex * ParticleByContext, LastIndex);
+					ProduceOverlapTasks.Add(PendingTask);
+				}
+				else
+				{
+					// If all context threads has been filled, create other tasks and link them to the previous one, 
+					FGraphEventArray Prereq = { ProduceOverlapTasks[ContextIndex] };
+					FGraphEventRef PendingTask = TGraphTask<Private::FProduceOverlapTask<bOnlyRigid, THandle, LambdaType, SpatialAccelerationType>>::CreateTask(&Prereq).
+						ConstructAndDispatchWhenReady(*this, BroadphaseContexts[ContextIndex],
+							InSpatialAcceleration, Dt, ReturnHandle, BatchIndex * ParticleByContext, LastIndex);
+					ProduceOverlapTasks[ContextIndex] = PendingTask;
+				}
+			}
+		}
+
 		// Generate the set of particles that overlap the specified particle and are allowed to collide with it
 		template<bool bOnlyRigid, typename T_SPATIALACCELERATION>
 		void ProduceParticleOverlaps(
@@ -609,7 +693,7 @@ namespace Chaos
 			}
 		}
 
-		// Redistribute the midphases among the contects to even out the per-core work a bit
+		// Redistribute the midphases among the contexts to even out the per-core work a bit
 		void RedistributeMidPhasesInContexts()
 		{
 			QUICK_SCOPE_CYCLE_COUNTER(STAT_Collisions_RedistributeMidPhases);
@@ -778,6 +862,61 @@ namespace Chaos
 
 		int32 NumBroadPhasePairs;
 		int32 NumMidPhases;
+
+		FGraphEventArray ProduceOverlapTasks;
+
+		template<bool bOnlyRigid, typename THandle, typename LambdaHandle, typename SpatialAccelerationType>
+		friend class Private::FProduceOverlapTask;
+		friend class Private::FAssignMidPhaseTask;
+		friend class Private::FProcessMidPhaseTask;
 	};
+	namespace Private
+	{
+		template<bool bOnlyRigid, typename THandle, typename LambdaHandle, typename SpatialAccelerationType>
+		class FProduceOverlapTask
+		{
+		public:
 
+			FProduceOverlapTask(FSpatialAccelerationBroadPhase& InSpatialAccelerationBroadPhase, Private::FBroadPhaseContext& InBroadPhaseContext,
+				const SpatialAccelerationType& InSpatialAcceleration, FReal InDt, LambdaHandle InReturnHandle, int32 InStartParticleIndex, int32 InEndParticleIndex)
+				: SpatialAccelerationBroadPhase(&InSpatialAccelerationBroadPhase)
+				, BroadPhaseContext(InBroadPhaseContext)
+				, SpatialAcceleration(InSpatialAcceleration)
+				, Dt(InDt)
+				, ReturnHandle(InReturnHandle)
+				, StartParticleIndex(InStartParticleIndex)
+				, EndParticleIndex(InEndParticleIndex)
+			{}
+
+			TStatId GetStatId() const { RETURN_QUICK_DECLARE_CYCLE_STAT(FProduceOverlapTask, STATGROUP_TaskGraphTasks); }
+			static CHAOS_API ENamedThreads::Type GetDesiredThread() { return ENamedThreads::AnyHiPriThreadHiPriTask; }
+			static ESubsequentsMode::Type GetSubsequentsMode() { return ESubsequentsMode::TrackSubsequents; }
+			void DoTask(ENamedThreads::Type CurrentThread, const FGraphEventRef& MyCompletionGraphEvent) { Process(); }
+
+		private:
+			using THandleBase = typename THandle::THandleBase;
+			using TTransientHandle = typename THandle::TTransientHandle;
+			void Process()
+			{
+				SCOPE_CYCLE_COUNTER(STAT_Collisions_ParticlePairBroadPhase);
+				for (int32 Index = StartParticleIndex; Index < EndParticleIndex; Index++)
+				{
+					THandleBase Handle = ReturnHandle(Index);
+					TTransientHandle& TransientHandle = static_cast<TTransientHandle&>(Handle);
+					if (!TransientHandle.LightWeightDisabled())
+					{
+						SpatialAccelerationBroadPhase->ProduceParticleOverlaps<bOnlyRigid>(Dt, TransientHandle.Handle(), SpatialAcceleration, BroadPhaseContext);
+					}
+				}
+			}
+
+			FSpatialAccelerationBroadPhase* SpatialAccelerationBroadPhase;
+			Private::FBroadPhaseContext& BroadPhaseContext;
+			const SpatialAccelerationType& SpatialAcceleration;
+			FReal Dt;
+			LambdaHandle ReturnHandle;
+			int32 StartParticleIndex;
+			int32 EndParticleIndex;
+		};
+	}
 }
