Remove package name from stringify output, we may introduce it again later but it will likely have to be provided by callers... alternatively the parser may be able to identify objects that are in the local buffer
#rnx

diff --git a/Engine/Source/Runtime/Experimental/JsonObjectGraph/Private/JsonStringifyImpl.cpp b/Engine/Source/Runtime/Experimental/JsonObjectGraph/Private/JsonStringifyImpl.cpp
--- a/Engine/Source/Runtime/Experimental/JsonObjectGraph/Private/JsonStringifyImpl.cpp
+++ b/Engine/Source/Runtime/Experimental/JsonObjectGraph/Private/JsonStringifyImpl.cpp
@@ -246,11 +246,6 @@ FJsonStringifyImpl::FJsonStringifyImpl(TConstArrayView<const UObject*> Roots, co
 	TArray<const UObject*> Exports;
 	GatherExports(Roots, {}, bFilterEditorOnly, RootObjects, Exports);
 
-	if(Exports.Num() > 0 )
-	{
-		PackageName = Exports[0]->GetPackage()->GetFName();
-	}
-
 	ObjectsToExport = TSet<const UObject*>(Exports);
 
 	MemoryWriter.SetIsPersistent(true);
@@ -848,14 +843,6 @@ void FJsonStringifyImpl::WritePackageSummary()
 	}
 #endif // WITH_TEXT_ARCHIVE_SUPPORT
 
-	// Write package name, this is useful for detecting file system level moves 
-	// and matches the behavior of FLinkerSave:
-	if(PackageName != FName())
-	{
-		Writer->WriteIdentifierPrefix(UE_JSON_PACKAGE_NAME_KEY_TCHAR);
-		Writer->WriteValueInline(PackageName.ToString());
-	}
-
 	// Note that there is no automatic compression here, if you add information to the package summary
 	// you should update ShouldWritePackageSummary
 	Writer->WriteObjectEnd();
@@ -870,9 +857,10 @@ bool FJsonStringifyImpl::ShouldWritePackageSummary() const
 {
 	return 
 #if WITH_TEXT_ARCHIVE_SUPPORT
-		Versions.Num() > 0 ||
+		Versions.Num() > 0;
+#else
+		false;
 #endif
-		PackageName != FName();
 }
 
 FJsonStringifyImpl::FPendingScope::FPendingScope(FJsonStringifyImpl* To, const TFunction<void()>& Prefix)
