Dataflow : add OnSubGraphsChanged delegate to listen to dataflow related edits
- also add an accessor for SubGraphs

diff --git a/Engine/Source/Runtime/Experimental/Dataflow/Engine/Private/Dataflow/DataflowObject.cpp b/Engine/Source/Runtime/Experimental/Dataflow/Engine/Private/Dataflow/DataflowObject.cpp
--- a/Engine/Source/Runtime/Experimental/Dataflow/Engine/Private/Dataflow/DataflowObject.cpp
+++ b/Engine/Source/Runtime/Experimental/Dataflow/Engine/Private/Dataflow/DataflowObject.cpp
@@ -15,6 +15,7 @@
 #define LOCTEXT_NAMESPACE "UDataflow"
 
 FDataflowAssetDelegates::FOnVariablesChanged FDataflowAssetDelegates::OnVariablesChanged;
+FDataflowAssetDelegates::FOnSubGraphsChanged FDataflowAssetDelegates::OnSubGraphsChanged;
 
 namespace UE::Dataflow::CVars
 {
@@ -298,6 +299,11 @@ void UDataflow::RemoveSubGraph(UDataflowSubGraph* SubGraph)
 	}
 }
 
+const TArray<TObjectPtr<UDataflowSubGraph>>& UDataflow::GetSubGraphs() const
+{
+	return DataflowSubGraphs;
+}
+
 #if WITH_EDITOR
 bool UDataflow::CanEditChange(const FProperty* InProperty) const
 {
diff --git a/Engine/Source/Runtime/Experimental/Dataflow/Engine/Public/Dataflow/DataflowObject.h b/Engine/Source/Runtime/Experimental/Dataflow/Engine/Public/Dataflow/DataflowObject.h
--- a/Engine/Source/Runtime/Experimental/Dataflow/Engine/Public/Dataflow/DataflowObject.h
+++ b/Engine/Source/Runtime/Experimental/Dataflow/Engine/Public/Dataflow/DataflowObject.h
@@ -30,9 +30,13 @@ namespace UE::Dataflow { class FGraph; }
 
 struct FDataflowAssetDelegates
 {
-	/** Callback for object property modifications, called by UObject::PostEditChangeProperty with a single property event */
+	/** Called when variables are edited ( add, remove, changetype , setvalue ) */
 	DECLARE_MULTICAST_DELEGATE_TwoParams(FOnVariablesChanged, const UDataflow*, FName);
 	static DATAFLOWENGINE_API FOnVariablesChanged OnVariablesChanged;
+
+	/** Called when sub-graphs are edited ( add, remove ) */
+	DECLARE_MULTICAST_DELEGATE_TwoParams(FOnSubGraphsChanged, const UDataflow*, FName);
+	static DATAFLOWENGINE_API FOnVariablesChanged OnSubGraphsChanged;
 };
 
 /**
@@ -181,6 +185,8 @@ public:
 	DATAFLOWENGINE_API void AddSubGraph(UDataflowSubGraph* SubGraph);
 	DATAFLOWENGINE_API void RemoveSubGraph(UDataflowSubGraph* SubGraph);
 
+	DATAFLOWENGINE_API const TArray<TObjectPtr<UDataflowSubGraph>>& GetSubGraphs() const;
+
 #if WITH_EDITORONLY_DATA
 
 	/*
