Assigned Saved Manifold optimization

diff --git a/Engine/Source/Runtime/Experimental/Chaos/Private/Chaos/Collision/PBDCollisionConstraint.cpp b/Engine/Source/Runtime/Experimental/Chaos/Private/Chaos/Collision/PBDCollisionConstraint.cpp
--- a/Engine/Source/Runtime/Experimental/Chaos/Private/Chaos/Collision/PBDCollisionConstraint.cpp
+++ b/Engine/Source/Runtime/Experimental/Chaos/Private/Chaos/Collision/PBDCollisionConstraint.cpp
@@ -74,6 +74,9 @@ namespace Chaos
 	bool bChaos_Collision_AllowGlobalInitialPhi = true;
 	FAutoConsoleVariableRef CVarChaos_Collision_AllowGlobalInitialPhi(TEXT("p.Chaos.Collision.AllowGlobalInitialPhi"), bChaos_Collision_AllowGlobalInitialPhi, TEXT(""));
 
+	bool bChaos_Collision_SimpleAssignContact = true;
+	FAutoConsoleVariableRef CVarChaos_Collision_SimpleAssignContacti(TEXT("p.Chaos.Collision.SimpleAssignContact"), bChaos_Collision_SimpleAssignContact, TEXT("Try to assign contact without worrying if the contact has been already assigned. Improve performance and ease to sleep when enabled."));
+
 
 	struct FCollisionTolerances
 	{
@@ -1140,6 +1143,48 @@ namespace Chaos
 		}
 	}
 
+	int32 FPBDCollisionConstraint::FindSimpleSavedManifoldPoint(const int32 ManifoldPointIndex) const
+	{
+		if (bChaos_Manifold_EnableFrictionRestore)
+		{
+			const FManifoldPoint& ManifoldPoint = ManifoldPoints[ManifoldPointIndex];
+			if (!ManifoldPoint.Flags.bDisabled)
+			{
+				const int32 NumAllowedSavedPoints = SavedManifoldPoints.Num();
+				const FRealSingle ExactDistanceToleranceSq = FMath::Square(Chaos_Manifold_FrictionExactPositionTolerance);
+				const FRealSingle NearDistanceToleranceSq = FMath::Square(Chaos_Manifold_FrictionNearPositionTolerance);
+
+				// Ignore points farther than NearDistanceTolerance
+				FRealSingle BestDistanceSq = NearDistanceToleranceSq;
+
+				int32 MatchAllowedPointIndex = INDEX_NONE;
+				for (int32 PointIndex = 0; PointIndex < NumAllowedSavedPoints; ++PointIndex)
+				{
+					const int32 MatchPointIndex = (PointIndex + ManifoldPointIndex) % NumAllowedSavedPoints;
+					const FSavedManifoldPoint& SavedManifoldPoint = SavedManifoldPoints[MatchPointIndex];
+					const FRealSingle DistanceSq = CalculateSavedManifoldPointDistanceSq(SavedManifoldPoint, ManifoldPoint);
+
+					// If this is an exact match, take the point without searching further
+					if (DistanceSq < ExactDistanceToleranceSq)
+					{
+						BestDistanceSq = DistanceSq;
+						MatchAllowedPointIndex = MatchPointIndex;
+						break;
+					}
+
+					// If this is a close match, we may use it if we can't find a better one
+					if (DistanceSq < BestDistanceSq)
+					{
+						BestDistanceSq = DistanceSq;
+						MatchAllowedPointIndex = MatchPointIndex;
+					}
+				}
+					return MatchAllowedPointIndex;
+			}
+		}
+		return INDEX_NONE;
+	}
+
 	int32 FPBDCollisionConstraint::FindSavedManifoldPoint(const int32 ManifoldPointIndex, int32* InOutAllowedSavedPointIndices, int32& InOutNumAllowedSavedPoints) const
 	{
 		if (bChaos_Manifold_EnableFrictionRestore)
@@ -1187,25 +1232,27 @@ namespace Chaos
 					{
 						InOutAllowedSavedPointIndices[MatchAllowedPointIndex] = InOutAllowedSavedPointIndices[InOutNumAllowedSavedPoints];
 					}
-
 					return MatchIndex;
 				}
 			}
 		}
-
 		return INDEX_NONE;
 	}
 
 	void FPBDCollisionConstraint::AssignSavedManifoldPoints()
 	{
 		TManifoldPointArray<int32> AllowedPointIndices;
-		AllowedPointIndices.SetNum(SavedManifoldPoints.Num());
-		for (int32 PointIndex = 0, PointEndIndex = SavedManifoldPoints.Num(); PointIndex < PointEndIndex; ++PointIndex)
+		int32 NumAllowedPointIndices = 0;
+		if (!bChaos_Collision_SimpleAssignContact)
 		{
-			AllowedPointIndices[PointIndex] = PointIndex;
+			AllowedPointIndices.SetNum(SavedManifoldPoints.Num());
+			for (int32 PointIndex = 0, PointEndIndex = SavedManifoldPoints.Num(); PointIndex < PointEndIndex; ++PointIndex)
+			{
+				AllowedPointIndices[PointIndex] = PointIndex;
+			}
+			NumAllowedPointIndices = AllowedPointIndices.Num();
 		}
 
-		int32 NumAllowedPointIndices = AllowedPointIndices.Num();
 		for (int32 PointIndex = 0, PointEndIndex = ManifoldPoints.Num(); PointIndex < PointEndIndex; ++PointIndex)
 		{
 			// @todo(chaos): ideally we would skip this step for restored manifold points but for now we need to 
@@ -1217,7 +1264,15 @@ namespace Chaos
 				ManifoldPoints[PointIndex].Flags.bInitialContact = true;
 
 				// See if we have saved data for this contact
-				const int32 SavedManifoldPointIndex = FindSavedManifoldPoint(PointIndex, AllowedPointIndices.GetData(), NumAllowedPointIndices);
+				int32 SavedManifoldPointIndex = INDEX_NONE;
+				if (bChaos_Collision_SimpleAssignContact)
+				{
+					SavedManifoldPointIndex = FindSimpleSavedManifoldPoint(PointIndex);
+				}
+				else
+				{
+					SavedManifoldPointIndex = FindSavedManifoldPoint(PointIndex, AllowedPointIndices.GetData(), NumAllowedPointIndices);
+				}
 				if (SavedManifoldPointIndex != INDEX_NONE)
 				{
 					ManifoldPoints[PointIndex].Flags.bHasStaticFrictionAnchor = true;
@@ -1226,7 +1281,7 @@ namespace Chaos
 					ManifoldPoints[PointIndex].InitialPhi = SavedManifoldPoints[SavedManifoldPointIndex].InitialPhi;
 					ManifoldPoints[PointIndex].Flags.bInitialContact = false;
 				}
-				// Nothing to do if no saved friction point because we already set the achor to the most recently detected contact point
+				// Nothing to do if no saved friction point because we already set the anchor to the most recently detected contact point
 				// (See InitManifoldPoint, And TryRestoreManifold)
 			}
 		}
diff --git a/Engine/Source/Runtime/Experimental/Chaos/Public/Chaos/Collision/PBDCollisionConstraint.h b/Engine/Source/Runtime/Experimental/Chaos/Public/Chaos/Collision/PBDCollisionConstraint.h
--- a/Engine/Source/Runtime/Experimental/Chaos/Public/Chaos/Collision/PBDCollisionConstraint.h
+++ b/Engine/Source/Runtime/Experimental/Chaos/Public/Chaos/Collision/PBDCollisionConstraint.h
@@ -884,6 +884,7 @@ namespace Chaos
 		CHAOS_API bool AreMatchingContactPoints(const FContactPoint& A, const FContactPoint& B, FReal& OutScore) const;
 		CHAOS_API int32 FindManifoldPoint(const FContactPoint& ContactPoint) const;
 
+		int32 FindSimpleSavedManifoldPoint(const int32 ManifoldPointIndex) const;
 		CHAOS_API int32 FindSavedManifoldPoint(const int32 ManifoldPointIndex, int32* InOutAllowedSavedPointIndices, int32& InOutNumAllowedSavedPoints) const;
 		CHAOS_API void AssignSavedManifoldPoints();
 
