Add the url of the current endpoint to 'ias.DisplayOnScreenStatistics'

#rnx

- FOnDemandIoBackendStats now takes and stores a reference to it's owning backend's FAvailableEndPoints.
--This means that we no longer need to provide one when calling ::StartAnalyticsRecording so that parameter has been removed.
- FAvailableEndPoints::GetCurrent now returns an empty string if there is no current endpoint rather than asserting/reading garbage memory.
- Fixed FBackendStatus::ToString to be const.

diff --git a/Engine/Source/Runtime/Experimental/IoStore/OnDemand/Private/AvailableEndPoints.h b/Engine/Source/Runtime/Experimental/IoStore/OnDemand/Private/AvailableEndPoints.h
--- a/Engine/Source/Runtime/Experimental/IoStore/OnDemand/Private/AvailableEndPoints.h
+++ b/Engine/Source/Runtime/Experimental/IoStore/OnDemand/Private/AvailableEndPoints.h
@@ -17,7 +17,15 @@ struct FAvailableEndPoints
 
 	const FString& GetCurrent() const
 	{
-		return Urls[Current];
+		if (HasCurrent())
+		{
+			return Urls[Current];
+		}
+		else
+		{
+			static FString EmptyString(TEXT("Disconnected"));
+			return EmptyString;
+		}
 	}
 
 	int32 Current = INDEX_NONE;
diff --git a/Engine/Source/Runtime/Experimental/IoStore/OnDemand/Private/OnDemandBackendStatus.cpp b/Engine/Source/Runtime/Experimental/IoStore/OnDemand/Private/OnDemandBackendStatus.cpp
--- a/Engine/Source/Runtime/Experimental/IoStore/OnDemand/Private/OnDemandBackendStatus.cpp
+++ b/Engine/Source/Runtime/Experimental/IoStore/OnDemand/Private/OnDemandBackendStatus.cpp
@@ -19,7 +19,7 @@ bool FBackendStatus::IsHttpEnabled(EIoChunkType ChunkType) const
 		((CurrentFlags & uint8(EFlags::HttpBulkOptionalDisabled)) == 0 && GIasHttpOptionalBulkDataEnabled));
 }
 
-void FBackendStatus::ToString(FStringBuilderBase& Builder)
+void FBackendStatus::ToString(FStringBuilderBase& Builder) const
 {
 	if (IsCacheEnabled())
 	{
diff --git a/Engine/Source/Runtime/Experimental/IoStore/OnDemand/Private/OnDemandBackendStatus.h b/Engine/Source/Runtime/Experimental/IoStore/OnDemand/Private/OnDemandBackendStatus.h
--- a/Engine/Source/Runtime/Experimental/IoStore/OnDemand/Private/OnDemandBackendStatus.h
+++ b/Engine/Source/Runtime/Experimental/IoStore/OnDemand/Private/OnDemandBackendStatus.h
@@ -87,7 +87,7 @@ public:
 		AddOrRemoveFlags(EFlags::AbandonCache, bAbandon, TEXT("Abandon cache"));
 	}
 
-	void ToString(FStringBuilderBase& Builder);
+	void ToString(FStringBuilderBase& Builder) const;
 
 private:
 
diff --git a/Engine/Source/Runtime/Experimental/IoStore/OnDemand/Private/OnDemandIoDispatcherBackend.cpp b/Engine/Source/Runtime/Experimental/IoStore/OnDemand/Private/OnDemandIoDispatcherBackend.cpp
--- a/Engine/Source/Runtime/Experimental/IoStore/OnDemand/Private/OnDemandIoDispatcherBackend.cpp
+++ b/Engine/Source/Runtime/Experimental/IoStore/OnDemand/Private/OnDemandIoDispatcherBackend.cpp
@@ -919,9 +919,9 @@ private:
 	FIoRequestQueue CompletedRequests;
 	FChunkRequestQueue HttpRequests;
 	TArray<FChunkRequest*> PendingCacheGets;
-	FBackendStatus BackendStatus; 
-	FOnDemandIoBackendStats Stats;
+	FBackendStatus BackendStatus;
 	FAvailableEndPoints AvailableEps;
+	FOnDemandIoBackendStats Stats;
 	FDistributedEndpointUrl DistributionUrl;
 	FEventRef DistributedEndpointEvent;
 
@@ -941,7 +941,7 @@ private:
 FOnDemandIoBackend::FOnDemandIoBackend(const FOnDemandEndpointConfig& Config, FOnDemandIoStore& InIoStore, TUniquePtr<IIasCache>&& InCache)
 	: IoStore(InIoStore)
 	, Cache(MoveTemp(InCache))
-	, Stats(BackendStatus)
+	, Stats(AvailableEps, BackendStatus)
 	, CurrentThreadPriority(EThreadPriority::TPri_Num)
 {
 	EndpointTestPath = Config.TocPath;
@@ -1484,7 +1484,7 @@ TUniquePtr<IAnalyticsRecording> FOnDemandIoBackend::StartAnalyticsRecording() co
 {
 	if (GIasReportAnalyticsEnabled)
 	{
-		return Stats.StartAnalyticsRecording(AvailableEps);
+		return Stats.StartAnalyticsRecording();
 	}
 
 	return TUniquePtr<IAnalyticsRecording>();
diff --git a/Engine/Source/Runtime/Experimental/IoStore/OnDemand/Private/Statistics.cpp b/Engine/Source/Runtime/Experimental/IoStore/OnDemand/Private/Statistics.cpp
--- a/Engine/Source/Runtime/Experimental/IoStore/OnDemand/Private/Statistics.cpp
+++ b/Engine/Source/Runtime/Experimental/IoStore/OnDemand/Private/Statistics.cpp
@@ -374,8 +374,9 @@ static FOnDemandIoBackendStats* GStatistics = nullptr;
 static FDelegateHandle GStatisticsEndFrameDelegateHandle;
 static FDelegateHandle GStatisticsOnScreenDelegateHandle;
 
-FOnDemandIoBackendStats::FOnDemandIoBackendStats(FBackendStatus& InStatus)
+FOnDemandIoBackendStats::FOnDemandIoBackendStats(const FAvailableEndPoints& InAvailableEndPoints, FBackendStatus& InStatus)
 	: BackendStatus(InStatus)
+	, AvailableEndPoints(InAvailableEndPoints)
 {
 	static constexpr float OneOver1024 = 1.0f / 1024.0f;
 
@@ -455,6 +456,16 @@ FOnDemandIoBackendStats::FOnDemandIoBackendStats(FBackendStatus& InStatus)
 					BackendStatusText << TEXT("IAS Backend Status: ");
 					BackendStatus.ToString(BackendStatusText);
 
+					FString CdnUrl = AvailableEndPoints.GetCurrent();
+					if (!CdnUrl.IsEmpty())
+					{
+						BackendStatusText << TEXT(" | ") << CdnUrl;
+					}
+					else
+					{
+						BackendStatusText << TEXT(" | Disconnected");
+					}
+
 					FCoreDelegates::EOnScreenMessageSeverity Verbosity = BackendStatus.IsHttpError() ?	FCoreDelegates::EOnScreenMessageSeverity::Error :
 																										FCoreDelegates::EOnScreenMessageSeverity::Info;
 
@@ -588,7 +599,7 @@ class FAnalyticsRecording : public IAnalyticsRecording
 {
 public:
 	FAnalyticsRecording() = delete;
-	FAnalyticsRecording(const FAvailableEndPoints& InAvailableEndPoints, FBackendStatus& InBackendStatus)
+	FAnalyticsRecording(const FAvailableEndPoints& InAvailableEndPoints, const FBackendStatus& InBackendStatus)
 		: AvailableEndPoints(InAvailableEndPoints)
 		, BackendStatus(InBackendStatus)
 	{
@@ -632,18 +643,21 @@ private:
 
 	void Report(TArray<FAnalyticsEventAttribute>& OutAnalyticsArray) const override
 	{
-		FString CdnUrl = AvailableEndPoints.GetCurrent();
-
-		// Strip the prefix from the url as some analytics systems may have trouble dealing with it
-		if (!CdnUrl.RemoveFromStart(TEXT("http://")))
+		if (AvailableEndPoints.HasCurrent())
 		{
-			CdnUrl.RemoveFromStart(TEXT("https://"));
+			FString CdnUrl = AvailableEndPoints.GetCurrent();
+
+			// Strip the prefix from the url as some analytics systems may have trouble dealing with it
+			if (!CdnUrl.RemoveFromStart(TEXT("http://")))
+			{
+				CdnUrl.RemoveFromStart(TEXT("https://"));
+			}
+
+			AppendAnalyticsEventAttributeArray(OutAnalyticsArray, TEXT("IasCdnUrl"), MoveTemp(CdnUrl));
 		}
 
-		AppendAnalyticsEventAttributeArray(OutAnalyticsArray
-			, TEXT("IasCdnUrl"), MoveTemp(CdnUrl)
-			, TEXT("IasHttpHasEverConnected"), GHttpConnectCount.Get() > 0 // Report if the system has ever actually managed to make a connection
-		);
+		// Report if the system has ever actually managed to make a connection
+		AppendAnalyticsEventAttributeArray(OutAnalyticsArray, TEXT("IasHttpHasEverConnected"), GHttpConnectCount.Get() > 0);
 
 		if (GIasReportHttpAnalyticsEnabled)
 		{
@@ -782,7 +796,7 @@ private:
 	bool bRecording = true;
 };
 
-TUniquePtr<IAnalyticsRecording> FOnDemandIoBackendStats::StartAnalyticsRecording(const FAvailableEndPoints& AvailableEndPoints) const
+TUniquePtr<IAnalyticsRecording> FOnDemandIoBackendStats::StartAnalyticsRecording() const
 {
 	TUniquePtr<FAnalyticsRecording> Recording = MakeUnique<FAnalyticsRecording>(AvailableEndPoints, BackendStatus);
 
diff --git a/Engine/Source/Runtime/Experimental/IoStore/OnDemand/Private/Statistics.h b/Engine/Source/Runtime/Experimental/IoStore/OnDemand/Private/Statistics.h
--- a/Engine/Source/Runtime/Experimental/IoStore/OnDemand/Private/Statistics.h
+++ b/Engine/Source/Runtime/Experimental/IoStore/OnDemand/Private/Statistics.h
@@ -34,7 +34,7 @@ class FOnDemandIoBackendStats
 public:
 	static FOnDemandIoBackendStats* Get() IAS_STATISTICS_IMPL(nullptr)
 
-	FOnDemandIoBackendStats(FBackendStatus& InStatus) IAS_STATISTICS_IMPL()
+	FOnDemandIoBackendStats(const FAvailableEndPoints& InAvailableEndPoints, FBackendStatus& InBackendStatus) IAS_STATISTICS_IMPL()
 	~FOnDemandIoBackendStats() IAS_STATISTICS_IMPL()
 
 	/** Report analytics not directly associated with a specific endpoint */
@@ -42,7 +42,7 @@ public:
 	/** Report analytics for the current endpoint */
 	void ReportEndPointAnalytics(TArray<FAnalyticsEventAttribute>& OutAnalyticsArray) const IAS_STATISTICS_IMPL()
 
-	TUniquePtr<IAnalyticsRecording> StartAnalyticsRecording(const FAvailableEndPoints& AvailableEndPoints) const IAS_STATISTICS_IMPL(TUniquePtr<IAnalyticsRecording>())
+	TUniquePtr<IAnalyticsRecording> StartAnalyticsRecording() const IAS_STATISTICS_IMPL(TUniquePtr<IAnalyticsRecording>())
 
 	void OnIoRequestEnqueue() IAS_STATISTICS_IMPL()
 	void OnIoRequestComplete(uint64 Size, uint64 DurationMs) IAS_STATISTICS_IMPL()
@@ -74,7 +74,8 @@ public:
 private:
 
 #if IAS_WITH_STATISTICS
-	FBackendStatus& BackendStatus;
+	const FBackendStatus& BackendStatus;
+	const FAvailableEndPoints& AvailableEndPoints;
 #endif //IAS_WITH_STATISTICS
 };
 
