Improve task perfromance when no work overlap is found

diff --git a/Engine/Source/Runtime/Experimental/Chaos/Private/Chaos/Collision/SpatialAccelerationBroadPhase.cpp b/Engine/Source/Runtime/Experimental/Chaos/Private/Chaos/Collision/SpatialAccelerationBroadPhase.cpp
--- a/Engine/Source/Runtime/Experimental/Chaos/Private/Chaos/Collision/SpatialAccelerationBroadPhase.cpp
+++ b/Engine/Source/Runtime/Experimental/Chaos/Private/Chaos/Collision/SpatialAccelerationBroadPhase.cpp
@@ -26,10 +26,7 @@ namespace Chaos
 		}
 		void FProcessMidPhaseTask::Process()
 		{
-			if (!BroadPhaseContext.Overlaps.IsEmpty())
-			{
-				SpatialAccelerationBroadPhase.ProcessMidPhases(Dt, BroadPhaseContext);
-			}
+			SpatialAccelerationBroadPhase.ProcessMidPhases(Dt, BroadPhaseContext);
 		}
 	}
 }
\ No newline at end of file
diff --git a/Engine/Source/Runtime/Experimental/Chaos/Public/Chaos/Collision/SpatialAccelerationBroadPhase.h b/Engine/Source/Runtime/Experimental/Chaos/Public/Chaos/Collision/SpatialAccelerationBroadPhase.h
--- a/Engine/Source/Runtime/Experimental/Chaos/Public/Chaos/Collision/SpatialAccelerationBroadPhase.h
+++ b/Engine/Source/Runtime/Experimental/Chaos/Public/Chaos/Collision/SpatialAccelerationBroadPhase.h
@@ -461,8 +461,12 @@ namespace Chaos
 			check(NumContextTask <= NumActiveBroadphaseContexts);
 			for (int32 ContextIndex = 0; ContextIndex < NumContextTask; ContextIndex++)
 			{
-				const FGraphEventRef ProcessMidPhaseTask = TGraphTask<Private::FProcessMidPhaseTask>::CreateTask().ConstructAndDispatchWhenReady(*this, BroadphaseContexts[ContextIndex], Dt);
-				ProcessMidPhaseTasks.Add(ProcessMidPhaseTask);
+				// Accessing the overlaps array is safe here, all overlaps have been finished to be computed at this point
+				if (!BroadphaseContexts[ContextIndex].Overlaps.IsEmpty())
+				{
+					const FGraphEventRef ProcessMidPhaseTask = TGraphTask<Private::FProcessMidPhaseTask>::CreateTask().ConstructAndDispatchWhenReady(*this, BroadphaseContexts[ContextIndex], Dt);
+					ProcessMidPhaseTasks.Add(ProcessMidPhaseTask);
+				}
 			}
 			FTaskGraphInterface::Get().WaitUntilTasksComplete(ProcessMidPhaseTasks);
 		}
