Chaos Debug Draw - fix Actor-centered debug draw

diff --git a/Engine/Source/Runtime/Experimental/ChaosSolverEngine/Private/Chaos/ChaosDebugDrawSubsystem.cpp b/Engine/Source/Runtime/Experimental/ChaosSolverEngine/Private/Chaos/ChaosDebugDrawSubsystem.cpp
--- a/Engine/Source/Runtime/Experimental/ChaosSolverEngine/Private/Chaos/ChaosDebugDrawSubsystem.cpp
+++ b/Engine/Source/Runtime/Experimental/ChaosSolverEngine/Private/Chaos/ChaosDebugDrawSubsystem.cpp
@@ -8,7 +8,10 @@
 #include "ChaosDebugDraw/ChaosDDTimeline.h"
 #include "DrawDebugHelpers.h"
 #include "Engine/Engine.h"
+#include "Engine/HitResult.h"
 #include "Engine/World.h"
+#include "GameFramework/Pawn.h"
+#include "GameFramework/PlayerController.h"
 #include "HAL/CriticalSection.h"
 #include "Physics/Experimental/PhysScene_Chaos.h"
 #include "VisualLogger/VisualLogger.h"
@@ -23,6 +26,10 @@ extern bool bChaosDebugDraw_Enabled;
 extern float ChaosDebugDraw_Radius;
 extern int32 bChaosDebugDraw_DrawMode;
 extern int32 ChaosDebugDraw_MaxElements;
+extern bool bChaosDebugDraw_SingleActor;
+extern float ChaosDebugDraw_SingleActorTraceLength;
+extern float ChaosDebugDraw_SingleActorMaxRadius;
+
 extern float CommandLifeTime(const float LifeTime);
 extern void DebugDrawChaosCommand(const UWorld* World, const Chaos::FLatentDrawCommand& Command);
 extern void VisLogChaosCommand(const AActor* Actor, const Chaos::FLatentDrawCommand& Command);
@@ -370,17 +377,86 @@ void UChaosDebugDrawSubsystem::OnWorldTickEnd(ELevelTick TickType, float Dt)
 	{
 		CDDWorldTimelineContext.EndFrame();
 
+		if (CDDScene.IsValid())
+		{
+			CDDScene->SetCommandBudget(ChaosDebugDraw_MaxElements);
+		}
+
+		UpdateDrawRegion();
+
+		RenderScene();
+	}
+}
+
+void UChaosDebugDrawSubsystem::UpdateDrawRegion()
+{
+	if (UWorld* World = GetWorld())
+	{
+		CDDWorldTimelineContext.EndFrame();
+
 		if (CDDScene.IsValid())
 		{
 			FSphere3d DrawRegion = CDDScene->GetDrawRegion();
-			DrawRegion.W = ChaosDebugDraw_Radius;
-			if (World->ViewLocationsRenderedLastFrame.Num() > 0)
+
+			// By default we center the draw region on wherever the world was rendered from,
+			// or optionally center the draw region around where the player is looking.
+			if (!bChaosDebugDraw_SingleActor)
+			{
+				DrawRegion.W = ChaosDebugDraw_Radius;
+				if (World->ViewLocationsRenderedLastFrame.Num() > 0)
+				{
+					DrawRegion.Center = World->ViewLocationsRenderedLastFrame[0];
+				}
+			}
+			else
 			{
-				DrawRegion.Center = World->ViewLocationsRenderedLastFrame[0];
+				FVector CamLoc = FVector::ZeroVector;
+				FVector CamLook = FVector::XAxisVector;
+				AActor* PlayerPawnActor = nullptr;
+				bool bHaveCamera = false;
+
+				// If we have a player, use their camera for the raycast as a fallback, but this does
+				// not work in PIE Simulate mode
+				if (const APlayerController* Controller = GEngine->GetFirstLocalPlayerController(World))
+				{
+					FRotator CamRot;
+					Controller->GetPlayerViewPoint(CamLoc, CamRot);
+					CamLook = CamRot.Vector();
+					PlayerPawnActor = Controller->GetPawn<AActor>();
+					bHaveCamera = true;
+				}
+
+				// Use the last rendered transform for the raycast, which works in PIE Simulate mode
+				if (World->CachedViewInfoRenderedLastFrame.Num() > 0)
+				{
+					// NOTE: ViewToWorld has Look=Z-Axis
+					FMatrix ViewToWorld = World->CachedViewInfoRenderedLastFrame[0].ViewToWorld;
+					FMatrix ViewRot = FRotationMatrix::MakeFromX(ViewToWorld.GetUnitAxis(EAxis::Z));
+					CamLoc = ViewToWorld.GetOrigin();
+					CamLook = ViewToWorld.GetUnitAxis(EAxis::Z);
+					bHaveCamera = true;
+				}
+
+				if (bHaveCamera)
+				{
+					FVector TraceStart = CamLoc;
+					FVector TraceEnd = CamLoc + CamLook * ChaosDebugDraw_SingleActorTraceLength;
+
+					FHitResult HitResult(ForceInit);
+					FCollisionQueryParams TraceParams(SCENE_QUERY_STAT(ChaosDebugVisibilityTrace), true, PlayerPawnActor);
+					bool bHit = World->LineTraceSingleByChannel(HitResult, TraceStart, TraceEnd, ECC_Visibility, TraceParams);
+					if (bHit && HitResult.GetActor() != nullptr)
+					{
+						FVector Origin, BoxExtent;
+						HitResult.GetActor()->GetActorBounds(true, Origin, BoxExtent);
+						const float Radius = FMath::Min(BoxExtent.Size(), ChaosDebugDraw_SingleActorMaxRadius);
+						DrawRegion.Center = Origin;
+						DrawRegion.W = Radius;
+					}
+				}
 			}
 
 			CDDScene->SetDrawRegion(DrawRegion);
-			CDDScene->SetCommandBudget(ChaosDebugDraw_MaxElements);
 
 			// @todo(chaos): This is a problem in multi-client PIE - the draw region on the server can flipflop between
 			// the client positions depending on which one renders last. Could support multiple draw regions maybe...
@@ -391,8 +467,6 @@ void UChaosDebugDrawSubsystem::OnWorldTickEnd(ELevelTick TickType, float Dt)
 				ServerScene->SetCommandBudget(ChaosDebugDraw_MaxElements);
 			}
 		}
-
-		RenderScene();
 	}
 }
 
diff --git a/Engine/Source/Runtime/Experimental/ChaosSolverEngine/Public/Chaos/ChaosDebugDrawSubsystem.h b/Engine/Source/Runtime/Experimental/ChaosSolverEngine/Public/Chaos/ChaosDebugDrawSubsystem.h
--- a/Engine/Source/Runtime/Experimental/ChaosSolverEngine/Public/Chaos/ChaosDebugDrawSubsystem.h
+++ b/Engine/Source/Runtime/Experimental/ChaosSolverEngine/Public/Chaos/ChaosDebugDrawSubsystem.h
@@ -36,6 +36,7 @@ protected:
 	void RenderFrame(FChaosDDRenderer& Renderer, const ChaosDD::Private::FChaosDDFramePtr& Frame);
 	void OnWorldTickStart(ELevelTick TickType, float Dt);
 	void OnWorldTickEnd(ELevelTick TickType, float Dt);
+	void UpdateDrawRegion();
 	void RenderScene();
 	static void StaticOnWorldTickStart(UWorld* World, ELevelTick TickType, float Dt);
 	static void StaticOnWorldTickEnd(UWorld* World, ELevelTick TickType, float Dt);
