Geometry collection : fix navigation mesh update not transforming properly the vertices
- Made sure to use the TransformIndex to access the GeomToComponent array
- Also use the ComponentSpaceTransforms to take advantage of the cached transforms instead of recomputing them every time

diff --git a/Engine/Source/Runtime/Experimental/GeometryCollectionEngine/Private/GeometryCollection/GeometryCollectionComponent.cpp b/Engine/Source/Runtime/Experimental/GeometryCollectionEngine/Private/GeometryCollection/GeometryCollectionComponent.cpp
--- a/Engine/Source/Runtime/Experimental/GeometryCollectionEngine/Private/GeometryCollection/GeometryCollectionComponent.cpp
+++ b/Engine/Source/Runtime/Experimental/GeometryCollectionEngine/Private/GeometryCollection/GeometryCollectionComponent.cpp
@@ -1527,8 +1527,7 @@ bool UGeometryCollectionComponent::DoCustomNavigableGeometryExport(FNavigableGeo
 	}
 
 	// Get all the geometry transforms in component space (they are stored natively in parent-bone space)
-	TArray<FTransform3f> GeomToComponent;
-	ComputeCurrentGlobalsMatrices(GeomToComponent);
+	const TArray<FTransform3f>& GeomToComponent = ComponentSpaceTransforms.RequestAllTransforms();
 
 	OutVertexBuffer.AddUninitialized(VertexCount);
 
@@ -1543,11 +1542,11 @@ bool UGeometryCollectionComponent::DoCustomNavigableGeometryExport(FNavigableGeo
 		int32 SourceGeometryVertexStart = VertexStartArray[GeometryIndex];
 		int32 SourceGeometryVertexCount = VertexCountArray[GeometryIndex];
 
-		ParallelFor(SourceGeometryVertexCount, [&OutVertexBuffer, &GeomToComponent, &Vertex, SourceGeometryVertexStart, SubsetIndex, DestVertex](int32 PointIdx)
+		ParallelFor(SourceGeometryVertexCount, [&OutVertexBuffer, &GeomToComponent, &Vertex, SourceGeometryVertexStart, TransformIndex, DestVertex](int32 PointIdx)
 			{
 				//extract vertex from source
 				int32 SourceGeometryVertexIndex = SourceGeometryVertexStart + PointIdx;
-				const FVector3f VertexInComponentSpace = GeomToComponent[SubsetIndex].TransformPosition(Vertex[SourceGeometryVertexIndex]);
+				const FVector3f VertexInComponentSpace = GeomToComponent[TransformIndex].TransformPosition(Vertex[SourceGeometryVertexIndex]);
 				int32 DestVertexIndex = DestVertex + PointIdx;
 				OutVertexBuffer[DestVertexIndex].X = VertexInComponentSpace.X;
 				OutVertexBuffer[DestVertexIndex].Y = VertexInComponentSpace.Y;
