Removed unused var, this should have been part of 38937731

#rnx

diff --git a/Engine/Source/Runtime/Experimental/JsonObjectGraph/Private/JsonStringifyImpl.h b/Engine/Source/Runtime/Experimental/JsonObjectGraph/Private/JsonStringifyImpl.h
--- a/Engine/Source/Runtime/Experimental/JsonObjectGraph/Private/JsonStringifyImpl.h
+++ b/Engine/Source/Runtime/Experimental/JsonObjectGraph/Private/JsonStringifyImpl.h
@@ -106,7 +106,6 @@ private:
 	TArray<const UObject*> RootObjects;
 	TSet<const UObject*> ObjectsToExport;
 	TSet<const UObject*> ObjectsExported;
-	FName PackageName;
 
 	TArray<FCustomVersion> Versions;
 };
