[NetPhysics] Fix post resim error getting skipped when GT ticks slower than PT

diff --git a/Engine/Source/Programs/HeadlessChaos/Private/HeadlessChaosTestRewind.cpp b/Engine/Source/Programs/HeadlessChaos/Private/HeadlessChaosTestRewind.cpp
--- a/Engine/Source/Programs/HeadlessChaos/Private/HeadlessChaosTestRewind.cpp
+++ b/Engine/Source/Programs/HeadlessChaos/Private/HeadlessChaosTestRewind.cpp
@@ -2043,8 +2043,10 @@ namespace ChaosTest {
 			});
 	}
 
-	GTEST_TEST(AllTraits, RewindTest_SimCallbackCorrectionInterpolation)
+	GTEST_TEST(AllTraits, DISABLED_RewindTest_SimCallbackCorrectionInterpolation)
 	{
+		// NOTE: Disabled because the logic that is being tested is no longer used.
+
 		/*want to interpolate between original sim and correction
 		edge cases to test:
 		- gt write should stomp interpolation
diff --git a/Engine/Source/Runtime/Experimental/Chaos/Private/Chaos/Framework/ChaosResultsManager.cpp b/Engine/Source/Runtime/Experimental/Chaos/Private/Chaos/Framework/ChaosResultsManager.cpp
--- a/Engine/Source/Runtime/Experimental/Chaos/Private/Chaos/Framework/ChaosResultsManager.cpp
+++ b/Engine/Source/Runtime/Experimental/Chaos/Private/Chaos/Framework/ChaosResultsManager.cpp
@@ -192,6 +192,12 @@ namespace Chaos
 			SetPrevNextDataHelper<ESetPrevNextDataMode::Prev>(*Results.Prev);
 
 			Results.Next = PotentialNext;
+
+			// Move any non applied DirtyRigidErrors from previous to next
+			if (Results.Prev->DirtyRigidErrors.Num() > 0)
+			{
+				Results.Next->DirtyRigidErrors.Append(Results.Prev->DirtyRigidErrors);
+			}
 			
 			if(PotentialNext->ExternalEndTime <= LatestTimeSeen)
 			{
