Replace duplicate code with a dedicated function in preparation for larger changes.

#rnx

- There were a number of places using the exact same logic for determining the correct task priority to use when calling FOnDemandIoBackend::CompleteRequest on the task system. This code has been replaced with a dedicated function 'GetRequestCompletionTaskPriority' so that the logic is easier to update in the future.

diff --git a/Engine/Source/Runtime/Experimental/IoStore/OnDemand/Private/OnDemandIoDispatcherBackend.cpp b/Engine/Source/Runtime/Experimental/IoStore/OnDemand/Private/OnDemandIoDispatcherBackend.cpp
--- a/Engine/Source/Runtime/Experimental/IoStore/OnDemand/Private/OnDemandIoDispatcherBackend.cpp
+++ b/Engine/Source/Runtime/Experimental/IoStore/OnDemand/Private/OnDemandIoDispatcherBackend.cpp
@@ -252,6 +252,12 @@ FAutoConsoleVariableRef CVarOnDemandBackendThreadPriority(
 	TEXT("Note that this is switchable at runtime"),
 	ECVF_Default);
 
+///////////////////////////////////////////////////////////////////////////////
+[[nodiscard]] static UE::Tasks::ETaskPriority GetRequestCompletionTaskPriority()
+{
+	return GCompleteMaterializeTaskPriorities[FMath::Clamp(GCompleteMaterializeTaskPriority, 0, UE_ARRAY_COUNT(GCompleteMaterializeTaskPriorities) - 1)];
+}
+
 ///////////////////////////////////////////////////////////////////////////////
 static bool LatencyTest(FStringView Url, FStringView Path)
 {
@@ -1237,10 +1243,9 @@ bool FOnDemandIoBackend::Resolve(FIoRequestImpl* Request)
 			check(Buffer.GetData() != nullptr);
 			ChunkRequest->bCached = true;
 
-			UE::Tasks::ETaskPriority TaskPriority = GCompleteMaterializeTaskPriorities[FMath::Clamp(GCompleteMaterializeTaskPriority, 0, UE_ARRAY_COUNT(GCompleteMaterializeTaskPriorities) - 1)];
 			Launch(UE_SOURCE_LOCATION, [this, ChunkRequest] {
 				CompleteRequest(ChunkRequest);
-			}, TaskPriority);
+			}, GetRequestCompletionTaskPriority());
 			return true;
 		}
 
@@ -1249,12 +1254,11 @@ bool FOnDemandIoBackend::Resolve(FIoRequestImpl* Request)
 			InflightCacheRequestCount.fetch_add(1, std::memory_order_relaxed);
 
 			FTaskEvent OnReadyEvent(TEXT("IasCacheMaterializeDone"));
-			UE::Tasks::ETaskPriority TaskPriority = GCompleteMaterializeTaskPriorities[FMath::Clamp(GCompleteMaterializeTaskPriority, 0, UE_ARRAY_COUNT(GCompleteMaterializeTaskPriorities)-1)];
 
 			Launch(UE_SOURCE_LOCATION, [this, ChunkRequest] {
 				InflightCacheRequestCount.fetch_sub(1, std::memory_order_relaxed);
 				CompleteMaterialize(ChunkRequest);
-			}, OnReadyEvent, TaskPriority);
+			}, OnReadyEvent, GetRequestCompletionTaskPriority());
 
 			EIoErrorCode& OutStatus = ChunkRequest->CacheGetStatus;
 			Cache->Materialize(Key, Buffer, OutStatus, MoveTemp(OnReadyEvent));
@@ -1586,11 +1590,10 @@ void FOnDemandIoBackend::ProcessHttpRequests(FHttpClient& HttpClient, FBitWindow
 							}
 							}
 
-							UE::Tasks::ETaskPriority TaskPriority = GCompleteMaterializeTaskPriorities[FMath::Clamp(GCompleteMaterializeTaskPriority, 0, UE_ARRAY_COUNT(GCompleteMaterializeTaskPriorities) - 1)];
 							UE::Tasks::Launch(UE_SOURCE_LOCATION, [this, ChunkRequest]()
 							{
 								CompleteRequest(ChunkRequest);
-							}, TaskPriority);
+							}, GetRequestCompletionTaskPriority());
 						});
 				}
 			}
