Iris
* Fix cull distance overrides on actors whose location did not change.
* Previously the GridFilter would only use the overrides on BeginReplication and ignore changes done after that.
* Removed NetCullDistanceOverrides manager
* NetCull override values are now stored in the WorldLocations manager
* Unified all cull distance usage to use the single Getter in WorldLocations
* Added LogIrisNetCull log category to isolate cull distance changes from other Iris logs

diff --git a/Engine/Plugins/Runtime/ReplicationSystemTestPlugin/Source/Private/Tests/ReplicationSystem/Filtering/TestGridFilter.cpp b/Engine/Plugins/Runtime/ReplicationSystemTestPlugin/Source/Private/Tests/ReplicationSystem/Filtering/TestGridFilter.cpp
--- a/Engine/Plugins/Runtime/ReplicationSystemTestPlugin/Source/Private/Tests/ReplicationSystem/Filtering/TestGridFilter.cpp
+++ b/Engine/Plugins/Runtime/ReplicationSystemTestPlugin/Source/Private/Tests/ReplicationSystem/Filtering/TestGridFilter.cpp
@@ -53,6 +53,21 @@ private:
 		}
 
 		DefinitionsProperty->CopyCompleteValue((void*)(UPTRINT(FilterDefinitions) + DefinitionsProperty->GetOffset_ForInternal()), &NewFilterDefinitions);
+
+		// Setup custom filter configs via the CDO. Ugly! We need to support non-CDO engine configs in Iris
+		UNetObjectGridFilterConfig* GridFilterConfig = GetMutableDefault<UNetObjectGridFilterConfig>();
+		
+		// Save original CDO state
+		OriginalFilterProfiles = GridFilterConfig->FilterProfiles;
+
+		// Add unique profiles
+		GridFilterConfig->FilterProfiles.Add(
+			// This profile disables the culling histerisis so that filtering out an object has impact on the current tick.
+			FNetObjectGridFilterProfile
+			{ 
+				.FilterProfileName = TEXT("CullDistanceTest"),
+				.FrameCountBeforeCulling = 1,
+			});
 	}
 
 	void RestoreFilterDefinitions()
@@ -66,6 +81,9 @@ private:
 
 		WorldLocFilterHandle = InvalidNetObjectFilterHandle;
 		WorldLocFilter = nullptr;
+
+		UNetObjectGridFilterConfig* GridFilterConfig = GetMutableDefault<UNetObjectGridFilterConfig>();
+		GridFilterConfig->FilterProfiles = OriginalFilterProfiles;
 	}
 
 	void InitFilterHandles()
@@ -79,9 +97,9 @@ protected:
 	FNetObjectFilterHandle WorldLocFilterHandle;
 private:
 	TArray<FNetObjectFilterDefinition> OriginalFilterDefinitions;
+	TArray<FNetObjectGridFilterProfile> OriginalFilterProfiles;
 };
 
-// Test world loc filter
 UE_NET_TEST_FIXTURE(FTestGridFilterFixture, TestWorldLocGridFilter)
 {
 	// Spawn object with WorldLocation's on server
@@ -102,29 +120,32 @@ UE_NET_TEST_FIXTURE(FTestGridFilterFixture, TestWorldLocGridFilter)
 		FVector Loc;
 		float CullDistance;
 	};
-	TMap<const UObject*, FObjectLoc> ObjectLocs;
+	TMap<UReplicatedTestObject*, FObjectLoc> ObjectLocs;
 
 	// Visible objects
-	ObjectLocs.Add(ServerObjectZero,	FObjectLoc{ FVector::ZeroVector, 1500.f} );
-	ObjectLocs.Add(ServerObjectNear,	FObjectLoc{ FVector(100.f, 100.f, 100.f), 1500.f} );
-	ObjectLocs.Add(ServerObjectLimit,	FObjectLoc{ FVector(DefaultGridConfig->CellSizeX, DefaultGridConfig->CellSizeY, 1500.f), DefaultGridConfig->MaxCullDistance});
+	ObjectLocs.Add(ServerObjectZero,	FObjectLoc{ .Loc=FVector::ZeroVector, .CullDistance=1500.f} );
+	ObjectLocs.Add(ServerObjectNear,	FObjectLoc{ .Loc=FVector(100.f, 100.f, 100.f), .CullDistance=1500.f} );
+
+	// Make an object that sits right in the limit of a grid cell
+	const FVector CellLimitPos(DefaultGridConfig->CellSizeX, DefaultGridConfig->CellSizeY, 0.0f);
+	ObjectLocs.Add(ServerObjectLimit,	FObjectLoc{ .Loc=CellLimitPos, .CullDistance=(float)CellLimitPos.Size()});
 
 	// Culled objects
-	ObjectLocs.Add(ServerObjectCulled,	FObjectLoc{ FVector(DefaultGridConfig->CellSizeX + 100.f, DefaultGridConfig->CellSizeY + 100.f, 100.f), 1500.f });
-	ObjectLocs.Add(ServerObjectVeryFar, FObjectLoc{ FVector(DefaultGridConfig->CellSizeX + 99999.f, DefaultGridConfig->CellSizeY + 99999.f, 99999.f), 1500.f });
+	ObjectLocs.Add(ServerObjectCulled,	FObjectLoc{ .Loc=FVector(DefaultGridConfig->CellSizeX + 100.f, DefaultGridConfig->CellSizeY + 100.f, 100.f), .CullDistance=1500.f });
+	ObjectLocs.Add(ServerObjectVeryFar, FObjectLoc{ .Loc=FVector(DefaultGridConfig->CellSizeX + 99999.f, DefaultGridConfig->CellSizeY + 99999.f, 99999.f), .CullDistance=1500.f });
 
 	Server->GetReplicationBridge()->SetExternalWorldLocationUpdateFunctor([&](FNetRefHandle NetHandle, const UObject* ReplicatedObject, FVector& OutLocation, float& OutCullDistance)
 	{
-		const FObjectLoc& ObjectLoc = ObjectLocs[ReplicatedObject];
+		const FObjectLoc& ObjectLoc = ObjectLocs[Cast<const UReplicatedTestObject>(ReplicatedObject)];
 		OutLocation = ObjectLoc.Loc;
 		OutCullDistance = ObjectLoc.CullDistance;
 	});
 
 	// Apply grid filter
-	Server->ReplicationSystem->SetFilter(ServerObjectZero->NetRefHandle, WorldLocFilterHandle);
-	Server->ReplicationSystem->SetFilter(ServerObjectLimit->NetRefHandle, WorldLocFilterHandle);
-	Server->ReplicationSystem->SetFilter(ServerObjectCulled->NetRefHandle, WorldLocFilterHandle);
-	Server->ReplicationSystem->SetFilter(ServerObjectVeryFar->NetRefHandle, WorldLocFilterHandle);
+	for (auto It = ObjectLocs.CreateConstIterator(); It; ++It)
+	{
+		Server->ReplicationSystem->SetFilter(It.Key()->NetRefHandle, WorldLocFilterHandle);
+	}
 
 	// Add client
 	FReplicationSystemTestClient* Client = CreateClient();
@@ -143,19 +164,18 @@ UE_NET_TEST_FIXTURE(FTestGridFilterFixture, TestWorldLocGridFilter)
 	// Test visible objects
 	UE_NET_ASSERT_NE(Client->GetReplicationBridge()->GetReplicatedObject(ServerObjectZero->NetRefHandle), nullptr);
 	UE_NET_ASSERT_NE(Client->GetReplicationBridge()->GetReplicatedObject(ServerObjectNear->NetRefHandle), nullptr);
+	UE_NET_ASSERT_NE(Client->GetReplicationBridge()->GetReplicatedObject(ServerObjectLimit->NetRefHandle), nullptr);
 
 	// Test culled objects
-	UE_NET_ASSERT_EQ(Client->GetReplicationBridge()->GetReplicatedObject(ServerObjectLimit->NetRefHandle), nullptr);
 	UE_NET_ASSERT_EQ(Client->GetReplicationBridge()->GetReplicatedObject(ServerObjectCulled->NetRefHandle), nullptr);
 	UE_NET_ASSERT_EQ(Client->GetReplicationBridge()->GetReplicatedObject(ServerObjectVeryFar->NetRefHandle), nullptr);
 
-	Server->DestroyObject(ServerObjectZero);
-	Server->DestroyObject(ServerObjectLimit);
-	Server->DestroyObject(ServerObjectCulled);
-	Server->DestroyObject(ServerObjectVeryFar);
+	for (auto It = ObjectLocs.CreateConstIterator(); It; ++It)
+	{
+		Server->DestroyObject(It.Key());
+	}
 }
 
-// Test world loc filter
 UE_NET_TEST_FIXTURE(FTestGridFilterFixture, TestWorldLocationIsFrequentlyUpdatedForNonDormantObject)
 {
 	struct FObjectLoc
@@ -227,6 +247,94 @@ UE_NET_TEST_FIXTURE(FTestGridFilterFixture, TestWorldLocationIsFrequentlyUpdated
 	UE_NET_ASSERT_NE(Client->GetReplicationBridge()->GetReplicatedObject(ServerObject->NetRefHandle), nullptr);
 }
 
+UE_NET_TEST_FIXTURE(FTestGridFilterFixture, TestWorldLocCullDistance)
+{
+	// Spawn object with WorldLocation's on server
+	UObjectReplicationBridge::FRootObjectReplicationParams Params;
+	Params.bNeedsWorldLocationUpdate = true;
+	Params.bUseClassConfigDynamicFilter = true;
+
+	UReplicatedTestObject* ServerObjectTest = Server->CreateObject(Params);
+
+	const UNetObjectGridFilterConfig* DefaultGridConfig = GetDefault<UNetObjectGridFilterConfig>();
+
+	struct FObjectLoc
+	{
+		FVector Loc;
+		float CullDistance;
+	};
+	TMap<UReplicatedTestObject*, FObjectLoc> ObjectLocs;
+
+	const float DefaultCullDistance = DefaultGridConfig->CellSizeX;
+
+	// Spawn an object that resides at the exact limit of his culldistance and at the boundary of a grid cell.
+	ObjectLocs.Add(ServerObjectTest, FObjectLoc{ .Loc=FVector(DefaultCullDistance, 0.0f, 0.0f), .CullDistance=DefaultCullDistance });
+
+	Server->GetReplicationBridge()->SetExternalWorldLocationUpdateFunctor([&](FNetRefHandle NetHandle, const UObject* ReplicatedObject, FVector& OutLocation, float& OutCullDistance)
+		{
+			const FObjectLoc& ObjectLoc = ObjectLocs[Cast<const UReplicatedTestObject>(ReplicatedObject)];
+			OutLocation = ObjectLoc.Loc;
+			OutCullDistance = ObjectLoc.CullDistance;
+		});
+
+	// Apply grid filter
+	for (auto It = ObjectLocs.CreateConstIterator(); It; ++It)
+	{
+		// Use the CullDistanceTest profile that disables the FrameCountBeforeCulling feature
+		Server->ReplicationSystem->SetFilter(It.Key()->NetRefHandle, WorldLocFilterHandle, TEXT("CullDistanceTest"));
+	}
+
+	// Add client
+	FReplicationSystemTestClient* Client = CreateClient();
+
+	// Set the view location of the client to (0,0,0)
+	FReplicationView ReplicationView;
+	FReplicationView::FView View = ReplicationView.Views.Emplace_GetRef();
+	View.Pos = FVector(0.0f, 0.0f, 0.0f);
+	Server->ReplicationSystem->SetReplicationView(Client->ConnectionIdOnServer, ReplicationView);
+
+	// Send and deliver packet
+	Server->UpdateAndSend({ Client });
+	
+	// Object should be relevant
+	UE_NET_ASSERT_NE(Client->GetReplicationBridge()->GetReplicatedObject(ServerObjectTest->NetRefHandle), nullptr);
+
+	// Reduce his culldistance via override
+	Server->ReplicationSystem->SetCullDistanceOverride(ServerObjectTest->NetRefHandle, DefaultCullDistance-1.0f);
+
+	// Send and deliver packet
+	Server->UpdateAndSend({ Client });
+
+	// Object should be filtered out
+	UE_NET_ASSERT_EQ(Client->GetReplicationBridge()->GetReplicatedObject(ServerObjectTest->NetRefHandle), nullptr);
+
+	// Go back to the initial cull distance
+	Server->ReplicationSystem->ClearCullDistanceOverride(ServerObjectTest->NetRefHandle);
+
+	// Send and deliver packet
+	Server->UpdateAndSend({ Client });
+
+	// Object should be relevant again
+	UE_NET_ASSERT_NE(Client->GetReplicationBridge()->GetReplicatedObject(ServerObjectTest->NetRefHandle), nullptr);
+
+	// Change the object's internal cull distance
+	ObjectLocs[ServerObjectTest].CullDistance = DefaultCullDistance-100.f;
+
+	// The object needs to be dirty for the cull distance update to be picked up
+	Server->ReplicationSystem->MarkDirty(ServerObjectTest->NetRefHandle);
+
+	// Send and deliver packet
+	Server->UpdateAndSend({ Client });
+
+	// Object should be filtered out
+	UE_NET_ASSERT_EQ(Client->GetReplicationBridge()->GetReplicatedObject(ServerObjectTest->NetRefHandle), nullptr);
+
+	for(auto It=ObjectLocs.CreateConstIterator(); It; ++It)
+	{
+		Server->DestroyObject(It.Key());
+	}
+}
+
 #if 0
 UE_NET_TEST_FIXTURE(FTestGridFilterFixture, TestFilterPerformance)
 {
diff --git a/Engine/Source/Runtime/Engine/Private/Net/Iris/ReplicationSystem/ReplicationSystemUtil.cpp b/Engine/Source/Runtime/Engine/Private/Net/Iris/ReplicationSystem/ReplicationSystemUtil.cpp
--- a/Engine/Source/Runtime/Engine/Private/Net/Iris/ReplicationSystem/ReplicationSystemUtil.cpp
+++ b/Engine/Source/Runtime/Engine/Private/Net/Iris/ReplicationSystem/ReplicationSystemUtil.cpp
@@ -625,7 +625,7 @@ void FReplicationSystemUtil::SetStaticPriority(const AActor* Actor, float Priori
 	});
 }
 
-void FReplicationSystemUtil::SetCullDistanceSqrOverride(const AActor* Actor, float CullDistSqr)
+void FReplicationSystemUtil::SetCullDistanceOverride(const AActor* Actor, float CullDist)
 {
 	FNetHandle ActorHandle = GetNetHandle(Actor);
 	if (!ActorHandle.IsValid())
@@ -633,7 +633,7 @@ void FReplicationSystemUtil::SetCullDistanceSqrOverride(const AActor* Actor, flo
 		return;
 	}
 	
-	ReplicationSystemUtil::ForEachReplicationSystem([ActorHandle, CullDistSqr](UReplicationSystem* ReplicationSystem)
+	ReplicationSystemUtil::ForEachReplicationSystem([ActorHandle, CullDist](UReplicationSystem* ReplicationSystem)
 	{
 		if (ReplicationSystem->IsServer())
 		{
@@ -642,14 +642,14 @@ void FReplicationSystemUtil::SetCullDistanceSqrOverride(const AActor* Actor, flo
 				FNetRefHandle RefHandle = Bridge->GetReplicatedRefHandle(ActorHandle);
 				if (RefHandle.IsValid())
 				{
-					ReplicationSystem->SetCullDistanceSqrOverride(RefHandle, CullDistSqr);
+					ReplicationSystem->SetCullDistanceOverride(RefHandle, CullDist);
 				}
 			}
 		}
 	});
 }
 
-void FReplicationSystemUtil::ClearCullDistanceSqrOverride(const AActor* Actor)
+void FReplicationSystemUtil::ClearCullDistanceOverride(const AActor* Actor)
 {
 	FNetHandle ActorHandle = GetNetHandle(Actor);
 	if (!ActorHandle.IsValid())
@@ -666,7 +666,7 @@ void FReplicationSystemUtil::ClearCullDistanceSqrOverride(const AActor* Actor)
 				FNetRefHandle RefHandle = Bridge->GetReplicatedRefHandle(ActorHandle);
 				if (RefHandle.IsValid())
 				{
-					ReplicationSystem->ClearCullDistanceSqrOverride(RefHandle);
+					ReplicationSystem->ClearCullDistanceOverride(RefHandle);
 				}
 			}
 		}
diff --git a/Engine/Source/Runtime/Engine/Public/Net/Iris/ReplicationSystem/ReplicationSystemUtil.h b/Engine/Source/Runtime/Engine/Public/Net/Iris/ReplicationSystem/ReplicationSystemUtil.h
--- a/Engine/Source/Runtime/Engine/Public/Net/Iris/ReplicationSystem/ReplicationSystemUtil.h
+++ b/Engine/Source/Runtime/Engine/Public/Net/Iris/ReplicationSystem/ReplicationSystemUtil.h
@@ -129,11 +129,11 @@ struct FReplicationSystemUtil
 	 */
 	ENGINE_API static void SetStaticPriority(const AActor* Actor, float Priority);
 
-	/** Set the squared cull distance for an actor. This will cause affected code to ignore the NetCullDistanceSquared property. */
-	ENGINE_API static void SetCullDistanceSqrOverride(const AActor* Actor, float CullDistSqr);
+	/** Set the cull distance for an actor. This will cause affected code to ignore the NetCullDistanceSquared property. */
+	ENGINE_API static void SetCullDistanceOverride(const AActor* Actor, float CullDistSqr);
 
-	/** Clears any previously set squared cull distance for an actor. This will cause affected code to respect the NetCullDistanceSquared property. */
-	ENGINE_API static void ClearCullDistanceSqrOverride(const AActor* Actor);
+	/** Clears any previously set cull distance override for an actor. This will cause affected code to respect the NetCullDistanceSquared property. */
+	ENGINE_API static void ClearCullDistanceOverride(const AActor* Actor);
 
 	/** Set the poll frequency for an object and its subobjects. */
 	ENGINE_API static void SetPollFrequency(const UObject* Object, float PollFrequency);
diff --git a/Engine/Source/Runtime/Experimental/Iris/Core/Private/Iris/Core/IrisLog.cpp b/Engine/Source/Runtime/Experimental/Iris/Core/Private/Iris/Core/IrisLog.cpp
--- a/Engine/Source/Runtime/Experimental/Iris/Core/Private/Iris/Core/IrisLog.cpp
+++ b/Engine/Source/Runtime/Experimental/Iris/Core/Private/Iris/Core/IrisLog.cpp
@@ -4,3 +4,4 @@
 
 DEFINE_LOG_CATEGORY(LogIris);
 DEFINE_LOG_CATEGORY(LogIrisFiltering);
+DEFINE_LOG_CATEGORY(LogIrisNetCull);
diff --git a/Engine/Source/Runtime/Experimental/Iris/Core/Private/Iris/ReplicationSystem/Filtering/NetObjectGridFilter.cpp b/Engine/Source/Runtime/Experimental/Iris/Core/Private/Iris/ReplicationSystem/Filtering/NetObjectGridFilter.cpp
--- a/Engine/Source/Runtime/Experimental/Iris/Core/Private/Iris/ReplicationSystem/Filtering/NetObjectGridFilter.cpp
+++ b/Engine/Source/Runtime/Experimental/Iris/Core/Private/Iris/ReplicationSystem/Filtering/NetObjectGridFilter.cpp
@@ -1,7 +1,6 @@
 // Copyright Epic Games, Inc. All Rights Reserved.
 
 #include "Iris/ReplicationSystem/Filtering/NetObjectGridFilter.h"
-#include "Iris/ReplicationSystem/NetCullDistanceOverrides.h"
 #include "Iris/ReplicationSystem/RepTag.h"
 #include "Iris/ReplicationSystem/ReplicationProtocol.h"
 #include "Iris/ReplicationSystem/ReplicationSystem.h"
@@ -21,16 +20,12 @@ void UNetObjectGridFilter::OnInit(const FNetObjectFilterInitParams& Params)
 	AssignedObjectInfoIndices.Init(Params.CurrentMaxInternalIndex);
 
 	PerConnectionInfos.SetNum(Params.MaxConnectionCount + 1);
-
-	NetCullDistanceOverrides = &Params.ReplicationSystem->GetNetCullDistanceOverrides();
 }
 
 void UNetObjectGridFilter::OnDeinit()
 {
 	Config = nullptr;
 
-	NetCullDistanceOverrides = nullptr;
-
 	PerConnectionInfos.Empty();
 	ObjectInfos.Empty();
 	AssignedObjectInfoIndices.Empty();
@@ -475,17 +470,7 @@ void UNetObjectGridFilter::UpdateCellInfoForObject(const FObjectLocationInfo& Ob
 
 void UNetObjectGridFilter::UpdatePositionAndCullDistance(const UNetObjectGridFilter::FObjectLocationInfo& ObjectLocationInfo, UNetObjectGridFilter::FPerObjectInfo& PerObjectInfo, const UE::Net::FReplicationInstanceProtocol* InstanceProtocol)
 {
-	// Update position & culldistance
-	{
-		UpdateObjectInfo(PerObjectInfo, ObjectLocationInfo, InstanceProtocol);
-	}
-
-	// Optionally update cull distance
-	if (NetCullDistanceOverrides->HasCullDistanceOverride(PerObjectInfo.ObjectIndex))
-	{
-		const float CullDistanceSq = NetCullDistanceOverrides->GetCullDistanceSqr(PerObjectInfo.ObjectIndex);
-		PerObjectInfo.SetCullDistanceSq(CullDistanceSq);
-	}
+	UpdateObjectInfo(PerObjectInfo, ObjectLocationInfo, InstanceProtocol);
 }
 
 void UNetObjectGridFilter::CalculateCellBox(const UNetObjectGridFilter::FPerObjectInfo& PerObjectInfo, UNetObjectGridFilter::FCellBox& OutCellBox)
@@ -605,10 +590,10 @@ void UNetObjectGridWorldLocFilter::UpdateObjectInfo(UNetObjectGridFilter::FPerOb
 {
 	check(ObjectLocationInfo.IsUsingWorldLocations());
 
-	const UE::Net::FWorldLocations::FObjectInfo& CachedObjectInfo = WorldLocations->GetObjectInfo(PerObjectInfo.ObjectIndex);
+	const UE::Net::FWorldLocations::FWorldInfo& Info = WorldLocations->GetWorldInfo(PerObjectInfo.ObjectIndex);
 	
-	PerObjectInfo.Position = CachedObjectInfo.WorldLocation;
-	PerObjectInfo.SetCullDistance(CachedObjectInfo.CullDistance);
+	PerObjectInfo.Position = Info.WorldLocation;
+	PerObjectInfo.SetCullDistance(Info.CullDistance);
 }
 
 bool UNetObjectGridWorldLocFilter::BuildObjectInfo(uint32 ObjectIndex, FNetObjectFilterAddObjectParams& Params)
diff --git a/Engine/Source/Runtime/Experimental/Iris/Core/Private/Iris/ReplicationSystem/NetCullDistanceOverrides.cpp b/Engine/Source/Runtime/Experimental/Iris/Core/Private/Iris/ReplicationSystem/NetCullDistanceOverrides.cpp
deleted file mode 100644
diff --git a/Engine/Source/Runtime/Experimental/Iris/Core/Private/Iris/ReplicationSystem/ObjectReplicationBridge.cpp b/Engine/Source/Runtime/Experimental/Iris/Core/Private/Iris/ReplicationSystem/ObjectReplicationBridge.cpp
--- a/Engine/Source/Runtime/Experimental/Iris/Core/Private/Iris/ReplicationSystem/ObjectReplicationBridge.cpp
+++ b/Engine/Source/Runtime/Experimental/Iris/Core/Private/Iris/ReplicationSystem/ObjectReplicationBridge.cpp
@@ -1452,13 +1452,7 @@ void UObjectReplicationBridge::ForceUpdateWorldLocation(FNetRefHandle NetRefHand
 				WorldInfo.WorldLocation.Y >= WorldLocations.GetWorldMinPos().Y && WorldInfo.WorldLocation.Y <= WorldLocations.GetWorldMaxPos().Y,
 				TEXT("Object %s with position %s lies outside configured world boundary."), ToCStr(NetRefHandleManager->PrintObjectFromIndex(InternalObjectIndex)), *WorldInfo.WorldLocation.ToString());
 
-	const FWorldLocations::FObjectInfo WorldObjectInfo
-	{
-		.WorldLocation = WorldInfo.WorldLocation,
-		.CullDistance = WorldInfo.CullDistance
-	};
-
-	WorldLocations.SetObjectInfo(InternalObjectIndex, WorldObjectInfo);
+	WorldLocations.SetObjectInfo(InternalObjectIndex, WorldInfo.WorldLocation, WorldInfo.CullDistance);
 }
 
 void UObjectReplicationBridge::UpdateInstancesWorldLocation()
@@ -1494,13 +1488,7 @@ void UObjectReplicationBridge::UpdateInstancesWorldLocation()
 					   WorldInfo.WorldLocation.Y >= WorldLocations.GetWorldMinPos().Y && WorldInfo.WorldLocation.Y <= WorldLocations.GetWorldMaxPos().Y,
                        TEXT("Object %s with position %s lies outside configured world boundary."), ToCStr(NetRefHandleManager->PrintObjectFromIndex((InternalObjectIndex))), *WorldInfo.WorldLocation.ToString());
 			
-            const FWorldLocations::FObjectInfo WorldObjectInfo
-            { 
-                .WorldLocation = WorldInfo.WorldLocation, 
-                .CullDistance = WorldInfo.CullDistance 
-            };
-
-            WorldLocations.SetObjectInfo(InternalObjectIndex, WorldObjectInfo);
+            WorldLocations.SetObjectInfo(InternalObjectIndex, WorldInfo.WorldLocation, WorldInfo.CullDistance);
 		}
 	};
 
diff --git a/Engine/Source/Runtime/Experimental/Iris/Core/Private/Iris/ReplicationSystem/ObjectReplicationBridgeDebugging.cpp b/Engine/Source/Runtime/Experimental/Iris/Core/Private/Iris/ReplicationSystem/ObjectReplicationBridgeDebugging.cpp
--- a/Engine/Source/Runtime/Experimental/Iris/Core/Private/Iris/ReplicationSystem/ObjectReplicationBridgeDebugging.cpp
+++ b/Engine/Source/Runtime/Experimental/Iris/Core/Private/Iris/ReplicationSystem/ObjectReplicationBridgeDebugging.cpp
@@ -7,8 +7,6 @@
 #include "Iris/Core/IrisLog.h"
 #include "Iris/Core/IrisDebugging.h"
 
-#include "Iris/ReplicationSystem/NetCullDistanceOverrides.h"
-
 #include "Iris/ReplicationSystem/ReplicationSystem.h"
 #include "Iris/ReplicationSystem/ReplicationSystemTypes.h"
 #include "Iris/ReplicationSystem/ReplicationSystemInternal.h"
@@ -870,7 +868,6 @@ void UObjectReplicationBridge::PrintNetCullDistances(const TArray<FString>& Args
 
 	FReplicationSystemInternal* ReplicationSystemInternal = GetReplicationSystem()->GetReplicationSystemInternal();
 	const FWorldLocations& WorldLocations = ReplicationSystemInternal->GetWorldLocations();
-	const FNetCullDistanceOverrides& CullDistanceOverrides = ReplicationSystemInternal->GetNetCullDistanceOverrides();
 
 	struct FCullDistanceInfo
 	{
@@ -880,6 +877,9 @@ void UObjectReplicationBridge::PrintNetCullDistances(const TArray<FString>& Args
 		// Total replicated root objects of this class
 		uint32 NumTotal = 0;
 
+		// Total objects with a cull distance override
+		uint32 NumOverrides = 0;
+
 		// Track unique culldistance values for replicated root objects
 		TMap<float /*CullDistance*/, uint32 /*ActorCount with culldistance value*/> UniqueCullDistances; 
 
@@ -921,16 +921,10 @@ void UObjectReplicationBridge::PrintNetCullDistances(const TArray<FString>& Args
 		if( UObject* RepObj = NetRefHandleManager->GetReplicatedObjectInstance(RootObjectIndex) )
 		{
 			UClass* RepObjClass = RepObj->GetClass();
-			// Find this object's current net cull distance
-			float RootObjectCullDistance = WorldLocations.GetCullDistance(RootObjectIndex);
-			if (CullDistanceOverrides.HasCullDistanceOverride(RootObjectIndex))
-			{
-				RootObjectCullDistance = FMath::Sqrt(CullDistanceOverrides.GetCullDistanceSqr(RootObjectIndex));
-			}
 
 			FCullDistanceInfo& Info = ClassCullDistanceMap.FindOrAdd(RepObjClass);
 
-			// Find the CullDistance of the CDO.
+			// Find the CullDistance of the CDO the first time we discover a new class 
 			if (Info.Class == nullptr)
 			{
 				UNetObjectFactory* NetFactory = GetNetFactory(NetRefHandleManager->GetReplicatedObjectData(RootObjectIndex).NetFactoryId);
@@ -950,18 +944,39 @@ void UObjectReplicationBridge::PrintNetCullDistances(const TArray<FString>& Args
 
 			Info.NumTotal++;
 
+			if (WorldLocations.HasCullDistanceOverride(RootObjectIndex))
+			{
+				Info.NumOverrides++;
+			}
+
+			// Find this object's current net cull distance
+			float RootObjectCullDistance = WorldLocations.GetCullDistance(RootObjectIndex);
+
 			uint32& NumUsingCullDistance = Info.UniqueCullDistances.FindOrAdd(RootObjectCullDistance, 0);
 			++NumUsingCullDistance;
 		}
 	});
 
-	// Sort from highest to lowest
-	ClassCullDistanceMap.ValueSort([](const FCullDistanceInfo& lhs, const FCullDistanceInfo& rhs) 
-	{ 
-		const float LHSSortingCullDistance = lhs.FindMostUsedCullDistance();
-		const float RHSSortingCullDistance = rhs.FindMostUsedCullDistance();
-		return LHSSortingCullDistance >= RHSSortingCullDistance;
-	});
+	EPrintTraits PrintArgs = FindPrintTraitsFromArgs(Args);
+
+
+	if (PrintArgs == EPrintTraits::Default)
+	{
+		// Sort from highest to lowest
+		ClassCullDistanceMap.ValueSort([](const FCullDistanceInfo& lhs, const FCullDistanceInfo& rhs) 
+		{ 
+			const float LHSSortingCullDistance = lhs.FindMostUsedCullDistance();
+			const float RHSSortingCullDistance = rhs.FindMostUsedCullDistance();
+			return LHSSortingCullDistance >= RHSSortingCullDistance;
+		});
+	}
+	else if (PrintArgs == EPrintTraits::SortByClass)
+	{
+		ClassCullDistanceMap.KeySort([](const UClass& lhs, const UClass& rhs)
+		{
+				return lhs.GetName() < rhs.GetName();
+		});
+	}
 
 	UE_LOG(LogIrisBridge, Display, TEXT("################ Start Printing NetCullDistance Values ################"));
 	UE_LOG(LogIrisBridge, Display, TEXT(""));
@@ -972,7 +987,7 @@ void UObjectReplicationBridge::PrintNetCullDistances(const TArray<FString>& Args
 		FCullDistanceInfo& Info = ClassIt.Value();
 		UClass* Class = Info.Class;
 
-		UE_LOG(LogIrisBridge, Display, TEXT("MostCommon NetCullDistance: %f | Class: %s | Instances: %u"), Info.FindMostUsedCullDistance(), *Info.Class->GetName(), Info.NumTotal);
+		UE_LOG(LogIrisBridge, Display, TEXT("MostCommon NetCullDistance: %f | Class: %s | Instances: %u | Overrides: %u"), Info.FindMostUsedCullDistance(), *Info.Class->GetName(), Info.NumTotal, Info.NumOverrides);
 
 		Info.UniqueCullDistances.KeySort([](const float& lhs, const float& rhs){ return lhs >= rhs; });
 		for (auto DivergentIt = Info.UniqueCullDistances.CreateConstIterator(); DivergentIt; ++DivergentIt)
diff --git a/Engine/Source/Runtime/Experimental/Iris/Core/Private/Iris/ReplicationSystem/ReplicationBridge.cpp b/Engine/Source/Runtime/Experimental/Iris/Core/Private/Iris/ReplicationSystem/ReplicationBridge.cpp
--- a/Engine/Source/Runtime/Experimental/Iris/Core/Private/Iris/ReplicationSystem/ReplicationBridge.cpp
+++ b/Engine/Source/Runtime/Experimental/Iris/Core/Private/Iris/ReplicationSystem/ReplicationBridge.cpp
@@ -497,9 +497,6 @@ void UReplicationBridge::InternalDestroyNetObject(FNetRefHandle Handle)
 	{
 		FReplicationSystemInternal* ReplicationSystemInternal = ReplicationSystem->GetReplicationSystemInternal();
 
-		FNetCullDistanceOverrides& NetCullDistanceOverrides = ReplicationSystemInternal->GetNetCullDistanceOverrides();
-		NetCullDistanceOverrides.ClearCullDistanceSqr(ObjectInternalIndex);
-
 		FWorldLocations& WorldLocations = ReplicationSystemInternal->GetWorldLocations();
 		WorldLocations.RemoveObjectInfoCache(ObjectInternalIndex);
 
@@ -1028,7 +1025,7 @@ UE::Net::FNetRefHandle UReplicationBridge::StoreDestructionInfo(FNetRefHandle Ha
 				   Parameters.Location.Y >= WorldLocations.GetWorldMinPos().Y && Parameters.Location.Y <= WorldLocations.GetWorldMaxPos().Y, TEXT("Object %s with position %s lies outside configured world boundary."), ToCStr(NetRefHandleManager->PrintObjectFromIndex((InternalReplicationIndex))), *Parameters.Location.ToString());
 
 		WorldLocations.InitObjectInfoCache(InternalReplicationIndex);
-		WorldLocations.UpdateWorldLocation(InternalReplicationIndex, Parameters.Location);
+		WorldLocations.SetObjectInfo(InternalReplicationIndex, Parameters.Location, 0.0f);
 
 		GetReplicationSystem()->SetPrioritizer(DestructionInfoHandle, DefaultSpatialNetObjectPrioritizerHandle);
 	}
diff --git a/Engine/Source/Runtime/Experimental/Iris/Core/Private/Iris/ReplicationSystem/ReplicationSystem.cpp b/Engine/Source/Runtime/Experimental/Iris/Core/Private/Iris/ReplicationSystem/ReplicationSystem.cpp
--- a/Engine/Source/Runtime/Experimental/Iris/Core/Private/Iris/ReplicationSystem/ReplicationSystem.cpp
+++ b/Engine/Source/Runtime/Experimental/Iris/Core/Private/Iris/ReplicationSystem/ReplicationSystem.cpp
@@ -197,13 +197,6 @@ public:
 			Registry.Init(InitParams);
 		}		
 
-		{
-			FNetCullDistanceOverrides& NetCullDistanceOverrides = ReplicationSystemInternal.GetNetCullDistanceOverrides();
-			FNetCullDistanceOverridesInitParams InitParams;
-			InitParams.MaxInternalNetRefIndex = CurrentMaxInternalNetRefIndex;
-			NetCullDistanceOverrides.Init(InitParams);
-		}
-
 		{
 			FWorldLocations& WorldLocations = ReplicationSystemInternal.GetWorldLocations();
 			FWorldLocationsInitParams InitParams;
@@ -338,7 +331,6 @@ public:
 	{
 		ReplicationSystemInternal.GetReplicationStateStorage().OnMaxInternalNetRefIndexIncreased(NewMaxInternalIndex);
 		ReplicationSystemInternal.GetGroups().OnMaxInternalNetRefIndexIncreased(NewMaxInternalIndex);
-		ReplicationSystemInternal.GetNetCullDistanceOverrides().OnMaxInternalNetRefIndexIncreased(NewMaxInternalIndex);
 		ReplicationSystemInternal.GetWorldLocations().OnMaxInternalNetRefIndexIncreased(NewMaxInternalIndex);
 		ReplicationSystemInternal.GetDeltaCompressionBaselineInvalidationTracker().OnMaxInternalNetRefIndexIncreased(NewMaxInternalIndex);
 		ReplicationSystemInternal.GetDeltaCompressionBaselineManager().OnMaxInternalNetRefIndexIncreased(NewMaxInternalIndex);
@@ -380,6 +372,9 @@ public:
 		// Reset baseline invalidation
 		ReplicationSystemInternal.GetDeltaCompressionBaselineInvalidationTracker().PostSendUpdate();
 
+		// Reset dirty info list for the next frame
+		ReplicationSystemInternal.GetWorldLocations().PostSendUpdate();
+
 		// Tell systems we finished PostSendUpdate
 		ReplicationSystemInternal.GetReplicationBridge()->OnPostSendUpdate();
 	}
@@ -391,17 +386,15 @@ public:
 
 	void UpdateDirtyListPostPoll()
 	{
+		// From here there shouldn't be any user code that calls public API functions.
+		ReplicationSystemInternal.GetWorldLocations().LockDirtyInfoList(true);
+
 		ReplicationSystemInternal.GetDirtyNetObjectTracker().UpdateAccumulatedDirtyList();
 	}
 
 	void UpdateWorldLocations()
 	{
 		IRIS_CSV_PROFILER_SCOPE(Iris, ReplicationSystem_UpdateWorldLocations);
-
-		// Reset dirty object info before updating.
-		FWorldLocations& WorldLocations = ReplicationSystemInternal.GetWorldLocations();
-		WorldLocations.ResetObjectsWithDirtyInfo();
-
 		ReplicationSystemInternal.GetReplicationBridge()->CallUpdateInstancesWorldLocation();
 	}
 
@@ -1900,17 +1893,12 @@ void UReplicationSystem::AddReferencedObjects(UObject* InThis, FReferenceCollect
 	Super::AddReferencedObjects(InThis, Collector);
 }
 
-const UE::Net::FNetCullDistanceOverrides& UReplicationSystem::GetNetCullDistanceOverrides() const
-{
-	return Impl->ReplicationSystemInternal.GetNetCullDistanceOverrides();
-}
-
 const UE::Net::FWorldLocations& UReplicationSystem::GetWorldLocations() const
 {
 	return Impl->ReplicationSystemInternal.GetWorldLocations();
 }
 
-void UReplicationSystem::SetCullDistanceSqrOverride(FNetRefHandle Handle, float DistSqr)
+void UReplicationSystem::SetCullDistanceOverride(FNetRefHandle Handle, float CullDistance)
 {
 	const UE::Net::Private::FInternalNetRefIndex ObjectInternalIndex = Impl->ReplicationSystemInternal.GetNetRefHandleManager().GetInternalIndex(Handle);
 	if (ObjectInternalIndex == UE::Net::Private::FNetRefHandleManager::InvalidInternalIndex)
@@ -1918,12 +1906,18 @@ void UReplicationSystem::SetCullDistanceSqrOverride(FNetRefHandle Handle, float
 		return;
 	}
 
-	UE_LOG(LogIris, Verbose, TEXT("UReplicationSystem::SetCullDistanceSqrOverride: %s is now overridden to %f"), *Impl->ReplicationSystemInternal.GetNetRefHandleManager().PrintObjectFromNetRefHandle(Handle), FMath::Sqrt(DistSqr));
-	
-	return Impl->ReplicationSystemInternal.GetNetCullDistanceOverrides().SetCullDistanceSqr(ObjectInternalIndex, DistSqr);
+	UE_LOG(LogIrisNetCull, Verbose, TEXT("UReplicationSystem::SetCullDistanceOverride: %s will be overridden to %f. Previous cull distance: %f"),
+		*Impl->ReplicationSystemInternal.GetNetRefHandleManager().PrintObjectFromIndex(ObjectInternalIndex),
+		CullDistance,
+		Impl->ReplicationSystemInternal.GetWorldLocations().GetCullDistance(ObjectInternalIndex));
+
+	const bool bSuccess = Impl->ReplicationSystemInternal.GetWorldLocations().SetCullDistanceOverride(ObjectInternalIndex, CullDistance);
+
+	ensureMsgf(bSuccess, TEXT("SetCullDistanceOverride failed for %s (cull distance: %f). The object does not use the world location cache."), 
+		*Impl->ReplicationSystemInternal.GetNetRefHandleManager().PrintObjectFromIndex(ObjectInternalIndex), CullDistance);
 }
 
-void UReplicationSystem::ClearCullDistanceSqrOverride(FNetRefHandle Handle)
+void UReplicationSystem::ClearCullDistanceOverride(FNetRefHandle Handle)
 {
 	const UE::Net::Private::FInternalNetRefIndex ObjectInternalIndex = Impl->ReplicationSystemInternal.GetNetRefHandleManager().GetInternalIndex(Handle);
 	if (ObjectInternalIndex == UE::Net::Private::FNetRefHandleManager::InvalidInternalIndex)
@@ -1931,11 +1925,14 @@ void UReplicationSystem::ClearCullDistanceSqrOverride(FNetRefHandle Handle)
 		return;
 	}
 
-	const bool bWasCullDistanceOverriden = Impl->ReplicationSystemInternal.GetNetCullDistanceOverrides().ClearCullDistanceSqr(ObjectInternalIndex);
-	UE_CLOG(bWasCullDistanceOverriden, LogIris, Verbose, TEXT("UReplicationSystem::ClearCullDistanceSqrOverride: %s is no longer overridden."), *Impl->ReplicationSystemInternal.GetNetRefHandleManager().PrintObjectFromNetRefHandle(Handle));
+	const bool bWasCullDistanceOverriden = Impl->ReplicationSystemInternal.GetWorldLocations().ClearCullDistanceOverride(ObjectInternalIndex);
+	
+	UE_CLOG(bWasCullDistanceOverriden, LogIrisNetCull, Verbose, TEXT("UReplicationSystem::ClearCullDistanceOverride: %s is no longer overridden. Now using cull distance: %f"),
+		*Impl->ReplicationSystemInternal.GetNetRefHandleManager().PrintObjectFromNetRefHandle(Handle),
+		Impl->ReplicationSystemInternal.GetWorldLocations().GetCullDistance(ObjectInternalIndex));
 }
 
-float UReplicationSystem::GetCullDistanceSqrOverride(FNetRefHandle Handle, float DefaultValue) const
+float UReplicationSystem::GetCullDistance(FNetRefHandle Handle, float DefaultValue) const
 {
 	const UE::Net::Private::FInternalNetRefIndex ObjectInternalIndex = Impl->ReplicationSystemInternal.GetNetRefHandleManager().GetInternalIndex(Handle);
 	if (ObjectInternalIndex == UE::Net::Private::FNetRefHandleManager::InvalidInternalIndex)
@@ -1943,7 +1940,13 @@ float UReplicationSystem::GetCullDistanceSqrOverride(FNetRefHandle Handle, float
 		return DefaultValue;
 	}
 
-	return Impl->ReplicationSystemInternal.GetNetCullDistanceOverrides().GetCullDistanceSqr(ObjectInternalIndex, DefaultValue);
+	if (!Impl->ReplicationSystemInternal.GetWorldLocations().HasInfoForObject(ObjectInternalIndex))
+	{
+		// This replicated object did not register world information
+		return DefaultValue;
+	}
+
+	return Impl->ReplicationSystemInternal.GetWorldLocations().GetCullDistance(ObjectInternalIndex);
 }
 
 void UReplicationSystem::ReportProtocolMismatch(uint64 NetRefHandleId, uint32 ConnectionId)
diff --git a/Engine/Source/Runtime/Experimental/Iris/Core/Private/Iris/ReplicationSystem/ReplicationSystemInternal.h b/Engine/Source/Runtime/Experimental/Iris/Core/Private/Iris/ReplicationSystem/ReplicationSystemInternal.h
--- a/Engine/Source/Runtime/Experimental/Iris/Core/Private/Iris/ReplicationSystem/ReplicationSystemInternal.h
+++ b/Engine/Source/Runtime/Experimental/Iris/Core/Private/Iris/ReplicationSystem/ReplicationSystemInternal.h
@@ -6,7 +6,6 @@
 #include "Containers/Array.h"
 #include "Iris/ReplicationState/ReplicationStateStorage.h"
 #include "Iris/ReplicationSystem/DirtyNetObjectTracker.h"
-#include "Iris/ReplicationSystem/NetCullDistanceOverrides.h"
 #include "Iris/ReplicationSystem/NetRefHandleManager.h"
 #include "Iris/ReplicationSystem/ChangeMaskCache.h"
 #include "Iris/ReplicationSystem/Conditionals/ReplicationConditionals.h"
@@ -91,8 +90,6 @@ public:
 	FNetBlobHandlerManager& GetNetBlobHandlerManager() { return NetBlobManager.GetNetBlobHandlerManager(); }
 	const FNetBlobHandlerManager& GetNetBlobHandlerManager() const { return NetBlobManager.GetNetBlobHandlerManager(); }
 
-	FNetCullDistanceOverrides& GetNetCullDistanceOverrides() { return NetCullDistanceOverrides; }
-
 	FWorldLocations& GetWorldLocations() { return WorldLocations; }
 
 	FDeltaCompressionBaselineManager& GetDeltaCompressionBaselineManager() { return DeltaCompressionBaselineManager; }
@@ -153,7 +150,6 @@ private:
 	FReplicationPrioritization Prioritization;
 	FObjectReferenceCache ObjectReferenceCache;
 	FNetBlobManager NetBlobManager;
-	FNetCullDistanceOverrides NetCullDistanceOverrides;
 	FWorldLocations WorldLocations;
 	FDeltaCompressionBaselineManager DeltaCompressionBaselineManager;
 	FDeltaCompressionBaselineInvalidationTracker DeltaCompressionBaselineInvalidationTracker;
diff --git a/Engine/Source/Runtime/Experimental/Iris/Core/Private/Iris/ReplicationSystem/WorldLocations.cpp b/Engine/Source/Runtime/Experimental/Iris/Core/Private/Iris/ReplicationSystem/WorldLocations.cpp
--- a/Engine/Source/Runtime/Experimental/Iris/Core/Private/Iris/ReplicationSystem/WorldLocations.cpp
+++ b/Engine/Source/Runtime/Experimental/Iris/Core/Private/Iris/ReplicationSystem/WorldLocations.cpp
@@ -14,19 +14,34 @@ void FWorldLocations::Init(const FWorldLocationsInitParams& InitParams)
 	ValidInfoIndexes.Init(InitParams.MaxInternalNetRefIndex);
 	ObjectsWithDirtyInfo.Init(InitParams.MaxInternalNetRefIndex);
 	ObjectsRequiringFrequentWorldLocationUpdate.Init(InitParams.MaxInternalNetRefIndex);
+	StorageIndexes.SetNum(InitParams.MaxInternalNetRefIndex);
+
+	ReservedStorageSlot.Init(InitParams.PreallocatedStorageCount);	
+	StoredObjectInfo = TNetChunkedArray<FObjectInfo, StorageElementsPerChunk>(InitParams.PreallocatedStorageCount, EInitMemory::Constructor);
 
 	MinWorldPos = GetDefault<UWorldLocationsConfig>()->MinPos;
 	MaxWorldPos = GetDefault<UWorldLocationsConfig>()->MaxPos;
 }
 
+void FWorldLocations::PostSendUpdate()
+{
+	// Clear the dirty info list to start fresh for the next frame
+#if DO_ENSURE
+	bLockdownDirtyList = false;
+#endif
+
+	ObjectsWithDirtyInfo.ClearAllBits();
+}
+
 void FWorldLocations::OnMaxInternalNetRefIndexIncreased(UE::Net::Private::FInternalNetRefIndex NewMaxInternalIndex)
 {
 	ValidInfoIndexes.SetNumBits(NewMaxInternalIndex);
 	ObjectsWithDirtyInfo.SetNumBits(NewMaxInternalIndex);
 	ObjectsRequiringFrequentWorldLocationUpdate.SetNumBits(NewMaxInternalIndex);
+	StorageIndexes.SetNum(NewMaxInternalIndex);
 }
 
-void FWorldLocations::InitObjectInfoCache(uint32 ObjectIndex)
+void FWorldLocations::InitObjectInfoCache(UE::Net::Private::FInternalNetRefIndex ObjectIndex)
 {
 	if (ValidInfoIndexes.IsBitSet(ObjectIndex))
 	{
@@ -35,54 +50,98 @@ void FWorldLocations::InitObjectInfoCache(uint32 ObjectIndex)
 	}
 
 	ValidInfoIndexes.SetBit(ObjectIndex);
-	
-	if (ObjectIndex >= uint32(StoredObjectInfo.Num()))
+
+	// Find an available slot
+	uint32 AvailableSlot = ReservedStorageSlot.FindFirstZero();
+
+	// No more slots available, grow the storage space by a single chunk
+	if (AvailableSlot == FNetBitArray::InvalidIndex)
 	{
 		LLM_SCOPE_BYTAG(Iris);
-		StoredObjectInfo.Add(ObjectIndex + 1U - StoredObjectInfo.Num());
+		AvailableSlot = ReservedStorageSlot.GetNumBits();
+		StoredObjectInfo.Add(StorageElementsPerChunk);
+		ReservedStorageSlot.SetNumBits(AvailableSlot + StorageElementsPerChunk);
 	}
 
-	StoredObjectInfo[ObjectIndex].WorldLocation = FVector::Zero();
-	StoredObjectInfo[ObjectIndex].CullDistance = 0.0f;
+	ReservedStorageSlot.SetBit(AvailableSlot);
+	StorageIndexes[ObjectIndex] = AvailableSlot;
 }
 
-void FWorldLocations::RemoveObjectInfoCache(uint32 ObjectIndex)
+void FWorldLocations::RemoveObjectInfoCache(UE::Net::Private::FInternalNetRefIndex ObjectIndex)
 {
+	if (!ValidInfoIndexes.IsBitSet(ObjectIndex))
+	{
+		// Object did not register a location
+		return;
+	}
+
 	ValidInfoIndexes.ClearBit(ObjectIndex);
 	ObjectsWithDirtyInfo.ClearBit(ObjectIndex);
 	ObjectsRequiringFrequentWorldLocationUpdate.ClearBit(ObjectIndex);
+
+	const uint32 StorageIndex = StorageIndexes[ObjectIndex];
+	StorageIndexes[ObjectIndex] = INDEX_NONE;
+
+	// Default construct the info since it can be reused in the future.
+	StoredObjectInfo[StorageIndex] = FObjectInfo();
+
+	ReservedStorageSlot.ClearBit(StorageIndex);
 }
 
-void FWorldLocations::SetObjectInfo(uint32 ObjectIndex, const FWorldLocations::FObjectInfo& ObjectInfo)
+void FWorldLocations::SetObjectInfo(UE::Net::Private::FInternalNetRefIndex ObjectIndex, const FVector& Location, const float NetCullDistance)
 {
+	ensure(!bLockdownDirtyList);
+
 	checkSlow(ValidInfoIndexes.IsBitSet(ObjectIndex));
-	FObjectInfo& TargetObjectInfo = StoredObjectInfo[ObjectIndex];
-	const bool bHasInfoChanged = ObjectsWithDirtyInfo.GetBit(ObjectIndex) || TargetObjectInfo.WorldLocation != ObjectInfo.WorldLocation || TargetObjectInfo.CullDistance != ObjectInfo.CullDistance;
-	TargetObjectInfo = ObjectInfo;
-	TargetObjectInfo.WorldLocation = ClampPositionToBoundary(ObjectInfo.WorldLocation);
+	FObjectInfo& TargetObjectInfo = GetObjectInfo(ObjectIndex);
+
+	const FVector ClampedLoc = ClampPositionToBoundary(Location);
+	
+	const bool bHasInfoChanged = ObjectsWithDirtyInfo.GetBit(ObjectIndex) || TargetObjectInfo.WorldLocation != ClampedLoc || TargetObjectInfo.CullDistance != NetCullDistance;
+
+	TargetObjectInfo.WorldLocation = ClampedLoc;
+	TargetObjectInfo.CullDistance = NetCullDistance;
 
-	ObjectsWithDirtyInfo.SetBitValue(ObjectIndex, bHasInfoChanged);
+	ObjectsWithDirtyInfo.OrBitValue(ObjectIndex, bHasInfoChanged);
 }
 
-void FWorldLocations::UpdateWorldLocation(uint32 ObjectIndex, const FVector& WorldLocation)
+bool FWorldLocations::SetCullDistanceOverride(UE::Net::Private::FInternalNetRefIndex ObjectIndex, float CullDistance)
 {
-	const FVector InBoundsWorldLocation = ClampPositionToBoundary(WorldLocation);
-	checkSlow(ValidInfoIndexes.GetBit(ObjectIndex));
-	FVector& TargetWorldLocation = StoredObjectInfo[ObjectIndex].WorldLocation;
-	TargetWorldLocation = InBoundsWorldLocation;
-	const bool bHasInfoChanged = ObjectsWithDirtyInfo.GetBit(ObjectIndex) || TargetWorldLocation != InBoundsWorldLocation;
-	
-	ObjectsWithDirtyInfo.SetBitValue(ObjectIndex, bHasInfoChanged);
+	ensure(!bLockdownDirtyList);
+
+	if (ValidInfoIndexes.IsBitSet(ObjectIndex))
+	{
+		//TODO: Check for zero or clamp down on huge values ?
+		if (GetObjectInfo(ObjectIndex).CullDistanceOverride != CullDistance)
+		{
+			GetObjectInfo(ObjectIndex).CullDistanceOverride = CullDistance;
+			ObjectsWithDirtyInfo.SetBitValue(ObjectIndex, true);
+		}
+		return true;
+	}
+
+	return false;
 }
 
-void FWorldLocations::ResetObjectsWithDirtyInfo()
+bool FWorldLocations::ClearCullDistanceOverride(UE::Net::Private::FInternalNetRefIndex ObjectIndex)
 {
-	ObjectsWithDirtyInfo.ClearAllBits();
+	ensure(!bLockdownDirtyList);
+
+	if (ValidInfoIndexes.IsBitSet(ObjectIndex) && GetObjectInfo(ObjectIndex).CullDistanceOverride != FLT_MAX)
+	{
+		GetObjectInfo(ObjectIndex).CullDistanceOverride = FLT_MAX;
+		ObjectsWithDirtyInfo.SetBitValue(ObjectIndex, true);
+		return true;
+	}
+
+	return false;
 }
 
-FVector FWorldLocations::ClampPositionToBoundary(const FVector& Position)
+void FWorldLocations::LockDirtyInfoList(bool bLock)
 {
-	return Position.BoundToBox(GetWorldMinPos(), GetWorldMaxPos());
+#if DO_ENSURE
+	bLockdownDirtyList = bLock;
+#endif
 }
 
-}
+} // end namespace UE::Net
diff --git a/Engine/Source/Runtime/Experimental/Iris/Core/Public/Iris/Core/IrisLog.h b/Engine/Source/Runtime/Experimental/Iris/Core/Public/Iris/Core/IrisLog.h
--- a/Engine/Source/Runtime/Experimental/Iris/Core/Public/Iris/Core/IrisLog.h
+++ b/Engine/Source/Runtime/Experimental/Iris/Core/Public/Iris/Core/IrisLog.h
@@ -6,3 +6,4 @@
 
 IRISCORE_API DECLARE_LOG_CATEGORY_EXTERN(LogIris, Log, All);
 IRISCORE_API DECLARE_LOG_CATEGORY_EXTERN(LogIrisFiltering, Log, All);
+IRISCORE_API DECLARE_LOG_CATEGORY_EXTERN(LogIrisNetCull, Log, All);
diff --git a/Engine/Source/Runtime/Experimental/Iris/Core/Public/Iris/ReplicationSystem/Filtering/NetObjectGridFilter.h b/Engine/Source/Runtime/Experimental/Iris/Core/Public/Iris/ReplicationSystem/Filtering/NetObjectGridFilter.h
--- a/Engine/Source/Runtime/Experimental/Iris/Core/Public/Iris/ReplicationSystem/Filtering/NetObjectGridFilter.h
+++ b/Engine/Source/Runtime/Experimental/Iris/Core/Public/Iris/ReplicationSystem/Filtering/NetObjectGridFilter.h
@@ -287,8 +287,6 @@ private:
 
 	// Cached state from UpdateCellInfoForObject().
 	FDebugUpdateCellInfo DebugUpdateCellInfo;
-
-	const UE::Net::FNetCullDistanceOverrides* NetCullDistanceOverrides = nullptr;
 };
 
 /**
diff --git a/Engine/Source/Runtime/Experimental/Iris/Core/Public/Iris/ReplicationSystem/NetCullDistanceOverrides.h b/Engine/Source/Runtime/Experimental/Iris/Core/Public/Iris/ReplicationSystem/NetCullDistanceOverrides.h
deleted file mode 100644
diff --git a/Engine/Source/Runtime/Experimental/Iris/Core/Public/Iris/ReplicationSystem/ReplicationSystem.h b/Engine/Source/Runtime/Experimental/Iris/Core/Public/Iris/ReplicationSystem/ReplicationSystem.h
--- a/Engine/Source/Runtime/Experimental/Iris/Core/Public/Iris/ReplicationSystem/ReplicationSystem.h
+++ b/Engine/Source/Runtime/Experimental/Iris/Core/Public/Iris/ReplicationSystem/ReplicationSystem.h
@@ -665,13 +665,6 @@ public:
 	 */
 	IRISCORE_API void MarkDirty(FNetRefHandle Handle);
 
-	/**
-	 * Retrieve the NetCullDistanceOverrides instance which holds cull distances for objects that have explicitly overriden it.
-	 * @return The NetCullDistanceOverrides instance.
-	 * @see UE::Net::FNetCullDistanceOverrides
-	 */
-	IRISCORE_API const UE::Net::FNetCullDistanceOverrides& GetNetCullDistanceOverrides() const;
-
 	/**
 	 * Retrieve the WorldLocations instance which holds world locations for all objects that support it. 
 	 * @return The WorldLocations instance.
@@ -701,14 +694,14 @@ public:
 
 	int32 GetPIEInstanceID() const { return PIEInstanceID; }
 
-	/** Set the squared cull distance for an object. This can be used by prioritizers and filters like UNetObjectGridFilter for example. For Actors this will override, not overwrite, the NetCullDistanceSquared property. */
-	IRISCORE_API void SetCullDistanceSqrOverride(FNetRefHandle Handle, float DistSqr);
+	/** Set the cull distance for an object temporarily instead of the default cull distance value associated with the object (ex: @see AActor::NetCullDistanceSquared) */
+	IRISCORE_API void SetCullDistanceOverride(FNetRefHandle Handle, float Dist);
 
-	/** Clears any previously set squared cull distance for an object. For Actors this will cause affected code to respect the NetCullDistanceSquared property. */
-	IRISCORE_API void ClearCullDistanceSqrOverride(FNetRefHandle Handle);
+	/** Clears any previously set cull distance override for an object. This will cause the objects to revert to their internal cull distance property. */
+	IRISCORE_API void ClearCullDistanceOverride(FNetRefHandle Handle);
 
-	/** Returns the previously set squared cull distance for an object, or DefaultValue if it wasn't or the Handle isn't valid. */
-	IRISCORE_API float GetCullDistanceSqrOverride(FNetRefHandle Handle, float DefaultValue = -1.0f) const;
+	/** Returns the current cull distance for an object or DefaultValue if it doesn't have a cull distance or is not a replicated object. */
+	IRISCORE_API float GetCullDistance(FNetRefHandle Handle, float DefaultValue = -1.0f) const;
 
 	/** Returns elapsed time in seconds since ReplicatonSystem was created */
 	double GetElapsedTime() const { return ElapsedTime; }
diff --git a/Engine/Source/Runtime/Experimental/Iris/Core/Public/Iris/ReplicationSystem/WorldLocations.h b/Engine/Source/Runtime/Experimental/Iris/Core/Public/Iris/ReplicationSystem/WorldLocations.h
--- a/Engine/Source/Runtime/Experimental/Iris/Core/Public/Iris/ReplicationSystem/WorldLocations.h
+++ b/Engine/Source/Runtime/Experimental/Iris/Core/Public/Iris/ReplicationSystem/WorldLocations.h
@@ -4,6 +4,7 @@
 
 #include "CoreTypes.h"
 #include "Containers/ChunkedArray.h"
+#include "Iris/Core/NetChunkedArray.h"
 #include "Net/Core/NetBitArray.h"
 #include "Math/Vector.h"
 #include "WorldLocations.generated.h"
@@ -37,146 +38,208 @@ namespace UE::Net
 struct FWorldLocationsInitParams
 {
 	UE::Net::Private::FInternalNetRefIndex MaxInternalNetRefIndex = 0;
+
+	/** How many world info storage slots to preallocate. */
+	uint32 PreallocatedStorageCount = 0;
 };
 
 class FWorldLocations
 {
-public:
-	void Init(const FWorldLocationsInitParams& InitParams);
-	void Deinit();
-
-	/** Returns whether the object has a valid cached data or not. */
-	inline bool HasInfoForObject(uint32 ObjectIndex) const;
-
-	/** Returns the object's world location if it's valid or a zero vector if it's not. */
-	inline FVector GetWorldLocation(uint32 ObjectIndex) const;
-
-	/** Get the object's last cull distance we stored. */
-	inline float GetCullDistance(uint32 ObjectIndex) const;
+public: 
 
-	/** Contains the cached object data we are storing. */
-	struct FObjectInfo
+	/** Publically available information of a replicated root object */
+	struct FWorldInfo
 	{
 		/** Absolute coordinate of the object */
 		FVector WorldLocation = FVector::ZeroVector;
-		/** Network cull distance of the object */
+
+		/** The current network cull distance of the object */
 		float CullDistance = 0.0f;
 	};
 
 public:
-	// Internal API
-	void UpdateWorldLocation(uint32 ObjectIndex, const FVector& WorldLocation);
 
-	void InitObjectInfoCache(uint32 ObjectIndex);
-	void RemoveObjectInfoCache(uint32 ObjectIndex);
+	void Init(const FWorldLocationsInitParams& InitParams);
+	void PostSendUpdate();
 
-	void SetObjectInfo(uint32 ObjectIndex, const FObjectInfo& ObjectInfo);
-    const FObjectInfo& GetObjectInfo(uint32 ObjectIndex) const;
+	/** Returns whether the object has a valid cached data or not. */
+	bool HasInfoForObject(UE::Net::Private::FInternalNetRefIndex ObjectIndex) const
+	{
+		return ValidInfoIndexes.IsBitSet(ObjectIndex);
+	}
+
+	/** Returns the object's world location if it's valid or a zero vector if it's not. */
+	inline FVector GetWorldLocation(UE::Net::Private::FInternalNetRefIndex ObjectIndex) const;
+
+	/** Get the object's current cull distance. */
+	inline float GetCullDistance(UE::Net::Private::FInternalNetRefIndex ObjectIndex) const;
+
+	/** Return the current stored world information of the given object. Returns a default struct if the object did not register in the world location cache. */
+	inline FWorldLocations::FWorldInfo GetWorldInfo(UE::Net::Private::FInternalNetRefIndex ObjectIndex) const;
+
+	/** Set the mandatory info of a replicated root object*/
+	void SetObjectInfo(UE::Net::Private::FInternalNetRefIndex ObjectIndex, const FVector& Location, float NetCullDistance);
+	
+	/** Assign a world information cache to the replicated object */
+	void InitObjectInfoCache(UE::Net::Private::FInternalNetRefIndex ObjectIndex);
+
+	/** Remove the world information cache of the replicated object */
+	void RemoveObjectInfoCache(UE::Net::Private::FInternalNetRefIndex ObjectIndex);
 
 	/**
-	 * Objects are not necessarily marked as dirty just because they're moving, such as objects attached to other objects. If such objects are spatially filtered they need to update their world locations in order for replication to work as expected.
+	 * Objects are not necessarily marked as dirty just because they're moving, such as objects attached to other objects. 
+	 * If such objects are spatially filtered they need to update their world locations in order for replication to work as expected.
 	 * Use SetObjectRequiresFrequentWorldLocationUpdate to force frequent world location update on an object.
 	 */
-	inline void SetObjectRequiresFrequentWorldLocationUpdate(uint32 ObjectIndex, bool bRequiresFrequentUpdate);
+	void SetObjectRequiresFrequentWorldLocationUpdate(UE::Net::Private::FInternalNetRefIndex ObjectIndex, bool bRequiresFrequentUpdate)
+	{
+		ObjectsRequiringFrequentWorldLocationUpdate.SetBitValue(ObjectIndex, ValidInfoIndexes.GetBit(ObjectIndex) && bRequiresFrequentUpdate);
+	}
 
 	/** Returns whether an object requires frequent world location updates. */
-	inline bool GetObjectRequiresFrequentWorldLocationUpdate(uint32 ObjectIndex) const;
+	bool GetObjectRequiresFrequentWorldLocationUpdate(UE::Net::Private::FInternalNetRefIndex ObjectIndex) const
+	{
+		return ObjectsRequiringFrequentWorldLocationUpdate.GetBit(ObjectIndex);
+	}
 
-	inline FNetBitArrayView GetObjectsRequiringFrequentWorldLocationUpdate() const;
+	/** 
+	 * Add a temporary net cull distance that will have priority over the regular net cull distance.
+	 * Returns true if the object had registered to use the world location cache and can store the override.
+	 */
+	bool SetCullDistanceOverride(UE::Net::Private::FInternalNetRefIndex ObjectIndex, float CullDistance);
 
-	void ResetObjectsWithDirtyInfo();
-	inline FNetBitArrayView GetObjectsWithDirtyInfo() const;
+	/** 
+	 * Remove the temporary net cull distance override and instead use the regular net cull distance 
+	 * Returns true if the object had registered to use the world location cache and had an override value previously set
+	 */
+	bool ClearCullDistanceOverride(UE::Net::Private::FInternalNetRefIndex ObjectIndex);
+
+	/** Returns true if the object was set a cull distance override and is using it instead of his default cull distance value */
+	bool HasCullDistanceOverride(UE::Net::Private::FInternalNetRefIndex ObjectIndex) const
+	{
+		return ValidInfoIndexes.IsBitSet(ObjectIndex) ? GetObjectInfo(ObjectIndex).CullDistanceOverride != FLT_MAX : false;
+	}
+
+	/** Returns the list of objects that need to check for a location change every frame */
+	const FNetBitArrayView GetObjectsRequiringFrequentWorldLocationUpdate() const { return MakeNetBitArrayView(ObjectsRequiringFrequentWorldLocationUpdate); }
+
+	/** Returns the list of objects that changed world location or cull distance this frame */
+	const FNetBitArrayView GetObjectsWithDirtyInfo() const { return MakeNetBitArrayView(ObjectsWithDirtyInfo); }
 
 	/** Returns the list of objects that registered world location information */
 	const FNetBitArrayView GetObjectsWithWorldInfo() const { return MakeNetBitArrayView(ValidInfoIndexes); }
 
+	/** Reset the list of objects that changed location or cull distance */
+	void ResetObjectsWithDirtyInfo();
+
+	/** Debug tool to track when its legal to modify the DirtyInfo list. */
+	void LockDirtyInfoList(bool bLock);
+
 	/** Return the world boundaries (min and max position). */
 	const FVector& GetWorldMinPos() const { return MinWorldPos; };
 	const FVector& GetWorldMaxPos() const { return MaxWorldPos; };
 	
 	/** Return a position clamped to the configured world boundary. */
-	FVector ClampPositionToBoundary(const FVector& Position);
+	FVector ClampPositionToBoundary(const FVector& Position)
+	{
+		return Position.BoundToBox(GetWorldMinPos(), GetWorldMaxPos());
+	}
 
 	void OnMaxInternalNetRefIndexIncreased(UE::Net::Private::FInternalNetRefIndex NewMaxInternalIndex);
 
 private:
+
+	/** Contains the cached object data we are storing. */
+	struct FObjectInfo
+	{
+		/** Absolute coordinate of the object */
+		FVector WorldLocation = FVector::ZeroVector;
+
+		/** The default network cull distance of the object */
+		float CullDistance = 0.0f;
+
+		/** The optional temporary cull distance override. Max means it is not used */
+		float CullDistanceOverride = FLT_MAX;
+	};
+
 	enum : uint32
 	{
-		BytesPerLocationChunk = 65536U,
+		StorageElementsPerChunk = 65536U / sizeof(FObjectInfo),
 	};
 
+private:
+
+	const int32 GetStorageIndex(UE::Net::Private::FInternalNetRefIndex ObjectIndex) const
+	{
+		return StorageIndexes[ObjectIndex];
+	}
+
+	const FObjectInfo& GetObjectInfo(UE::Net::Private::FInternalNetRefIndex ObjectIndex) const
+	{
+		return StoredObjectInfo[GetStorageIndex(ObjectIndex)];
+	}
+
+	FObjectInfo& GetObjectInfo(UE::Net::Private::FInternalNetRefIndex ObjectIndex)
+	{
+		return StoredObjectInfo[GetStorageIndex(ObjectIndex)];
+	}
+
+private:
+
 	/** Set bits indicate that we have stored information for this internal object index */
 	FNetBitArray ValidInfoIndexes;
+	
 	/** Set bits indicate that the world location or net cull distance has changed since last update */
 	FNetBitArray ObjectsWithDirtyInfo;
+	
 	/** Set bits indicate that the object requires frequent world location updates */
 	FNetBitArray ObjectsRequiringFrequentWorldLocationUpdate;
 
-	TChunkedArray<FObjectInfo, BytesPerLocationChunk> StoredObjectInfo;
+	/** Map that returns the storage index for the world info of a registered object. */
+	TArray<int32> StorageIndexes;
+
+	/** List detailing if a given index slot is free to be used to store world info. */
+	FNetBitArray ReservedStorageSlot;
+
+	TNetChunkedArray<FObjectInfo, StorageElementsPerChunk> StoredObjectInfo;
 
 	/** World boundaries (min and max position). */
 	FVector MinWorldPos;
 	FVector MaxWorldPos;
+
+	/** Controls if the dirty list can be modified */
+	bool bLockdownDirtyList = false;
 };
 
-inline bool FWorldLocations::HasInfoForObject(uint32 ObjectIndex) const
+inline FVector FWorldLocations::GetWorldLocation(UE::Net::Private::FInternalNetRefIndex ObjectIndex) const
 {
-	return ValidInfoIndexes.IsBitSet(ObjectIndex);
+	return ValidInfoIndexes.IsBitSet(ObjectIndex) ? GetObjectInfo(ObjectIndex).WorldLocation : FVector::Zero();
 }
 
-inline const FWorldLocations::FObjectInfo& FWorldLocations::GetObjectInfo(uint32 ObjectIndex) const
+inline float FWorldLocations::GetCullDistance(UE::Net::Private::FInternalNetRefIndex ObjectIndex) const
 {
 	if (ValidInfoIndexes.IsBitSet(ObjectIndex))
 	{
-		return StoredObjectInfo[ObjectIndex];
+		const FObjectInfo& Info = GetObjectInfo(ObjectIndex);
+		return Info.CullDistanceOverride == FLT_MAX ? Info.CullDistance : Info.CullDistanceOverride;
 	}
 
-	static FObjectInfo EmptyInfo;
-	return EmptyInfo;
+	return 0.0f;
 }
 
-inline FVector FWorldLocations::GetWorldLocation(uint32 ObjectIndex) const
+inline FWorldLocations::FWorldInfo FWorldLocations::GetWorldInfo(UE::Net::Private::FInternalNetRefIndex ObjectIndex) const
 {
 	if (ValidInfoIndexes.IsBitSet(ObjectIndex))
 	{
-		return StoredObjectInfo[ObjectIndex].WorldLocation;
+		const FObjectInfo& Info = GetObjectInfo(ObjectIndex);
+		return FWorldInfo
+		{ 
+			.WorldLocation = Info.WorldLocation, 
+			.CullDistance = Info.CullDistanceOverride == FLT_MAX ? Info.CullDistance : Info.CullDistanceOverride,
+		};
 	}
-	else
-	{
-		return FVector::Zero();
-	}
-}
-
-inline float FWorldLocations::GetCullDistance(uint32 ObjectIndex) const
-{
-	if (ValidInfoIndexes.IsBitSet(ObjectIndex))
-	{
-		return StoredObjectInfo[ObjectIndex].CullDistance;
-	}
-	else
-	{
-		return 0.0f;
-	}
-}
 
-inline void FWorldLocations::SetObjectRequiresFrequentWorldLocationUpdate(uint32 ObjectIndex, bool bRequiresFrequentUpdate)
-{
-	ObjectsRequiringFrequentWorldLocationUpdate.SetBitValue(ObjectIndex, ValidInfoIndexes.GetBit(ObjectIndex) && bRequiresFrequentUpdate);
-}
-
-inline bool FWorldLocations::GetObjectRequiresFrequentWorldLocationUpdate(uint32 ObjectIndex) const
-{
-	return ObjectsRequiringFrequentWorldLocationUpdate.GetBit(ObjectIndex);
+	return FWorldInfo();
 }
 
-inline FNetBitArrayView FWorldLocations::GetObjectsRequiringFrequentWorldLocationUpdate() const
-{
-	return MakeNetBitArrayView(ObjectsRequiringFrequentWorldLocationUpdate);
-}
-
-inline FNetBitArrayView FWorldLocations::GetObjectsWithDirtyInfo() const
-{
-	return MakeNetBitArrayView(ObjectsWithDirtyInfo);
-}
-
-}
+} // end namespace UE::Net
