Dataflow : make sub-graphs can be properly de-serialized
- Make sure all search through the EdNodes also include the dataflow sub-graphs

diff --git a/Engine/Source/Runtime/Experimental/Dataflow/Engine/Private/Dataflow/DataflowObject.cpp b/Engine/Source/Runtime/Experimental/Dataflow/Engine/Private/Dataflow/DataflowObject.cpp
--- a/Engine/Source/Runtime/Experimental/Dataflow/Engine/Private/Dataflow/DataflowObject.cpp
+++ b/Engine/Source/Runtime/Experimental/Dataflow/Engine/Private/Dataflow/DataflowObject.cpp
@@ -27,6 +27,85 @@ namespace UE::Dataflow::CVars
 			ECVF_Default);
 }
 
+namespace UE::Dataflow::Private
+{
+	// EdNode visitor helper function (non-const version)
+	template<typename ActionType>
+	bool ForEachEdGraphNodeInEdGraph(UEdGraph* EdGraph, ActionType Action)
+	{
+		if (!EdGraph)
+		{
+			return false;
+		}
+		for (UEdGraphNode* EdNode : EdGraph->Nodes)
+		{
+			if (EdNode)
+			{
+				if (!Action(*EdGraph, *EdNode))
+				{
+					return false;
+				}
+			}
+		}
+		return true;
+	}
+
+	// EdNode visitor helper function (const version)
+	template<typename ActionType>
+	bool ForEachEdGraphNodeInEdGraph(const UEdGraph* EdGraph, ActionType Action)
+	{
+		if (!EdGraph)
+		{
+			return false;
+		}
+		for (const UEdGraphNode* EdNode : EdGraph->Nodes)
+		{
+			if (EdNode)
+			{
+				if (!Action(*EdGraph, *EdNode))
+				{
+					return false;
+				}
+			}
+		}
+		return true;
+	}
+
+	// EdNode visitor helper function (non-const version)
+	// this explore the subgraphs as well
+	template<typename ActionType>
+	void ForEachEdGraphNodeInDataflowAsset(UDataflow* DataflowAsset, ActionType Action)
+	{
+		if (ForEachEdGraphNodeInEdGraph(DataflowAsset, Action))
+		{
+			for (UEdGraph* SubGraph : DataflowAsset->GetSubGraphs())
+			{
+				if (!ForEachEdGraphNodeInEdGraph(SubGraph, Action))
+				{
+					return;
+				}
+			}
+		}
+	}
+
+	// EdNode visitor helper function (const version)
+	// this explore the subgraphs as well
+	template<typename ActionType>
+	void ForEachEdGraphNodeInDataflowAsset(const UDataflow* DataflowAsset, ActionType Action)
+	{
+		if (ForEachEdGraphNodeInEdGraph(DataflowAsset, Action))
+		{
+			for (const UEdGraph* SubGraph : DataflowAsset->GetSubGraphs())
+			{
+				if (!ForEachEdGraphNodeInEdGraph(SubGraph, Action))
+				{
+					return;
+				}
+			}
+		}
+	}
+}
+
 FDataflowAssetEdit::FDataflowAssetEdit(UDataflow* InAsset, FPostEditFunctionCallback InCallback)
 	: PostEditCallback(InCallback)
 	, Asset(InAsset)
@@ -97,22 +176,22 @@ void UDataflow::PostLoad()
 #if WITH_EDITOR
 	const TSet<FName>& DisabledNodes = Dataflow->GetDisabledNodes();
 
-	for (UEdGraphNode* EdNode : Nodes)
-	{
-		UDataflowEdNode* DataflowEdNode = Cast<UDataflowEdNode>(EdNode);
-
-		// Not all nodes are UDataflowEdNode (There is now UDataflowEdNodeComment)
-		if (DataflowEdNode)
+	UE::Dataflow::Private::ForEachEdGraphNodeInDataflowAsset(this, 
+		[this, &DisabledNodes](UEdGraph& EdGraph, UEdGraphNode& EdNode)
 		{
-			DataflowEdNode->SetDataflowGraph(Dataflow);
-			DataflowEdNode->UpdatePinsFromDataflowNode();
-		}
+			// Not all nodes are UDataflowEdNode (There is now UDataflowEdNodeComment)
+			if (UDataflowEdNode* DataflowEdNode = Cast<UDataflowEdNode>(&EdNode))
+			{
+				DataflowEdNode->SetDataflowGraph(Dataflow);
+				DataflowEdNode->UpdatePinsFromDataflowNode();
+			}
 
-		if (DisabledNodes.Contains(FName(EdNode->GetName())))
-		{
-			EdNode->SetEnabledState(ENodeEnabledState::Disabled);
-		}
-	}
+			if (DisabledNodes.Contains(FName(EdNode.GetName())))
+			{
+				EdNode.SetEnabledState(ENodeEnabledState::Disabled);
+			}
+			return true; // visit all nodes
+		});
 
 	// Resync connections (nodes might have redirected connections
 	for (const UE::Dataflow::FLink& Link : Dataflow->GetConnections())
@@ -196,32 +275,44 @@ void UDataflow::Serialize(FArchive& Ar)
 
 TObjectPtr<const UDataflowEdNode> UDataflow::FindEdNodeByDataflowNodeGuid(const FGuid& Guid) const
 {
-	for (const UEdGraphNode* const EdNode : Nodes)
-	{
-		if (const UDataflowEdNode* const DataflowEdNode = Cast<UDataflowEdNode>(EdNode))
+	TObjectPtr<const UDataflowEdNode> FoundNode(nullptr);
+
+	UE::Dataflow::Private::ForEachEdGraphNodeInDataflowAsset(this, 
+		[&Guid, &FoundNode](const UEdGraph& EdGraph, const UEdGraphNode& EdNode)
 		{
-			if (DataflowEdNode->GetDataflowNodeGuid() == Guid)
+			if (const UDataflowEdNode* const DataflowEdNode = Cast<UDataflowEdNode>(&EdNode))
 			{
-				return TObjectPtr<const UDataflowEdNode>(DataflowEdNode);
+				if (DataflowEdNode->GetDataflowNodeGuid() == Guid)
+				{
+					FoundNode = DataflowEdNode;
+					return false; // early exit
+				}
 			}
-		}
-	}
-	return TObjectPtr<const UDataflowEdNode>(nullptr);
+			return true; // continue visiting 
+		});
+
+	return FoundNode;
 }
 
 TObjectPtr<UDataflowEdNode> UDataflow::FindEdNodeByDataflowNodeGuid(const FGuid& Guid)
 {
-	for (UEdGraphNode* const EdNode : Nodes)
-	{
-		if (UDataflowEdNode* const DataflowEdNode = Cast<UDataflowEdNode>(EdNode))
+	TObjectPtr<UDataflowEdNode> FoundNode(nullptr);
+
+	UE::Dataflow::Private::ForEachEdGraphNodeInDataflowAsset(this,
+		[&Guid, &FoundNode](UEdGraph& EdGraph, UEdGraphNode& EdNode)
 		{
-			if (DataflowEdNode->GetDataflowNodeGuid() == Guid)
+			if (UDataflowEdNode* DataflowEdNode = Cast<UDataflowEdNode>(&EdNode))
 			{
-				return TObjectPtr<UDataflowEdNode>(DataflowEdNode);
+				if (DataflowEdNode->GetDataflowNodeGuid() == Guid)
+				{
+					FoundNode = DataflowEdNode;
+					return false; // early exit
+				}
 			}
-		}
-	}
-	return TObjectPtr<UDataflowEdNode>(nullptr);
+			return true; // continue visiting 
+		});
+
+	return FoundNode;
 }
 
 UDataflow* UDataflow::GetDataflowAssetFromEdGraph(UEdGraph* EdGraph)
diff --git a/Engine/Source/Runtime/Experimental/Dataflow/Engine/Public/Dataflow/DataflowObject.h b/Engine/Source/Runtime/Experimental/Dataflow/Engine/Public/Dataflow/DataflowObject.h
--- a/Engine/Source/Runtime/Experimental/Dataflow/Engine/Public/Dataflow/DataflowObject.h
+++ b/Engine/Source/Runtime/Experimental/Dataflow/Engine/Public/Dataflow/DataflowObject.h
@@ -227,8 +227,9 @@ private:
 private:
 	/** 
 	* List of Dataflow SubGraphs 
-	* in editor they also exists in the SubGraphs parent class property
+	* In editor they also exists in the SubGraphs parent class property
 	*/
+	UPROPERTY()
 	TArray<TObjectPtr<UDataflowSubGraph>> DataflowSubGraphs; 
 
 private:
