Dataflow : make sure to get the dataflow asset from the correct graph when handling pin connection changes

diff --git a/Engine/Source/Runtime/Experimental/Dataflow/Engine/Private/Dataflow/DataflowEdNode.cpp b/Engine/Source/Runtime/Experimental/Dataflow/Engine/Private/Dataflow/DataflowEdNode.cpp
--- a/Engine/Source/Runtime/Experimental/Dataflow/Engine/Private/Dataflow/DataflowEdNode.cpp
+++ b/Engine/Source/Runtime/Experimental/Dataflow/Engine/Private/Dataflow/DataflowEdNode.cpp
@@ -420,10 +420,10 @@ void UDataflowEdNode::PinConnectionListChanged(UEdGraphPin* Pin)
 										return MatchesConnection(LinkedCon, ConnectedInput);
 									});
 							});
-					UDataflow* const EdGraph = CastChecked<UDataflow>(GetGraph());
+					UDataflow* const DataflowAsset = UDataflow::GetDataflowAssetFromEdGraph(GetGraph());
 					for (FDataflowInput* const ConnectedInput : InputsToDisconnect)
 					{
-						const TObjectPtr<UDataflowEdNode> InputEdNode = EdGraph->FindEdNodeByDataflowNodeGuid(ConnectedInput->GetOwningNodeGuid());
+						const TObjectPtr<UDataflowEdNode> InputEdNode = DataflowAsset->FindEdNodeByDataflowNodeGuid(ConnectedInput->GetOwningNodeGuid());
 						if (ensure(InputEdNode))
 						{
 							UEdGraphPin* const InputPin = InputEdNode->FindPin(ConnectedInput->GetName(), EEdGraphPinDirection::EGPD_Input);
