[ChaosVD] Initial AutoRTFM Support

Currently attempting to start a CVD recording if AutoRTFM is enabled will crash because we use atomics and locks. This CL replaces the uses of RWLocks to the new Transactionally safe lock data type.

That said, CVD depends on TraceLog which also uses atomics, so even with the fixes in the CVD side, the crash is still happening. We also have the issue of that we rely on the order and timing on how the data is serialized across multiple threads, and the Abort Callbacks provided by AutoRTFM are done in the game thread after the fact. This is an issue because it means we will need to implement an additional buffer layer to remove rolled back data before tracing it (instead of doing it on the spot), or during load of a CVD file, find data we already processed and removed if the transaction that created it failed. Both options are not free perf wise and not trivial to implement.

As a workaround this CL does

- Frame and Solver Stages, Start and End events, are traced regardless if there was a rollback. (This is unlikely to happen with these, as currently physics is not transactable)
- Actual binary data, is traced on the spot, but the event that allows the data to be used, is only traced if the transaction is completely successful. This causes other potential issues described in a code comment
- Added a warning for when a CVD file recorded in a build with AutoRTFM enabled is loaded, indicating there might bee accuracy issues

Once we coordinate with the team owning Trace log to get a solution for that system, we can start working on properly supporting AutoRTFM regardless the context in which any data is traced.

diff --git a/Engine/Plugins/ChaosVD/Source/ChaosVD/Private/Trace/ChaosVDTraceAnalyzer.cpp b/Engine/Plugins/ChaosVD/Source/ChaosVD/Private/Trace/ChaosVDTraceAnalyzer.cpp
--- a/Engine/Plugins/ChaosVD/Source/ChaosVD/Private/Trace/ChaosVDTraceAnalyzer.cpp
+++ b/Engine/Plugins/ChaosVD/Source/ChaosVD/Private/Trace/ChaosVDTraceAnalyzer.cpp
@@ -3,6 +3,8 @@
 #include "Trace/ChaosVDTraceAnalyzer.h"
 
 #include "ChaosVDModule.h"
+#include "Containers/Ticker.h"
+#include "Misc/MessageDialog.h"
 #include "Trace/ChaosVDTraceProvider.h"
 #include "TraceServices/Model/AnalysisSession.h"
 #include "TraceServices/Model/Frames.h"
@@ -26,6 +28,8 @@ void FChaosVDTraceAnalyzer::OnAnalysisBegin(const FOnAnalysisContext& Context)
 	Builder.RouteEvent(RouteId_ChaosVDNonSolverLocation, "ChaosVDLogger", "ChaosVDNonSolverLocation");
 	Builder.RouteEvent(RouteId_ChaosVDNonSolverTransform, "ChaosVDLogger", "ChaosVDNonSolverTransform");
 	Builder.RouteEvent(RouteId_ChaosVDNetworkTickOffset, "ChaosVDLogger", "ChaosVDNetworkTickOffset");
+	Builder.RouteEvent(RouteId_ChaosVDRolledBackDataID, "ChaosVDLogger", "ChaosVDRolledBackDataID");
+	Builder.RouteEvent(RouteId_ChaosVDUsesAutoRTFM, "ChaosVDLogger", "ChaosVDUsesAutoRTFM");
 
 	Builder.RouteEvent(RouteId_BeginFrame, "Misc", "BeginFrame");
 	Builder.RouteEvent(RouteId_EndFrame, "Misc", "EndFrame");
@@ -272,6 +276,30 @@ bool FChaosVDTraceAnalyzer::OnEvent(uint16 RouteId, EStyle Style, const FOnEvent
 
 			ChaosVDTraceProvider->GetCurrentTickOffsetsBySolverID().FindOrAdd(SolverID, TickOffset);
 
+			break;
+		}
+	case RouteId_ChaosVDRolledBackDataID:
+		{
+			const int32 DataID = EventData.GetValue<int32>("DataID");
+			ChaosVDTraceProvider->RemoveUnprocessedData(DataID);
+
+			break;
+		}	
+	case RouteId_ChaosVDUsesAutoRTFM:
+		{
+			const bool bUsingAutoRTFM = EventData.GetValue<bool>("bUsingAutoRTFM");
+
+			if (bUsingAutoRTFM)
+			{
+				FText MissingDataProcessorsMessage = NSLOCTEXT("ChaosVisualDebugger","AutoRTFMWarningMessage", "This recording was made with AutoRTFM enabled. \n\nAutoRTFM is not fully supported and framing/timing of recorded data during transactions might be off.");
+
+				FTSTicker::GetCoreTicker().AddTicker(FTickerDelegate::CreateLambda([MissingDataProcessorsMessage](float DeltaTime)
+				{
+					FMessageDialog::Open(EAppMsgCategory::Warning, EAppMsgType::Ok, MissingDataProcessorsMessage, NSLOCTEXT("ChaosVisualDebugger", "AutoRTFMWarningMessageTitle", "Partially unsupported CVD Recording"));
+					return false;
+				}));
+			}
+
 			break;
 		}
 
diff --git a/Engine/Plugins/ChaosVD/Source/ChaosVD/Private/Trace/ChaosVDTraceAnalyzer.h b/Engine/Plugins/ChaosVD/Source/ChaosVD/Private/Trace/ChaosVDTraceAnalyzer.h
--- a/Engine/Plugins/ChaosVD/Source/ChaosVD/Private/Trace/ChaosVDTraceAnalyzer.h
+++ b/Engine/Plugins/ChaosVD/Source/ChaosVD/Private/Trace/ChaosVDTraceAnalyzer.h
@@ -31,7 +31,10 @@ public:
 	virtual void OnAnalysisEnd() override;
 	virtual bool OnEvent(uint16 RouteId, EStyle Style, const FOnEventContext& Context) override;
 
-	FChaosVDTraceAnalysisComplete& OnAnalysisComplete() { return ChaosVDTraceAnalysisCompleteDelegate; }
+	FChaosVDTraceAnalysisComplete& OnAnalysisComplete()
+	{
+		return ChaosVDTraceAnalysisCompleteDelegate;
+	}
 
 private:
 
@@ -51,6 +54,8 @@ private:
 		RouteId_ChaosVDNonSolverLocation,
 		RouteId_ChaosVDNonSolverTransform,
 		RouteId_ChaosVDNetworkTickOffset,
+		RouteId_ChaosVDRolledBackDataID,
+		RouteId_ChaosVDUsesAutoRTFM,
 	};
 
 	TraceServices::IAnalysisSession& Session;
diff --git a/Engine/Plugins/ChaosVD/Source/ChaosVD/Private/Trace/ChaosVDTraceProvider.cpp b/Engine/Plugins/ChaosVD/Source/ChaosVD/Private/Trace/ChaosVDTraceProvider.cpp
--- a/Engine/Plugins/ChaosVD/Source/ChaosVD/Private/Trace/ChaosVDTraceProvider.cpp
+++ b/Engine/Plugins/ChaosVD/Source/ChaosVD/Private/Trace/ChaosVDTraceProvider.cpp
@@ -182,6 +182,12 @@ FChaosVDTraceProvider::FBinaryDataContainer& FChaosVDTraceProvider::FindOrAddUnp
 	}
 }
 
+void FChaosVDTraceProvider::RemoveUnprocessedData(const int32 DataID)
+{
+	// The removal call should always come before the data is processed
+	ensure(UnprocessedDataByID.Remove(DataID) != 0);
+}
+
 bool FChaosVDTraceProvider::ProcessBinaryData(const int32 DataID)
 {
 	RegisterDefaultDataProcessorsIfNeeded();
diff --git a/Engine/Plugins/ChaosVD/Source/ChaosVD/Public/Trace/ChaosVDTraceProvider.h b/Engine/Plugins/ChaosVD/Source/ChaosVD/Public/Trace/ChaosVDTraceProvider.h
--- a/Engine/Plugins/ChaosVD/Source/ChaosVD/Public/Trace/ChaosVDTraceProvider.h
+++ b/Engine/Plugins/ChaosVD/Source/ChaosVD/Public/Trace/ChaosVDTraceProvider.h
@@ -106,6 +106,7 @@ private:
 	void HandleAnalysisComplete();
 	
 	FBinaryDataContainer& FindOrAddUnprocessedData(const int32 DataID);
+	void RemoveUnprocessedData(const int32 DataID);
 
 	bool ProcessBinaryData(const int32 DataID);
 
diff --git a/Engine/Source/Runtime/Experimental/Chaos/Private/ChaosVisualDebugger/ChaosVDSerializedNameTable.cpp b/Engine/Source/Runtime/Experimental/Chaos/Private/ChaosVisualDebugger/ChaosVDSerializedNameTable.cpp
--- a/Engine/Source/Runtime/Experimental/Chaos/Private/ChaosVisualDebugger/ChaosVDSerializedNameTable.cpp
+++ b/Engine/Source/Runtime/Experimental/Chaos/Private/ChaosVisualDebugger/ChaosVDSerializedNameTable.cpp
@@ -19,14 +19,14 @@ namespace Chaos::VisualDebugger
 		}
 
 		{
-			FReadScopeLock ReadLock(NamesByIDLock);
+			UE::TReadScopeLock ReadLock(NamesByIDLock);
 			if (NamesByID.Contains(NameID))
 			{
 				return NameID;
 			}
 		}
 		{
-			FWriteScopeLock WriteLock(NamesByIDLock);			
+			UE::TWriteScopeLock WriteLock(NamesByIDLock);			
 			NamesByID.Add(NameID, Name);
 		}
 
@@ -49,7 +49,7 @@ namespace Chaos::VisualDebugger
 
 	uint64 FChaosVDSerializableNameTable::AddNameToTable(const FChaosVDSerializedNameEntry& InNameEntry)
 	{
-		FWriteScopeLock WriteLock(NamesByIDLock);
+		UE::TWriteScopeLock WriteLock(NamesByIDLock);
 		NamesByID.Add(InNameEntry.NameID, FName(InNameEntry.Name, InNameEntry.NameNumber));
 
 		return InNameEntry.NameID;
@@ -64,7 +64,7 @@ namespace Chaos::VisualDebugger
 		}
 
 		{
-			FReadScopeLock ReadLock(NamesByIDLock);
+			UE::TReadScopeLock ReadLock(NamesByIDLock);
 			if (const FName* FoundName = NamesByID.Find(NameID))
 			{
 				return *FoundName;
@@ -77,7 +77,7 @@ namespace Chaos::VisualDebugger
 	void FChaosVDSerializableNameTable::ResetTable()
 	{
 		{
-			FWriteScopeLock WriteLock(NamesByIDLock);
+			UE::TWriteScopeLock WriteLock(NamesByIDLock);
 			NamesByID.Empty();
 		}
 	}
diff --git a/Engine/Source/Runtime/Experimental/Chaos/Private/ChaosVisualDebugger/ChaosVisualDebuggerTrace.cpp b/Engine/Source/Runtime/Experimental/Chaos/Private/ChaosVisualDebugger/ChaosVisualDebuggerTrace.cpp
--- a/Engine/Source/Runtime/Experimental/Chaos/Private/ChaosVisualDebugger/ChaosVisualDebuggerTrace.cpp
+++ b/Engine/Source/Runtime/Experimental/Chaos/Private/ChaosVisualDebugger/ChaosVisualDebuggerTrace.cpp
@@ -27,6 +27,7 @@
 #include "DataWrappers/ChaosVDQueryDataWrappers.h"
 #include "HAL/CriticalSection.h"
 #include "Misc/ScopeRWLock.h"
+#include "Misc/TransactionallySafeRWLock.h"
 #include "Templates/UniquePtr.h"
 
 UE_TRACE_EVENT_DEFINE(ChaosVDLogger, ChaosVDSolverFrameStart)
@@ -46,6 +47,8 @@ UE_TRACE_EVENT_DEFINE(ChaosVDLogger, ChaosVDDummyEvent)
 UE_TRACE_EVENT_DEFINE(ChaosVDLogger, ChaosVDNonSolverLocation)
 UE_TRACE_EVENT_DEFINE(ChaosVDLogger, ChaosVDNonSolverTransform)
 UE_TRACE_EVENT_DEFINE(ChaosVDLogger, ChaosVDNetworkTickOffset)
+UE_TRACE_EVENT_DEFINE(ChaosVDLogger, ChaosVDRolledBackDataID)
+UE_TRACE_EVENT_DEFINE(ChaosVDLogger, ChaosVDUsesAutoRTFM)
 
 namespace Chaos::VisualDebugger::Cvars
 {
@@ -65,8 +68,8 @@ namespace Chaos::VisualDebugger::Cvars
 /** Struct where we keep track of the geometry we are tracing */
 struct FChaosVDGeometryTraceContext
 {
-	FRWLock TracedGeometrySetLock;
-	FRWLock CachedGeometryHashesLock;
+	FTransactionallySafeRWLock TracedGeometrySetLock;
+	FTransactionallySafeRWLock CachedGeometryHashesLock;
 	TSet<uint32> GeometryTracedIDs;
 
 	uint32 GetGeometryHashForImplicit(const Chaos::FImplicitObject* Implicit)
@@ -77,7 +80,7 @@ struct FChaosVDGeometryTraceContext
 		}
 
 		{
-			FReadScopeLock ReadLock(CachedGeometryHashesLock);
+			UE::TReadScopeLock ReadLock(CachedGeometryHashesLock);
 			if (uint32* FoundHashPtr = CachedGeometryHashes.Find((void*) Implicit))
 			{
 				return *FoundHashPtr;
@@ -87,7 +90,7 @@ struct FChaosVDGeometryTraceContext
 		{
 			uint32 Hash = Implicit->GetTypeHash();
 
-			FWriteScopeLock WriteLock(CachedGeometryHashesLock);
+			UE::TWriteScopeLock WriteLock(CachedGeometryHashesLock);
 			CachedGeometryHashes.Add((void*)Implicit, Hash);
 			return Hash;
 		}
@@ -100,7 +103,7 @@ struct FChaosVDGeometryTraceContext
 			return;
 		}
 
-		FWriteScopeLock WriteLock(CachedGeometryHashesLock);
+		UE::TWriteScopeLock WriteLock(CachedGeometryHashesLock);
 		CachedGeometryHashes.Remove((void*)Implicit);
 	}
 
@@ -221,10 +224,17 @@ void FChaosVisualDebuggerTrace::TraceParticleDestroyed(const Chaos::FGeometryPar
 		return;
 	}
 
-	UE_TRACE_LOG(ChaosVDLogger, ChaosVDParticleDestroyed, ChaosVDChannel)
-		<< ChaosVDParticleDestroyed.SolverID(CVDContextData->Id)
-		<< ChaosVDParticleDestroyed.Cycle(FPlatformTime::Cycles64())
-		<< ChaosVDParticleDestroyed.ParticleID(ParticleHandle->UniqueIdx().Idx);
+	int32 ParticleID = ParticleHandle->UniqueIdx().Idx;
+	int32 SolverID = CVDContextData->Id;
+	uint64 Cycle = FPlatformTime::Cycles64();
+
+	UE_AUTORTFM_ONCOMMIT(SolverID, Cycle, ParticleID)
+	{
+		UE_TRACE_LOG(ChaosVDLogger, ChaosVDParticleDestroyed, ChaosVDChannel)
+			<< ChaosVDParticleDestroyed.SolverID(SolverID)
+			<< ChaosVDParticleDestroyed.Cycle(Cycle)
+			<< ChaosVDParticleDestroyed.ParticleID(ParticleID);
+	};
 }
 
 void FChaosVisualDebuggerTrace::TraceParticleClusterChildData(const Chaos::TParticleView<Chaos::TPBDRigidParticles<Chaos::FReal, 3>>& ParticlesView, Chaos::FRigidClustering* ClusteringData, const FChaosVDContext& CVDContextData)
@@ -297,7 +307,7 @@ void FChaosVisualDebuggerTrace::SetupForFullCaptureIfNeeded(int32 SolverID, bool
 
 	if (bOutFullCaptureRequested)
 	{
-		FWriteScopeLock WriteLock(DeltaRecordingStatesLock);
+		UE::TWriteScopeLock WriteLock(DeltaRecordingStatesLock);
 		SolverIDsForDeltaRecording.Remove(SolverID);
 		RequestedFullCaptureSolverIDs.Remove(SolverID);
 	}
@@ -310,7 +320,7 @@ int32 FChaosVisualDebuggerTrace::GetSolverID(Chaos::FPhysicsSolverBase& Solver)
 
 bool FChaosVisualDebuggerTrace::ShouldPerformFullCapture(int32 SolverID)
 {
-	FReadScopeLock ReadLock(DeltaRecordingStatesLock);
+	UE::TReadScopeLock ReadLock(DeltaRecordingStatesLock);
 	int32* FoundSolverID = SolverIDsForDeltaRecording.Find(SolverID);
 
 	// If the solver ID is on the SolverIDsForDeltaRecording set, it means we should NOT perform a full capture
@@ -543,13 +553,16 @@ void FChaosVisualDebuggerTrace::TraceSolverFrameStart(const FChaosVDContext& Con
 	bool bOutIsFullCaptureRequested;
 	SetupForFullCaptureIfNeeded(ContextData.Id, bOutIsFullCaptureRequested);
 
-	UE_TRACE_LOG(ChaosVDLogger, ChaosVDSolverFrameStart, ChaosVDChannel)
-		<< ChaosVDSolverFrameStart.SolverID(ContextData.Id)
-		<< ChaosVDSolverFrameStart.Cycle(FPlatformTime::Cycles64())
-		<< ChaosVDSolverFrameStart.DebugName(*InDebugName, InDebugName.Len())
-		<< ChaosVDSolverFrameStart.IsKeyFrame(bOutIsFullCaptureRequested)
-		<< ChaosVDSolverFrameStart.IsReSimulated(bIsReSimulatedFrame)
-		<< ChaosVDSolverFrameStart.CurrentFrameNumber(FrameNumber);
+	UE_AUTORTFM_OPEN
+	{
+		UE_TRACE_LOG(ChaosVDLogger, ChaosVDSolverFrameStart, ChaosVDChannel)
+			<< ChaosVDSolverFrameStart.SolverID(ContextData.Id)
+			<< ChaosVDSolverFrameStart.Cycle(FPlatformTime::Cycles64())
+			<< ChaosVDSolverFrameStart.DebugName(*InDebugName, InDebugName.Len())
+			<< ChaosVDSolverFrameStart.IsKeyFrame(bOutIsFullCaptureRequested)
+			<< ChaosVDSolverFrameStart.IsReSimulated(bIsReSimulatedFrame)
+			<< ChaosVDSolverFrameStart.CurrentFrameNumber(FrameNumber);
+	};
 }
 
 void FChaosVisualDebuggerTrace::TraceSolverFrameEnd(const FChaosVDContext& ContextData)
@@ -567,16 +580,19 @@ void FChaosVisualDebuggerTrace::TraceSolverFrameEnd(const FChaosVDContext& Conte
 	}
 
 	{
-		FWriteScopeLock WriteLock(DeltaRecordingStatesLock);
+		UE::TWriteScopeLock WriteLock(DeltaRecordingStatesLock);
 		if (!SolverIDsForDeltaRecording.Contains(ContextData.Id))
 		{
 			SolverIDsForDeltaRecording.Add(ContextData.Id);
 		}
 	}
 
-	UE_TRACE_LOG(ChaosVDLogger, ChaosVDSolverFrameEnd, ChaosVDChannel)
-		<< ChaosVDSolverFrameEnd.SolverID(ContextData.Id)
-		<< ChaosVDSolverFrameEnd.Cycle(FPlatformTime::Cycles64());
+	UE_AUTORTFM_OPEN
+	{
+		UE_TRACE_LOG(ChaosVDLogger, ChaosVDSolverFrameEnd, ChaosVDChannel)
+			<< ChaosVDSolverFrameEnd.SolverID(ContextData.Id)
+			<< ChaosVDSolverFrameEnd.Cycle(FPlatformTime::Cycles64());
+	};
 }
 
 void FChaosVisualDebuggerTrace::TraceSolverStepStart(FStringView StepName)
@@ -594,10 +610,13 @@ void FChaosVisualDebuggerTrace::TraceSolverStepStart(FStringView StepName)
 		return;
 	}
 
-	UE_TRACE_LOG(ChaosVDLogger, ChaosVDSolverStepStart, ChaosVDChannel)
-		<< ChaosVDSolverStepStart.Cycle(FPlatformTime::Cycles64())
-		<< ChaosVDSolverStepStart.SolverID(CVDContextData->Id)
-		<< ChaosVDSolverStepStart.StepName(StepName.GetData(), GetNum(StepName));
+	UE_AUTORTFM_OPEN
+	{
+		UE_TRACE_LOG(ChaosVDLogger, ChaosVDSolverStepStart, ChaosVDChannel)
+			<< ChaosVDSolverStepStart.Cycle(FPlatformTime::Cycles64())
+			<< ChaosVDSolverStepStart.SolverID(CVDContextData->Id)
+			<< ChaosVDSolverStepStart.StepName(StepName.GetData(), GetNum(StepName));
+	};
 }
 
 void FChaosVisualDebuggerTrace::TraceSolverStepEnd()
@@ -615,9 +634,12 @@ void FChaosVisualDebuggerTrace::TraceSolverStepEnd()
 		return;
 	}
 
-	UE_TRACE_LOG(ChaosVDLogger, ChaosVDSolverStepEnd, ChaosVDChannel)
-		<< ChaosVDSolverStepEnd.Cycle(FPlatformTime::Cycles64())
-		<< ChaosVDSolverStepEnd.SolverID(CVDContextData->Id);
+	UE_AUTORTFM_OPEN
+	{
+		UE_TRACE_LOG(ChaosVDLogger, ChaosVDSolverStepEnd, ChaosVDChannel)
+			<< ChaosVDSolverStepEnd.Cycle(FPlatformTime::Cycles64())
+			<< ChaosVDSolverStepEnd.SolverID(CVDContextData->Id);
+	};
 }
 
 void FChaosVisualDebuggerTrace::TraceSolverSimulationSpace(const Chaos::FRigidTransform3& Transform)
@@ -634,12 +656,15 @@ void FChaosVisualDebuggerTrace::TraceSolverSimulationSpace(const Chaos::FRigidTr
 	{
 		return;
 	}
-	
-	UE_TRACE_LOG(ChaosVDLogger, ChaosVDSolverSimulationSpace, ChaosVDChannel)
-		<< ChaosVDSolverSimulationSpace.Cycle(FPlatformTime::Cycles64())
-		<< ChaosVDSolverSimulationSpace.SolverID(CVDContextData->Id)
-		<< CVD_TRACE_VECTOR_ON_EVENT(ChaosVDSolverSimulationSpace, Position, Transform.GetLocation())
-		<< CVD_TRACE_ROTATOR_ON_EVENT(ChaosVDSolverSimulationSpace, Rotation, Transform.GetRotation());
+
+	UE_AUTORTFM_OPEN
+	{
+		UE_TRACE_LOG(ChaosVDLogger, ChaosVDSolverSimulationSpace, ChaosVDChannel)
+			<< ChaosVDSolverSimulationSpace.Cycle(FPlatformTime::Cycles64())
+			<< ChaosVDSolverSimulationSpace.SolverID(CVDContextData->Id)
+			<< CVD_TRACE_VECTOR_ON_EVENT(ChaosVDSolverSimulationSpace, Position, Transform.GetLocation())
+			<< CVD_TRACE_ROTATOR_ON_EVENT(ChaosVDSolverSimulationSpace, Rotation, Transform.GetRotation());
+	};
 }
 
 void FChaosVisualDebuggerTrace::TraceBinaryData(TConstArrayView<uint8> InData, FStringView TypeName, EChaosVDTraceBinaryDataOptions Options)
@@ -650,57 +675,79 @@ void FChaosVisualDebuggerTrace::TraceBinaryData(TConstArrayView<uint8> InData, F
 	}
 
 	//TODO: This might overflow
-	static FThreadSafeCounter LastDataID;
-
-	const int32 DataID = LastDataID.Increment();
+	static std::atomic<int32> LastDataID;
+	
+	int32 DataID = INDEX_NONE;
+	
+	UE_AUTORTFM_OPEN
+	{
+		DataID = LastDataID++;
 
-	ensure(DataID < TNumericLimits<int32>::Max());
+		ensure(DataID < TNumericLimits<int32>::Max());
 
-	TConstArrayView<uint8> DataViewToTrace = InData;
+		TConstArrayView<uint8> DataViewToTrace = InData;
 
-	// Handle Compression if enabled
-	const bool bIsCompressed = Chaos::VisualDebugger::Cvars::bCompressBinaryData;
-	TArray<uint8> CompressedData;
-	if (bIsCompressed)
-	{
-		CompressedData.Reserve(CompressedData.Num());
-		FOodleCompressedArray::CompressData(CompressedData, InData.GetData(),InData.Num(), FOodleDataCompression::ECompressor::Kraken,
-			static_cast<FOodleDataCompression::ECompressionLevel>(Chaos::VisualDebugger::Cvars::CompressionMode));
+		// Handle Compression if enabled
+		const bool bIsCompressed = Chaos::VisualDebugger::Cvars::bCompressBinaryData;
+		TArray<uint8> CompressedData;
+		if (bIsCompressed)
+		{
+			CompressedData.Reserve(CompressedData.Num());
+			FOodleCompressedArray::CompressData(CompressedData, InData.GetData(),InData.Num(), FOodleDataCompression::ECompressor::Kraken,
+				static_cast<FOodleDataCompression::ECompressionLevel>(Chaos::VisualDebugger::Cvars::CompressionMode));
 
-		DataViewToTrace = CompressedData;
-	}
+			DataViewToTrace = CompressedData;
+		}
 
-	const uint32 DataSize = static_cast<uint32>(DataViewToTrace.Num());
-	constexpr uint32 MaxChunkSize = TNumericLimits<uint16>::Max();
-	const uint32 ChunkNum = (DataSize + MaxChunkSize - 1) / MaxChunkSize;
+		const uint32 DataSize = static_cast<uint32>(DataViewToTrace.Num());
+		constexpr uint32 MaxChunkSize = TNumericLimits<uint16>::Max();
+		const uint32 ChunkNum = (DataSize + MaxChunkSize - 1) / MaxChunkSize;
 
-	UE_TRACE_LOG(ChaosVDLogger, ChaosVDBinaryDataStart, ChaosVDChannel)
-		<< ChaosVDBinaryDataStart.Cycle(FPlatformTime::Cycles64())
-		<< ChaosVDBinaryDataStart.TypeName(TypeName.GetData(), TypeName.Len())
-		<< ChaosVDBinaryDataStart.DataID(DataID)
-		<< ChaosVDBinaryDataStart.DataSize(DataSize)
-		<< ChaosVDBinaryDataStart.OriginalSize(InData.Num())
-		<< ChaosVDBinaryDataStart.IsCompressed(bIsCompressed);
+		UE_TRACE_LOG(ChaosVDLogger, ChaosVDBinaryDataStart, ChaosVDChannel)
+			<< ChaosVDBinaryDataStart.Cycle(FPlatformTime::Cycles64())
+			<< ChaosVDBinaryDataStart.TypeName(TypeName.GetData(), TypeName.Len())
+			<< ChaosVDBinaryDataStart.DataID(DataID)
+			<< ChaosVDBinaryDataStart.DataSize(DataSize)
+			<< ChaosVDBinaryDataStart.OriginalSize(InData.Num())
+			<< ChaosVDBinaryDataStart.IsCompressed(bIsCompressed);
 
-	uint32 RemainingSize = DataSize;
-	for (uint32 Index = 0; Index < ChunkNum; ++Index)
-	{
-		const uint16 Size = static_cast<uint16>(FMath::Min(RemainingSize, MaxChunkSize));
-		const uint8* ChunkData = DataViewToTrace.GetData() + MaxChunkSize * Index;
+		uint32 RemainingSize = DataSize;
+		for (uint32 Index = 0; Index < ChunkNum; ++Index)
+		{
+			const uint16 Size = static_cast<uint16>(FMath::Min(RemainingSize, MaxChunkSize));
+			const uint8* ChunkData = DataViewToTrace.GetData() + MaxChunkSize * Index;
 
-		UE_TRACE_LOG(ChaosVDLogger, ChaosVDBinaryDataContent, ChaosVDChannel)
-			<< ChaosVDBinaryDataContent.Cycle(FPlatformTime::Cycles64())
-			<< ChaosVDBinaryDataContent.DataID(DataID)
-			<< ChaosVDBinaryDataContent.RawData(ChunkData, Size);
+			UE_TRACE_LOG(ChaosVDLogger, ChaosVDBinaryDataContent, ChaosVDChannel)
+				<< ChaosVDBinaryDataContent.Cycle(FPlatformTime::Cycles64())
+				<< ChaosVDBinaryDataContent.DataID(DataID)
+				<< ChaosVDBinaryDataContent.RawData(ChunkData, Size);
 
-		RemainingSize -= Size;
-	}
+			RemainingSize -= Size;
+		}
 
-	UE_TRACE_LOG(ChaosVDLogger, ChaosVDBinaryDataEnd, ChaosVDChannel)
-		<< ChaosVDBinaryDataEnd.Cycle(FPlatformTime::Cycles64())
-		<< ChaosVDBinaryDataEnd.DataID(DataID);
+		ensure(RemainingSize == 0);
+	};
+	
+	// Note: AutoRTFM is only partially supported at the moment.
+	// The approach taken here is that we trace serialized data regardless if the transaction will fail or not (to avoid allocating a new buffer and copying the data)
+	// but we only commit the last trace event that tells the CVD editor that the data is ready to be processed if the transaction succeeds (ChaosVDBinaryDataEnd). This allows us to ensure we don't load rolled back data automatically.
+	// In non-transacted callstacks this will be executed immediately (therefore the behaviour is the same as usual), but in transacted callstacks the commit will be done after the transaction completes in a call done from the game thread.
+	// This might pose an issue for data recorded outside the Game Thread as the assumption CVD relies on (all data is loaded in the exact order as it was recorded relative to other trace events) will not be valid.
+	// This means we might trace the ChaosVDBinaryDataEnd event when the [ Frame / Solver Stage ] end event to which the data belong, was already issued and in consequence CVD will load the data in the incorrect frame
+	// Currently the only calls done within a transaction should only be scene queries, and as they are done from the Game Thread, the framing / timing during load of the CVD recording should still be correct.
+	
+	UE_AUTORTFM_ONCOMMIT(DataID)
+	{
+		UE_TRACE_LOG(ChaosVDLogger, ChaosVDBinaryDataEnd, ChaosVDChannel)
+			<< ChaosVDBinaryDataEnd.Cycle(FPlatformTime::Cycles64())
+			<< ChaosVDBinaryDataEnd.DataID(DataID);
+	};
 
-	ensure(RemainingSize == 0);
+	UE_AUTORTFM_ONABORT(DataID)
+	{
+		UE_TRACE_LOG(ChaosVDLogger, ChaosVDRolledBackDataID, ChaosVDChannel)
+			<< ChaosVDRolledBackDataID.DataID(DataID);
+	};
 }
 
 void FChaosVisualDebuggerTrace::TraceImplicitObject(FChaosVDImplicitObjectWrapper WrappedGeometryData)
@@ -712,7 +759,7 @@ void FChaosVisualDebuggerTrace::TraceImplicitObject(FChaosVDImplicitObjectWrappe
 
 	uint32 GeometryID = WrappedGeometryData.Hash;
 	{
-		FReadScopeLock ReadLock(GeometryTracerObject.TracedGeometrySetLock);
+		UE::TReadScopeLock ReadLock(GeometryTracerObject.TracedGeometrySetLock);
 		if (GeometryTracerObject.GeometryTracedIDs.Find(GeometryID))
 		{
 			return;
@@ -720,7 +767,7 @@ void FChaosVisualDebuggerTrace::TraceImplicitObject(FChaosVDImplicitObjectWrappe
 	}
 
 	{
-		FWriteScopeLock WriteLock(GeometryTracerObject.TracedGeometrySetLock);
+		UE::TWriteScopeLock WriteLock(GeometryTracerObject.TracedGeometrySetLock);
 		GeometryTracerObject.GeometryTracedIDs.Add(GeometryID);
 	}
 
@@ -891,10 +938,13 @@ void FChaosVisualDebuggerTrace::TraceNetworkTickOffset(int32 TickOffset, int32 S
 	{
 		return;
 	}
-	
-	UE_TRACE_LOG(ChaosVDLogger, ChaosVDNetworkTickOffset, ChaosVDChannel)
-		<< ChaosVDNetworkTickOffset.Offset(TickOffset)
-		<< ChaosVDNetworkTickOffset.SolverID(SolverID);
+
+	UE_AUTORTFM_OPEN
+	{
+		UE_TRACE_LOG(ChaosVDLogger, ChaosVDNetworkTickOffset, ChaosVDChannel)
+			<< ChaosVDNetworkTickOffset.Offset(TickOffset)
+			<< ChaosVDNetworkTickOffset.SolverID(SolverID);
+	};
 }
 
 void FChaosVisualDebuggerTrace::TraceDebugDrawBox(const FBox& InBox, FName Tag, FColor Color, int32 SolverID)
@@ -1019,7 +1069,7 @@ bool FChaosVisualDebuggerTrace::IsTracing()
 void FChaosVisualDebuggerTrace::RegisterEventHandlers()
 {
 	{
-		FWriteScopeLock WriteLock(DeltaRecordingStatesLock);
+		UE::TWriteScopeLock WriteLock(DeltaRecordingStatesLock);
 
 		
 		if (!RecordingStartedDelegateHandle.IsValid())
@@ -1041,7 +1091,7 @@ void FChaosVisualDebuggerTrace::RegisterEventHandlers()
 
 void FChaosVisualDebuggerTrace::UnregisterEventHandlers()
 {
-	FWriteScopeLock WriteLock(DeltaRecordingStatesLock);
+	UE::TWriteScopeLock WriteLock(DeltaRecordingStatesLock);
 	if (RecordingStartedDelegateHandle.IsValid())
 	{
 		FChaosVDRuntimeModule::RemoveRecordingStartedCallback(RecordingStartedDelegateHandle);
@@ -1070,16 +1120,16 @@ void FChaosVisualDebuggerTrace::Reset()
 	CVDNameTable->ResetTable();
 
 	{
-		FWriteScopeLock WriteLock(DeltaRecordingStatesLock);
+		UE::TWriteScopeLock WriteLock(DeltaRecordingStatesLock);
 		RequestedFullCaptureSolverIDs.Reset();
 		SolverIDsForDeltaRecording.Reset();
 	}
 
 	{
-		FWriteScopeLock GeometryWriteLock(GeometryTracerObject.TracedGeometrySetLock);
+		UE::TWriteScopeLock GeometryWriteLock(GeometryTracerObject.TracedGeometrySetLock);
 		GeometryTracerObject.GeometryTracedIDs.Reset();
 
-		FWriteScopeLock GeometryHashWriteLock(GeometryTracerObject.CachedGeometryHashesLock);
+		UE::TWriteScopeLock GeometryHashWriteLock(GeometryTracerObject.CachedGeometryHashesLock);
 		GeometryTracerObject.CachedGeometryHashes.Reset();
 	}
 }
@@ -1137,6 +1187,9 @@ void FChaosVisualDebuggerTrace::HandleRecordingStart()
 
 	TraceArchiveHeader();
 
+	UE_TRACE_LOG(ChaosVDLogger, ChaosVDUsesAutoRTFM, ChaosVDChannel)
+			<< ChaosVDUsesAutoRTFM.bUsingAutoRTFM(static_cast<bool>(UE_AUTORTFM));
+
 	bIsTracing = true;
 }
 
@@ -1144,13 +1197,13 @@ void FChaosVisualDebuggerTrace::PerformFullCapture(EChaosVDFullCaptureFlags Capt
 {
 	if (EnumHasAnyFlags(CaptureOptions, EChaosVDFullCaptureFlags::Particles))
 	{
-		FWriteScopeLock WriteLock(DeltaRecordingStatesLock);
+		UE::TWriteScopeLock WriteLock(DeltaRecordingStatesLock);
 		RequestedFullCaptureSolverIDs.Append(SolverIDsForDeltaRecording);
 	}
 
 	if (EnumHasAnyFlags(CaptureOptions, EChaosVDFullCaptureFlags::Geometry))
 	{
-		FWriteScopeLock GeometryWriteLock(GeometryTracerObject.TracedGeometrySetLock);
+		UE::TWriteScopeLock GeometryWriteLock(GeometryTracerObject.TracedGeometrySetLock);
 		GeometryTracerObject.GeometryTracedIDs.Reset();
 	}
 }
diff --git a/Engine/Source/Runtime/Experimental/Chaos/Public/ChaosVisualDebugger/ChaosVDSerializedNameTable.h b/Engine/Source/Runtime/Experimental/Chaos/Public/ChaosVisualDebugger/ChaosVDSerializedNameTable.h
--- a/Engine/Source/Runtime/Experimental/Chaos/Public/ChaosVisualDebugger/ChaosVDSerializedNameTable.h
+++ b/Engine/Source/Runtime/Experimental/Chaos/Public/ChaosVisualDebugger/ChaosVDSerializedNameTable.h
@@ -5,6 +5,7 @@
 #include "Containers/Map.h"
 #include "Containers/StringView.h"
 #include "HAL/Platform.h"
+#include "Misc/TransactionallySafeRWLock.h"
 #include "Templates/SharedPointer.h"
 
 namespace Chaos::VisualDebugger
@@ -43,7 +44,8 @@ namespace Chaos::VisualDebugger
 	private:
 			
 		TMap<uint64, FName> NamesByID;
-		FRWLock NamesByIDLock;
+		
+		FTransactionallySafeRWLock NamesByIDLock;
 	};
 
 	inline FArchive& operator<<(FArchive& Ar, FChaosVDSerializedNameEntry& NameEntry)
diff --git a/Engine/Source/Runtime/Experimental/Chaos/Public/ChaosVisualDebugger/ChaosVisualDebuggerTrace.h b/Engine/Source/Runtime/Experimental/Chaos/Public/ChaosVisualDebugger/ChaosVisualDebuggerTrace.h
--- a/Engine/Source/Runtime/Experimental/Chaos/Public/ChaosVisualDebugger/ChaosVisualDebuggerTrace.h
+++ b/Engine/Source/Runtime/Experimental/Chaos/Public/ChaosVisualDebugger/ChaosVisualDebuggerTrace.h
@@ -122,6 +122,14 @@ UE_TRACE_EVENT_BEGIN_EXTERN(ChaosVDLogger, ChaosVDNetworkTickOffset)
 	UE_TRACE_EVENT_FIELD(int32, SolverID)
 UE_TRACE_EVENT_END()
 
+UE_TRACE_EVENT_BEGIN_EXTERN(ChaosVDLogger, ChaosVDRolledBackDataID)
+	UE_TRACE_EVENT_FIELD(int32, DataID)
+UE_TRACE_EVENT_END()
+
+UE_TRACE_EVENT_BEGIN_EXTERN(ChaosVDLogger, ChaosVDUsesAutoRTFM)
+	UE_TRACE_EVENT_FIELD(bool, bUsingAutoRTFM)
+UE_TRACE_EVENT_END()
+
 UE_TRACE_EVENT_BEGIN_EXTERN(ChaosVDLogger, ChaosVDDummyEvent)
 	UE_TRACE_EVENT_FIELD(int32, SolverID)
 UE_TRACE_EVENT_END()
diff --git a/Engine/Source/Runtime/Experimental/ChaosVisualDebugger/Private/ChaosVDRuntimeModule.cpp b/Engine/Source/Runtime/Experimental/ChaosVisualDebugger/Private/ChaosVDRuntimeModule.cpp
--- a/Engine/Source/Runtime/Experimental/ChaosVisualDebugger/Private/ChaosVDRuntimeModule.cpp
+++ b/Engine/Source/Runtime/Experimental/ChaosVisualDebugger/Private/ChaosVDRuntimeModule.cpp
@@ -29,7 +29,10 @@ FAutoConsoleCommand ChaosVDStartRecordingCommand(
 	TEXT("Turn on the recording of debugging data"),
 	FConsoleCommandWithArgsDelegate::CreateLambda([](const TArray<FString>& Args)
 	{
-		FChaosVDRuntimeModule::Get().StartRecording(Args);
+		UE_AUTORTFM_ONCOMMIT(Args)
+		{
+			FChaosVDRuntimeModule::Get().StartRecording(Args);
+		};
 	})
 );
 
@@ -38,7 +41,10 @@ FAutoConsoleCommand StopVDStartRecordingCommand(
 	TEXT("Turn off the recording of debugging data"),
 	FConsoleCommandWithArgsDelegate::CreateLambda([](const TArray<FString>& Args)
 	{
-		FChaosVDRuntimeModule::Get().StopRecording();
+		UE_AUTORTFM_ONCOMMIT(Args)
+		{
+			FChaosVDRuntimeModule::Get().StopRecording();
+		};
 	})
 );
 
@@ -56,7 +62,7 @@ FChaosVDRecordingStateChangedDelegate FChaosVDRuntimeModule::RecordingStartedDel
 FChaosVDRecordingStateChangedDelegate FChaosVDRuntimeModule::RecordingStopDelegate = FChaosVDRecordingStateChangedDelegate();
 FChaosVDRecordingStartFailedDelegate FChaosVDRuntimeModule::RecordingStartFailedDelegate = FChaosVDRecordingStartFailedDelegate();
 FChaosVDCaptureRequestDelegate FChaosVDRuntimeModule::PerformFullCaptureDelegate = FChaosVDCaptureRequestDelegate();
-FRWLock FChaosVDRuntimeModule::DelegatesRWLock = FRWLock();
+FTransactionallySafeRWLock FChaosVDRuntimeModule::DelegatesRWLock = FTransactionallySafeRWLock();
 
 FChaosVDRuntimeModule& FChaosVDRuntimeModule::Get()
 {
@@ -106,7 +112,13 @@ void FChaosVDRuntimeModule::ShutdownModule()
 
 int32 FChaosVDRuntimeModule::GenerateUniqueID()
 {
-	return LastGeneratedID.Increment();
+	int32 NewID = 0;
+	UE_AUTORTFM_OPEN
+	{
+		NewID = LastGeneratedID++;
+	};
+
+	return NewID;
 }
 
 FString FChaosVDRuntimeModule::GetLastRecordingFileNamePath() const
@@ -128,7 +140,7 @@ void FChaosVDRuntimeModule::GenerateRecordingFileName(FString& OutFileName)
 bool FChaosVDRuntimeModule::RequestFullCapture(float DeltaTime)
 {
 	// Full capture intervals are clamped to be no lower than 1 sec
-	FReadScopeLock ReadLock (DelegatesRWLock);
+	UE::TReadScopeLock ReadLock (DelegatesRWLock);
 	PerformFullCaptureDelegate.Broadcast(EChaosVDFullCaptureFlags::Particles);
 	return true;
 }
@@ -213,7 +225,7 @@ void FChaosVDRuntimeModule::StartRecording(TConstArrayView<FString> Args)
 	if (ensure(bIsRecording))
 	{
 		{
-			FReadScopeLock ReadLock(DelegatesRWLock);
+			UE::TReadScopeLock ReadLock(DelegatesRWLock);
 			RecordingStartedDelegate.Broadcast();
 		}
 		
@@ -238,7 +250,7 @@ void FChaosVDRuntimeModule::StartRecording(TConstArrayView<FString> Args)
 #endif
 
 		{
-			FReadScopeLock ReadLock(DelegatesRWLock);
+			UE::TReadScopeLock ReadLock(DelegatesRWLock);
 			RecordingStartFailedDelegate.Broadcast(FailureReason);
 		}	
 	}
@@ -274,7 +286,7 @@ void FChaosVDRuntimeModule::StopRecording()
 	}
 
 	{
-		FReadScopeLock ReadLock(DelegatesRWLock);
+		UE::TReadScopeLock ReadLock(DelegatesRWLock);
 		RecordingStopDelegate.Broadcast();
 	}
 
diff --git a/Engine/Source/Runtime/Experimental/ChaosVisualDebugger/Public/ChaosVDRuntimeModule.h b/Engine/Source/Runtime/Experimental/ChaosVisualDebugger/Public/ChaosVDRuntimeModule.h
--- a/Engine/Source/Runtime/Experimental/ChaosVisualDebugger/Public/ChaosVDRuntimeModule.h
+++ b/Engine/Source/Runtime/Experimental/ChaosVisualDebugger/Public/ChaosVDRuntimeModule.h
@@ -7,6 +7,7 @@
 #include "ProfilingDebugging/TraceAuxiliary.h"
 #include "Containers/Ticker.h"
 #include "Misc/ScopeRWLock.h"
+#include "Misc/TransactionallySafeRWLock.h"
 
 struct FChaosVDRecording;
 class FText;
@@ -49,49 +50,49 @@ public:
 
 	static FDelegateHandle RegisterRecordingStartedCallback(const FChaosVDRecordingStateChangedDelegate::FDelegate& InCallback)
 	{
-		FWriteScopeLock WriteLock(DelegatesRWLock);
+		UE::TWriteScopeLock WriteLock(DelegatesRWLock);
 		return RecordingStartedDelegate.Add(InCallback);
 	}
 
 	static FDelegateHandle RegisterRecordingStopCallback(const FChaosVDRecordingStateChangedDelegate::FDelegate& InCallback)
 	{
-		FWriteScopeLock WriteLock(DelegatesRWLock);
+		UE::TWriteScopeLock WriteLock(DelegatesRWLock);
 		return RecordingStopDelegate.Add(InCallback);
 	}
 
 	static FDelegateHandle RegisterRecordingStartFailedCallback(const FChaosVDRecordingStartFailedDelegate::FDelegate& InCallback)
 	{
-		FWriteScopeLock WriteLock(DelegatesRWLock);
+		UE::TWriteScopeLock WriteLock(DelegatesRWLock);
 		return RecordingStartFailedDelegate.Add(InCallback);
 	}
 	
 	static FDelegateHandle RegisterFullCaptureRequestedCallback(const FChaosVDCaptureRequestDelegate::FDelegate& InCallback)
 	{
-		FWriteScopeLock WriteLock(DelegatesRWLock);
+		UE::TWriteScopeLock WriteLock(DelegatesRWLock);
 		return PerformFullCaptureDelegate.Add(InCallback);
 	}
 
 	static bool RemoveRecordingStartedCallback(const FDelegateHandle& InDelegateToRemove)
 	{
-		FWriteScopeLock WriteLock(DelegatesRWLock);
+		UE::TWriteScopeLock WriteLock(DelegatesRWLock);
 		return RecordingStartedDelegate.Remove(InDelegateToRemove);
 	}
 
 	static bool RemoveRecordingStopCallback(const FDelegateHandle& InDelegateToRemove)
 	{
-		FWriteScopeLock WriteLock(DelegatesRWLock);
+		UE::TWriteScopeLock WriteLock(DelegatesRWLock);
 		return RecordingStopDelegate.Remove(InDelegateToRemove);
 	}
 
 	static bool RemoveRecordingStartFailedCallback(const FDelegateHandle& InDelegateToRemove)
 	{
-		FWriteScopeLock WriteLock(DelegatesRWLock);
+		UE::TWriteScopeLock WriteLock(DelegatesRWLock);
 		return RecordingStartFailedDelegate.Remove(InDelegateToRemove);
 	}
 	
 	static bool RemoveFullCaptureRequestedCallback(const FDelegateHandle& InDelegateToRemove)
 	{
-		FWriteScopeLock WriteLock(DelegatesRWLock);
+		UE::TWriteScopeLock WriteLock(DelegatesRWLock);
 		return PerformFullCaptureDelegate.Remove(InDelegateToRemove);
 	}
 
@@ -138,11 +139,11 @@ private:
 	static FChaosVDRecordingStartFailedDelegate RecordingStartFailedDelegate;
 	static FChaosVDCaptureRequestDelegate PerformFullCaptureDelegate;
 
-	FThreadSafeCounter LastGeneratedID;
+	std::atomic<int32> LastGeneratedID;
 
 	FString LastRecordingFileNamePath;
 
 	TMap<FString, bool> OriginalTraceChannelsState;
 
-	static FRWLock DelegatesRWLock;
+	static FTransactionallySafeRWLock DelegatesRWLock;
 };
