Iris: Group destruction: Lower log level to Verbose

Removes "Destroyed GroupHandle named" log message from the default log level, matching the "Created GroupHandle named" message.

#rnx

diff --git a/Engine/Source/Runtime/Experimental/Iris/Core/Private/Iris/ReplicationSystem/Filtering/NetObjectGroups.cpp b/Engine/Source/Runtime/Experimental/Iris/Core/Private/Iris/ReplicationSystem/Filtering/NetObjectGroups.cpp
--- a/Engine/Source/Runtime/Experimental/Iris/Core/Private/Iris/ReplicationSystem/Filtering/NetObjectGroups.cpp
+++ b/Engine/Source/Runtime/Experimental/Iris/Core/Private/Iris/ReplicationSystem/Filtering/NetObjectGroups.cpp
@@ -127,7 +127,7 @@ void FNetObjectGroups::DestroyGroup(FNetObjectGroupHandle GroupHandle)
 	{
 		FNetObjectGroup& Group = Groups[GroupHandle.GetGroupIndex()];
 
-		UE_LOG(LogIrisGroup, Log, TEXT("Destroyed GroupHandle named: %s Index: %u Id: %u"), *Group.GroupName.ToString(), GroupHandle.Index, Group.GroupId);
+		UE_LOG(LogIrisGroup, Verbose, TEXT("Destroyed GroupHandle named: %s Index: %u Id: %u"), *Group.GroupName.ToString(), GroupHandle.Index, Group.GroupId);
 
 		ClearGroup(GroupHandle);		
 		Groups.RemoveAt(GroupHandle.GetGroupIndex());
