Dataflow : Better support for subgraph

diff --git a/Engine/Source/Runtime/Experimental/Dataflow/Engine/Private/Dataflow/DataflowEdNode.cpp b/Engine/Source/Runtime/Experimental/Dataflow/Engine/Private/Dataflow/DataflowEdNode.cpp
--- a/Engine/Source/Runtime/Experimental/Dataflow/Engine/Private/Dataflow/DataflowEdNode.cpp
+++ b/Engine/Source/Runtime/Experimental/Dataflow/Engine/Private/Dataflow/DataflowEdNode.cpp
@@ -91,7 +91,7 @@ void UDataflowEdNode::SetShouldWireframeRenderNode(bool bInRender)
 	if (IsBound())
 	{
 #if WITH_EDITOR && !(UE_BUILD_SHIPPING || UE_BUILD_TEST)
-		if (UDataflow* const DataflowObject = Cast<UDataflow>(GetGraph()))
+		if (UDataflow* const DataflowObject = UDataflow::GetDataflowAssetFromEdGraph(GetGraph()))
 		{
 			if (bRenderWireframeInAssetEditor)
 			{
@@ -452,7 +452,7 @@ void UDataflowEdNode::Serialize(FArchive& Ar)
 #if WITH_EDITOR
 	if (Ar.IsTransacting())
 	{
-		const UDataflow* const DataflowObject = Cast<UDataflow>(GetGraph());
+		const UDataflow* const DataflowObject = UDataflow::GetDataflowAssetFromEdGraph(GetGraph());
 		bool bCanSerializeNode = !DataflowObject || DataflowObject->IsPerNodeTransactionSerializationEnabled();
 		if (bCanSerializeNode)
 		{
