We now track IAS decode errors as two different values, one for buffers coming from http requests and another for buffers that we loaded from the cache.

#rnx

- Up until now we have only been tracking the number of decode errors but that doesn't tell us if there is a problem with the network or the disk cache and historically has made it difficult to know where to target our efforts.
- FOnDemandIoBackend::CompleteRequest now takes an enum that describes the source of the request, so if the decode fails we can inform the statistics system if we failed to decode a buffer that was accessed via http or loaded from disk.

diff --git a/Engine/Source/Runtime/Experimental/IoStore/OnDemand/Private/OnDemandIoDispatcherBackend.cpp b/Engine/Source/Runtime/Experimental/IoStore/OnDemand/Private/OnDemandIoDispatcherBackend.cpp
--- a/Engine/Source/Runtime/Experimental/IoStore/OnDemand/Private/OnDemandIoDispatcherBackend.cpp
+++ b/Engine/Source/Runtime/Experimental/IoStore/OnDemand/Private/OnDemandIoDispatcherBackend.cpp
@@ -1108,7 +1108,15 @@ void FOnDemandIoBackend::CompleteRequest(FChunkRequest* ChunkRequest)
 
 			if (!bDecoded)
 			{
-				Stats.OnIoDecodeError();
+				if (ChunkRequest->bCached)
+				{
+					Stats.OnCacheDecodeError();
+				}
+				else
+				{
+					// Currently not being cached implies that the request was made via http
+					Stats.OnHttpDecodeError();
+				}
 			}
 		}
 		
diff --git a/Engine/Source/Runtime/Experimental/IoStore/OnDemand/Private/Statistics.cpp b/Engine/Source/Runtime/Experimental/IoStore/OnDemand/Private/Statistics.cpp
--- a/Engine/Source/Runtime/Experimental/IoStore/OnDemand/Private/Statistics.cpp
+++ b/Engine/Source/Runtime/Experimental/IoStore/OnDemand/Private/Statistics.cpp
@@ -258,9 +258,9 @@ FCounterAtomicInt		GIoRequestReadCount(TEXT("Ias/IoRequestReadCount"), TraceCoun
 FCounterAtomicInt		GIoRequestReadBytes(TEXT("Ias/IoRequestReadBytes"), TraceCounterDisplayHint_Memory);
 FCounterInt				GIoRequestCancelCount(TEXT("Ias/IoRequestCancelCount"), TraceCounterDisplayHint_None);
 FCounterAtomicInt		GIoRequestErrorCount(TEXT("Ias/IoRequestErrorCount"), TraceCounterDisplayHint_None);
-FCounterAtomicInt		GIoDecodeErrorCount(TEXT("Ias/IoDecodeErrorCount"), TraceCounterDisplayHint_None);
 // cache stats
 FCounterAtomicInt		GCacheErrorCount(TEXT("Ias/CacheErrorCount"), TraceCounterDisplayHint_None);
+FCounterAtomicInt		GCacheDecodeErrorCount(TEXT("Ias/CacheDecodeErrorCount"), TraceCounterDisplayHint_None);
 FCounterAtomicInt		GCacheGetCount(TEXT("Ias/CacheGetCount"), TraceCounterDisplayHint_None);
 FCounterAtomicInt		GCachePutCount(TEXT("Ias/CachePutCount"), TraceCounterDisplayHint_None);
 FCounterAtomicInt		GCachePutExistingCount(TEXT("Ias/CachePutExistingCount"), TraceCounterDisplayHint_None);
@@ -276,6 +276,7 @@ FCounterInt				GHttpConnectCount(TEXT("Ias/HttpConnectCount"), TraceCounterDispl
 FCounterInt				GHttpDisconnectCount(TEXT("Ias/HttpDisconnectCount"), TraceCounterDisplayHint_None);
 FCounterInt				GHttpGetCount(TEXT("Ias/HttpGetCount"), TraceCounterDisplayHint_None);
 FCounterInt				GHttpErrorCount(TEXT("Ias/HttpErrorCount"), TraceCounterDisplayHint_None);
+FCounterAtomicInt		GHttpDecodeErrorCount(TEXT("Ias/HttpDecodeErrorCount"), TraceCounterDisplayHint_None);
 FCounterInt				GHttpRetryCount(TEXT("Ias/HttpRetryCount"), TraceCounterDisplayHint_None);
 FCounterInt				GHttpCancelCount(TEXT("Ias/HttpCancelCount"), TraceCounterDisplayHint_None);
 FCounterAtomicInt		GHttpPendingCount(TEXT("Ias/HttpPendingCount"), TraceCounterDisplayHint_None);
@@ -498,11 +499,12 @@ FOnDemandIoBackendStats::FOnDemandIoBackendStats(const FAvailableEndPoints& InAv
 
 				Out.Add(FCoreDelegates::EOnScreenMessageSeverity::Info, FText::FromStringView(Builder.ToView()));
 
-				if (GIoDecodeErrorCount.Get() > 0 || GHttpErrorCount.Get() > 0)
+				if (GHttpDecodeErrorCount.Get() > 0 || GCacheDecodeErrorCount.Get() > 0 || GHttpErrorCount.Get() > 0)
 				{
 					TStringBuilder<256> Warning;
-					Warning.Appendf(TEXT("IAS - Decode Errors %d | Http Errors %d"),
-						GIoDecodeErrorCount.Get(),
+					Warning.Appendf(TEXT("IAS - Cache Decode Errors %d | Http Decode Errors %d | Http Errors %d"),
+						GCacheDecodeErrorCount.Get(),
+						GHttpDecodeErrorCount.Get(),
 						GHttpErrorCount.Get()
 					);
 
@@ -547,14 +549,13 @@ void FOnDemandIoBackendStats::ReportEndPointAnalytics(TArray<FAnalyticsEventAttr
 	{
 		AppendAnalyticsEventAttributeArray(OutAnalyticsArray
 			,TRACK_DELTA("IasHttpErrorCount", GHttpErrorCount.Get())
+			,TRACK_DELTA("IasHttpDecodeErrors", GHttpDecodeErrorCount.Get())
 			,TRACK_DELTA("IasHttpRetryCount", GHttpRetryCount.Get())
 			,TRACK_DELTA("IasHttpGetCount", GHttpGetCount.Get())
 			,TRACK_DELTA("IasHttpDownloadedBytes", GHttpDownloadedBytes.Get())
 			,TEXT("IasHttpDurationMeanAvg"), GHttpAvgDuration.GetMean()
 			,TEXT("IasHttpDurationStdDev"), GHttpAvgDuration.GetDeviation()
 
-			,TRACK_DELTA("IasDecodeErrors", GIoDecodeErrorCount.Get())
-
 			,TEXT("IasHttpDuration0"), GHttpDurationBucket0
 			,TEXT("IasHttpDuration1"), GHttpDurationBucket1
 			,TEXT("IasHttpDuration2"), GHttpDurationBucket2
@@ -578,6 +579,7 @@ void FOnDemandIoBackendStats::ReportEndPointAnalytics(TArray<FAnalyticsEventAttr
 
 			,TRACK_DELTA("IasCacheTotalCount", CacheTotalCount)
 			,TRACK_DELTA("IasCacheErrorCount", GCacheErrorCount.Get())
+			,TRACK_DELTA("IasCacheDecodeErrors", GCacheDecodeErrorCount.Get())
 			,TRACK_DELTA("IasCacheGetCount", GCacheGetCount.Get())
 			,TRACK_DELTA("IasCachePutCount", GCachePutCount.Get())
 
@@ -618,6 +620,7 @@ private:
 		}
 	
 		Http.ErrorCount.Stop();
+		Http.DecodeErrorCount.Stop();
 		Http.RetryCount.Stop();
 		Http.GetCount.Stop();
 		Http.DownloadedBytes.Stop();
@@ -628,9 +631,8 @@ private:
 		Http.DurationBuckets3.Stop();
 		Http.DurationBuckets4.Stop();
 
-		Io.DecodeErrors.Stop();
-
 		Cache.ErrorCount.Stop();
+		Cache.DecodeErrorCount.Stop();
 		Cache.GetCount.Stop();
 		Cache.PutCount.Stop();
 
@@ -663,6 +665,7 @@ private:
 		{
 			AppendAnalyticsEventAttributeArray(OutAnalyticsArray
 				, TEXT("IasHttpErrorCount"), Http.ErrorCount.GetValue()
+				, TEXT("IasHttpDecodeErrorCount"), Http.DecodeErrorCount.GetValue()
 				, TEXT("IasHttpRetryCount"), Http.RetryCount.GetValue()
 				, TEXT("IasHttpGetCount"), Http.GetCount.GetValue()
 				, TEXT("IasHttpDownloadedBytes"), Http.DownloadedBytes.GetValue()
@@ -677,10 +680,6 @@ private:
 			);
 		}
 
-		AppendAnalyticsEventAttributeArray(OutAnalyticsArray
-			, TEXT("IasDecodeErrors"), GIoDecodeErrorCount.Get()
-		);
-
 		if (GIasReportCacheAnalyticsEnabled)
 		{
 			const float CacheUsagePercent = GCacheMaxBytes > 0 ? 100.f * (float(GCacheCachedBytes.Get()) / float(GCacheMaxBytes)) : 0.f;
@@ -693,6 +692,7 @@ private:
 				, TEXT("IasCacheUsagePercent"), CacheUsagePercent
 
 				, TEXT("IasCacheErrorCount"), Cache.ErrorCount.GetValue()
+				, TEXT("IasCacheDecodeErrorCount"), Cache.DecodeErrorCount.GetValue()
 				, TEXT("IasCacheGetCount"), Cache.GetCount.GetValue()
 				, TEXT("IasCachePutCount"), Cache.PutCount.GetValue()
 
@@ -760,34 +760,30 @@ private:
 	/** Http Stats */
 	struct FHttpStats
 	{
-		FTrackedValue<FCounterInt, GHttpErrorCount>			ErrorCount;
-		FTrackedValue<FCounterInt, GHttpRetryCount>			RetryCount;
-		FTrackedValue<FCounterInt, GHttpGetCount>			GetCount;
-		FTrackedValue<FCounterInt, GHttpDownloadedBytes>	DownloadedBytes;
-
-		FTrackedValue<uint32, GHttpDurationBucket0>			DurationBuckets0;
-		FTrackedValue<uint32, GHttpDurationBucket1>			DurationBuckets1;
-		FTrackedValue<uint32, GHttpDurationBucket2>			DurationBuckets2;
-		FTrackedValue<uint32, GHttpDurationBucket3>			DurationBuckets3;
-		FTrackedValue<uint32, GHttpDurationBucket4>			DurationBuckets4;
+		FTrackedValue<FCounterInt, GHttpErrorCount>				ErrorCount;
+		FTrackedValue<FCounterAtomicInt, GHttpDecodeErrorCount> DecodeErrorCount;
+		FTrackedValue<FCounterInt, GHttpRetryCount>				RetryCount;
+		FTrackedValue<FCounterInt, GHttpGetCount>				GetCount;
+		FTrackedValue<FCounterInt, GHttpDownloadedBytes>		DownloadedBytes;
+
+		FTrackedValue<uint32, GHttpDurationBucket0>				DurationBuckets0;
+		FTrackedValue<uint32, GHttpDurationBucket1>				DurationBuckets1;
+		FTrackedValue<uint32, GHttpDurationBucket2>				DurationBuckets2;
+		FTrackedValue<uint32, GHttpDurationBucket3>				DurationBuckets3;
+		FTrackedValue<uint32, GHttpDurationBucket4>				DurationBuckets4;
 	} Http;
 
-	/** Io Stats */
-	struct FIoStats
-	{
-		FTrackedValue<FCounterAtomicInt, GIoDecodeErrorCount> DecodeErrors;
-	} Io;
-
 	/** Cache Stats */
 	struct FCacheStats
 	{
-		FTrackedValue<FCounterAtomicInt, GCacheErrorCount>		ErrorCount;
-		FTrackedValue<FCounterAtomicInt, GCacheGetCount>		GetCount;
-		FTrackedValue<FCounterAtomicInt, GCachePutCount>		PutCount;
-
-		FTrackedValue<FCounterAtomicInt, GCacheWrittenBytes>	WrittenBytes;
-		FTrackedValue<FCounterAtomicInt, GCacheReadBytes>		ReadBytes;
-		FTrackedValue<FCounterAtomicInt, GCacheRejectBytes>		RejectBytes;
+		FTrackedValue<FCounterAtomicInt, GCacheErrorCount>			ErrorCount;
+		FTrackedValue<FCounterAtomicInt, GCacheDecodeErrorCount>	DecodeErrorCount;
+		FTrackedValue<FCounterAtomicInt, GCacheGetCount>			GetCount;
+		FTrackedValue<FCounterAtomicInt, GCachePutCount>			PutCount;
+
+		FTrackedValue<FCounterAtomicInt, GCacheWrittenBytes>		WrittenBytes;
+		FTrackedValue<FCounterAtomicInt, GCacheReadBytes>			ReadBytes;
+		FTrackedValue<FCounterAtomicInt, GCacheRejectBytes>			RejectBytes;
 	} Cache;
 
 	const FAvailableEndPoints& AvailableEndPoints;
@@ -830,14 +826,14 @@ void FOnDemandIoBackendStats::OnIoRequestError()
 	CSV_CUSTOM_STAT_DEFINED(FrameIoRequestErrorCount, int32(GIoRequestErrorCount.Get()), ECsvCustomStatOp::Set);
 }
 
-void FOnDemandIoBackendStats::OnIoDecodeError()
+void FOnDemandIoBackendStats::OnCacheError()
 {
-	GIoDecodeErrorCount.Add(1);
+	GCacheErrorCount.Add(1);
 }
 
-void FOnDemandIoBackendStats::OnCacheError()
+void FOnDemandIoBackendStats::OnCacheDecodeError()
 {
-	GCacheErrorCount.Add(1);
+	GCacheDecodeErrorCount.Add(1);
 }
 
 void FOnDemandIoBackendStats::OnCacheGet(uint64 DataSize)
@@ -960,6 +956,11 @@ void FOnDemandIoBackendStats::OnHttpError()
 	GHttpErrorCount.Add(1);
 }
 
+void FOnDemandIoBackendStats::OnHttpDecodeError()
+{
+	GHttpDecodeErrorCount.Add(1);
+}
+
 } // namespace UE::IoStore
 
 #undef UE_ENABLE_ONSCREEN_STATISTICS
diff --git a/Engine/Source/Runtime/Experimental/IoStore/OnDemand/Private/Statistics.h b/Engine/Source/Runtime/Experimental/IoStore/OnDemand/Private/Statistics.h
--- a/Engine/Source/Runtime/Experimental/IoStore/OnDemand/Private/Statistics.h
+++ b/Engine/Source/Runtime/Experimental/IoStore/OnDemand/Private/Statistics.h
@@ -49,9 +49,8 @@ public:
 	void OnIoRequestCancel() IAS_STATISTICS_IMPL()
 	void OnIoRequestError() IAS_STATISTICS_IMPL()
 
-	void OnIoDecodeError() IAS_STATISTICS_IMPL()
-
 	void OnCacheError() IAS_STATISTICS_IMPL()
+	void OnCacheDecodeError() IAS_STATISTICS_IMPL()
 	void OnCacheGet(uint64 DataSize) IAS_STATISTICS_IMPL()
 	void OnCachePut() IAS_STATISTICS_IMPL()
 	void OnCachePutExisting(uint64 DataSize) IAS_STATISTICS_IMPL()
@@ -70,6 +69,7 @@ public:
 	void OnHttpGet(uint64 SizeBytes, uint64 DurationMs) IAS_STATISTICS_IMPL()
 	void OnHttpRetry() IAS_STATISTICS_IMPL()
 	void OnHttpError() IAS_STATISTICS_IMPL()
+	void OnHttpDecodeError() IAS_STATISTICS_IMPL()
 
 private:
 
