Iris - Fixed issue with atomic destroy where subobject infos would not be updated properly. Print more info when an object switches state from an unsupported state.

diff --git a/Engine/Source/Runtime/Experimental/Iris/Core/Private/Iris/ReplicationSystem/ReplicationWriter.cpp b/Engine/Source/Runtime/Experimental/Iris/Core/Private/Iris/ReplicationSystem/ReplicationWriter.cpp
--- a/Engine/Source/Runtime/Experimental/Iris/Core/Private/Iris/ReplicationSystem/ReplicationWriter.cpp
+++ b/Engine/Source/Runtime/Experimental/Iris/Core/Private/Iris/ReplicationSystem/ReplicationWriter.cpp
@@ -123,80 +123,6 @@ const TCHAR* FReplicationWriter::LexToString(const EReplicatedObjectState State)
 	return State < EReplicatedObjectState::Count ? Names[(uint32)State] : TEXT("");
 }
 
-void FReplicationWriter::FReplicationInfo::SetState(EReplicatedObjectState NewState)
-{
-	EReplicatedObjectState CurrentState = GetState();
-	switch (NewState)
-	{
-		case EReplicatedObjectState::PendingCreate:
-		{
-			ensureMsgf(CurrentState == EReplicatedObjectState::Invalid || CurrentState == EReplicatedObjectState::WaitOnCreateConfirmation, TEXT("Trying to set state %s when state is %s. IsDestructionInfo: %u IsSubObject: %u"), LexToString(NewState), LexToString(CurrentState), IsDestructionInfo, IsSubObject);
-		}
-		break;
-		case EReplicatedObjectState::WaitOnCreateConfirmation:
-		{
-			ensureMsgf(CurrentState == EReplicatedObjectState::PendingCreate || CurrentState == EReplicatedObjectState::CancelPendingDestroy, TEXT("Trying to set state %s when state is %s. IsDestructionInfo: %u IsSubObject: %u"), LexToString(NewState), LexToString(CurrentState), IsDestructionInfo, IsSubObject);
-		}
-		break;
-		case EReplicatedObjectState::Created:
-		{
-			ensureMsgf(CurrentState == EReplicatedObjectState::PendingCreate || CurrentState == EReplicatedObjectState::WaitOnCreateConfirmation || CurrentState == EReplicatedObjectState::CancelPendingDestroy || CurrentState == EReplicatedObjectState::WaitOnFlush, TEXT("Trying to set state %s when state is %s. IsDestructionInfo: %u IsSubObject: %u"), LexToString(NewState), LexToString(CurrentState), IsDestructionInfo, IsSubObject);
-		}
-		break;
-		case EReplicatedObjectState::PendingTearOff:
-		{
-			ensureMsgf(CurrentState == EReplicatedObjectState::PendingTearOff || CurrentState == EReplicatedObjectState::WaitOnFlush || CurrentState == EReplicatedObjectState::WaitOnCreateConfirmation || CurrentState == EReplicatedObjectState::Created || CurrentState == EReplicatedObjectState::WaitOnDestroyConfirmation, TEXT("Trying to set state %s when state is %s. IsDestructionInfo: %u IsSubObject: %u"), LexToString(NewState), LexToString(CurrentState), IsDestructionInfo, IsSubObject);
-		}
-		break;
-		case EReplicatedObjectState::SubObjectPendingDestroy:
-		{
-			ensureMsgf(CurrentState == EReplicatedObjectState::PendingDestroy || CurrentState == EReplicatedObjectState::SubObjectPendingDestroy || CurrentState == EReplicatedObjectState::WaitOnCreateConfirmation || CurrentState == EReplicatedObjectState::Created || CurrentState == EReplicatedObjectState::WaitOnFlush || CurrentState == EReplicatedObjectState::WaitOnDestroyConfirmation, TEXT("Trying to set state %s when state is %s. IsDestructionInfo: %u IsSubObject: %u"), LexToString(NewState), LexToString(CurrentState), IsDestructionInfo, IsSubObject);
-		}
-		break;
-		case EReplicatedObjectState::WaitOnFlush:
-		{
-			ensureMsgf(CurrentState != EReplicatedObjectState::Invalid, TEXT("Trying to set state %s when state is %s. IsDestructionInfo: %u IsSubObject: %u"), LexToString(NewState), LexToString(CurrentState), IsDestructionInfo, IsSubObject);
-		}
-		break;
-		case EReplicatedObjectState::PendingDestroy:
-		{
-			ensureMsgf(CurrentState != EReplicatedObjectState::Invalid, TEXT("Trying to set state %s when state is %s. IsDestructionInfo: %u IsSubObject: %u"), LexToString(NewState), LexToString(CurrentState), IsDestructionInfo, IsSubObject);
-		}
-		break;
-		case EReplicatedObjectState::WaitOnDestroyConfirmation:
-		{
-			ensureMsgf(CurrentState >= EReplicatedObjectState::PendingTearOff, TEXT("Trying to set state %s when state is %s. IsDestructionInfo: %u IsSubObject: %u"), LexToString(NewState), LexToString(CurrentState), IsDestructionInfo, IsSubObject);
-		}
-		break;
-		case EReplicatedObjectState::CancelPendingDestroy:
-		{
-			ensureMsgf(CurrentState == EReplicatedObjectState::WaitOnDestroyConfirmation || CurrentState == EReplicatedObjectState::CancelPendingDestroy, TEXT("Trying to set state %s when state is %s. IsDestructionInfo: %u IsSubObject: %u"), LexToString(NewState), LexToString(CurrentState), IsDestructionInfo, IsSubObject);
-		}
-		break;
-		case EReplicatedObjectState::Destroyed:
-		{
-			ensureMsgf(CurrentState == EReplicatedObjectState::WaitOnDestroyConfirmation || CurrentState == EReplicatedObjectState::PendingTearOff || CurrentState == EReplicatedObjectState::CancelPendingDestroy, TEXT("Trying to set state %s when state is %s. IsDestructionInfo: %u IsSubObject: %u"), LexToString(NewState), LexToString(CurrentState), IsDestructionInfo, IsSubObject);
-		}
-		break;
-		case EReplicatedObjectState::PermanentlyDestroyed:
-		{
-			ensureMsgf(CurrentState == EReplicatedObjectState::Invalid || CurrentState == EReplicatedObjectState::WaitOnCreateConfirmation, TEXT("Trying to set state %s when state is %s. IsDestructionInfo: %u IsSubObject: %u"), LexToString(NewState), LexToString(CurrentState), IsDestructionInfo, IsSubObject);
-		}
-		break;
-		case EReplicatedObjectState::Invalid:
-		{
-			ensureMsgf(CurrentState == EReplicatedObjectState::PermanentlyDestroyed || CurrentState == EReplicatedObjectState::Destroyed || CurrentState == EReplicatedObjectState::PendingCreate, TEXT("Trying to set state %s when state is %s. IsDestructionInfo: %u IsSubObject: %u"), LexToString(NewState), LexToString(CurrentState), IsDestructionInfo, IsSubObject);
-		}
-		break;
-
-		default:
-			ensureMsgf(false, TEXT("Trying to set state %s when state is %s. IsDestructionInfo: %u IsSubObject: %u"), ToCStr(FString::FromInt(int(NewState))), LexToString(CurrentState), IsDestructionInfo, IsSubObject);
-			break;
-	};
-
-	State = (uint32)NewState;
-}
-
 
 // Default allocator for changemasks
 static FGlobalChangeMaskAllocator s_DefaultChangeMaskAllocator;
@@ -403,8 +329,78 @@ void FReplicationWriter::SetState(uint32 InternalIndex, EReplicatedObjectState N
 {
 	FReplicationInfo& Info = GetReplicationInfo(InternalIndex);
 
-	UE_LOG_REPLICATIONWRITER_CONN(TEXT("ReplicationWriter.SetState for ( InternalIndex: %u ) %s -> %s"), InternalIndex, LexToString(Info.GetState()), LexToString(NewState));
-	Info.SetState(NewState);
+	const EReplicatedObjectState CurrentState = Info.GetState();
+	UE_LOG_REPLICATIONWRITER_CONN(TEXT("ReplicationWriter.SetState for ( InternalIndex: %u ) %s -> %s"), InternalIndex, LexToString(CurrentState), LexToString(NewState));
+
+	switch (NewState)
+	{
+		case EReplicatedObjectState::PendingCreate:
+		{
+			ensureMsgf(CurrentState == EReplicatedObjectState::Invalid || CurrentState == EReplicatedObjectState::WaitOnCreateConfirmation, TEXT("Trying to set state %s when state is %s. Object: %s IsDestructionInfo: %u IsSubObject: %u"), LexToString(NewState), LexToString(CurrentState), ToCStr(NetRefHandleManager->PrintObjectFromIndex(InternalIndex)), Info.IsDestructionInfo, Info.IsSubObject);
+		}
+		break;
+		case EReplicatedObjectState::WaitOnCreateConfirmation:
+		{
+			ensureMsgf(CurrentState == EReplicatedObjectState::PendingCreate || CurrentState == EReplicatedObjectState::CancelPendingDestroy, TEXT("Trying to set state %s when state is %s. Object: %s IsDestructionInfo: %u IsSubObject: %u"), LexToString(NewState), LexToString(CurrentState), ToCStr(NetRefHandleManager->PrintObjectFromIndex(InternalIndex)), Info.IsDestructionInfo, Info.IsSubObject);
+		}
+		break;
+		case EReplicatedObjectState::Created:
+		{
+			ensureMsgf(CurrentState == EReplicatedObjectState::PendingCreate || CurrentState == EReplicatedObjectState::WaitOnCreateConfirmation || CurrentState == EReplicatedObjectState::CancelPendingDestroy || CurrentState == EReplicatedObjectState::WaitOnFlush, TEXT("Trying to set state %s when state is %s. Object: %s IsDestructionInfo: %u IsSubObject: %u"), LexToString(NewState), LexToString(CurrentState), ToCStr(NetRefHandleManager->PrintObjectFromIndex(InternalIndex)), Info.IsDestructionInfo, Info.IsSubObject);
+		}
+		break;
+		case EReplicatedObjectState::PendingTearOff:
+		{
+			ensureMsgf(CurrentState == EReplicatedObjectState::PendingTearOff || CurrentState == EReplicatedObjectState::WaitOnFlush || CurrentState == EReplicatedObjectState::WaitOnCreateConfirmation || CurrentState == EReplicatedObjectState::Created || CurrentState == EReplicatedObjectState::WaitOnDestroyConfirmation, TEXT("Trying to set state %s when state is %s. Object: %s IsDestructionInfo: %u IsSubObject: %u HasAttachments: %u"), LexToString(NewState), LexToString(CurrentState), ToCStr(NetRefHandleManager->PrintObjectFromIndex(InternalIndex)), Info.IsDestructionInfo, Info.IsSubObject, Info.HasAttachments);
+		}
+		break;
+		case EReplicatedObjectState::SubObjectPendingDestroy:
+		{
+			ensureMsgf(CurrentState == EReplicatedObjectState::PendingDestroy || CurrentState == EReplicatedObjectState::SubObjectPendingDestroy || CurrentState == EReplicatedObjectState::WaitOnCreateConfirmation || CurrentState == EReplicatedObjectState::Created || CurrentState == EReplicatedObjectState::WaitOnFlush || CurrentState == EReplicatedObjectState::WaitOnDestroyConfirmation, TEXT("Trying to set state %s when state is %s. Object: %s IsDestructionInfo: %u IsSubObject: %u"), LexToString(NewState), LexToString(CurrentState), ToCStr(NetRefHandleManager->PrintObjectFromIndex(InternalIndex)), Info.IsDestructionInfo, Info.IsSubObject);
+		}
+		break;
+		case EReplicatedObjectState::WaitOnFlush:
+		{
+			ensureMsgf(CurrentState != EReplicatedObjectState::Invalid, TEXT("Trying to set state %s when state is %s. Object: %s IsDestructionInfo: %u IsSubObject: %u"), LexToString(NewState), LexToString(CurrentState), ToCStr(NetRefHandleManager->PrintObjectFromIndex(InternalIndex)), Info.IsDestructionInfo, Info.IsSubObject);
+		}
+		break;
+		case EReplicatedObjectState::PendingDestroy:
+		{
+			ensureMsgf(CurrentState != EReplicatedObjectState::Invalid, TEXT("Trying to set state %s when state is %s. Object: %s IsDestructionInfo: %u IsSubObject: %u"), LexToString(NewState), LexToString(CurrentState), ToCStr(NetRefHandleManager->PrintObjectFromIndex(InternalIndex)), Info.IsDestructionInfo, Info.IsSubObject);
+		}
+		break;
+		case EReplicatedObjectState::WaitOnDestroyConfirmation:
+		{
+			ensureMsgf(CurrentState >= EReplicatedObjectState::PendingTearOff, TEXT("Trying to set state %s when state is %s. Object: %s IsDestructionInfo: %u IsSubObject: %u"), LexToString(NewState), LexToString(CurrentState), ToCStr(NetRefHandleManager->PrintObjectFromIndex(InternalIndex)), Info.IsDestructionInfo, Info.IsSubObject);
+		}
+		break;
+		case EReplicatedObjectState::CancelPendingDestroy:
+		{
+			ensureMsgf(CurrentState == EReplicatedObjectState::WaitOnDestroyConfirmation || CurrentState == EReplicatedObjectState::CancelPendingDestroy, TEXT("Trying to set state %s when state is %s. Object: %s IsDestructionInfo: %u IsSubObject: %u"), LexToString(NewState), LexToString(CurrentState), ToCStr(NetRefHandleManager->PrintObjectFromIndex(InternalIndex)), Info.IsDestructionInfo, Info.IsSubObject);
+		}
+		break;
+		case EReplicatedObjectState::Destroyed:
+		{
+			ensureMsgf(CurrentState == EReplicatedObjectState::WaitOnDestroyConfirmation || CurrentState == EReplicatedObjectState::PendingTearOff || CurrentState == EReplicatedObjectState::CancelPendingDestroy, TEXT("Trying to set state %s when state is %s. Object: %s IsDestructionInfo: %u IsSubObject: %u"), LexToString(NewState), LexToString(CurrentState), ToCStr(NetRefHandleManager->PrintObjectFromIndex(InternalIndex)), Info.IsDestructionInfo, Info.IsSubObject);
+		}
+		break;
+		case EReplicatedObjectState::PermanentlyDestroyed:
+		{
+			ensureMsgf(CurrentState == EReplicatedObjectState::Invalid || CurrentState == EReplicatedObjectState::WaitOnCreateConfirmation, TEXT("Trying to set state %s when state is %s. Object: %s IsDestructionInfo: %u IsSubObject: %u"), LexToString(NewState), LexToString(CurrentState), ToCStr(NetRefHandleManager->PrintObjectFromIndex(InternalIndex)), Info.IsDestructionInfo, Info.IsSubObject);
+		}
+		break;
+		case EReplicatedObjectState::Invalid:
+		{
+			ensureMsgf(CurrentState == EReplicatedObjectState::PermanentlyDestroyed || CurrentState == EReplicatedObjectState::Destroyed || CurrentState == EReplicatedObjectState::PendingCreate, TEXT("Trying to set state %s when state is %s. Object: %s IsDestructionInfo: %u IsSubObject: %u"), LexToString(NewState), LexToString(CurrentState), ToCStr(NetRefHandleManager->PrintObjectFromIndex(InternalIndex)), Info.IsDestructionInfo, Info.IsSubObject);
+		}
+		break;
+
+		default:
+			ensureMsgf(false, TEXT("Trying to set state %s when state is %s. Object: %s IsDestructionInfo: %u IsSubObject: %u"), ToCStr(FString::FromInt(int(NewState))), LexToString(CurrentState), ToCStr(NetRefHandleManager->PrintObjectFromIndex(InternalIndex)), Info.IsDestructionInfo, Info.IsSubObject);
+			break;
+	};
+
+	Info.State = (uint32)NewState;
 }
 
 void FReplicationWriter::Init(const FReplicationParameters& InParameters)
@@ -709,7 +705,7 @@ void FReplicationWriter::SetPendingDestroyOrSubObjectPendingDestroyState(uint32
 			FReplicationInfo& SubObjectInfo = GetReplicationInfo(SubObjectIndex);
 			if (SubObjectInfo.GetState() == EReplicatedObjectState::SubObjectPendingDestroy)
 			{
-				SubObjectInfo.SetState(EReplicatedObjectState::PendingDestroy);
+				SetState(SubObjectIndex, EReplicatedObjectState::PendingDestroy);
 				SubObjectInfo.SubObjectPendingDestroy = 0U;
 				ObjectsWithDirtyChanges.ClearBit(SubObjectIndex);
 			}
@@ -804,7 +800,7 @@ void FReplicationWriter::UpdateScope(const FNetBitArrayView& UpdatedScope)
 					FReplicationInfo& SubObjectInfo = GetReplicationInfo(SubObjectIndex);
 					if (SubObjectInfo.GetState() == EReplicatedObjectState::PendingDestroy)
 					{
-						SubObjectInfo.SetState(EReplicatedObjectState::SubObjectPendingDestroy);
+						SetState(SubObjectIndex, EReplicatedObjectState::SubObjectPendingDestroy);
 						SubObjectInfo.SubObjectPendingDestroy = 1U;
 
 						ObjectsWithDirtyChanges.SetBit(SubObjectIndex);
@@ -2172,19 +2168,19 @@ uint32 FReplicationWriter::WriteRootObjectsPendingDestroy(FNetSerializationConte
 			// Fill in subobject record
 			for (const uint32 SubObjectIndex : NetRefHandleManager->GetSubObjects(InternalIndex))
 			{
-				const FReplicationInfo& SubObjectInfo = GetReplicationInfo(SubObjectIndex);
+				FReplicationInfo& SubObjectInfo = GetReplicationInfo(SubObjectIndex);
 				if (SubObjectInfo.GetState() == EReplicatedObjectState::Invalid)
 				{
 					continue;
 				}
 
 				SetState(SubObjectIndex, EReplicatedObjectState::WaitOnDestroyConfirmation);
-				Info.HasDirtyChangeMask = 0U;
-				Info.SubObjectPendingDestroy = 0U;
+				SubObjectInfo.HasDirtyChangeMask = 0U;
+				SubObjectInfo.SubObjectPendingDestroy = 0U;
 
 				FReplicationRecord::FSubObjectRecord::FSubObjectInfo& SubObjectRecordInfo = SubObjectRecord.SubObjectInfos.AddDefaulted_GetRef();
 				SubObjectRecordInfo.Index = SubObjectIndex;
-				SubObjectRecordInfo.ReplicatedObjectState = static_cast<uint32>(Info.GetState());
+				SubObjectRecordInfo.ReplicatedObjectState = static_cast<uint32>(SubObjectInfo.GetState());
 			}
 
 			ObjectRecord.Record.HasSubObjectRecord = !SubObjectRecord.SubObjectInfos.IsEmpty();
diff --git a/Engine/Source/Runtime/Experimental/Iris/Core/Private/Iris/ReplicationSystem/ReplicationWriter.h b/Engine/Source/Runtime/Experimental/Iris/Core/Private/Iris/ReplicationSystem/ReplicationWriter.h
--- a/Engine/Source/Runtime/Experimental/Iris/Core/Private/Iris/ReplicationSystem/ReplicationWriter.h
+++ b/Engine/Source/Runtime/Experimental/Iris/Core/Private/Iris/ReplicationSystem/ReplicationWriter.h
@@ -137,7 +137,6 @@ public:
 
 		static const uint32 LocalChangeMaskMaxBitCount = 64u;
 
-		void SetState(EReplicatedObjectState NewState);
 		EReplicatedObjectState GetState() const { return (EReplicatedObjectState)State; }
 
 		ChangeMaskStorageType* GetChangeMaskStoragePointer() { return ChangeMaskOrPtr.GetPointer(ChangeMaskBitCount); }
