.jrn truncation tests

#rnx

diff --git a/Engine/Source/Runtime/Experimental/IoStore/OnDemand/Private/JournaledCache.cpp b/Engine/Source/Runtime/Experimental/IoStore/OnDemand/Private/JournaledCache.cpp
--- a/Engine/Source/Runtime/Experimental/IoStore/OnDemand/Private/JournaledCache.cpp
+++ b/Engine/Source/Runtime/Experimental/IoStore/OnDemand/Private/JournaledCache.cpp
@@ -983,6 +983,7 @@ static int32 LoadCache(FDiskCache& DiskCache)
 		{
 			return nullptr;
 		}
+
 		const uint32 BytesToConsume = sizeof(FDataEntry) * (FirstEntry->EntryCount);
 		if (IsOob(Cursor + BytesToConsume))
 		{
@@ -3030,6 +3031,71 @@ static void MiscTests(FSupport& Support)
 	}
 }
 
+////////////////////////////////////////////////////////////////////////////////
+static void Truncation(FSupport& Support)
+{
+	// To really stress-test truncated .jrn file, use -stompmalloc
+
+	using namespace JournaledCache;
+
+	const ANSICHAR* CacheName = "truncate";
+	FString CachePath = Support.TestDir / CacheName;
+
+	for (uint32 i = 0; i < 512; ++i)
+	{
+		Support.CleanFs();
+
+		auto NewCache = [&CachePath] {
+			FCache::FConfig Config;
+			Config.Path = CachePath;
+			Config.MemoryQuota = uint32(2_Mi);
+			Config.DiskQuota = 512_Mi;
+			Config.JournalQuota = uint32(128_Ki);
+			Config.DropCache = false;
+			return new FCache(MoveTemp(Config));
+		};
+
+		// Fill
+		FCache* Cache = NewCache();
+		enum { PhraseNum = 31, EntryNum = 13 };
+		for (uint32 p = 0; p < PhraseNum; ++p)
+		{
+			for (uint32 e = 0; e < EntryNum; ++e)
+			{
+				uint32 Size = (e & 1) ? 37 : 11;
+				FIoBuffer Data = Support.DummyData(Size);
+				Cache->Put(KeyGen(Data), Data);
+				Cache->WriteMemToDisk(1 << 30);
+			}
+			Cache->Flush();
+		}
+		delete Cache;
+
+		// Ordinary load
+		Cache = NewCache();
+		Cache->Load();
+		delete Cache;
+
+		// Truncate .jrn file
+		FString JrnPath = Support.TestDir;
+		JrnPath /= FString(CacheName);
+		JrnPath += GetCacheJrnSuffix();
+
+		IPlatformFile& Ipf = IPlatformFile::GetPlatformPhysical();
+		IFileHandle* JrnFile = Ipf.OpenWrite(*JrnPath, true, true);
+		uint64 JrnSize = JrnFile->Size();
+		check(JrnSize > i);
+		JrnSize -= i;
+		JrnFile->Truncate(JrnSize);
+		delete JrnFile;
+
+		// Survive the truncation
+		Cache = NewCache();
+		Cache->Load();
+		delete Cache;
+	}
+}
+
 ////////////////////////////////////////////////////////////////////////////////
 IOSTOREONDEMAND_API void Tests(const TCHAR* CacheDir=nullptr)
 {
@@ -3038,6 +3104,7 @@ IOSTOREONDEMAND_API void Tests(const TCHAR* CacheDir=nullptr)
 	MemCacheTests(Support);
 	CacheTests(Support);
 	BigCache(Support);
+	Truncation(Support);
 }
 
 } // namespace IasJournaledFileCacheTest
