Iris
-Added FArchive& operator<<(FArchive& Ar, UE::Net::FNetRefHandle& RefHandle) to facilitate writing NetRefHandles to FArchive

diff --git a/Engine/Source/Runtime/Experimental/Iris/Core/Private/Iris/ReplicationSystem/NetRefHandle.cpp b/Engine/Source/Runtime/Experimental/Iris/Core/Private/Iris/ReplicationSystem/NetRefHandle.cpp
--- a/Engine/Source/Runtime/Experimental/Iris/Core/Private/Iris/ReplicationSystem/NetRefHandle.cpp
+++ b/Engine/Source/Runtime/Experimental/Iris/Core/Private/Iris/ReplicationSystem/NetRefHandle.cpp
@@ -1,6 +1,7 @@
 // Copyright Epic Games, Inc. All Rights Reserved.
 
 #include "Iris/ReplicationSystem/NetRefHandle.h"
+#include "Iris/ReplicationSystem/NetRefHandleManager.h"
 #include "Containers/UnrealString.h"
 #include "Misc/StringBuilder.h"
 
@@ -51,3 +52,38 @@ FAnsiStringBuilderBase& operator<<(FAnsiStringBuilderBase& Builder, const UE::Ne
 
 	return Builder;
 }
+
+FArchive& operator<<(FArchive& Ar, UE::Net::FNetRefHandle& RefHandle)
+{
+	using namespace UE::Net;
+
+	if (Ar.IsSaving())
+	{
+		bool bIsValid = RefHandle.IsValid();
+		Ar.SerializeBits(&bIsValid, 1U);
+		if (bIsValid)
+		{
+			uint64 IdBits = 0U;
+			Ar.SerializeIntPacked64(IdBits);
+		}
+	}
+	else if (Ar.IsLoading())
+	{
+		FNetRefHandle Handle;
+
+		bool bIsValid = false;
+		Ar.SerializeBits(&bIsValid, 1U);
+		if (bIsValid)
+		{
+			uint64 NetId = 0U;
+			Ar.SerializeIntPacked64(NetId);		
+			if (!Ar.IsError())
+			{
+				Handle = Private::FNetRefHandleManager::MakeNetRefHandleFromId(NetId);
+			}
+		}
+
+		RefHandle = Handle;
+	}
+	return Ar;
+}
diff --git a/Engine/Source/Runtime/Experimental/Iris/Core/Public/Iris/ReplicationSystem/NetRefHandle.h b/Engine/Source/Runtime/Experimental/Iris/Core/Public/Iris/ReplicationSystem/NetRefHandle.h
--- a/Engine/Source/Runtime/Experimental/Iris/Core/Public/Iris/ReplicationSystem/NetRefHandle.h
+++ b/Engine/Source/Runtime/Experimental/Iris/Core/Public/Iris/ReplicationSystem/NetRefHandle.h
@@ -90,3 +90,4 @@ FORCEINLINE uint64 GetObjectIdForNetTrace(const FNetRefHandle& Handle)
 
 IRISCORE_API FStringBuilderBase& operator<<(FStringBuilderBase& Builder, const UE::Net::FNetRefHandle& NetHandle);
 IRISCORE_API FAnsiStringBuilderBase& operator<<(FAnsiStringBuilderBase& Builder, const UE::Net::FNetRefHandle& NetHandle);
+IRISCORE_API FArchive& operator<<(FArchive& Ar, UE::Net::FNetRefHandle& RefHandle);
