Crash in Rigid Solver
Make sure the evolution is not null

diff --git a/Engine/Source/Runtime/Experimental/Chaos/Private/PBDRigidsSolver.cpp b/Engine/Source/Runtime/Experimental/Chaos/Private/PBDRigidsSolver.cpp
--- a/Engine/Source/Runtime/Experimental/Chaos/Private/PBDRigidsSolver.cpp
+++ b/Engine/Source/Runtime/Experimental/Chaos/Private/PBDRigidsSolver.cpp
@@ -773,6 +773,12 @@ namespace Chaos
 		// NOTE: Do we really need these lists of proxies if we can just
 		// access them through the GTParticles list?
 		
+		// Evolution should be valid at this point, adding this check to avoid a server crash
+		if (!ensure(GetEvolution()))
+		{
+			UE_LOG(LogPBDRigidsSolver, Warning, TEXT("FPBDRigidsSolver::RegisterObject: Regiester a proxy while Evolution is still null."));
+			return;
+		}
 
 		RigidBody_External.SetUniqueIdx(GetEvolution()->GenerateUniqueIdx());
 		TrackGTParticle_External(*Proxy->GetParticle_LowLevel());	//todo: remove this
