Fix for serializing the ID in FNetRefHandle's operator<<.

diff --git a/Engine/Source/Runtime/Experimental/Iris/Core/Private/Iris/ReplicationSystem/NetRefHandle.cpp b/Engine/Source/Runtime/Experimental/Iris/Core/Private/Iris/ReplicationSystem/NetRefHandle.cpp
--- a/Engine/Source/Runtime/Experimental/Iris/Core/Private/Iris/ReplicationSystem/NetRefHandle.cpp
+++ b/Engine/Source/Runtime/Experimental/Iris/Core/Private/Iris/ReplicationSystem/NetRefHandle.cpp
@@ -63,7 +63,7 @@ FArchive& operator<<(FArchive& Ar, UE::Net::FNetRefHandle& RefHandle)
 		Ar.SerializeBits(&bIsValid, 1U);
 		if (bIsValid)
 		{
-			uint64 IdBits = 0U;
+			uint64 IdBits = RefHandle.GetId();
 			Ar.SerializeIntPacked64(IdBits);
 		}
 	}
