[NetPhysics] Fix bug where modular vehicle input type got cleared if input was not networked due to having a zero value.

diff --git a/Engine/Plugins/Experimental/ChaosModularVehicle/Source/ChaosModularVehicleEngine/Private/ChaosModularVehicle/ChaosSimModuleManagerAsyncCallback.cpp b/Engine/Plugins/Experimental/ChaosModularVehicle/Source/ChaosModularVehicleEngine/Private/ChaosModularVehicle/ChaosSimModuleManagerAsyncCallback.cpp
--- a/Engine/Plugins/Experimental/ChaosModularVehicle/Source/ChaosModularVehicleEngine/Private/ChaosModularVehicle/ChaosSimModuleManagerAsyncCallback.cpp
+++ b/Engine/Plugins/Experimental/ChaosModularVehicle/Source/ChaosModularVehicleEngine/Private/ChaosModularVehicle/ChaosSimModuleManagerAsyncCallback.cpp
@@ -389,6 +389,7 @@ bool FNetworkModularVehicleStates::NetSerialize(FArchive& Ar, class UPackageMap*
 		}
 	}
 
+	bOutSuccess = true;
 	return true;
 }
 
diff --git a/Engine/Source/Runtime/Experimental/ChaosVehicles/ChaosVehiclesCore/Private/SimModule/ModuleInput.cpp b/Engine/Source/Runtime/Experimental/ChaosVehicles/ChaosVehiclesCore/Private/SimModule/ModuleInput.cpp
--- a/Engine/Source/Runtime/Experimental/ChaosVehicles/ChaosVehiclesCore/Private/SimModule/ModuleInput.cpp
+++ b/Engine/Source/Runtime/Experimental/ChaosVehicles/ChaosVehiclesCore/Private/SimModule/ModuleInput.cpp
@@ -64,33 +64,43 @@ void FModuleInputValue::SetMagnitude(float NewSize)
 void FModuleInputValue::Serialize(FArchive& Ar, UPackageMap* Map, bool& bOutSuccess)
 {
 	// #TODO: send only value changes/deltas
-	Ar << ValueType;
 
-	switch (GetValueType())
+	bool bIsNonZero = Ar.IsLoading() ? false : IsQuantizedNonZero();
+	Ar.SerializeBits(&bIsNonZero, 1);
+	Ar.SerializeBits(&ValueType, 2);
+	Ar.SerializeBits(&bApplyInputDecay, 1);
+
+	if (bIsNonZero)
 	{
-	case EModuleInputValueType::MBoolean:
+		switch (GetValueType())
 		{
-			uint8 State = (Value.X != 0);
-			Ar << State;
-			if (Ar.IsLoading())
+		case EModuleInputValueType::MBoolean:
 			{
-				Value.X = State ? 1.0f : 0.0f;
+				bool bState = (Value.X != 0);
+				Ar.SerializeBits(&bState, 1);
+				if (Ar.IsLoading())
+				{
+					Value.X = bState ? 1.0f : 0.0f;
+				}
 			}
+			break;
+
+		case EModuleInputValueType::MAxis3D:
+			ModularQuantize::SerializeFixedFloat<1, 16>(Value.Z, Ar);
+		case EModuleInputValueType::MAxis2D:
+			ModularQuantize::SerializeFixedFloat<1, 16>(Value.Y, Ar);
+		case EModuleInputValueType::MAxis1D:
+			ModularQuantize::SerializeFixedFloat<1, 16>(Value.X, Ar);
+			break;
+
+		default:
+			checkf(false, TEXT("Unsupported value type for module input value!"));
+			break;
 		}
-		break;
-	case EModuleInputValueType::MAxis1D:
-		ModularQuantize::SerializeFixedFloat<1, 16>(Value.X, Ar);
-		break;
-	case EModuleInputValueType::MAxis2D:
-		ModularQuantize::SerializeFixedFloat<1, 16>(Value.X, Ar);
-		ModularQuantize::SerializeFixedFloat<1, 16>(Value.Y, Ar);
-		break;
-	case EModuleInputValueType::MAxis3D:
-		Value.NetSerialize(Ar, Map, bOutSuccess);
-		break;
-	default:
-		checkf(false, TEXT("Unsupported value type for module input value!"));
-		break;
+	}
+	else
+	{
+		Reset();
 	}
 	bOutSuccess = true;
 }
@@ -249,54 +259,14 @@ void FModuleInputContainer::Serialize(FArchive& Ar, UPackageMap* Map, bool& bOut
 	uint32 Number = InputValues.Num();
 	Ar.SerializeIntPacked(Number);
 
-	// Array of bits to mark which inputs to serialize or not
-	FNetBitArray InputsBitArray(Number);
-	FNetBitArray InputDecayBitArray(Number);
-
 	if (Ar.IsLoading())
 	{
-		Ar.SerializeBits(InputsBitArray.GetData(), Number);
-		Ar.SerializeBits(InputDecayBitArray.GetData(), Number);
-
-		InputValues.Reset(Number);
-		for (uint32 I = 0; I < Number; I++)
-		{
-			FModuleInputValue Value;
-
-			const bool bHasNonZeroValue = InputsBitArray.IsBitSet(I);
-			if (bHasNonZeroValue)
-			{
-				Value.Serialize(Ar, Map, bOutSuccess);
-				Value.SetApplyInputDecay(InputDecayBitArray.IsBitSet(I));
-			}
-			InputValues.EmplaceAt(I, Value);
-		}
+		InputValues.SetNum(Number);
 	}
-	else
-	{
-		// Only mark inputs for serialization if they have a non-zero value
-		for (uint32 I = 0; I < Number; I++)
-		{
-			if (InputValues[I].IsNonZero())
-			{
-				InputsBitArray.SetBit(I);
-				if (InputValues[I].ShouldApplyInputDecay())
-				{
-					InputDecayBitArray.SetBit(I);
-				}
-			}
-		}
-		Ar.SerializeBits(InputsBitArray.GetData(), Number);
-		Ar.SerializeBits(InputDecayBitArray.GetData(), Number);
 
-		for (uint32 I = 0; I < Number; I++)
-		{
-			const bool bHasNonZeroValue = InputsBitArray.IsBitSet(I);
-			if (bHasNonZeroValue)
-			{
-				InputValues[I].Serialize(Ar, Map, bOutSuccess);
-			}
-		}
+	for (uint32 I = 0; I < Number; I++)
+	{
+		InputValues[I].Serialize(Ar, Map, bOutSuccess);
 	}
 
 	return;
diff --git a/Engine/Source/Runtime/Experimental/ChaosVehicles/ChaosVehiclesCore/Public/SimModule/ModuleInput.h b/Engine/Source/Runtime/Experimental/ChaosVehicles/ChaosVehiclesCore/Public/SimModule/ModuleInput.h
--- a/Engine/Source/Runtime/Experimental/ChaosVehicles/ChaosVehiclesCore/Public/SimModule/ModuleInput.h
+++ b/Engine/Source/Runtime/Experimental/ChaosVehicles/ChaosVehiclesCore/Public/SimModule/ModuleInput.h
@@ -19,6 +19,8 @@ enum class EModuleInputValueType : uint8
 	MAxis1D		UMETA(DisplayName = "Axis1D (float)"),
 	MAxis2D		UMETA(DisplayName = "Axis2D (Vector2D)"),
 	MAxis3D		UMETA(DisplayName = "Axis3D (Vector)"),
+
+	// NOTE, If adding an entry here, update the number of bits to serialize in FModuleInputValue::Serialize
 };
 
 /** Input Options */
@@ -301,6 +303,7 @@ public:
 	float operator[](int32 Index) const { return Value[Index]; }
 
 	bool IsNonZero(float Tolerance = KINDA_SMALL_NUMBER) const { return Value.SizeSquared() >= Tolerance * Tolerance; }
+	bool IsQuantizedNonZero(float Tolerance = KINDA_SMALL_NUMBER) const { return ReturnQuantized().IsNonZero(Tolerance); }
 
 	// In-place type conversion
 	FModuleInputValue& ConvertToType(EModuleInputValueType Type)
