IAS cache eviction support

#rnx

diff --git a/Engine/Source/Runtime/Experimental/IoStore/OnDemand/Private/IasCache.h b/Engine/Source/Runtime/Experimental/IoStore/OnDemand/Private/IasCache.h
--- a/Engine/Source/Runtime/Experimental/IoStore/OnDemand/Private/IasCache.h
+++ b/Engine/Source/Runtime/Experimental/IoStore/OnDemand/Private/IasCache.h
@@ -51,6 +51,9 @@ public:
 
 	/** Get cache usage information. */
 	virtual void GetCacheUsage(uint64& OutDiskSize, uint64& OutMaxDiskSize) const = 0;
+
+	/** Evict an item from the cache. For example it might be corrupt */
+	virtual FIoStatus Evict(const FIoHash& Key) = 0;
 };
 
 struct FIasCacheConfig
diff --git a/Engine/Source/Runtime/Experimental/IoStore/OnDemand/Private/JournaledCache.cpp b/Engine/Source/Runtime/Experimental/IoStore/OnDemand/Private/JournaledCache.cpp
--- a/Engine/Source/Runtime/Experimental/IoStore/OnDemand/Private/JournaledCache.cpp
+++ b/Engine/Source/Runtime/Experimental/IoStore/OnDemand/Private/JournaledCache.cpp
@@ -124,6 +124,7 @@ public:
 	uint32			GetMax() const		{ return MaxSize; }
 	const FIoBuffer*Get(uint64 Key) const;
 	bool			Put(uint64 Key, FIoBuffer&& Data);
+	bool			Evict(uint64 Key, bool bAlways);
 	int32			Peel(int32 PeelThreshold, PeelItems& Out);
 	uint32			DebugVisit(void* Param, FDebugCacheEntry::Callback* Callback);
 
@@ -155,7 +156,7 @@ const FIoBuffer* FMemCache::Get(uint64 Key) const
 	{
 		if (Item.Key == Key)
 		{
-			return &(Item.Data);
+			return (Item.Data.GetSize() > 0) ? &(Item.Data) : nullptr;
 		}
 	}
 
@@ -192,19 +193,59 @@ bool FMemCache::Put(uint64 Key, FIoBuffer&& Data)
 	return true;
 }
 
+////////////////////////////////////////////////////////////////////////////////
+bool FMemCache::Evict(uint64 Key, bool bAlways)
+{
+	bool bEviction = false;
+	for (uint32 i = 0, n = Items.Num(); i < n; ++i)
+	{
+		FItem& Item = Items[i];
+		if (Item.Key != Key)
+		{
+			continue;
+		}
+
+		UsedSize -= int32(Item.Data.GetSize());
+		bEviction = true;
+
+		Item = MoveTemp(Items[n - 1]);
+		Items.Pop();
+		break;
+	}
+
+	// Even if there was no item cached in memory we may want to flow a dataless
+	// item through such that a record is made of the eviction.
+	if (bAlways)
+	{
+		Items.Add({ Key });
+	}
+
+	return bEviction;
+}
+
 ////////////////////////////////////////////////////////////////////////////////
 int32 FMemCache::Peel(int32 PeelThreshold, PeelItems& Out)
 {
 	TRACE_CPUPROFILER_EVENT_SCOPE(IasCache::Peel);
 
-	Algo::Sort(Items, [] (const FItem& Lhs, const FItem& Rhs) {
-		return Lhs.Data.GetSize() < Rhs.Data.GetSize();
+	// Items are sorted so that those we want to peel first are at the end. At
+	// the time of writing, largest items peel first. The exception is evicted
+	// items which have no data. Said evictions peeling priority.
+	Algo::Sort(Items, [] (const FItem& Lhs, const FItem& Rhs)
+	{
+		uint32 OpL = uint32(Lhs.Data.GetSize());
+		uint32 OpR = uint32(Rhs.Data.GetSize());
+		if (!OpL || !OpR)
+		{
+			return OpL > OpR;
+		}
+		return OpL < OpR;
 	});
 
 	// Add large items
 	int32 NumItems = Items.Num();
 	int32 DropSize = 0;
-	for (int32 i=Items.Num() - 1; i >= 0 && DropSize < PeelThreshold; --i)
+	for (int32 i = NumItems - 1; i >= 0 && DropSize < PeelThreshold; --i)
 	{
 		FItem& Item = Items[i];
 
@@ -565,6 +606,7 @@ public:
 	FDiskPhrase				OpenPhrase(uint32 DataSize);
 	void					ClosePhrase(FDiskPhrase&& Phrase);
 	bool					Has(uint64 Key) const;
+	bool					Evict(uint64 Key);
 	EIoErrorCode			Materialize(uint64 Key, FIoBuffer& Out, uint32 Offset=0) const;
 	int32					Flush();
 	void					Drop();
@@ -706,6 +748,13 @@ bool FDiskCache::Has(uint64 Key) const
 	return (DataMap.Find(Key) != nullptr);
 }
 
+////////////////////////////////////////////////////////////////////////////////
+bool FDiskCache::Evict(uint64 Key)
+{
+	FWriteScopeLock _(Lock);
+	return DataMap.Remove(Key) > 0;
+}
+
 ////////////////////////////////////////////////////////////////////////////////
 EIoErrorCode FDiskCache::Materialize(uint64 Key, FIoBuffer& Out, uint32 Offset) const
 {
@@ -758,6 +807,12 @@ uint64 FDiskCache::Insert(uint64 DataBase, const FDataEntry& Entry)
 		return 0;
 	}
 
+	if (Entry.Size == 0)
+	{
+		DataMap.Remove(Entry.Key);
+		return 0;
+	}
+
 	FMapEntry Value;
 	Value.DataCursor = DataBase + Entry.Offset;
 	check(Value.DataCursor < MaxDataSize);
@@ -1213,6 +1268,7 @@ public:
 	bool			Has(uint64 Key) const;
 	FGetToken		Get(uint64 Key, FIoBuffer& OutData) const;
 	bool			Put(uint64 Key, FIoBuffer& Data);
+	bool			Evict(uint64 Key);
 	EIoErrorCode	Materialize(FGetToken Token, FIoBuffer& OutData, uint32 Offset=0) const;
 	uint32			Flush();
 	uint32			WriteMemToDisk(int32 Allowance);
@@ -1331,6 +1387,24 @@ bool FCache::Put(uint64 Key, FIoBuffer& Data)
 	return Ok;
 }
 
+////////////////////////////////////////////////////////////////////////////////
+bool FCache::Evict(uint64 Key)
+{
+	bool Ok = DiskCache.Evict(Key);
+	{
+		FWriteScopeLock _(MemLock);
+		Ok |= MemCache.Evict(Key, Ok);
+	}
+
+	if (Partial.Key == Key)
+	{
+		Partial = FPartial();
+		Ok = true;
+	}
+
+	return (Ok != 0);
+}
+
 ////////////////////////////////////////////////////////////////////////////////
 EIoErrorCode FCache::Materialize(FGetToken Token, FIoBuffer& OutData, uint32 Offset) const
 {
@@ -2138,6 +2212,7 @@ public:
 	virtual void				Cancel(FIoBuffer& GivenDest) override;
 	virtual FIoStatus			Put(const FIoHash& Key, FIoBuffer& Data) override;
 	virtual void				GetCacheUsage(uint64& OutDiskSize, uint64& OutMaxDiskSize) const override;
+	virtual FIoStatus			Evict(const FIoHash& Key) override;
 
 private:
 	struct FMaterialOp
@@ -2312,6 +2387,14 @@ void FJournaledCache::GetCacheUsage(uint64& OutDiskSize, uint64& OutMaxDiskSize)
 	Cache->GetDiskUsage(OutDiskSize, OutMaxDiskSize);
 }
 
+////////////////////////////////////////////////////////////////////////////////
+FIoStatus FJournaledCache::Evict(const FIoHash& Key)
+{
+	uint64 InnerKey = ReduceKey(Key);
+	bool Ok = Cache->Evict(InnerKey);
+	return Ok ? FIoStatus::Ok : FIoStatus(EIoErrorCode::Unknown);
+}
+
 ////////////////////////////////////////////////////////////////////////////////
 uint64 FJournaledCache::ReduceKey(const FIoHash& Key)
 {
@@ -2707,6 +2790,58 @@ static void CacheTests(FSupport& Support)
 		NewCache(true);
 	}
 
+	// eviction
+	{
+		FIoBuffer Data;
+		Data = Support.DummyData(5); Cache->Put(5, Data);
+		Data = Support.DummyData(6); Cache->Put(6, Data);
+		Data = Support.DummyData(7); Cache->Put(7, Data);
+		Data = Support.DummyData(8); Cache->Put(8, Data);
+
+		// mem-cache evict
+		{
+			Data = FIoBuffer();
+			Cache->Get(6, Data); check(Data.GetData() != nullptr);
+			Cache->Evict(6);
+			Data = FIoBuffer();
+			check(!Cache->Has(6));
+		}
+
+		// disk-cache evict
+		{
+			Cache->WriteMemToDisk(uint32(1_Ki));
+			Cache->Flush();
+
+			Data = FIoBuffer();
+			check(Cache->Get(7, Data) != 0);
+			Cache->Evict(7);
+			Data = FIoBuffer();
+			check(!Cache->Has(7));
+
+			Cache->WriteMemToDisk(uint32(1_Ki));
+			Cache->Flush();
+
+			NewCache();
+			Cache->Load();
+			check( Cache->Has(5));
+			check(!Cache->Has(6));
+			check(!Cache->Has(7));
+			check( Cache->Has(8));
+		}
+
+		// partial evict
+		{
+			Data = Support.DummyData(2_Ki); Cache->Put(9, Data);
+			Cache->WriteMemToDisk(uint32(1_Ki));
+			Cache->Flush();
+			check(Cache->Has(9));
+			Cache->Evict(9);
+			check(!Cache->Has(9));
+		}
+
+		NewCache(true);
+	}
+
 	// journal wrap
 	{
 		FCache::FConfig Config;
