Fix for Resims not being able to successfully complete due to certain particles being marked bResimAsFollower.

diff --git a/Engine/Source/Runtime/Experimental/Chaos/Private/PBDRigidsSolver.cpp b/Engine/Source/Runtime/Experimental/Chaos/Private/PBDRigidsSolver.cpp
--- a/Engine/Source/Runtime/Experimental/Chaos/Private/PBDRigidsSolver.cpp
+++ b/Engine/Source/Runtime/Experimental/Chaos/Private/PBDRigidsSolver.cpp
@@ -1932,6 +1932,7 @@ namespace Chaos
 			{
 				// Clear ResimFrame if no valid resim frame was found
 				MRewindData->SetResimFrame(INDEX_NONE);
+				GetEvolution()->GetIslandManager().ResetParticleResimFrame();
 				return;
 			}
 
diff --git a/Engine/Source/Runtime/Experimental/Chaos/Private/RewindData.cpp b/Engine/Source/Runtime/Experimental/Chaos/Private/RewindData.cpp
--- a/Engine/Source/Runtime/Experimental/Chaos/Private/RewindData.cpp
+++ b/Engine/Source/Runtime/Experimental/Chaos/Private/RewindData.cpp
@@ -1188,6 +1188,8 @@ CHAOS_API int32 ResimFrameValidation = (int32)EResimFrameValidation::IslandValid
 FAutoConsoleVariableRef CVarResimFrameValidationLeniency(TEXT("p.Resim.ResimFrameValidation"), ResimFrameValidation, TEXT("0 = no leniency, all dirty particles need a valid target. 1 = Island leniency, all particles in resim islands need a valid target. 2 = Full leniency, only the particle triggering the resim need a valid target."));
 CHAOS_API bool bResimIncompleteHistory = false;
 FAutoConsoleVariableRef CVarResimIncompleteHistory(TEXT("p.Resim.IncompleteHistory"), bResimIncompleteHistory, TEXT("If a valid resim frame can't be found, use the requested resim frame and perform a resimulation with incomplete data."));
+CHAOS_API bool bUseParticleResimAsFollowerDuringTargetValidation = false;
+FAutoConsoleVariableRef CVarUseParticleResimAsFollowerDuringTargetValidation(TEXT("p.Resim.UseParticleResimAsFollowerDuringTargetValidation"), bUseParticleResimAsFollowerDuringTargetValidation, TEXT("If disabled, do not use the particle's ResimAsFollower flag when trying to find a valid resim frame."));
 
 int32 FRewindData::FindValidResimFrame(const int32 RequestedFrame)
 {
@@ -1207,7 +1209,11 @@ int32 FRewindData::FindValidResimFrame(const int32 RequestedFrame)
 	auto TargetFinderHelper = [&](FDirtyParticleInfo* DirtyParticleInfo, const FFrameAndPhase FrameAndPhase) -> bool
 	{
 		bool bValid = true;
-		const bool bResimAsFollower = DirtyParticleInfo->bResimAsFollower;
+		bool bResimAsFollower = false;
+		if (bUseParticleResimAsFollowerDuringTargetValidation)
+		{
+			bResimAsFollower = DirtyParticleInfo->bResimAsFollower;
+		}
 		FGeometryParticleStateBase& History = DirtyParticleInfo->GetHistory();
 		if (const FParticleDynamicMisc* DynamicMisc = History.DynamicsMisc.Read(FrameAndPhase, PropertiesPool))
 		{
