Geometry collection : Make sure direct dataflow asset properties are no longer used
- Set getter and setter on DataflowAsset property so that Blueprint write to the DataflowINstance correctly
- Deprecate the Terminal name property

diff --git a/Engine/Source/Runtime/Experimental/GeometryCollectionEngine/Private/GeometryCollection/GeometryCollectionObject.cpp b/Engine/Source/Runtime/Experimental/GeometryCollectionEngine/Private/GeometryCollection/GeometryCollectionObject.cpp
--- a/Engine/Source/Runtime/Experimental/GeometryCollectionEngine/Private/GeometryCollection/GeometryCollectionObject.cpp
+++ b/Engine/Source/Runtime/Experimental/GeometryCollectionEngine/Private/GeometryCollection/GeometryCollectionObject.cpp
@@ -1041,7 +1041,9 @@ void UGeometryCollection::Serialize(FArchive& Ar)
 	// Important : this needs to remain after the call to Super::Serialize
 	if (bStrippedDataflowData)
 	{
+#if WITH_EDITORONLY_DATA
 		DataflowAsset = StrippedDataflowAsset;
+#endif
 		DataflowInstance = StrippedDataflowInstance;
 	}
 
@@ -2046,10 +2048,10 @@ void UGeometryCollection::MigrateDeprecatedDataflowData()
 	if (DataflowAsset != nullptr)
 	{
 		DataflowInstance.SetDataflowAsset(DataflowAsset);
-		DataflowInstance.SetDataflowTerminal(FName(DataflowTerminal));
+		DataflowInstance.SetDataflowTerminal(FName(DataflowTerminal_DEPRECATED));
 
 		DataflowAsset = nullptr;
-		DataflowTerminal.Empty();
+		DataflowTerminal_DEPRECATED.Empty();
 
 		// todo(ccaillaud) do not migrate the overrides until we have a solid plan to make sure they won't sync and risk of loosing them 
 	
@@ -2165,6 +2167,16 @@ const TArray<UAssetUserData*>* UGeometryCollection::GetAssetUserDataArray() cons
 	return &ToRawPtrTArrayUnsafe(AssetUserData);
 }
 
+void UGeometryCollection::SetDataflowAsset(UDataflow* InDataflowAsset)
+{
+	DataflowInstance.SetDataflowAsset(InDataflowAsset);
+}
+
+UDataflow* UGeometryCollection::GetDataflowAsset() const
+{
+	return DataflowInstance.GetDataflowAsset();
+}
+
 TObjectPtr<UDataflowBaseContent> UGeometryCollection::CreateDataflowContent()
 {
 	TObjectPtr<UDataflowBaseContent> BaseContent = NewObject<UDataflowBaseContent>();
diff --git a/Engine/Source/Runtime/Experimental/GeometryCollectionEngine/Public/GeometryCollection/GeometryCollectionObject.h b/Engine/Source/Runtime/Experimental/GeometryCollectionEngine/Public/GeometryCollection/GeometryCollectionObject.h
--- a/Engine/Source/Runtime/Experimental/GeometryCollectionEngine/Public/GeometryCollection/GeometryCollectionObject.h
+++ b/Engine/Source/Runtime/Experimental/GeometryCollectionEngine/Public/GeometryCollection/GeometryCollectionObject.h
@@ -831,11 +831,15 @@ public:
 	//
 	// Dataflow
 	//
-	UPROPERTY(BlueprintReadWrite, Category = "Dataflow")
+#if WITH_EDITORONLY_DATA
+	// this is deprecated data replaced with Setter and Getter but it is exposed to blueprint 
+	// so we use the Setter/Getter to make sure the blueprint will use them 
+	UPROPERTY(BlueprintReadWrite, BlueprintSetter=SetDataflowAsset, BlueprintGetter = GetDataflowAsset, Category = "Dataflow")
 	TObjectPtr<UDataflow> DataflowAsset;
 
-	UPROPERTY()
-	FString DataflowTerminal = "GeometryCollectionTerminal";
+	UPROPERTY(meta = (DeprecatedProperty, DeprecationMessage = "Dataflow terminal name is now stored in DataflowInstance"))
+	FString DataflowTerminal_DEPRECATED = "GeometryCollectionTerminal";
+#endif
 
 	UPROPERTY(EditAnywhere, BlueprintReadWrite, Category = "Dataflow", DisplayName = "DataFlow Overrides", meta = (DisplayAfter = "DataflowInstance"))
 	TMap<FString, FString> Overrides;
@@ -843,10 +847,16 @@ public:
 	UPROPERTY(EditAnywhere, Category = "Dataflow")
 	FDataflowInstance DataflowInstance;
 
+	UFUNCTION(BlueprintCallable, Category = "Dataflow")
+	void SetDataflowAsset(UDataflow* InDataflowAsset);
+
+	UFUNCTION(BlueprintCallable, Category = "Dataflow")
+	UDataflow* GetDataflowAsset() const;
+
 	//~ Begin IDataflowGeneratedAssetInterface
 	virtual const FDataflowInstance& GetDataflowInstance() const { return DataflowInstance; };
 	//~ End IDataflowGeneratedAssetInterface
-	
+
 	//~ Begin IDataflowContentOwner interface
 	virtual TObjectPtr<UDataflowBaseContent> CreateDataflowContent() override;
 	virtual void WriteDataflowContent(const TObjectPtr<UDataflowBaseContent>& DataflowContent) const override;
