[ChaosVD] Fix ensure that happens during a CVD recording when a midphase in the broad phase is recorded

This seems to be fallout of the refactor done to remove the use of parallelfor in the broadphase in CL 38937679

CVD uses TLS to gather the context in which data is serialized (Just 2 integers used as ID), and that context needs to be pushed to any worked thread TLS before data is traced. After the refactor, the new tasks in the task graph where not pushing the CVD context in the worker's thread TLS causing the issue.

This means that since that recent refactor, the broad phase data was not being recorded.

diff --git a/Engine/Source/Runtime/Experimental/Chaos/Private/Chaos/Collision/SpatialAccelerationBroadPhase.cpp b/Engine/Source/Runtime/Experimental/Chaos/Private/Chaos/Collision/SpatialAccelerationBroadPhase.cpp
--- a/Engine/Source/Runtime/Experimental/Chaos/Private/Chaos/Collision/SpatialAccelerationBroadPhase.cpp
+++ b/Engine/Source/Runtime/Experimental/Chaos/Private/Chaos/Collision/SpatialAccelerationBroadPhase.cpp
@@ -10,6 +10,12 @@ namespace Chaos
 	}
 	namespace Private
 	{
+		FAssignMidPhaseTask::FAssignMidPhaseTask(FSpatialAccelerationBroadPhase& InSpatialAccelerationBroadPhase,Private::FBroadPhaseContext& InBroadPhaseContext)
+		: SpatialAccelerationBroadPhase(InSpatialAccelerationBroadPhase)
+		, BroadPhaseContext(InBroadPhaseContext)
+		{
+			CVD_GET_CURRENT_CONTEXT(CVDContext);
+		}
 
 		ENamedThreads::Type FAssignMidPhaseTask::GetDesiredThread()
 		{
@@ -17,15 +23,25 @@ namespace Chaos
 		}
 		void FAssignMidPhaseTask::Process()
 		{
+			CVD_SCOPE_CONTEXT(CVDContext);
 			SpatialAccelerationBroadPhase.AssignMidPhases(BroadPhaseContext);
 		}
 
+		FProcessMidPhaseTask::FProcessMidPhaseTask(FSpatialAccelerationBroadPhase& InSpatialAccelerationBroadPhase, Private::FBroadPhaseContext& InBroadPhaseContext, FReal InDt)
+		: SpatialAccelerationBroadPhase(InSpatialAccelerationBroadPhase)
+		, BroadPhaseContext(InBroadPhaseContext)
+		, Dt(InDt)
+		{
+			CVD_GET_CURRENT_CONTEXT(CVDContext);
+		}
+
 		ENamedThreads::Type FProcessMidPhaseTask::GetDesiredThread()
 		{
 			return ENamedThreads::AnyHiPriThreadHiPriTask;
 		}
 		void FProcessMidPhaseTask::Process()
 		{
+			CVD_SCOPE_CONTEXT(CVDContext);
 			SpatialAccelerationBroadPhase.ProcessMidPhases(Dt, BroadPhaseContext);
 		}
 	}
diff --git a/Engine/Source/Runtime/Experimental/Chaos/Public/Chaos/Collision/SpatialAccelerationBroadPhase.h b/Engine/Source/Runtime/Experimental/Chaos/Public/Chaos/Collision/SpatialAccelerationBroadPhase.h
--- a/Engine/Source/Runtime/Experimental/Chaos/Public/Chaos/Collision/SpatialAccelerationBroadPhase.h
+++ b/Engine/Source/Runtime/Experimental/Chaos/Public/Chaos/Collision/SpatialAccelerationBroadPhase.h
@@ -301,9 +301,7 @@ namespace Chaos
 		class FAssignMidPhaseTask
 		{
 		public:
-			FAssignMidPhaseTask(FSpatialAccelerationBroadPhase& InSpatialAccelerationBroadPhase, Private::FBroadPhaseContext& InBroadPhaseContext)
-				: SpatialAccelerationBroadPhase(InSpatialAccelerationBroadPhase)
-				, BroadPhaseContext(InBroadPhaseContext) {}
+			FAssignMidPhaseTask(FSpatialAccelerationBroadPhase& InSpatialAccelerationBroadPhase, Private::FBroadPhaseContext& InBroadPhaseContext);
 
 			TStatId GetStatId() const { RETURN_QUICK_DECLARE_CYCLE_STAT(FAssignMidPhaseTask, STATGROUP_TaskGraphTasks); }
 			static CHAOS_API ENamedThreads::Type GetDesiredThread();
@@ -314,15 +312,16 @@ namespace Chaos
 			void Process();
 			FSpatialAccelerationBroadPhase& SpatialAccelerationBroadPhase;
 			Private::FBroadPhaseContext& BroadPhaseContext;
+
+#if WITH_CHAOS_VISUAL_DEBUGGER
+			FChaosVDContext CVDContext;
+#endif
 		};
 
 		class FProcessMidPhaseTask
 		{
 		public:
-			FProcessMidPhaseTask(FSpatialAccelerationBroadPhase& InSpatialAccelerationBroadPhase, Private::FBroadPhaseContext& InBroadPhaseContext, FReal InDt)
-				: SpatialAccelerationBroadPhase(InSpatialAccelerationBroadPhase)
-				, BroadPhaseContext(InBroadPhaseContext)
-				, Dt(InDt) {}
+			FProcessMidPhaseTask(FSpatialAccelerationBroadPhase& InSpatialAccelerationBroadPhase, Private::FBroadPhaseContext& InBroadPhaseContext, FReal InDt);
 
 			TStatId GetStatId() const { RETURN_QUICK_DECLARE_CYCLE_STAT(FProcessMidPhaseTask, STATGROUP_TaskGraphTasks); }
 			static CHAOS_API ENamedThreads::Type GetDesiredThread();
@@ -333,6 +332,10 @@ namespace Chaos
 			FSpatialAccelerationBroadPhase& SpatialAccelerationBroadPhase;
 			Private::FBroadPhaseContext& BroadPhaseContext;
 			FReal Dt;
+
+#if WITH_CHAOS_VISUAL_DEBUGGER
+			FChaosVDContext CVDContext;
+#endif
 		};
 	}
 
