Fix API to allow dataflow anytype to be registered outside of the dataflow modules

diff --git a/Engine/Source/Runtime/Experimental/Dataflow/Core/Public/Dataflow/DataflowAnyTypeRegistry.h b/Engine/Source/Runtime/Experimental/Dataflow/Core/Public/Dataflow/DataflowAnyTypeRegistry.h
--- a/Engine/Source/Runtime/Experimental/Dataflow/Core/Public/Dataflow/DataflowAnyTypeRegistry.h
+++ b/Engine/Source/Runtime/Experimental/Dataflow/Core/Public/Dataflow/DataflowAnyTypeRegistry.h
@@ -19,7 +19,7 @@ namespace UE::Dataflow
 		static bool AreTypesCompatibleStatic(FName TypeA, FName TypeB);
 		static FName GetStorageTypeStatic(FName Type);
 		static bool IsAnyTypeStatic(FName Type);
-		static FAnyTypesRegistry& GetInstance();
+		DATAFLOWCORE_API static FAnyTypesRegistry& GetInstance();
 
 	private:
 		typedef bool(*FSupportTypeFunction)(FName);
diff --git a/Engine/Source/Runtime/Experimental/Dataflow/Core/Public/Dataflow/DataflowTypePolicy.h b/Engine/Source/Runtime/Experimental/Dataflow/Core/Public/Dataflow/DataflowTypePolicy.h
--- a/Engine/Source/Runtime/Experimental/Dataflow/Core/Public/Dataflow/DataflowTypePolicy.h
+++ b/Engine/Source/Runtime/Experimental/Dataflow/Core/Public/Dataflow/DataflowTypePolicy.h
@@ -83,8 +83,6 @@ UE_DATAFLOW_POLICY_DECLARE_TYPENAME(FDataflowGeometrySelection)
 UE_DATAFLOW_POLICY_DECLARE_TYPENAME(FDataflowMaterialSelection)
 
 
-#undef UE_DATAFLOW_POLICY_DECLARE_TYPENAME
-
 struct IDataflowTypePolicy
 {
 	virtual bool SupportsType(FName Type) const = 0;
