Dataflow asset : Fix variable nodes not refreshing properly when their type changed

diff --git a/Engine/Source/Runtime/Experimental/Dataflow/Engine/Private/Dataflow/DataflowObject.cpp b/Engine/Source/Runtime/Experimental/Dataflow/Engine/Private/Dataflow/DataflowObject.cpp
--- a/Engine/Source/Runtime/Experimental/Dataflow/Engine/Private/Dataflow/DataflowObject.cpp
+++ b/Engine/Source/Runtime/Experimental/Dataflow/Engine/Private/Dataflow/DataflowObject.cpp
@@ -395,6 +395,34 @@ const TArray<TObjectPtr<UDataflowSubGraph>>& UDataflow::GetSubGraphs() const
 	return DataflowSubGraphs;
 }
 
+void UDataflow::RefreshEdNodeByGuid(const FGuid NodeGuid)
+{
+	RefreshEdNode(FindEdNodeByDataflowNodeGuid(NodeGuid));
+}
+
+void UDataflow::RefreshEdNode(TObjectPtr<UDataflowEdNode> EdNode)
+{
+	if (EdNode)
+	{
+		EdNode->UpdatePinsFromDataflowNode();
+		if (Nodes.Contains(EdNode))
+		{
+			NotifyNodeChanged(EdNode);
+			//NotifyGraphChanged();
+			return;
+		}
+		for (TObjectPtr<UDataflowSubGraph> Subgraph : DataflowSubGraphs)
+		{
+			if (Subgraph && Subgraph->Nodes.Contains(EdNode))
+			{
+				Subgraph->NotifyNodeChanged(EdNode);
+				//Subgraph->NotifyGraphChanged();
+				return;
+			}
+		}
+	}
+}
+
 #if WITH_EDITOR
 bool UDataflow::CanEditChange(const FProperty* InProperty) const
 {
diff --git a/Engine/Source/Runtime/Experimental/Dataflow/Engine/Private/Dataflow/DataflowVariableNodes.cpp b/Engine/Source/Runtime/Experimental/Dataflow/Engine/Private/Dataflow/DataflowVariableNodes.cpp
--- a/Engine/Source/Runtime/Experimental/Dataflow/Engine/Private/Dataflow/DataflowVariableNodes.cpp
+++ b/Engine/Source/Runtime/Experimental/Dataflow/Engine/Private/Dataflow/DataflowVariableNodes.cpp
@@ -171,6 +171,11 @@ void FGetDataflowVariableNode::UpdateOutputTypes(const FPropertyBagPropertyDesc&
 	}
 
 	Output->LockType();
+
+	if (TStrongObjectPtr<UDataflow> DataflowAsset = WeakDataflowPtr.Pin())
+	{
+		DataflowAsset->RefreshEdNodeByGuid(GetGuid());
+	}
 }
 
 void FGetDataflowVariableNode::OnPropertyChanged(UE::Dataflow::FContext& Context, const FPropertyChangedEvent& InPropertyChangedEvent)
diff --git a/Engine/Source/Runtime/Experimental/Dataflow/Engine/Public/Dataflow/DataflowObject.h b/Engine/Source/Runtime/Experimental/Dataflow/Engine/Public/Dataflow/DataflowObject.h
--- a/Engine/Source/Runtime/Experimental/Dataflow/Engine/Public/Dataflow/DataflowObject.h
+++ b/Engine/Source/Runtime/Experimental/Dataflow/Engine/Public/Dataflow/DataflowObject.h
@@ -187,6 +187,13 @@ public:
 
 	DATAFLOWENGINE_API const TArray<TObjectPtr<UDataflowSubGraph>>& GetSubGraphs() const;
 
+	/** 
+	* Make sure EdNode gets refreshed from its dataflow node 
+	* This is useful when connections are changed
+	*/
+	DATAFLOWENGINE_API void RefreshEdNode(TObjectPtr<UDataflowEdNode> EdNode);
+	DATAFLOWENGINE_API void RefreshEdNodeByGuid(const FGuid NodeGuid);
+
 #if WITH_EDITORONLY_DATA
 
 	/*
