Add integer option to modular vehicle inputs

diff --git a/Engine/Plugins/Experimental/ChaosModularVehicle/Source/ChaosModularVehicleEngine/Private/ChaosModularVehicle/ModularVehicleBaseComponent.cpp b/Engine/Plugins/Experimental/ChaosModularVehicle/Source/ChaosModularVehicleEngine/Private/ChaosModularVehicle/ModularVehicleBaseComponent.cpp
--- a/Engine/Plugins/Experimental/ChaosModularVehicle/Source/ChaosModularVehicleEngine/Private/ChaosModularVehicle/ModularVehicleBaseComponent.cpp
+++ b/Engine/Plugins/Experimental/ChaosModularVehicle/Source/ChaosModularVehicleEngine/Private/ChaosModularVehicle/ModularVehicleBaseComponent.cpp
@@ -1415,6 +1415,15 @@ void UModularVehicleBaseComponent::SetInputBool(const FName Name, const bool Val
 	}
 }
 
+void UModularVehicleBaseComponent::SetInputInteger(const FName Name, const int32 Value)
+{
+	if (InputProducer)
+	{
+		InputProducer->BufferInput(InputNameMap, Name, Value);
+	}
+}
+
+
 void UModularVehicleBaseComponent::SetInputAxis1D(const FName Name, const double Value)
 {
 	if (InputProducer)
diff --git a/Engine/Plugins/Experimental/ChaosModularVehicle/Source/ChaosModularVehicleEngine/Public/ChaosModularVehicle/ModularVehicleBaseComponent.h b/Engine/Plugins/Experimental/ChaosModularVehicle/Source/ChaosModularVehicleEngine/Public/ChaosModularVehicle/ModularVehicleBaseComponent.h
--- a/Engine/Plugins/Experimental/ChaosModularVehicle/Source/ChaosModularVehicleEngine/Public/ChaosModularVehicle/ModularVehicleBaseComponent.h
+++ b/Engine/Plugins/Experimental/ChaosModularVehicle/Source/ChaosModularVehicleEngine/Public/ChaosModularVehicle/ModularVehicleBaseComponent.h
@@ -213,6 +213,9 @@ public:
 	UFUNCTION(BlueprintCallable, Category = "Game|Components|ModularVehicle")
 	void SetInputBool(const FName Name, const bool Value);
 
+	UFUNCTION(BlueprintCallable, Category = "Game|Components|ModularVehicle")
+	void SetInputInteger(const FName Name, const int32 Value);
+
 	UFUNCTION(BlueprintCallable, Category = "Game|Components|ModularVehicle")
 	void SetInputAxis1D(const FName Name, const double Value);
 
diff --git a/Engine/Source/Runtime/Experimental/ChaosVehicles/ChaosVehiclesCore/Private/SimModule/ModuleInput.cpp b/Engine/Source/Runtime/Experimental/ChaosVehicles/ChaosVehiclesCore/Private/SimModule/ModuleInput.cpp
--- a/Engine/Source/Runtime/Experimental/ChaosVehicles/ChaosVehiclesCore/Private/SimModule/ModuleInput.cpp
+++ b/Engine/Source/Runtime/Experimental/ChaosVehicles/ChaosVehiclesCore/Private/SimModule/ModuleInput.cpp
@@ -7,11 +7,31 @@ TArray<struct FModuleInputSetup>* FScopedModuleInputInitializer::InitSetupData =
 
 DEFINE_LOG_CATEGORY(LogModularInput);
 
+bool FModuleInputValue::IsNonZero(float Tolerance) const
+{
+	switch (ValueType)
+	{
+	case EModuleInputValueType::MBoolean:
+	case EModuleInputValueType::MInteger:
+		return ValueInt != 0;
+
+	case EModuleInputValueType::MAxis2D:
+	case EModuleInputValueType::MAxis1D:
+	case EModuleInputValueType::MAxis3D:
+		return Value.SizeSquared() >= Tolerance * Tolerance;
+	}
+
+	checkf(false, TEXT("Unsupported value type for module input value!"));
+	return false;
+}
 float FModuleInputValue::GetMagnitudeSq() const
 {
 	switch (GetValueType())
 	{
 	case EModuleInputValueType::MBoolean:
+	case EModuleInputValueType::MInteger:
+		return ValueInt * ValueInt;
+
 	case EModuleInputValueType::MAxis1D:
 		return Value.X * Value.X;
 	case EModuleInputValueType::MAxis2D:
@@ -29,6 +49,10 @@ float FModuleInputValue::GetMagnitude() const
 	switch (GetValueType())
 	{
 	case EModuleInputValueType::MBoolean:
+	case EModuleInputValueType::MInteger:
+		return ValueInt;
+	break;
+
 	case EModuleInputValueType::MAxis1D:
 		return Value.X;
 	case EModuleInputValueType::MAxis2D:
@@ -41,11 +65,36 @@ float FModuleInputValue::GetMagnitude() const
 	return 0.f;
 }
 
+int32 FModuleInputValue::GetMagnitudeInt() const
+{
+	float Tolerance = KINDA_SMALL_NUMBER;
+
+	switch (GetValueType())
+	{
+	case EModuleInputValueType::MBoolean:
+	case EModuleInputValueType::MInteger:
+		return ValueInt;
+
+	case EModuleInputValueType::MAxis2D:
+	case EModuleInputValueType::MAxis1D:
+	case EModuleInputValueType::MAxis3D:
+		return Value.SizeSquared() >= Tolerance * Tolerance;
+
+	}
+
+	checkf(false, TEXT("Unsupported value type for module input value!"));
+	return 0.f;
+}
+
 void FModuleInputValue::SetMagnitude(float NewSize)
 {
 	switch (GetValueType())
 	{
 	case EModuleInputValueType::MBoolean:
+	case EModuleInputValueType::MInteger:
+		ValueInt = (int32)(NewSize);
+	break;
+
 	case EModuleInputValueType::MAxis1D:
 		Value.X = NewSize;
 		break;
@@ -67,7 +116,7 @@ void FModuleInputValue::Serialize(FArchive& Ar, UPackageMap* Map, bool& bOutSucc
 
 	bool bIsNonZero = Ar.IsLoading() ? false : IsQuantizedNonZero();
 	Ar.SerializeBits(&bIsNonZero, 1);
-	Ar.SerializeBits(&ValueType, 2);
+	Ar.SerializeBits(&ValueType, 3);
 	Ar.SerializeBits(&bApplyInputDecay, 1);
 
 	if (bIsNonZero)
@@ -76,15 +125,21 @@ void FModuleInputValue::Serialize(FArchive& Ar, UPackageMap* Map, bool& bOutSucc
 		{
 		case EModuleInputValueType::MBoolean:
 			{
-				bool bState = (Value.X != 0);
+				bool bState = (ValueInt != 0);
 				Ar.SerializeBits(&bState, 1);
 				if (Ar.IsLoading())
 				{
-					Value.X = bState ? 1.0f : 0.0f;
+					ValueInt = bState ? 1.0f : 0.0f;
 				}
 			}
 			break;
 
+		case EModuleInputValueType::MInteger:
+		{
+			Ar.SerializeIntPacked((uint32&)ValueInt);
+		}
+		break;
+
 		case EModuleInputValueType::MAxis3D:
 			ModularQuantize::SerializeFixedFloat<1, 16>(Value.Z, Ar);
 		case EModuleInputValueType::MAxis2D:
@@ -105,17 +160,37 @@ void FModuleInputValue::Serialize(FArchive& Ar, UPackageMap* Map, bool& bOutSucc
 	bOutSuccess = true;
 }
 
+void FModuleInputValue::Lerp(const FModuleInputValue& Min, const FModuleInputValue& Max, float Alpha)
+{
+	switch (GetValueType())
+	{
+	case EModuleInputValueType::MBoolean:
+	case EModuleInputValueType::MInteger:
+		ValueInt = Max.ValueInt; // don't lerp just use last value, is this correct for both cases?
+		break;
+
+	default:
+		Value = FMath::Lerp(Min.Value, Max.Value, Alpha);
+		break;
+	}
+}
+
 void FModuleInputValue::Merge(const FModuleInputValue& From)
 {
 	switch (GetValueType())
 	{
 	case EModuleInputValueType::MBoolean:
 		// ensure we capture edges of digitial inputs by taking largest absolute value
-		if (FMath::Abs(From.Value.X) >= FMath::Abs(Value.X))
+		if (FMath::Abs(From.ValueInt) >= FMath::Abs(ValueInt))
 		{
-			Value.X = From.Value.X;
+			ValueInt = From.ValueInt;
 		}
 		break;
+
+	case EModuleInputValueType::MInteger:
+		ValueInt = From.ValueInt; // use the last known value
+		break;
+
 	case EModuleInputValueType::MAxis1D:
 	case EModuleInputValueType::MAxis2D:
 	case EModuleInputValueType::MAxis3D:
@@ -140,6 +215,9 @@ void FModuleInputValue::Decay(const float DecayAmount)
 	case EModuleInputValueType::MBoolean:
 		// Don't decay boolean
 		break;
+	case EModuleInputValueType::MInteger:
+		// Don't decay integers
+		break;
 	case EModuleInputValueType::MAxis1D:
 	case EModuleInputValueType::MAxis2D:
 	case EModuleInputValueType::MAxis3D:
@@ -158,6 +236,8 @@ FString FModuleInputValue::ToString() const
 	{
 	case EModuleInputValueType::MBoolean:
 		return FString(IsNonZero() ? TEXT("true") : TEXT("false"));
+	case EModuleInputValueType::MInteger:
+		return FString(TEXT("%d"), ValueInt);
 	case EModuleInputValueType::MAxis1D:
 		return FString::Printf(TEXT("%3.3f"), Value.X);
 	case EModuleInputValueType::MAxis2D:
@@ -274,7 +354,17 @@ void FModuleInputContainer::Serialize(FArchive& Ar, UPackageMap* Map, bool& bOut
 
 int FModuleInputContainer::AddInput(EModuleInputValueType Type, TSubclassOf<UDefaultModularVehicleInputModifier>& InputModifierClass)
 {
-	InputValues.Add(FModuleInputValue(Type, FVector::ZeroVector));
+	switch(Type)
+	{
+		case EModuleInputValueType::MBoolean:
+		case EModuleInputValueType::MInteger:
+			InputValues.Add(FModuleInputValue(Type, 0));
+		break;
+
+		default:
+			InputValues.Add(FModuleInputValue(Type, FVector::ZeroVector));
+		break;
+	}
 	return InputValues.Num() - 1;
 }
 
@@ -369,6 +459,19 @@ float FInputInterface::GetMagnitude(const FName& InName) const
 	return 0.0f;
 }
 
+int32 FInputInterface::GetMagnitudeInt(const FName& InName) const
+{
+	if (ValueContainer.GetNumInputs() > 0)
+	{
+		if (const int* Index = NameMap.Find(InName))
+		{
+			return ValueContainer.GetValueAtIndex(*Index).GetMagnitudeInt();
+		}
+	}
+
+	return 0;
+}
+
 bool FInputInterface::InputsNonZero() const
 {
 	for (int I = 0; I < ValueContainer.GetNumInputs(); I++)
diff --git a/Engine/Source/Runtime/Experimental/ChaosVehicles/ChaosVehiclesCore/Public/SimModule/ModuleInput.h b/Engine/Source/Runtime/Experimental/ChaosVehicles/ChaosVehiclesCore/Public/SimModule/ModuleInput.h
--- a/Engine/Source/Runtime/Experimental/ChaosVehicles/ChaosVehiclesCore/Public/SimModule/ModuleInput.h
+++ b/Engine/Source/Runtime/Experimental/ChaosVehicles/ChaosVehiclesCore/Public/SimModule/ModuleInput.h
@@ -16,6 +16,7 @@ enum class EModuleInputValueType : uint8
 	// Value types in increasing size order (used for type promotion)
 
 	MBoolean	UMETA(DisplayName = "Digital (bool)"),
+	MInteger	UMETA(DisplayName = "Digital (int32)"),
 	MAxis1D		UMETA(DisplayName = "Axis1D (float)"),
 	MAxis2D		UMETA(DisplayName = "Axis2D (Vector2D)"),
 	MAxis3D		UMETA(DisplayName = "Axis3D (Vector)"),
@@ -162,6 +163,7 @@ public:
 	using MAxis1D = double; // #TODO: Axis1D (float) double or float? FVector2D and FVector are double
 	using MAxis2D = FVector2D;
 	using MAxis3D = FVector;
+	using MInteger = int32;
 
 	// Support all relevant default constructors (FModuleInputValue isn't movable)
 	FModuleInputValue() = default;
@@ -170,10 +172,11 @@ public:
 
 	// Specialized constructors for supported types
 	// Converting a value to a different type (e.g. Val = FVector(1, 1, 1); Val = true;) zeroes out any unused components to ensure getters continue to function correctly.
-	explicit FModuleInputValue(bool bInValue) : Value(bInValue ? 1.f : 0.f, 0.f, 0.f), ValueType(EModuleInputValueType::MBoolean) {}
-	FModuleInputValue(MAxis1D InValue) : Value(InValue, 0.f, 0.f), ValueType(EModuleInputValueType::MAxis1D) {}
-	FModuleInputValue(MAxis2D InValue) : Value(InValue.X, InValue.Y, 0.f), ValueType(EModuleInputValueType::MAxis2D) {}
-	FModuleInputValue(MAxis3D InValue) : Value(InValue), ValueType(EModuleInputValueType::MAxis3D) {}
+	FModuleInputValue(bool bInValue) : Value(FVector::ZeroVector), ValueInt(bInValue), ValueType(EModuleInputValueType::MBoolean) {}
+	FModuleInputValue(MInteger InValue) : Value(FVector::ZeroVector), ValueInt(InValue), ValueType(EModuleInputValueType::MInteger) {}
+	FModuleInputValue(MAxis1D InValue) : Value(InValue, 0.f, 0.f), ValueInt(0), ValueType(EModuleInputValueType::MAxis1D) {}
+	FModuleInputValue(MAxis2D InValue) : Value(InValue.X, InValue.Y, 0.f), ValueInt(0), ValueType(EModuleInputValueType::MAxis2D) {}
+	FModuleInputValue(MAxis3D InValue) : Value(InValue), ValueInt(0), ValueType(EModuleInputValueType::MAxis3D) {}
 
 	FModuleInputValue ReturnQuantized() const
 	{
@@ -182,6 +185,10 @@ public:
 		switch (ValueType)
 		{
 			case EModuleInputValueType::MBoolean:
+			case EModuleInputValueType::MInteger:
+				OutValue = FModuleInputValue(ValueType, ValueInt);
+			break;
+
 			case EModuleInputValueType::MAxis1D:
 				ModularQuantize::QuantizeValue<1, 16>(OutValue.Value.X);
 			break;
@@ -211,6 +218,8 @@ public:
 	// Build a specific type with an arbitrary Axis3D value
 	FModuleInputValue(EModuleInputValueType InValueType, MAxis3D InValue) : Value(InValue), ValueType(InValueType)
 	{
+		ValueInt = 0.0f; // not used in this case
+
 		// Clear out value components to match type
 		switch (ValueType)
 		{
@@ -227,14 +236,26 @@ public:
 		}
 	}
 
+	// Build a specific type with an Integer value
+	FModuleInputValue(EModuleInputValueType InValueType, MInteger InValue) : ValueInt(InValue), ValueType(InValueType)
+	{
+		// not used in thos case clear for good measure
+		Value.X = 0.f;
+		Value.Y = 0.f;
+		Value.Z = 0.f;
+	}
+
 	// Resets Value without affecting ValueType
 	void Reset()
 	{
 		Value = FVector::ZeroVector;
+		ValueInt = 0;
 	}
 
 	FModuleInputValue& operator+=(const FModuleInputValue& Rhs)
 	{
+		ensure(ValueType == Rhs.ValueType);
+
 		Value += Rhs.Value;
 		// Promote value type to largest number of bits.
 		ValueType = FMath::Max(ValueType, Rhs.ValueType);
@@ -243,6 +264,8 @@ public:
 
 	friend FModuleInputValue operator+(const FModuleInputValue& Lhs, const FModuleInputValue& Rhs)
 	{
+		ensure(Lhs.ValueType == Rhs.ValueType);
+
 		FModuleInputValue Result(Lhs);
 		Result += Rhs;
 		return Result;
@@ -250,6 +273,8 @@ public:
 
 	FModuleInputValue& operator-=(const FModuleInputValue& Rhs)
 	{
+		ensure(ValueType == Rhs.ValueType);
+
 		Value -= Rhs.Value;
 		// Promote value type to largest number of bits.
 		ValueType = FMath::Max(ValueType, Rhs.ValueType);
@@ -258,6 +283,8 @@ public:
 
 	friend FModuleInputValue operator-(const FModuleInputValue& Lhs, const FModuleInputValue& Rhs)
 	{
+		ensure(Lhs.ValueType == Rhs.ValueType);
+
 		FModuleInputValue Result(Lhs);
 		Result -= Rhs;
 		return Result;
@@ -266,12 +293,18 @@ public:
 	// Scalar operators
 	FModuleInputValue& operator*=(float Scalar)
 	{
+		ensure(ValueType != EModuleInputValueType::MBoolean);
+		ensure(ValueType != EModuleInputValueType::MInteger);
+
 		Value *= Scalar;
 		return *this;
 	}
 
 	friend FModuleInputValue operator*(const FModuleInputValue& Lhs, const float Rhs)
 	{
+		ensure(Lhs.GetValueType() != EModuleInputValueType::MBoolean);
+		ensure(Lhs.GetValueType() != EModuleInputValueType::MInteger);
+
 		FModuleInputValue Result(Lhs);
 		Result *= Rhs;
 		return Result;
@@ -302,8 +335,8 @@ public:
 	// Read only index based value accessor, doesn't care about type. Expect 0 when accessing unused components.
 	float operator[](int32 Index) const { return Value[Index]; }
 
-	bool IsNonZero(float Tolerance = KINDA_SMALL_NUMBER) const { return Value.SizeSquared() >= Tolerance * Tolerance; }
 	bool IsQuantizedNonZero(float Tolerance = KINDA_SMALL_NUMBER) const { return ReturnQuantized().IsNonZero(Tolerance); }
+	bool IsNonZero(float Tolerance = KINDA_SMALL_NUMBER) const;
 
 	// In-place type conversion
 	FModuleInputValue& ConvertToType(EModuleInputValueType Type)
@@ -320,14 +353,12 @@ public:
 
 	float GetMagnitudeSq() const;
 	float GetMagnitude() const;
+	int32 GetMagnitudeInt() const;
 
 	// Serialize values
 	void Serialize(FArchive& Ar, UPackageMap* Map, bool& bOutSuccess);
 
-	void Lerp(const FModuleInputValue& Min, const FModuleInputValue& Max, float Alpha)
-	{
-		Value = FMath::Lerp(Min.Value, Max.Value, Alpha);
-	}
+	void Lerp(const FModuleInputValue& Min, const FModuleInputValue& Max, float Alpha);
 
 	void Merge(const FModuleInputValue& From);
 
@@ -346,6 +377,9 @@ protected:
 	UPROPERTY()
 	FVector Value = FVector::ZeroVector;
 
+	UPROPERTY()
+	int32 ValueInt = 0;
+
 	UPROPERTY()
 	EModuleInputValueType ValueType = EModuleInputValueType::MBoolean;
 
@@ -379,6 +413,12 @@ inline FModuleInputValue::MAxis3D FModuleInputValue::Get() const
 	return Value;
 }
 
+template<>
+inline FModuleInputValue::MInteger FModuleInputValue::Get() const
+{
+	return ValueInt;
+}
+
 
 class CHAOSVEHICLESCORE_API FModuleInputConversion
 {
@@ -592,6 +632,7 @@ public:
 	void MergeValue(const FName& InName, const FModuleInputValue& InValue);
 	FModuleInputValue GetValue(const FName& InName) const;
 	float GetMagnitude(const FName& InName) const;
+	int32 GetMagnitudeInt(const FName& InName) const;
 	bool InputsNonZero() const;
 
 	const FInputNameMap& NameMap;			// per vehicle
