Keep persisted bytes stat up to date when evicting items from IAS' cache.
#rnx

diff --git a/Engine/Source/Runtime/Experimental/IoStore/OnDemand/Private/JournaledCache.cpp b/Engine/Source/Runtime/Experimental/IoStore/OnDemand/Private/JournaledCache.cpp
--- a/Engine/Source/Runtime/Experimental/IoStore/OnDemand/Private/JournaledCache.cpp
+++ b/Engine/Source/Runtime/Experimental/IoStore/OnDemand/Private/JournaledCache.cpp
@@ -628,7 +628,7 @@ private:
 	void					OpenDataFile();
 	void					Wrap();
 	void					Spam();
-	uint64					Insert(uint64 DataBase, const FDataEntry& Entry);
+	int64					Insert(uint64 DataBase, const FDataEntry& Entry);
 	uint64					Insert(uint64 DataBase, const FDataEntry* Entries, uint32 EntryCount);
 	void					Prune(uint64 DataBase, uint32 Size);
 	mutable FRWLock			Lock;
@@ -752,7 +752,19 @@ bool FDiskCache::Has(uint64 Key) const
 bool FDiskCache::Evict(uint64 Key)
 {
 	FWriteScopeLock _(Lock);
-	return DataMap.Remove(Key) > 0;
+
+	FSetElementId Id = DataMap.FindId(Key);
+	if (!Id.IsValidId())
+	{
+		return false;
+	}
+
+	const TPair<uint64, FMapEntry>& Element = DataMap.Get(Id);
+	MappedBytes -= Element.Value.Size;
+	FOnDemandIoBackendStats::Get()->OnCachePersistedBytes(MappedBytes);
+
+	DataMap.Remove(Id);
+	return true;
 }
 
 ////////////////////////////////////////////////////////////////////////////////
@@ -800,7 +812,7 @@ EIoErrorCode FDiskCache::Materialize(uint64 Key, FIoBuffer& Out, uint32 Offset)
 }
 
 ////////////////////////////////////////////////////////////////////////////////
-uint64 FDiskCache::Insert(uint64 DataBase, const FDataEntry& Entry)
+int64 FDiskCache::Insert(uint64 DataBase, const FDataEntry& Entry)
 {
 	if (Entry.Key == 0)
 	{
@@ -809,8 +821,14 @@ uint64 FDiskCache::Insert(uint64 DataBase, const FDataEntry& Entry)
 
 	if (Entry.Size == 0)
 	{
+		int64 Ret = 0;
+		if (const FMapEntry* Existing = DataMap.Find(Entry.Key); Existing != nullptr)
+		{
+			Ret = -int64(Existing->Size);
+		}
+
 		DataMap.Remove(Entry.Key);
-		return 0;
+		return Ret;
 	}
 
 	FMapEntry Value;
@@ -830,7 +848,7 @@ uint64 FDiskCache::Insert(uint64 DataBase, const FDataEntry* Entries, uint32 Ent
 		PartialBias = Entries[0].Offset;
 	}
 
-	uint64 TotalSize = 0;
+	int64 TotalSize = 0;
 	for (uint32 i = 0; i < EntryCount; ++i)
 	{
 		check(int64(DataBase - PartialBias) >= 0);
