Evict IAS cache items that fail to decode
#rnx
#ushell-unshelve of 39147949

diff --git a/Engine/Source/Runtime/Experimental/IoStore/OnDemand/Private/OnDemandIoDispatcherBackend.cpp b/Engine/Source/Runtime/Experimental/IoStore/OnDemand/Private/OnDemandIoDispatcherBackend.cpp
--- a/Engine/Source/Runtime/Experimental/IoStore/OnDemand/Private/OnDemandIoDispatcherBackend.cpp
+++ b/Engine/Source/Runtime/Experimental/IoStore/OnDemand/Private/OnDemandIoDispatcherBackend.cpp
@@ -1111,6 +1111,9 @@ void FOnDemandIoBackend::CompleteRequest(FChunkRequest* ChunkRequest)
 				if (ChunkRequest->bCached)
 				{
 					Stats.OnCacheDecodeError();
+
+					check(Cache.IsValid());
+					Cache->Evict(ChunkRequest->Params.ChunkKey);
 				}
 				else
 				{
