Cache poisoning to stress test eviction path
#rnx
#ushell-unshelve of 39147949

diff --git a/Engine/Source/Runtime/Experimental/IoStore/OnDemand/Private/OnDemandIoDispatcherBackend.cpp b/Engine/Source/Runtime/Experimental/IoStore/OnDemand/Private/OnDemandIoDispatcherBackend.cpp
--- a/Engine/Source/Runtime/Experimental/IoStore/OnDemand/Private/OnDemandIoDispatcherBackend.cpp
+++ b/Engine/Source/Runtime/Experimental/IoStore/OnDemand/Private/OnDemandIoDispatcherBackend.cpp
@@ -205,6 +205,16 @@ static FAutoConsoleVariableRef CVar_WriteOnlyDecoding(
 	TEXT("Enables the use of 'WriteOnly' flag when decoding to buffers with the 'HardwareTargetBuffer' flag")
 );
 
+#if !UE_BUILD_SHIPPING
+static int32 GIasPoisonCache = false;
+static FAutoConsoleVariableRef CVar_IasPoisonCache(
+	TEXT("ias.PoisonCache"),
+	GIasPoisonCache,
+	TEXT("Fills all data materialized from the cache with 0x4d")
+);
+#endif
+
+
 
 // These priorities are indexed using the cvar below
 static UE::Tasks::ETaskPriority GCompleteMaterializeTaskPriorities[] =
@@ -1174,6 +1184,16 @@ void FOnDemandIoBackend::CompleteMaterialize(FChunkRequest* ChunkRequest)
 	{
 		case EIoErrorCode::Ok:
 		check(ChunkRequest->Chunk.GetData() != nullptr);
+#if !UE_BUILD_SHIPPING
+		if (GIasPoisonCache)
+		{
+			FIoBuffer& Data = ChunkRequest->Chunk;
+			for (uint64 i = Data.GetSize(); i--;)
+			{
+				Data.GetData()[i] = 0x4d;
+			}
+		}
+#endif
 		ChunkRequest->bCached = true;
 		CompleteRequest(ChunkRequest);
 		return;
