CompleteMaterialize() -> CompleteCacheRead()
#rnx
#ushell-unshelve of 39147949

diff --git a/Engine/Source/Runtime/Experimental/IoStore/OnDemand/Private/OnDemandIoDispatcherBackend.cpp b/Engine/Source/Runtime/Experimental/IoStore/OnDemand/Private/OnDemandIoDispatcherBackend.cpp
--- a/Engine/Source/Runtime/Experimental/IoStore/OnDemand/Private/OnDemandIoDispatcherBackend.cpp
+++ b/Engine/Source/Runtime/Experimental/IoStore/OnDemand/Private/OnDemandIoDispatcherBackend.cpp
@@ -230,7 +230,7 @@ static int32 GCompleteMaterializeTaskPriority = 3;
 FAutoConsoleVariableRef CVarCompleteMaterializeTaskPriority(
 	TEXT("ias.CompleteMaterializeTaskPriority"),
 	GCompleteMaterializeTaskPriority,
-	TEXT("Task priority for the CompleteMaterialize task (0 = foreground/high, 1 = foreground/normal, 2 = background/high, 3 = background/normal, 4 = background/low)."),
+	TEXT("Task priority for the CompleteCacheRead task (0 = foreground/high, 1 = foreground/normal, 2 = background/high, 3 = background/normal, 4 = background/low)."),
 	ECVF_Default
 );
 
@@ -918,7 +918,7 @@ private:
 	FString GetEndpointTestPath() const;
 	void ConditionallyStartBackendThread();
 	void CompleteRequest(FChunkRequest* ChunkRequest);
-	void CompleteMaterialize(FChunkRequest* ChunkRequest);
+	void CompleteCacheRead(FChunkRequest* ChunkRequest);
 	bool ResolveDistributedEndpoint(const FDistributedEndpointUrl& Url);
 	void ProcessHttpRequests(FHttpClient& HttpClient, FBitWindow& HttpErrors, int32 MaxConcurrentRequests);
 	int32 WaitForCompleteRequestTasks(float WaitTimeSeconds, float PollTimeSeconds);
@@ -1175,9 +1175,9 @@ void FOnDemandIoBackend::CompleteRequest(FChunkRequest* ChunkRequest)
 	ChunkRequests.Release(ChunkRequest);
 }
 
-void FOnDemandIoBackend::CompleteMaterialize(FChunkRequest* ChunkRequest)
+void FOnDemandIoBackend::CompleteCacheRead(FChunkRequest* ChunkRequest)
 {
-	TRACE_CPUPROFILER_EVENT_SCOPE(IasBackend::CompleteMaterialize);
+	TRACE_CPUPROFILER_EVENT_SCOPE(IasBackend::CompleteCacheRead);
 
 	bool bWasCancelled = false;
 	switch (ChunkRequest->CacheGetStatus)
@@ -1288,7 +1288,7 @@ bool FOnDemandIoBackend::Resolve(FIoRequestImpl* Request)
 
 			Launch(UE_SOURCE_LOCATION, [this, ChunkRequest] {
 				InflightCacheRequestCount.fetch_sub(1, std::memory_order_relaxed);
-				CompleteMaterialize(ChunkRequest);
+				CompleteCacheRead(ChunkRequest);
 			}, OnReadyEvent, GetRequestCompletionTaskPriority());
 
 			EIoErrorCode& OutStatus = ChunkRequest->CacheGetStatus;
