Fix FMultiEndpointHttpClient from attempting too many retries.

#rnx

- FMultiEndpointHttpClient::RetryRequest was incorrectly waiting for two retry attempts before cycling the host to try instead of cycling after the first retry.
- Fix off by one error when checking the number of retries to attempt.
- Added code documentation to the header detailing how retry attempts should work.

diff --git a/Engine/Source/Runtime/Experimental/IoStore/OnDemand/Private/OnDemandHttpClient.cpp b/Engine/Source/Runtime/Experimental/IoStore/OnDemand/Private/OnDemandHttpClient.cpp
--- a/Engine/Source/Runtime/Experimental/IoStore/OnDemand/Private/OnDemandHttpClient.cpp
+++ b/Engine/Source/Runtime/Experimental/IoStore/OnDemand/Private/OnDemandHttpClient.cpp
@@ -484,7 +484,7 @@ void FMultiEndpointHttpClient::IssueRequest(FRequest&& Request)
 				const bool bError		= TicketStatus.GetId() == FTicketStatus::EId::Error;
 				const bool bServerError = Request.StatusCode > 499 && Request.StatusCode < 600;
 
-				if ((bError || bServerError) && Request.RetryCount <= Config.MaxRetryCount)
+				if ((bError || bServerError) && Request.RetryCount < Config.MaxRetryCount)
 				{
 					RetryRequest(MoveTemp(Request));
 				}
@@ -532,7 +532,7 @@ void FMultiEndpointHttpClient::RetryRequest(FRequest&& Request)
 	FConnection& Connection = Request.Connection;
 
 	// Try a different host URL after the first retry
-	if (Request.RetryCount > 1 && Request.Host == Request.Connection.CurrentHost)
+	if (Request.RetryCount > 0 && Request.Host == Request.Connection.CurrentHost)
 	{
 		FAnsiStringView HostUrl = Request.Connection.HostGroup.CycleHost(Connection.CurrentHost);
 		if (Connection.Pools[Connection.CurrentHost].IsValid() == false)
diff --git a/Engine/Source/Runtime/Experimental/IoStore/OnDemand/Private/OnDemandHttpClient.h b/Engine/Source/Runtime/Experimental/IoStore/OnDemand/Private/OnDemandHttpClient.h
--- a/Engine/Source/Runtime/Experimental/IoStore/OnDemand/Private/OnDemandHttpClient.h
+++ b/Engine/Source/Runtime/Experimental/IoStore/OnDemand/Private/OnDemandHttpClient.h
@@ -102,6 +102,16 @@ struct FMultiEndpointHttpClientResponse
 	int32		HostIndex = INDEX_NONE;
 };
 
+/**
+ * Todo - More documentation on the client behavior.
+ * 
+ * Retry policy:
+ * If a request fails then the client will retry it up to FMultiEndpointHttpClientConfig::MaxRetryCount
+ * times. The first retry attempt will use the primary host with each subsequent attempt cycling to the
+ * next host in the FOnDemandHostGroup. If the end of the group is reached with retries remaining then
+ * the cycle will begin again at the start of the group.
+ * @see FOnDemandHostGroup for more info on how the host cycling works.
+ */
 class FMultiEndpointHttpClient
 {
 public:
