GCISMPool deprecation warning fix

diff --git a/Engine/Source/Runtime/Experimental/GeometryCollectionEngine/Public/GeometryCollection/GeometryCollectionISMPoolActor.h b/Engine/Source/Runtime/Experimental/GeometryCollectionEngine/Public/GeometryCollection/GeometryCollectionISMPoolActor.h
--- a/Engine/Source/Runtime/Experimental/GeometryCollectionEngine/Public/GeometryCollection/GeometryCollectionISMPoolActor.h
+++ b/Engine/Source/Runtime/Experimental/GeometryCollectionEngine/Public/GeometryCollection/GeometryCollectionISMPoolActor.h
@@ -4,8 +4,8 @@
 #include "GameFramework/Actor.h"
 #include "GeometryCollectionISMPoolActor.generated.h"
 
-class UGeometryCollectionISMPoolComponent;
-class UGeometryCollectionISMPoolDebugDrawComponent;
+class UE_DEPRECATED(5.6, "UGeometryCollectionISMPoolComponent is deprecated, please use UISMPoolComponent instead.") UGeometryCollectionISMPoolComponent;
+class UE_DEPRECATED(5.6, "UGeometryCollectionISMPoolDebugDrawComponent is deprecated, please use UISMPoolDebugDrawComponent instead.") UGeometryCollectionISMPoolDebugDrawComponent;
 
 UCLASS(ConversionRoot, ComponentWrapperClass, MinimalAPI)
 class UE_DEPRECATED(5.6, "AGeometryCollectionISMPoolActor is deprecated, please use AISMPoolActor instead.") AGeometryCollectionISMPoolActor : public AActor
@@ -13,11 +13,15 @@ class UE_DEPRECATED(5.6, "AGeometryCollectionISMPoolActor is deprecated, please
 	GENERATED_UCLASS_BODY()
 
 private:
+PRAGMA_DISABLE_DEPRECATION_WARNINGS
 	UPROPERTY(Category = ISMPoolActor, VisibleAnywhere, BlueprintReadOnly, meta = (AllowPrivateAccess = "true"))
 	TObjectPtr<class UGeometryCollectionISMPoolComponent> ISMPoolComp;
+PRAGMA_ENABLE_DEPRECATION_WARNINGS
 
+PRAGMA_DISABLE_DEPRECATION_WARNINGS
 	UPROPERTY(Category = ISMPoolActor, VisibleAnywhere, BlueprintReadOnly, meta = (AllowPrivateAccess = "true"))
 	TObjectPtr<class UGeometryCollectionISMPoolDebugDrawComponent> ISMPoolDebugDrawComp;
+PRAGMA_ENABLE_DEPRECATION_WARNINGS
 
 public:
 	/** Returns ISMPoolComp subobject **/
@@ -26,5 +30,3 @@ PRAGMA_DISABLE_DEPRECATION_WARNINGS
 PRAGMA_ENABLE_DEPRECATION_WARNINGS
 };
 
-
-
diff --git a/Engine/Source/Runtime/Experimental/GeometryCollectionEngine/Public/GeometryCollection/GeometryCollectionISMPoolComponent.h b/Engine/Source/Runtime/Experimental/GeometryCollectionEngine/Public/GeometryCollection/GeometryCollectionISMPoolComponent.h
--- a/Engine/Source/Runtime/Experimental/GeometryCollectionEngine/Public/GeometryCollection/GeometryCollectionISMPoolComponent.h
+++ b/Engine/Source/Runtime/Experimental/GeometryCollectionEngine/Public/GeometryCollection/GeometryCollectionISMPoolComponent.h
@@ -11,7 +11,7 @@
 #include "GeometryCollectionISMPoolComponent.generated.h"
 
 class AActor;
-class UGeometryCollectionISMPoolComponent;
+class UE_DEPRECATED(5.6, "UGeometryCollectionISMPoolComponent is deprecated, please use UISMPoolComponent instead.") UGeometryCollectionISMPoolComponent;
 class UInstancedStaticMeshComponent;
 
 /** 
@@ -283,11 +283,17 @@ struct FGeometryCollectionISMPool
 	FGeometryCollectionISMPool();
 
 	/** Find or add an ISM and return an ISM index handle. */
+	PRAGMA_DISABLE_DEPRECATION_WARNINGS
 	FISMIndex GetOrAddISM(UGeometryCollectionISMPoolComponent* OwningComponent, const FGeometryCollectionStaticMeshInstance& MeshInstance, bool& bOutISMCreated);
+	PRAGMA_ENABLE_DEPRECATION_WARNINGS
+
 	/** Remove an ISM. */
 	void RemoveISM(FISMIndex ISMIndex, bool bKeepAlive, bool bRecycle);
 	/** Add instances to ISM and return a mesh info handle. */
+	PRAGMA_DISABLE_DEPRECATION_WARNINGS
 	FGeometryCollectionMeshInfo AddInstancesToISM(UGeometryCollectionISMPoolComponent* OwningComponent, const FGeometryCollectionStaticMeshInstance& MeshInstance, int32 InstanceCount, TArrayView<const float> CustomDataFloats);
+	PRAGMA_ENABLE_DEPRECATION_WARNINGS
+
 	/** Remove instances from an ISM. */
 	void RemoveInstancesFromISM(const FGeometryCollectionMeshInfo& MeshInfo);
 	/** Update ISM contents. */
@@ -298,12 +304,16 @@ struct FGeometryCollectionISMPool
 	void Clear();
 
 	/** Tick maintenance of free list and preallocation. */
+	PRAGMA_DISABLE_DEPRECATION_WARNINGS
 	void Tick(UGeometryCollectionISMPoolComponent* OwningComponent);
+	PRAGMA_ENABLE_DEPRECATION_WARNINGS
 
 	/** Add an ISM description to the preallocation queue. */
 	void RequestPreallocateMeshInstance(const FGeometryCollectionStaticMeshInstance& MeshInstances);
 	/** Process the preallocation queue. Processing is timesliced so that only some of the queue will be processed in every call. */
+	PRAGMA_DISABLE_DEPRECATION_WARNINGS
 	void ProcessPreallocationRequests(UGeometryCollectionISMPoolComponent* OwningComponent, int32 MaxPreallocations);
+	PRAGMA_ENABLE_DEPRECATION_WARNINGS
 
 	void UpdateAbsoluteTransforms(const FTransform& BaseTransform, EUpdateTransformFlags UpdateTransformFlags, ETeleportType Teleport);
 
@@ -333,6 +343,8 @@ struct FGeometryCollectionISMPool
 * UGeometryCollectionISMPoolComponent.
 * Component that manages a pool of ISM components in order to allow multiple client components that use the same meshes to the share ISMs.
 */
+class UE_DEPRECATED(5.6, "UGeometryCollectionISMPoolDebugDrawComponent is deprecated, please use UISMPoolDebugDrawComponent instead.") UGeometryCollectionISMPoolDebugDrawComponent;
+
 UCLASS(meta = (BlueprintSpawnableComponent), MinimalAPI)
 class UE_DEPRECATED(5.6, "UGeometryCollectionISMPoolComponent is deprecated, please use UISMPoolComponent instead.") UGeometryCollectionISMPoolComponent: public USceneComponent
 {
@@ -387,5 +399,7 @@ private:
 	FGeometryCollectionISMPool Pool;
 
 	// Expose internals for debug draw support.
-	friend class UGeometryCollectionISMPoolDebugDrawComponent;
+PRAGMA_DISABLE_DEPRECATION_WARNINGS
+	friend UGeometryCollectionISMPoolDebugDrawComponent;
+PRAGMA_ENABLE_DEPRECATION_WARNINGS
 };
