Make TPBDActiveView threadsafe when activate/deactivating ranges in parallel (how they are used by Cloth).

It was discovered by a licensee that FMath::Abs is not an atomic operation, so it was possible for FMath::Abs to return a negative value if it was called at the same time another thread was flipping its active sign, which would in turn produce incorrect results when using UpperBound to find the correct range index.

diff --git a/Engine/Source/Runtime/Experimental/Chaos/Private/Chaos/Deformable/ChaosDeformableSolver.cpp b/Engine/Source/Runtime/Experimental/Chaos/Private/Chaos/Deformable/ChaosDeformableSolver.cpp
--- a/Engine/Source/Runtime/Experimental/Chaos/Private/Chaos/Deformable/ChaosDeformableSolver.cpp
+++ b/Engine/Source/Runtime/Experimental/Chaos/Private/Chaos/Deformable/ChaosDeformableSolver.cpp
@@ -710,7 +710,7 @@ namespace Chaos::Softs
 							if (!Proxy.CollisionBodies.Contains(AddBody.Key))
 							{
 								int32 Index = Evolution->AddCollisionParticle(INDEX_NONE, true);
-								int32 ViewIndex = Evolution->CollisionParticlesActiveView().GetRanges().Num() - 1;
+								int32 ViewIndex = Evolution->CollisionParticlesActiveView().GetNumRanges() - 1;
 								Evolution->CollisionParticles().SetX(Index, AddBody.Transform.GetTranslation());
 								Evolution->CollisionParticles().SetR(Index, AddBody.Transform.GetRotation());
 								Chaos::FImplicitObjectPtr UniquePtr(AddBody.Shapes); AddBody.Shapes = nullptr;
diff --git a/Engine/Source/Runtime/Experimental/Chaos/Public/Chaos/PBDActiveView.h b/Engine/Source/Runtime/Experimental/Chaos/Public/Chaos/PBDActiveView.h
--- a/Engine/Source/Runtime/Experimental/Chaos/Public/Chaos/PBDActiveView.h
+++ b/Engine/Source/Runtime/Experimental/Chaos/Public/Chaos/PBDActiveView.h
@@ -8,6 +8,39 @@
 
 namespace Chaos
 {
+	class FActiveViewRange final
+	{
+	public:
+		FActiveViewRange()
+			: bIsActive(0)
+			, Offset(0)
+		{}
+
+		explicit FActiveViewRange(int32 InOffset, bool bActive = false)
+			: bIsActive(bActive ? 1 : 0)
+			, Offset((uint32)InOffset)
+		{}
+
+		bool IsActive() const
+		{
+			return !!bIsActive && Offset > 0;
+		}
+
+		int32 GetOffset() const
+		{
+			return (int32)Offset;
+		}
+
+		void SetActive(bool bActivate)
+		{
+			bIsActive = bActivate;
+		}
+
+	private:
+		uint32 bIsActive : 1;
+		uint32 Offset : 31;
+	};
+
 	// Index based view, specialized for working with several ranges within a same array such as particles.
 	template<typename TItemsType>
 	class TPBDActiveView
@@ -57,21 +90,34 @@ namespace Chaos
 		TArray<TVector<int32, 2>, TInlineAllocator<8>> GetActiveRanges() const;
 
 		// Return internal ranges.
-		TConstArrayView<int32> GetRanges() const { return Ranges; }
+		TConstArrayView<FActiveViewRange> GetAllRanges() const { return Ranges; }
+		UE_DEPRECATED(5.6, "This method has been deprecated as the underlying type of Ranges has changed. Use GetAllRanges instead.")
+		TConstArrayView<int32> GetRanges() const { return TConstArrayView<int32>(); }
+
+		int32 GetNumRanges() const { return Ranges.Num(); }
 
 	private:
+
+		int32 GetRangeIndex(int32 Offset) const
+		{
+			// Binary search upper bound range of this offset
+			const int32 Index = Algo::UpperBound(Ranges, FActiveViewRange(Offset), [](const FActiveViewRange A, const FActiveViewRange B) { return A.GetOffset() < B.GetOffset(); });
+			check(Ranges.IsValidIndex(Index));
+			return Index;
+		}
+
 		TItemsType& Items;
-		TArray<int32> Ranges;
+		TArray<FActiveViewRange> Ranges;
 	};
 
 	template <class TItemsType>
 	int32 TPBDActiveView<TItemsType>::AddRange(int32 NumItems, bool bActivate)
 	{
-		const int32 Offset = Ranges.Num() ? FMath::Abs(Ranges.Last()) : 0;
+		const int32 Offset = Ranges.Num() ? Ranges.Last().GetOffset() : 0;
 		if (NumItems)
 		{
 			const int32 Size = Offset + NumItems;
-			Ranges.Emplace(bActivate ? Size : -Size);
+			Ranges.Emplace(Size, bActivate);
 		}
 		return Offset;
 	}
@@ -79,48 +125,33 @@ namespace Chaos
 	template <class TItemsType>
 	int32 TPBDActiveView<TItemsType>::GetRangeSize(int32 Offset) const
 	{
-		// Binary search upper bound range of this offset
-		const int32 Index = Algo::UpperBound(Ranges, Offset, [](const int32 A, const int32 B) { return FMath::Abs(A) < FMath::Abs(B); });
-		check(Ranges.IsValidIndex(Index));
-
+		const int32 Index = GetRangeIndex(Offset);
 		// Return size regardless or activation state
-		return FMath::Abs(Ranges[Index]) - Offset;
+		return Ranges[Index].GetOffset() - Offset;
 	}
 
 	template <class TItemsType>
 	void TPBDActiveView<TItemsType>::ActivateRange(int32 Offset, bool bActivate)
 	{
-		// Binary search upper bound range of this offset
-		const int32 Index = Algo::UpperBound(Ranges, Offset, [](const int32 A, const int32 B) { return FMath::Abs(A) < FMath::Abs(B); });
-		check(Ranges.IsValidIndex(Index));
-
-		// Change activation state (sign) if needed
-		if ((Ranges[Index] > 0) != bActivate)
-		{
-			Ranges[Index] = -Ranges[Index];
-		}
+		const int32 Index = GetRangeIndex(Offset);
+		Ranges[Index].SetActive(bActivate);
 	}
 
 	template <class TItemsType>
 	void TPBDActiveView<TItemsType>::SequentialFor(TFunctionRef<void(TItemsType&, int32)> Function) const
 	{
 		int32 Offset = 0;
-		for (int32 Range : Ranges)
+		for (const FActiveViewRange Range : Ranges)
 		{
-			if (Range > 0)
+			if (Range.IsActive())
 			{
 				// Active range
-				for (int32 Index = Offset; Index < Range; ++Index)
+				for (int32 Index = Offset; Index < Range.GetOffset(); ++Index)
 				{
 					Function(Items, Index);
 				}
-				Offset = Range;
-			}
-			else
-			{
-				// Inactive range
-				Offset = -Range;
 			}
+			Offset = Range.GetOffset();
 		}
 	}
 
@@ -128,23 +159,18 @@ namespace Chaos
 	void TPBDActiveView<TItemsType>::ParallelFor(TFunctionRef<void(TItemsType&, int32)> Function, int32 MinParallelBatchSize) const
 	{
 		int32 Offset = 0;
-		for (int32 Range : Ranges)
+		for (const FActiveViewRange Range : Ranges)
 		{
-			if (Range > 0)
+			if (Range.IsActive())
 			{
 				// Active range
-				const int32 RangeSize = Range - Offset;
+				const int32 RangeSize = Range.GetOffset() - Offset;
 				PhysicsParallelFor(RangeSize, [this, Offset, &Function](int32 Index)
 				{
 					Function(Items, Offset + Index);
 				}, /*bForceSingleThreaded =*/ RangeSize < MinParallelBatchSize);
-				Offset = Range;
-			}
-			else
-			{
-				// Inactive range
-				Offset = -Range;
 			}
+			Offset = Range.GetOffset();
 		}
 	}
 
@@ -168,19 +194,14 @@ namespace Chaos
 	void TPBDActiveView<TItemsType>::RangeFor(TFunctionRef<void(TItemsType&, int32, int32)> Function) const
 	{
 		int32 Offset = 0;
-		for (int32 Range : Ranges)
+		for (const FActiveViewRange Range : Ranges)
 		{
-			if (Range > 0)
+			if (Range.IsActive())
 			{
 				// Active Range
-				Function(Items, Offset, Range);
-				Offset = Range;
-			}
-			else
-			{
-				// Inactive range
-				Offset = -Range;
+				Function(Items, Offset, Range.GetOffset());
 			}
+			Offset = Range.GetOffset();
 		}
 	}
 
@@ -203,7 +224,7 @@ namespace Chaos
 	{
 		for (int32 Index = 0; Index < Ranges.Num(); ++Index)
 		{
-			if (FMath::Abs(Ranges[Index]) > Offset)
+			if (Ranges[Index].GetOffset() > Offset)
 			{
 				Ranges.SetNum(Index);
 				break;
@@ -214,9 +235,9 @@ namespace Chaos
 	template <class TItemsType>
 	bool TPBDActiveView<TItemsType>::HasActiveRange() const
 	{
-		for (int32 Range : Ranges)
+		for (const FActiveViewRange Range : Ranges)
 		{
-			if (Range > 0)
+			if (Range.IsActive())
 			{
 				return true;
 			}
@@ -229,17 +250,13 @@ namespace Chaos
 	{
 		int32 ActiveSize = 0;
 		int32 Offset = 0;
-		for (int32 Range : Ranges)
+		for (const FActiveViewRange Range : Ranges)
 		{
-			if (Range > 0)
+			if (Range.IsActive())
 			{
-				ActiveSize += Range - Offset;
-				Offset = Range;
-			}
-			else
-			{
-				Offset = -Range;
+				ActiveSize += Range.GetOffset() - Offset;
 			}
+			Offset = Range.GetOffset();
 		}
 		return ActiveSize;
 	}
@@ -250,17 +267,13 @@ namespace Chaos
 		TArray<TVector<int32, 2>, TInlineAllocator<8>> ActiveRanges;
 
 		int32 Offset = 0;
-		for (int32 Range : Ranges)
+		for (const FActiveViewRange Range : Ranges)
 		{
-			if (Range > 0)
-			{
-				ActiveRanges.Add(TVector<int32, 2>(Offset, Range));
-				Offset = Range;
-			}
-			else
+			if (Range.IsActive())
 			{
-				Offset = -Range;
+				ActiveRanges.Add(TVector<int32, 2>(Offset, Range.GetOffset()));
 			}
+			Offset = Range.GetOffset();
 		}
 		return ActiveRanges;
 	}
