Support MACD between Spheres and Triangle Meshes

diff --git a/Engine/Source/Runtime/Experimental/Chaos/Private/Chaos/Collision/ConvexTriangleContactPoint.cpp b/Engine/Source/Runtime/Experimental/Chaos/Private/Chaos/Collision/ConvexTriangleContactPoint.cpp
--- a/Engine/Source/Runtime/Experimental/Chaos/Private/Chaos/Collision/ConvexTriangleContactPoint.cpp
+++ b/Engine/Source/Runtime/Experimental/Chaos/Private/Chaos/Collision/ConvexTriangleContactPoint.cpp
@@ -884,6 +884,42 @@ namespace Chaos::Private
 		return true;
 	}
 
+	template<>
+	bool FindClosestFeatures(const Chaos::TSphere<Chaos::FReal, 3>& Convex, const FTriangle& Triangle, const FVec3& TriangleNormal, const FVec3& ConvexRelativeMovement, const FReal CullDistance, FConvexContactPoint& OutContactPoint)
+	{
+		const FVec3 Center = Convex.GetCenterf();
+		FVec3 Bary = FVec3(0);
+
+		const FVec3 NearestPointTriangle = FindClosestPointAndBaryOnTriangle(Triangle.GetVertex(0), Triangle.GetVertex(1), Triangle.GetVertex(2), Center, Bary);
+		const FReal Distance = FVec3::Dist(NearestPointTriangle, Center) - Convex.GetRadiusf();
+
+		OutContactPoint.ShapeContactPoints[0] = Convex.FindClosestPoint(NearestPointTriangle);
+		OutContactPoint.ShapeContactPoints[1] = NearestPointTriangle;
+		OutContactPoint.ShapeContactNormal = TriangleNormal;
+		OutContactPoint.ShapeContactNormal.Normalize();
+		OutContactPoint.Phi = Distance;
+
+		// Now check the cull distance, taking movement into account
+		const FReal SeparationWithMotion = OutContactPoint.Phi + FVec3::DotProduct(ConvexRelativeMovement, OutContactPoint.ShapeContactNormal);
+		if ((OutContactPoint.Phi > CullDistance) && (SeparationWithMotion > CullDistance))
+		{
+			return false;
+		}
+
+		// Find the triangle feature at the contact point
+		if (!Private::GetTriangleFeature(Triangle, TriangleNormal, OutContactPoint.ShapeContactPoints[1], OutContactPoint.ShapeContactNormal, OutContactPoint.Features[1]))
+		{
+			return false;
+		}
+
+		// Find the convex feature at the contact point
+		OutContactPoint.Features[0].FeatureType = EConvexFeatureType::Vertex;
+		OutContactPoint.Features[0].PlaneIndex = 0;
+		OutContactPoint.Features[0].PlaneFeatureIndex = 0;
+
+		return true;
+	}
+
 	template bool FindClosestFeatures(
 		const FImplicitConvex3& Convex,
 		const FTriangle& Triangle, 
diff --git a/Engine/Source/Runtime/Experimental/Chaos/Private/Chaos/CollisionOneShotManifoldsMeshShapes.cpp b/Engine/Source/Runtime/Experimental/Chaos/Private/Chaos/CollisionOneShotManifoldsMeshShapes.cpp
--- a/Engine/Source/Runtime/Experimental/Chaos/Private/Chaos/CollisionOneShotManifoldsMeshShapes.cpp
+++ b/Engine/Source/Runtime/Experimental/Chaos/Private/Chaos/CollisionOneShotManifoldsMeshShapes.cpp
@@ -661,6 +661,134 @@ namespace Chaos
 			}
 		}
 
+		template<>
+		void GenerateConvexTriangleOneShotManifoldMACD(
+			const Chaos::TSphere<Chaos::FReal, 3>& Convex,
+			const FRigidTransform3& ConvexTransform,
+			const FVec3& InConvexRelativeMovement,
+			Private::FMeshContactGenerator& ContactGenerator,
+			const int32 TriangleIndex,
+			const FReal CullDistance,
+			FContactPointManifold& OutContactPoints)
+		{
+			FTriangle Triangle = ContactGenerator.GetTriangle(TriangleIndex);
+			const FVec3 TriangleNormal = ContactGenerator.GetTriangleNormal(TriangleIndex);
+			const FReal ConvexTriangleDistanceAtP = FVec3::DotProduct(Convex.GetCenterOfMass() - Triangle.GetVertex(0), TriangleNormal);
+
+			// If we are outside the plane of the triangle at P, collide at P
+			bool bUseMACD = !InConvexRelativeMovement.IsZero() && (ConvexTriangleDistanceAtP < 0);
+
+			// If desired, shift the triangle so it is relative to the convex when it was at X
+			FVec3 ConvexRelativeMovement = FVec3(0);
+			FReal ConvexTriangleDistance = ConvexTriangleDistanceAtP;
+			if (bUseMACD)
+			{
+				ConvexRelativeMovement = InConvexRelativeMovement;
+				Triangle[0] += ConvexRelativeMovement;
+				Triangle[1] += ConvexRelativeMovement;
+				Triangle[2] += ConvexRelativeMovement;
+				ConvexTriangleDistance -= FVec3::DotProduct(ConvexRelativeMovement, TriangleNormal);
+			}
+
+			// If we were inside the triangle at X and P we ignore this triangle
+			if ((ConvexTriangleDistance < 0) && (ConvexTriangleDistanceAtP < 0))
+			{
+				return;
+			}
+
+			// Find the closest feature pair on the triangle and convex
+			Private::FConvexContactPoint ClosestContact;
+			if (Private::FindClosestFeatures(Convex, Triangle, TriangleNormal, ConvexRelativeMovement, CullDistance, ClosestContact))
+			{
+				ClosestContact.Features[0].ObjectIndex = 0;
+				ClosestContact.Features[1].ObjectIndex = TriangleIndex;
+
+#if CHAOS_DEBUG_DRAW
+				if (CVars::ChaosSolverDebugDrawMeshContacts && FDebugDrawQueue::GetInstance().IsDebugDrawingEnabled())
+				{
+					const FVec3 P = ConvexTransform.TransformPositionNoScale(ClosestContact.ShapeContactPoints[1] - ConvexRelativeMovement);
+					const FVec3 N = ConvexTransform.TransformVectorNoScale(ClosestContact.ShapeContactNormal);
+					FDebugDrawQueue::GetInstance().DrawDebugLine(P, P + 10.0f * N, FColor::Black, false, CVars::ChaosSolverDebugDebugDrawSettings.DrawDuration, (uint8)CVars::ChaosSolverDebugDebugDrawSettings.DrawPriority, 1.25f * CVars::ChaosSolverDebugDebugDrawSettings.LineThickness);
+				}
+				if (CVars::ChaosSolverDebugDrawMeshContactDetails && FDebugDrawQueue::GetInstance().IsDebugDrawingEnabled())
+				{
+					const FVec3 XConvex = ConvexTransform.TransformPositionNoScale(-ConvexRelativeMovement);
+					const FColor Color = (bUseMACD) ? FColor::Red : FColor::Green;
+					FDebugDrawQueue::GetInstance().DrawDebugPoint(XConvex, Color, false, CVars::ChaosSolverDebugDebugDrawSettings.DrawDuration, (uint8)CVars::ChaosSolverDebugDebugDrawSettings.DrawPriority, 20.0f * CVars::ChaosSolverDebugDebugDrawSettings.LineThickness);
+				}
+#endif
+
+				// Back face culling
+				const FReal TriangleDotNormal = FVec3::DotProduct(TriangleNormal, ClosestContact.ShapeContactNormal);
+				if (TriangleDotNormal < 0)
+				{
+					return;
+				}
+
+				FContactPoint& Contact = OutContactPoints[OutContactPoints.AddUninitialized()];
+				Contact.ShapeContactPoints[0] = ClosestContact.ShapeContactPoints[0];
+				Contact.ShapeContactPoints[1] = ClosestContact.ShapeContactPoints[1];
+				Contact.ShapeContactNormal = ClosestContact.ShapeContactNormal;
+				Contact.Phi = ClosestContact.Phi;
+				Contact.ContactType = ClosestContact.GetContactPointType();
+				Contact.FaceIndex = INDEX_NONE;
+
+#if CHAOS_DEBUG_DRAW
+				if (CVars::ChaosSolverDebugDrawMeshContacts && FDebugDrawQueue::GetInstance().IsDebugDrawingEnabled())
+				{
+					for (int32 ContactIndex = 0; ContactIndex < OutContactPoints.Num(); ++ContactIndex)
+					{
+						FContactPoint& ContactPoint = OutContactPoints[ContactIndex];
+						const FVec3 P = ConvexTransform.TransformPositionNoScale(ContactPoint.ShapeContactPoints[1] - ConvexRelativeMovement);
+						const FVec3 N = ConvexTransform.TransformVectorNoScale(ContactPoint.ShapeContactNormal);
+						FColor Color = FColor::Black;
+						switch (ClosestContact.Features[1].FeatureType)
+						{
+						case Private::EConvexFeatureType::Plane:
+							Color = FColor::White;
+							break;
+						case Private::EConvexFeatureType::Edge:
+							Color = FColor::Cyan;
+							break;
+						case Private::EConvexFeatureType::Vertex:
+							Color = FColor::Magenta;
+							break;
+						}
+						FDebugDrawQueue::GetInstance().DrawDebugLine(P, P + 10.0f * N, Color, false, CVars::ChaosSolverDebugDebugDrawSettings.DrawDuration, (uint8)CVars::ChaosSolverDebugDebugDrawSettings.DrawPriority, 1.25f * CVars::ChaosSolverDebugDebugDrawSettings.LineThickness);
+					}
+				}
+#endif
+
+				// Adjust the closest feature if it is invalid. E.g., if we collide with a triangle edge and the normal is outside the range allowed by
+				// the triangles sharing the edge we will project the normal into the valid range. If the feature gets chnaged, we will alter all
+				// of the manifold points to use the new normal.
+				const bool bFeatureChanged = ContactGenerator.FixFeature(TriangleIndex, ClosestContact.Features[1].FeatureType, ClosestContact.Features[1].PlaneFeatureIndex, ClosestContact.ShapeContactNormal);
+				if (bFeatureChanged)
+				{
+					FContactPoint& ContactPoint = OutContactPoints[0];
+
+					// Update the normal and recalculate the separation
+					ContactPoint.ShapeContactNormal = ClosestContact.ShapeContactNormal;
+					ContactPoint.Phi = FVec3::DotProduct(ContactPoint.ShapeContactPoints[0] - ContactPoint.ShapeContactPoints[1], ContactPoint.ShapeContactNormal);
+
+					// Remap the triangle contact onto the new plane, keeping the contact point on the convex shape where it is.
+					// NOTE: This means that the triangle contact point may be outside the triangle, but for contact separation
+					// we really only care about the contact plane. This is required for static friction, which assumes the contacts 
+					// have zero tangential separation on the frame they are generated.
+					ContactPoint.ShapeContactPoints[1] = ContactPoint.ShapeContactPoints[0] - (ContactPoint.Phi * ContactPoint.ShapeContactNormal);
+				}
+
+				// Correct the contact points if we ran collision detection at X rather than P
+				if (bUseMACD)
+				{
+					FContactPoint& ContactPoint = OutContactPoints[0];
+
+					const FReal ShiftDotNormal = FVec3::DotProduct(ConvexRelativeMovement, ContactPoint.ShapeContactNormal);
+					ContactPoint.ShapeContactPoints[1] += -ShiftDotNormal * ContactPoint.ShapeContactNormal;
+					ContactPoint.Phi += ShiftDotNormal;
+				}
+			}
+		}
 		// MACD: Motion-Aware Collision Detection
 		template<typename ConvexType, typename MeshType>
 		void ConstructConvexMeshOneShotManifoldMACD(
@@ -786,9 +914,11 @@ namespace Chaos
 				ContactGeneratorSettings.bSortByPhi = bChaos_Collision_ConvexTriMeshSortByPhi;
 				Private::FMeshContactGenerator ContactGenerator(ContactGeneratorSettings);
 
+				const FVec3 RelativeMovement = FVec3(Constraint.GetRelativeMovement());
+
 				if (const FImplicitSphere3* Sphere = Quadratic.template GetObject<FImplicitSphere3>())
 				{
-					ConstructConvexMeshOneShotManifold2(*Sphere, QuadraticTransform, *Mesh, MeshTransform, MeshScale, CullDistance, ContactGenerator);
+					ConstructConvexMeshOneShotManifoldMACD(*Sphere, QuadraticTransform, *Mesh, MeshTransform, MeshScale, RelativeMovement, CullDistance, ContactGenerator);
 				}
 				else if (const FImplicitCapsule3* Capsule = Quadratic.template GetObject<FImplicitCapsule3>())
 				{
