add options for geometry collection convex merging to also merge convexes on leaf nodes

diff --git a/Engine/Plugins/Experimental/GeometryCollectionPlugin/Source/GeometryCollectionNodes/Private/Dataflow/GeometryCollectionUtilityNodes.cpp b/Engine/Plugins/Experimental/GeometryCollectionPlugin/Source/GeometryCollectionNodes/Private/Dataflow/GeometryCollectionUtilityNodes.cpp
--- a/Engine/Plugins/Experimental/GeometryCollectionPlugin/Source/GeometryCollectionNodes/Private/Dataflow/GeometryCollectionUtilityNodes.cpp
+++ b/Engine/Plugins/Experimental/GeometryCollectionPlugin/Source/GeometryCollectionNodes/Private/Dataflow/GeometryCollectionUtilityNodes.cpp
@@ -219,6 +219,7 @@ FGenerateClusterConvexHullsFromLeafHullsDataflowNode::FGenerateClusterConvexHull
 	RegisterInputConnection(&ConvexCount);
 	RegisterInputConnection(&ErrorTolerance);
 	RegisterInputConnection(&OptionalSelectionFilter);
+	RegisterInputConnection(&bAllowMergingLeafHulls);
 	RegisterInputConnection(&bProtectNegativeSpace);
 	RegisterInputConnection(&TargetNumSamples);
 	RegisterInputConnection(&MinSampleSpacing);
@@ -268,6 +269,7 @@ void FGenerateClusterConvexHullsFromLeafHullsDataflowNode::Evaluate(UE::Dataflow
 			const double InErrorToleranceInCm = GetValue(Context, &ErrorTolerance);
 			FGeometryCollectionConvexUtility::FClusterConvexHullSettings HullMergeSettings(InConvexCount, InErrorToleranceInCm, bPreferExternalCollisionShapes);
 			HullMergeSettings.AllowMergesMethod = AllowMerges;
+			HullMergeSettings.bAllowMergingLeafHulls = GetValue(Context, &bAllowMergingLeafHulls);
 			HullMergeSettings.EmptySpace = bHasNegativeSpace ? &NegativeSpace : nullptr;
 
 			if (bHasSelectionFilter)
diff --git a/Engine/Plugins/Experimental/GeometryCollectionPlugin/Source/GeometryCollectionNodes/Public/Dataflow/GeometryCollectionUtilityNodes.h b/Engine/Plugins/Experimental/GeometryCollectionPlugin/Source/GeometryCollectionNodes/Public/Dataflow/GeometryCollectionUtilityNodes.h
--- a/Engine/Plugins/Experimental/GeometryCollectionPlugin/Source/GeometryCollectionNodes/Public/Dataflow/GeometryCollectionUtilityNodes.h
+++ b/Engine/Plugins/Experimental/GeometryCollectionPlugin/Source/GeometryCollectionNodes/Public/Dataflow/GeometryCollectionUtilityNodes.h
@@ -296,6 +296,10 @@ public:
 	UPROPERTY(meta = (DataflowInput))
 	FDataflowTransformSelection OptionalSelectionFilter;
 
+	/** Also allow the same hull merging process to run on leaf hulls (merging hulls on leaves in the selection) */
+	UPROPERTY(EditAnywhere, Category = "Convex", meta = (DataflowInput))
+	bool bAllowMergingLeafHulls = false;
+
 	/** Whether to use a sphere cover to define negative space that should not be covered by convex hulls */
 	UPROPERTY(EditAnywhere, Category = NegativeSpace, meta = (DataflowInput))
 	bool bProtectNegativeSpace = false;
diff --git a/Engine/Source/Runtime/Experimental/Chaos/Private/GeometryCollection/GeometryCollectionConvexUtility.cpp b/Engine/Source/Runtime/Experimental/Chaos/Private/GeometryCollection/GeometryCollectionConvexUtility.cpp
--- a/Engine/Source/Runtime/Experimental/Chaos/Private/GeometryCollection/GeometryCollectionConvexUtility.cpp
+++ b/Engine/Source/Runtime/Experimental/Chaos/Private/GeometryCollection/GeometryCollectionConvexUtility.cpp
@@ -1994,9 +1994,9 @@ void FGeometryCollectionConvexUtility::GenerateClusterConvexHullsFromLeafOrChild
 			return;
 		}
 
-		if (bUseDirectChildren)
+		if (bUseDirectChildren || Settings.bAllowMergingLeafHulls)
 		{
-			// If using direct children, order matters -- use depth-first order
+			// If using direct children, or merging leaf hulls, order matters -- use depth-first order
 			TSet<int32> ToProcessSet;
 			ToProcessSet.Append(TransformSubset);
 			TransformsToProcess.Reserve(TransformSubset.Num());
@@ -2032,12 +2032,16 @@ void FGeometryCollectionConvexUtility::GenerateClusterConvexHullsFromLeafOrChild
 			UE_LOG(LogChaos, Warning, TEXT("Transform Index %d out of bounds [0,%d)"), TransformIndex, SimulationTypeAttribute.Num());
 			continue;
 		}
-		// only do this for clusters
+		// only do this for clusters, unless we allow merging for leaf hulls too
 		const bool bIsCluster = (SimulationTypeAttribute.Get()[TransformIndex] == FGeometryCollection::ESimulationTypes::FST_Clustered);
-		if (bIsCluster)
+		if (bIsCluster || Settings.bAllowMergingLeafHulls)
 		{
 			TArray<int32> SourceTransformIndices;
-			if (bUseDirectChildren)
+			if (!bIsCluster)
+			{
+				SourceTransformIndices.Add(TransformIndex);
+			}
+			else if (bUseDirectChildren)
 			{
 				SourceTransformIndices = HierarchyFacade.GetChildrenAsArray(TransformIndex);
 				// TODO: consider expanding this recursively to children for any nodes without convex hulls, until we hit a leaf/rigid node
@@ -2100,7 +2104,7 @@ void FGeometryCollectionConvexUtility::GenerateClusterConvexHullsFromLeafOrChild
 					}
 				}
 				// Set HullProximity for convex merges -- the connections between convex hulls that can be merged
-				if (bUseDirectChildren || Settings.AllowMergesMethod == EAllowConvexMergeMethod::Any)
+				if (!bIsCluster || bUseDirectChildren || Settings.AllowMergesMethod == EAllowConvexMergeMethod::Any)
 				{
 					// With MergeMethod of Any, any pair of convex hulls can be merged
 					// Note: Currently bUseDirectChildren only supports this method
@@ -2126,7 +2130,7 @@ void FGeometryCollectionConvexUtility::GenerateClusterConvexHullsFromLeafOrChild
 								}
 							}
 							// Consider merges between convex hulls within a transform
-							if (bConsiderIntraTransformConvexMerges)
+							if (!bIsCluster || bConsiderIntraTransformConvexMerges)
 							{
 								for (int32 HullSubIdx = 0; HullSubIdx + 1 < HullIndicesA->Num(); ++HullSubIdx)
 								{
diff --git a/Engine/Source/Runtime/Experimental/Chaos/Public/GeometryCollection/GeometryCollectionConvexUtility.h b/Engine/Source/Runtime/Experimental/Chaos/Public/GeometryCollection/GeometryCollectionConvexUtility.h
--- a/Engine/Source/Runtime/Experimental/Chaos/Public/GeometryCollection/GeometryCollectionConvexUtility.h
+++ b/Engine/Source/Runtime/Experimental/Chaos/Public/GeometryCollection/GeometryCollectionConvexUtility.h
@@ -124,6 +124,8 @@ public:
 		int32 ConvexCount = 4;
 		double ErrorToleranceInCm = 0.0;
 		bool bUseExternalCollisionIfAvailable = true;
+		// Option to apply the same hull merging settings to leaf hulls
+		bool bAllowMergingLeafHulls = false;
 		EAllowConvexMergeMethod AllowMergesMethod = EAllowConvexMergeMethod::ByProximity;
 		UE::Geometry::FSphereCovering* EmptySpace = nullptr;
 	};
