Geometry collection : Fix issue where setting rest collection after the component is registered prevent it to using custom renderers
- making sure the custom renderer is initialized properly when calling SetRestCollection ( moved the code from Onregister to its own shared method )

diff --git a/Engine/Source/Runtime/Experimental/GeometryCollectionEngine/Private/GeometryCollection/GeometryCollectionComponent.cpp b/Engine/Source/Runtime/Experimental/GeometryCollectionEngine/Private/GeometryCollection/GeometryCollectionComponent.cpp
--- a/Engine/Source/Runtime/Experimental/GeometryCollectionEngine/Private/GeometryCollection/GeometryCollectionComponent.cpp
+++ b/Engine/Source/Runtime/Experimental/GeometryCollectionEngine/Private/GeometryCollection/GeometryCollectionComponent.cpp
@@ -3537,6 +3537,20 @@ void UGeometryCollectionComponent::OnRegister()
 
 	InitializeEmbeddedGeometry();
 
+	InitializeCustomRenderer();
+
+	Super::OnRegister();
+}
+
+void UGeometryCollectionComponent::OnUnregister()
+{
+	Super::OnUnregister();
+
+	UnregisterCustomRenderer();
+}
+
+void UGeometryCollectionComponent::InitializeCustomRenderer()
+{
 	// Create a custom renderer object if a type is set on the component or the GC asset.
 	UClass* Type = bOverrideCustomRenderer ? CustomRendererType : (RestCollection ? RestCollection->CustomRendererType : nullptr);
 	if (Type && Type->ImplementsInterface(UGeometryCollectionExternalRenderInterface::StaticClass()))
@@ -3559,15 +3573,6 @@ void UGeometryCollectionComponent::OnRegister()
 	{
 		CustomRenderer = nullptr;
 	}
-
-	Super::OnRegister();
-}
-
-void UGeometryCollectionComponent::OnUnregister()
-{
-	Super::OnUnregister();
-
-	UnregisterCustomRenderer();
 }
 
 void UGeometryCollectionComponent::RegisterCustomRenderer()
@@ -4973,6 +4978,8 @@ void UGeometryCollectionComponent::SetRestCollection(const UGeometryCollection*
 			InitializeEmbeddedGeometry();
 		}
 
+		InitializeCustomRenderer();
+
 		if (CustomRenderer)
 		{
 			UnregisterCustomRenderer();
diff --git a/Engine/Source/Runtime/Experimental/GeometryCollectionEngine/Public/GeometryCollection/GeometryCollectionComponent.h b/Engine/Source/Runtime/Experimental/GeometryCollectionEngine/Public/GeometryCollection/GeometryCollectionComponent.h
--- a/Engine/Source/Runtime/Experimental/GeometryCollectionEngine/Public/GeometryCollection/GeometryCollectionComponent.h
+++ b/Engine/Source/Runtime/Experimental/GeometryCollectionEngine/Public/GeometryCollection/GeometryCollectionComponent.h
@@ -1776,6 +1776,7 @@ private:
 	/** return true if a custom renderer is active */
 	GEOMETRYCOLLECTIONENGINE_API bool IsUsingCustomRenderer() const;
 
+	void InitializeCustomRenderer();
 	GEOMETRYCOLLECTIONENGINE_API void RegisterCustomRenderer();
 	GEOMETRYCOLLECTIONENGINE_API void UnregisterCustomRenderer();
 
