Move the cvar 'ias.HttpFailTimeOutMs' from the OnDemandHttpClient to the DispatcherBackend

#rnx

- FEventLoop::FImpl::SetFailTimeout now accepts zero and negative values as a way to request that the timeout be reset to the general default.
- FHttpClient and FMultiEndpointHttpClient now accept a timeout value via their config structures when first created rather than checking 'ias.HttpFailTimeOutMs' directly.
- The 'ias.HttpFailTimeOutMs' cvar has been moved to the dispatcher backend code, when the backend creates it's client we pass the value of the cvar in as the timeout.
-- This will allow multiple systems (like IAS, IAD etc) to use different timeout values while sharing the same client code.

diff --git a/Engine/Source/Runtime/Experimental/IoStore/HttpClient/Private/Loop.inl b/Engine/Source/Runtime/Experimental/IoStore/HttpClient/Private/Loop.inl
--- a/Engine/Source/Runtime/Experimental/IoStore/HttpClient/Private/Loop.inl
+++ b/Engine/Source/Runtime/Experimental/IoStore/HttpClient/Private/Loop.inl
@@ -1312,9 +1312,14 @@ void FEventLoop::FImpl::Throttle(uint32 KiBPerSec)
 ////////////////////////////////////////////////////////////////////////////////
 void FEventLoop::FImpl::SetFailTimeout(int32 TimeoutMs)
 {
-	// While TimeoutMs must be >=0, it is signed so that MAX_uint32 can't be used
-	check(TimeoutMs > 0);
-	FailTimeoutMs = TimeoutMs;
+	if (TimeoutMs > 0)
+	{
+		FailTimeoutMs = TimeoutMs;
+	}
+	else
+	{
+		FailTimeoutMs = GIdleMs; // Reset to default
+	}
 }
 
 ////////////////////////////////////////////////////////////////////////////////
diff --git a/Engine/Source/Runtime/Experimental/IoStore/OnDemand/Private/OnDemandHttpClient.cpp b/Engine/Source/Runtime/Experimental/IoStore/OnDemand/Private/OnDemandHttpClient.cpp
--- a/Engine/Source/Runtime/Experimental/IoStore/OnDemand/Private/OnDemandHttpClient.cpp
+++ b/Engine/Source/Runtime/Experimental/IoStore/OnDemand/Private/OnDemandHttpClient.cpp
@@ -14,14 +14,6 @@
 namespace UE::IoStore
 {
 
-static int32 GIasHttpFailTimeOutMs = 4 * 1000;
-static FAutoConsoleVariableRef CVar_IasHttpFailTimeOutMs(
-	TEXT("ias.HttpFailTimeOutMs"),
-	GIasHttpFailTimeOutMs,
-	TEXT("Fail infinite network waits that take longer than this (in ms, 0=disabled)")
-);
-
-
 static bool GIasHttpAllowChunkedXfer = false;
 static FAutoConsoleVariableRef CVar_IasHttpAllowChunkedXfer(
 	TEXT("ias.HttpAllowChunkedXfer"),
@@ -156,7 +148,7 @@ FHttpClient::FHttpClient(FHttpClientConfig&& ClientConfig)
 	Connections.SetNum(Config.Endpoints.Num());
 	SetEndpoint(Config.PrimaryEndpoint);
 
-	EventLoop.SetFailTimeout(GIasHttpFailTimeOutMs);
+	EventLoop.SetFailTimeout(Config.TimeoutMs);
 }
 
 TUniquePtr<FHttpClient> FHttpClient::Create(FHttpClientConfig&& ClientConfig)
@@ -390,6 +382,7 @@ TIoStatusOr<FIoBuffer> FHttpClient::Get(FStringView Url, uint32 RetryCount, EHtt
 FMultiEndpointHttpClient::FMultiEndpointHttpClient(const FMultiEndpointHttpClientConfig& InConfig)
 	: Config(InConfig)
 {
+	EventLoop.SetFailTimeout(Config.TimeoutMs);
 }
 
 TUniquePtr<FMultiEndpointHttpClient> FMultiEndpointHttpClient::Create(const FMultiEndpointHttpClientConfig& Config)
diff --git a/Engine/Source/Runtime/Experimental/IoStore/OnDemand/Private/OnDemandHttpClient.h b/Engine/Source/Runtime/Experimental/IoStore/OnDemand/Private/OnDemandHttpClient.h
--- a/Engine/Source/Runtime/Experimental/IoStore/OnDemand/Private/OnDemandHttpClient.h
+++ b/Engine/Source/Runtime/Experimental/IoStore/OnDemand/Private/OnDemandHttpClient.h
@@ -33,6 +33,7 @@ struct FHttpClientConfig
 	int32 MaxConnectionCount = 8;
 	int32 MaxRetryCount = 1;
 	int32 ReceiveBufferSize = -1;
+	int32 TimeoutMs = 0;
 	bool bChangeEndpointAfterSuccessfulRetry = true;
 };
 
@@ -87,6 +88,7 @@ struct FMultiEndpointHttpClientConfig
 	int32			MaxConnectionCount = 4;
 	int32			ReceiveBufferSize = -1;
 	uint32			MaxRetryCount = 0;
+	int32			TimeoutMs = 0;
 	EHttpRedirects	Redirects = EHttpRedirects::Follow;
 };
 
diff --git a/Engine/Source/Runtime/Experimental/IoStore/OnDemand/Private/OnDemandIoDispatcherBackend.cpp b/Engine/Source/Runtime/Experimental/IoStore/OnDemand/Private/OnDemandIoDispatcherBackend.cpp
--- a/Engine/Source/Runtime/Experimental/IoStore/OnDemand/Private/OnDemandIoDispatcherBackend.cpp
+++ b/Engine/Source/Runtime/Experimental/IoStore/OnDemand/Private/OnDemandIoDispatcherBackend.cpp
@@ -109,6 +109,13 @@ static FAutoConsoleVariableRef CVar_IasHttpConnectionCount(
 	TEXT("Number of open HTTP connections to the on demand endpoint(s).")
 );
 
+static int32 GIasHttpFailTimeOutMs = 4 * 1000;
+static FAutoConsoleVariableRef CVar_IasHttpFailTimeOutMs(
+	TEXT("ias.HttpFailTimeOutMs"),
+	GIasHttpFailTimeOutMs,
+	TEXT("Fail infinite network waits that take longer than this (in ms, a value of zero will use the default)")
+);
+
 /**
  *This is only applied when the connection was made to a single ServiceUrl rather than a DistributedUrl.
  * In the latter case we will make two attempts on the primary CDN followed by a single attempt for the
@@ -1722,6 +1729,7 @@ uint32 FOnDemandIoBackend::Run()
 		.MaxConnectionCount = GIasHttpConnectionCount,
 		.MaxRetryCount = FMath::Max(AvailableEps.Urls.Num(), GIasHttpRetryCount),
 		.ReceiveBufferSize = GIasHttpRecvBufKiB >= 0 ? GIasHttpRecvBufKiB << 10 : -1,
+		.TimeoutMs = GIasHttpFailTimeOutMs,
 		.bChangeEndpointAfterSuccessfulRetry = GIasHttpChangeEndpointAfterSuccessfulRetry,
 	});
 
