Misc fixes and tweaks

#rnx
#ushell-unshelve of 38505511

diff --git a/Engine/Source/Runtime/Experimental/IoStore/HttpClient/Private/ConnectionPool.inl b/Engine/Source/Runtime/Experimental/IoStore/HttpClient/Private/ConnectionPool.inl
--- a/Engine/Source/Runtime/Experimental/IoStore/HttpClient/Private/ConnectionPool.inl
+++ b/Engine/Source/Runtime/Experimental/IoStore/HttpClient/Private/ConnectionPool.inl
@@ -226,7 +226,7 @@ FOutcome FHost::Connect(FSocket& Socket)
 	}
 
 	Socket = MoveTemp(Candidate);
-	return FOutcome::Waiting();
+	return Outcome;
 }
 
 
diff --git a/Engine/Source/Runtime/Experimental/IoStore/HttpClient/Private/Socket.inl b/Engine/Source/Runtime/Experimental/IoStore/HttpClient/Private/Socket.inl
--- a/Engine/Source/Runtime/Experimental/IoStore/HttpClient/Private/Socket.inl
+++ b/Engine/Source/Runtime/Experimental/IoStore/HttpClient/Private/Socket.inl
@@ -33,7 +33,8 @@
 		};
 
 #	define IAS_HTTP_USE_POLL
-		template <typename... ArgTypes> auto poll(ArgTypes... Args)
+		template <typename... ArgTypes>
+		static auto poll(ArgTypes&&... Args)
 		{
 			return WSAPoll(Forward<ArgTypes>(Args)...);
 		}
@@ -63,10 +64,10 @@
 namespace UE::IoStore::HTTP
 {
 
-// {{{1 socket .................................................................
+// {{{1 wait ...................................................................
 
 ////////////////////////////////////////////////////////////////////////////////
-static const SocketType InvalidSocket = ~SocketType(0);
+static const SocketType InvalidSocket = SocketType(~0ull);
 
 static_assert(sizeof(sockaddr_in::sin_addr) == sizeof(uint32));
 
@@ -126,14 +127,14 @@ public:
 					FWaiter(FWaitable&& Waitable);
 					FWaiter(FWaiter&&) = default;
 	FWaiter&		operator = (FWaiter&&) = default;
-	bool			IsValid() const				{ return Socket != 0; }
+	bool			IsValid() const				{ return Socket != InvalidSocket; }
 	bool			IsReady() const				{ return Ready != 0; }
 	void			WaitFor(EWhat What)			{ WaitOn = What; }
 	void			SetIndex(uint32 InIndex)	{ Index = uint16(InIndex); }
 	uint32			GetIndex() const			{ return Index; }
 
 private:
-	SocketType		Socket = 0;
+	SocketType		Socket = InvalidSocket;
 	uint16			Index = 0;
 	EWhat			WaitOn = EWhat::None;
 	uint8			Ready = 0;
@@ -172,7 +173,7 @@ int32 FWaiter::Wait(TArrayView<FWaiter> Waiters, int32 TimeoutMs)
 	// to the next. To cleanly set members to zero and to not get narrowing
 	// warnings from the compiler, we list-init and don't assume POD types.
 	using PollEventType = decltype(FSelect::events);
-	PollEventType Events[] = { POLLERR, POLLOUT, POLLIN, POLLOUT|POLLOUT };
+	PollEventType Events[] = { POLLERR, POLLOUT, POLLIN, POLLIN|POLLOUT };
 	TArray<FSelect, TFixedAllocator<64>> Selects;
 	for (FWaiter& Waiter : Waiters)
 	{
@@ -256,6 +257,8 @@ int32 FWaiter::Wait(TArrayView<FWaiter> Waiters, int32 TimeoutMs)
 
 
 
+// {{{1 socket .................................................................
+
 ////////////////////////////////////////////////////////////////////////////////
 class FSocket
 {
@@ -436,8 +439,6 @@ bool FSocket::SetRecvBufSize(int32 Size)
 	return 0 == setsockopt(Socket, SOL_SOCKET, SO_RCVBUF, &(char&)Size, sizeof(Size));
 }
 
-
-
 // }}}
 
 } // namespace UE::IoStore::HTTP
