[ModuleInput] Enums can only have values added to the end of the enum otherwise existing assets get their enum type changed unintenionally.

diff --git a/Engine/Source/Runtime/Experimental/ChaosVehicles/ChaosVehiclesCore/Public/SimModule/ModuleInput.h b/Engine/Source/Runtime/Experimental/ChaosVehicles/ChaosVehiclesCore/Public/SimModule/ModuleInput.h
--- a/Engine/Source/Runtime/Experimental/ChaosVehicles/ChaosVehiclesCore/Public/SimModule/ModuleInput.h
+++ b/Engine/Source/Runtime/Experimental/ChaosVehicles/ChaosVehiclesCore/Public/SimModule/ModuleInput.h
@@ -13,15 +13,14 @@ DECLARE_LOG_CATEGORY_EXTERN(LogModularInput, Log, All);
 UENUM(BlueprintType)
 enum class EModuleInputValueType : uint8
 {
-	// Value types in increasing size order (used for type promotion)
-
 	MBoolean	UMETA(DisplayName = "Digital (bool)"),
-	MInteger	UMETA(DisplayName = "Digital (int32)"),
 	MAxis1D		UMETA(DisplayName = "Axis1D (float)"),
 	MAxis2D		UMETA(DisplayName = "Axis2D (Vector2D)"),
 	MAxis3D		UMETA(DisplayName = "Axis3D (Vector)"),
+	MInteger	UMETA(DisplayName = "Digital (int32)"),
 
-	// NOTE, If adding an entry here, update the number of bits to serialize in FModuleInputValue::Serialize
+	// NOTE, If adding an entry here, add it to the end, otherwise existing assets get deserialized improperly.
+	// ALSO update the number of bits to serialize in FModuleInputValue::Serialize
 };
 
 /** Input Options */
