Minor improvements to FMultiEndpointHttpClient logging.

#rnx

- Add retry info to the logging.

diff --git a/Engine/Source/Runtime/Experimental/IoStore/OnDemand/Private/OnDemandHttpClient.cpp b/Engine/Source/Runtime/Experimental/IoStore/OnDemand/Private/OnDemandHttpClient.cpp
--- a/Engine/Source/Runtime/Experimental/IoStore/OnDemand/Private/OnDemandHttpClient.cpp
+++ b/Engine/Source/Runtime/Experimental/IoStore/OnDemand/Private/OnDemandHttpClient.cpp
@@ -586,7 +586,7 @@ void FMultiEndpointHttpClient::Log(const FRequest& Request)
 	Sb << TEXT("GET ")
 		<< Request.Connection.HostGroup.Host(Request.Host)
 		<< Request.RelativeUrl;
-	Sb.Appendf(TEXT(" (Status: %u, %llums, %.2lf KiB)"), Request.StatusCode, DurationMs, double(Request.Body.GetSize()) / 1024.0);
+	Sb.Appendf(TEXT(" (Status: %u, Attempt %d/%d, %llums, %.2lf KiB)"), Request.StatusCode, Request.RetryCount, Config.MaxRetryCount, DurationMs, double(Request.Body.GetSize()) / 1024.0);
 
 	UE_LOG(LogIoStoreOnDemand, Log, TEXT("%s"), Sb.ToString());
 }
