Make FOnDemandHostGroup::CycleHost const.

#rnx

- This method returns the next host in the list rather than changing any internal state. Making the method const better relfects the intent to the caller.

diff --git a/Engine/Source/Runtime/Experimental/IoStore/OnDemand/Private/OnDemandHostGroup.cpp b/Engine/Source/Runtime/Experimental/IoStore/OnDemand/Private/OnDemandHostGroup.cpp
--- a/Engine/Source/Runtime/Experimental/IoStore/OnDemand/Private/OnDemandHostGroup.cpp
+++ b/Engine/Source/Runtime/Experimental/IoStore/OnDemand/Private/OnDemandHostGroup.cpp
@@ -57,7 +57,7 @@ FAnsiStringView FOnDemandHostGroup::Host(int32 Index) const
 	return Impl->HostUrls.IsEmpty() ? FAnsiStringView() : Impl->HostUrls[Index];
 }
 
-FAnsiStringView FOnDemandHostGroup::CycleHost(int32& InOutIndex)
+FAnsiStringView FOnDemandHostGroup::CycleHost(int32& InOutIndex) const
 {
 	InOutIndex = (InOutIndex + 1) % Impl->HostUrls.Num();
 	return Host(InOutIndex);
diff --git a/Engine/Source/Runtime/Experimental/IoStore/OnDemand/Public/IO/OnDemandHostGroup.h b/Engine/Source/Runtime/Experimental/IoStore/OnDemand/Public/IO/OnDemandHostGroup.h
--- a/Engine/Source/Runtime/Experimental/IoStore/OnDemand/Public/IO/OnDemandHostGroup.h
+++ b/Engine/Source/Runtime/Experimental/IoStore/OnDemand/Public/IO/OnDemandHostGroup.h
@@ -29,7 +29,7 @@ public:
 	/** Get the URL at the specified index. */
 	UE_API FAnsiStringView Host(int32 Index) const;
 	/** Get the next available host starting from the specified index. */
-	UE_API FAnsiStringView CycleHost(int32& InOutIndex);
+	UE_API FAnsiStringView CycleHost(int32& InOutIndex) const;
 	/** Set the primary host URL. */
 	UE_API void SetPrimaryHost(int32 Index);
 	/** Get the primary host URL. */
