Repacked cache work queue so it is Arm MTE friendly
#rnx
#ushell-unshelve of 39205013

diff --git a/Engine/Source/Runtime/Experimental/IoStore/OnDemand/Private/JournaledCache.cpp b/Engine/Source/Runtime/Experimental/IoStore/OnDemand/Private/JournaledCache.cpp
--- a/Engine/Source/Runtime/Experimental/IoStore/OnDemand/Private/JournaledCache.cpp
+++ b/Engine/Source/Runtime/Experimental/IoStore/OnDemand/Private/JournaledCache.cpp
@@ -1757,10 +1757,10 @@ public:
 	void					SetGovernorDemand(uint32 Threshold, uint32 Boost, uint32 SuperBoost);
 	uint32					ClaimReadId();
 	void					BeginRead(const FCache* Cache, const FReadRequest& Request);
-	void					CancelRead(const void* GivenDest);
+	void					CancelRead(const FIoBuffer* GivenDest);
 
 private:
-	struct FWork
+	union FWork
 	{
 		enum {
 			Work_Register,
@@ -1771,36 +1771,38 @@ private:
 			Work_Cancel,
 		};
 
-		void	SetPtr(const void* In)	{ Ptr = UPTRINT(In) >> 3; check((UPTRINT(In) & 0x7) == 0); }
-		void*	GetPtr() const			{ return (FCache*)(Ptr << 3); }
+		template <typename Type>
+		struct TField
+		{
+			void	operator = (auto In)	{ ::memcpy(Value, &In, sizeof(Value)); }
+					operator Type () const	{ Type Ret; ::memcpy(&Ret, Value, sizeof(Ret)); return Ret; }
+			uint16	Value[sizeof(Type) / sizeof(uint16)];
+		};
 
-		union {
-			struct {			// reg / unreg / read / cancel
-				UPTRINT			What : 3;
-				UPTRINT			Ptr : 45;
-				UPTRINT			ReadId : 16;
-			};
-			struct {			// rate
-				uint16			_What0 : 3;
-				uint16			Ops : 13;
-				uint16			Seconds;
-				uint32			Allowance;
-			};
-			struct {			// demand
-				uint16			_What1 : 3;
-				uint16			_Unused1 : 13;
-				uint16			Threshold;
-				uint16			Boost;
-				uint16			SuperBoost;
-			};
-			FReadSink*			Sink;
+		struct {				// reg / unreg / read / cancel
+			uint16				What;
+			TField<FCache*>		Cache;
+			TField<uint64>		Key;
+		};
+		struct {				// read.param
+			uint16				ReadId;
+			TField<FReadSink*>	Sink;
+			TField<FIoBuffer*>	Dest;
+		};
+		struct {				// rate
+			uint16				_What0;
+			uint16				Ops;
+			uint16				Seconds;
+			TField<uint32>		Allowance;
 		};
-		union {
-			uint64				Key;	// read
-			FIoBuffer*			Dest;	// read.param
+		struct {				// demand
+			uint16				_What1;
+			uint16				Threshold;
+			uint16				Boost;
+			uint16				SuperBoost;
 		};
 	};
-	static_assert(sizeof(FWork) == sizeof(UPTRINT) * 2);
+	static_assert(sizeof(FWork) == sizeof(uint16) * 9);
 
 	void						StartThread();
 	int32						Update();
@@ -1860,7 +1862,7 @@ void FServiceThread::RegisterCache(TUniquePtr<FCache> Cache)
 
 	FWork Work;
 	Work.What = FWork::Work_Register;
-	Work.SetPtr(RawPtr);
+	Work.Cache = RawPtr;
 	SubmitWork(&Work, 1);
 
 	if (PrevRunCount == 0)
@@ -1874,7 +1876,7 @@ void FServiceThread::UnregisterCache(FCache* Cache)
 {
 	FWork Work;
 	Work.What = FWork::Work_Unregister;
-	Work.SetPtr(Cache);
+	Work.Cache = Cache;
 	SubmitWork(&Work, 1);
 
 	int32 PrevRunCount = RunCount.fetch_sub(1, std::memory_order_relaxed);
@@ -1917,9 +1919,7 @@ void FServiceThread::SetGovernorDemand(uint32 Threshold, uint32 Boost, uint32 Su
 uint32 FServiceThread::ClaimReadId()
 {
 	++ReadIdCounter;
-	FWork BitfieldNarrowing;
-	BitfieldNarrowing.ReadId = ReadIdCounter;
-	return BitfieldNarrowing.ReadId;
+	return uint16(ReadIdCounter);
 }
 
 ////////////////////////////////////////////////////////////////////////////////
@@ -1927,11 +1927,11 @@ void FServiceThread::BeginRead(const FCache* Cache, const FReadRequest& Request)
 {
 	FWork Works[2] = {};
 
-	Works[0].SetPtr(Cache);
+	Works[0].Cache = Cache;
 	Works[0].What = FWork::Work_Read;
 	Works[0].Key = Request.Key;
-	Works[0].ReadId = Request.ReadId;
 
+	Works[1].ReadId = uint16(Request.ReadId);
 	Works[1].Sink = Request.Sink;
 	Works[1].Dest = Request.Dest;
 
@@ -1939,11 +1939,11 @@ void FServiceThread::BeginRead(const FCache* Cache, const FReadRequest& Request)
 }
 
 ////////////////////////////////////////////////////////////////////////////////
-void FServiceThread::CancelRead(const void* GivenData)
+void FServiceThread::CancelRead(const FIoBuffer* GivenData)
 {
 	FWork Work;
 	Work.What = FWork::Work_Cancel;
-	Work.SetPtr(GivenData);
+	Work.Dest = GivenData;
 	SubmitWork(&Work, 1);
 }
 
@@ -2039,12 +2039,13 @@ int32 FServiceThread::Update()
 		EIoErrorCode Status = EIoErrorCode::Cancelled;
 		if (Param.Dest != nullptr)
 		{
-			auto* Cache = (FCache*)(Work.GetPtr());
+			FCache* Cache = Work.Cache;
 			Status = Cache->Materialize(Work.Key, *Param.Dest);
 		}
 
-		FReadSink::FReadResult Result = { uint16(Work.ReadId), uint16(Status) };
-		Param.Sink->OnRead(&Result, 1);
+		auto Sink = (FReadSink*)(Param.Sink);
+		FReadSink::FReadResult Result = { uint16(Param.ReadId), uint16(Status) };
+		Sink->OnRead(&Result, 1);
 
 		Cycle = FPlatformTime::Cycles64();
 		if (Cycle >= StopReadsCycle)
@@ -2083,7 +2084,7 @@ void FServiceThread::ReceiveWork()
 			continue;
 		}
 
-		auto* CachePtr = (FCache*)(Work.GetPtr());
+		FCache* CachePtr = Work.Cache;
 		bool bFound = false;
 		for (int32 i = 0, n = Caches.Num(); i < n; ++i)
 		{
@@ -2103,7 +2104,7 @@ void FServiceThread::ReceiveWork()
 		{
 			for(const FWork& RegWork : InboundWork)
 			{
-				if (RegWork.What == FWork::Work_Register && RegWork.GetPtr() == CachePtr)
+				if (RegWork.What == FWork::Work_Register && RegWork.Cache == CachePtr)
 				{
 					SubmitWork(&Work, 1);
 				}
@@ -2123,7 +2124,7 @@ void FServiceThread::ReceiveWork()
 
 		if (Work.What == FWork::Work_Register)
 		{
-			auto* Cache = (FCache*)(Work.GetPtr());
+			FCache* Cache = Work.Cache;
 			Cache->Load();
 			Caches.Add(TUniquePtr<FCache>(Cache));
 			continue;
@@ -2154,7 +2155,7 @@ void FServiceThread::ReceiveWork()
 			continue;
 		}
 
-		const void* ToCancel = Work.GetPtr();
+		const void* ToCancel = Work.Dest;
 		for (uint32 i = 1, n = ActiveReads.Num(); i < n; i += 2)
 		{
 			if (ActiveReads[i].Dest == ToCancel)
