Chaos - improve logging for AABBTree Nan detection

diff --git a/Engine/Source/Runtime/Experimental/Chaos/Public/Chaos/AABBTree.h b/Engine/Source/Runtime/Experimental/Chaos/Public/Chaos/AABBTree.h
--- a/Engine/Source/Runtime/Experimental/Chaos/Public/Chaos/AABBTree.h
+++ b/Engine/Source/Runtime/Experimental/Chaos/Public/Chaos/AABBTree.h
@@ -1860,22 +1860,29 @@ public:
 		if (bHasBounds && ValidateBounds(NewBounds) == false)
 		{
 			bHasBounds = false;
+			TSharedPtr<FString, ESPMode::ThreadSafe> DebugName;
 #if CHAOS_DEBUG_NAME			
 			if constexpr (std::is_same_v<TPayloadType, FAccelerationStructureHandle>)
 			{
-				if (const FGeometryParticleHandle* Particle = Payload.GetGeometryParticleHandle_PhysicsThread())
+				if (IsInPhysicsThreadContext())
 				{
-					const TSharedPtr<FString, ESPMode::ThreadSafe>& DebugName = Particle->DebugName();
-					if (IsInPhysicsThreadContext())
+					if (const FGeometryParticleHandle* Particle = Payload.GetGeometryParticleHandle_PhysicsThread())
 					{
-						FString DebugStr = DebugName ? *DebugName : TEXT("No Name");
-						UE_LOG(LogChaos, Warning, TEXT("AABBTree encountered invalid bounds input : %s"), *DebugStr);
+						DebugName = Particle->DebugName();
+					}
+				}
+				else
+				{
+					if (const FGeometryParticle* Particle = Payload.GetExternalGeometryParticle_ExternalThread())
+					{
+						DebugName = Particle->DebugName();
 					}
 				}
 			}
 #endif
-			UE_LOG(LogChaos, Warning, TEXT("AABBTree encountered invalid bounds input.Forcing element to global payload. Min: %s Max: %s"),
-				*NewBounds.Min().ToString(), *NewBounds.Max().ToString());
+			UE_LOG(LogChaos, Warning, TEXT("AABBTree encountered invalid bounds. Min: %s Max: %s Name: %s")
+				, *NewBounds.Min().ToString() , *NewBounds.Max().ToString()
+				, DebugName.IsValid() ? **DebugName : TEXT("<Unknown>"));
 		}
 
 		bool bElementExisted = true;
