Fix IAS latency tests to use the correct timeout value

#rnx

- The timeout value was not being applied to the FEventLoop via ::SetFailTimeout it was only ever being used in the Tick method. The timeout for the tick doesn't actually timeout the request it only times out the tick.
- Removing the timeout from the tick seemed to cause the value set in ::SetFailTimeout to no longer work which I believe is a bug, I have left the value there for now and will follow up on that in a future work item.
- Fixed some minor coding standards issues while I was in the code area.

diff --git a/Engine/Source/Runtime/Experimental/IoStore/OnDemand/Private/LatencyTesting.cpp b/Engine/Source/Runtime/Experimental/IoStore/OnDemand/Private/LatencyTesting.cpp
--- a/Engine/Source/Runtime/Experimental/IoStore/OnDemand/Private/LatencyTesting.cpp
+++ b/Engine/Source/Runtime/Experimental/IoStore/OnDemand/Private/LatencyTesting.cpp
@@ -28,6 +28,7 @@ void LatencyTest(FAnsiStringView InUrl, FStringView InPath, uint32 InTimeOutMs,
 	AnsiPath << InPath;
 
 	FEventLoop Loop;
+	Loop.SetFailTimeout(InTimeOutMs);
 
 	for (int32& Result : OutResults)
 	{
@@ -44,7 +45,12 @@ void LatencyTest(FAnsiStringView InUrl, FStringView InPath, uint32 InTimeOutMs,
 			});
 
 		uint64 Cycles = FPlatformTime::Cycles64();
-		while (Loop.Tick(InTimeOutMs) != 0);
+		
+		while (Loop.Tick(InTimeOutMs) != 0)
+		{
+
+		}
+
 		Cycles = FPlatformTime::Cycles64() - Cycles;
 
 		Result = Ok ? int32(Cycles) : -1;
@@ -54,7 +60,9 @@ void LatencyTest(FAnsiStringView InUrl, FStringView InPath, uint32 InTimeOutMs,
 	for (int32& Result : OutResults)
 	{
 		if (Result == -1)
+		{
 			continue;
+		}
 
 		Result = int32((int64(Result) * 1000) / Freq);
 	}
