Missed this small function in previous checkin wher adding integer option to vehicle inputs

diff --git a/Engine/Source/Runtime/Experimental/ChaosVehicles/ChaosVehiclesCore/Public/SimModule/ModuleInput.h b/Engine/Source/Runtime/Experimental/ChaosVehicles/ChaosVehiclesCore/Public/SimModule/ModuleInput.h
--- a/Engine/Source/Runtime/Experimental/ChaosVehicles/ChaosVehiclesCore/Public/SimModule/ModuleInput.h
+++ b/Engine/Source/Runtime/Experimental/ChaosVehicles/ChaosVehiclesCore/Public/SimModule/ModuleInput.h
@@ -443,6 +443,11 @@ public:
 		return InValue.Get<FModuleInputValue::MAxis3D>();
 	}
 
+	static int32 ToInteger(FModuleInputValue& InValue)
+	{
+		return InValue.Get<FModuleInputValue::MInteger>();
+	}
+
 	static FString ToString(FModuleInputValue& ActionValue)
 	{
 		return ActionValue.ToString();
