Fix trimesh queries for negatively scaled meshes. premultiply NdotD for backface testing to get the correct direction for the normal.

diff --git a/Engine/Source/Runtime/Experimental/Chaos/Private/Chaos/TriangleMeshImplicitObject.cpp b/Engine/Source/Runtime/Experimental/Chaos/Private/Chaos/TriangleMeshImplicitObject.cpp
--- a/Engine/Source/Runtime/Experimental/Chaos/Private/Chaos/TriangleMeshImplicitObject.cpp
+++ b/Engine/Source/Runtime/Experimental/Chaos/Private/Chaos/TriangleMeshImplicitObject.cpp
@@ -1306,14 +1306,17 @@ struct FTriangleMeshSweepVisitor
 
 			const VectorRegister4Float PositiveEpsilon = MakeVectorRegister(UE_KINDA_SMALL_NUMBER, UE_KINDA_SMALL_NUMBER, UE_KINDA_SMALL_NUMBER, UE_KINDA_SMALL_NUMBER);
 			const VectorRegister4Float NegativeEpsilon = VectorNegate(PositiveEpsilon);
-			const VectorRegister4Float IsAlongPositiveNormal = VectorCompareGT(NDotD, PositiveEpsilon);
-			const VectorRegister4Float IsAlongNegativeNormal = VectorCompareLT(NDotD, NegativeEpsilon);
 
-			const VectorRegister4Float IsAlongPositiveNormalMask = VectorMultiply(VectorCullsBackFaceSweepsCode, IsAlongPositiveNormal);
-			const VectorRegister4Float IsAlongNegativeNormalMask = VectorMultiply(VectorCullsBackFaceSweepsCode, IsAlongNegativeNormal);
+			// NDotD is only in the correct direction when all scale components are positive.
+			// Here we correct it if we need to have changed the winding order of the triangle
+			const VectorRegister4Float ScaleCorrectedNDotD = VectorMultiply(NDotD, VectorCullsBackFaceSweepsCode);
+
+			// Now NDotD is correct we can see if we're approaching the front or back of the triangle
+			const VectorRegister4Float IsAlongPositiveNormal = VectorCompareGT(ScaleCorrectedNDotD, PositiveEpsilon);
+			const VectorRegister4Float IsAlongNegativeNormal = VectorCompareLT(ScaleCorrectedNDotD, NegativeEpsilon);
+			const bool bIsBackFace = VectorMaskBits(IsAlongPositiveNormal) != 0;
+			const bool bIsFrontFace = VectorMaskBits(IsAlongNegativeNormal) != 0;
 
-			const bool bIsBackFace = VectorMaskBits(IsAlongPositiveNormalMask) != 0;
-			const bool bIsFrontFace = VectorMaskBits(IsAlongNegativeNormalMask) != 0;
 			bIsParallelSweepAndCullsBackFace = (!bIsBackFace && !bIsFrontFace);
 
 			// If this is a definitive back-face hit, always skip, even if there is an initial overlap.
