Remove atomic or from DirtyNetObjectTracker to make it safe for AutoRTFM transactions

diff --git a/Engine/Source/Runtime/Experimental/Iris/Core/Private/Iris/ReplicationSystem/DirtyNetObjectTracker.cpp b/Engine/Source/Runtime/Experimental/Iris/Core/Private/Iris/ReplicationSystem/DirtyNetObjectTracker.cpp
--- a/Engine/Source/Runtime/Experimental/Iris/Core/Private/Iris/ReplicationSystem/DirtyNetObjectTracker.cpp
+++ b/Engine/Source/Runtime/Experimental/Iris/Core/Private/Iris/ReplicationSystem/DirtyNetObjectTracker.cpp
@@ -1,7 +1,6 @@
 // Copyright Epic Games, Inc. All Rights Reserved.
 
 #include "DirtyNetObjectTracker.h"
-#include "HAL/PlatformAtomics.h"
 #include "Iris/IrisConstants.h"
 #include "Iris/Core/IrisProfiler.h"
 #include "Iris/ReplicationSystem/ReplicationSystem.h"
@@ -161,8 +160,7 @@ void FDirtyNetObjectTracker::MarkNetObjectDirty(FInternalNetRefIndex NetObjectIn
 
 	uint32* DirtyNetObjectsData = DirtyNetObjects.GetData();
 
-	// ideally we'd have c++20 std::atomic_ref for this
-	FPlatformAtomics::InterlockedOr((TSignedIntType<sizeof(StorageType)>::Type*)(&DirtyNetObjectsData[BitOffset/StorageTypeBitCount]), BitMask);
+	DirtyNetObjectsData[BitOffset/StorageTypeBitCount] |= BitMask;
 
 	UE_LOG(LogIrisDirtyTracker, Verbose, TEXT("FDirtyNetObjectTracker::MarkNetObjectDirty[%u]: %s"), ReplicationSystemId, *NetRefHandleManager->PrintObjectFromIndex(NetObjectIndex));
 }
