Remove uses of atomics from DirtyNetObjectTracker to make it AutoRTFM safe

diff --git a/Engine/Source/Runtime/Experimental/Iris/Core/Private/Iris/ReplicationSystem/DirtyNetObjectTracker.cpp b/Engine/Source/Runtime/Experimental/Iris/Core/Private/Iris/ReplicationSystem/DirtyNetObjectTracker.cpp
--- a/Engine/Source/Runtime/Experimental/Iris/Core/Private/Iris/ReplicationSystem/DirtyNetObjectTracker.cpp
+++ b/Engine/Source/Runtime/Experimental/Iris/Core/Private/Iris/ReplicationSystem/DirtyNetObjectTracker.cpp
@@ -9,9 +9,6 @@
 #include "ProfilingDebugging/CsvProfiler.h"
 #include "Traits/IntType.h"
 
-#include <atomic>
-
-
 DEFINE_LOG_CATEGORY(LogIrisDirtyTracker)
 
 namespace UE::Net::Private
@@ -221,8 +218,6 @@ void FDirtyNetObjectTracker::ReconcilePolledList(const FNetBitArrayView& Objects
 	DirtyNetObjects.ClearAllBits();
 
 	AllowExternalAccess();
-
-	std::atomic_thread_fence(std::memory_order_seq_cst);
 }
 
 #if UE_NET_IRIS_CSV_STATS
diff --git a/Engine/Source/Runtime/Experimental/Iris/Core/Private/Iris/ReplicationSystem/DirtyNetObjectTracker.h b/Engine/Source/Runtime/Experimental/Iris/Core/Private/Iris/ReplicationSystem/DirtyNetObjectTracker.h
--- a/Engine/Source/Runtime/Experimental/Iris/Core/Private/Iris/ReplicationSystem/DirtyNetObjectTracker.h
+++ b/Engine/Source/Runtime/Experimental/Iris/Core/Private/Iris/ReplicationSystem/DirtyNetObjectTracker.h
@@ -124,7 +124,7 @@ private:
 	bool bShouldResetPolledGlobalDirtyTracker = false;
 
 #if UE_NET_THREAD_SAFETY_CHECK
-	std::atomic_bool bIsExternalAccessAllowed = false;
+	bool bIsExternalAccessAllowed = false;
 #endif
 
 #if UE_NET_IRIS_CSV_STATS
