Geometry collection : Fix the custom rednerer try to be initialized even if the component was not properly registered when calling SetRestCollection

diff --git a/Engine/Source/Runtime/Experimental/GeometryCollectionEngine/Private/GeometryCollection/GeometryCollectionComponent.cpp b/Engine/Source/Runtime/Experimental/GeometryCollectionEngine/Private/GeometryCollection/GeometryCollectionComponent.cpp
--- a/Engine/Source/Runtime/Experimental/GeometryCollectionEngine/Private/GeometryCollection/GeometryCollectionComponent.cpp
+++ b/Engine/Source/Runtime/Experimental/GeometryCollectionEngine/Private/GeometryCollection/GeometryCollectionComponent.cpp
@@ -4978,7 +4978,10 @@ void UGeometryCollectionComponent::SetRestCollection(const UGeometryCollection*
 			InitializeEmbeddedGeometry();
 		}
 
-		InitializeCustomRenderer();
+		if (IsRegistered())
+		{
+			InitializeCustomRenderer();
+		}
 
 		if (CustomRenderer)
 		{
