Add some safety checks to verify that certain methods in FMultiEndpointHttpClient are all called on the same thread.

#rnx

- The system is disabled in shipping configuration and can be disabled by changing UE_HTTPCLIENT_THREADSAFE_CHECKS to 0 in IoStoreOnDemand.Build.cs
- The checks are also disabled by default for each FMultiEndpointHttpClient created. The caller must opt into the test via FMultiEndpointHttpClientConfig::bEnableThreadSafetyChecks = true

diff --git a/Engine/Source/Runtime/Experimental/IoStore/OnDemand/IoStoreOnDemand.Build.cs b/Engine/Source/Runtime/Experimental/IoStore/OnDemand/IoStoreOnDemand.Build.cs
--- a/Engine/Source/Runtime/Experimental/IoStore/OnDemand/IoStoreOnDemand.Build.cs
+++ b/Engine/Source/Runtime/Experimental/IoStore/OnDemand/IoStoreOnDemand.Build.cs
@@ -25,5 +25,14 @@ public class IoStoreOnDemand : ModuleRules
 		{
 			PrivateDependencyModuleNames.AddRange(new string[] { "S3Client", "RSA" });
 		}
+
+		if (Target.Configuration != UnrealTargetConfiguration.Shipping)
+		{
+			PrivateDefinitions.Add("UE_HTTPCLIENT_THREADSAFE_CHECKS=1");
+		}
+		else
+		{
+			PrivateDefinitions.Add("UE_HTTPCLIENT_THREADSAFE_CHECKS=0");
+		}
 	}
 }
diff --git a/Engine/Source/Runtime/Experimental/IoStore/OnDemand/Private/OnDemandHttpClient.cpp b/Engine/Source/Runtime/Experimental/IoStore/OnDemand/Private/OnDemandHttpClient.cpp
--- a/Engine/Source/Runtime/Experimental/IoStore/OnDemand/Private/OnDemandHttpClient.cpp
+++ b/Engine/Source/Runtime/Experimental/IoStore/OnDemand/Private/OnDemandHttpClient.cpp
@@ -11,6 +11,12 @@
 #include "Misc/ConfigCacheIni.h"
 #include "Misc/Paths.h"
 
+#if UE_HTTPCLIENT_THREADSAFE_CHECKS
+	#define UE_CHECK_THREADSAFETY(MethodName) checkf(OwningThread == 0 || OwningThread == FPlatformTLS::GetCurrentThreadId(), #MethodName TEXT("was called on the on the wrong thread!"))
+#else
+	#define UE_CHECK_THREADSAFETY(MethodName)
+#endif //UE_HTTPCLIENT_THREADSAFE_CHECKS
+
 namespace UE::IoStore
 {
 
@@ -383,6 +389,13 @@ FMultiEndpointHttpClient::FMultiEndpointHttpClient(const FMultiEndpointHttpClien
 	: Config(InConfig)
 {
 	EventLoop.SetFailTimeout(Config.TimeoutMs);
+
+#if UE_HTTPCLIENT_THREADSAFE_CHECKS
+	if (Config.bEnableThreadSafetyChecks)
+	{
+		OwningThread = FPlatformTLS::GetCurrentThreadId();
+	}
+#endif //UE_HTTPCLIENT_THREADSAFE_CHECKS
 }
 
 TUniquePtr<FMultiEndpointHttpClient> FMultiEndpointHttpClient::Create(const FMultiEndpointHttpClientConfig& Config)
@@ -398,6 +411,8 @@ void FMultiEndpointHttpClient::Get(const FOnDemandHostGroup& Endpoint, FAnsiStri
 
 void FMultiEndpointHttpClient::Get(const FOnDemandHostGroup& Endpoint, FAnsiStringView RelativeUrl, const FIoOffsetAndLength& ChunkRange, FOnHttpResponse&& OnResponse)
 {
+	UE_CHECK_THREADSAFETY(FMultiEndpointHttpClient::Get);
+
 	FConnection& Connection = GetConnection(Endpoint);
 
 	IssueRequest(FRequest
@@ -413,6 +428,8 @@ void FMultiEndpointHttpClient::Get(const FOnDemandHostGroup& Endpoint, FAnsiStri
 
 bool FMultiEndpointHttpClient::Tick(int32 WaitTimeMs, uint32 MaxKiBPerSecond)
 {
+	UE_CHECK_THREADSAFETY(FMultiEndpointHttpClient::Tick);
+
 	EventLoop.Throttle(MaxKiBPerSecond);
 	const uint32 TicketCount = EventLoop.Tick(WaitTimeMs);
 
@@ -447,6 +464,8 @@ bool FMultiEndpointHttpClient::Tick(int32 WaitTimeMs, uint32 MaxKiBPerSecond)
 
 void FMultiEndpointHttpClient::IssueRequest(FRequest&& Request)
 {
+	UE_CHECK_THREADSAFETY(FMultiEndpointHttpClient::IssueRequest);
+
 	using namespace UE::IoStore::HTTP;
 
 	check(Request.Connection.HostGroup.IsEmpty() == false);
@@ -522,6 +541,8 @@ void FMultiEndpointHttpClient::IssueRequest(FRequest&& Request)
 
 void FMultiEndpointHttpClient::RetryRequest(FRequest&& Request)
 {
+	UE_CHECK_THREADSAFETY(FMultiEndpointHttpClient::RetryRequest);
+
 	FConnection& Connection = Request.Connection;
 
 	// Try a different host URL after the first retry
@@ -549,6 +570,8 @@ void FMultiEndpointHttpClient::RetryRequest(FRequest&& Request)
 
 FMultiEndpointHttpClient::FConnection& FMultiEndpointHttpClient::GetConnection(const FOnDemandHostGroup& HostGroup)
 {
+	UE_CHECK_THREADSAFETY(FMultiEndpointHttpClient::GetConnection);
+
 	for (TUniquePtr<FConnection>& Conn : Connections)
 	{
 		check(Conn.IsValid());
@@ -592,3 +615,5 @@ void FMultiEndpointHttpClient::Log(const FRequest& Request)
 }
 
 } // namespace UE::IoStore
+
+#undef UE_CHECK_THREADSAFETY
diff --git a/Engine/Source/Runtime/Experimental/IoStore/OnDemand/Private/OnDemandHttpClient.h b/Engine/Source/Runtime/Experimental/IoStore/OnDemand/Private/OnDemandHttpClient.h
--- a/Engine/Source/Runtime/Experimental/IoStore/OnDemand/Private/OnDemandHttpClient.h
+++ b/Engine/Source/Runtime/Experimental/IoStore/OnDemand/Private/OnDemandHttpClient.h
@@ -90,6 +90,7 @@ struct FMultiEndpointHttpClientConfig
 	uint32			MaxRetryCount = 0;
 	int32			TimeoutMs = 0;
 	EHttpRedirects	Redirects = EHttpRedirects::Follow;
+	bool			bEnableThreadSafetyChecks = false;
 };
 
 struct FMultiEndpointHttpClientResponse
@@ -162,6 +163,10 @@ private:
 	TArray<TUniquePtr<FConnection>>	Connections;
 	HTTP::FEventLoop				EventLoop;
 	TArray<FRequest>				Retries;
+
+#if UE_HTTPCLIENT_THREADSAFE_CHECKS
+	uint32							OwningThread = 0;
+#endif //UE_HTTPCLIENT_THREADSAFE_CHECKS
 };
 
 } // namespace UE::IoStore
