dataflow : add clone method for context cache entries

diff --git a/Engine/Source/Runtime/Experimental/Dataflow/Core/Private/Dataflow/DataflowNodeParameters.cpp b/Engine/Source/Runtime/Experimental/Dataflow/Core/Private/Dataflow/DataflowNodeParameters.cpp
--- a/Engine/Source/Runtime/Experimental/Dataflow/Core/Private/Dataflow/DataflowNodeParameters.cpp
+++ b/Engine/Source/Runtime/Experimental/Dataflow/Core/Private/Dataflow/DataflowNodeParameters.cpp
@@ -79,11 +79,11 @@ namespace UE::Dataflow
 		Context.PopFromCallstack(Connection);
 	}
 
-	void BeginContextEvaluation(FContext& Context, const FDataflowNode* Node, const FDataflowOutput* Output)
+	void FContext::BeginContextEvaluation(const FDataflowNode * Node, const FDataflowOutput * Output)
 	{
 		if (Output)
 		{
-			Context.Evaluate(*Output);
+			Evaluate(*Output);
 		}
 		else if (Node)
 		{
@@ -91,7 +91,7 @@ namespace UE::Dataflow
 			{
 				for (const FDataflowOutput* const NodeOutput : Node->GetOutputs())
 				{
-					Context.Evaluate(*NodeOutput);
+					Evaluate(*NodeOutput);
 				}
 			}
 			// Note: If the node is deactivated and has an output (like above), then the output might still need to be forwarded.
@@ -103,7 +103,7 @@ namespace UE::Dataflow
 				//       The purpose of the node evaluation function is to evaluate outputs.
 				//       Therefore if a node has no outputs, then it shouldn't need any evaluation.
 				UE_LOG(LogChaosDataflow, Verbose, TEXT("FDataflowNode::Evaluate(): Node [%s], Output [nullptr], NodeTimestamp [%lu]"), *Node->GetName().ToString(), Node->GetTimestamp().Value);
-				Node->Evaluate(Context, nullptr);
+				Node->Evaluate(*this, nullptr);
 			}
 		}
 		else
@@ -114,7 +114,7 @@ namespace UE::Dataflow
 
 	void FContextSingle::Evaluate(const FDataflowNode* Node, const FDataflowOutput* Output)
 	{
-		BeginContextEvaluation(*this, Node, Output);
+		BeginContextEvaluation(Node, Output);
 	}
 
 	bool FContextSingle::Evaluate(const FDataflowOutput& Connection)
@@ -127,7 +127,7 @@ namespace UE::Dataflow
 
 	void FContextThreaded::Evaluate(const FDataflowNode* Node, const FDataflowOutput* Output)
 	{
-		BeginContextEvaluation(*this, Node, Output);
+		BeginContextEvaluation(Node, Output);
 	}
 
 	bool FContextThreaded::Evaluate(const FDataflowOutput& Connection)
diff --git a/Engine/Source/Runtime/Experimental/Dataflow/Core/Public/Dataflow/DataflowNodeParameters.h b/Engine/Source/Runtime/Experimental/Dataflow/Core/Public/Dataflow/DataflowNodeParameters.h
--- a/Engine/Source/Runtime/Experimental/Dataflow/Core/Public/Dataflow/DataflowNodeParameters.h
+++ b/Engine/Source/Runtime/Experimental/Dataflow/Core/Public/Dataflow/DataflowNodeParameters.h
@@ -85,6 +85,9 @@ namespace UE::Dataflow
 		// InReferenceDataKey is the key of the cache element this function is called on 
 		virtual TUniquePtr<FContextCacheElementBase> CreateReference(FContextCacheKey InReferenceDataKey) const = 0;
 
+		// clone the cache entry
+		virtual TUniquePtr<FContextCacheElementBase> Clone(const FContext& Context) const = 0;
+
 		template<typename T>
 		inline const T& GetTypedData(const FContext& Context, const FProperty* PropertyIn, const T& Default) const;
 
@@ -122,6 +125,11 @@ namespace UE::Dataflow
 			: FContextCacheElementBase(EType::CacheElementTyped, InNodeGuid, InProperty, InNodeHash, Timestamp)
 			, Data(Forward<T>(InData))
 		{}
+
+		TContextCacheElement(const TContextCacheElement<T>& Other)
+			: FContextCacheElementBase(EType::CacheElementTyped, Other.GetNodeGuid(), Other.GetProperty(), Other.GetNodeHash(), Other.GetTimestamp())
+			, Data(Other.Data)
+		{}
 		
 		inline const T& GetData(const FContext& Context, const FProperty* PropertyIn, const T& Default) const;
 
@@ -134,6 +142,8 @@ namespace UE::Dataflow
 
 		inline virtual TUniquePtr<FContextCacheElementBase> CreateReference(FContextCacheKey InReferenceDataKey) const override;
 
+		inline virtual TUniquePtr<FContextCacheElementBase> Clone(const FContext& Context) const override;
+
 	private:
 		typedef typename TDecay<T>::Type FDataType;  // Using universal references here means T could be either const& or an rvalue reference
 		const FDataType Data;                        // Decaying T removes any reference and gets the correct underlying storage data type
@@ -156,6 +166,8 @@ namespace UE::Dataflow
 
 		inline virtual TUniquePtr<FContextCacheElementBase> CreateReference(FContextCacheKey InReferenceDataKey) const override;
 
+		inline virtual TUniquePtr<FContextCacheElementBase> Clone(const FContext& Context) const override;
+
 	private:
 		const FContextCacheKey DataKey; // this is a key to another cache element
 	};
@@ -175,6 +187,8 @@ namespace UE::Dataflow
 		}
 
 		inline virtual TUniquePtr<FContextCacheElementBase> CreateReference(FContextCacheKey InReferenceDataKey) const override;
+
+		inline virtual TUniquePtr<FContextCacheElementBase> Clone(const FContext& Context) const override;
 	};
 
 	//--------------------------------------------------------------------
@@ -188,6 +202,11 @@ namespace UE::Dataflow
 			, ObjectPtr(InObjectPtr)
 		{}
 
+		TContextCacheElementUObject(const TContextCacheElementUObject<T>& Other)
+			: FContextCacheElementBase(EType::CacheElementUObject, Other.GetNodeGuid(), Other.GetProperty(), Other.GetNodeHash(), Other.GetTimestamp())
+			, ObjectPtr(Other.ObjectPtr)
+		{}
+
 		const T& GetData(const FContext& /*Context*/, const FProperty* /*PropertyIn*/, const T& /*Default*/) const
 		{
 			return ObjectPtr;
@@ -200,6 +219,8 @@ namespace UE::Dataflow
 
 		inline virtual TUniquePtr<FContextCacheElementBase> CreateReference(FContextCacheKey InReferenceDataKey) const override;
 
+		inline virtual TUniquePtr<FContextCacheElementBase> Clone(const FContext& Context) const override;
+
 		//~ Begin FGCObject interface
 		virtual void AddReferencedObjects(FReferenceCollector& Collector) override { Collector.AddReferencedObject(ObjectPtr); }
 		virtual FString GetReferencerName() const override { return TEXT("TContextCacheElementUObject"); }
@@ -222,6 +243,12 @@ namespace UE::Dataflow
 			, ObjectPtrArray(Forward<ArrayT>(InObjectPtrArray))
 		{}
 
+
+		TContextCacheElementUObjectArray(const TContextCacheElementUObjectArray<T>& Other)
+			: FContextCacheElementBase(EType::CacheElementUObject, Other.GetNodeGuid(), Other.GetProperty(), Other.GetNodeHash(), Other.GetTimestamp())
+			, ObjectPtrArray(Other.ObjectPtrArray)
+		{}
+
 		const TArray<T>& GetData(const FContext& /*Context*/, const FProperty* /*PropertyIn*/, const TArray<T>& /*Default*/) const
 		{
 			return ObjectPtrArray;
@@ -234,6 +261,8 @@ namespace UE::Dataflow
 
 		inline virtual TUniquePtr<FContextCacheElementBase> CreateReference(FContextCacheKey InReferenceDataKey) const override;
 
+		inline virtual TUniquePtr<FContextCacheElementBase> Clone(const FContext& Context) const override;
+
 		//~ Begin FGCObject interface
 		virtual void AddReferencedObjects(FReferenceCollector& Collector) override { Collector.AddReferencedObjects(ObjectPtrArray); }
 		virtual FString GetReferencerName() const override { return TEXT("TContextCacheElementUObjectArray"); }
@@ -337,7 +366,17 @@ namespace UE::Dataflow
 		{
 			if (IsA(T::StaticType()))
 			{
-				return (T*)this;
+				return static_cast<const T*>(this);
+			}
+			return nullptr;
+		}
+
+		template<class T>
+		T* AsType()
+		{
+			if (IsA(T::StaticType()))
+			{
+				return static_cast<T*>(this);
 			}
 			return nullptr;
 		}
@@ -450,6 +489,9 @@ namespace UE::Dataflow
 
 		DATAFLOWCORE_API bool IsCacheEntryAfterTimestamp(FContextCacheKey InKey, const FTimestamp InTimestamp) const;
 
+	protected:
+		DATAFLOWCORE_API void BeginContextEvaluation(const FDataflowNode* Node, const FDataflowOutput* Output);
+
 	private:
 #if DATAFLOW_EDITOR_EVALUATION
 		TArray<const FDataflowConnection*> Callstack;
@@ -576,6 +618,12 @@ namespace UE::Dataflow
 		return MakeUnique<TContextCacheElementReference<T>>(GetNodeGuid(), GetProperty(), InReferenceDataKey, GetNodeHash(), GetTimestamp());
 	}
 
+	template<class T>
+	TUniquePtr<FContextCacheElementBase> TContextCacheElement<T>::Clone(const FContext& Context) const
+	{
+		return MakeUnique<TContextCacheElement<T>>(*this);
+	}
+
 	template<class T>
 	const T& TContextCacheElementReference<T>::GetData(const FContext& Context, const FProperty* PropertyIn, const T& Default) const
 	{
@@ -594,23 +642,54 @@ namespace UE::Dataflow
 		return MakeUnique<TContextCacheElementReference<T>>(GetNodeGuid(), GetProperty(), InReferenceDataKey, GetNodeHash(), GetTimestamp());
 	}
 
+	template<class T>
+	TUniquePtr<FContextCacheElementBase> TContextCacheElementReference<T>::Clone(const FContext& Context) const
+	{
+		if (const TUniquePtr<FContextCacheElementBase>* ReferencedCacheEntry = Context.GetDataImpl(DataKey))
+		{
+			if (ReferencedCacheEntry && ReferencedCacheEntry->IsValid())
+			{
+				return (*ReferencedCacheEntry)->Clone(Context);
+			}
+		}
+		return MakeUnique<FContextCacheElementNull>(GetNodeGuid(), GetProperty(), FContextCacheKey(), GetNodeHash(), GetTimestamp());
+	}
+
 	inline TUniquePtr<FContextCacheElementBase> FContextCacheElementNull::CreateReference(FContextCacheKey InReferenceDataKey) const
 	{
 		// a null entry will always return a null entry as a reference
 		return MakeUnique<FContextCacheElementNull>(GetNodeGuid(), GetProperty(), InReferenceDataKey, GetNodeHash(), GetTimestamp());
 	}
 
+
+	TUniquePtr<FContextCacheElementBase> FContextCacheElementNull::Clone(const FContext& Context) const
+	{
+		return MakeUnique<FContextCacheElementNull>(GetNodeGuid(), GetProperty(), FContextCacheKey(), GetNodeHash(), GetTimestamp());
+	}
+
 	template<class T>
 	TUniquePtr<FContextCacheElementBase> TContextCacheElementUObject<T>::CreateReference(FContextCacheKey InReferenceDataKey) const
 	{
 		return MakeUnique<TContextCacheElementReference<T>>(GetNodeGuid(), GetProperty(), InReferenceDataKey, GetNodeHash(), GetTimestamp());
 	}
 
+	template<class T>
+	TUniquePtr<FContextCacheElementBase> TContextCacheElementUObject<T>::Clone(const FContext& Context) const
+	{
+		return MakeUnique<TContextCacheElementUObject<T>>(*this);
+	}
+
 	template<class T>
 	TUniquePtr<FContextCacheElementBase> TContextCacheElementUObjectArray<T>::CreateReference(FContextCacheKey InReferenceDataKey) const
 	{
 		return MakeUnique<TContextCacheElementReference<T>>(GetNodeGuid(), GetProperty(), InReferenceDataKey, GetNodeHash(), GetTimestamp());
 	}
+
+	template<class T>
+	TUniquePtr<FContextCacheElementBase> TContextCacheElementUObjectArray<T>::Clone(const FContext& Context) const
+	{
+		return MakeUnique<TContextCacheElementUObjectArray<T>>(*this);
+	}
 }
 
 
