Dataflow : fix ordering when reconciliating pins from the dataflow node in dataflow graph nodes
- also add a way to reconcile the pin connections

diff --git a/Engine/Source/Runtime/Experimental/Dataflow/Engine/Private/Dataflow/DataflowEdNode.cpp b/Engine/Source/Runtime/Experimental/Dataflow/Engine/Private/Dataflow/DataflowEdNode.cpp
--- a/Engine/Source/Runtime/Experimental/Dataflow/Engine/Private/Dataflow/DataflowEdNode.cpp
+++ b/Engine/Source/Runtime/Experimental/Dataflow/Engine/Private/Dataflow/DataflowEdNode.cpp
@@ -228,6 +228,128 @@ void UDataflowEdNode::UpdatePinsFromDataflowNode()
 					}
 					EdPin->bHidden = Pin.bHidden;
 				}
+
+				// reorder the pins 
+				TArray<UEdGraphPin*> OrderedPins;
+				OrderedPins.Reserve(GetAllPins().Num());
+				for (const FDataflowInput* DataflowInput : DataflowNode->GetInputs())
+				{
+					if (DataflowInput)
+					{
+						if (UEdGraphPin* EdPin = FindPin(DataflowInput->GetName(), EEdGraphPinDirection::EGPD_Input))
+						{
+							OrderedPins.Add(EdPin);
+						}
+					}
+				}
+				for (const FDataflowOutput* DataflowOutput : DataflowNode->GetOutputs())
+				{
+					if (DataflowOutput)
+					{
+						if (UEdGraphPin* EdPin = FindPin(DataflowOutput->GetName(), EEdGraphPinDirection::EGPD_Output))
+						{
+							OrderedPins.Add(EdPin);
+						}
+					}
+				}
+				Pins = MoveTemp(OrderedPins);
+			}
+		}
+	}
+#endif // WITH_EDITOR && !UE_BUILD_SHIPPING
+}
+
+void UDataflowEdNode::UpdatePinsConnectionsFromDataflowNode()
+{
+	UE_LOG(DATAFLOWNODE_LOG, Verbose, TEXT("UDataflowEdNode::UpdatePinsConnectionsFromDataflowNode()"));
+	// called on node creation from UI. 
+
+#if WITH_EDITOR && !(UE_BUILD_SHIPPING || UE_BUILD_TEST)
+	if (DataflowGraph && DataflowNodeGuid.IsValid())
+	{
+		if (TSharedPtr<FDataflowNode> DataflowNode = DataflowGraph->FindBaseNode(DataflowNodeGuid))
+		{
+			// remove pins that do not match inputs / outputs anymore
+			TArray<UEdGraphPin*> PinsToRemove;
+			for (UEdGraphPin* Pin : GetAllPins())
+			{
+				if (ensure(Pin))
+				{
+					if (Pin->Direction == EEdGraphPinDirection::EGPD_Input)
+					{
+						if (const FDataflowInput* DataflowInput = DataflowNode->FindInput(Pin->GetFName()))
+						{
+							// first break all connection 
+							Pin->BreakAllPinLinks();
+
+							// then regenerate them 
+							if (const FDataflowOutput* ConnectedDataflowOutput = DataflowInput->GetConnection())
+							{
+								const FGuid OutputOwnerGuid = ConnectedDataflowOutput->GetOwningNodeGuid();
+								auto SearchByNodeGuidPredicate = [&OutputOwnerGuid](UEdGraphNode* EdNode)
+									{
+										if (UDataflowEdNode* DataflowEdnode = Cast<UDataflowEdNode>(EdNode))
+										{
+											return (DataflowEdnode->DataflowNodeGuid == OutputOwnerGuid);
+										}
+										return false;
+									};
+
+								if (TObjectPtr<UEdGraphNode>* EdNodeToConnect = GetGraph()->Nodes.FindByPredicate(SearchByNodeGuidPredicate))
+								{
+									if (UDataflowEdNode* EdDataflowNodeToConnect = Cast<UDataflowEdNode>(*EdNodeToConnect))
+									{
+										if (UEdGraphPin* PinToConnect = EdDataflowNodeToConnect->FindPin(ConnectedDataflowOutput->GetName(), EEdGraphPinDirection::EGPD_Output))
+										{
+											PinToConnect->MakeLinkTo(Pin);
+											EdDataflowNodeToConnect->UpdatePinsFromDataflowNode();
+											GetGraph()->NotifyNodeChanged(EdDataflowNodeToConnect);
+										}
+									}
+								}
+							}
+						}
+					}
+					else if (Pin->Direction == EEdGraphPinDirection::EGPD_Output)
+					{
+						if (const FDataflowOutput* DataflowOutput = DataflowNode->FindOutput(Pin->GetFName()))
+						{
+							// first break all connection 
+							Pin->BreakAllPinLinks();
+
+							// then regenerate them 
+							const TArray<const FDataflowInput*> ConnectedDataflowInputs = DataflowOutput->GetConnectedInputs();
+							for (const FDataflowInput* ConnectedDataflowInput : ConnectedDataflowInputs)
+							{
+								if (ConnectedDataflowInput)
+								{
+									const FGuid InputOwnerGuid = ConnectedDataflowInput->GetOwningNodeGuid();
+									auto SearchByNodeGuidPredicate = [&InputOwnerGuid](UEdGraphNode* EdNode)
+										{
+											if (UDataflowEdNode* DataflowEdnode = Cast<UDataflowEdNode>(EdNode))
+											{
+												return (DataflowEdnode->DataflowNodeGuid == InputOwnerGuid);
+											}
+											return false;
+										};
+
+									if (TObjectPtr<UEdGraphNode>* EdNodeToConnect = GetGraph()->Nodes.FindByPredicate(SearchByNodeGuidPredicate))
+									{
+										if (UDataflowEdNode* EdDataflowNodeToConnect = Cast<UDataflowEdNode>(*EdNodeToConnect))
+										{
+											if (UEdGraphPin* PinToConnect = EdDataflowNodeToConnect->FindPin(ConnectedDataflowInput->GetName(), EEdGraphPinDirection::EGPD_Input))
+											{
+												Pin->MakeLinkTo(PinToConnect);
+												EdDataflowNodeToConnect->UpdatePinsFromDataflowNode();
+												GetGraph()->NotifyNodeChanged(EdDataflowNodeToConnect);
+											}
+										}
+									}
+								}
+							}
+						}
+					}
+				}
 			}
 		}
 	}
diff --git a/Engine/Source/Runtime/Experimental/Dataflow/Engine/Public/Dataflow/DataflowEdNode.h b/Engine/Source/Runtime/Experimental/Dataflow/Engine/Public/Dataflow/DataflowEdNode.h
--- a/Engine/Source/Runtime/Experimental/Dataflow/Engine/Public/Dataflow/DataflowEdNode.h
+++ b/Engine/Source/Runtime/Experimental/Dataflow/Engine/Public/Dataflow/DataflowEdNode.h
@@ -64,6 +64,7 @@ public:
 	void SetDataflowGraph(TSharedPtr<UE::Dataflow::FGraph> InDataflowGraph) { DataflowGraph = InDataflowGraph; }
 
 	DATAFLOWENGINE_API void UpdatePinsFromDataflowNode();
+	DATAFLOWENGINE_API void UpdatePinsConnectionsFromDataflowNode();
 
 	FGuid GetDataflowNodeGuid() const { return DataflowNodeGuid; }
 	void SetDataflowNodeGuid(FGuid InGuid) { DataflowNodeGuid = InGuid; }
