Dataflow : add method to connect two connection that can be out of order
the method will figure out which one is the output and input to connect them properly

diff --git a/Engine/Source/Runtime/Experimental/Dataflow/Core/Private/Dataflow/DataflowGraph.cpp b/Engine/Source/Runtime/Experimental/Dataflow/Core/Private/Dataflow/DataflowGraph.cpp
--- a/Engine/Source/Runtime/Experimental/Dataflow/Core/Private/Dataflow/DataflowGraph.cpp
+++ b/Engine/Source/Runtime/Experimental/Dataflow/Core/Private/Dataflow/DataflowGraph.cpp
@@ -208,6 +208,31 @@ namespace UE::Dataflow
 
 	}
 
+	bool FGraph::Connect(FDataflowConnection* ConnectionA, FDataflowConnection* ConnectionB)
+	{
+		using namespace UE::Dataflow;
+		if (ConnectionA && ConnectionB)
+		{
+			if (ConnectionA->GetDirection() == FPin::EDirection::OUTPUT
+				&& ConnectionB->GetDirection() == FPin::EDirection::INPUT)
+			{
+				FDataflowOutput* Output = static_cast<FDataflowOutput*>(ConnectionA);
+				FDataflowInput* Input = static_cast<FDataflowInput*>(ConnectionB);
+				Connect(Output, Input);
+				return true;
+			}
+			if (ConnectionA->GetDirection() == FPin::EDirection::INPUT
+				&& ConnectionB->GetDirection() == FPin::EDirection::OUTPUT)
+			{
+				FDataflowOutput* Output = static_cast<FDataflowOutput*>(ConnectionB);
+				FDataflowInput* Input = static_cast<FDataflowInput*>(ConnectionA);
+				Connect(Output, Input);
+				return true;
+			}
+		}
+		return false;
+	}
+
 	void FGraph::Connect(FDataflowOutput* OutputConnection, FDataflowInput* InputConnection)
 	{
 		if (ensure(OutputConnection && InputConnection))
diff --git a/Engine/Source/Runtime/Experimental/Dataflow/Core/Public/Dataflow/DataflowGraph.h b/Engine/Source/Runtime/Experimental/Dataflow/Core/Public/Dataflow/DataflowGraph.h
--- a/Engine/Source/Runtime/Experimental/Dataflow/Core/Public/Dataflow/DataflowGraph.h
+++ b/Engine/Source/Runtime/Experimental/Dataflow/Core/Public/Dataflow/DataflowGraph.h
@@ -184,6 +184,7 @@ namespace UE::Dataflow
 		DATAFLOWCORE_API EConnectType GetConnectType(const FDataflowOutput& Output, const FDataflowInput& Input) const;
 		DATAFLOWCORE_API bool Connect(FDataflowOutput& Output, FDataflowInput& Input);
 
+		DATAFLOWCORE_API bool Connect(FDataflowConnection* ConnectionA, FDataflowConnection* ConnectionB);
 		DATAFLOWCORE_API void Connect(FDataflowOutput* OutputConnection, FDataflowInput* InputConnection);
 		DATAFLOWCORE_API void Disconnect(FDataflowOutput* OutputConnection, FDataflowInput* InputConnection);
 
