dataflow : Remove unecessary local variable array in recent change

diff --git a/Engine/Source/Runtime/Experimental/Dataflow/Engine/Private/Dataflow/DataflowEdNode.cpp b/Engine/Source/Runtime/Experimental/Dataflow/Engine/Private/Dataflow/DataflowEdNode.cpp
--- a/Engine/Source/Runtime/Experimental/Dataflow/Engine/Private/Dataflow/DataflowEdNode.cpp
+++ b/Engine/Source/Runtime/Experimental/Dataflow/Engine/Private/Dataflow/DataflowEdNode.cpp
@@ -269,8 +269,6 @@ void UDataflowEdNode::UpdatePinsConnectionsFromDataflowNode()
 	{
 		if (TSharedPtr<FDataflowNode> DataflowNode = DataflowGraph->FindBaseNode(DataflowNodeGuid))
 		{
-			// remove pins that do not match inputs / outputs anymore
-			TArray<UEdGraphPin*> PinsToRemove;
 			for (UEdGraphPin* Pin : GetAllPins())
 			{
 				if (ensure(Pin))
