Chaosflesh: -supported restart async simulation
-supported reading cached simulation when running async simulation
-fixed simulation length: now including start and end time.
-supported reading/writing per-vertex muscle activation in ChaosCacheCollection
-added color attribute in dynamic collection to reflect muscle activation level
-updated muscle activation computation in deformable solver
-added muscle activation group clean-up after RemoveVertices
-supported multiple LOD sections for surface/tetrahedral embedding

diff --git a/Engine/Plugins/Experimental/ChaosCaching/Source/ChaosCaching/Private/Chaos/CacheManagerActor.cpp b/Engine/Plugins/Experimental/ChaosCaching/Source/ChaosCaching/Private/Chaos/CacheManagerActor.cpp
--- a/Engine/Plugins/Experimental/ChaosCaching/Source/ChaosCaching/Private/Chaos/CacheManagerActor.cpp
+++ b/Engine/Plugins/Experimental/ChaosCaching/Source/ChaosCaching/Private/Chaos/CacheManagerActor.cpp
@@ -429,6 +429,11 @@ void AChaosCacheManager::BeginEvaluate()
 		bool                    bRequiresRecord = false;
 		FComponentCacheAdapter* CurrAdapter     = ActiveAdapters[Index];
 		check(CurrAdapter);    // should definitely have added one above
+		if (CurrAdapter->SupportsRestartSimulation())
+		{
+			CurrAdapter->SetRestartSimulation(bRestartSimulation);
+			CurrAdapter->SetRestartTimeRange(RestartTimeStart, RestartTimeEnd);
+		}
 
 		// We only need to register callbacks if in play/record mode
 		if(CacheMode != ECacheMode::None)
@@ -682,8 +687,13 @@ void AChaosCacheManager::HandlePostSolve(Chaos::FReal InDt, Chaos::FPhysicsSolve
 			if (Observed.TimeSinceTrigger > InCache->GetDuration())
 			{
 				FPendingFrameWrite NewFrame;
-				NewFrame.Time = Observed.TimeSinceTrigger;
-				InAdapter->Record_PostSolve(Comp, Observed.TickRecord.GetSpaceTransform(), NewFrame, Observed.TimeSinceTrigger);
+				float ActualTime = Observed.TimeSinceTrigger;
+				if (bRestartSimulation)
+				{
+					ActualTime += RestartTimeStart;
+				}
+				NewFrame.Time = ActualTime;
+				InAdapter->Record_PostSolve(Comp, Observed.TickRecord.GetSpaceTransform(), NewFrame, ActualTime);
 
 				InCache->AddFrame_Concurrent(MoveTemp(NewFrame));
 			}
@@ -730,7 +740,7 @@ void AChaosCacheManager::OnStartFrameChanged(Chaos::FReal InTime)
 			}
 
 			FCacheUserToken Token = Observed.Cache->BeginPlayback();
-			if (Token.IsOpen())
+			if (Token.IsOpen() || Observed.BestFitAdapter->SupportsPlayWhileRecording())
 			{
 				if (Observed.BestFitAdapter->ValidForPlayback(Comp, Observed.Cache))
 				{
@@ -739,7 +749,54 @@ void AChaosCacheManager::OnStartFrameChanged(Chaos::FReal InTime)
 					Observed.BestFitAdapter->InitializeForLoad(Comp, Observed);
 					Observed.BestFitAdapter->SetRestState(Comp, Observed.Cache, GetTransform(), InTime);
 				}
-				Observed.Cache->EndPlayback(Token);
+				if (Token.IsOpen())
+				{
+					Observed.Cache->EndPlayback(Token);
+				}
+			}
+		}
+	}
+}
+
+void AChaosCacheManager::ReadRestartData()
+{
+	if (!CacheCollection)
+	{
+		return;
+	}
+	for (int32 ObservedIdx = 0; ObservedIdx < ObservedComponents.Num(); ++ObservedIdx)
+	{
+		FObservedComponent& Observed = ObservedComponents[ObservedIdx];
+		if (UPrimitiveComponent* Comp = Observed.GetComponent(this))
+		{
+			if (!Observed.BestFitAdapter)
+			{
+				// Cache the best fit adapater so that we don't need to look it up each time we evaluate a start time.
+				Observed.BestFitAdapter = Chaos::FAdapterUtil::GetBestAdapterForClass(Comp->GetClass(), false);
+			}
+
+			if (!Observed.Cache)
+			{
+				Observed.Cache = CacheCollection->FindCache(Observed.CacheName);
+			}
+
+			if (!Observed.Cache || !Observed.BestFitAdapter)
+			{
+				continue;
+			}
+
+			FCacheUserToken Token = Observed.Cache->BeginPlayback();
+			if (Token.IsOpen() || Observed.BestFitAdapter->SupportsPlayWhileRecording())
+			{
+				if (Observed.BestFitAdapter->ValidForPlayback(Comp, Observed.Cache))
+				{
+					Observed.BestFitAdapter->InitializeForLoad(Comp, Observed);
+					Observed.BestFitAdapter->SetRestState(Comp, Observed.Cache, GetTransform(), RestartTimeStart);
+				}
+				if (Token.IsOpen())
+				{
+					Observed.Cache->EndPlayback(Token);
+				}
 			}
 		}
 	}
diff --git a/Engine/Plugins/Experimental/ChaosCaching/Source/ChaosCaching/Public/Chaos/Adapters/CacheAdapter.h b/Engine/Plugins/Experimental/ChaosCaching/Source/ChaosCaching/Public/Chaos/Adapters/CacheAdapter.h
--- a/Engine/Plugins/Experimental/ChaosCaching/Source/ChaosCaching/Public/Chaos/Adapters/CacheAdapter.h
+++ b/Engine/Plugins/Experimental/ChaosCaching/Source/ChaosCaching/Public/Chaos/Adapters/CacheAdapter.h
@@ -192,6 +192,16 @@ namespace Chaos
 		 */
 		virtual bool ValidForPlayback(UPrimitiveComponent* InComponent, UChaosCache* InCache) const = 0;
 
+		/**
+		 * if the cache adapter supports playback when recording
+		 */
+		virtual bool SupportsPlayWhileRecording() const { return false;  };
+		virtual bool SupportsRestartSimulation() const { return false; };
+		void SetRestartSimulation(bool InbRestartSimulation) { bRestartSimulation = InbRestartSimulation; };
+		void SetRestartTimeRange(float InRestartTimeStart, float InRestartTimeEnd) {
+			RestartTimeStart = InRestartTimeStart; RestartTimeEnd = InRestartTimeEnd;
+		};
+
 		/**
 		* Stall on any solver tasks that could broadcast events via the EventsSolver. Called from game thread.
 		* This is necessary so that the cache manager can teardown and setup new event handling.
@@ -199,6 +209,10 @@ namespace Chaos
 		virtual void WaitForSolverTasks(UPrimitiveComponent* InComponent) const {};
 
 	protected:
+		bool bRestartSimulation = false;
+		float RestartTimeStart = TNumericLimits<double>::Max(); // if bRestartSimulation, use RestartTimeStart as an offset when saving, and clear cache data in [RestartTimeStart, RestartTimeEnd] 
+		float RestartTimeEnd = -TNumericLimits<double>::Max();
+
 	private:
 	};
 
diff --git a/Engine/Plugins/Experimental/ChaosCaching/Source/ChaosCaching/Public/Chaos/CacheManagerActor.h b/Engine/Plugins/Experimental/ChaosCaching/Source/ChaosCaching/Public/Chaos/CacheManagerActor.h
--- a/Engine/Plugins/Experimental/ChaosCaching/Source/ChaosCaching/Public/Chaos/CacheManagerActor.h
+++ b/Engine/Plugins/Experimental/ChaosCaching/Source/ChaosCaching/Public/Chaos/CacheManagerActor.h
@@ -286,8 +286,13 @@ public:
 
 	CHAOSCACHING_API void SetEditorIconVisibility(bool bVisibility);
 #endif
-	
-
+	void SetRestartSimulation(bool InbRestartSimulation) { bRestartSimulation = InbRestartSimulation; };
+	void SetRestartTimeRange(float InRestartTimeStart, float InRestartTimeEnd) {
+		RestartTimeStart = InRestartTimeStart; RestartTimeEnd = InRestartTimeEnd;
+	};
+	void SetIsSimulating(bool InbIsSimulating) { bIsSimulating = InbIsSimulating; };
+	/** sets state for all observed components at the restart time. */
+	CHAOSCACHING_API void ReadRestartData();
 
 protected:
 
@@ -394,6 +399,10 @@ private:
 	// since the object persists
 	float StartTimeAtBeginPlay;
 
+	bool bRestartSimulation = false; // if restart for a specified time range
+	float RestartTimeStart = TNumericLimits<double>::Max(); // if bRestartSimulation, use RestartTimeStart as an offset when saving, and clear cache data in [RestartTimeStart, RestartTimeEnd] 
+	float RestartTimeEnd = -TNumericLimits<double>::Max();
+
 #if WITH_EDITORONLY_DATA
 	UPROPERTY(Transient)
 	TObjectPtr<UBillboardComponent> EditorIconComponent;
diff --git a/Engine/Plugins/Experimental/ChaosCachingUSD/Source/ChaosCachingUSD/Private/ChaosCachingUSD/Operations.cpp b/Engine/Plugins/Experimental/ChaosCachingUSD/Source/ChaosCachingUSD/Private/ChaosCachingUSD/Operations.cpp
--- a/Engine/Plugins/Experimental/ChaosCachingUSD/Source/ChaosCachingUSD/Private/ChaosCachingUSD/Operations.cpp
+++ b/Engine/Plugins/Experimental/ChaosCachingUSD/Source/ChaosCachingUSD/Private/ChaosCachingUSD/Operations.cpp
@@ -630,16 +630,16 @@ bool UE::ChaosCachingUSD::WritePoints(
 	const double Time, 
 	const TArray<Chaos::TVector<float, 3>>& Points, 
 	const TArray<Chaos::TVector<float, 3>>& Vels,
-	const FIntVector2& PointsRange)
+	const Chaos::FRange& PointsRange)
 {
 	FScopedUsdAllocs UsdAllocs; // Use USD memory allocator
 
-	const size_t NumPoints = PointsRange[1];
+	const size_t NumPoints = PointsRange.Count;
 
 	pxr::VtArray<pxr::GfVec3f> VtPoints(NumPoints);
 	pxr::VtArray<pxr::GfVec3f> VtVels(NumPoints);
 
-	int32 GlobalIndex = PointsRange[0];
+	int32 GlobalIndex = PointsRange.Start;
 	for (int32 LocalIndex = 0; LocalIndex < NumPoints; ++LocalIndex, ++GlobalIndex)
 	{
 		VtPoints[LocalIndex].Set(Points[GlobalIndex][0], Points[GlobalIndex][1], Points[GlobalIndex][2]);
@@ -649,6 +649,50 @@ bool UE::ChaosCachingUSD::WritePoints(
 	return WritePoints(Stage, PrimPath, Time, VtPoints, VtVels);
 }
 
+bool UE::ChaosCachingUSD::WriteActivations(
+	UE::FUsdStage& Stage,
+	const FString& PrimPath,
+	const double Time,
+	pxr::VtArray<float>& VtActivations)
+{
+	FScopedUsdAllocs UsdAllocs; // Use USD memory allocator
+
+	pxr::GfMatrix4d GfMat(1.0);
+
+	UE::FSdfPath Path(*PrimPath);
+	pxr::UsdPrim Prim = Stage.GetPrimAtPath(Path);
+	if (!Prim)
+	{
+		// We're authoring a prim that doesn't exist.  Use an override prim.
+		Prim = Stage.OverridePrim(Path);
+	}
+	return UE::ChaosCachingUSDUtil::Private::SetActivations(
+			Prim, VtActivations,
+			Time == -TNumericLimits<double>::Max() ? pxr::UsdTimeCode::Default() : pxr::UsdTimeCode(Time));
+}
+
+bool UE::ChaosCachingUSD::WriteActivations(
+	UE::FUsdStage& Stage,
+	const FString& PrimPath,
+	const double Time,
+	const TArray<float>& Activations,
+	const Chaos::FRange& PointsRange)
+{
+	FScopedUsdAllocs UsdAllocs; // Use USD memory allocator
+
+	const size_t NumPoints = PointsRange.Count;
+
+	pxr::VtArray<float> VtActivations(NumPoints);
+
+	int32 GlobalIndex = PointsRange.Start;
+	for (int32 LocalIndex = 0; LocalIndex < NumPoints; ++LocalIndex, ++GlobalIndex)
+	{
+		VtActivations[LocalIndex] = Activations[GlobalIndex];
+	}
+
+	return WriteActivations(Stage, PrimPath, Time, VtActivations);
+}
+
 bool
 UE::ChaosCachingUSD::ReadTimeSamples(
 	const UE::FUsdStage& Stage, 
@@ -773,11 +817,17 @@ UE::ChaosCachingUSD::ReadPoints(
 	}
 
 	// VtArray is copy-on-write. Be careful to use const API.
-	static_cast<pxr::UsdAttribute&>(Attr).Get(
-		&VtPoints, 
-		Time == -TNumericLimits<double>::Max() ? pxr::UsdTimeCode::Default() : pxr::UsdTimeCode(Time));
-
-	return true;
+	if (Attr)
+	{
+		static_cast<pxr::UsdAttribute&>(Attr).Get(
+			&VtPoints,
+			Time == -TNumericLimits<double>::Max() ? pxr::UsdTimeCode::Default() : pxr::UsdTimeCode(Time));
+		return true;
+	}
+	else
+	{
+		return false;
+	}
 }
 
 bool 
@@ -815,8 +865,66 @@ UE::ChaosCachingUSD::ReadPoints(
 	pxr::VtArray<pxr::GfVec3f>& VtPoints,
 	pxr::VtArray<pxr::GfVec3f>& VtVels)
 {
+	//Velocity is optional, e.g. for quasistatics
 	ChaosCachingUSD::ReadPoints(Stage, PrimPath, FString(pxr::UsdGeomTokens->velocities.GetText()), Time, VtVels);
 	return ChaosCachingUSD::ReadPoints(Stage, PrimPath, FString(pxr::UsdGeomTokens->points.GetText()), Time, VtPoints);
 }
 
+bool
+UE::ChaosCachingUSD::ReadMuscleActivation(
+	const UE::FUsdStage& Stage,
+	const FString& PrimPath,
+	const FString& AttrName,
+	const double Time,
+	pxr::VtArray<float>& VtActivations)
+{
+	FScopedUsdAllocs UsdAllocs; // Use USD memory allocator
+
+	UE::FSdfPath Path(*PrimPath);
+	UE::FUsdPrim Prim = Stage.GetPrimAtPath(Path);
+	if (!Prim)
+	{
+		UE_LOG(LogUsd, Error, TEXT("No prim found at path '%s'."), *PrimPath);
+		return false;
+	}
+	UE::FUsdAttribute Attr = Prim.GetAttribute(*AttrName);
+	if (!Attr)
+	{
+		return false;
+	}
+
+	// VtArray is copy-on-write. Be careful to use const API.
+	static_cast<pxr::UsdAttribute&>(Attr).Get(
+		&VtActivations,
+		Time == -TNumericLimits<double>::Max() ? pxr::UsdTimeCode::Default() : pxr::UsdTimeCode(Time));
+
+	return true;
+}
+
+bool
+UE::ChaosCachingUSD::ReadMuscleActivation(
+	const UE::FUsdStage& Stage,
+	const FString& PrimPath,
+	const double Time,
+	pxr::VtArray<float>& VtActivations)
+{
+	return ChaosCachingUSD::ReadMuscleActivation(Stage, PrimPath, FString("activation"), Time, VtActivations);
+}
+
+bool
+UE::ChaosCachingUSD::DeletePointsInTimeRange(UE::FUsdStage& Stage, const FString& PrimPath, const double StartTime, const double EndTime)
+{
+	FScopedUsdAllocs UsdAllocs; // Use USD memory allocator
+
+	UE::FSdfPath Path(*PrimPath);
+	pxr::UsdPrim Prim = Stage.GetPrimAtPath(Path);
+	if (!Prim)
+	{
+		// We're authoring a prim that doesn't exist. Use an override prim.
+		Prim = Stage.OverridePrim(Path);
+		UE::ChaosCachingUSDUtil::Private::DefineAncestorTransforms(Stage, Path);
+	}
+	return UE::ChaosCachingUSDUtil::Private::DeletePointsInTimeRange(
+			Prim, StartTime, EndTime);
+}
 #endif // USE_USD_SDK
diff --git a/Engine/Plugins/Experimental/ChaosCachingUSD/Source/ChaosCachingUSD/Private/ChaosCachingUSD/Util.cpp b/Engine/Plugins/Experimental/ChaosCachingUSD/Source/ChaosCachingUSD/Private/ChaosCachingUSD/Util.cpp
--- a/Engine/Plugins/Experimental/ChaosCachingUSD/Source/ChaosCachingUSD/Private/ChaosCachingUSD/Util.cpp
+++ b/Engine/Plugins/Experimental/ChaosCachingUSD/Source/ChaosCachingUSD/Private/ChaosCachingUSD/Util.cpp
@@ -97,4 +97,64 @@ UE::ChaosCachingUSDUtil::Private::SetPointsExtentAndPrimXform(
 	return UE::ChaosCachingUSDUtil::Private::SetPointsExtentAndPrimXform(Prim, PrimXf, Points, Vels, Time);
 }
 
+bool
+UE::ChaosCachingUSDUtil::Private::SetActivations(
+	pxr::UsdPrim& Prim,
+	const pxr::VtArray<float>& Activations,
+	const pxr::UsdTimeCode& Time)
+{
+	FScopedUsdAllocs UEAllocs; // Use USD memory allocator
+
+	if (auto ActivationAttr = Prim.CreateAttribute(
+		pxr::TfToken("activation"), pxr::SdfValueTypeNames->FloatArray))
+	{
+		ActivationAttr.Set(Activations, Time);
+		return true;
+	}
+	return false;
+}
+
+bool
+UE::ChaosCachingUSDUtil::Private::DeletePointsInTimeRange(pxr::UsdPrim& Prim, float StartTime, float EndTime)
+{
+	FScopedUsdAllocs UEAllocs; // Use USD memory allocator
+	
+	if (!Prim)
+	{
+		return false;
+	}
+
+	pxr::UsdGeomPointBased PointBased(Prim);
+
+	DeleteUsdAttributeInTimeRange(PointBased.MakeMatrixXform(), StartTime, EndTime);
+	DeleteUsdAttributeInTimeRange(PointBased.GetExtentAttr(), StartTime, EndTime);
+	DeleteUsdAttributeInTimeRange(PointBased.GetVelocitiesAttr(), StartTime, EndTime);
+	DeleteUsdAttributeInTimeRange(Prim.GetAttribute(pxr::TfToken("activation")), StartTime, EndTime);
+	return DeleteUsdAttributeInTimeRange(PointBased.GetPointsAttr(), StartTime, EndTime);
+}
+
+bool
+UE::ChaosCachingUSDUtil::Private::DeleteUsdAttributeInTimeRange(pxr::UsdAttribute Attribute, float StartTime, float EndTime)
+{
+	if (!Attribute)
+	{
+		return false;
+	}
+	// Get all time samples for the attribute
+	std::vector<double> TimeSamples;
+	if (!Attribute.GetTimeSamples(&TimeSamples))
+	{
+		return false; // No time samples available
+	}
+
+	// Iterate over time samples and remove those within the specified range
+	for (double TimeSample : TimeSamples)
+	{
+		if (TimeSample >= StartTime && TimeSample <= EndTime)
+		{
+			Attribute.ClearAtTime(TimeSample); // Remove the specific time sample
+		}
+	}
+	return true;
+}
 #endif // USE_USD_SDK
diff --git a/Engine/Plugins/Experimental/ChaosCachingUSD/Source/ChaosCachingUSD/Private/ChaosCachingUSD/Util.h b/Engine/Plugins/Experimental/ChaosCachingUSD/Source/ChaosCachingUSD/Private/ChaosCachingUSD/Util.h
--- a/Engine/Plugins/Experimental/ChaosCachingUSD/Source/ChaosCachingUSD/Private/ChaosCachingUSD/Util.h
+++ b/Engine/Plugins/Experimental/ChaosCachingUSD/Source/ChaosCachingUSD/Private/ChaosCachingUSD/Util.h
@@ -95,7 +95,9 @@ namespace UE::ChaosCachingUSDUtil::Private {
 	void ComputeExtent(const pxr::VtArray<pxr::GfVec3f>& Points, pxr::VtVec3fArray& Extent);
 	bool SetPointsExtentAndPrimXform(pxr::UsdPrim& Prim, const pxr::GfMatrix4d& PrimXf, const pxr::VtArray<pxr::GfVec3f>& Points, const pxr::VtArray<pxr::GfVec3f>& Vels, const pxr::UsdTimeCode& Time = pxr::UsdTimeCode::Default());
 	bool SetPointsExtentAndPrimXform(pxr::UsdPrim& Prim, const pxr::GfMatrix4d& PrimXf, const pxr::VtArray<pxr::GfVec3f>& Points, const pxr::UsdTimeCode& Time = pxr::UsdTimeCode::Default());
-
+	bool SetActivations(pxr::UsdPrim& Prim, const pxr::VtArray<float>& Activations, const pxr::UsdTimeCode& Time);
+	bool DeletePointsInTimeRange(pxr::UsdPrim& Prim, float StartTime, float EndTime);
+	bool DeleteUsdAttributeInTimeRange(pxr::UsdAttribute Attribute, float StartTime, float EndTime);
 } // namespace UE::ChaosCachingUSDUtil::Private
 
 #endif // USE_USD_SDK
diff --git a/Engine/Plugins/Experimental/ChaosCachingUSD/Source/ChaosCachingUSD/Public/ChaosCachingUSD/Operations.h b/Engine/Plugins/Experimental/ChaosCachingUSD/Source/ChaosCachingUSD/Public/ChaosCachingUSD/Operations.h
--- a/Engine/Plugins/Experimental/ChaosCachingUSD/Source/ChaosCachingUSD/Public/ChaosCachingUSD/Operations.h
+++ b/Engine/Plugins/Experimental/ChaosCachingUSD/Source/ChaosCachingUSD/Public/ChaosCachingUSD/Operations.h
@@ -3,6 +3,7 @@
 #pragma once
 #if USE_USD_SDK
 
+#include "Chaos/Range.h"
 #include "Chaos/Vector.h"
 #include "GeometryCollection/GeometryCollection.h"
 #include "GeometryCollection/ManagedArray.h"
@@ -16,6 +17,7 @@
 
 namespace UE::ChaosCachingUSD
 {
+
 	bool ValuesDiffer(const pxr::VtArray<pxr::GfVec3f>& A, const pxr::VtArray<pxr::GfVec3f>& B, const float Tolerance = 1.0e-8, const uint64 stride = 1)
 	{
 		if (A.size() != B.size()) return true;
@@ -130,8 +132,10 @@ namespace UE::ChaosCachingUSD
 	CHAOSCACHINGUSD_API bool WritePoints(UE::FUsdStage& Stage, const FString& PrimPath, const double Time, const TArray<Chaos::TVector<float, 3>>& Points, const TArray<Chaos::TVector<float, 3>>& Vels);
 
 	/** Copy \p Points and \p Vels within the specified indices range to \c VtArray (with USD memory allocator), then write to USD stage. */
-	CHAOSCACHINGUSD_API bool WritePoints(UE::FUsdStage& Stage, const FString& PrimPath, const double Time, const TArray<Chaos::TVector<float, 3>>& Points, const TArray<Chaos::TVector<float, 3>>& Vels, const FIntVector2& PointsRange);
+	CHAOSCACHINGUSD_API bool WritePoints(UE::FUsdStage& Stage, const FString& PrimPath, const double Time, const TArray<Chaos::TVector<float, 3>>& Points, const TArray<Chaos::TVector<float, 3>>& Vels, const Chaos::FRange& PointsRange);
 	
+	CHAOSCACHINGUSD_API bool WriteActivations(UE::FUsdStage& Stage, const FString& PrimPath, const double Time, pxr::VtArray<float>& VtActivations);
+	CHAOSCACHINGUSD_API bool WriteActivations(UE::FUsdStage& Stage, const FString& PrimPath, const double Time, const TArray<float>& Activations, const Chaos::FRange& PointsRange);
 
 	/** Get time samples for an attribute. */
 	CHAOSCACHINGUSD_API bool ReadTimeSamples(const UE::FUsdStage& Stage, const FString& PrimPath, const FString& AttrName, TArray<double>& TimeSamples);
@@ -150,5 +154,10 @@ namespace UE::ChaosCachingUSD
 	/** Get points from the points attribute. Default time is used if \p Time is \c -DBL_MAX. */
 	CHAOSCACHINGUSD_API bool ReadPoints(const UE::FUsdStage& Stage, const FString& PrimPath, const double Time, pxr::VtArray<pxr::GfVec3f>& Points, pxr::VtArray<pxr::GfVec3f>& VtVels);
 
+	CHAOSCACHINGUSD_API bool ReadMuscleActivation(const UE::FUsdStage& Stage, const FString& PrimPath, const FString& AttrName, const double Time, pxr::VtArray<float>& VtActivations);
+	CHAOSCACHINGUSD_API bool ReadMuscleActivation(const UE::FUsdStage& Stage, const FString& PrimPath, const double Time, pxr::VtArray<float>& VtActivations);
+	
+	CHAOSCACHINGUSD_API bool DeletePointsInTimeRange(UE::FUsdStage& Stage, const FString& PrimPath, const double StartTime, const double EndTime);
+	
 } // UE::ChaosCachingUSD
 #endif // USE_USD_SDK
diff --git a/Engine/Plugins/Experimental/ChaosFlesh/Source/ChaosFlesh/Private/ChaosFlesh/TetrahedralCollection.cpp b/Engine/Plugins/Experimental/ChaosFlesh/Source/ChaosFlesh/Private/ChaosFlesh/TetrahedralCollection.cpp
--- a/Engine/Plugins/Experimental/ChaosFlesh/Source/ChaosFlesh/Private/ChaosFlesh/TetrahedralCollection.cpp
+++ b/Engine/Plugins/Experimental/ChaosFlesh/Source/ChaosFlesh/Private/ChaosFlesh/TetrahedralCollection.cpp
@@ -1,6 +1,7 @@
 // Copyright Epic Games, Inc. All Rights Reserved.
 #include "ChaosFlesh/TetrahedralCollection.h"
 #include "GeometryCollection/Facades/CollectionPositionTargetFacade.h"
+#include "GeometryCollection/Facades/CollectionMuscleActivationFacade.h"
 
 DEFINE_LOG_CATEGORY_STATIC(FTetrahedralCollectionLogging, Log, All);
 
@@ -366,24 +367,12 @@ void FTetrahedralCollection::RemoveVertices(const TArray<int32>& SortedVertexInd
 
 	// --- Incident Elements (Tetrahedral Dependency) --
 	InitIncidentElements();
-	/*for (int32 Idx = 0; Idx < NumVertices; Idx++)
-	{
-		TArray<int32> NewIncidentElements;
-		TArray<int32> NewIncidentElementsLocal;
-		for (int32 i = 0; i < IncidentElements[Idx].Num(); ++i)
-		{
-			const int32 NewTetIndex = TetrahedralElementsRemap[IncidentElements[Idx][i]];
-			if (NewTetIndex >= 0)
-			{
-				NewIncidentElements.Add(NewTetIndex);
-				NewIncidentElementsLocal.Add(IncidentElementsLocalIndex[Idx][i]);
-			}
-		}
-		IncidentElements[Idx] = NewIncidentElements;
-		IncidentElementsLocalIndex[Idx] = NewIncidentElementsLocal;
-	}*/
 
-	// --- Position Targets (Vertex dependency)--
+	// --- Position Targets (Vertex group dependency)--
 	GeometryCollection::Facades::FPositionTargetFacade PositionTargetFacade(*this);
 	PositionTargetFacade.RemoveInvalidPositionTarget();
+
+	// --- Muscle Activation (Geometry group dependency)--
+	GeometryCollection::Facades::FMuscleActivationFacade MuscleActivationFacade(*this);
+	MuscleActivationFacade.RemoveInvalidMuscles();
 }
\ No newline at end of file
diff --git a/Engine/Plugins/Experimental/ChaosFlesh/Source/ChaosFleshEngine/Private/ChaosCache/FleshComponentCacheAdapter.cpp b/Engine/Plugins/Experimental/ChaosFlesh/Source/ChaosFleshEngine/Private/ChaosCache/FleshComponentCacheAdapter.cpp
--- a/Engine/Plugins/Experimental/ChaosFlesh/Source/ChaosFleshEngine/Private/ChaosCache/FleshComponentCacheAdapter.cpp
+++ b/Engine/Plugins/Experimental/ChaosFlesh/Source/ChaosFleshEngine/Private/ChaosCache/FleshComponentCacheAdapter.cpp
@@ -98,7 +98,7 @@ namespace Chaos
 					FScopedUsdAllocs UsdAllocs; // Use USD memory allocator
 
 					FPrimitiveStage& PrimitiveStage = PrimitiveStages.FindOrAdd(InComponent);
-					const FIntVector2 ParticleRange = GetParticleRange(InComponent, NumParticles);
+					const Chaos::FRange ParticleRange = GetParticleRange(InComponent, NumParticles);
 
 					// Update time range.
 					PrimitiveStage.MinTime = FMath::Min(InTime, PrimitiveStage.MinTime);
@@ -114,6 +114,16 @@ namespace Chaos
 							return;
 						}
 
+						// Write muscle activation
+						const TArray<float> ParticleMuscleActivation = PhysicsThreadAccess.GetParticleMuscleActivation();
+						if (!UE::ChaosCachingUSD::WriteActivations(PrimitiveStage.MonolithStage, PrimitiveStage.PrimPath, InTime, ParticleMuscleActivation, ParticleRange))
+						{
+							UE_LOG(LogChaosFleshCache, Error,
+								TEXT("Failed to write activations '%s' at time %g to file: '%s'"),
+								*PrimitiveStage.PrimPath, InTime, *PrimitiveStage.MonolithStage.GetRootLayer().GetDisplayName());
+							return;
+						}
+
 						uint64 NumTimeSamples = UE::ChaosCachingUSD::GetNumTimeSamples(PrimitiveStage.MonolithStage, PrimitiveStage.PrimPath, UE::ChaosCachingUSD::GetPointsAttrName());
 						if (CVarParams.SaveFrequency >= 1 && NumTimeSamples % CVarParams.SaveFrequency == 0)
 						{
@@ -170,7 +180,7 @@ namespace Chaos
 	}
 
 	void FFleshCacheAdapter::LoadCacheAtTime(UPrimitiveComponent* PrimitiveComponent, Chaos::FReal TargetTime, const int32 NumParticles, const bool bNeedsRange,
-		const TFunction<void(const int32, const int32, const FVector3f&, const FVector3f&)>& LoadFunction) const
+		const TFunction<void(const int32, const int32, const FVector3f&, const FVector3f&, const float Activation)>& LoadFunction) const
 	{
 #if USE_USD_SDK && DO_USD_CACHING
 
@@ -178,12 +188,13 @@ namespace Chaos
 		if (PrimitiveStage.MonolithStage)
 		{
 			// GetParticleRange only called in play/record mode (not load), so PT only
-			const FIntVector2 ParticleRange = bNeedsRange ? GetParticleRange(PrimitiveComponent, NumParticles) : FIntVector2(0,NumParticles);
-			const int32 NumComponentParticles = ParticleRange[1];
+			const Chaos::FRange ParticleRange = bNeedsRange ? GetParticleRange(PrimitiveComponent, NumParticles) : Chaos::FRange(0,NumParticles);
+			const int32 NumComponentParticles = ParticleRange.Count;
 			
 			FScopedUsdAllocs UEAllocs; // Use USD memory allocator
 			pxr::VtArray<pxr::GfVec3f> Points0, Points1;
 			pxr::VtArray<pxr::GfVec3f> Vels0, Vels1;
+			pxr::VtArray<float> Activations0, Activations1;
 
 			double Prev = -TNumericLimits<double>::Max();
 			double Next = -TNumericLimits<double>::Max();
@@ -214,6 +225,13 @@ namespace Chaos
 				return;
 			}
 
+			// load muscle activation
+			if (!UE::ChaosCachingUSD::ReadMuscleActivation(PrimitiveStage.MonolithStage, PrimitiveStage.PrimPath, Prev, Activations0) ||
+				Points0.size() != Activations0.size())
+			{
+				Activations0.assign(Points0.size(), -1.f);
+			}
+
 			int32 NumCachedParticles = static_cast<int32>(Points0.size());
 			if (NumCachedParticles > NumComponentParticles)
 			{
@@ -232,7 +250,7 @@ namespace Chaos
 					const pxr::GfVec3f& Pos = Points0.cdata()[CachedIndex];
 					const pxr::GfVec3f& Vel = Vels0.cdata()[CachedIndex];
 
-					LoadFunction(CachedIndex, ParticleRange[0], FVector3f(Pos[0], Pos[1], Pos[2]), FVector3f(Vel[0], Vel[1], Vel[2]));
+					LoadFunction(CachedIndex, ParticleRange.Start, FVector3f(Pos[0], Pos[1], Pos[2]), FVector3f(Vel[0], Vel[1], Vel[2]), Activations0.cdata()[CachedIndex]);
 				}
 				return;
 			}
@@ -246,6 +264,14 @@ namespace Chaos
 					*PrimitiveStage.PrimPath, Next, *PrimitiveStage.MonolithStage.GetRootLayer().GetDisplayName());
 				return;
 			}
+
+			// load muscle activation
+			if (!UE::ChaosCachingUSD::ReadMuscleActivation(PrimitiveStage.MonolithStage, PrimitiveStage.PrimPath, Prev, Activations1) ||
+				Points1.size() != Activations1.size())
+			{
+				Activations1.assign(Points1.size(), -1.f);
+			}
+
 			double Duration = Next - Prev;
 			double Alpha = Duration > UE_SMALL_NUMBER ? (TargetTime - Prev) / Duration : 0.5;
 			for (int32 CachedIndex = 0; CachedIndex < NumCachedParticles; ++CachedIndex)
@@ -258,8 +284,8 @@ namespace Chaos
 				const pxr::GfVec3f& V1 = Vels1.cdata()[CachedIndex];
 				pxr::GfVec3f Pos = (1.0 - Alpha) * P0 + Alpha * P1;
 				pxr::GfVec3f Vel = (1.0 - Alpha) * V0 + Alpha * V1;
-				
-				LoadFunction(CachedIndex, ParticleRange[0], FVector3f(Pos[0], Pos[1], Pos[2]), FVector3f(Vel[0], Vel[1], Vel[2]));
+				float Activation = (1.0 - Alpha) * Activations0.cdata()[CachedIndex] + Alpha * Activations1.cdata()[CachedIndex];
+				LoadFunction(CachedIndex, ParticleRange.Start, FVector3f(Pos[0], Pos[1], Pos[2]), FVector3f(Vel[0], Vel[1], Vel[2]), Activation);
 			}
 		}
 	#endif
@@ -283,7 +309,7 @@ namespace Chaos
 				Softs::FPAndInvM* ParticlePAndInvMs = &Particles.PAndInvM(0);
 				
 				LoadCacheAtTime(InComponent, InTime, NumParticles, true, [&ParticleXs, &ParticleVs, &ParticleInvMs, &ParticlePAndInvMs](
-					const int32 ParticleIndex, const int32 ParticleOffset, const FVector3f& ParticlePosition, const FVector3f& ParticleVelocity)
+					const int32 ParticleIndex, const int32 ParticleOffset, const FVector3f& ParticlePosition, const FVector3f& ParticleVelocity, const float MuscleActivation)
 				{
 					const int32 GlobalIndex = ParticleIndex + ParticleOffset;
 					
@@ -300,8 +326,8 @@ namespace Chaos
 					{
 						if (const Chaos::Softs::FFleshThreadingProxy* Proxy = PhysicsProxy->As<Chaos::Softs::FFleshThreadingProxy>())
 						{
-							FIntVector2 ParticleRange = Proxy->GetSolverParticleRange();
-							Evolution->ActivateParticleRange(ParticleRange[0], false);
+							Chaos::FRange ParticleRange = Proxy->GetSolverParticleRange();
+							Evolution->ActivateParticleRange(ParticleRange.Start, false);
 						}
 					}
 				}
@@ -390,9 +416,9 @@ namespace Chaos
 		return nullptr;
 	}
 
-	FIntVector2 FFleshCacheAdapter::GetParticleRange(UPrimitiveComponent* InComponent, const int32 NumParticles)
+	Chaos::FRange FFleshCacheAdapter::GetParticleRange(UPrimitiveComponent* InComponent, const int32 NumParticles)
 	{
-		FIntVector2 ParticleRange(0, NumParticles);
+		Chaos::FRange ParticleRange(0, NumParticles);
 		if(InComponent)
 		{
 			if(UFleshComponent* FleshComp = CastChecked<UFleshComponent>(InComponent))
@@ -470,13 +496,14 @@ namespace Chaos
 			if (UFleshDynamicAsset* DynamicCollection = FleshComp->GetDynamicCollection())
 			{
 				TManagedArray<FVector3f>& DynamicVertex = DynamicCollection->GetPositions();
+				TManagedArray<float>& DynamicVertexActivation = DynamicCollection->GetActivations();
 				const int32 NumDynamicVertex = DynamicVertex.Num();
 
 				auto UEVertd = [](Chaos::FVec3 V) { return FVector3d(V.X, V.Y, V.Z); };
 				auto UEVertf = [](FVector3d V) { return FVector3f((float)V.X, (float)V.Y, (float)V.Z); };
-				
-				LoadCacheAtTime(InComponent, InTime, NumDynamicVertex, false, [&DynamicVertex, FleshComp, UEVertd, UEVertf](
-					const int32 ParticleIndex, const int32 ParticleOffset, const FVector3f& ParticlePosition, const FVector3f&)
+
+				LoadCacheAtTime(InComponent, InTime, NumDynamicVertex, false, [&DynamicVertex, &DynamicVertexActivation, FleshComp, UEVertd, UEVertf](
+					const int32 ParticleIndex, const int32 ParticleOffset, const FVector3f& ParticlePosition, const FVector3f&, const float& MuscleActivation)
 				{
 					const int32 GlobalIndex = ParticleIndex + ParticleOffset;
 
@@ -484,6 +511,7 @@ namespace Chaos
 					DynamicVertex[GlobalIndex] =
 									UEVertf(FleshComp->GetComponentTransform().InverseTransformPosition(
 										UEVertd(DynamicVertex[GlobalIndex])));
+					DynamicVertexActivation[GlobalIndex] = MuscleActivation;
 				});
 			}
 
@@ -692,17 +720,27 @@ namespace Chaos
 								FString UniqueCompName = FString::Printf(TEXT("%s_%d.%s"), *CompName, i++, *Ext);
 								UniqueFilePath = FPaths::Combine(CacheDir, UniqueCompName);
 							} while (PlatformFile.FileExists(*UniqueFilePath));
-
-							if (!PlatformFile.MoveFile(*UniqueFilePath, *PrimitiveStage.FilePath))
+							if (bRestartSimulation)
 							{
-								UE_LOG(LogChaosFleshCache, Error, TEXT("Failed to rename file from '%s' to '%s'."), *PrimitiveStage.FilePath, *UniqueFilePath);
-								return false;
+								if (!PlatformFile.CopyFile(*UniqueFilePath, *PrimitiveStage.FilePath))
+								{
+									UE_LOG(LogChaosFleshCache, Error, TEXT("Failed to copy file from '%s' to '%s'."), *PrimitiveStage.FilePath, *UniqueFilePath);
+									return false;
+								}
+							}
+							else
+							{
+								if (!PlatformFile.MoveFile(*UniqueFilePath, *PrimitiveStage.FilePath))
+								{
+									UE_LOG(LogChaosFleshCache, Error, TEXT("Failed to rename file from '%s' to '%s'."), *PrimitiveStage.FilePath, *UniqueFilePath);
+									return false;
+								}
 							}
 						}
 					}
 					else
 					{
-						if (!PlatformFile.DeleteFile(*PrimitiveStage.FilePath))
+						if (!bRestartSimulation && !PlatformFile.DeleteFile(*PrimitiveStage.FilePath))
 						{
 							UE_LOG(LogChaosFleshCache, Error, TEXT("Failed to remove existing cache file: '%s'"), *PrimitiveStage.FilePath);
 							return false;
@@ -713,15 +751,31 @@ namespace Chaos
 					{
 						UE::ChaosCachingUSD::CloseStage(PrimitiveStage.MonolithStage);
 					}
-					if (!UE::ChaosCachingUSD::NewStage(PrimitiveStage.FilePath, PrimitiveStage.MonolithStage))
+					if (bRestartSimulation)
 					{
-						UE_LOG(LogChaosFleshCache, Error, TEXT("Failed to create new USD file: '%s'"), *PrimitiveStage.FilePath);
-						return false;
+						if (!UE::ChaosCachingUSD::OpenStage(PrimitiveStage.FilePath, PrimitiveStage.MonolithStage))
+						{
+							UE_LOG(LogChaosFleshCache, Error, TEXT("Failed to read restart USD file: '%s'"), *PrimitiveStage.FilePath);
+							return false;
+						}
+						if (!UE::ChaosCachingUSD::DeletePointsInTimeRange(PrimitiveStage.MonolithStage, PrimitiveStage.PrimPath, RestartTimeStart, RestartTimeEnd))
+						{
+							UE_LOG(LogChaosFleshCache, Error, TEXT("Failed to delete USD data in file: '%s', delete range: [%f, %f]"), *PrimitiveStage.FilePath, RestartTimeStart, RestartTimeEnd);
+							return false;
+						}
 					}
-					if (!UE::ChaosCachingUSD::WriteTetMesh(PrimitiveStage.MonolithStage, PrimitiveStage.PrimPath, *RestCollection))
+					else
 					{
-						UE_LOG(LogChaosFleshCache, Error, TEXT("Failed to write tetrahedron mesh '%s' to USD file: '%s'"), *PrimitiveStage.PrimPath, *PrimitiveStage.FilePath);
-						return false;
+						if (!UE::ChaosCachingUSD::NewStage(PrimitiveStage.FilePath, PrimitiveStage.MonolithStage))
+						{
+							UE_LOG(LogChaosFleshCache, Error, TEXT("Failed to create new USD file: '%s'"), *PrimitiveStage.FilePath);
+							return false;
+						}
+						if (!UE::ChaosCachingUSD::WriteTetMesh(PrimitiveStage.MonolithStage, PrimitiveStage.PrimPath, *RestCollection))
+						{
+							UE_LOG(LogChaosFleshCache, Error, TEXT("Failed to write tetrahedron mesh '%s' to USD file: '%s'"), *PrimitiveStage.PrimPath, *PrimitiveStage.FilePath);
+							return false;
+						}
 					}
 				}
 
diff --git a/Engine/Plugins/Experimental/ChaosFlesh/Source/ChaosFleshEngine/Private/ChaosFlesh/ChaosDeformableSolverComponent.cpp b/Engine/Plugins/Experimental/ChaosFlesh/Source/ChaosFleshEngine/Private/ChaosFlesh/ChaosDeformableSolverComponent.cpp
--- a/Engine/Plugins/Experimental/ChaosFlesh/Source/ChaosFleshEngine/Private/ChaosFlesh/ChaosDeformableSolverComponent.cpp
+++ b/Engine/Plugins/Experimental/ChaosFlesh/Source/ChaosFleshEngine/Private/ChaosFlesh/ChaosDeformableSolverComponent.cpp
@@ -437,6 +437,32 @@ void UDeformableSolverComponent::ReadFromSimulation(float DeltaTime, const bool
 	}
 }
 
+void UDeformableSolverComponent::ReadRestartData()
+{
+	if (FleshSolverProxy.Solver)
+	{
+		FDeformableSolver::FGameThreadAccess GameThreadSolver = GameThreadAccess();
+
+		Chaos::Softs::FDeformableDataMap DataMap;
+		for (TObjectPtr<UDeformablePhysicsComponent>& DeformableComponent : ConnectedObjects.DeformableComponents)
+		{
+			if (DeformableComponent)
+			{
+				if (IsSimulating(DeformableComponent))
+				{
+					DeformableComponent->PreSolverUpdate();
+					if (FDataMapValue Value = DeformableComponent->WriteRestartDataToSimulation())
+					{
+						DataMap.Add(DeformableComponent, Value);
+					}
+				}
+			}
+		}
+
+		GameThreadSolver.PushRestartPackage(GameThreadSolver.GetFrame(), MoveTemp(DataMap));
+	}
+}
+
 #if WITH_EDITOR
 
 bool UDeformableSolverComponent::CanEditChange(const FProperty* InProperty) const
diff --git a/Engine/Plugins/Experimental/ChaosFlesh/Source/ChaosFleshEngine/Private/ChaosFlesh/ChaosDeformableTetrahedralComponent.cpp b/Engine/Plugins/Experimental/ChaosFlesh/Source/ChaosFleshEngine/Private/ChaosFlesh/ChaosDeformableTetrahedralComponent.cpp
--- a/Engine/Plugins/Experimental/ChaosFlesh/Source/ChaosFleshEngine/Private/ChaosFlesh/ChaosDeformableTetrahedralComponent.cpp
+++ b/Engine/Plugins/Experimental/ChaosFlesh/Source/ChaosFleshEngine/Private/ChaosFlesh/ChaosDeformableTetrahedralComponent.cpp
@@ -320,11 +320,11 @@ FTransform UDeformableTetrahedralComponent::GetSimSpaceRestTransform() const
 	return ComponentToBone;
 }
 
-void UDeformableTetrahedralComponent::UpdateFromSimulation(const FDataMapValue* SimualtionBuffer)
+void UDeformableTetrahedralComponent::UpdateFromSimulation(const FDataMapValue* SimulationBuffer)
 {
 	PERF_SCOPE(STAT_ChaosDeformable_UDeformableTetrahedralComponent_UpdateFromSimualtion);
 
-	if (const FFleshThreadingProxy::FFleshOutputBuffer* FleshBuffer = (*SimualtionBuffer)->As<FFleshThreadingProxy::FFleshOutputBuffer>())
+	if (const FFleshThreadingProxy::FFleshOutputBuffer* FleshBuffer = (*SimulationBuffer)->As<FFleshThreadingProxy::FFleshOutputBuffer>())
 	{
 		if (GetDynamicCollection())
 		{
@@ -370,6 +370,16 @@ void UDeformableTetrahedralComponent::UpdateFromSimulation(const FDataMapValue*
 	}
 }
 
+UDeformablePhysicsComponent::FDataMapValue UDeformableTetrahedralComponent::WriteRestartDataToSimulation()
+{
+	if (!GetDynamicCollection())
+	{
+		ResetDynamicCollection();
+	}
+	return FDataMapValue(
+		new Chaos::Softs::FFleshThreadingProxy::FFleshRestartBuffer(*GetDynamicCollection()->GetCollection(), this));
+}
+
 void UDeformableTetrahedralComponent::UpdateLocalBounds()
 {
 	if (bBoundsNeedsUpdate && RestCollection)
@@ -550,12 +560,33 @@ void UDeformableTetrahedralComponent::RenderProceduralMesh()
 							{
 								TArray<FVector3f> RenderVertex;
 								Flesh.ComponentSpaceVertices(RenderVertex);
-
-								if (GetDynamicCollection())
+								const UFleshDynamicAsset* ConstDynamicCollection = GetDynamicCollection();
+								if (ConstDynamicCollection)
 								{
-									const TManagedArray<FVector3f>& DynamicVertex = GetDynamicCollection()->GetPositions();
-									if ((DynamicVertex.Num() > 0) && (DynamicVertex.Num() == RenderVertex.Num())) 
-										RenderVertex = DynamicVertex.GetConstArray();
+									const TManagedArray<FVector3f>* DynamicVertex = ConstDynamicCollection->FindPositions();
+									if ((DynamicVertex->Num() > 0) && (DynamicVertex->Num() == RenderVertex.Num())) 
+										RenderVertex = DynamicVertex->GetConstArray();
+									auto ActivationToColor = [](float Activation) { return FLinearColor(Activation, 0.f, 1.f - Activation); };
+									//update muscle activation color
+									const TManagedArray<float>* VertexActivations = ConstDynamicCollection->FindActivations();
+									if (VertexActivations && RenderMesh->Colors.Num() == 3 * NumFaces && VertexActivations->Num() == NumVertices)
+									{
+										for (int32 i = 0; i < NumFaces; ++i)
+										{
+											for (int j = 0; j < 3; ++j)
+											{
+												float Activation = (*VertexActivations)[Flesh.Indices[i][j]];
+												if (Activation >= 0)
+												{
+													RenderMesh->Colors[3 * i + j] = ActivationToColor(Activation);
+												}
+												else
+												{
+													RenderMesh->Colors[3 * i + j] = FLinearColor::White;
+												}
+											}
+										}
+									}
 								}
 								auto InRange = [](int32 Size, int32 Val) { return 0 <= Val && Val < Size; };
 
@@ -685,18 +716,19 @@ TOptional<TArray<int32>> UDeformableTetrahedralComponent::GetMeshImportVertexMap
 	{
 		ChaosFlesh::BoundSurfacePositions(SkeletalMeshAsset, RestCollection->GetCollection(), RestVertices, RestVertices, Positions);
 
-		//@todo(Flesh Sections) : Add checks for multiple sections. 
 		int32 NumSections = MLDLOD.Sections.Num();
-		if (NumSections != 1)
+		int32 VertexOffset = 0;
+		int32 MLDVertexCount = 0;
+		for (int32 SectionIndex = 0; SectionIndex < NumSections; ++SectionIndex)
 		{
-			UE_LOG(LogDeformableTetrahedralComponentInternal, Warning, TEXT("SkeletalMeshAsset should have only one section."));
-			return None;
+			const FSkelMeshSection& MLDSection = MLDLOD.Sections[SectionIndex];
+			MLDVertexCount += MLDSection.NumVertices;
 		}
-
 		for (int32 SectionIndex = 0; SectionIndex < NumSections; ++SectionIndex)
 		{
 			const FSkelMeshSection& MLDSection = MLDLOD.Sections[SectionIndex];
-			if (MLDSection.NumVertices != Positions.Num())
+			if (MLDSection.NumVertices > Positions.Num()
+				|| (SectionIndex == NumSections - 1 && VertexOffset + MLDSection.NumVertices != Positions.Num()))
 			{
 				UE_LOG(LogDeformableTetrahedralComponentInternal, Warning, TEXT("SkeletalMeshAsset and FleshAsset have different number of vertices in section %d. Check if the assets have the same mesh."), SectionIndex);
 				return None;
@@ -705,12 +737,13 @@ TOptional<TArray<int32>> UDeformableTetrahedralComponent::GetMeshImportVertexMap
 			for (int32 VertexIndex = 0; VertexIndex < MLDSection.NumVertices; ++VertexIndex)
 			{
 				const FVector3f& MLDPosition = MLDSection.SoftVertices[VertexIndex].Position;
-				const FVector3f& FleshPosition = Positions[VertexIndex];
+				const FVector3f& FleshPosition = Positions[VertexIndex + VertexOffset];
 				if (!MLDPosition.Equals(FleshPosition, UE_KINDA_SMALL_NUMBER))
 				{
 					UE_LOG(LogDeformableTetrahedralComponentInternal, Warning, TEXT("SkeletalMeshAsset and FleshAsset have different vertex positions. Check if the assets have the same vertex order or some points are not embedded."));
 				}
 			}
+			VertexOffset += MLDSection.NumVertices;
 		}
 	}
 
diff --git a/Engine/Plugins/Experimental/ChaosFlesh/Source/ChaosFleshEngine/Private/ChaosFlesh/FleshDynamicAsset.cpp b/Engine/Plugins/Experimental/ChaosFlesh/Source/ChaosFleshEngine/Private/ChaosFlesh/FleshDynamicAsset.cpp
--- a/Engine/Plugins/Experimental/ChaosFlesh/Source/ChaosFleshEngine/Private/ChaosFlesh/FleshDynamicAsset.cpp
+++ b/Engine/Plugins/Experimental/ChaosFlesh/Source/ChaosFleshEngine/Private/ChaosFlesh/FleshDynamicAsset.cpp
@@ -22,6 +22,7 @@ void UFleshDynamicAsset::Init()
 {
 	DynamicCollection->AddAttribute<FVector3f>("Vertex", FGeometryCollection::VerticesGroup);
 	DynamicCollection->AddAttribute<int32>("ObjectState", FGeometryCollection::VerticesGroup);
+	DynamicCollection->AddAttribute<float>("Activation", FGeometryCollection::VerticesGroup);
 }
 
 void UFleshDynamicAsset::Reset(const FManagedArrayCollection* InCopyFrom)
@@ -60,3 +61,12 @@ TManagedArray<int32>& UFleshDynamicAsset::GetObjectState()
 	return DynamicCollection->ModifyAttribute<int32>("ObjectState", FGeometryCollection::VerticesGroup);
 }
 
+TManagedArray<float>& UFleshDynamicAsset::GetActivations()
+{
+	return DynamicCollection->ModifyAttribute<float>("Activation", FGeometryCollection::VerticesGroup);
+}
+
+const TManagedArray<float>* UFleshDynamicAsset::FindActivations() const
+{
+	return DynamicCollection->FindAttributeTyped<float>("Activation", FGeometryCollection::VerticesGroup);
+}
\ No newline at end of file
diff --git a/Engine/Plugins/Experimental/ChaosFlesh/Source/ChaosFleshEngine/Public/ChaosCache/FleshComponentCacheAdapter.h b/Engine/Plugins/Experimental/ChaosFlesh/Source/ChaosFleshEngine/Public/ChaosCache/FleshComponentCacheAdapter.h
--- a/Engine/Plugins/Experimental/ChaosFlesh/Source/ChaosFleshEngine/Public/ChaosCache/FleshComponentCacheAdapter.h
+++ b/Engine/Plugins/Experimental/ChaosFlesh/Source/ChaosFleshEngine/Public/ChaosCache/FleshComponentCacheAdapter.h
@@ -6,7 +6,7 @@
 #include "Chaos/Deformable/ChaosDeformableSolver.h"
 #include "Chaos/Deformable/ChaosDeformableSolverProxy.h"
 #include "Chaos/Deformable/ChaosDeformableSolverTypes.h"
-
+#include "Chaos/Range.h"
 #if USE_USD_SDK && DO_USD_CACHING
 
 #include "UsdWrappers/UsdStage.h"
@@ -62,6 +62,7 @@ namespace Chaos
 		virtual uint8						GetPriority() const override;
 		virtual FGuid						GetGuid() const override;
 		virtual bool						ValidForPlayback(UPrimitiveComponent* InComponent, UChaosCache* InCache) const override;
+		virtual bool						SupportsPlayWhileRecording() const override { return true; };
 		virtual Chaos::FPhysicsSolver*		GetComponentSolver(UPrimitiveComponent* InComponent) const override;
 		virtual Chaos::FPhysicsSolverEvents* BuildEventsSolver(UPrimitiveComponent* InComponent) const override;
 		virtual void						SetRestState(UPrimitiveComponent* InComponent, UChaosCache* InCache, const FTransform& InRootTransform, Chaos::FReal InTime) const override;
@@ -81,13 +82,14 @@ namespace Chaos
 		static FDeformableSolver* GetDeformableSolver(UPrimitiveComponent* InComponent);
 
 		/** Get the particle indices range for the component in the evolution particles list on the physics thread */
-		static FIntVector2 GetParticleRange(UPrimitiveComponent* InComponent, const int32 NumParticles);
+		static Chaos::FRange GetParticleRange(UPrimitiveComponent* InComponent, const int32 NumParticles);
 		
-	private:
+		bool SupportsRestartSimulation() const override { return true; };
 
+	private:
 		/** Load the cache for a given component at a specific time */
 		void LoadCacheAtTime(UPrimitiveComponent* PrimitiveComponent, Chaos::FReal TargetTime, const int32 NumParticles, const bool bNeedsRange,
-				const TFunction<void(const int32, const int32, const FVector3f&, const FVector3f&)>& LoadFunction) const;
+				const TFunction<void(const int32, const int32, const FVector3f&, const FVector3f&, const float MuscleActivation)>& LoadFunction) const;
 		
 #if USE_USD_SDK && DO_USD_CACHING
 
diff --git a/Engine/Plugins/Experimental/ChaosFlesh/Source/ChaosFleshEngine/Public/ChaosFlesh/ChaosDeformablePhysicsComponent.h b/Engine/Plugins/Experimental/ChaosFlesh/Source/ChaosFleshEngine/Public/ChaosFlesh/ChaosDeformablePhysicsComponent.h
--- a/Engine/Plugins/Experimental/ChaosFlesh/Source/ChaosFleshEngine/Public/ChaosFlesh/ChaosDeformablePhysicsComponent.h
+++ b/Engine/Plugins/Experimental/ChaosFlesh/Source/ChaosFleshEngine/Public/ChaosFlesh/ChaosDeformablePhysicsComponent.h
@@ -44,7 +44,7 @@ public:
 	virtual void PreSolverUpdate() {};
 	virtual FDataMapValue NewDeformableData() { return FDataMapValue(nullptr); }
 	virtual void UpdateFromSimulation(const FDataMapValue* SimulationBuffer) {}
-
+	virtual FDataMapValue WriteRestartDataToSimulation() { return FDataMapValue(nullptr); }
 
 	virtual void OnCreatePhysicsState() override;
 	virtual void OnDestroyPhysicsState() override;
diff --git a/Engine/Plugins/Experimental/ChaosFlesh/Source/ChaosFleshEngine/Public/ChaosFlesh/ChaosDeformableSolverComponent.h b/Engine/Plugins/Experimental/ChaosFlesh/Source/ChaosFleshEngine/Public/ChaosFlesh/ChaosDeformableSolverComponent.h
--- a/Engine/Plugins/Experimental/ChaosFlesh/Source/ChaosFleshEngine/Public/ChaosFlesh/ChaosDeformableSolverComponent.h
+++ b/Engine/Plugins/Experimental/ChaosFlesh/Source/ChaosFleshEngine/Public/ChaosFlesh/ChaosDeformableSolverComponent.h
@@ -94,6 +94,7 @@ public:
 	virtual void ResetSimulationProxy() override;
 	virtual void WriteToSimulation(const float DeltaTime, const bool bAsyncTask) override;
 	virtual void ReadFromSimulation(const float DeltaTime, const bool bAsyncTask) override;
+	virtual void ReadRestartData() override;
 	// End IDataflowPhysicsSolverInterface overrides
 
 	//~ Begin UObject Interface
diff --git a/Engine/Plugins/Experimental/ChaosFlesh/Source/ChaosFleshEngine/Public/ChaosFlesh/ChaosDeformableTetrahedralComponent.h b/Engine/Plugins/Experimental/ChaosFlesh/Source/ChaosFleshEngine/Public/ChaosFlesh/ChaosDeformableTetrahedralComponent.h
--- a/Engine/Plugins/Experimental/ChaosFlesh/Source/ChaosFleshEngine/Public/ChaosFlesh/ChaosDeformableTetrahedralComponent.h
+++ b/Engine/Plugins/Experimental/ChaosFlesh/Source/ChaosFleshEngine/Public/ChaosFlesh/ChaosDeformableTetrahedralComponent.h
@@ -114,7 +114,8 @@ public:
 	/** Simulation Interface */
 	virtual FThreadingProxy* NewProxy() override;
 	virtual FDataMapValue NewDeformableData() override;
-	virtual void UpdateFromSimulation(const FDataMapValue* SimualtionBuffer) override;
+	virtual void UpdateFromSimulation(const FDataMapValue* SimulationBuffer) override;
+	virtual FDataMapValue WriteRestartDataToSimulation() override;
 
 	/** RestCollection */
 	UFUNCTION(BlueprintCallable, Category = "Physics")
diff --git a/Engine/Plugins/Experimental/ChaosFlesh/Source/ChaosFleshEngine/Public/ChaosFlesh/FleshDynamicAsset.h b/Engine/Plugins/Experimental/ChaosFlesh/Source/ChaosFleshEngine/Public/ChaosFlesh/FleshDynamicAsset.h
--- a/Engine/Plugins/Experimental/ChaosFlesh/Source/ChaosFleshEngine/Public/ChaosFlesh/FleshDynamicAsset.h
+++ b/Engine/Plugins/Experimental/ChaosFlesh/Source/ChaosFleshEngine/Public/ChaosFlesh/FleshDynamicAsset.h
@@ -33,6 +33,8 @@ public:
 	TManagedArray<FVector3f>& GetPositions();
 	const TManagedArray<FVector3f>* FindPositions() const;
 	TManagedArray<int32>& GetObjectState();
+	TManagedArray<float>& GetActivations();
+	const TManagedArray<float>* FindActivations() const;
 
 	const FManagedArrayCollection* GetCollection() const { return DynamicCollection.Get(); }
 	FManagedArrayCollection* GetCollection() { return DynamicCollection.Get(); }
diff --git a/Engine/Plugins/Experimental/Dataflow/Source/DataflowEditor/Private/Dataflow/DataflowSimulationGenerator.cpp b/Engine/Plugins/Experimental/Dataflow/Source/DataflowEditor/Private/Dataflow/DataflowSimulationGenerator.cpp
--- a/Engine/Plugins/Experimental/Dataflow/Source/DataflowEditor/Private/Dataflow/DataflowSimulationGenerator.cpp
+++ b/Engine/Plugins/Experimental/Dataflow/Source/DataflowEditor/Private/Dataflow/DataflowSimulationGenerator.cpp
@@ -18,9 +18,9 @@ namespace UE::Dataflow
 {
 void FDataflowSimulationTask::DoWork()
 {
-	const int32 StartFrame = bAsyncCaching ? 1 : 0;
-	const int32 NumFrames = (MaxTime-MinTime) / DeltaTime;
-	for (int32 FrameIndex = StartFrame; FrameIndex < NumFrames; ++FrameIndex)
+	const int32 NumFrames = (MaxTime-MinTime+ UE_KINDA_SMALL_NUMBER) / DeltaTime;
+
+	for (int32 FrameIndex = 0; FrameIndex < NumFrames; ++FrameIndex)
 	{
 		if (!TaskManager->bCancelled.load())
 		{
@@ -74,14 +74,24 @@ bool FDataflowTaskManager::AllocateSimulationResource(const FVector2f& TimeRange
 	WorldContext.SetCurrentWorld(SimulationWorld);
 	
 	CacheManager = SimulationWorld->SpawnActor<AChaosCacheManager>();
+	UDataflowSimulationManager* DataflowManager = SimulationWorld->GetSubsystem<UDataflowSimulationManager>();
 
 	PreviewActor = UE::Dataflow::SpawnSimulatedActor(ActorClass, CacheManager, CacheAsset, true, DataflowContent, BlueprintTransform);
 	UE::Dataflow::SetupSkeletonAnimation(PreviewActor, bSkeletalMeshVisibility);
 
-	// Init the cache manager
+	// Set simulation restart time
+	if (SimulationTask->GetTask().bRestartSimulation)
+	{
+		CacheManager->SetRestartSimulation(SimulationTask->GetTask().bRestartSimulation);
+		CacheManager->SetRestartTimeRange(TimeRange[0], TimeRange[1]);
+		// Read restart positions
+		CacheManager->ReadRestartData();
+		DataflowManager->ReadRestartData();
+	}
 	CacheManager->SetObservedComponentProperties(CacheManager->CacheMode);
+	// Init the cache manager
 	CacheManager->BeginEvaluate();
-	
+
 	SimulationResource = MakeShared<FDataflowSimulationResource>();
 	SimulationResource->NumSimulatedFrames = &NumSimulatedFrames;
 	SimulationResource->bCancelled = &bCancelled;
@@ -94,14 +104,14 @@ bool FDataflowTaskManager::AllocateSimulationResource(const FVector2f& TimeRange
 		SimulationTask->GetTask().MaxTime = TimeRange[1];
 		SimulationTask->GetTask().DeltaTime = (TimeRange[1]-TimeRange[0]) / NumFrames;
 
-		UDataflowSimulationManager* DataflowManager = SimulationWorld->GetSubsystem<UDataflowSimulationManager>();
 		if(SimulationTask->GetTask().bAsyncCaching)
 		{
 			// Update all the skelmesh animations at the simulation time
-			UE::Dataflow::UpdateSkeletonAnimation(PreviewActor, SimulationTask->GetTask().MinTime+SimulationTask->GetTask().DeltaTime);
-				
+			UE::Dataflow::UpdateSkeletonAnimation(PreviewActor, SimulationTask->GetTask().MinTime);
+			
 			// Foreground task : Run the world ticking 
-			SimulationWorld->Tick(ELevelTick::LEVELTICK_All, SimulationTask->GetTask().DeltaTime);
+			// DeltaSeconds = 0 because this is the preroll, i.e., frame 0 simulation
+			SimulationWorld->Tick(ELevelTick::LEVELTICK_All, 0 /*DeltaSeconds*/);
 
 			// Init simulation proxies from interface
 			DataflowManager->InitSimulationInterfaces();
@@ -127,6 +137,7 @@ void FDataflowTaskManager::FreeSimulationResource()
 
 	if(CacheManager)
 	{
+		// Writes data into caches
 		CacheManager->EndEvaluate();
 
 		//  Clear the observed components
@@ -189,8 +200,9 @@ void FDataflowSimulationGenerator::StartGenerateSimulation()
 	TaskManager->SimulationTask = MakeUnique<FAsyncTask<FDataflowSimulationTask>>();
 	TaskManager->SimulationTask->GetTask().TaskManager = TaskManager;
 	TaskManager->SimulationTask->GetTask().bAsyncCaching = CacheParams.bAsyncCaching;
-	
-	TaskManager->AllocateSimulationResource(CacheParams.TimeRange, CacheParams.FrameRate, CacheAsset, BlueprintClass, DataflowContent, BlueprintTransform, bSkeletalMeshVisibility);
+	TaskManager->SimulationTask->GetTask().bRestartSimulation = CacheParams.bRestartSimulation;
+	FVector2f TimeRange = CacheParams.bRestartSimulation ? CacheParams.RestartTimeRange : CacheParams.TimeRange;
+	TaskManager->AllocateSimulationResource(TimeRange, CacheParams.FrameRate, CacheAsset, BlueprintClass, DataflowContent, BlueprintTransform, bSkeletalMeshVisibility);
 
 	if(CacheParams.bAsyncCaching)
 	{
diff --git a/Engine/Plugins/Experimental/Dataflow/Source/DataflowEditor/Private/Dataflow/DataflowSimulationScene.cpp b/Engine/Plugins/Experimental/Dataflow/Source/DataflowEditor/Private/Dataflow/DataflowSimulationScene.cpp
--- a/Engine/Plugins/Experimental/Dataflow/Source/DataflowEditor/Private/Dataflow/DataflowSimulationScene.cpp
+++ b/Engine/Plugins/Experimental/Dataflow/Source/DataflowEditor/Private/Dataflow/DataflowSimulationScene.cpp
@@ -246,7 +246,10 @@ void FDataflowSimulationScene::UpdateSimulationCache(float& DeltaSeconds)
 		if(!bIsRecordingCache)
 		{
 			// Update the cached simulation at some point in time
-			CacheManager->SetStartTime(SimulationTime);
+			if (SimulationTime != CacheManager->StartTime)
+			{
+				CacheManager->SetStartTime(SimulationTime);
+			}
 		}
 		else
 		{
diff --git a/Engine/Plugins/Experimental/Dataflow/Source/DataflowEditor/Public/Dataflow/DataflowSimulationGenerator.h b/Engine/Plugins/Experimental/Dataflow/Source/DataflowEditor/Public/Dataflow/DataflowSimulationGenerator.h
--- a/Engine/Plugins/Experimental/Dataflow/Source/DataflowEditor/Public/Dataflow/DataflowSimulationGenerator.h
+++ b/Engine/Plugins/Experimental/Dataflow/Source/DataflowEditor/Public/Dataflow/DataflowSimulationGenerator.h
@@ -55,6 +55,9 @@ namespace UE::Dataflow
 
 		/** Boolean to check if we are running the task in the background */
 		bool bAsyncCaching = true;
+
+		/** Boolean to check if we should restart the simulation in the range */
+		bool bRestartSimulation = false;
 	};
 
 	/** Async simulation resource that will be used while simulating */
diff --git a/Engine/Source/Runtime/Experimental/Chaos/Private/Chaos/Deformable/ChaosDeformableSolver.cpp b/Engine/Source/Runtime/Experimental/Chaos/Private/Chaos/Deformable/ChaosDeformableSolver.cpp
--- a/Engine/Source/Runtime/Experimental/Chaos/Private/Chaos/Deformable/ChaosDeformableSolver.cpp
+++ b/Engine/Source/Runtime/Experimental/Chaos/Private/Chaos/Deformable/ChaosDeformableSolver.cpp
@@ -57,6 +57,8 @@ DECLARE_CYCLE_STAT(TEXT("Chaos.Deformable.Solver.PullOutputPackage"), STAT_Chaos
 DECLARE_CYCLE_STAT(TEXT("Chaos.Deformable.Solver.PushOutputPackage"), STAT_ChaosDeformableSolver_PushOutputPackage, STATGROUP_Chaos);
 DECLARE_CYCLE_STAT(TEXT("Chaos.Deformable.Solver.PullInputPackage"), STAT_ChaosDeformableSolver_PullInputPackage, STATGROUP_Chaos);
 DECLARE_CYCLE_STAT(TEXT("Chaos.Deformable.Solver.PushInputPackage"), STAT_ChaosDeformableSolver_PushInputPackage, STATGROUP_Chaos);
+DECLARE_CYCLE_STAT(TEXT("Chaos.Deformable.Solver.PullRestartPackage"), STAT_ChaosDeformableSolver_PullRestartPackage, STATGROUP_Chaos);
+DECLARE_CYCLE_STAT(TEXT("Chaos.Deformable.Solver.PushRestartPackage"), STAT_ChaosDeformableSolver_PushRestartPackage, STATGROUP_Chaos);
 DECLARE_CYCLE_STAT(TEXT("Chaos.Deformable.Solver.InitializeSimulationObjects"), STAT_ChaosDeformableSolver_InitializeSimulationObjects, STATGROUP_Chaos);
 DECLARE_CYCLE_STAT(TEXT("Chaos.Deformable.Solver.InitializeSimulationObject"), STAT_ChaosDeformableSolver_InitializeSimulationObject, STATGROUP_Chaos);
 DECLARE_CYCLE_STAT(TEXT("Chaos.Deformable.Solver.InitializeDeformableParticles"), STAT_ChaosDeformableSolver_InitializeDeformableParticles, STATGROUP_Chaos);
@@ -103,6 +105,7 @@ namespace Chaos::Softs
 	FCriticalSection FDeformableSolver::RemovalMutex;
 	FCriticalSection FDeformableSolver::PackageOutputMutex;
 	FCriticalSection FDeformableSolver::PackageInputMutex;
+	FCriticalSection FDeformableSolver::PackageRestartMutex;
 	FCriticalSection FDeformableSolver::SolverEnabledMutex;
 
 	int32 GSParallelMax = 100;
@@ -202,6 +205,51 @@ namespace Chaos::Softs
 		}
 	}
 
+	void FDeformableSolver::LoadRestartData()
+	{
+		// pull CurrentRestartPackage. This is after UDeformableSolverComponent::ReadRestartData() that pushes it.
+		UpdateProxyRestartPackages();
+		// update evolution particles
+		UpdateRestartParticlePositions();
+	}
+
+	void FDeformableSolver::UpdateProxyRestartPackages()
+	{
+		TUniquePtr < FDeformablePackage > TailPackage = PullRestartPackage();
+		while (TailPackage)
+		{
+			CurrentRestartPackage = TUniquePtr < FDeformablePackage >(TailPackage.Release());
+			TailPackage = PullRestartPackage();
+		}
+	}
+
+	void FDeformableSolver::UpdateRestartParticlePositions()
+	{
+		for (const TPair<FThreadingProxy::FKey, TUniquePtr<FThreadingProxy>>& Entry : Proxies)
+		{
+			const FThreadingProxy::FKey& Owner = Entry.Key;
+			if (const FFleshThreadingProxy* Proxy = Entry.Value->As<FFleshThreadingProxy>())
+			{
+				if (this->CurrentRestartPackage)
+				{
+					if (FFleshThreadingProxy::FFleshRestartBuffer* FleshRestartBuffer =
+						this->CurrentRestartPackage->ObjectMap.Contains(Owner) ?
+						this->CurrentRestartPackage->ObjectMap[Owner]->As<FFleshThreadingProxy::FFleshRestartBuffer>() :
+						nullptr)
+					{
+						if (const TManagedArray<FVector3f>* DynamicVertex = FleshRestartBuffer->Dynamic.FindAttribute<FVector3f>("Vertex", FGeometryCollection::VerticesGroup))
+						{
+							const Chaos::FRange& Range = Proxy->GetSolverParticleRange();
+							for (int32 i = Range.Start; i < Range.Start + Range.Count; ++i)
+							{
+								Evolution->Particles().SetX(i, (*DynamicVertex)[i]);
+							}							
+						} // if has DynamicVertex
+					} // if flesh input buffer
+				} // if CurrentRestartPackage
+			}
+		} // for all proxies
+	}
 	void FDeformableSolver::Simulate(FSolverReal DeltaTime)
 	{
 		PERF_SCOPE(STAT_ChaosDeformableSolver_Simulate);
@@ -211,6 +259,11 @@ namespace Chaos::Softs
 			UpdateProxyInputPackages();
 			InitializeSimulationObjects();
 			InitializeSimulationSpace();
+			if (bPendingRestart)
+			{
+				LoadRestartData();
+				bPendingRestart = false;
+			}
 			AdvanceDt(DeltaTime);
 			DebugDrawSimulationData();
 		}
@@ -233,7 +286,7 @@ namespace Chaos::Softs
 						GeometryCollection::Facades::FConstraintOverrideTargetFacade CnstrTargets(FleshInputBuffer->SimulationCollection);
 						if (CnstrTargets.IsValid() && CnstrTargets.Num())
 						{
-							const FIntVector2& Range = Proxy->GetSolverParticleRange();
+							const Chaos::FRange& Range = Proxy->GetSolverParticleRange();
 							const FSolverReal CurrentRatio = FSolverReal(this->Iteration) / FSolverReal(this->Property.NumSolverSubSteps);
 							const FTransform WorldToSim = Proxy->GetCurrentPointsTransform();
 
@@ -243,7 +296,7 @@ namespace Chaos::Softs
 								for (int32 i = 0; i < CnstrTargets.Num(); i++)
 								{
 									int32 LocalIndex = CnstrTargets.GetIndex(i);
-									int32 ParticleIndex = Range[0] + LocalIndex;
+									int32 ParticleIndex = Range.Start + LocalIndex;
 									// Set particle kinematic state to kinematic, saving prior state.
 									TransientConstraintBuffer.Add(
 										TPair<int32, TTuple<float, float, FVector3f>>(
@@ -264,7 +317,7 @@ namespace Chaos::Softs
 							for (int32 i = 0; i < CnstrTargets.Num(); i++)
 							{
 								int32 LocalIndex = CnstrTargets.GetIndex(i);
-								int32 ParticleIndex = Range[0] + LocalIndex;
+								int32 ParticleIndex = Range.Start + LocalIndex;
 
 								const FVector3f& WorldSpaceTarget = CnstrTargets.GetPosition(i);
 								FVector3f SimSpaceTarget = ToSingle(WorldToSim.TransformPosition(ToDouble(WorldSpaceTarget)));
@@ -284,7 +337,7 @@ namespace Chaos::Softs
 								for (int32 i = 0; i < CnstrTargets.Num(); i++)
 								{
 									int32 LocalIndex = CnstrTargets.GetIndex(i);
-									int32 ParticleIndex = Range[0] + LocalIndex;
+									int32 ParticleIndex = Range.Start + LocalIndex;
 									Chaos::FDebugDrawQueue::GetInstance().DrawDebugPoint(ToDouble(Evolution->Particles().GetX(ParticleIndex)), FColor::Orange, false, -1.0f, 0, GDeformableDebugParams.ParticleRadius);
 								}
 							}
@@ -555,7 +608,7 @@ namespace Chaos::Softs
 		PERF_SCOPE(STAT_ChaosDeformableSolver_InitializeKinematicParticles);
 
 		const FManagedArrayCollection& Rest = Proxy.GetRestCollection();
-		const FIntVector2& Range = Proxy.GetSolverParticleRange();
+		const Chaos::FRange& Range = Proxy.GetSolverParticleRange();
 
 		if (Property.bEnableKinematics)
 		{
@@ -564,7 +617,7 @@ namespace Chaos::Softs
 			{
 				if (VertexBoneWeightsFacade.IsKinematicVertex(VertexIdx))
 				{
-					int32 ParticleIndex = Range[0] + VertexIdx;
+					int32 ParticleIndex = Range.Start + VertexIdx;
 					Evolution->Particles().InvM(ParticleIndex) = 0.f;
 					Evolution->Particles().PAndInvM(ParticleIndex).InvM = 0.f;
 				}
@@ -595,7 +648,7 @@ namespace Chaos::Softs
 								bHavePrintedLog = true;
 								UE_LOG(LogChaosDeformableSolver, Warning, TEXT("Detected deprecated kinematic initialization, reevaluate input asset"));
 							}
-							int32 ParticleIndex = Range[0] + vdx;
+							int32 ParticleIndex = Range.Start + vdx;
 							Evolution->Particles().InvM(ParticleIndex) = 0.f;
 							Evolution->Particles().PAndInvM(ParticleIndex).InvM = 0.f;
 						}
@@ -611,7 +664,7 @@ namespace Chaos::Softs
 		PERF_SCOPE(STAT_ChaosDeformableSolver_InitializeWeakConstraints);
 
 		const FManagedArrayCollection& Rest = Proxy.GetRestCollection();
-		const FIntVector2& Range = Proxy.GetSolverParticleRange();
+		const Chaos::FRange& Range = Proxy.GetSolverParticleRange();
 		if (Property.bEnablePositionTargets)
 		{
 			typedef GeometryCollection::Facades::FPositionTargetFacade FPositionTargets;
@@ -648,11 +701,11 @@ namespace Chaos::Softs
 				{
 					for (int32 j = 0; j < PositionTargetIndices[i].Num(); j++)
 					{
-						PositionTargetIndices[i][j] += Range[0];
+						PositionTargetIndices[i][j] += Range.Start;
 					}
 					for (int32 j = 0; j < PositionTargetSecondIndices[i].Num(); j++)
 					{
-						PositionTargetSecondIndices[i][j] += Range[0];
+						PositionTargetSecondIndices[i][j] += Range.Start;
 					}
 				}
 				GSWeakConstraints->AddExtraConstraints(PositionTargetIndices, PositionTargetWeights, PositionTargetStiffness, PositionTargetSecondIndices, PositionTargetSecondWeights);
@@ -811,11 +864,11 @@ namespace Chaos::Softs
 							FFleshThreadingProxy* TargetProxy = Proxies[AddConstraints.Get<1>()]->As<FFleshThreadingProxy>();
 							if (SourceProxy && TargetProxy)
 							{
-								FIntVector2 SampleRange = SourceProxy->GetSolverParticleRange();
-								TConstArrayView<Softs::FSolverVec3> Samples(&Evolution->Particles().GetX(SampleRange[0]), SampleRange[1]);
+								Chaos::FRange SampleRange = SourceProxy->GetSolverParticleRange();
+								TConstArrayView<Softs::FSolverVec3> Samples(&Evolution->Particles().GetX(SampleRange.Start), SampleRange.Count);
 
-								FIntVector2 TargetRange = TargetProxy->GetSolverParticleRange();
-								TConstArrayView<Softs::FSolverVec3> TetVertices(&Evolution->Particles().GetX(TargetRange[0]), TargetRange[1]);
+								Chaos::FRange TargetRange = TargetProxy->GetSolverParticleRange();
+								TConstArrayView<Softs::FSolverVec3> TetVertices(&Evolution->Particles().GetX(TargetRange.Start), TargetRange.Count);
 
 								const GeometryCollection::Facades::FTetrahedralFacade Geom(TargetProxy->GetRestCollection());
 								TArray<GeometryCollection::Facades::TetrahedralParticleEmbedding> Intersections;
@@ -847,7 +900,7 @@ namespace Chaos::Softs
 		auto ChaosVert = [](FVector3d V) { return Chaos::FVec3(V.X, V.Y, V.Z); };
 		auto DoubleVert = [](FVector3f V) { return FVector3d(V.X, V.Y, V.Z); };
 
-		const FIntVector2& Range = Proxy.GetSolverParticleRange();
+		const Chaos::FRange& Range = Proxy.GetSolverParticleRange();
 		const FManagedArrayCollection& Rest = Proxy.GetRestCollection();
 		const TManagedArray<FIntVector4>& Tetrahedron = Rest.GetAttribute<FIntVector4>("Tetrahedron", "Tetrahedral");
 		if (uint32 NumElements = Tetrahedron.Num())
@@ -855,7 +908,7 @@ namespace Chaos::Softs
 			const Chaos::Softs::FSolverParticles& P = Evolution->Particles();
 			for (uint32 edx = 0; edx < NumElements; ++edx)
 			{
-				auto T = ChaosTet(Tetrahedron[edx], Range[0]);
+				auto T = ChaosTet(Tetrahedron[edx], Range.Start);
 				Chaos::FDebugDrawQueue::GetInstance().DrawDebugPoint(
 					DoubleVert(P.GetX(T[0])), FColor::Blue, false, -1.0f, 0, GDeformableDebugParams.ParticleRadius);
 			}
@@ -873,7 +926,7 @@ namespace Chaos::Softs
 		auto ChaosTet = [](FIntVector4 V, int32 dp) { return Chaos::TVec4<int32>(dp + V.X, dp + V.Y, dp + V.Z, dp + V.W); };
 
 		TArray<FSolverReal> StiffnessWithMultiplier;
-		const FIntVector2& Range = Proxy.GetSolverParticleRange();
+		const Chaos::FRange& Range = Proxy.GetSolverParticleRange();
 		if (Rest.HasAttributes({ FManagedArrayCollection::TManagedType<FSolverReal>("Stiffness", FGeometryCollection::VerticesGroup) }))
 		{
 			uint32 NumParticles = Rest.NumElements(FGeometryCollection::VerticesGroup);
@@ -883,7 +936,7 @@ namespace Chaos::Softs
 			FSolverReal IncompressibilityMultiplier = 1.f;
 			FSolverReal InflationMultiplier = 1.f;
 
-			if (const UObject* Owner = this->MObjects[Range[0]]) {
+			if (const UObject* Owner = this->MObjects[Range.Start]) {
 				FFleshThreadingProxy::FFleshInputBuffer* FleshInputBuffer = nullptr;
 				if (this->CurrentInputPackage->ObjectMap.Contains(Owner))
 				{
@@ -912,7 +965,7 @@ namespace Chaos::Softs
 			Elements.SetNum(NumElements);
 			for (uint32 edx = 0; edx < NumElements; ++edx)
 			{
-				Elements[edx] = ChaosTet(Tetrahedron[edx], Range[0]);
+				Elements[edx] = ChaosTet(Tetrahedron[edx], Range.Start);
 			}
 
 			if (Property.bUseGridBasedConstraints)
@@ -921,7 +974,7 @@ namespace Chaos::Softs
 				AllElements->SetNum(ElementsOffset + NumElements);
 				for (uint32 edx = 0; edx < NumElements; ++edx)
 				{
-					(*AllElements)[edx + ElementsOffset] = ChaosTet(Tetrahedron[edx], Range[0]);
+					(*AllElements)[edx + ElementsOffset] = ChaosTet(Tetrahedron[edx], Range.Start);
 				}
 
 			}
@@ -935,7 +988,7 @@ namespace Chaos::Softs
 				FSolverReal IncompressibilityMultiplier = 1.f;
 				FSolverReal InflationMultiplier = 1.f;
 
-				if (const UObject* Owner = this->MObjects[Range[0]]) {
+				if (const UObject* Owner = this->MObjects[Range.Start]) {
 					FFleshThreadingProxy::FFleshInputBuffer* FleshInputBuffer = nullptr;
 					if (this->CurrentInputPackage->ObjectMap.Contains(Owner))
 					{
@@ -1063,7 +1116,7 @@ namespace Chaos::Softs
 
 		}
 
-		FIntVector2 AllUnconstrainedSurfaceElementsSkinRange(INDEX_NONE), AllUnconstrainedSurfaceElementsCorotatedCodRange(INDEX_NONE);
+		Chaos::FRange AllUnconstrainedSurfaceElementsSkinRange(INDEX_NONE, INDEX_NONE), AllUnconstrainedSurfaceElementsCorotatedCodRange(INDEX_NONE, INDEX_NONE);
 
 		if (const TManagedArray<int32>* TriangleMeshIndices = Rest.FindAttribute<int32>("ObjectIndices", "TriangleMesh"))
 		{
@@ -1077,10 +1130,10 @@ namespace Chaos::Softs
 						{
 							if (const TManagedArray<int32>* VertexCounts = Rest.FindAttribute<int32>("VertexCount", FGeometryCollection::GeometryGroup))
 							{
-								AllUnconstrainedSurfaceElementsCorotatedCodRange[0] =  AllUnconstrainedSurfaceElementsCorotatedCod->Num();
-								AllUnconstrainedSurfaceElementsSkinRange[0] =  AllUnconstrainedSurfaceElementsSkin->Num();
-								AllUnconstrainedSurfaceElementsCorotatedCodRange[1] = 0;
-								AllUnconstrainedSurfaceElementsSkinRange[1] = 0;
+								AllUnconstrainedSurfaceElementsCorotatedCodRange.Start =  AllUnconstrainedSurfaceElementsCorotatedCod->Num();
+								AllUnconstrainedSurfaceElementsSkinRange.Start =  AllUnconstrainedSurfaceElementsSkin->Num();
+								AllUnconstrainedSurfaceElementsCorotatedCodRange.Count = 0;
+								AllUnconstrainedSurfaceElementsSkinRange.Count = 0;
 								if (const TManagedArray<bool>* bUseSkinConstraints = Rest.FindAttribute<bool>("SkinConstraints", "TriangleMesh"))
 								{
 									for (int32 i = 0; i < TriangleMeshIndices->Num(); i++)
@@ -1096,10 +1149,10 @@ namespace Chaos::Softs
 											{
 												for (int32 j = 0; j < 3; j++)
 												{
-													(*AllUnconstrainedSurfaceElementsSkin)[e - FaceStartIndex + SurfaceOffset][j] = (*Indices)[e][j] + Range[0];
+													(*AllUnconstrainedSurfaceElementsSkin)[e - FaceStartIndex + SurfaceOffset][j] = (*Indices)[e][j] + Range.Start;
 												}
 											}
-											AllUnconstrainedSurfaceElementsSkinRange[1] += FaceNum;
+											AllUnconstrainedSurfaceElementsSkinRange.Count += FaceNum;
 										}
 										else
 										{
@@ -1109,10 +1162,10 @@ namespace Chaos::Softs
 											{
 												for (int32 j = 0; j < 3; j++)
 												{
-													(*AllUnconstrainedSurfaceElementsCorotatedCod)[e - FaceStartIndex + SurfaceOffset][j] = (*Indices)[e][j] + Range[0];
+													(*AllUnconstrainedSurfaceElementsCorotatedCod)[e - FaceStartIndex + SurfaceOffset][j] = (*Indices)[e][j] + Range.Start;
 												}
 											}
-											AllUnconstrainedSurfaceElementsCorotatedCodRange[1] += FaceNum;
+											AllUnconstrainedSurfaceElementsCorotatedCodRange.Count += FaceNum;
 										}
 									}
 								}
@@ -1129,10 +1182,10 @@ namespace Chaos::Softs
 										{
 											for (int32 j = 0; j < 3; j++)
 											{
-												(*AllUnconstrainedSurfaceElementsCorotatedCod)[e - FaceStartIndex + SurfaceOffset][j] = (*Indices)[e][j] + Range[0];
+												(*AllUnconstrainedSurfaceElementsCorotatedCod)[e - FaceStartIndex + SurfaceOffset][j] = (*Indices)[e][j] + Range.Start;
 											}
 										}
-										AllUnconstrainedSurfaceElementsCorotatedCodRange[1] += FaceNum;
+										AllUnconstrainedSurfaceElementsCorotatedCodRange.Count += FaceNum;
 									}
 								}
 							}
@@ -1143,10 +1196,10 @@ namespace Chaos::Softs
 		}
 		if (AllCorotatedCodEMeshArray)
 		{
-			if (AllUnconstrainedSurfaceElementsCorotatedCodRange[1] > 0 )
+			if (AllUnconstrainedSurfaceElementsCorotatedCodRange.Count > 0 )
 			{
-				AllCorotatedCodEMeshArray->SetNum(AllCorotatedCodEMeshArray->Num() + AllUnconstrainedSurfaceElementsCorotatedCodRange[1]);
-				for (int32 i = AllUnconstrainedSurfaceElementsCorotatedCodRange[0]; i < AllUnconstrainedSurfaceElementsCorotatedCodRange[0] + AllUnconstrainedSurfaceElementsCorotatedCodRange[1]; i++)
+				AllCorotatedCodEMeshArray->SetNum(AllCorotatedCodEMeshArray->Num() + AllUnconstrainedSurfaceElementsCorotatedCodRange.Count);
+				for (int32 i = AllUnconstrainedSurfaceElementsCorotatedCodRange.Start; i < AllUnconstrainedSurfaceElementsCorotatedCodRange.Start + AllUnconstrainedSurfaceElementsCorotatedCodRange.Count; i++)
 				{
 					(*AllCorotatedCodEMeshArray)[i] = 0.f;
 				}
@@ -1154,10 +1207,10 @@ namespace Chaos::Softs
 		}
 		if (AllSkinEMeshArray)
 		{
-			if (AllUnconstrainedSurfaceElementsSkinRange[1] > 0)
+			if (AllUnconstrainedSurfaceElementsSkinRange.Count > 0)
 			{
-				AllSkinEMeshArray->SetNum(AllSkinEMeshArray->Num() + AllUnconstrainedSurfaceElementsSkinRange[1]);
-				for (int32 i = AllUnconstrainedSurfaceElementsSkinRange[0]; i < AllUnconstrainedSurfaceElementsSkinRange[0] + AllUnconstrainedSurfaceElementsSkinRange[1]; i++)
+				AllSkinEMeshArray->SetNum(AllSkinEMeshArray->Num() + AllUnconstrainedSurfaceElementsSkinRange.Count);
+				for (int32 i = AllUnconstrainedSurfaceElementsSkinRange.Start; i < AllUnconstrainedSurfaceElementsSkinRange.Start + AllUnconstrainedSurfaceElementsSkinRange.Count; i++)
 				{
 					(*AllSkinEMeshArray)[i] = 0.f;
 				}
@@ -1165,20 +1218,20 @@ namespace Chaos::Softs
 		}
 		if (StiffnessWithMultiplier.Num() > 0)
 		{
-			if (AllUnconstrainedSurfaceElementsCorotatedCodRange[1] > 0 && AllCorotatedCodEMeshArray)
+			if (AllUnconstrainedSurfaceElementsCorotatedCodRange.Count > 0 && AllCorotatedCodEMeshArray)
 			{
-				for (int32 i = AllUnconstrainedSurfaceElementsCorotatedCodRange[0]; i < AllUnconstrainedSurfaceElementsCorotatedCodRange[0] + AllUnconstrainedSurfaceElementsCorotatedCodRange[1]; i++)
+				for (int32 i = AllUnconstrainedSurfaceElementsCorotatedCodRange.Start; i < AllUnconstrainedSurfaceElementsCorotatedCodRange.Start + AllUnconstrainedSurfaceElementsCorotatedCodRange.Count; i++)
 				{
-					(*AllCorotatedCodEMeshArray)[i] = (StiffnessWithMultiplier[(*AllUnconstrainedSurfaceElementsCorotatedCod)[i][0] - Range[0]] + StiffnessWithMultiplier[(*AllUnconstrainedSurfaceElementsCorotatedCod)[i][1] - Range[0]]
-									+ StiffnessWithMultiplier[(*AllUnconstrainedSurfaceElementsCorotatedCod)[i][2] - Range[0]]) / 3.f;
+					(*AllCorotatedCodEMeshArray)[i] = (StiffnessWithMultiplier[(*AllUnconstrainedSurfaceElementsCorotatedCod)[i][0] - Range.Start] + StiffnessWithMultiplier[(*AllUnconstrainedSurfaceElementsCorotatedCod)[i][1] - Range.Start]
+									+ StiffnessWithMultiplier[(*AllUnconstrainedSurfaceElementsCorotatedCod)[i][2] - Range.Start]) / 3.f;
 				}
 			}
-			if (AllSkinEMeshArray && AllUnconstrainedSurfaceElementsSkinRange[1] > 0)
+			if (AllSkinEMeshArray && AllUnconstrainedSurfaceElementsSkinRange.Count > 0)
 			{
-				for (int32 i = AllUnconstrainedSurfaceElementsSkinRange[0]; i < AllUnconstrainedSurfaceElementsSkinRange[0] + AllUnconstrainedSurfaceElementsSkinRange[1]; i++)
+				for (int32 i = AllUnconstrainedSurfaceElementsSkinRange.Start; i < AllUnconstrainedSurfaceElementsSkinRange.Start + AllUnconstrainedSurfaceElementsSkinRange.Count; i++)
 				{
-					(*AllSkinEMeshArray)[i] = (StiffnessWithMultiplier[(*AllUnconstrainedSurfaceElementsSkin)[i][0] - Range[0]] + StiffnessWithMultiplier[(*AllUnconstrainedSurfaceElementsSkin)[i][1] - Range[0]]
-									+ StiffnessWithMultiplier[(*AllUnconstrainedSurfaceElementsSkin)[i][2] - Range[0]]) / 3.f;
+					(*AllSkinEMeshArray)[i] = (StiffnessWithMultiplier[(*AllUnconstrainedSurfaceElementsSkin)[i][0] - Range.Start] + StiffnessWithMultiplier[(*AllUnconstrainedSurfaceElementsSkin)[i][1] - Range.Start]
+									+ StiffnessWithMultiplier[(*AllUnconstrainedSurfaceElementsSkin)[i][2] - Range.Start]) / 3.f;
 				}
 			}
 		}
@@ -1198,13 +1251,13 @@ namespace Chaos::Softs
 
 			if (uint32 NumElements = Tetrahedron.Num())
 			{
-				const FIntVector2& Range = Proxy.GetSolverParticleRange();
+				const Chaos::FRange& Range = Proxy.GetSolverParticleRange();
 
 				int32 ElementsOffset = AllElements->Num();
 				AllElements->SetNum(ElementsOffset + NumElements);
 				for (uint32 edx = 0; edx < NumElements; ++edx)
 				{
-					(*AllElements)[edx + ElementsOffset] = ChaosTet(Tetrahedron[edx], Range[0]);
+					(*AllElements)[edx + ElementsOffset] = ChaosTet(Tetrahedron[edx], Range.Start);
 				}
 			}
 		}
@@ -1227,17 +1280,17 @@ namespace Chaos::Softs
 
 			if (uint32 NumElements = Tetrahedron.Num())
 			{
-				const FIntVector2& Range = Proxy.GetSolverParticleRange();
+				const Chaos::FRange& Range = Proxy.GetSolverParticleRange();
 
 				int32 ElementsOffset = AllElements->Num();
 				AllElements->SetNum(ElementsOffset + NumElements);
 				for (uint32 edx = 0; edx < NumElements; ++edx)
 				{
-					(*AllElements)[edx + ElementsOffset] = ChaosTet(Tetrahedron[edx], Range[0]);
+					(*AllElements)[edx + ElementsOffset] = ChaosTet(Tetrahedron[edx], Range.Start);
 				}
 				if (uint32 NumIncident = IncidentElementsPointer->Num())
 				{
-					int32 IncidentOffSet = Range[0];
+					int32 IncidentOffSet = Range.Start;
 					AllIncidentElements->SetNum(IncidentOffSet + NumIncident); 
 					AllIncidentElementsLocal->SetNum(IncidentOffSet + NumIncident);
 					for (uint32 i = 0; i < NumIncident; ++i)
@@ -1276,7 +1329,7 @@ namespace Chaos::Softs
 						if (!Proxy->GetIsCached())
 						{
 							FTransform GlobalTransform = Proxy->GetCurrentPointsTransform();
-							const FIntVector2& Range = Proxy->GetSolverParticleRange();
+							const Chaos::FRange& Range = Proxy->GetSolverParticleRange();
 							const FManagedArrayCollection& Rest = Proxy->GetRestCollection();
 
 							if (Rest.FindAttributeTyped<FVector3f>("Vertex", FGeometryCollection::VerticesGroup))
@@ -1285,7 +1338,7 @@ namespace Chaos::Softs
 								// @todo(chaos) : reduce conversions
 								auto ChaosVert = [](FVector3f V) { return Chaos::FVec3(V.X, V.Y, V.Z); };
 								auto ChaosVertfloat = [](FVector3f V) { return Chaos::TVector<FSolverReal, 3>(V.X, V.Y, V.Z); };
-								auto SolverParticleToObjectVertexIndex = [&](int32 SolverParticleIndex) {return SolverParticleIndex - Range[0]; };
+								auto SolverParticleToObjectVertexIndex = [&](int32 SolverParticleIndex) {return SolverParticleIndex - Range.Start; };
 
 								FFleshThreadingProxy::FFleshInputBuffer* FleshInputBuffer = nullptr;
 								if (this->CurrentInputPackage && this->CurrentInputPackage->ObjectMap.Contains(Owner))
@@ -1329,7 +1382,7 @@ namespace Chaos::Softs
 
 														// @todo(flesh) : Add the pre-cached component space rest transforms to the rest collection. 
 														// see  UFleshComponent::NewDeformableData for how its pulled from the SkeletalMesh
-														FVec3 LocalPoint = FleshInputBuffer->RestTransforms[BoneIndices[i]].InverseTransformPosition(ChaosVert(Vertex[Index - Range[0]]));
+														FVec3 LocalPoint = FleshInputBuffer->RestTransforms[BoneIndices[i]].InverseTransformPosition(ChaosVert(Vertex[Index - Range.Start]));
 														FVec3 ComponentPointAtT = FleshInputBuffer->Transforms[BoneIndices[i]].TransformPosition(LocalPoint);
 
 														if (PreviousFleshBuffer)
@@ -1349,9 +1402,9 @@ namespace Chaos::Softs
 										}
 									}
 								}
-								if (!bParticleTouched && ensure(Vertex.IsValidIndex(Index - Range[0])))
+								if (!bParticleTouched && ensure(Vertex.IsValidIndex(Index - Range.Start)))
 								{
-									MParticles.SetX(Index, GlobalTransform.TransformPosition(ChaosVert(Vertex[Index - Range[0]])));
+									MParticles.SetX(Index, GlobalTransform.TransformPosition(ChaosVert(Vertex[Index - Range.Start])));
 									MParticles.PAndInvM(Index).P = MParticles.GetX(Index);
 								}
 							}
@@ -1870,7 +1923,7 @@ namespace Chaos::Softs
 			GeometryCollection::Facades::FMuscleActivationFacade MuscleActivationFacade(Rest);
 			if (MuscleActivationFacade.IsValid())
 			{
-				int32 VertexOffset = Proxy.GetSolverParticleRange()[0];
+				int32 VertexOffset = Proxy.GetSolverParticleRange().Start;
 				int32 ElementOffset = AllElements->Num(); //InitializeMuscleActivation has to go before adding tetrahedrons
 				MuscleActivationConstraints->AddMuscles(MuscleActivationFacade, VertexOffset, ElementOffset);
 			}
@@ -1901,7 +1954,7 @@ namespace Chaos::Softs
 
 				if(FleshProxyA && FleshProxyB)
 				{
-					return FleshProxyA->GetSolverParticleRange()[0] < FleshProxyB->GetSolverParticleRange()[0];
+					return FleshProxyA->GetSolverParticleRange().Start < FleshProxyB->GetSolverParticleRange().Start;
 				}
 				return true;
 			});
@@ -1913,11 +1966,11 @@ namespace Chaos::Softs
 				{
 					if (Proxy->CanSimulate())
 					{
-						FIntVector2 Indices = Proxy->GetSolverParticleRange();
-						if (Indices[1] > 0)
+						Chaos::FRange Indices = Proxy->GetSolverParticleRange();
+						if (Indices.Count > 0)
 						{
-							Proxies.FindAndRemoveChecked(MObjects[Indices[0]]);
-							Evolution->Particles().RemoveAt(Indices[0], Indices[1]);
+							Proxies.FindAndRemoveChecked(MObjects[Indices.Start]);
+							Evolution->Particles().RemoveAt(Indices.Start, Indices.Count);
 						}
 					}
 				}
@@ -1936,10 +1989,10 @@ namespace Chaos::Softs
 						{
 							if (FFleshThreadingProxy* MovedProxy = Proxies[CurrentObject]->As<FFleshThreadingProxy>())
 							{
-								FIntVector2 Range = MovedProxy->GetSolverParticleRange();
-								MovedProxy->SetSolverParticleRange(Index, Range[1]);
-								int32 Offset = Evolution->AddParticleRange(Range[1]);
-								//ensure(Offset == Range[0]);
+								Chaos::FRange Range = MovedProxy->GetSolverParticleRange();
+								MovedProxy->SetSolverParticleRange(Range.Start, Range.Count);
+								int32 Offset = Evolution->AddParticleRange(Range.Count);
+								//ensure(Offset == Range.Start);
 							}
 						}
 					}
@@ -2077,6 +2130,27 @@ namespace Chaos::Softs
 		return TUniquePtr<FDeformablePackage>(nullptr);
 	}
 
+	void FDeformableSolver::PushRestartPackage(int32 InFrame, FDeformableDataMap&& InPackage)
+	{
+		PERF_SCOPE(STAT_ChaosDeformableSolver_PushRestartPackage);
+		FScopeLock Lock(&PackageRestartMutex);
+		TRACE_CPUPROFILER_EVENT_SCOPE(ChaosDeformableSolver_PushRestartPackage);
+		BufferedRestartPackages.Push(TUniquePtr< FDeformablePackage >(new FDeformablePackage(InFrame, MoveTemp(InPackage))));
+		bPendingRestart = true;
+	}
+
+	TUniquePtr<FDeformablePackage> FDeformableSolver::PullRestartPackage()
+	{
+		PERF_SCOPE(STAT_ChaosDeformableSolver_PullRestartPackage);
+		FScopeLock Lock(&PackageRestartMutex);
+		TRACE_CPUPROFILER_EVENT_SCOPE(ChaosDeformableSolver_PullRestartPackage);
+		if (BufferedRestartPackages.Num())
+		{
+			return BufferedRestartPackages.Pop();
+		}
+		return {};
+	}
+
 	void FDeformableSolver::UpdateProxyInputPackages()
 	{
 		PERF_SCOPE(STAT_ChaosDeformableSolver_UpdateProxyInputPackages);
@@ -2180,8 +2254,8 @@ namespace Chaos::Softs
 
 		if (FFleshThreadingProxy* Proxy = InProxy.As<FFleshThreadingProxy>())
 		{
-			const FIntVector2& Range = Proxy->GetSolverParticleRange();
-			if (0 <= Range[0])
+			const Chaos::FRange& Range = Proxy->GetSolverParticleRange();
+			if (0 <= Range.Start)
 			{
 				// @todo(chaos) : reduce conversions
 				auto UEVertd = [](Chaos::FVec3 V) { return FVector3d(V.X, V.Y, V.Z); };
@@ -2189,7 +2263,7 @@ namespace Chaos::Softs
 
 				TManagedArray<FVector3f>& Position = Proxy->GetDynamicCollection().ModifyAttribute<FVector3f>("Vertex", FGeometryCollection::VerticesGroup);
 
-				if((Position.Num() + Range[0]) <= (int32)(Evolution->Particles().Size()))
+				if((Position.Num() + Range.Start) <= (int32)(Evolution->Particles().Size()))
 				{
 					// The final transform gets us from whatever the simulation space is,
 					// to component space.
@@ -2198,7 +2272,7 @@ namespace Chaos::Softs
 					{
 						for (int32 vdx = 0; vdx < Position.Num(); vdx++)
 						{
-							const Chaos::FVec3f& Pos = Evolution->Particles().GetX(vdx + Range[0]);
+							const Chaos::FVec3f& Pos = Evolution->Particles().GetX(vdx + Range.Start);
 							FVector PosD = UEVertd(Pos);
 							Position[vdx] = UEVertf(FinalXf.TransformPosition(PosD));
 						}
@@ -2207,7 +2281,7 @@ namespace Chaos::Softs
 					{
 						for (int32 vdx = 0; vdx < Position.Num(); vdx++)
 						{
-							Position[vdx] = UEVertf(UEVertd(Evolution->Particles().GetX(vdx + Range[0])));
+							Position[vdx] = UEVertf(UEVertd(Evolution->Particles().GetX(vdx + Range.Start)));
 						}
 					}
 				}
diff --git a/Engine/Source/Runtime/Experimental/Chaos/Private/Chaos/Deformable/ChaosDeformableSolverThreadingWrapper.cpp b/Engine/Source/Runtime/Experimental/Chaos/Private/Chaos/Deformable/ChaosDeformableSolverThreadingWrapper.cpp
--- a/Engine/Source/Runtime/Experimental/Chaos/Private/Chaos/Deformable/ChaosDeformableSolverThreadingWrapper.cpp
+++ b/Engine/Source/Runtime/Experimental/Chaos/Private/Chaos/Deformable/ChaosDeformableSolverThreadingWrapper.cpp
@@ -39,6 +39,14 @@ namespace Chaos::Softs
 			Solver->Reset(InProps);
 	}
 
+	void FDeformableSolver::FPhysicsThreadAccess::LoadRestartData()
+	{
+		if (Solver)
+		{
+			Solver->LoadRestartData();
+		}
+	}
+
 	void FDeformableSolver::FPhysicsThreadAccess::AdvanceDt(FSolverReal DeltaTime)
 	{
 		if (Solver)
@@ -87,6 +95,14 @@ namespace Chaos::Softs
 			Solver->PushInputPackage(InFrame, MoveTemp(InPackage));
 	}
 
+	void FDeformableSolver::FGameThreadAccess::PushRestartPackage(int32 InFrame, FDeformableDataMap&& InPackage)
+	{
+		if (Solver)
+		{
+			Solver->PushRestartPackage(InFrame, MoveTemp(InPackage));
+		}	
+	}
+
 	TUniquePtr<FDeformablePackage> FDeformableSolver::FPhysicsThreadAccess::PullInputPackage()
 	{
 		if (Solver)
@@ -156,4 +172,22 @@ namespace Chaos::Softs
 			Solver->UpdateOutputState(InProxy);
 	}
 
+	TArray<float> FDeformableSolver::FPhysicsThreadAccess::GetParticleMuscleActivation()
+	{
+		TArray<float> ParticleMuscleActivations;
+		ParticleMuscleActivations.Init(-1.f, Solver->Evolution->Particles().Size());
+		if (Solver->MuscleActivationConstraints)
+		{
+			for (int32 MuscleIndex = 0; MuscleIndex < Solver->MuscleActivationConstraints->NumMuscles(); ++MuscleIndex)
+			{
+				const int32 VertexOffset = Solver->MuscleActivationConstraints->GetMuscleVertexOffset(MuscleIndex);
+				const int32 VertexCount = Solver->MuscleActivationConstraints->GetMuscleVertexCount(MuscleIndex);
+				for (int32 VertexIdx = VertexOffset; VertexIdx < VertexOffset + VertexCount; ++VertexIdx)
+				{
+					ParticleMuscleActivations[VertexIdx] = Solver->MuscleActivationConstraints->GetMuscleActivation(MuscleIndex);
+				}
+			}
+		}
+		return ParticleMuscleActivations;
+	}
 }; // Namespace Chaos::Softs
\ No newline at end of file
diff --git a/Engine/Source/Runtime/Experimental/Chaos/Private/GeometryCollection/Facades/CollectionMuscleActivationFacade.cpp b/Engine/Source/Runtime/Experimental/Chaos/Private/GeometryCollection/Facades/CollectionMuscleActivationFacade.cpp
--- a/Engine/Source/Runtime/Experimental/Chaos/Private/GeometryCollection/Facades/CollectionMuscleActivationFacade.cpp
+++ b/Engine/Source/Runtime/Experimental/Chaos/Private/GeometryCollection/Facades/CollectionMuscleActivationFacade.cpp
@@ -132,20 +132,64 @@ namespace GeometryCollection::Facades
 		return ReturnData;
 	}
 
+	int32 FMuscleActivationFacade::MuscleVertexOffset(const int32 MuscleIndex) const
+	{
+		const TManagedArray<int32>* VertexStart = ConstCollection.FindAttributeTyped<int32>("VertexStart", FGeometryCollection::GeometryGroup);
+		if (GeometryGroupIndexAttribute.IsValidIndex(MuscleIndex) && VertexStart->IsValidIndex(GeometryGroupIndexAttribute[MuscleIndex]))
+		{
+			return (*VertexStart)[GeometryGroupIndexAttribute[MuscleIndex]];
+		}
+		else
+		{
+			return 0;
+		}
+	}
+	
+	int32 FMuscleActivationFacade::NumMuscleVertices(const int32 MuscleIndex) const
+	{
+		const TManagedArray<int32>* VertexCount = ConstCollection.FindAttributeTyped<int32>("VertexCount", FGeometryCollection::GeometryGroup);
+		if (GeometryGroupIndexAttribute.IsValidIndex(MuscleIndex) && VertexCount->IsValidIndex(GeometryGroupIndexAttribute[MuscleIndex]))
+		{
+			return (*VertexCount)[GeometryGroupIndexAttribute[MuscleIndex]];
+		}
+		else
+		{
+			return 0;
+		}
+	}
+
+	int32 FMuscleActivationFacade::RemoveInvalidMuscles()
+	{
+		check(!IsConst());
+		TArray<int32> InvalidMuscleIndices;
+		if (IsValid())
+		{
+			for (int32 MuscleIndex = 0; MuscleIndex < GeometryGroupIndexAttribute.Num(); ++MuscleIndex)
+			{
+				if (GeometryGroupIndexAttribute[MuscleIndex] < 0)
+				{
+					InvalidMuscleIndices.Add(MuscleIndex);
+				}
+			}
+			Collection->RemoveElements(GroupName, InvalidMuscleIndices);
+		}
+		return InvalidMuscleIndices.Num();
+	}
+
 	bool FMuscleActivationFacade::SetUpMuscleActivation(const TArray<int32>& InOrigin, const TArray<int32>& Insertion,
 		float InContractionVolumeScale)
 	{
 		// Vertices and fiber field
 		if (!(InOrigin.Num() > 0
 			&& Insertion.Num() > 0
-			&& ConstCollection.FindAttribute<FVector3f>("Vertex", "Vertices")
+			&& ConstCollection.FindAttribute<FVector3f>("Vertex", FGeometryCollection::VerticesGroup)
 			&& ConstCollection.FindAttribute<FIntVector4>("Tetrahedron", "Tetrahedral")
 			&& ConstCollection.FindAttribute<FVector3f>("FiberDirection", "Tetrahedral")))
 		{
 			return false;
 		}
 		TArray<int32> Origin = InOrigin;
-		const TArray<FVector3f>& Vertices = ConstCollection.FindAttribute<FVector3f>("Vertex", "Vertices")->GetConstArray();
+		const TArray<FVector3f>& Vertices = ConstCollection.FindAttribute<FVector3f>("Vertex", FGeometryCollection::VerticesGroup)->GetConstArray();
 		const TArray<FIntVector4>& Elements = ConstCollection.FindAttribute<FIntVector4>("Tetrahedron", "Tetrahedral")->GetConstArray();
 		const TArray<FVector3f>& FiberDirections = ConstCollection.FindAttribute<FVector3f>("FiberDirection", "Tetrahedral")->GetConstArray();
 		TArray<TArray<int32>> MuscleActivationElements;
@@ -247,7 +291,7 @@ namespace GeometryCollection::Facades
 		TArray<TArray<TArray<FVector3f>>> MuscleLineSegments;
 		TArray<int32> StreamlineStartElements;
 		// Vertices and fiber field
-		if (!(ConstCollection.FindAttribute<FVector3f>("Vertex", "Vertices")
+		if (!(ConstCollection.FindAttribute<FVector3f>("Vertex", FGeometryCollection::VerticesGroup)
 			&& ConstCollection.FindAttribute<FIntVector4>("Tetrahedron", "Tetrahedral")
 			&& ConstCollection.FindAttribute<FVector3f>("FiberDirection", "Tetrahedral")
 			&& ConstCollection.FindAttribute<int32>("TetrahedronStart", FGeometryCollection::GeometryGroup)
@@ -255,7 +299,7 @@ namespace GeometryCollection::Facades
 		{
 			return MuscleLineSegments;
 		}
-		const TArray<FVector3f>& Vertices = ConstCollection.FindAttribute<FVector3f>("Vertex", "Vertices")->GetConstArray();
+		const TArray<FVector3f>& Vertices = ConstCollection.FindAttribute<FVector3f>("Vertex", FGeometryCollection::VerticesGroup)->GetConstArray();
 		const TArray<FIntVector4>& Elements = ConstCollection.FindAttribute<FIntVector4>("Tetrahedron", "Tetrahedral")->GetConstArray();
 		const TArray<FVector3f>& FiberDirections = ConstCollection.FindAttribute<FVector3f>("FiberDirection", "Tetrahedral")->GetConstArray();
 		const TArray<int32>& TetrahedronStart = ConstCollection.FindAttribute<int32>("TetrahedronStart", FGeometryCollection::GeometryGroup)->GetConstArray();
diff --git a/Engine/Source/Runtime/Experimental/Chaos/Public/Chaos/Deformable/ChaosDeformableSolver.h b/Engine/Source/Runtime/Experimental/Chaos/Public/Chaos/Deformable/ChaosDeformableSolver.h
--- a/Engine/Source/Runtime/Experimental/Chaos/Public/Chaos/Deformable/ChaosDeformableSolver.h
+++ b/Engine/Source/Runtime/Experimental/Chaos/Public/Chaos/Deformable/ChaosDeformableSolver.h
@@ -50,6 +50,9 @@ namespace Chaos::Softs
 			FPhysicsThreadAccess(FDeformableSolver* InSolver, const FPhysicsThreadAccessor&) : Solver(InSolver) {}
 			bool operator()() { return Solver != nullptr; }
 
+			/* Pre-Simulation Advance */
+			CHAOS_API void LoadRestartData();
+
 			/* Simulation Advance */
 			CHAOS_API void UpdateProxyInputPackages();
 			CHAOS_API void Simulate(FSolverReal DeltaTime);
@@ -75,6 +78,8 @@ namespace Chaos::Softs
 			TArrayCollectionArray<const UObject*>& GetObjectsMap() { return Solver->MObjects; }
 			const TArrayCollectionArray<const UObject*>& GetObjectsMap() const { return Solver->MObjects; }
 
+			/* Returns per-particle muscle activation: 0-1 for contractibles, -1 for non-contractibles */
+			CHAOS_API TArray<float> GetParticleMuscleActivation();
 
 		private:
 			FDeformableSolver* Solver;
@@ -93,7 +98,7 @@ namespace Chaos::Softs
 			CHAOS_API void AddProxy(FThreadingProxy* InObject);
 			CHAOS_API void RemoveProxy(FThreadingProxy* InObject);
 			CHAOS_API void PushInputPackage(int32 Frame, FDeformableDataMap&& InPackage);
-
+			CHAOS_API void PushRestartPackage(int32 Frame, FDeformableDataMap&& InPackage);
 			CHAOS_API void SetEnableSolver(bool InbEnableSolver);
 			CHAOS_API bool GetEnableSolver();
 
@@ -108,6 +113,11 @@ namespace Chaos::Softs
 		void SetEnableSolver(bool InbEnableSolver) {FScopeLock Lock(&SolverEnabledMutex); bEnableSolver = InbEnableSolver; }
 		bool GetEnableSolver() const { return bEnableSolver; }
 
+		/* Pre-Simulation Advance */
+		CHAOS_API void LoadRestartData();
+		CHAOS_API void UpdateProxyRestartPackages();
+		CHAOS_API void UpdateRestartParticlePositions();
+
 		/* Simulation Advance */
 		int32 GetFrame() const { return Frame; }
 		CHAOS_API void UpdateProxyInputPackages();
@@ -118,7 +128,8 @@ namespace Chaos::Softs
 		CHAOS_API void UpdateSimulationObjects(FSolverReal DeltaTime);
 		CHAOS_API void UpdateOutputState(FThreadingProxy&);
 		CHAOS_API void PushOutputPackage(int32 Frame, FDeformableDataMap&& Package);
-		CHAOS_API TUniquePtr<FDeformablePackage> PullInputPackage( );
+		CHAOS_API TUniquePtr<FDeformablePackage> PullInputPackage();
+		CHAOS_API TUniquePtr<FDeformablePackage> PullRestartPackage();
 
 		/* Iteration Advance */
 		CHAOS_API void InitializeSimulationSpace();
@@ -145,7 +156,6 @@ namespace Chaos::Softs
 		CHAOS_API void RemoveSimulationObjects();
 		CHAOS_API TArray<Chaos::TVec3<FSolverReal>> ComputeParticleTargets(const TArray<TArray<int32>>& ParticleIndices);
 
-
 		/*Debug Output*/
 		CHAOS_API void DebugDrawSimulationData();
 		CHAOS_API void DebugDrawTetrahedralParticles(FFleshThreadingProxy& Proxy);
@@ -158,6 +168,7 @@ namespace Chaos::Softs
 		CHAOS_API void RemoveProxy(FThreadingProxy* InObject);
 		CHAOS_API TUniquePtr<FDeformablePackage> PullOutputPackage();
 		CHAOS_API void PushInputPackage(int32 Frame, FDeformableDataMap&& InPackage);
+		CHAOS_API void PushRestartPackage(int32 InFrame, FDeformableDataMap&& InPackage);
 
 		const FDeformableSolverProperties& GetProperties() const { return Property; }
 
@@ -168,19 +179,22 @@ namespace Chaos::Softs
 		static CHAOS_API FCriticalSection	RemovalMutex; // @todo(flesh) : change to threaded commands to prevent the lock. 
 		static CHAOS_API FCriticalSection	PackageOutputMutex;
 		static CHAOS_API FCriticalSection	PackageInputMutex;
+		static CHAOS_API FCriticalSection   PackageRestartMutex;
 		static CHAOS_API FCriticalSection	SolverEnabledMutex;
 
 		TArray< FThreadingProxy* > RemovedProxys_Internal;
 		TArray< FThreadingProxy* > UninitializedProxys_Internal;
 		TArray< TUniquePtr<FDeformablePackage>  > BufferedInputPackages;
 		TArray< TUniquePtr<FDeformablePackage>  > BufferedOutputPackages;
+		TArray< TUniquePtr<FDeformablePackage>  > BufferedRestartPackages;
 		TUniquePtr < FDeformablePackage > CurrentInputPackage;
 		TUniquePtr < FDeformablePackage > PreviousInputPackage;
+		TUniquePtr < FDeformablePackage > CurrentRestartPackage;
+		bool bPendingRestart = false;
 
 		TSet< const UObject* > InitializedObjects_External;
 		TMap< FThreadingProxy::FKey, TUniquePtr<FThreadingProxy> > Proxies;
 
-
 		// User Configuration
 		FDeformableSolverProperties Property;
 
diff --git a/Engine/Source/Runtime/Experimental/Chaos/Public/Chaos/Deformable/ChaosDeformableSolverProxy.h b/Engine/Source/Runtime/Experimental/Chaos/Public/Chaos/Deformable/ChaosDeformableSolverProxy.h
--- a/Engine/Source/Runtime/Experimental/Chaos/Public/Chaos/Deformable/ChaosDeformableSolverProxy.h
+++ b/Engine/Source/Runtime/Experimental/Chaos/Public/Chaos/Deformable/ChaosDeformableSolverProxy.h
@@ -5,7 +5,7 @@
 #include "UObject/ObjectMacros.h"
 #include "GeometryCollection/GeometryCollection.h"
 #include "GeometryCollection/ManagedArrayCollection.h"
-
+#include "Chaos/Range.h"
 #include "ChaosDeformableSolverProxy.generated.h"
 
 /**
@@ -86,7 +86,7 @@ namespace Chaos::Softs
 
 		const FManagedArrayCollection& Rest;
 		FManagedArrayCollection Dynamic;
-		FIntVector2 SolverParticleRange = FIntVector2(0, 0);
+		Chaos::FRange SolverParticleRange = Chaos::FRange(0, 0);
 		ChaosDeformableSimSpace SimSpace = ChaosDeformableSimSpace::World;
 		bool bIsCached = false;
 
@@ -235,11 +235,10 @@ namespace Chaos::Softs
 
 		void SetSolverParticleRange(int32 InStart, int32 InRange)
 		{
-			SolverParticleRange[0] = InStart;
-			SolverParticleRange[1] = InRange;
+			SolverParticleRange = Chaos::FRange(InStart, InRange);
 		}
 
-		const FIntVector2& GetSolverParticleRange() const
+		const Chaos::FRange& GetSolverParticleRange() const
 		{
 			return SolverParticleRange;
 		}
@@ -346,5 +345,22 @@ namespace Chaos::Softs
 			FManagedArrayCollection Dynamic;
 		};
 
+		class FFleshRestartBuffer : public FThreadingProxy::FBuffer
+		{
+			typedef FThreadingProxy::FBuffer Super;
+
+		public:
+			typedef FFleshThreadingProxy Source;
+			FFleshRestartBuffer(const FManagedArrayCollection& InDynamicCollection,
+				const UObject* InOwner = nullptr)
+				: Super(InOwner, FFleshThreadingProxy::TypeName())
+			{
+				Dynamic.AddAttribute<FVector3f>("Vertex", FGeometryCollection::VerticesGroup);
+				Dynamic.CopyMatchingAttributesFrom(InDynamicCollection);
+			}
+			virtual ~FFleshRestartBuffer() {}
+
+			FManagedArrayCollection Dynamic;
+		};
 	};
 }// namespace Chaos::Softs
diff --git a/Engine/Source/Runtime/Experimental/Chaos/Public/Chaos/Deformable/GaussSeidelDynamicWeakConstraints.h b/Engine/Source/Runtime/Experimental/Chaos/Public/Chaos/Deformable/GaussSeidelDynamicWeakConstraints.h
--- a/Engine/Source/Runtime/Experimental/Chaos/Public/Chaos/Deformable/GaussSeidelDynamicWeakConstraints.h
+++ b/Engine/Source/Runtime/Experimental/Chaos/Public/Chaos/Deformable/GaussSeidelDynamicWeakConstraints.h
@@ -45,16 +45,16 @@ namespace Chaos::Softs
 			}
 		}
 		
-		const GSConstraintHandle* AddSingleParticleTetrahedraConstraint(const GeometryCollection::Facades::FTetrahedralFacade& InTargetGeom, const ParticleType& AllParticles, const GeometryCollection::Facades::TetrahedralParticleEmbedding& InIntersection, const int32 InConstraintIndex, const FIntVector2& SourceRange, const FIntVector2& TargetRange, const T PositionTargetStiffness)
+		const GSConstraintHandle* AddSingleParticleTetrahedraConstraint(const GeometryCollection::Facades::FTetrahedralFacade& InTargetGeom, const ParticleType& AllParticles, const GeometryCollection::Facades::TetrahedralParticleEmbedding& InIntersection, const int32 InConstraintIndex, const Chaos::FRange& SourceRange, const Chaos::FRange& TargetRange, const T PositionTargetStiffness)
 		{
 			GSConstraintHandle* NewConstraintPtr = new GSConstraintHandle();
 			NewConstraintPtr -> ConstraintIndex = InConstraintIndex;
 			Handles[InConstraintIndex] = NewConstraintPtr;
 
-			const FIntVector4 TargetTet = InTargetGeom.Tetrahedron[InTargetGeom.TetrahedronStart[InIntersection.GeometryIndex] + InIntersection.TetrahedronIndex] + FIntVector4(TargetRange[0]);
+			const FIntVector4 TargetTet = InTargetGeom.Tetrahedron[InTargetGeom.TetrahedronStart[InIntersection.GeometryIndex] + InIntersection.TetrahedronIndex] + FIntVector4(TargetRange.Start);
 
 			FGaussSeidelWeakConstraintSingleData<T> SingleConstraintData;
-			SingleConstraintData.SingleIndices = { InIntersection.ParticleIndex + SourceRange[0] };
+			SingleConstraintData.SingleIndices = { InIntersection.ParticleIndex + SourceRange.Start };
 			SingleConstraintData.SingleSecondIndices =  {TargetTet[0], TargetTet[1], TargetTet[2], TargetTet[3]};
 			SingleConstraintData.SingleWeights = {(T)1.};
 			SingleConstraintData.SingleSecondWeights = {InIntersection.BarycentricWeights[0], InIntersection.BarycentricWeights[1], InIntersection.BarycentricWeights[2], InIntersection.BarycentricWeights[3]};
@@ -118,7 +118,7 @@ namespace Chaos::Softs
 			
 		}
 
-		TArray<const GSConstraintHandle*> AddParticleTetrahedraConstraints(const GeometryCollection::Facades::FTetrahedralFacade& InTargetGeom, const ParticleType& AllParticles, const TArray<GeometryCollection::Facades::TetrahedralParticleEmbedding>& InIntersections, const FIntVector2& SourceRange, const FIntVector2& TargetRange, const T ConstraintStiffness)
+		TArray<const GSConstraintHandle*> AddParticleTetrahedraConstraints(const GeometryCollection::Facades::FTetrahedralFacade& InTargetGeom, const ParticleType& AllParticles, const TArray<GeometryCollection::Facades::TetrahedralParticleEmbedding>& InIntersections, const Chaos::FRange& SourceRange, const Chaos::FRange& TargetRange, const T ConstraintStiffness)
 		{
 			TArray<const GSConstraintHandle*> ConstraintHandles;
 			ConstraintHandles.SetNum(InIntersections.Num());
diff --git a/Engine/Source/Runtime/Experimental/Chaos/Public/Chaos/Deformable/MuscleActivationConstraints.h b/Engine/Source/Runtime/Experimental/Chaos/Public/Chaos/Deformable/MuscleActivationConstraints.h
--- a/Engine/Source/Runtime/Experimental/Chaos/Public/Chaos/Deformable/MuscleActivationConstraints.h
+++ b/Engine/Source/Runtime/Experimental/Chaos/Public/Chaos/Deformable/MuscleActivationConstraints.h
@@ -12,17 +12,14 @@ namespace Chaos::Softs
 	typedef GeometryCollection::Facades::FMuscleActivationFacade Facade;
 	typedef GeometryCollection::Facades::FMuscleActivationData Data;
 	template <typename T, typename ParticleType>
-	struct FMuscleActivationConstraints
+	class FMuscleActivationConstraints
 	{
+	public:
 		//Handles muscle activation data
 		FMuscleActivationConstraints(){}
 
 		virtual ~FMuscleActivationConstraints() {}
 
-		void Init(const FSolverParticles& InParticles, const T Dt) const
-		{
-		}
-
 		void AddMuscles(const Facade& FMuscleActivation, int32 VertexOffset = 0, int32 ElementOffset = 0)
 		{
 			for (int MuscleIdx = 0; MuscleIdx < FMuscleActivation.NumMuscles(); MuscleIdx++)
@@ -47,23 +44,23 @@ namespace Chaos::Softs
 					MuscleActivationData.OriginInsertionPair[1] += VertexOffset;
 					OriginInsertionPair.Add(MuscleActivationData.OriginInsertionPair);
 					OriginInsertionRestLength.Add(MuscleActivationData.OriginInsertionRestLength);
-					Activation.Add(0.f);
+					FiberLengthRatio.Add(1.f);
+
+					MuscleVertexOffset.Add(FMuscleActivation.MuscleVertexOffset(MuscleIdx));
+					MuscleVertexCount.Add(FMuscleActivation.NumMuscleVertices(MuscleIdx));
 				}
 			}
 		}
 
 		void UpdateLengthBasedMuscleActivation(const ParticleType& InParticles)
 		{
-			float AverageFiberLengthScale = 0;
 			for (int32 MuscleIdx = 0; MuscleIdx < MuscleActivationElement.Num(); MuscleIdx++)
 			{
 				if (ensureMsgf(MuscleActivationElement[MuscleIdx].Num() == FiberDirectionMatrix[MuscleIdx].Num(),
 					TEXT("MuscleActivationElement[%d].Num() = %d, not equal to FiberDirectionMatrix[%d].Num() = %d"), MuscleIdx, MuscleActivationElement[MuscleIdx].Num(), MuscleIdx, FiberDirectionMatrix[MuscleIdx].Num()))
 				{
-					//Infer activation
-					float CurrentFiberLengthScale = (InParticles.X(OriginInsertionPair[MuscleIdx][0]) - InParticles.X(OriginInsertionPair[MuscleIdx][1])).Size() / OriginInsertionRestLength[MuscleIdx];
-					AverageFiberLengthScale += CurrentFiberLengthScale;
-					Activation[MuscleIdx] = CurrentFiberLengthScale * CurrentFiberLengthScale;
+					//update origin/insertion length
+					FiberLengthRatio[MuscleIdx] = (InParticles.X(OriginInsertionPair[MuscleIdx][0]) - InParticles.X(OriginInsertionPair[MuscleIdx][1])).Size() / OriginInsertionRestLength[MuscleIdx];
 				}
 			}
 		}
@@ -77,18 +74,29 @@ namespace Chaos::Softs
 				{
 					for (int32 ElemIdx = 0; ElemIdx < MuscleActivationElement[MuscleIdx].Num(); ElemIdx++)
 					{
-						Constraints.ModifyDmInverseFromFiberLength(MuscleActivationElement[MuscleIdx][ElemIdx], Activation[MuscleIdx], FiberDirectionMatrix[MuscleIdx][ElemIdx], ContractionVolumeScale[MuscleIdx][ElemIdx]);
+						Constraints.ModifyDmInverseFromFiberLength(MuscleActivationElement[MuscleIdx][ElemIdx], FiberLengthRatio[MuscleIdx] * FiberLengthRatio[MuscleIdx], FiberDirectionMatrix[MuscleIdx][ElemIdx], ContractionVolumeScale[MuscleIdx][ElemIdx]);
 					}
 				}
 			}
 		}
-		
+
+		int32 NumMuscles() { return MuscleActivationElement.Num(); };
+		int32 GetMuscleVertexOffset(int32 MuscleIndex) { return MuscleVertexOffset[MuscleIndex]; };
+		int32 GetMuscleVertexCount(int32 MuscleIndex) { return MuscleVertexCount[MuscleIndex]; };
+
+		/* returns muscle activation in [0,1]. max muscle activation 1 at 0.75 fiber length ratio or less */
+		float GetMuscleActivation(int32 MuscleIndex) { return FMath::Max(FMath::Min(4.f * (1.f - FiberLengthRatio[MuscleIndex]), 1), 0); };
+
+	private:	
 		TArray<TArray<int32>> MuscleActivationElement;
 		TArray<FIntVector2> OriginInsertionPair;
 		TArray<float> OriginInsertionRestLength;
-		TArray<float> Activation;
+		TArray<float> FiberLengthRatio;
 		TArray<TArray<Chaos::PMatrix33d>> FiberDirectionMatrix;
 		TArray<TArray<float>> ContractionVolumeScale;
+
+		TArray<int32> MuscleVertexOffset;
+		TArray<int32> MuscleVertexCount;
 	};
 
 
diff --git a/Engine/Source/Runtime/Experimental/Chaos/Public/GeometryCollection/Facades/CollectionMuscleActivationFacade.h b/Engine/Source/Runtime/Experimental/Chaos/Public/GeometryCollection/Facades/CollectionMuscleActivationFacade.h
--- a/Engine/Source/Runtime/Experimental/Chaos/Public/GeometryCollection/Facades/CollectionMuscleActivationFacade.h
+++ b/Engine/Source/Runtime/Experimental/Chaos/Public/GeometryCollection/Facades/CollectionMuscleActivationFacade.h
@@ -59,6 +59,9 @@ namespace GeometryCollection::Facades
 		CHAOS_API bool IsValidGeometryIndex(const int32 Index) const { return 0 <= Index && Index < ConstCollection.NumElements(FGeometryCollection::GeometryGroup); };
 		CHAOS_API bool IsValidElementIndex(const int32 Index) const { return 0 <= Index && Index < ConstCollection.NumElements("Tetrahedral"); };
 		CHAOS_API int32 NumMuscles() const { return MuscleActivationElementAttribute.Num(); }
+		CHAOS_API int32 MuscleVertexOffset(const int32 MuscleIndex) const;
+		CHAOS_API int32 NumMuscleVertices(const int32 MuscleIndex) const;
+		CHAOS_API int32 RemoveInvalidMuscles();
 		CHAOS_API bool SetUpMuscleActivation(const TArray<int32>& Origin, const TArray<int32>& Insertion, float ContractionVolumeScale);
 		CHAOS_API TArray<TArray<TArray<FVector3f>>> BuildStreamlines(const TArray<int32>& Origin, const TArray<int32>& Insertion,
 			 int32 NumLinesMultiplier, int32 MaxStreamlineIterations, int32 MaxPointsPerLine);
diff --git a/Engine/Source/Runtime/Experimental/ChaosCore/Public/Chaos/Range.h b/Engine/Source/Runtime/Experimental/ChaosCore/Public/Chaos/Range.h
new file mode 100644
--- /dev/null
+++ b/Engine/Source/Runtime/Experimental/ChaosCore/Public/Chaos/Range.h
@@ -0,0 +1,19 @@
+// Copyright Epic Games, Inc. All Rights Reserved.
+#pragma once
+
+#include "CoreTypes.h"
+
+namespace Chaos
+{
+struct FRange
+{
+	FRange(int32 InStart, int32 InCount) : Start(InStart), Count(InCount) {}
+	int32 Start;
+	int32 Count;
+};
+
+inline bool operator==(const FRange& First, const FRange& Second)
+{
+	return First.Start == Second.Start && First.Count == Second.Count;
+}
+}
diff --git a/Engine/Source/Runtime/Experimental/Dataflow/Engine/Public/Dataflow/DataflowPreview.h b/Engine/Source/Runtime/Experimental/Dataflow/Engine/Public/Dataflow/DataflowPreview.h
--- a/Engine/Source/Runtime/Experimental/Dataflow/Engine/Public/Dataflow/DataflowPreview.h
+++ b/Engine/Source/Runtime/Experimental/Dataflow/Engine/Public/Dataflow/DataflowPreview.h
@@ -1,4 +1,4 @@
-﻿// Copyright Epic Games, Inc. All Rights Reserved.
+// Copyright Epic Games, Inc. All Rights Reserved.
 
 #pragma once
 
@@ -22,10 +22,18 @@ public :
 	UPROPERTY(EditAnywhere, Category="Caching")
 	int32 FrameRate = 30;
 
-	/** Number of sampling frames per second for caching*/
+	/** Time range of the simulation*/
 	UPROPERTY(EditAnywhere, Category="Caching")
 	FVector2f TimeRange = FVector2f(0.0f, 5.0f);
 
+	/** If enabled, the simulation will restart on Time Range without modifying anything outside of Time Range */
+	UPROPERTY(EditAnywhere, AdvancedDisplay, Category = "Caching")
+	bool bRestartSimulation = false;
+
+	/** Time range of the restarted simulation*/
+	UPROPERTY(EditAnywhere, Category = "Caching", meta = (EditCondition = "bRestartSimulation"))
+	FVector2f RestartTimeRange = FVector2f(0.0f, 0.0f);
+
 	/** Boolean to check if the caching will be done on an async thread (if yes no GT dependency) */
 	UPROPERTY(EditAnywhere, AdvancedDisplay, Category = "Caching")
 	bool bAsyncCaching = true;
diff --git a/Engine/Source/Runtime/Experimental/Dataflow/Simulation/Private/Dataflow/DataflowSimulationManager.cpp b/Engine/Source/Runtime/Experimental/Dataflow/Simulation/Private/Dataflow/DataflowSimulationManager.cpp
--- a/Engine/Source/Runtime/Experimental/Dataflow/Simulation/Private/Dataflow/DataflowSimulationManager.cpp
+++ b/Engine/Source/Runtime/Experimental/Dataflow/Simulation/Private/Dataflow/DataflowSimulationManager.cpp
@@ -411,6 +411,23 @@ void UDataflowSimulationManager::WriteSimulationInterfaces(const float DeltaTime
 	//Dataflow::PostSimulationTick(GetWorld(), GetWorld()->GetTimeSeconds(), DeltaTime);
 }
 
+void UDataflowSimulationManager::ReadRestartData()
+{
+	for (const TPair<TObjectPtr<UDataflow>, UE::Dataflow::Private::FDataflowSimulationData>& DataflowData : SimulationData)
+	{
+		for (const TPair<FString, TSet<IDataflowSimulationInterface*>>& SimulationInterfaces : DataflowData.Value.SimulationInterfaces)
+		{
+			for (IDataflowSimulationInterface* SimulationInterface : SimulationInterfaces.Value)
+			{
+				if (SimulationInterface)
+				{
+					SimulationInterface->ReadRestartData();
+				}
+			}
+		}
+	}
+}
+
 void UDataflowSimulationManager::AdvanceSimulationProxies(const float DeltaTime, const float SimulationTime)
 {
 	for(TPair<TObjectPtr<UDataflow>, UE::Dataflow::Private::FDataflowSimulationData>& DataflowData : SimulationData)
diff --git a/Engine/Source/Runtime/Experimental/Dataflow/Simulation/Public/Dataflow/DataflowSimulationInterface.h b/Engine/Source/Runtime/Experimental/Dataflow/Simulation/Public/Dataflow/DataflowSimulationInterface.h
--- a/Engine/Source/Runtime/Experimental/Dataflow/Simulation/Public/Dataflow/DataflowSimulationInterface.h
+++ b/Engine/Source/Runtime/Experimental/Dataflow/Simulation/Public/Dataflow/DataflowSimulationInterface.h
@@ -1,4 +1,4 @@
-﻿// Copyright Epic Games, Inc. All Rights Reserved.
+// Copyright Epic Games, Inc. All Rights Reserved.
 
 #pragma once
 
@@ -73,6 +73,9 @@ public:
 	/** Read data received from the simulation proxy */
 	virtual void ReadFromSimulation(const float DeltaTime, const bool bAsyncTask) {};
 
+	/** Read restart data (positions) from simulation proxy */
+	virtual void ReadRestartData() {};
+
 	/** Postprocess data after simulation */
 	virtual void PostProcessSimulation(const float DeltaTime) {};
 	
diff --git a/Engine/Source/Runtime/Experimental/Dataflow/Simulation/Public/Dataflow/DataflowSimulationManager.h b/Engine/Source/Runtime/Experimental/Dataflow/Simulation/Public/Dataflow/DataflowSimulationManager.h
--- a/Engine/Source/Runtime/Experimental/Dataflow/Simulation/Public/Dataflow/DataflowSimulationManager.h
+++ b/Engine/Source/Runtime/Experimental/Dataflow/Simulation/Public/Dataflow/DataflowSimulationManager.h
@@ -1,4 +1,4 @@
-﻿// Copyright Epic Games, Inc. All Rights Reserved.
+// Copyright Epic Games, Inc. All Rights Reserved.
 
 #pragma once
 
@@ -91,6 +91,9 @@ public:
 	/** Read all the data from the simulation proxies and write the result onto the interfaces (received from PT->GT) */
 	void WriteSimulationInterfaces(const float DeltaTime, const bool bAsyncTask);
 
+	/** Read the restart data and Write it to the simulation proxies (to be send from GT->PT) */
+	void ReadRestartData();
+
 	/** Init all the simulation interfaces*/
 	void InitSimulationInterfaces();
 
diff --git a/Engine/Source/Runtime/Experimental/Dataflow/Simulation/Public/Dataflow/Interfaces/DataflowPhysicsSolver.h b/Engine/Source/Runtime/Experimental/Dataflow/Simulation/Public/Dataflow/Interfaces/DataflowPhysicsSolver.h
--- a/Engine/Source/Runtime/Experimental/Dataflow/Simulation/Public/Dataflow/Interfaces/DataflowPhysicsSolver.h
+++ b/Engine/Source/Runtime/Experimental/Dataflow/Simulation/Public/Dataflow/Interfaces/DataflowPhysicsSolver.h
@@ -1,4 +1,4 @@
-﻿// Copyright Epic Games, Inc. All Rights Reserved.
+// Copyright Epic Games, Inc. All Rights Reserved.
 
 #pragma once
 #include "Dataflow/DataflowSimulationInterface.h"
@@ -21,7 +21,10 @@ struct DATAFLOWSIMULATION_API FDataflowPhysicsSolverProxy : public FDataflowSimu
 	{
 		return StaticStruct();
 	}
-	
+
+	/** Advance the solver datas in time */
+	virtual void ReadRestartData() {};
+
 	/** Advance the solver datas in time */
 	virtual void AdvanceSolverDatas(const float DeltaTime) {}
 
