Datflow object : reconcile pin connections when rtefresh EdNode

diff --git a/Engine/Source/Runtime/Experimental/Dataflow/Engine/Private/Dataflow/DataflowObject.cpp b/Engine/Source/Runtime/Experimental/Dataflow/Engine/Private/Dataflow/DataflowObject.cpp
--- a/Engine/Source/Runtime/Experimental/Dataflow/Engine/Private/Dataflow/DataflowObject.cpp
+++ b/Engine/Source/Runtime/Experimental/Dataflow/Engine/Private/Dataflow/DataflowObject.cpp
@@ -405,6 +405,7 @@ void UDataflow::RefreshEdNode(TObjectPtr<UDataflowEdNode> EdNode)
 	if (EdNode)
 	{
 		EdNode->UpdatePinsFromDataflowNode();
+		EdNode->UpdatePinsConnectionsFromDataflowNode();
 		if (Nodes.Contains(EdNode))
 		{
 			NotifyNodeChanged(EdNode);
