Add a way to update an existing FMultiEndpointHttpClient connection if we want to change the current host being used.

#rnx

- Add a new public method ::UpdateConnection that will find the current connection for the given hostgroup and if such a connection exists it will be set to use the hostgroup's current primary host.
-- This is useful to call when the primary host for the group has changed or if the caller believes that connection has been re-established with a previously failed host and we want to try it again.
- Added a new method ::FindConnection that will return the connection for a hostgroup is one already exists. If no connection exists then a nullptr is returned.
- Condensed duplicate code for creating a new connection to a common method ::CreateConnection

diff --git a/Engine/Source/Runtime/Experimental/IoStore/OnDemand/Private/OnDemandHttpClient.cpp b/Engine/Source/Runtime/Experimental/IoStore/OnDemand/Private/OnDemandHttpClient.cpp
--- a/Engine/Source/Runtime/Experimental/IoStore/OnDemand/Private/OnDemandHttpClient.cpp
+++ b/Engine/Source/Runtime/Experimental/IoStore/OnDemand/Private/OnDemandHttpClient.cpp
@@ -462,6 +462,19 @@ bool FMultiEndpointHttpClient::Tick(int32 WaitTimeMs, uint32 MaxKiBPerSecond)
 	return bIsIdle == false;
 }
 
+void FMultiEndpointHttpClient::UpdateConnection(const FOnDemandHostGroup& HostGroup)
+{
+	if (FConnection* Connection = FindConnection(HostGroup))
+	{
+		Connection->CurrentHost = HostGroup.PrimaryHostIndex();
+
+		if (Connection->Pools[Connection->CurrentHost].IsValid() == false)
+		{
+			Connection->Pools[Connection->CurrentHost] = CreateConnection(HostGroup.PrimaryHost());
+		}
+	}
+}
+
 void FMultiEndpointHttpClient::IssueRequest(FRequest&& Request)
 {
 	UE_CHECK_THREADSAFETY(FMultiEndpointHttpClient::IssueRequest);
@@ -551,14 +564,7 @@ void FMultiEndpointHttpClient::RetryRequest(FRequest&& Request)
 		FAnsiStringView HostUrl = Request.Connection.HostGroup.CycleHost(Connection.CurrentHost);
 		if (Connection.Pools[Connection.CurrentHost].IsValid() == false)
 		{
-			HTTP::FConnectionPool::FParams Params;
-			ensure(Params.SetHostFromUrl(HostUrl) >= 0);
-			if (Config.ReceiveBufferSize >= 0)
-			{
-				Params.RecvBufSize = Config.ReceiveBufferSize;
-			}
-			Params.ConnectionCount = uint16(Config.MaxConnectionCount);
-			Connection.Pools[Connection.CurrentHost] = MakeUnique<HTTP::FConnectionPool>(Params);
+			Connection.Pools[Connection.CurrentHost] = CreateConnection(HostUrl);
 		}
 	}
 
@@ -568,6 +574,19 @@ void FMultiEndpointHttpClient::RetryRequest(FRequest&& Request)
 	Retries.Emplace(MoveTemp(Request));
 }
 
+TUniquePtr<HTTP::FConnectionPool> FMultiEndpointHttpClient::CreateConnection(FAnsiStringView HostUrl) const
+{
+	HTTP::FConnectionPool::FParams Params;
+	ensure(Params.SetHostFromUrl(HostUrl) >= 0);
+	if (Config.ReceiveBufferSize >= 0)
+	{
+		Params.RecvBufSize = Config.ReceiveBufferSize;
+	}
+	Params.ConnectionCount = uint16(Config.MaxConnectionCount);
+
+	return MakeUnique<HTTP::FConnectionPool>(Params);
+}
+
 FMultiEndpointHttpClient::FConnection& FMultiEndpointHttpClient::GetConnection(const FOnDemandHostGroup& HostGroup)
 {
 	UE_CHECK_THREADSAFETY(FMultiEndpointHttpClient::GetConnection);
@@ -589,16 +608,24 @@ FMultiEndpointHttpClient::FConnection& FMultiEndpointHttpClient::GetConnection(c
 	Conn.Pools.SetNum(HostGroup.Hosts().Num());
 	Conn.CurrentHost = HostGroup.PrimaryHostIndex();
 
-	HTTP::FConnectionPool::FParams Params;
-	ensure(Params.SetHostFromUrl(HostGroup.PrimaryHost()) >= 0);
-	if (Config.ReceiveBufferSize >= 0)
+	Conn.Pools[Conn.CurrentHost] = CreateConnection(HostGroup.PrimaryHost());
+	return Conn;
+}
+
+FMultiEndpointHttpClient::FConnection* FMultiEndpointHttpClient::FindConnection(const FOnDemandHostGroup& HostGroup)
+{
+	UE_CHECK_THREADSAFETY(FMultiEndpointHttpClient::FindConnection);
+
+	for (TUniquePtr<FConnection>& Conn : Connections)
 	{
-		Params.RecvBufSize = Config.ReceiveBufferSize;
+		check(Conn.IsValid());
+		if (Conn->HostGroup == HostGroup)
+		{
+			return Conn.Get();
+		}
 	}
-	Params.ConnectionCount = uint16(Config.MaxConnectionCount);
 
-	Conn.Pools[Conn.CurrentHost] = MakeUnique<HTTP::FConnectionPool>(Params);
-	return Conn;
+	return nullptr;
 }
 
 void FMultiEndpointHttpClient::Log(const FRequest& Request)
diff --git a/Engine/Source/Runtime/Experimental/IoStore/OnDemand/Private/OnDemandHttpClient.h b/Engine/Source/Runtime/Experimental/IoStore/OnDemand/Private/OnDemandHttpClient.h
--- a/Engine/Source/Runtime/Experimental/IoStore/OnDemand/Private/OnDemandHttpClient.h
+++ b/Engine/Source/Runtime/Experimental/IoStore/OnDemand/Private/OnDemandHttpClient.h
@@ -124,12 +124,14 @@ public:
 
 	static TUniquePtr<FMultiEndpointHttpClient> Create(const FMultiEndpointHttpClientConfig& Config);
 
-	void Get(const FOnDemandHostGroup& Endpoint, FAnsiStringView RelativeUrl, FOnHttpResponse&& OnResponse);
-	void Get(const FOnDemandHostGroup& Endpoint, FAnsiStringView RelativeUrl, const FIoOffsetAndLength& ChunkRange, FOnHttpResponse&& OnResponse);
+	void Get(const FOnDemandHostGroup& HostGroup, FAnsiStringView RelativeUrl, FOnHttpResponse&& OnResponse);
+	void Get(const FOnDemandHostGroup& HostGroup, FAnsiStringView RelativeUrl, const FIoOffsetAndLength& ChunkRange, FOnHttpResponse&& OnResponse);
 
 	bool Tick(int32 WaitTimeMs, uint32 MaxKiBPerSecond);
 	bool Tick() { return Tick(-1, 0); }
 
+	void UpdateConnection(const FOnDemandHostGroup& HostGroup);
+
 private:
 	using FHttpConnectionPools = TArray<TUniquePtr<HTTP::FConnectionPool>, TInlineAllocator<4>>;
 
@@ -153,16 +155,20 @@ private:
 		int32				Host = INDEX_NONE;
 	};
 
-									FMultiEndpointHttpClient(const FMultiEndpointHttpClientConfig& Config);
-	void							IssueRequest(FRequest&& Request);
-	void							RetryRequest(FRequest&& Request);
-	FConnection&					GetConnection(const FOnDemandHostGroup& Endpoint);
-	void							Log(const FRequest& Request);
+										FMultiEndpointHttpClient(const FMultiEndpointHttpClientConfig& Config);
+	void								IssueRequest(FRequest&& Request);
+	void								RetryRequest(FRequest&& Request);
+
+	TUniquePtr<HTTP::FConnectionPool>	CreateConnection(FAnsiStringView HostUrl) const;
+	FConnection&						GetConnection(const FOnDemandHostGroup& HostGroup);
+	FConnection*						FindConnection(const FOnDemandHostGroup& HostGroup);
+
+	void								Log(const FRequest& Request);
 
-	FMultiEndpointHttpClientConfig	Config;
-	TArray<TUniquePtr<FConnection>>	Connections;
-	HTTP::FEventLoop				EventLoop;
-	TArray<FRequest>				Retries;
+	FMultiEndpointHttpClientConfig		Config;
+	TArray<TUniquePtr<FConnection>>		Connections;
+	HTTP::FEventLoop					EventLoop;
+	TArray<FRequest>					Retries;
 
 #if UE_HTTPCLIENT_THREADSAFE_CHECKS
 	uint32							OwningThread = 0;
