Log reasons for endpoint resolve failures.
#rnx

diff --git a/Engine/Source/Runtime/Experimental/IoStore/OnDemand/Private/DistributionEndpoints.cpp b/Engine/Source/Runtime/Experimental/IoStore/OnDemand/Private/DistributionEndpoints.cpp
--- a/Engine/Source/Runtime/Experimental/IoStore/OnDemand/Private/DistributionEndpoints.cpp
+++ b/Engine/Source/Runtime/Experimental/IoStore/OnDemand/Private/DistributionEndpoints.cpp
@@ -59,6 +59,12 @@ FDistributionEndpoints::EResult FDistributionEndpoints::ResolveEndpoints(const F
 		Dest=FIoBuffer()
 		] (const FTicketStatus& Status) mutable
 	{
+		if (Status.GetId() == FTicketStatus::EId::Error)
+		{
+			const char* ErrorReason = Status.GetErrorReason();
+			UE_LOG(LogIas, Log, TEXT("Resolve endpoint error: '%S'"), ErrorReason);
+		}
+
 		if (Status.GetId() >= FTicketStatus::EId::Cancelled)
 		{
 			return;
@@ -75,6 +81,8 @@ FDistributionEndpoints::EResult FDistributionEndpoints::ResolveEndpoints(const F
 		FResponse& Response = Status.GetResponse();
 		if (Response.GetStatus() != EStatusCodeClass::Successful)
 		{
+			uint32 StatusCode = Response.GetStatusCode();
+			UE_LOG(LogIas, Log, TEXT("Resolve HTTP status code: '%d'"), StatusCode);
 			return;
 		}
 
