Surface outcome error codes

#rnx

diff --git a/Engine/Source/Runtime/Experimental/IoStore/HttpClient/Private/Activity.inl b/Engine/Source/Runtime/Experimental/IoStore/HttpClient/Private/Activity.inl
--- a/Engine/Source/Runtime/Experimental/IoStore/HttpClient/Private/Activity.inl
+++ b/Engine/Source/Runtime/Experimental/IoStore/HttpClient/Private/Activity.inl
@@ -180,12 +180,19 @@ static void Activity_Free(FActivity* Activity)
 }
 
 ////////////////////////////////////////////////////////////////////////////////
-static void Activity_SetError(FActivity* Activity, const char* Reason)
+static void Activity_SetError(FActivity* Activity, const char* Reason, int32 Code=-1)
 {
 	Activity->IsKeepAlive = 0;
 	Activity->ErrorReason = Reason;
 
-	Activity_ChangeState(Activity, FActivity::EState::Failed, LastSocketResult());
+	Code = (Code < 0) ? LastSocketResult() : Code;
+	Activity_ChangeState(Activity, FActivity::EState::Failed, Code);
+}
+
+////////////////////////////////////////////////////////////////////////////////
+static void Activity_SetError(FActivity* Activity, const FOutcome& Outcome)
+{
+	Activity_SetError(Activity, Outcome.GetMessage().GetData(), Outcome.GetResult());
 }
 
 // }}}
diff --git a/Engine/Source/Runtime/Experimental/IoStore/HttpClient/Private/Loop.inl b/Engine/Source/Runtime/Experimental/IoStore/HttpClient/Private/Loop.inl
--- a/Engine/Source/Runtime/Experimental/IoStore/HttpClient/Private/Loop.inl
+++ b/Engine/Source/Runtime/Experimental/IoStore/HttpClient/Private/Loop.inl
@@ -33,7 +33,7 @@ static FOutcome DoSend(FActivity* Activity, FPeerType& Peer)
 
 	if (Outcome.IsError())
 	{
-		Activity_SetError(Activity, Outcome.GetMessage().GetData());
+		Activity_SetError(Activity, Outcome);
 		return Outcome;
 	}
 
@@ -84,7 +84,7 @@ static FOutcome DoRecvPeer(FActivity* Activity, FPeerType& Peer, int32& MaxRecvS
 
 	if (Outcome.IsError())
 	{
-		Activity_SetError(Activity, Outcome.GetMessage().GetData());
+		Activity_SetError(Activity, Outcome);
 		return Outcome;
 	}
 
@@ -121,7 +121,7 @@ static FOutcome DoRecvMessage(FActivity* Activity, FPeerType& Peer)
 
 		if (Outcome.IsError())
 		{
-			Activity_SetError(Activity, Outcome.GetMessage().GetData());
+			Activity_SetError(Activity, Outcome);
 			return Outcome;
 		}
 
diff --git a/Engine/Source/Runtime/Experimental/IoStore/HttpClient/Private/Socket.inl b/Engine/Source/Runtime/Experimental/IoStore/HttpClient/Private/Socket.inl
--- a/Engine/Source/Runtime/Experimental/IoStore/HttpClient/Private/Socket.inl
+++ b/Engine/Source/Runtime/Experimental/IoStore/HttpClient/Private/Socket.inl
@@ -380,7 +380,8 @@ FOutcome FSocket::Send(const char* Data, uint32 Size)
 		Result = getsockopt(Socket, SOL_SOCKET, SO_ERROR, (char*)&Error, &ErrorSize);
 		if (Result < 0 || Error != 0)
 		{
-			return FOutcome::Error("Error while connecting");
+			Result = (Result < 0) ? LastSocketResult() : Error;
+			return FOutcome::Error("Error while connecting", Result);
 		}
 
 		return FOutcome::Waiting();
