Optimize constraint sorting for determinism
A multitreaded sort has been added
Up to 6% improvement in the entire collision detection in some scenario

diff --git a/Engine/Source/Runtime/Experimental/Chaos/Private/Chaos/Collision/CollisionConstraintAllocator.cpp b/Engine/Source/Runtime/Experimental/Chaos/Private/Chaos/Collision/CollisionConstraintAllocator.cpp
--- a/Engine/Source/Runtime/Experimental/Chaos/Private/Chaos/Collision/CollisionConstraintAllocator.cpp
+++ b/Engine/Source/Runtime/Experimental/Chaos/Private/Chaos/Collision/CollisionConstraintAllocator.cpp
@@ -93,59 +93,30 @@ namespace Chaos
 		void FCollisionConstraintAllocator::ProcessNewConstraints()
 		{
 			QUICK_SCOPE_CYCLE_COUNTER(STAT_Collisions_ProcessNewConstraints);
-
 			FMemMark Mark(FMemStack::Get());
-			TArray<FPBDCollisionConstraint*, TMemStackAllocator<>> NewConstraints;
 
-			// Count the constraints so we can initialize the array
-			int32 NumNewConstraints = 0;
-			for (TUniquePtr<FCollisionContextAllocator>& ContextAllocator : ContextAllocators)
+			if (!ContextAllocators.IsEmpty())
 			{
-				NumNewConstraints += ContextAllocator->NewActiveConstraints.Num();
-			}
-			NewConstraints.Reserve(NumNewConstraints);
-
-			// Collect the constraints
-			int32 NewConstraintBegin = 0;
-			for (TUniquePtr<FCollisionContextAllocator>& ContextAllocator : ContextAllocators)
-			{
-				NewConstraints.Append(ContextAllocator->NewActiveConstraints);
-				ContextAllocator->NewActiveConstraints.Reset();
-			}
-			check(NewConstraints.Num() == NumNewConstraints);
+				const TArray<FPBDCollisionConstraint*>& NewConstraints = ContextAllocators[0]->NewActiveConstraints;
+				// Add the constraints to the active list (and the CCD list if appropriate)
+				for (FPBDCollisionConstraint* Constraint : NewConstraints)
+				{
+					FPBDCollisionConstraintContainerCookie& Cookie = Constraint->GetContainerCookie();
 
-			// For deterministic behaviour we need to sort the constraints so that any systems that iterate over
-			// active constraints will do so in a repeatable way.
-			// @todo(chaos): we could sort each context's array and then process them in order here instead
-			// @todo(chaos): if the context held the sort keys as well we could avoid some cache misses?
-			if (bIsDeteministic)
-			{
-				QUICK_SCOPE_CYCLE_COUNTER(STAT_Collisions_SortConstraints);
+					// Add the constraint to the active list and update its epoch
+					checkSlow(ActiveConstraints.Find(Constraint) == INDEX_NONE);
+					Cookie.ConstraintIndex = ActiveConstraints.Add(Constraint);
 
-				NewConstraints.Sort(
-					[](const FPBDCollisionConstraint& L, const FPBDCollisionConstraint& R)
+					// If the constraint uses CCD, keep it in another list so we don't have to search the full list for CCD contacts
+					if (Constraint->GetCCDEnabled())
 					{
-						return L.GetCollisionSortKey() < R.GetCollisionSortKey();
-					});
-			}
+						checkSlow(ActiveCCDConstraints.Find(Constraint) == INDEX_NONE);
+						Cookie.CCDConstraintIndex = ActiveCCDConstraints.Add(Constraint);
+					}
 
-			// Add the constraints to the active list (and the CCD list if appropriate)
-			for (FPBDCollisionConstraint* Constraint : NewConstraints)
-			{
-				FPBDCollisionConstraintContainerCookie& Cookie = Constraint->GetContainerCookie();
-
-				// Add the constraint to the active list and update its epoch
-				checkSlow(ActiveConstraints.Find(Constraint) == INDEX_NONE);
-				Cookie.ConstraintIndex = ActiveConstraints.Add(Constraint);
-
-				// If the constraint uses CCD, keep it in another list so we don't have to search the full list for CCD contacts
-				if (Constraint->GetCCDEnabled())
-				{
-					checkSlow(ActiveCCDConstraints.Find(Constraint) == INDEX_NONE);
-					Cookie.CCDConstraintIndex = ActiveCCDConstraints.Add(Constraint);
+					Cookie.LastUsedEpoch = CurrentEpoch;
 				}
-
-				Cookie.LastUsedEpoch = CurrentEpoch;
+				ContextAllocators[0]->NewActiveConstraints.Reset();
 			}
 		}
 
diff --git a/Engine/Source/Runtime/Experimental/Chaos/Private/Chaos/Collision/SpatialAccelerationBroadPhase.cpp b/Engine/Source/Runtime/Experimental/Chaos/Private/Chaos/Collision/SpatialAccelerationBroadPhase.cpp
--- a/Engine/Source/Runtime/Experimental/Chaos/Private/Chaos/Collision/SpatialAccelerationBroadPhase.cpp
+++ b/Engine/Source/Runtime/Experimental/Chaos/Private/Chaos/Collision/SpatialAccelerationBroadPhase.cpp
@@ -1,5 +1,7 @@
 // Copyright Epic Games, Inc. All Rights Reserved.
 #include "Chaos/Collision/SpatialAccelerationBroadPhase.h"
+#include "Chaos/PBDCollisionConstraints.h"
+
 
 namespace Chaos
 {
@@ -44,5 +46,79 @@ namespace Chaos
 			CVD_SCOPE_CONTEXT(CVDContext);
 			SpatialAccelerationBroadPhase.ProcessMidPhases(Dt, BroadPhaseContext);
 		}
+
+		FSortConstraintTask::FSortConstraintTask(Private::FBroadPhaseContext& InBroadPhaseContext)
+			: BroadPhaseContext(InBroadPhaseContext)
+		{
+			CVD_GET_CURRENT_CONTEXT(CVDContext);
+		}
+
+		ENamedThreads::Type FSortConstraintTask::GetDesiredThread()
+		{
+			return ENamedThreads::AnyNormalThreadNormalTask;
+		}
+
+		void FSortConstraintTask::Process()
+		{
+			QUICK_SCOPE_CYCLE_COUNTER(STAT_SortConstraintBatch);
+			CVD_SCOPE_CONTEXT(CVDContext);
+			BroadPhaseContext.CollisionContext.Allocator->NewActiveConstraints.Sort(
+				[](const FPBDCollisionConstraintHandle& L, const FPBDCollisionConstraintHandle& R)
+				{
+					return L.GetContact().GetCollisionSortKey() < R.GetContact().GetCollisionSortKey();
+				});
+		}
+
+		FMergeConstraintsTask::FMergeConstraintsTask(Private::FBroadPhaseContext& InBroadPhaseContext, Private::FBroadPhaseContext& InOtherBroadPhaseContext, bool bInNeedSorting)
+			: BroadPhaseContext(InBroadPhaseContext)
+			, OtherBroadPhaseContext(InOtherBroadPhaseContext)
+			, bNeedSorting(bInNeedSorting)
+		{
+			CVD_GET_CURRENT_CONTEXT(CVDContext);
+		}
+
+		ENamedThreads::Type FMergeConstraintsTask::GetDesiredThread()
+		{
+			return ENamedThreads::AnyNormalThreadNormalTask;
+		}
+		void FMergeConstraintsTask::Process()
+		{
+			QUICK_SCOPE_CYCLE_COUNTER(STAT_MergeConstraints);
+			CVD_SCOPE_CONTEXT(CVDContext);
+
+			TArray<FPBDCollisionConstraint*>& ThisArray = BroadPhaseContext.CollisionContext.Allocator->NewActiveConstraints;
+			TArray<FPBDCollisionConstraint*>& OtherArray = OtherBroadPhaseContext.CollisionContext.Allocator->NewActiveConstraints;
+
+			if (bNeedSorting)
+			{
+				TArray<FPBDCollisionConstraint*> ResultArray;
+				const int32 NumThis = ThisArray.Num();
+				const int32 NumOther = OtherArray.Num();
+				ResultArray.Reserve(NumThis + NumOther);
+
+				int32 ThisIndex = 0;
+				int32 OtherIndex = 0;
+
+				while (ThisIndex < NumThis || OtherIndex < NumOther)
+				{
+					if (ThisIndex < NumThis && (OtherIndex >= NumOther || ThisArray[ThisIndex]->GetContact().GetCollisionSortKey() < OtherArray[OtherIndex]->GetContact().GetCollisionSortKey()))
+					{
+						ResultArray.Add(ThisArray[ThisIndex++]);
+					}
+					else
+					{
+						check(OtherIndex < NumOther);
+						ResultArray.Add(OtherArray[OtherIndex++]);
+					}
+				}
+
+				ThisArray = MoveTemp(ResultArray);
+			}
+			else
+			{
+				ThisArray.Append(OtherArray);
+			}
+			OtherArray.Reset();
+		}
 	}
 }
\ No newline at end of file
diff --git a/Engine/Source/Runtime/Experimental/Chaos/Private/Chaos/PBDCollisionConstraints.cpp b/Engine/Source/Runtime/Experimental/Chaos/Private/Chaos/PBDCollisionConstraints.cpp
--- a/Engine/Source/Runtime/Experimental/Chaos/Private/Chaos/PBDCollisionConstraints.cpp
+++ b/Engine/Source/Runtime/Experimental/Chaos/Private/Chaos/PBDCollisionConstraints.cpp
@@ -404,6 +404,11 @@ namespace Chaos
 		ConstraintAllocator.SetIsDeterministic(bInIsDeterministic);
 	}
 
+	bool FPBDCollisionConstraints::GetIsDeterministic() const
+	{
+		return bIsDeterministic;
+	}
+
 	void FPBDCollisionConstraints::DisableHandles()
 	{
 		check(NumConstraints() == 0);
diff --git a/Engine/Source/Runtime/Experimental/Chaos/Public/Chaos/Collision/CollisionConstraintAllocator.h b/Engine/Source/Runtime/Experimental/Chaos/Public/Chaos/Collision/CollisionConstraintAllocator.h
--- a/Engine/Source/Runtime/Experimental/Chaos/Public/Chaos/Collision/CollisionConstraintAllocator.h
+++ b/Engine/Source/Runtime/Experimental/Chaos/Public/Chaos/Collision/CollisionConstraintAllocator.h
@@ -16,6 +16,7 @@
 namespace Chaos
 {
 	class FPBDCollisionConstraints;
+	class FSpatialAccelerationBroadPhase;
 
 	namespace Private
 	{
@@ -130,6 +131,10 @@ namespace Chaos
 
 		private:
 			friend class Private::FCollisionConstraintAllocator;
+			friend class FSortConstraintTask;
+			friend class FMergeConstraintsTask;
+			friend class Chaos::FSpatialAccelerationBroadPhase;
+			
 
 			void Reset()
 			{
diff --git a/Engine/Source/Runtime/Experimental/Chaos/Public/Chaos/Collision/SpatialAccelerationBroadPhase.h b/Engine/Source/Runtime/Experimental/Chaos/Public/Chaos/Collision/SpatialAccelerationBroadPhase.h
--- a/Engine/Source/Runtime/Experimental/Chaos/Public/Chaos/Collision/SpatialAccelerationBroadPhase.h
+++ b/Engine/Source/Runtime/Experimental/Chaos/Public/Chaos/Collision/SpatialAccelerationBroadPhase.h
@@ -337,6 +337,45 @@ namespace Chaos
 			FChaosVDContext CVDContext;
 #endif
 		};
+
+		class FSortConstraintTask
+		{
+		public:
+			FSortConstraintTask(Private::FBroadPhaseContext& InBroadPhaseContext);
+
+			TStatId GetStatId() const { RETURN_QUICK_DECLARE_CYCLE_STAT(FSortConstraintTask, STATGROUP_TaskGraphTasks); }
+			static CHAOS_API ENamedThreads::Type GetDesiredThread();
+			static ESubsequentsMode::Type GetSubsequentsMode() { return ESubsequentsMode::TrackSubsequents; }
+			void DoTask(ENamedThreads::Type CurrentThread, const FGraphEventRef& MyCompletionGraphEvent) { Process(); }
+		private:
+			void Process();
+			Private::FBroadPhaseContext& BroadPhaseContext;
+
+#if WITH_CHAOS_VISUAL_DEBUGGER
+			FChaosVDContext CVDContext;
+#endif
+		};
+
+		class FMergeConstraintsTask
+		{
+		public:
+			FMergeConstraintsTask(Private::FBroadPhaseContext& InBroadPhaseContext, Private::FBroadPhaseContext& InOtherBroadPhaseContext, bool bNeedSorting);
+
+			TStatId GetStatId() const { RETURN_QUICK_DECLARE_CYCLE_STAT(FMergeConstraintsTask, STATGROUP_TaskGraphTasks); }
+			static CHAOS_API ENamedThreads::Type GetDesiredThread();
+			static ESubsequentsMode::Type GetSubsequentsMode() { return ESubsequentsMode::TrackSubsequents; }
+			void DoTask(ENamedThreads::Type CurrentThread, const FGraphEventRef& MyCompletionGraphEvent) { Process(); }
+		private:
+			void Process();
+			Private::FBroadPhaseContext& BroadPhaseContext;
+			Private::FBroadPhaseContext& OtherBroadPhaseContext;
+			bool bNeedSorting;
+
+#if WITH_CHAOS_VISUAL_DEBUGGER
+			FChaosVDContext CVDContext;
+#endif
+		};
+
 	}
 
 	/**
@@ -416,11 +455,11 @@ namespace Chaos
 				CVD_GET_WRAPPED_CURRENT_CONTEXT(CVDContext);
 
 				FGraphEventArray NextPendingTasks;
-				const int32 NumContextTask = ProduceOverlapTasks.Num();
+				const int32 NumContextTask = PendingTasks.Num();
 				check(NumContextTask <= NumActiveBroadphaseContexts);
 				for (int32 ContextIndex = 0; ContextIndex < NumContextTask; ContextIndex++)
 				{
-					const FGraphEventArray Prereq = { ProduceOverlapTasks[ContextIndex] };
+					const FGraphEventArray Prereq = { PendingTasks[ContextIndex] };
 					const FGraphEventRef PendingTask = TGraphTask<Private::FAssignMidPhaseTask>::CreateTask(&Prereq).ConstructAndDispatchWhenReady(*this, BroadphaseContexts[ContextIndex]);
 					NextPendingTasks.Add(PendingTask);
 				}
@@ -458,9 +497,10 @@ namespace Chaos
 			// For that we would need to queue and redistribute the NarrowPhases, but currently each
 			// MidPhase runs the NarrowPhase in ProcessMidPhase
 			RedistributeMidPhasesInContexts();
-			const int32 NumContextTask = ProduceOverlapTasks.Num();
+			const int32 NumContextTask = PendingTasks.Num();
 			FGraphEventArray ProcessMidPhaseTasks;
-			ProcessMidPhaseTasks.Reserve(NumContextTask);
+			PendingTasks.Reset();
+			ContextsWithConstraints.Reset();
 			check(NumContextTask <= NumActiveBroadphaseContexts);
 			for (int32 ContextIndex = 0; ContextIndex < NumContextTask; ContextIndex++)
 			{
@@ -468,10 +508,57 @@ namespace Chaos
 				if (!BroadphaseContexts[ContextIndex].Overlaps.IsEmpty())
 				{
 					const FGraphEventRef ProcessMidPhaseTask = TGraphTask<Private::FProcessMidPhaseTask>::CreateTask().ConstructAndDispatchWhenReady(*this, BroadphaseContexts[ContextIndex], Dt);
-					ProcessMidPhaseTasks.Add(ProcessMidPhaseTask);
+					PendingTasks.Add(ProcessMidPhaseTask);
+					ContextsWithConstraints.Add(ContextIndex);
+				}
+			}
+		}
+
+		void GatherConstraints(bool bIsDeterministic)
+		{
+			SCOPE_CYCLE_COUNTER(STAT_Collisions_NarrowPhase);
+			if (bIsDeterministic)
+			{
+				FGraphEventArray ProcessSorting;
+				for (int32 ContextIndex = 0; ContextIndex < PendingTasks.Num(); ContextIndex++)
+				{
+					FGraphEventArray Prereq = { PendingTasks[ContextIndex] };
+					const FGraphEventRef SortConstraintTask = TGraphTask<Private::FSortConstraintTask>::CreateTask(&Prereq).ConstructAndDispatchWhenReady(BroadphaseContexts[ContextsWithConstraints[ContextIndex]]);
+					ProcessSorting.Add(SortConstraintTask);
+				}
+				PendingTasks = MoveTemp(ProcessSorting);
+			}
+			// Create a binary tree of task to merge all sorted constraints
+			int32 NeighborStep = 1;
+			while (PendingTasks.Num() >= 2)
+			{
+				int32 NumProcessSorting = PendingTasks.Num();
+				FGraphEventArray ProcessMerging;
+				// We merge data in the context that have a power of 2 index, so for first layer 0, 2, 4..., for the second layer 0, 4, 8... and so on.
+				for (int32 ProcessIndex = 0, ContextIndex = 0; ProcessIndex < NumProcessSorting - 1; ProcessIndex += 2, ContextIndex += 2*NeighborStep)
+				{
+					FGraphEventArray Prereq = { PendingTasks[ProcessIndex], PendingTasks[ProcessIndex + 1] };
+					const FGraphEventRef MergeConstraintTask = TGraphTask<Private::FMergeConstraintsTask>::CreateTask(&Prereq).ConstructAndDispatchWhenReady(BroadphaseContexts[ContextsWithConstraints[ContextIndex]], BroadphaseContexts[ContextsWithConstraints[ContextIndex + NeighborStep]], bIsDeterministic);
+					ProcessMerging.Add(MergeConstraintTask);
 				}
+				// Add the last one for odd process number
+				if (NumProcessSorting % 2 == 1)
+				{
+					ProcessMerging.Add(PendingTasks[NumProcessSorting - 1]);
+				}
+				PendingTasks = MoveTemp(ProcessMerging);
+				NeighborStep *= 2;
+			}
+			ensure(PendingTasks.Num() <= 1);
+
+			FTaskGraphInterface::Get().WaitUntilTasksComplete(PendingTasks);
+
+			// Lets move all constraints in the first allocator to be picked up by the CollisionConstraintAllocator
+			if (!ContextsWithConstraints.IsEmpty() && ContextsWithConstraints[0] != 0)
+			{
+				BroadphaseContexts[0].CollisionContext.Allocator->NewActiveConstraints = 
+					MoveTemp(BroadphaseContexts[ContextsWithConstraints[0]].CollisionContext.Allocator->NewActiveConstraints);
 			}
-			FTaskGraphInterface::Get().WaitUntilTasksComplete(ProcessMidPhaseTasks);
 		}
 
 		/** @brief This function is the outer loop of collision detection. It loops over the
@@ -496,7 +583,7 @@ namespace Chaos
 			}
 
 			// The number of contexts that may have new overlaps in them
-			NumActiveBroadphaseContexts = !!GSingleThreadedPhysics ? 1 : FMath::Max(FMath::Min(Chaos::CVars::NumWorkerCollisionFactor*(FTaskGraphInterface::Get().GetNumWorkerThreads()-1), Chaos::MaxNumWorkers), 1);
+			NumActiveBroadphaseContexts = !!GSingleThreadedPhysics ? 1 : FMath::Max(FMath::Min(Chaos::CVars::NumWorkerCollisionFactor * (FTaskGraphInterface::Get().GetNumWorkerThreads() - 1), Chaos::MaxNumWorkers), 1);
 			NumActiveBroadphaseContexts = FMath::Max(1, NumActiveBroadphaseContexts);
 
 			BroadphaseContexts.SetNum(NumActiveBroadphaseContexts);
@@ -512,7 +599,7 @@ namespace Chaos
 			using TSOA = typename ViewType::TSOA;
 			using THandle = typename TSOA::THandleType;
 			using THandleBase = typename THandle::THandleBase;
-			ProduceOverlapTasks.Reset();
+			PendingTasks.Reset();
 
 			int32 ContextIndex = 0;
 
@@ -615,21 +702,21 @@ namespace Chaos
 				int32 LastIndex = (BatchIndex + 1) * ParticleByContext;
 				LastIndex = BatchIndex == NumBatch - 1 ? FMath::Min(NumParticles, LastIndex) : LastIndex;
 
-				if (ContextIndex >= ProduceOverlapTasks.Num())
+				if (ContextIndex >= PendingTasks.Num())
 				{
 					FGraphEventRef PendingTask = TGraphTask<Private::FProduceOverlapTask<bOnlyRigid, THandle, LambdaType, SpatialAccelerationType>>::CreateTask().
 						ConstructAndDispatchWhenReady(*this, BroadphaseContexts[ContextIndex],
 							InSpatialAcceleration, Dt, ReturnHandle, BatchIndex * ParticleByContext, LastIndex);
-					ProduceOverlapTasks.Add(PendingTask);
+					PendingTasks.Add(PendingTask);
 				}
 				else
 				{
 					// If all context threads has been filled, create other tasks and link them to the previous one, 
-					FGraphEventArray Prereq = { ProduceOverlapTasks[ContextIndex] };
+					FGraphEventArray Prereq = { PendingTasks[ContextIndex] };
 					FGraphEventRef PendingTask = TGraphTask<Private::FProduceOverlapTask<bOnlyRigid, THandle, LambdaType, SpatialAccelerationType>>::CreateTask(&Prereq).
 						ConstructAndDispatchWhenReady(*this, BroadphaseContexts[ContextIndex],
 							InSpatialAcceleration, Dt, ReturnHandle, BatchIndex * ParticleByContext, LastIndex);
-					ProduceOverlapTasks[ContextIndex] = PendingTask;
+					PendingTasks[ContextIndex] = PendingTask;
 				}
 			}
 		}
@@ -870,7 +957,8 @@ namespace Chaos
 		int32 NumBroadPhasePairs;
 		int32 NumMidPhases;
 
-		FGraphEventArray ProduceOverlapTasks;
+		FGraphEventArray PendingTasks;
+		TArray<int32> ContextsWithConstraints;
 
 		template<bool bOnlyRigid, typename THandle, typename LambdaHandle, typename SpatialAccelerationType>
 		friend class Private::FProduceOverlapTask;
diff --git a/Engine/Source/Runtime/Experimental/Chaos/Public/Chaos/Collision/SpatialAccelerationCollisionDetector.h b/Engine/Source/Runtime/Experimental/Chaos/Public/Chaos/Collision/SpatialAccelerationCollisionDetector.h
--- a/Engine/Source/Runtime/Experimental/Chaos/Public/Chaos/Collision/SpatialAccelerationCollisionDetector.h
+++ b/Engine/Source/Runtime/Experimental/Chaos/Public/Chaos/Collision/SpatialAccelerationCollisionDetector.h
@@ -47,6 +47,7 @@ namespace Chaos
 			// Run collision detection on the output of the broadphase
 			BroadPhase.ProduceCollisions(Dt);
 
+			BroadPhase.GatherConstraints(GetCollisionContainer().GetIsDeterministic());
 			GetCollisionContainer().EndDetectCollisions();
 
 			// If we have a resim cache restore and save contacts
diff --git a/Engine/Source/Runtime/Experimental/Chaos/Public/Chaos/PBDCollisionConstraints.h b/Engine/Source/Runtime/Experimental/Chaos/Public/Chaos/PBDCollisionConstraints.h
--- a/Engine/Source/Runtime/Experimental/Chaos/Public/Chaos/PBDCollisionConstraints.h
+++ b/Engine/Source/Runtime/Experimental/Chaos/Public/Chaos/PBDCollisionConstraints.h
@@ -85,6 +85,12 @@ public:
 	*/
 	void SetIsDeterministic(const bool bInIsDeterministic);
 
+	/**
+	 * @brief Get determinism.
+	 * Support for determinism requires that we sort active constraints each tick, so there is additional cost.
+	*/
+	bool GetIsDeterministic() const;
+
 	/**
 	 *  Clears the list of active constraints.
 	 * @todo(chaos): This is only required because of the way events work (see AdvanceOneTimeStepTask::DoWork)
