Disable cvar for determinism to be able to change it through physics settings

diff --git a/Engine/Source/Runtime/Experimental/Chaos/Private/PBDRigidsSolver.cpp b/Engine/Source/Runtime/Experimental/Chaos/Private/PBDRigidsSolver.cpp
--- a/Engine/Source/Runtime/Experimental/Chaos/Private/PBDRigidsSolver.cpp
+++ b/Engine/Source/Runtime/Experimental/Chaos/Private/PBDRigidsSolver.cpp
@@ -329,7 +329,7 @@ namespace Chaos
 		int32 ChaosSolverProjectionIterations = -1;
 		FAutoConsoleVariableRef CVarChaosSolverProjectionIterations(TEXT("p.Chaos.Solver.Iterations.Projection"), ChaosSolverProjectionIterations, TEXT("Override number of solver projection iterations (-1 to use config)"));
 
-		int32 ChaosSolverDeterministic = 1;
+		int32 ChaosSolverDeterministic = -1;
 		FAutoConsoleVariableRef CVarChaosSolverDeterministic(TEXT("p.Chaos.Solver.Deterministic"), ChaosSolverDeterministic, TEXT("Override determinism. 0: disabled; 1: enabled; -1: use config"));
 
 		// Copied from RBAN
