Iris
* Fix Iris only sending 1 packet per tick even if net.DisableBandwithThrottling was true
* Minor refactor to Iris packet sending code to improve readability and add comments

diff --git a/Engine/Source/Runtime/Engine/Classes/Engine/NetConnection.h b/Engine/Source/Runtime/Engine/Classes/Engine/NetConnection.h
--- a/Engine/Source/Runtime/Engine/Classes/Engine/NetConnection.h
+++ b/Engine/Source/Runtime/Engine/Classes/Engine/NetConnection.h
@@ -67,6 +67,10 @@ namespace NetConnectionHelper
 {
 	/** Number of bits to use in the packet header for sending the milliseconds on the clock when the packet is sent */
 	inline constexpr int32 NumBitsForJitterClockTimeInHeader = 10;
+
+#if !(UE_BUILD_SHIPPING || UE_BUILD_TEST)
+	ENGINE_API bool HasDisabledBandwidthThrottling();
+#endif
 }
 
 extern ENGINE_API TAutoConsoleVariable<int32> CVarNetEnableCongestionControl;
diff --git a/Engine/Source/Runtime/Engine/Private/Net/Experimental/Iris/DataStreamChannel.cpp b/Engine/Source/Runtime/Engine/Private/Net/Experimental/Iris/DataStreamChannel.cpp
--- a/Engine/Source/Runtime/Engine/Private/Net/Experimental/Iris/DataStreamChannel.cpp
+++ b/Engine/Source/Runtime/Engine/Private/Net/Experimental/Iris/DataStreamChannel.cpp
@@ -40,6 +40,14 @@ namespace UE::Net::Private
 		IrisPacketSequenceSafetyMargin,
 		TEXT("How many packets to spare before considering packet sequence full. This allows a few non-DataStreamChannel packets to be sent without messing up packet acking."
 		));
+
+	/**
+	 * net.Iris.ReplicationWriterMaxAllowedPacketsIfNotHugeObject
+	 * Limit the amount of packets that can be sent in a single batch even if the netspeed allows it.
+	 */
+	static int32 GReplicationWriterMaxAllowedPacketsIfNotHugeObject = 3;
+	static FAutoConsoleVariableRef CVarReplicationWriterMaxAllowedPacketsIfNotHugeObject(TEXT("net.Iris.ReplicationWriterMaxAllowedPacketsIfNotHugeObject"), GReplicationWriterMaxAllowedPacketsIfNotHugeObject,
+		TEXT("Allow ReplicationWriter to overcommit data if we have more data to write."));
 }
 
 UDataStreamChannel::UDataStreamChannel(const FObjectInitializer& ObjectInitializer)
@@ -270,19 +278,34 @@ void UDataStreamChannel::WriteData(UE::Net::EDataStreamWriteMode WriteMode)
 	const int64 MaxBunchBits = MaxBytes*8;
 
 	// Try to determine if we have headroom to write more than a single packet if needed.
-	UDataStream::FBeginWriteParameters BeginWriteParams;
-	BeginWriteParams.WriteMode = WriteMode;
+	UDataStream::FBeginWriteParameters BeginWriteParams
+	{
+		.WriteMode = WriteMode,
+	};
 
 	if (WriteMode == EDataStreamWriteMode::Full)
 	{
 		int32 CurrentQueuedBits = Connection->QueuedBits + Connection->SendBuffer.GetNumBits();
+		
+		// If the connection hasn't gone over the limit yet (negative = bits allowed to send)
 		if (CurrentQueuedBits < 0)
 		{
-			const int32 NumBunchesThatMightBeWritten = (-CurrentQueuedBits) / MaxBunchBits;
+			const int32 MaxAllowedPackets = FMath::Min((-CurrentQueuedBits) / MaxBunchBits, UE::Net::Private::GReplicationWriterMaxAllowedPacketsIfNotHugeObject);
+
+			BeginWriteParams.MaxPackets = (uint32)FMath::Max(MaxAllowedPackets, 1);
+		}
 
-			// Indicate that DataStream can request more data
-			BeginWriteParams.bCanWriteMoreData = NumBunchesThatMightBeWritten > 1;
+#if !(UE_BUILD_SHIPPING || UE_BUILD_TEST)
+		if (NetConnectionHelper::HasDisabledBandwidthThrottling())
+		{
+			BeginWriteParams.MaxPackets = 0U;
 		}
+#endif
+	}
+	else
+	{
+		// Only allow 1 packet for PostTickDispatch
+		BeginWriteParams.MaxPackets = 1U;
 	}
 
 	// Currently we want to use a full bunch so we flush if we have to
@@ -387,21 +410,19 @@ void UDataStreamChannel::WriteData(UE::Net::EDataStreamWriteMode WriteMode)
 		return;
 	}
 
-
 	bool bContinueWriting = false;
 	do 
 	{
-		bContinueWriting = false;
-
 		// Write data until we are not allowed to write more
 		const UDataStream::EWriteResult Result = WriteDataFunction();
 
-		if (Result == UDataStream::EWriteResult::HasMoreData)
+		if (Result == UDataStream::EWriteResult::HasMoreData && IsNetReady() && !IsPacketWindowFull())
 		{
-			if (IsNetReady() && !IsPacketWindowFull())
-			{
-				bContinueWriting = true;
-			}
+			bContinueWriting = true;
+		}
+		else
+		{
+			bContinueWriting = false;
 		}
 
 	}
diff --git a/Engine/Source/Runtime/Engine/Private/NetConnection.cpp b/Engine/Source/Runtime/Engine/Private/NetConnection.cpp
--- a/Engine/Source/Runtime/Engine/Private/NetConnection.cpp
+++ b/Engine/Source/Runtime/Engine/Private/NetConnection.cpp
@@ -94,9 +94,9 @@ static TAutoConsoleVariable<int32> CVarNetPacketOrderMaxMissingPackets(TEXT("net
 static TAutoConsoleVariable<int32> CVarNetPacketOrderMaxCachedPackets(TEXT("net.PacketOrderMaxCachedPackets"), 32,
 	TEXT("(NOTE: Must be power of 2!) The maximum number of packets to cache while waiting for missing packet sequences, before treating missing packets as lost."));
 
-#if !UE_BUILD_SHIPPING
+#if !(UE_BUILD_SHIPPING || UE_BUILD_TEST)
 static TAutoConsoleVariable<int32> CVarDisableBandwithThrottling(TEXT("net.DisableBandwithThrottling"), 0,
-	TEXT("Forces IsNetReady to always return true. Not available in shipping builds."));
+	TEXT("Forces IsNetReady to always return true. Not available in shipping or test builds."));
 #endif
 
 TAutoConsoleVariable<int32> CVarNetEnableCongestionControl(TEXT("net.EnableCongestionControl"), 0,
@@ -123,6 +123,16 @@ DECLARE_CYCLE_STAT(TEXT("NetConnection Tick"), Stat_NetConnectionTick, STATGROUP
 DECLARE_CYCLE_STAT(TEXT("NetConnection ReceivedNak"), Stat_NetConnectionReceivedNak, STATGROUP_Net);
 DECLARE_CYCLE_STAT(TEXT("NetConnection NetConnectionReceivedAcks"), Stat_NetConnectionReceivedAcks, STATGROUP_Net);
 
+namespace NetConnectionHelper
+{
+#if !(UE_BUILD_SHIPPING || UE_BUILD_TEST)
+	bool HasDisabledBandwidthThrottling()
+	{
+		return CVarDisableBandwithThrottling.GetValueOnAnyThread() > 0;
+	}
+#endif
+}
+
 namespace UE::Net::Connection::Private
 {
 	struct FValidateLevelVisibilityResult
@@ -2564,7 +2574,7 @@ bool UNetConnection::IsNetReady() const
 	}
 
 #if !(UE_BUILD_SHIPPING || UE_BUILD_TEST)
-	if (CVarDisableBandwithThrottling.GetValueOnAnyThread() > 0)
+	if (NetConnectionHelper::HasDisabledBandwidthThrottling())
 	{
 		return true;
 	}
diff --git a/Engine/Source/Runtime/Experimental/Iris/Core/Private/Iris/ReplicationSystem/ReplicationWriter.cpp b/Engine/Source/Runtime/Experimental/Iris/Core/Private/Iris/ReplicationSystem/ReplicationWriter.cpp
--- a/Engine/Source/Runtime/Experimental/Iris/Core/Private/Iris/ReplicationSystem/ReplicationWriter.cpp
+++ b/Engine/Source/Runtime/Experimental/Iris/Core/Private/Iris/ReplicationSystem/ReplicationWriter.cpp
@@ -51,10 +51,6 @@ static FAutoConsoleVariableRef CVarWarnAboutDroppedAttachmentsToObjectsNotInScop
 	TEXT("Warn when attachments are dropped due to object not in scope. Default is false."
 	));
 
-static int32 GReplicationWriterMaxAllowedPacketsIfNotHugeObject = 3;
-static FAutoConsoleVariableRef CVarReplicationWriterMaxAllowedPacketsIfNotHugeObject(TEXT("net.Iris.ReplicationWriterMaxAllowedPacketsIfNotHugeObject"), GReplicationWriterMaxAllowedPacketsIfNotHugeObject,
-	TEXT("Allow ReplicationWriter to overcommit data if we have more data to write."));
-
 /*
  * net.Iris.ReplicationWriterMaxHugeObjectsInTransit
  * There's a tradeoff mainly between the connection characteristics to support and normal object replication scheduling when tweaking this value.
@@ -2904,15 +2900,15 @@ FReplicationWriter::EWriteObjectStatus FReplicationWriter::WriteObjectInBatch(FN
 	return EWriteObjectStatus::Success;
 }
 
-int FReplicationWriter::PrepareAndSendHugeObjectPayload(FNetSerializationContext& Context, uint32 InternalIndex)
+FReplicationWriter::EWriteStatus FReplicationWriter::PrepareAndSendHugeObjectPayload(FNetSerializationContext& Context, FInternalNetRefIndex InternalIndex)
 {
 	IRIS_PROFILER_SCOPE(FReplicationWriter_PrepareAndSendHugeObjectPayload);
 
 	// Sanity check
 	if (HugeObjectSendQueue.IsFull() || HugeObjectSendQueue.IsObjectInQueue(InternalIndex, false))
 	{
-		ensureMsgf(false, TEXT("HugeObjectSendQueue should not be full or already transmitting. ( InternalIndex: %u )"), InternalIndex);
-		return 0;
+		ensureMsgf(false, TEXT("HugeObjectSendQueue should not be full or already transmitting: %s"), *NetRefHandleManager->PrintObjectFromIndex(InternalIndex));
+		return EWriteStatus::Skipped;
 	}
 
 	typedef uint32 HugeObjectStorageType;
@@ -2977,11 +2973,13 @@ int FReplicationWriter::PrepareAndSendHugeObjectPayload(FNetSerializationContext
 			HandleObjectBatchFailure(WriteHugeObjectStatus, BatchInfo, WriteBitStreamInfo);
 
 			const FNetRefHandleManager::FReplicatedObjectData& ObjectData = NetRefHandleManager->GetReplicatedObjectDataNoCheck(InternalIndex);
-			UE_LOG(LogIris, Error, TEXT("Unable to fit object %s ( InternalIndex: %u ) %s in maximum combined payload of %u bytes. Connection %u will be disconnected."), *ObjectData.RefHandle.ToString(), InternalIndex, ToCStr(ObjectData.Protocol ? ObjectData.Protocol->DebugName : nullptr), MaxHugeObjectPayLoadBytes, Context.GetLocalConnectionId());
+			UE_LOG(LogIris, Error, TEXT("Unable to fit %s (protocol: %s) in maximum combined payload of %u bytes. Connection %u will be disconnected."), 
+				*NetRefHandleManager->PrintObjectFromIndex(InternalIndex), ToCStr(ObjectData.Protocol ? ObjectData.Protocol->DebugName : nullptr), 
+				MaxHugeObjectPayLoadBytes, Context.GetLocalConnectionId());
 			ensure(false);
 
 			Context.SetError(NetError_ObjectStateTooLarge);
-			return -1;
+			return EWriteStatus::Abort;
 		}
 
 		// If we encounter some other error we can try sending a smaller object in the meantime.
@@ -2990,8 +2988,8 @@ int FReplicationWriter::PrepareAndSendHugeObjectPayload(FNetSerializationContext
 			// Cleanup data from batch
 			HandleObjectBatchFailure(WriteHugeObjectStatus, BatchInfo, WriteBitStreamInfo);
 
-			UE_LOG(LogIris, Verbose, TEXT("Problem writing huge object ( InternalIndex: %u ). WriteObjectStatus: %u. Trying smaller object."), InternalIndex, unsigned(WriteHugeObjectStatus));
-			return 0;
+			UE_LOG(LogIris, Verbose, TEXT("Problem writing huge object %s. WriteObjectStatus: %u. Trying smaller object."), *NetRefHandleManager->PrintObjectFromIndex(InternalIndex), unsigned(WriteHugeObjectStatus));
+			return EWriteStatus::Skipped;
 		}
 	}
 
@@ -3027,9 +3025,9 @@ int FReplicationWriter::PrepareAndSendHugeObjectPayload(FNetSerializationContext
 	const bool bSplitSuccess = PartialNetObjectAttachmentHandler->SplitRawDataNetBlob(TRefCountPtr<FRawDataNetBlob>(NetObjectBlob.GetReference()), PartialNetBlobs, HugeObjectSendQueue.DebugName);
 	if (!bSplitSuccess)
 	{
-		UE_LOG(LogIris, Error, TEXT("Unable to split huge object ( InternalIndex: %u ) payload. Connection %u will be disconnected."), InternalIndex, Context.GetLocalConnectionId());
+		UE_LOG(LogIris, Error, TEXT("Unable to split huge object %s payload. Connection %u will be disconnected."), *NetRefHandleManager->PrintObjectFromIndex(InternalIndex), Context.GetLocalConnectionId());
 		Context.SetError(NetError_ObjectStateTooLarge);
-		return -1;
+		return EWriteStatus::Abort;
 	}
 
 	// Enqueue attachments
@@ -3037,9 +3035,9 @@ int FReplicationWriter::PrepareAndSendHugeObjectPayload(FNetSerializationContext
 	check(bEnqueueSuccess);
 	if (!bEnqueueSuccess)
 	{
-		UE_LOG(LogIris, Error, TEXT("Unable to enqueue huge object attachments ( InternalIndex: %u ). Connection %u will be disconnected."), InternalIndex, Context.GetLocalConnectionId());
+		UE_LOG(LogIris, Error, TEXT("Unable to enqueue huge object attachments: %s. Connection %u will be disconnected."), *NetRefHandleManager->PrintObjectFromIndex(InternalIndex), Context.GetLocalConnectionId());
 		Context.SetError(GNetError_InternalError);
-		return -1;
+		return EWriteStatus::Abort;
 	}
 
 	// Add huge object to queue
@@ -3048,9 +3046,9 @@ int FReplicationWriter::PrepareAndSendHugeObjectPayload(FNetSerializationContext
 	check(bHugeObjectWasEnqueued);
 	if (!bHugeObjectWasEnqueued)
 	{
-		UE_LOG(LogIris, Error, TEXT("Unable to enqueue huge object ( InternalIndex: %u ). Connection %u will be disconnected."), InternalIndex, Context.GetLocalConnectionId());
+		UE_LOG(LogIris, Error, TEXT("Unable to enqueue huge object: %s. Connection %u will be disconnected."), *NetRefHandleManager->PrintObjectFromIndex(InternalIndex), Context.GetLocalConnectionId());
 		Context.SetError(GNetError_InternalError);
-		return -1;
+		return EWriteStatus::Abort;
 	}
 
 	// Write huge object attachment(s)
@@ -3077,7 +3075,7 @@ int FReplicationWriter::PrepareAndSendHugeObjectPayload(FNetSerializationContext
 			WriteContext.bHasHugeObjectToSend = 1;
 
 			// Try to fit a smaller object.
-			return 0;
+			return EWriteStatus::Skipped;
 		}
 
 		FBatchRecord BatchRecord;
@@ -3092,21 +3090,26 @@ int FReplicationWriter::PrepareAndSendHugeObjectPayload(FNetSerializationContext
 			HugeObjectSendQueue.Stats.EndSendingTime = FPlatformTime::Cycles64();
 		}
 
-		return 1;
+		return EWriteStatus::Written;
 	}
 
 	// Should not get here.
-	return -1;
+	return EWriteStatus::Abort;
 }
 
-int FReplicationWriter::WriteObjectBatch(FNetSerializationContext& Context, uint32 InternalIndex, uint32 WriteObjectFlags)
+FReplicationWriter::FWriteBatchResult FReplicationWriter::WriteObjectBatch(FNetSerializationContext& Context, FInternalNetRefIndex InternalIndex, uint32 WriteObjectFlags)
 {
 	IRIS_PROFILER_SCOPE(FReplicationWriter_WriteObjectBatch);
 
 	// If this is a destruction info we treat it differently and just write the information required to destruct the object
 	if (GetReplicationInfo(InternalIndex).IsDestructionInfo)
 	{
-		return WriteDestructionInfo(Context, InternalIndex);
+		const EWriteStatus Status = WriteDestructionInfo(Context, InternalIndex);
+		return FWriteBatchResult
+		{
+			.Status = Status,
+			.NumWritten = (Status == EWriteStatus::Written) ? 1U : 0U,
+		};
 	}
 
 	FNetBitStreamWriter& Writer = *Context.GetBitStreamWriter();
@@ -3132,13 +3135,18 @@ int FReplicationWriter::WriteObjectBatch(FNetSerializationContext& Context, uint
 		if (IsWriteObjectSuccess(WriteObjectStatus))
 		{
 			FBatchRecord BatchRecord;
-			const int WrittenObjectCount = HandleObjectBatchSuccess(BatchInfo, BatchRecord);
+			const uint32 WrittenObjectCount = HandleObjectBatchSuccess(BatchInfo, BatchRecord);
 
 			// As a single batch also might include dependent objects which are treated as separate batches on the receiving end we need to account for this when tracking the written batch count.
 			WriteContext.WrittenBatchCount += BatchRecord.BatchCount;
 
 			CommitBatchRecord(BatchRecord);
-			return WrittenObjectCount;
+
+			return FWriteBatchResult
+			{
+				.Status = EWriteStatus::Written,
+				.NumWritten = WrittenObjectCount,
+			};
 		}
 
 		const EWriteObjectRetryMode WriteRetryMode = HandleObjectBatchFailure(WriteObjectStatus, BatchInfo, WriteBitStreamInfo);
@@ -3153,13 +3161,13 @@ int FReplicationWriter::WriteObjectBatch(FNetSerializationContext& Context, uint
 		{
 			case EWriteObjectRetryMode::Abort:
 			{
-				return -1;
+				return FWriteBatchResult{ .Status = EWriteStatus::Abort };
 			}
 
 			case EWriteObjectRetryMode::TrySmallObject:
 			{
 				++WriteContext.FailedToWriteSmallObjectCount;
-				return 0;
+				return FWriteBatchResult{ .Status = EWriteStatus::Skipped };
 			}
 
 			case EWriteObjectRetryMode::SplitHugeObject:
@@ -3176,26 +3184,30 @@ int FReplicationWriter::WriteObjectBatch(FNetSerializationContext& Context, uint
 
 	// Try #2 - Object will be serialized to a temporary buffer of maximum supported size and split into multiple chunks.
 	{
-		const int SendHugeObjectStatus = PrepareAndSendHugeObjectPayload(Context, InternalIndex);
+		const EWriteStatus SendHugeObjectStatus = PrepareAndSendHugeObjectPayload(Context, InternalIndex);
 
 		// If the huge object wrote data it will be tracked as a single batch
-		if (SendHugeObjectStatus == 1)
+		if (SendHugeObjectStatus == EWriteStatus::Written)
 		{
 			++WriteContext.WrittenBatchCount;
 		}
 				
-		return SendHugeObjectStatus;
+		return FWriteBatchResult
+		{
+			.Status = SendHugeObjectStatus,
+			.NumWritten = (SendHugeObjectStatus == EWriteStatus::Written) ? 1U : 0U,
+		};
 	}
 }
 
-int FReplicationWriter::WriteDestructionInfo(FNetSerializationContext& Context, uint32 InternalIndex)
+FReplicationWriter::EWriteStatus FReplicationWriter::WriteDestructionInfo(FNetSerializationContext& Context, uint32 InternalIndex)
 {
 	IRIS_PROFILER_SCOPE(FReplicationWriter_WriteDestructionInfo);
 
 	const int32 MaxDestroyObjectsPerFrame = GReplicationWriterMaxDestroyObjectsPerFrame;
 	if ((MaxDestroyObjectsPerFrame > 0) && (WriteContext.WrittenDestroyObjectCount >= static_cast<uint32>(MaxDestroyObjectsPerFrame)))
 	{
-		return 0;
+		return EWriteStatus::Skipped;
 	}
 
 	FNetBitStreamWriter& Writer = *Context.GetBitStreamWriter();
@@ -3221,7 +3233,7 @@ int FReplicationWriter::WriteDestructionInfo(FNetSerializationContext& Context,
 		// Trigger Rollback
 		Writer.DoOverflow();
 
-		return 0;
+		return EWriteStatus::Skipped;
 	}
 
 	WriteSentinel(&Writer, TEXT("DestructionInfo"));
@@ -3258,7 +3270,7 @@ int FReplicationWriter::WriteDestructionInfo(FNetSerializationContext& Context,
 		++WriteContext.WrittenDestroyObjectCount;
 	}
 
-	return Writer.IsOverflown() ? -1 :  1;
+	return Writer.IsOverflown() ? EWriteStatus::Abort : EWriteStatus::Written;
 }
 
 uint32 FReplicationWriter::WriteOOBAttachments(FNetSerializationContext& Context)
@@ -3270,14 +3282,14 @@ uint32 FReplicationWriter::WriteOOBAttachments(FNetSerializationContext& Context
 		if (WriteContext.bHasOOBAttachmentsToSend && CanSendObject(ObjectIndexForOOBAttachment) && (ReplicationRecord.GetUnusedInfoCount() > 0U))
 		{
 			IRIS_PROFILER_SCOPE(FReplicationWriter_WriteOOBAttachments);
-			const int32 Result = WriteObjectBatch(Context, ObjectIndexForOOBAttachment, WriteObjectFlag_Attachments);
-			if (Result == -1)
+			const FWriteBatchResult Result = WriteObjectBatch(Context, ObjectIndexForOOBAttachment, WriteObjectFlag_Attachments);
+			if (Result.Status == EWriteStatus::Abort)
 			{
 				return WrittenObjectCount;
 			}
 
 			WriteContext.bHasOOBAttachmentsToSend = Attachments.HasUnsentUnreliableAttachments(ENetObjectAttachmentType::OutOfBand, ObjectIndexForOOBAttachment);
-			WrittenObjectCount += Result;
+			WrittenObjectCount += Result.NumWritten;
 		}
 	}
 	else
@@ -3285,8 +3297,8 @@ uint32 FReplicationWriter::WriteOOBAttachments(FNetSerializationContext& Context
 		if (WriteContext.bHasHugeObjectToSend && (ReplicationRecord.GetUnusedInfoCount() > 0U))
 		{
 			IRIS_PROFILER_SCOPE(FReplicationWriter_WriteHugeObjectAttachments);
-			const int32 Result = WriteObjectBatch(Context, ObjectIndexForOOBAttachment, WriteObjectFlag_Attachments | WriteObjectFlag_HugeObject);
-			if (Result == -1)
+			const FWriteBatchResult Result = WriteObjectBatch(Context, ObjectIndexForOOBAttachment, WriteObjectFlag_Attachments | WriteObjectFlag_HugeObject);
+			if (Result.Status == EWriteStatus::Abort)
 			{
 				return WrittenObjectCount;
 			}
@@ -3298,20 +3310,20 @@ uint32 FReplicationWriter::WriteOOBAttachments(FNetSerializationContext& Context
 				HugeObjectSendQueue.Stats.EndSendingTime = FPlatformTime::Cycles64();
 			}
 
-			WrittenObjectCount += Result;
+			WrittenObjectCount += Result.NumWritten;
 		}
 
 		if (WriteContext.bHasOOBAttachmentsToSend && CanSendObject(ObjectIndexForOOBAttachment) && (ReplicationRecord.GetUnusedInfoCount() > 0U))
 		{
 			IRIS_PROFILER_SCOPE(FReplicationWriter_WriteOOBAttachments);
-			const int32 Result = WriteObjectBatch(Context, ObjectIndexForOOBAttachment, WriteObjectFlag_Attachments);
-			if (Result == -1)
+			const FWriteBatchResult Result = WriteObjectBatch(Context, ObjectIndexForOOBAttachment, WriteObjectFlag_Attachments);
+			if (Result.Status == EWriteStatus::Abort)
 			{
 				return WrittenObjectCount;
 			}
 
 			WriteContext.bHasOOBAttachmentsToSend = Attachments.HasUnsentAttachments(ENetObjectAttachmentType::OutOfBand, ObjectIndexForOOBAttachment);
-			WrittenObjectCount += Result;
+			WrittenObjectCount += Result.NumWritten;
 		}
 	}
 
@@ -3331,21 +3343,22 @@ uint32 FReplicationWriter::WriteObjects(FNetSerializationContext& Context)
 	// Write only if there are objects to send and we're not in replication record starvation mode.
 	bool bContinue = WriteContext.bHasUpdatedObjectsToSend && !WriteContext.bIsInReplicationRecordStarvation;
 
-	auto SendObjectFunction = [this, &Context, &WrittenObjectCount](FInternalNetRefIndex InternalIndex)
+	auto SendObjectFunction = [this, &Context, &WrittenObjectCount](FInternalNetRefIndex InternalIndex) -> EWriteStatus
 	{
-		if (!this->WriteContext.ObjectsWrittenThisPacket.GetBit(InternalIndex) && this->CanSendObject(InternalIndex))
+		if (this->WriteContext.ObjectsWrittenThisPacket.GetBit(InternalIndex) || !this->CanSendObject(InternalIndex))
 		{
-			const int32 Result = this->WriteObjectBatch(Context, InternalIndex, WriteObjectFlag_State | WriteObjectFlag_Attachments);
-			if (Result >= 0)
-			{
-				WrittenObjectCount += Result;
-			}
-			else
-			{
-				return false;
-			}
+			// Even if the object cannot be sent, consider it written.
+			return EWriteStatus::Written;
 		}
-		return true;
+
+		const FWriteBatchResult Result = this->WriteObjectBatch(Context, InternalIndex, WriteObjectFlag_State | WriteObjectFlag_Attachments);
+
+		if (Result.Status == EWriteStatus::Written)
+		{
+			WrittenObjectCount += Result.NumWritten;
+		}
+
+		return Result.Status;
 	};
 	
 
@@ -3367,34 +3380,53 @@ uint32 FReplicationWriter::WriteObjects(FNetSerializationContext& Context)
 				const uint32 InternalIndex = WriteContext.DependentObjectsPendingSend.Pop();
 				checkSlow(InternalIndex != ObjectIndexForOOBAttachment);
 				ensureMsgf(GetReplicationInfo(InternalIndex).GetState() != EReplicatedObjectState::Invalid, TEXT("DependentObject with InternalIndex %u is not in scope"), InternalIndex);
-				if (!SendObjectFunction(InternalIndex))
+
+				const EWriteStatus Status = SendObjectFunction(InternalIndex);
+
+				if (Status == EWriteStatus::Abort)
 				{
 					// If we fail, we put the object back on the pending send stack and try again in the next packet of the batch
 					// The reason for pop before using the index is that the SendObjectFunction will push new dependent objects on the stack
-					WriteContext.DependentObjectsPendingSend.Push(InternalIndex);				
+					WriteContext.DependentObjectsPendingSend.Push(InternalIndex);
 					bContinue = false;
 					break;
 				}
+
+				// If we overflown and could not fit the object in the packet
+				if (Status == EWriteStatus::Skipped)
+				{
+					// TODO: We need to ensure it's ok to skip dependant objects and wait for the next frame to resend them if want want Dependent objects to be truly atomic
+				}
 			}
 
-			// Normal send, dependent objects will be pushed as needed
+			// Normal send
 			if (ObjectListIt < SortedCount)
 			{
 				const uint32 InternalIndex = ObjectList[ObjectListIt].Index;
 
 				checkSlow(InternalIndex != ObjectIndexForOOBAttachment);
-				if (!SendObjectFunction(InternalIndex))
+
+				const EWriteStatus Status = SendObjectFunction(InternalIndex);
+				
+				if (Status == EWriteStatus::Abort)
 				{
 					bContinue = false;
 					break;
 				}
 
+				if (Status == EWriteStatus::Skipped)
+				{
+					// Continue trying to send objects as we might find an object with a smaller footprint that fits in the available buffer space.
+					// Note that skipped objects won't be reserialized again even if the more packets are available after this batch.
+				}
+
 				++ObjectListIt;
 			}
 			else
 			{
 				break;
 			}
+
 		}
 	}
 
@@ -3409,10 +3441,10 @@ uint32 FReplicationWriter::WriteObjects(FNetSerializationContext& Context)
 	return WrittenObjectCount;
 }
 
-int FReplicationWriter::HandleObjectBatchSuccess(const FBatchInfo& BatchInfo, FReplicationWriter::FBatchRecord& OutRecord)
+uint32 FReplicationWriter::HandleObjectBatchSuccess(const FBatchInfo& BatchInfo, FReplicationWriter::FBatchRecord& OutRecord)
 {
-	int WrittenObjectCount = 0;
-	int WrittenBatchCount = 0;
+	uint32 WrittenObjectCount = 0;
+	uint32 WrittenBatchCount = 0;
 
 	const bool bTrackObjectStats = BatchInfo.Type != EBatchInfoType::Internal;
 	uint32 ObjectCount = 0;
@@ -3565,6 +3597,7 @@ FReplicationWriter::EWriteObjectRetryMode FReplicationWriter::HandleObjectBatchF
 		return EWriteObjectRetryMode::TrySmallObject;
 	}
 
+	// If there is not enough space left to fit any object, stop serializing more objects immediately
 	const uint32 BitsLeft = BatchBitStreamInfo.ReplicationCapacity - (BatchBitStreamInfo.BatchStartPos - BatchBitStreamInfo.ReplicationStartPos);
 	if (BitsLeft < Parameters.SmallObjectBitThreshold)
 	{
@@ -3587,9 +3620,14 @@ FReplicationWriter::EWriteObjectRetryMode FReplicationWriter::HandleObjectBatchF
 		IRIS_PROFILER_SCOPE(FReplicationWriter_BlockedByHugeOBjectAlreadyBeingSent);
 	}
 
-	// If we are allowed to request more packets to write, we should abort once we failed to write a small object
-	const uint32 MaxFailedSmallObjectCount = WriteContext.bCanWriteMoreData ? 1U : Parameters.MaxFailedSmallObjectCount;
-	if (WriteContext.FailedToWriteSmallObjectCount >= MaxFailedSmallObjectCount)	
+	// If we have more packets available to send, don't try to fit smaller objects in the leftover buffer space
+	if (WriteContext.MaxPacketsToSend == 0U || WriteContext.NumWrittenPacketsInThisBatch < WriteContext.MaxPacketsToSend-1)
+	{
+		return EWriteObjectRetryMode::Abort;
+	}
+
+	// For the last packet of this tick, try to fit other small objects until we skip too many (default: 10)
+	if (WriteContext.FailedToWriteSmallObjectCount >= Parameters.MaxFailedSmallObjectCount)
 	{
 		return EWriteObjectRetryMode::Abort;
 	}
@@ -3635,7 +3673,6 @@ UDataStream::EWriteResult FReplicationWriter::BeginWrite(const UDataStream::FBeg
 			return UDataStream::EWriteResult::NoData;
 		}
 		WriteContext.bHasOOBAttachmentsToSend = bHasUnsentOOBAttachments;
-		WriteContext.bCanWriteMoreData = Params.bCanWriteMoreData;
 	}
 	else
 	{
@@ -3673,7 +3710,7 @@ UDataStream::EWriteResult FReplicationWriter::BeginWrite(const UDataStream::FBeg
 	WriteContext.WrittenDestroyObjectCount = 0U;
 	WriteContext.SortedObjectCount = 0U;
 	WriteContext.NumWrittenPacketsInThisBatch = 0U;
-	WriteContext.bCanWriteMoreData = Params.bCanWriteMoreData;
+	WriteContext.MaxPacketsToSend = Params.MaxPackets;
 
 	// Updated properly in Write
 	WriteContext.bIsInReplicationRecordStarvation = 0;
@@ -3767,8 +3804,9 @@ UDataStream::EWriteResult FReplicationWriter::Write(FNetSerializationContext& Co
 	WriteContext.WrittenBatchCount = 0U;
 	// Check whether we are running low on replication records or not and need to throttle object replication.
 	WriteContext.bIsInReplicationRecordStarvation = ReplicationRecord.GetUnusedInfoCount() < static_cast<uint32>(GReplicationWriterReplicationRecordStarvationThreshold);
-	// Check whether this is to be considered an OOB packet which should avoid writing things like objects pending destroy
-	WriteContext.bIsOOBPacket = (WriteContext.WriteMode == EDataStreamWriteMode::PostTickDispatch) || (static_cast<int32>(WriteContext.NumWrittenPacketsInThisBatch) >= GReplicationWriterMaxAllowedPacketsIfNotHugeObject);
+	
+	// Packets going over the bandwidth limit are considered an OOB packet which should avoid writing things like objects pending destroy.  PostTickDispath packets are also flagged OOB to minimize extra serialization data.
+	WriteContext.bIsOOBPacket = (WriteContext.WriteMode == EDataStreamWriteMode::PostTickDispatch) || (WriteContext.MaxPacketsToSend != 0 && WriteContext.NumWrittenPacketsInThisBatch >= WriteContext.MaxPacketsToSend);
 
 	FNetBitStreamWriter& Writer = *Context.GetBitStreamWriter();
 
@@ -3825,7 +3863,7 @@ UDataStream::EWriteResult FReplicationWriter::Write(FNetSerializationContext& Co
 		{
 			WriteResult = UDataStream::EWriteResult::HasMoreData;
 		}
-		else if (!WriteContext.bIsInReplicationRecordStarvation && WriteContext.bCanWriteMoreData && ((int32)WriteContext.NumWrittenPacketsInThisBatch < GReplicationWriterMaxAllowedPacketsIfNotHugeObject))
+		else if (!WriteContext.bIsInReplicationRecordStarvation && (WriteContext.MaxPacketsToSend == 0 || (WriteContext.NumWrittenPacketsInThisBatch < WriteContext.MaxPacketsToSend-1)))
 		{
 			WriteResult = UDataStream::EWriteResult::HasMoreData;
 		}
diff --git a/Engine/Source/Runtime/Experimental/Iris/Core/Private/Iris/ReplicationSystem/ReplicationWriter.h b/Engine/Source/Runtime/Experimental/Iris/Core/Private/Iris/ReplicationSystem/ReplicationWriter.h
--- a/Engine/Source/Runtime/Experimental/Iris/Core/Private/Iris/ReplicationSystem/ReplicationWriter.h
+++ b/Engine/Source/Runtime/Experimental/Iris/Core/Private/Iris/ReplicationSystem/ReplicationWriter.h
@@ -241,8 +241,10 @@ private:
 
 		// How many packets have we written to?
 		uint32 NumWrittenPacketsInThisBatch = 0U;
-		bool bCanWriteMoreData = false;
-		
+
+		// How many packets are we allowed to send in this batch. Based off the netspeed limit. When 0 = unlimited packets
+		uint32 MaxPacketsToSend = 1U;
+
 		EDataStreamWriteMode WriteMode = EDataStreamWriteMode::Full;	
 	
 		// Whether we're running low on ReplicationRecords or not. If starving only OOB and huge objects are sent, if at all possible.
@@ -485,15 +487,31 @@ private:
 	// Write objects recursive
 	EWriteObjectStatus WriteObjectInBatch(FNetSerializationContext& Context, uint32 InternalIndex, uint32 WriteObjectFlags, FBatchInfo& OutBatchInfo);
 
+	enum class EWriteStatus : int32
+	{
+		// Stream is full and we should stop serializing more objects
+		Abort = -1,
+		// Object was skipped and won't be retried again. Ex: failed dependency, waiting for creation confirmation or buffer won't fit the object and its the last packet to send.
+		Skipped = 0,
+		// Object was successfully written
+		Written = 1,
+	};
+
 	// Write destruction info for an object that should be destroyed
 	// returns > 1 if Objects was written, 0 if the objects was skipped (failed dependency or waiting for creation confirmation) -1 if the stream is full and we should stop
-	int WriteDestructionInfo(FNetSerializationContext& Context, uint32 InternalIndex);
+	EWriteStatus WriteDestructionInfo(FNetSerializationContext& Context, uint32 InternalIndex);
+
+	struct FWriteBatchResult
+	{
+		EWriteStatus Status = EWriteStatus::Skipped;
+		uint32 NumWritten = 0;
+	};
 
 	// Write Object and any subobject(s) to stream as an atomic batch
 	// returns > 1 if Objects was written, 0 if the objects was skipped (failed dependency or waiting for creation confirmation) -1 if the stream is full and we should stop
-	int WriteObjectBatch(FNetSerializationContext& Context, uint32 InternalIndex, uint32 WriteObjectFlags);
+	FWriteBatchResult WriteObjectBatch(FNetSerializationContext& Context, FInternalNetRefIndex InternalIndex, uint32 WriteObjectFlags);
 
-	int PrepareAndSendHugeObjectPayload(FNetSerializationContext& Context, uint32 InternalIndex);
+	EWriteStatus PrepareAndSendHugeObjectPayload(FNetSerializationContext& Context, FInternalNetRefIndex InternalIndex);
 
 	// Write OOBAttachments
 	uint32 WriteOOBAttachments(FNetSerializationContext& Context);
@@ -502,7 +520,7 @@ private:
 	uint32 WriteObjects(FNetSerializationContext& Context);
 
 	// Updates ReplicationInfos, pushes ReplicationRecords etc after a successful call to WriteObjectInBatch() on a top level object
-	int HandleObjectBatchSuccess(const FBatchInfo& BatchInfo, FBatchRecord& OutRecord);
+	uint32 HandleObjectBatchSuccess(const FBatchInfo& BatchInfo, FBatchRecord& OutRecord);
 
 	// Determines the best course of action after a WriteObjectBatch() call failed.
 	EWriteObjectRetryMode HandleObjectBatchFailure(EWriteObjectStatus WriteObjectStatus, const FBatchInfo& BatchInfo, const FBitStreamInfo& BatchBitStreamInfo);
diff --git a/Engine/Source/Runtime/Experimental/Iris/Core/Public/Iris/DataStream/DataStream.h b/Engine/Source/Runtime/Experimental/Iris/Core/Public/Iris/DataStream/DataStream.h
--- a/Engine/Source/Runtime/Experimental/Iris/Core/Public/Iris/DataStream/DataStream.h
+++ b/Engine/Source/Runtime/Experimental/Iris/Core/Public/Iris/DataStream/DataStream.h
@@ -98,7 +98,9 @@ public:
 	struct FBeginWriteParameters
 	{
 		UE::Net::EDataStreamWriteMode WriteMode = UE::Net::EDataStreamWriteMode::Full;
-		bool bCanWriteMoreData = false;
+
+		// Default to sending 1 packet per write. If 0 = unlimited packets
+		uint32 MaxPackets = 1U;
 	};
 
 public:
