Dataflow : Add support for output array
- Add also support for rename connections
- Expose virtual GetConnections method on the connection base class

diff --git a/Engine/Source/Runtime/Experimental/Dataflow/Core/Private/Dataflow/DataflowConnection.cpp b/Engine/Source/Runtime/Experimental/Dataflow/Core/Private/Dataflow/DataflowConnection.cpp
--- a/Engine/Source/Runtime/Experimental/Dataflow/Core/Private/Dataflow/DataflowConnection.cpp
+++ b/Engine/Source/Runtime/Experimental/Dataflow/Core/Private/Dataflow/DataflowConnection.cpp
@@ -14,12 +14,12 @@ const UE::Dataflow::FPin UE::Dataflow::FPin::InvalidPin = { UE::Dataflow::FPin::
 FDataflowConnection::FDataflowConnection(UE::Dataflow::FPin::EDirection InDirection, FName InType, FName InName, FDataflowNode* InOwningNode, const FProperty* InProperty, FGuid InGuid)
 	: Type(InType)
 	, OriginalType(InType)
-	, Name(InName)
 	, OwningNode(InOwningNode)
 	, Property(InProperty)
 	, Guid(InGuid)
 	, Offset(OwningNode ? OwningNode->GetPropertyOffset(Name) : INDEX_NONE)
 	, Direction(InDirection)
+	, Name(InName)
 {
 	InitFromType();
 }
@@ -27,12 +27,12 @@ FDataflowConnection::FDataflowConnection(UE::Dataflow::FPin::EDirection InDirect
 FDataflowConnection::FDataflowConnection(UE::Dataflow::FPin::EDirection InDirection, const UE::Dataflow::FConnectionParameters& Params)
 	: Type(Params.Type)
 	, OriginalType(Params.Type)
-	, Name(Params.Name)
 	, OwningNode(Params.Owner)
 	, Property(Params.Property)
 	, Guid(Params.Guid)
 	, Offset(Params.Offset)
 	, Direction(InDirection)
+	, Name(Params.Name)
 {
 	InitFromType();
 }
@@ -161,6 +161,11 @@ bool FDataflowConnection::ResetToOriginalType()
 	return false;
 }
 
+void FDataflowConnection::Rename(FName NewName)
+{
+	Name = NewName;
+}
+
 bool FDataflowConnection::IsAnytypeDependencyConnected() const
 {
 	if (OwningNode && !TypeDependencyGroup.IsNone())
diff --git a/Engine/Source/Runtime/Experimental/Dataflow/Core/Private/Dataflow/DataflowInputOutput.cpp b/Engine/Source/Runtime/Experimental/Dataflow/Core/Private/Dataflow/DataflowInputOutput.cpp
--- a/Engine/Source/Runtime/Experimental/Dataflow/Core/Private/Dataflow/DataflowInputOutput.cpp
+++ b/Engine/Source/Runtime/Experimental/Dataflow/Core/Private/Dataflow/DataflowInputOutput.cpp
@@ -54,6 +54,11 @@ bool FDataflowInput::RemoveConnection(FDataflowConnection* InOutput)
 	return false;
 }
 
+void FDataflowInput::GetConnections(TArray<FDataflowConnection*>& OutConnections) const
+{
+	OutConnections.Add(Connection);
+}
+
 TArray< FDataflowOutput* > FDataflowInput::GetConnectedOutputs()
 {
 	TArray<FDataflowOutput* > RetList;
@@ -137,7 +142,7 @@ void* FDataflowArrayInput::RealAddress() const
 	{
 		if (ArrayProperty)
 		{
-			if (const void* const AddressAtIndex = ArrayProperty->GetValueAddressAtIndex_Direct(ArrayProperty->Inner, ContainerRealAddress, Index))
+			if (void* const AddressAtIndex = ArrayProperty->GetValueAddressAtIndex_Direct(ArrayProperty->Inner, ContainerRealAddress, Index))
 			{
 				return (void*)((size_t)AddressAtIndex + (size_t)ElementOffset);
 			}
@@ -207,6 +212,11 @@ bool FDataflowOutput::RemoveConnection(FDataflowConnection* InInput)
 	Connections.RemoveSwap((FDataflowInput*)InInput); return true;
 }
 
+void FDataflowOutput::GetConnections(TArray<FDataflowConnection*>& OutConnections) const
+{
+	OutConnections.Append(Connections);
+}
+
 FDataflowOutput& FDataflowOutput::SetPassthroughInput(const UE::Dataflow::FConnectionReference& Reference)
 {
 	check(OwningNode);
@@ -413,3 +423,27 @@ void FDataflowOutput::FixAndPropagateType(FName InType)
 		}
 	}
 }
+
+/////////////////////////////////////////////////////////////////////////////////////////////////////////
+
+FDataflowArrayOutput::FDataflowArrayOutput(int32 InIndex, const UE::Dataflow::FArrayOutputParameters& Param)
+	: FDataflowOutput(Param)
+	, Index(InIndex)
+	, ElementOffset(Param.InnerOffset)
+	, ArrayProperty(Param.ArrayProperty)
+{}
+
+void* FDataflowArrayOutput::RealAddress() const
+{
+	if (void* ContainerRealAddress = Super::RealAddress())
+	{
+		if (ArrayProperty)
+		{
+			if (void* const AddressAtIndex = ArrayProperty->GetValueAddressAtIndex_Direct(ArrayProperty->Inner, ContainerRealAddress, Index))
+			{
+				return (void*)((size_t)AddressAtIndex + (size_t)ElementOffset);
+			}
+		}
+	}
+	return nullptr;
+}
diff --git a/Engine/Source/Runtime/Experimental/Dataflow/Core/Private/Dataflow/DataflowNode.cpp b/Engine/Source/Runtime/Experimental/Dataflow/Core/Private/Dataflow/DataflowNode.cpp
--- a/Engine/Source/Runtime/Experimental/Dataflow/Core/Private/Dataflow/DataflowNode.cpp
+++ b/Engine/Source/Runtime/Experimental/Dataflow/Core/Private/Dataflow/DataflowNode.cpp
@@ -306,6 +306,7 @@ void FDataflowNode::ClearInputs()
 		delete Con;
 	}
 	ExpandedInputs.Reset();
+	InputArrayProperties.Reset();
 }
 
 bool FDataflowNode::HasHideableInputs() const
@@ -343,9 +344,7 @@ void FDataflowNode::AddOutput(FDataflowOutput* InPtr)
 {
 	if (InPtr)
 	{
-PRAGMA_DISABLE_DEPRECATION_WARNINGS  // Until Outputs becomes private
-		for (const TPair<int32, FDataflowOutput*>& Elem : Outputs)
-PRAGMA_ENABLE_DEPRECATION_WARNINGS
+		for (const TPair<UE::Dataflow::FConnectionKey, FDataflowOutput*>& Elem : ExpandedOutputs)
 		{
 			const FDataflowOutput* const Out = Elem.Value;
 			ensureMsgf(!Out->GetName().IsEqual(InPtr->GetName()), TEXT("Add Output Failed: Existing Node output already defined with name (%s)"), *InPtr->GetName().ToString());
@@ -353,21 +352,20 @@ PRAGMA_ENABLE_DEPRECATION_WARNINGS
 
 		check(InPtr->GetOwningNode() == this);
 
-		const uint32 PropertyOffset = InPtr->GetOffset();
-PRAGMA_DISABLE_DEPRECATION_WARNINGS  // Until Outputs becomes private
-		if (ensure(!Outputs.Contains(PropertyOffset)))
+
+		check(InPtr->GetOwningNode() == this);
+
+		const UE::Dataflow::FConnectionKey Key(InPtr->GetOffset(), InPtr->GetContainerIndex(), InPtr->GetContainerElementOffset());
+		if (ensure(!ExpandedOutputs.Contains(Key)))
 		{
-			Outputs.Add(PropertyOffset, InPtr);
+			ExpandedOutputs.Add(Key, InPtr);
 		}
-PRAGMA_ENABLE_DEPRECATION_WARNINGS
 	}
 }
 
 FDataflowOutput* FDataflowNode::FindOutput(uint32 InGuidHash)
 {
-PRAGMA_DISABLE_DEPRECATION_WARNINGS  // Until Outputs becomes private
-	for (TPair<int32, FDataflowOutput*>& Elem : Outputs)
-PRAGMA_ENABLE_DEPRECATION_WARNINGS
+	for (TPair<UE::Dataflow::FConnectionKey, FDataflowOutput*>& Elem : ExpandedOutputs)
 	{
 		FDataflowOutput* const Con = Elem.Value;
 		if (GetTypeHash(Con->GetGuid()) == InGuidHash)
@@ -381,9 +379,7 @@ PRAGMA_ENABLE_DEPRECATION_WARNINGS
 
 FDataflowOutput* FDataflowNode::FindOutput(FName InName)
 {
-PRAGMA_DISABLE_DEPRECATION_WARNINGS  // Until Outputs becomes private
-	for (TPair<int32, FDataflowOutput*>& Elem : Outputs)
-PRAGMA_ENABLE_DEPRECATION_WARNINGS
+	for (TPair<UE::Dataflow::FConnectionKey, FDataflowOutput*>& Elem : ExpandedOutputs)
 	{
 		FDataflowOutput* const Con = Elem.Value;
 		if (Con->GetName().IsEqual(InName))
@@ -396,9 +392,7 @@ PRAGMA_ENABLE_DEPRECATION_WARNINGS
 
 const FDataflowOutput* FDataflowNode::FindOutput(FName InName) const
 {
-PRAGMA_DISABLE_DEPRECATION_WARNINGS  // Until Outputs becomes private
-	for (const TPair<int32, FDataflowOutput*>& Elem : Outputs)
-PRAGMA_ENABLE_DEPRECATION_WARNINGS
+	for (const TPair<UE::Dataflow::FConnectionKey, FDataflowOutput*>& Elem : ExpandedOutputs)
 	{
 		const FDataflowOutput* const Con = Elem.Value;
 		if (Con->GetName().IsEqual(InName))
@@ -411,9 +405,7 @@ PRAGMA_ENABLE_DEPRECATION_WARNINGS
 
 const FDataflowOutput* FDataflowNode::FindOutput(uint32 InGuidHash) const
 {
-PRAGMA_DISABLE_DEPRECATION_WARNINGS  // Until Outputs becomes private
-	for (const TPair<int32, FDataflowOutput*>& Elem : Outputs)
-PRAGMA_ENABLE_DEPRECATION_WARNINGS
+	for (const TPair<UE::Dataflow::FConnectionKey, FDataflowOutput*>& Elem : ExpandedOutputs)
 	{
 		const FDataflowOutput* const Con = Elem.Value;
 		if (GetTypeHash(Con->GetGuid()) == InGuidHash)
@@ -426,9 +418,7 @@ PRAGMA_ENABLE_DEPRECATION_WARNINGS
 
 const FDataflowOutput* FDataflowNode::FindOutput(const UE::Dataflow::FConnectionKey& Key) const
 {
-PRAGMA_DISABLE_DEPRECATION_WARNINGS  // Until Outputs becomes private
-	if (const FDataflowOutput* const* Con = Outputs.Find(Key.Offset))
-PRAGMA_ENABLE_DEPRECATION_WARNINGS
+	if (const FDataflowOutput* const* Con = ExpandedOutputs.Find(Key))
 	{
 		check(*Con);
 		return *Con;
@@ -444,14 +434,24 @@ const FDataflowOutput* FDataflowNode::FindOutput(const UE::Dataflow::FConnection
 		check(Con->RealAddress() == Reference.Reference);
 		return Con;
 	}
+	if (Reference.ContainerReference == nullptr && !OutputArrayProperties.IsEmpty())
+	{
+		// Search through all connections to see if Reference is the RealAddress of an array property.
+		for (const TPair<UE::Dataflow::FConnectionKey, FDataflowOutput*>& Elem : ExpandedOutputs)
+		{
+			FDataflowOutput* const Con = Elem.Value;
+			if (Con->RealAddress() == Reference.Reference)
+			{
+				return Con;
+			}
+		}
+	}
 	return nullptr;
 }
 
 FDataflowOutput* FDataflowNode::FindOutput(const UE::Dataflow::FConnectionKey& Key)
 {
-PRAGMA_DISABLE_DEPRECATION_WARNINGS  // Until Outputs becomes private
-	if (FDataflowOutput* const* Con = Outputs.Find(Key.Offset))
-PRAGMA_ENABLE_DEPRECATION_WARNINGS
+	if (FDataflowOutput* const* Con = ExpandedOutputs.Find(Key))
 	{
 		check(*Con);
 		return *Con;
@@ -467,14 +467,24 @@ FDataflowOutput* FDataflowNode::FindOutput(const UE::Dataflow::FConnectionRefere
 		check(Con->RealAddress() == Reference.Reference);
 		return Con;
 	}
+	if (Reference.ContainerReference == nullptr && !OutputArrayProperties.IsEmpty())
+	{
+		// Search through all connections to see if Reference is the RealAddress of an array property.
+		for (TPair<UE::Dataflow::FConnectionKey, FDataflowOutput*>& Elem : ExpandedOutputs)
+		{
+			FDataflowOutput* Con = Elem.Value;
+			if (Con->RealAddress() == Reference.Reference)
+			{
+				return Con;
+			}
+		}
+	}
 	return nullptr;
 }
 
 const FDataflowOutput* FDataflowNode::FindOutput(const FGuid& InGuid) const
 {
-PRAGMA_DISABLE_DEPRECATION_WARNINGS  // Until Outputs becomes private
-	for (const TPair<int32, FDataflowOutput*>& Elem : Outputs)
-PRAGMA_ENABLE_DEPRECATION_WARNINGS
+	for (const TPair<UE::Dataflow::FConnectionKey, FDataflowOutput*>& Elem : ExpandedOutputs)
 	{
 		const FDataflowOutput* const Con = Elem.Value;
 		if (Con->GetGuid() == InGuid)
@@ -487,41 +497,32 @@ PRAGMA_ENABLE_DEPRECATION_WARNINGS
 
 int32 FDataflowNode::NumOutputs() const
 {
-PRAGMA_DISABLE_DEPRECATION_WARNINGS  // Until Outputs becomes private
-	return Outputs.Num();
-PRAGMA_ENABLE_DEPRECATION_WARNINGS
+	return ExpandedOutputs.Num();
 }
 
 
 TArray< FDataflowOutput* > FDataflowNode::GetOutputs() const
 {
 	TArray< FDataflowOutput* > Result;
-PRAGMA_DISABLE_DEPRECATION_WARNINGS  // Until Outputs becomes private
-	Outputs.GenerateValueArray(Result);
-PRAGMA_ENABLE_DEPRECATION_WARNINGS
+	ExpandedOutputs.GenerateValueArray(Result);
 	return Result;
 }
 
 
 void FDataflowNode::ClearOutputs()
 {
-PRAGMA_DISABLE_DEPRECATION_WARNINGS  // Until Outputs becomes private
-	for (TPair<int32, FDataflowOutput*>& Elem : Outputs)
-PRAGMA_ENABLE_DEPRECATION_WARNINGS
+	for (TPair<UE::Dataflow::FConnectionKey, FDataflowOutput*>& Elem : ExpandedOutputs)
 	{
 		FDataflowOutput* const Con = Elem.Value;
 		delete Con;
 	}
-PRAGMA_DISABLE_DEPRECATION_WARNINGS  // Until Outputs becomes private
-	Outputs.Reset();
-PRAGMA_ENABLE_DEPRECATION_WARNINGS
+	ExpandedOutputs.Reset();
+	OutputArrayProperties.Reset();
 }
 
 bool FDataflowNode::HasHideableOutputs() const
 {
-PRAGMA_DISABLE_DEPRECATION_WARNINGS  // Until Outputs becomes private
-	for (const TPair<int32, FDataflowOutput*>& Elem : Outputs)
-PRAGMA_ENABLE_DEPRECATION_WARNINGS
+	for (const TPair<UE::Dataflow::FConnectionKey, FDataflowOutput*>& Elem : ExpandedOutputs)
 	{
 		const FDataflowOutput* const Con = Elem.Value;
 		if (Con->GetCanHidePin())
@@ -534,9 +535,7 @@ PRAGMA_ENABLE_DEPRECATION_WARNINGS
 
 bool FDataflowNode::HasHiddenOutputs() const
 {
-PRAGMA_DISABLE_DEPRECATION_WARNINGS  // Until Outputs becomes private
-	for (const TPair<int32, FDataflowOutput*>& Elem : Outputs)
-PRAGMA_ENABLE_DEPRECATION_WARNINGS
+	for (const TPair<UE::Dataflow::FConnectionKey, FDataflowOutput*>& Elem : ExpandedOutputs)
 	{
 		const FDataflowOutput* const Con = Elem.Value;
 		if (Con->GetPinIsHidden())
@@ -547,6 +546,46 @@ PRAGMA_ENABLE_DEPRECATION_WARNINGS
 	return false;
 }
 
+bool FDataflowNode::TryRenameInput(const UE::Dataflow::FConnectionReference& Reference, FName NewName)
+{
+	if (FDataflowInput* InputToRename = FindInput(Reference))
+	{
+		return TryRenameInput(*InputToRename, NewName);
+	}
+	return false;
+}
+
+bool FDataflowNode::TryRenameInput(FDataflowInput& InputToRename, FName NewName)
+{
+	const FDataflowInput* ExistingInputFromName = FindInput(NewName);
+	if (ExistingInputFromName == nullptr || ExistingInputFromName == &InputToRename)
+	{
+		InputToRename.Rename(NewName);
+		return true;
+	}
+	return false;
+}
+
+bool FDataflowNode::TryRenameOutput(const UE::Dataflow::FConnectionReference& Reference, FName NewName)
+{
+	if (FDataflowOutput* OutputToRename = FindOutput(Reference))
+	{
+		return TryRenameOutput(*OutputToRename, NewName);
+	}
+	return false;
+}
+
+bool FDataflowNode::TryRenameOutput(FDataflowOutput& OutputToRename, FName NewName)
+{
+	const FDataflowOutput* ExistingOutputFromName = FindOutput(NewName);
+	if (ExistingOutputFromName == nullptr || ExistingOutputFromName == &OutputToRename)
+	{
+		OutputToRename.Rename(NewName);
+		return true;
+	}
+	return false;
+}
+
 TArray<UE::Dataflow::FPin> FDataflowNode::GetPins() const
 {
 	TArray<UE::Dataflow::FPin> RetVal;
@@ -555,9 +594,7 @@ TArray<UE::Dataflow::FPin> FDataflowNode::GetPins() const
 		const FDataflowInput* const Con = Elem.Value;
 		RetVal.Add({ UE::Dataflow::FPin::EDirection::INPUT,Con->GetType(), Con->GetName(), Con->GetPinIsHidden()});
 	}
-PRAGMA_DISABLE_DEPRECATION_WARNINGS  // Until Outputs becomes private
-	for (const TPair<int32, FDataflowOutput*>& Elem : Outputs)
-PRAGMA_ENABLE_DEPRECATION_WARNINGS
+	for (const TPair<UE::Dataflow::FConnectionKey, FDataflowOutput*>& Elem : ExpandedOutputs)
 	{
 		const FDataflowOutput* const Con = Elem.Value;
 		RetVal.Add({ UE::Dataflow::FPin::EDirection::OUTPUT,Con->GetType(), Con->GetName(), Con->GetPinIsHidden() });
@@ -585,9 +622,7 @@ void FDataflowNode::UnregisterPinConnection(const UE::Dataflow::FPin& Pin)
 	}
 	else if (Pin.Direction == UE::Dataflow::FPin::EDirection::OUTPUT)
 	{
-PRAGMA_DISABLE_DEPRECATION_WARNINGS  // Until Outputs becomes private
-		for (TMap<int32, FDataflowOutput*>::TIterator Iter = Outputs.CreateIterator(); Iter; ++Iter)
-PRAGMA_ENABLE_DEPRECATION_WARNINGS
+		for (TMap<UE::Dataflow::FConnectionKey, FDataflowOutput*>::TIterator Iter = ExpandedOutputs.CreateIterator(); Iter; ++Iter)
 		{
 			FDataflowOutput* Con = Iter.Value();
 			if (Con->GetName().IsEqual(Pin.Name) && Con->GetType().IsEqual(Pin.Type))
@@ -644,8 +679,8 @@ PRAGMA_DISABLE_DEPRECATION_WARNINGS  // Until LastModifiedTimestamp becomes priv
 	if (LastModifiedTimestamp < InModifiedTimestamp)
 	{
 		LastModifiedTimestamp = InModifiedTimestamp;
-		for (TPair<int32, FDataflowOutput*>& Elem : Outputs)
 PRAGMA_ENABLE_DEPRECATION_WARNINGS
+		for (TPair<UE::Dataflow::FConnectionKey, FDataflowOutput*>& Elem : ExpandedOutputs)
 		{
 			FDataflowOutput* const Con = Elem.Value;
 			Con->Invalidate(InModifiedTimestamp);
@@ -731,6 +766,10 @@ UE::Dataflow::FConnectionKey FDataflowNode::GetKeyFromReference(const UE::Datafl
 	{
 		Key.ContainerElementOffset = UE::Dataflow::Private::GetArrayElementOffsetFromReference(*ArrayProperty, Reference);
 	}
+	if (const FArrayProperty* const* ArrayProperty = OutputArrayProperties.Find(Key.Offset))
+	{
+		Key.ContainerElementOffset = UE::Dataflow::Private::GetArrayElementOffsetFromReference(*ArrayProperty, Reference);
+	}
 	return Key;
 }
 
@@ -917,6 +956,61 @@ FDataflowInput& FDataflowNode::RegisterInputArrayConnectionInternal(const UE::Da
 	return *Input;
 }
 
+FDataflowOutput& FDataflowNode::RegisterOutputArrayConnectionInternal(const UE::Dataflow::FConnectionReference& Reference, const FName& ElementPropertyName, const FName& ArrayPropertyName)
+{
+	TUniquePtr<FStructOnScope> ScriptOnStruct = TUniquePtr<FStructOnScope>(NewStructOnScope());
+	check(ScriptOnStruct);
+	const UStruct* Struct = ScriptOnStruct->GetStruct();
+	check(Struct);
+	UE::Dataflow::FArrayOutputParameters OutputParams;
+	OutputParams.Owner = this;
+
+	// Find the Array property.
+	TArray<const FProperty*> ArrayPropertyChain;
+	for (FPropertyValueIterator PropertyIt(FArrayProperty::StaticClass(), Struct, this); PropertyIt; ++PropertyIt)
+	{
+		if (Reference.ContainerReference == PropertyIt.Value() && (ArrayPropertyName == NAME_None || ArrayPropertyName == PropertyIt.Key()->GetName()))
+		{
+			OutputParams.ArrayProperty = CastFieldChecked<FArrayProperty>(PropertyIt.Key());
+			OutputParams.Offset = GetConnectionOffsetFromReference(Reference.ContainerReference);
+			PropertyIt.GetPropertyChain(ArrayPropertyChain);
+			break;
+		}
+	}
+
+	check(OutputParams.ArrayProperty);
+
+	// Find the element property
+	TArray<const FProperty*> PropertyChain;
+	const void* const AddressAtIndex = OutputParams.ArrayProperty->GetValueAddressAtIndex_Direct(OutputParams.ArrayProperty->Inner, const_cast<void*>(Reference.ContainerReference), Reference.Index);
+	if (AddressAtIndex == Reference.Reference && (ElementPropertyName == NAME_None || ElementPropertyName == OutputParams.ArrayProperty->Inner->GetName()))
+	{
+		OutputParams.Property = OutputParams.ArrayProperty->Inner;
+		PropertyChain = { OutputParams.ArrayProperty->Inner };
+	}
+	else if (const FStructProperty* const InnerStruct = CastField<FStructProperty>(OutputParams.ArrayProperty->Inner))
+	{
+		OutputParams.Property = &UE::Dataflow::Private::FindPropertyChecked(InnerStruct->Struct, AddressAtIndex, Reference.Reference, ElementPropertyName, &PropertyChain);
+		PropertyChain.Add(InnerStruct);
+	}
+
+	check(OutputParams.Property);
+
+	PropertyChain.Append(ArrayPropertyChain);
+	FString ExtendedType;
+	const FString CPPType = OutputParams.Property->GetCPPType(&ExtendedType);
+	OutputParams.Type = FName(CPPType + ExtendedType);
+	OutputParams.Name = GetPropertyFullName(PropertyChain, Reference.Index);
+	OutputParams.InnerOffset = UE::Dataflow::Private::GetArrayElementOffsetFromReference(OutputParams.ArrayProperty, Reference);
+
+	InputArrayProperties.Emplace(OutputParams.Offset, OutputParams.ArrayProperty);
+
+	FDataflowOutput* const Output = new FDataflowArrayOutput(Reference.Index, OutputParams);
+	AddOutput(Output);
+	check(FindOutput(Reference) == Output);
+	return *Output;
+}
+
 void FDataflowNode::UnregisterInputConnection(const UE::Dataflow::FConnectionKey& Key)
 {
 	if (ExpandedInputs.Remove(Key))
@@ -926,6 +1020,15 @@ void FDataflowNode::UnregisterInputConnection(const UE::Dataflow::FConnectionKey
 	}
 }
 
+void FDataflowNode::UnregisterOutputConnection(const UE::Dataflow::FConnectionKey& Key)
+{
+	if (ExpandedOutputs.Remove(Key))
+	{
+		// Invalidate graph as this output might have had connections
+		Invalidate();
+	}
+}
+
 FDataflowOutput& FDataflowNode::RegisterOutputConnectionInternal(const UE::Dataflow::FConnectionReference& Reference, const FName& PropertyName)
 {
 	TUniquePtr<FStructOnScope> ScriptOnStruct = TUniquePtr<FStructOnScope>(NewStructOnScope());
@@ -1125,7 +1228,7 @@ bool FDataflowNode::ValidateConnections()
 			}
 
 PRAGMA_DISABLE_DEPRECATION_WARNINGS  // Until Outputs becomes private
-			for (const TPair<int, FDataflowOutput*>& Output : Outputs)
+			for (const TPair<UE::Dataflow::FConnectionKey, FDataflowOutput*>& Output : ExpandedOutputs)
 PRAGMA_ENABLE_DEPRECATION_WARNINGS
 			{
 				FDataflowOutput* OutputVal = Output.Value;
@@ -1419,9 +1522,7 @@ bool FDataflowNode::SetAllConnectionConcreteType(FName NewType)
 		}
 	}
 
-	PRAGMA_DISABLE_DEPRECATION_WARNINGS  // Until Outputs becomes private
-	for (const TPair<int32, FDataflowOutput*>& OutputEntry : Outputs)
-	PRAGMA_ENABLE_DEPRECATION_WARNINGS
+	for (const TPair<UE::Dataflow::FConnectionKey, FDataflowOutput*>& OutputEntry : ExpandedOutputs)
 	{
 		FDataflowOutput* const Output = OutputEntry.Value;
 		if (Output && Output->GetType() != NewType)
@@ -1457,4 +1558,4 @@ bool FDataflowNode::IsAnytypeDependencyConnected(FName DependencyGroup, const FD
 			return true;
 		});
 	return bIsAnyDependencyConnected;
-}
\ No newline at end of file
+}
diff --git a/Engine/Source/Runtime/Experimental/Dataflow/Core/Public/Dataflow/DataflowConnection.h b/Engine/Source/Runtime/Experimental/Dataflow/Core/Public/Dataflow/DataflowConnection.h
--- a/Engine/Source/Runtime/Experimental/Dataflow/Core/Public/Dataflow/DataflowConnection.h
+++ b/Engine/Source/Runtime/Experimental/Dataflow/Core/Public/Dataflow/DataflowConnection.h
@@ -134,7 +134,6 @@ struct FDataflowConnection
 protected:
 	FName Type;
 	FName OriginalType;
-	FName Name;
 	FDataflowNode* OwningNode = nullptr;
 	const FProperty* Property = nullptr;
 	FGuid  Guid;
@@ -146,6 +145,9 @@ protected:
 	bool bCanHidePin:1 = false;
 	bool bPinIsHidden:1 = false;
 
+private:
+	FName Name;
+
 	/** Lock the type from being changed */
 	bool bLockType:1 = false;
 
@@ -174,12 +176,6 @@ protected:
 	*/
 	DATAFLOWCORE_API bool IsExtendedType(FName InType) const;
 
-	/** 
-	* Set the concrete type of an anytype connection 
-	* @return true if the type was effectivelly changed
-	*/
-	DATAFLOWCORE_API bool SetConcreteType(FName InType);
-
 	/**
 	* reset the output to its original type 
 	* no-op for non enytypes types
@@ -188,6 +184,9 @@ protected:
 	*/
 	DATAFLOWCORE_API bool ResetToOriginalType();
 
+	/* Rename the connection */
+	DATAFLOWCORE_API void Rename(FName NewName);
+
 public:
 	FDataflowConnection() {};
 	UE_DEPRECATED(5.5, "Deprecated constructor : use FConnectionParameters to pass parameters")
@@ -216,6 +215,18 @@ public:
 		return UE::Dataflow::FConnectionKey(GetOffset(), GetContainerIndex(), GetContainerElementOffset());
 	}
 
+	/**
+	* Set the concrete type of an anytype connection
+	* Changing the type may fail or assert if :
+	*	- the connection is connected to another one, 
+	*	- the type is incompatible with the original type 
+	*	- the type has any type dependencies connected
+	*	- the type is locked ( see LockType() )
+	* (see IsSafeToTryChangingType)
+	* @return true if the type was effectivelly changed
+	*/
+	DATAFLOWCORE_API bool SetConcreteType(FName InType);
+
 	void LockType() { bLockType = true; }
 	void UnlockType() { bLockType = false; }
 
@@ -242,6 +253,7 @@ public:
 	virtual bool IsConnected() const { return false; }
 	virtual bool AddConnection(FDataflowConnection* In) { return false; };
 	virtual bool RemoveConnection(FDataflowConnection* In) { return false; }
+	virtual void GetConnections(TArray<FDataflowConnection*>& OutConnections) const {};
 
 	DATAFLOWCORE_API bool IsAnyType() const { return bIsAnyType; }
 	DATAFLOWCORE_API static bool IsAnyType(const FName& InType);
diff --git a/Engine/Source/Runtime/Experimental/Dataflow/Core/Public/Dataflow/DataflowInputOutput.h b/Engine/Source/Runtime/Experimental/Dataflow/Core/Public/Dataflow/DataflowInputOutput.h
--- a/Engine/Source/Runtime/Experimental/Dataflow/Core/Public/Dataflow/DataflowInputOutput.h
+++ b/Engine/Source/Runtime/Experimental/Dataflow/Core/Public/Dataflow/DataflowInputOutput.h
@@ -61,6 +61,8 @@ public:
 	virtual bool AddConnection(FDataflowConnection* InOutput) override;
 	virtual bool RemoveConnection(FDataflowConnection* InOutput) override;
 
+	virtual void GetConnections(TArray<FDataflowConnection*>& OutConnections) const override;
+
 	FDataflowOutput* GetConnection() { return Connection; }
 	const FDataflowOutput* GetConnection() const { return Connection; }
 	bool HasAnyConnections() const { return Connection != nullptr; }
@@ -120,6 +122,15 @@ namespace UE::Dataflow
 			: FConnectionParameters(InType, InName, InOwner, InProperty, InOffset, InGuid)
 		{}
 	};
+
+	struct FArrayOutputParameters : public FOutputParameters
+	{
+		FArrayOutputParameters()
+			: FOutputParameters()
+		{}
+		const FArrayProperty* ArrayProperty = nullptr;
+		uint32 InnerOffset = INDEX_NONE;
+	};
 }
 
 USTRUCT()
@@ -170,6 +181,8 @@ public:
 
 	DATAFLOWCORE_API virtual bool RemoveConnection(FDataflowConnection* InInput) override;
 
+	virtual void GetConnections(TArray<FDataflowConnection*>& OutConnections) const override;
+
 	UE_DEPRECATED(5.5, "Use SetPassthroughInput instead")
 	virtual void SetPassthroughOffset(const uint32 InPassthroughOffset)
 	{
@@ -278,6 +291,33 @@ struct TStructOpsTypeTraits<FDataflowOutput> : public  TStructOpsTypeTraitsBase2
 	};
 };
 
+USTRUCT()
+struct FDataflowArrayOutput : public FDataflowOutput
+{
+	GENERATED_USTRUCT_BODY()
+
+private:
+	int32 Index;
+	uint32 ElementOffset; // Offset to Property inside an array element
+	const FArrayProperty* ArrayProperty = nullptr;
+
+public:
+	explicit FDataflowArrayOutput(int32 InIndex = INDEX_NONE, const UE::Dataflow::FArrayOutputParameters& Param = {});
+
+	DATAFLOWCORE_API virtual void* RealAddress() const override;
+	virtual int32 GetContainerIndex() const override { return Index; }
+	virtual uint32 GetContainerElementOffset() const override { return ElementOffset; }
+};
+
+template<>
+struct TStructOpsTypeTraits<FDataflowArrayOutput> : public  TStructOpsTypeTraitsBase2<FDataflowArrayOutput>
+{
+	enum
+	{
+		WithCopy = false,
+	};
+};
+
 template<typename T>
 const T& FDataflowInput::GetValue(UE::Dataflow::FContext& Context, const T& Default) const
 {
diff --git a/Engine/Source/Runtime/Experimental/Dataflow/Core/Public/Dataflow/DataflowNode.h b/Engine/Source/Runtime/Experimental/Dataflow/Core/Public/Dataflow/DataflowNode.h
--- a/Engine/Source/Runtime/Experimental/Dataflow/Core/Public/Dataflow/DataflowNode.h
+++ b/Engine/Source/Runtime/Experimental/Dataflow/Core/Public/Dataflow/DataflowNode.h
@@ -49,9 +49,9 @@ PRAGMA_ENABLE_DEPRECATION_WARNINGS
 	UE_DEPRECATED(5.5, "Will be made private in 5.7, use Get/Set methods instead.")
 	UE::Dataflow::FTimestamp LastModifiedTimestamp;
 
-	UE_DEPRECATED(5.5, "Inputs type has changed and has been made private")
+	UE_DEPRECATED(5.5, "Inputs type has changed and has been made private (see ExpandedInputs)")
 	TMap< int, FDataflowInput* > Inputs;
-	UE_DEPRECATED(5.5, "Will be made private in 5.7, use Get/Set methods instead.")
+	UE_DEPRECATED(5.6, "Outputs type has changed and has been made private (see ExpandedOutputs)")
 	TMap< int, FDataflowOutput* > Outputs;
 
 	UE_DEPRECATED(5.6, "Will be made private in 5.8, use IsActive function instead.")
@@ -201,6 +201,11 @@ PRAGMA_ENABLE_DEPRECATION_WARNINGS
 	DATAFLOWCORE_API const FDataflowOutput* FindOutput(const UE::Dataflow::FConnectionReference& Reference) const;
 	DATAFLOWCORE_API const FDataflowOutput* FindOutput(const FGuid& InGuid) const;
 
+	DATAFLOWCORE_API bool TryRenameInput(const UE::Dataflow::FConnectionReference& Reference, FName NewName);
+	DATAFLOWCORE_API bool TryRenameInput(FDataflowInput& InputToRename, FName NewName);
+	DATAFLOWCORE_API bool TryRenameOutput(const UE::Dataflow::FConnectionReference& Reference, FName NewName);
+	DATAFLOWCORE_API bool TryRenameOutput(FDataflowOutput& OutputToRename, FName NewName);
+
 	/** Return a property's byte offset from the dataflow base node address using the full property name (must includes its parent struct property names). Does not work with array properties.*/
 	uint32 GetPropertyOffset(const FName& PropertyFullName) const;
 
@@ -309,6 +314,35 @@ PRAGMA_ENABLE_DEPRECATION_WARNINGS
 		return Input;
 	}
 
+	template<typename T>
+	FDataflowOutput& RegisterOutputArrayConnection(const UE::Dataflow::TConnectionReference<T>& Reference, const FName& ElementPropertyName = NAME_None,
+		const FName& ArrayPropertyName = NAME_None)
+	{
+		FDataflowOutput& Output = RegisterOutputArrayConnectionInternal(Reference, ElementPropertyName, ArrayPropertyName);
+		if constexpr (std::is_base_of_v<FDataflowAnyType, T>)
+		{
+			Output.SetTypePolicy(T::FPolicyType::GetInterface());
+		}
+		return Output;
+	}
+
+	template<typename T>
+	FDataflowOutput& FindOrRegisterOutputArrayConnection(const UE::Dataflow::TConnectionReference<T>& Reference, const FName& ElementPropertyName = NAME_None,
+		const FName& ArrayPropertyName = NAME_None)
+	{
+		if (FDataflowOutput* const FoundOutput = FindOutput(Reference))
+		{
+			return *FoundOutput;
+		}
+
+		FDataflowOutput& Output = RegisterOutputArrayConnectionInternal(Reference, ElementPropertyName, ArrayPropertyName);
+		if constexpr (std::is_base_of_v<FDataflowAnyType, T>)
+		{
+			Output.SetTypePolicy(T::FPolicyType::GetInterface());
+		}
+		return Output;
+	}
+
 	/** Unregister the input connection if one exists matching this property, and then invalidate the graph. */
 	void UnregisterInputConnection(const UE::Dataflow::FConnectionReference& Reference)
 	{
@@ -323,6 +357,12 @@ PRAGMA_ENABLE_DEPRECATION_WARNINGS
 	/** Unregister the connection if one exists matching this pin, then invalidate the graph. */
 	DATAFLOWCORE_API void UnregisterPinConnection(const UE::Dataflow::FPin& Pin);
 
+	void UnregisterOutputConnection(const UE::Dataflow::FConnectionReference& Reference)
+	{
+		UnregisterOutputConnection(GetKeyFromReference(Reference));
+	}
+	DATAFLOWCORE_API void UnregisterOutputConnection(const UE::Dataflow::FConnectionKey& Key);
+
 	//
 	// Evaluation
 	//
@@ -527,9 +567,7 @@ PRAGMA_ENABLE_DEPRECATION_WARNINGS
 				return;
 			}
 		}
-		PRAGMA_DISABLE_DEPRECATION_WARNINGS  // Until Outputs becomes private
-		for (const TPair<int32, FDataflowOutput*>& Elem : Outputs)
-		PRAGMA_ENABLE_DEPRECATION_WARNINGS
+		for (const TPair<UE::Dataflow::FConnectionKey, FDataflowOutput*>& Elem : ExpandedOutputs)
 		{
 			if (Function(Elem.Value) == false)
 			{
@@ -667,8 +705,8 @@ protected:
 
 	DATAFLOWCORE_API FDataflowInput& RegisterInputConnectionInternal(const UE::Dataflow::FConnectionReference& Reference, const FName& PropertyName = NAME_None);
 	DATAFLOWCORE_API FDataflowOutput& RegisterOutputConnectionInternal(const UE::Dataflow::FConnectionReference& Reference, const FName& PropertyName = NAME_None);
-	DATAFLOWCORE_API FDataflowInput& RegisterInputArrayConnectionInternal(const UE::Dataflow::FConnectionReference& Reference, const FName& ElementPropertyName = NAME_None,
-		const FName& ArrayPropertyName = NAME_None);
+	DATAFLOWCORE_API FDataflowInput& RegisterInputArrayConnectionInternal(const UE::Dataflow::FConnectionReference& Reference, const FName& ElementPropertyName = NAME_None, const FName& ArrayPropertyName = NAME_None);
+	DATAFLOWCORE_API FDataflowOutput& RegisterOutputArrayConnectionInternal(const UE::Dataflow::FConnectionReference& Reference, const FName& ElementPropertyName = NAME_None, const FName& ArrayPropertyName = NAME_None);
 
 private:
 
@@ -699,6 +737,9 @@ private:
 	TMap<UE::Dataflow::FConnectionKey, FDataflowInput*> ExpandedInputs;
 	TMap<uint32, const FArrayProperty*> InputArrayProperties; // Used to calculate ContainerElementOffsets
 
+	TMap<UE::Dataflow::FConnectionKey, FDataflowOutput*> ExpandedOutputs;
+	TMap<uint32, const FArrayProperty*> OutputArrayProperties; // Used to calculate ContainerElementOffsets
+
 	UPROPERTY()
 	FInstancedPropertyBag FrozenProperties;  // Property bag for frozen outputs
 
