Dataflow : Fix compilation warning about out of order initialization for members of FDataflowConnection

diff --git a/Engine/Source/Runtime/Experimental/Dataflow/Core/Private/Dataflow/DataflowConnection.cpp b/Engine/Source/Runtime/Experimental/Dataflow/Core/Private/Dataflow/DataflowConnection.cpp
--- a/Engine/Source/Runtime/Experimental/Dataflow/Core/Private/Dataflow/DataflowConnection.cpp
+++ b/Engine/Source/Runtime/Experimental/Dataflow/Core/Private/Dataflow/DataflowConnection.cpp
@@ -17,10 +17,12 @@ FDataflowConnection::FDataflowConnection(UE::Dataflow::FPin::EDirection InDirect
 	, OwningNode(InOwningNode)
 	, Property(InProperty)
 	, Guid(InGuid)
-	, Offset(OwningNode ? OwningNode->GetPropertyOffset(Name) : INDEX_NONE)
 	, Direction(InDirection)
 	, Name(InName)
 {
+	// this cannot be in the initializer because of  warning V670 about order of initialization
+	Offset = OwningNode ? OwningNode->GetPropertyOffset(Name) : INDEX_NONE;
+
 	InitFromType();
 }
 
