Fix the geometry collection set rest collection to make sure the custom renderer is properly deregistered and reregistered
this fixes issues with geometry collection editor thumbnail that were showing stacking up when using custom renderers

diff --git a/Engine/Plugins/Experimental/GeometryCollectionPlugin/Source/GeometryCollectionEditor/Private/GeometryCollection/GeometryCollectionThumbnailScene.cpp b/Engine/Plugins/Experimental/GeometryCollectionPlugin/Source/GeometryCollectionEditor/Private/GeometryCollection/GeometryCollectionThumbnailScene.cpp
--- a/Engine/Plugins/Experimental/GeometryCollectionPlugin/Source/GeometryCollectionEditor/Private/GeometryCollection/GeometryCollectionThumbnailScene.cpp
+++ b/Engine/Plugins/Experimental/GeometryCollectionPlugin/Source/GeometryCollectionEditor/Private/GeometryCollection/GeometryCollectionThumbnailScene.cpp
@@ -28,6 +28,7 @@ void FGeometryCollectionThumbnailScene::SetGeometryCollection(UGeometryCollectio
 	if (UGeometryCollectionComponent* Component = PreviewActor->GetGeometryCollectionComponent())
 	{
 		Component->SetRestCollection(GeometryCollection);
+		Component->ForceNativeRendering(true);
 
 		if (GeometryCollection)
 		{
diff --git a/Engine/Source/Runtime/Experimental/GeometryCollectionEngine/Private/GeometryCollection/GeometryCollectionComponent.cpp b/Engine/Source/Runtime/Experimental/GeometryCollectionEngine/Private/GeometryCollection/GeometryCollectionComponent.cpp
--- a/Engine/Source/Runtime/Experimental/GeometryCollectionEngine/Private/GeometryCollection/GeometryCollectionComponent.cpp
+++ b/Engine/Source/Runtime/Experimental/GeometryCollectionEngine/Private/GeometryCollection/GeometryCollectionComponent.cpp
@@ -4958,6 +4958,11 @@ void UGeometryCollectionComponent::ApplyAssetDefaults()
 
 void UGeometryCollectionComponent::SetRestCollection(const UGeometryCollection* RestCollectionIn, bool bApplyAssetDefaults)
 {
+	if (CustomRenderer)
+	{
+		UnregisterCustomRenderer();
+	}
+
 	//UE_LOG(UGCC_LOG, Log, TEXT("GeometryCollectionComponent[%p]::SetRestCollection()"), this);
 	if (RestCollectionIn)
 	{
@@ -4985,7 +4990,6 @@ void UGeometryCollectionComponent::SetRestCollection(const UGeometryCollection*
 
 		if (CustomRenderer)
 		{
-			UnregisterCustomRenderer();
 			RegisterCustomRenderer();
 			RefreshCustomRenderer();
 		}
