Chaos - add debug (cvar) option to randomize constraint solver order

diff --git a/Engine/Source/Runtime/Experimental/Chaos/Private/Chaos/Island/IslandManager.cpp b/Engine/Source/Runtime/Experimental/Chaos/Private/Chaos/Island/IslandManager.cpp
--- a/Engine/Source/Runtime/Experimental/Chaos/Private/Chaos/Island/IslandManager.cpp
+++ b/Engine/Source/Runtime/Experimental/Chaos/Private/Chaos/Island/IslandManager.cpp
@@ -85,6 +85,9 @@ namespace Chaos::CVars
 	/** Cvar to enable/disable computing max iterations if island is dirty */
 	bool bChaosSolverComputeIterationSettings = true;
 	FAutoConsoleVariableRef CVarChaosSolverComputeIterationSettingsEnabled(TEXT("p.Chaos.Solver.ComputeIterationSettings.Enabled"), bChaosSolverComputeIterationSettings, TEXT("Recompute iteration settings every time an island is changed"));
+
+	bool bChaosRandomizeConstraintOrder = false;
+	FAutoConsoleVariableRef CVarChaosRandomizeConstraintOrder(TEXT("p.Chaos.Solver.RandomizeConstraintOrder"), bChaosRandomizeConstraintOrder, TEXT("Randomize constraint order for testing"));
 }
 
 
@@ -2013,6 +2016,22 @@ namespace Chaos::Private
 		}
 	}
 
+	void FPBDIslandManager::RandomizeIslandEdges(FPBDIsland* Island)
+	{
+		for (int32 ContainerIndex = 0; ContainerIndex < Island->ContainerEdges.Num(); ++ContainerIndex)
+		{
+			TArray<FPBDIslandConstraint*>& ContainerConstraints = Island->ContainerEdges[ContainerIndex];
+			for (int32 SwapIndex = 0; SwapIndex < ContainerConstraints.Num(); ++SwapIndex)
+			{
+				const int32 IndexA = FMath::RandRange(0, ContainerConstraints.Num() - 1);
+				const int32 IndexB = FMath::RandRange(0, ContainerConstraints.Num() - 1);
+				Swap(ContainerConstraints[IndexA], ContainerConstraints[IndexB]);
+				ContainerConstraints[IndexA]->IslandArrayIndex = IndexA;
+				ContainerConstraints[IndexB]->IslandArrayIndex = IndexB;
+			}
+		}
+	}
+
 	void FPBDIslandManager::ProcessWakes()
 	{
 		for (FPBDIsland* Island : Islands)
@@ -2098,6 +2117,12 @@ namespace Chaos::Private
 				}
 			}
 
+			// Randomize constraint order for testing
+			if (CVars::bChaosRandomizeConstraintOrder)
+			{
+				RandomizeIslandEdges(Island);
+			}
+
 			Island->Flags.bWasSleeping = Island->Flags.bIsSleeping;
 			Island->Flags.bCheckSleep = false;
 			Island->Flags.bItemsAdded = false;
diff --git a/Engine/Source/Runtime/Experimental/Chaos/Public/Chaos/Island/IslandManager.h b/Engine/Source/Runtime/Experimental/Chaos/Public/Chaos/Island/IslandManager.h
--- a/Engine/Source/Runtime/Experimental/Chaos/Public/Chaos/Island/IslandManager.h
+++ b/Engine/Source/Runtime/Experimental/Chaos/Public/Chaos/Island/IslandManager.h
@@ -580,6 +580,7 @@ namespace Chaos::Private
 		// Sorting
 		CHAOS_API void AssignIslandLevels(FPBDIsland* Island);
 		CHAOS_API void SortIslandEdges(FPBDIsland* Island);
+		CHAOS_API void RandomizeIslandEdges(FPBDIsland* Island);
 
 		// Sleeping
 		CHAOS_API void ProcessSleep(const FRealSingle Dt);
