Add a new cvar 'ias.RunRequestTest' that will initiate a large number of Io requests for chunks that exist in the IAS system.

#rnx

- This is intended for development only and is an easy way to kick off a number of downloads. By default it will spawn 1000 requests but this can be overriden by providing the number of requests that you want as the first arg.
-- At the moment we just find the first X requests in the IoStore and request those, we don't randomize them so running the test multiple times will start hitting the cache. If you want to access the data via http then run with -Ias.NoCache.
- Add FOnDemandIoStore::DebugFindStreamingChunkIds that allows us to extract the first N FIoChunkIds in the backend.

diff --git a/Engine/Source/Runtime/Experimental/IoStore/OnDemand/Private/OnDemandIoDispatcherBackend.cpp b/Engine/Source/Runtime/Experimental/IoStore/OnDemand/Private/OnDemandIoDispatcherBackend.cpp
--- a/Engine/Source/Runtime/Experimental/IoStore/OnDemand/Private/OnDemandIoDispatcherBackend.cpp
+++ b/Engine/Source/Runtime/Experimental/IoStore/OnDemand/Private/OnDemandIoDispatcherBackend.cpp
@@ -35,6 +35,7 @@
 #include "OnDemandIoStore.h"
 #include "ProfilingDebugging/IoStoreTrace.h"
 #include "Statistics.h"
+#include "String/Numeric.h"
 #include "Tasks/Task.h"
 
 #include <atomic>
@@ -932,6 +933,10 @@ private:
 	void DrainHttpRequests();
 	void UpdateThreadPriorityIfNeeded();
 
+#if UE_IAS_DEBUG_CONSOLE_CMDS
+	void RunRequestTest(const TArray<FString>& Args) const;
+#endif // UE_IAS_DEBUG_CONSOLE_CMDS
+
 
 	FOnDemandIoStore& IoStore;
 	TUniquePtr<IIasCache> Cache;
@@ -999,8 +1004,17 @@ FOnDemandIoBackend::FOnDemandIoBackend(const FOnDemandEndpointConfig& Config, FO
 
 				TickBackendEvent->Trigger();
 			}),
-		ECVF_Cheat)
+		ECVF_Default)
+	);
+
+	DynamicConsoleCommands.Emplace(
+		IConsoleManager::Get().RegisterConsoleCommand(
+			TEXT("ias.RunRequestTest"),
+			TEXT("[optional Number of requests to make] Creates a number of requests for chunks that can be found in that IAS system"),
+			FConsoleCommandWithArgsDelegate::CreateRaw(this, &FOnDemandIoBackend::RunRequestTest),
+			ECVF_Default)
 	);
+
 #endif // UE_IAS_DEBUG_CONSOLE_CMDS
 }
 
@@ -1834,6 +1848,49 @@ int32 FOnDemandIoBackend::WaitForCompleteRequestTasks(float WaitTimeSeconds, flo
 	return ChunkRequests.Num();
 }
 
+#if UE_IAS_DEBUG_CONSOLE_CMDS
+void FOnDemandIoBackend::RunRequestTest(const TArray<FString>& Args) const
+{
+	int32 NumToRequest = 1000;
+
+	if (Args.Num() > 1)
+	{
+		UE_LOG(LogIas, Error, TEXT("Too many args for 'ias.RequestTestCount'"));
+		return;
+	}
+	else if (Args.Num() == 1)
+	{
+		int32 ArgValue = TCString<TCHAR>::Atoi(*Args[0]);
+
+		if (!UE::String::IsNumericOnlyDigits(Args[0]) || ArgValue <= 0)
+		{
+			UE_LOG(LogIas, Error, TEXT("Invalid arg for 'ias.RequestTestCount', value must be a postive integer"));
+			return;
+		}
+
+		NumToRequest = ArgValue;
+	}
+
+	UE_LOG(LogIas, Display, TEXT("Running IAS Request Test with %d requests..."), NumToRequest);
+
+	TArray<FIoChunkId> RequestIds = IoStore.DebugFindStreamingChunkIds(NumToRequest);
+
+	FIoBatch IoBatch;
+	for (const FIoChunkId& Id : RequestIds)
+	{
+		IoBatch.Read(Id, FIoReadOptions(), IoDispatcherPriority_Medium);
+	}
+
+	const double StartTime = FPlatformTime::Seconds();
+
+	IoBatch.IssueWithCallback([StartTime]()
+		{
+			const double TotalTime = FPlatformTime::Seconds() - StartTime;
+			UE_LOG(LogIas, Display, TEXT("IAS Request Test took %.3f(s)"), TotalTime);
+		});
+}
+#endif // UE_IAS_DEBUG_CONSOLE_CMDS
+
 TSharedPtr<IOnDemandIoDispatcherBackend> MakeOnDemandIoDispatcherBackend(
 	const FOnDemandEndpointConfig& Config,
 	FOnDemandIoStore& IoStore,
diff --git a/Engine/Source/Runtime/Experimental/IoStore/OnDemand/Private/OnDemandIoStore.cpp b/Engine/Source/Runtime/Experimental/IoStore/OnDemand/Private/OnDemandIoStore.cpp
--- a/Engine/Source/Runtime/Experimental/IoStore/OnDemand/Private/OnDemandIoStore.cpp
+++ b/Engine/Source/Runtime/Experimental/IoStore/OnDemand/Private/OnDemandIoStore.cpp
@@ -819,6 +819,34 @@ FOnDemandChunkInfo FOnDemandIoStore::GetInstalledChunkInfo(const FIoChunkId& Chu
 	return GetChunkInfo(ChunkId, EOnDemandContainerFlags::Mounted | EOnDemandContainerFlags::InstallOnDemand);
 }
 
+#if !UE_BUILD_SHIPPING
+TArray<FIoChunkId> FOnDemandIoStore::DebugFindStreamingChunkIds(int32 NumToFind)
+{
+	TArray<FIoChunkId> IdsContainer;
+	IdsContainer.Reserve(NumToFind);
+
+	TUniqueLock Lock(ContainerMutex);
+
+	for (FSharedOnDemandContainer& Container : Containers)
+	{
+		if (EnumHasAllFlags(Container->Flags, EOnDemandContainerFlags::Mounted | EOnDemandContainerFlags::StreamOnDemand))
+		{
+			for(const FIoChunkId& Id : Container->ChunkIds)
+			{
+				IdsContainer.Add(Id);
+
+				if (IdsContainer.Num() == NumToFind)
+				{
+					return IdsContainer;
+				}
+			}
+		}
+	}
+
+	return IdsContainer;
+}
+#endif
+
 FIoStatus FOnDemandIoStore::InitializeInstallCache()
 {
 	if (FForkProcessHelper::IsForkRequested() && !FForkProcessHelper::IsForkedChildProcess())
diff --git a/Engine/Source/Runtime/Experimental/IoStore/OnDemand/Private/OnDemandIoStore.h b/Engine/Source/Runtime/Experimental/IoStore/OnDemand/Private/OnDemandIoStore.h
--- a/Engine/Source/Runtime/Experimental/IoStore/OnDemand/Private/OnDemandIoStore.h
+++ b/Engine/Source/Runtime/Experimental/IoStore/OnDemand/Private/OnDemandIoStore.h
@@ -320,8 +320,15 @@ public:
 
 	FIoStatus							Initialize();
 	FIoStatus							InitializeStreamingBackend(const FOnDemandEndpointConfig& EndpointConfig);
+
 	FOnDemandChunkInfo					GetStreamingChunkInfo(const FIoChunkId& ChunkId);
 	FOnDemandChunkInfo					GetInstalledChunkInfo(const FIoChunkId& ChunkId);
+
+#if !UE_BUILD_SHIPPING
+	/** Used to access a number of FIoChunkId that exist in the system when running tests and is not intended for production use */
+	TArray<FIoChunkId>					DebugFindStreamingChunkIds(int32 NumToFind);
+#endif //!UE_BUILD_SHIPPING
+
 	void								ReleaseContent(FOnDemandInternalContentHandle& ContentHandle);
 	void								GetReferencedContent(TArray<FSharedOnDemandContainer>& OutContainers, TArray<TBitArray<>>& OutChunkEntryIndices);
 	void								CancelInstallRequest(FSharedInternalInstallRequest InstallRequest);
