Fix regression where ClothAssetTerminal name is no longer set by default.

diff --git a/Engine/Plugins/ChaosClothAsset/Source/ChaosClothAssetEngine/Private/ChaosClothAsset/ClothAsset.cpp b/Engine/Plugins/ChaosClothAsset/Source/ChaosClothAssetEngine/Private/ChaosClothAsset/ClothAsset.cpp
--- a/Engine/Plugins/ChaosClothAsset/Source/ChaosClothAssetEngine/Private/ChaosClothAsset/ClothAsset.cpp
+++ b/Engine/Plugins/ChaosClothAsset/Source/ChaosClothAssetEngine/Private/ChaosClothAsset/ClothAsset.cpp
@@ -150,6 +150,7 @@ TArray<TSharedRef<FManagedArrayCollection>> TrimOnCook(const FString InAssetName
 
 UChaosClothAsset::UChaosClothAsset(const FObjectInitializer& ObjectInitializer)
 	: Super(ObjectInitializer)
+	, DataflowInstance(this, nullptr, FName(TEXT("ClothAssetTerminal")))
 	, MinQualityLevelLOD(0)
 	, DisableBelowMinLodStripping(FPerPlatformBool(false))
 	, MinLod(0)
diff --git a/Engine/Source/Runtime/Experimental/Dataflow/Engine/Private/Dataflow/DataflowInstance.cpp b/Engine/Source/Runtime/Experimental/Dataflow/Engine/Private/Dataflow/DataflowInstance.cpp
--- a/Engine/Source/Runtime/Experimental/Dataflow/Engine/Private/Dataflow/DataflowInstance.cpp
+++ b/Engine/Source/Runtime/Experimental/Dataflow/Engine/Private/Dataflow/DataflowInstance.cpp
@@ -259,8 +259,10 @@ void FDataflowVariableOverrides::OnDataflowVariablesChanged(const UDataflow* Dat
 // FDataflowInstance
 //---------------------------------------------------------------------------
 
-FDataflowInstance::FDataflowInstance(UObject* InOwner)
-	: VariableOverrides(this)
+FDataflowInstance::FDataflowInstance(UObject* InOwner, UDataflow* InDataflowAsset, FName InTerminalNodeName)
+	: DataflowAsset(InDataflowAsset)
+	, DataflowTerminal(InTerminalNodeName)
+	, VariableOverrides(this)
 	, Owner(InOwner)
 {
 #if WITH_EDITOR
diff --git a/Engine/Source/Runtime/Experimental/Dataflow/Engine/Public/Dataflow/DataflowInstance.h b/Engine/Source/Runtime/Experimental/Dataflow/Engine/Public/Dataflow/DataflowInstance.h
--- a/Engine/Source/Runtime/Experimental/Dataflow/Engine/Public/Dataflow/DataflowInstance.h
+++ b/Engine/Source/Runtime/Experimental/Dataflow/Engine/Public/Dataflow/DataflowInstance.h
@@ -86,7 +86,7 @@ struct FDataflowInstance
 {
 	GENERATED_USTRUCT_BODY();
 
-	DATAFLOWENGINE_API FDataflowInstance(UObject* InOwner = nullptr);
+	DATAFLOWENGINE_API explicit FDataflowInstance(UObject* InOwner = nullptr, UDataflow* InDataflowAsset = nullptr, FName InTerminalNodeName = NAME_None);
 	DATAFLOWENGINE_API ~FDataflowInstance();
 
 public:
