Dataflow : SubGraph nodes : Input / Output / Call

diff --git a/Engine/Source/Runtime/Experimental/Dataflow/Engine/Private/Dataflow/DataflowEngineModule.cpp b/Engine/Source/Runtime/Experimental/Dataflow/Engine/Private/Dataflow/DataflowEngineModule.cpp
--- a/Engine/Source/Runtime/Experimental/Dataflow/Engine/Private/Dataflow/DataflowEngineModule.cpp
+++ b/Engine/Source/Runtime/Experimental/Dataflow/Engine/Private/Dataflow/DataflowEngineModule.cpp
@@ -7,6 +7,7 @@
 
 #include "Dataflow/DataflowTextureAssetNodes.h"
 #include "Dataflow/DataflowVariableNodes.h"
+#include "Dataflow/DataflowSubGraphNodes.h"
 
 /**
  * The public interface to this module
@@ -21,6 +22,7 @@ public:
 	{
 		UE::Dataflow::RegisterTextureAssetNodes();
 		UE::Dataflow::RegisterVariableNodes();
+		UE::Dataflow::RegisterSubGraphNodes();
 	}
 
 	virtual void ShutdownModule() override { }
diff --git a/Engine/Source/Runtime/Experimental/Dataflow/Engine/Private/Dataflow/DataflowSubGraph.cpp b/Engine/Source/Runtime/Experimental/Dataflow/Engine/Private/Dataflow/DataflowSubGraph.cpp
--- a/Engine/Source/Runtime/Experimental/Dataflow/Engine/Private/Dataflow/DataflowSubGraph.cpp
+++ b/Engine/Source/Runtime/Experimental/Dataflow/Engine/Private/Dataflow/DataflowSubGraph.cpp
@@ -1,7 +1,65 @@
 // Copyright Epic Games, Inc. All Rights Reserved.
 
 #include "Dataflow/DataflowSubGraph.h"
+#include "Dataflow/DataflowEdNode.h"
+#include "Dataflow/DataflowSubGraphNodes.h"
+
+namespace UE::Dataflow::DataflowSubGraph::Private
+{
+	template<typename T>
+	T* GetNodeByType(const UEdGraph& EdGraph)
+	{
+		for (UEdGraphNode* EdNode : EdGraph.Nodes)
+		{
+			if (UDataflowEdNode* DataflowEdNode = Cast<UDataflowEdNode>(EdNode))
+			{
+				if (TSharedPtr<FDataflowNode> DataflowNode = DataflowEdNode->GetDataflowNode())
+				{
+					if (T* TypedDataflowNode = DataflowNode->AsType<T>())
+					{
+						return TypedDataflowNode;
+					}
+				}
+			}
+		}
+		return nullptr;
+	}
+}
+
+/////////////////////////////////////////////////////////////////////////////////////////////////////////////
 
 UDataflowSubGraph::UDataflowSubGraph(class FObjectInitializer const& ObjectInitializer)
 	: UEdGraph(ObjectInitializer)
 {}
+
+FDataflowSubGraphInputNode* UDataflowSubGraph::GetInputNode() const
+{
+	// TODO : we should probably cache the node Guid or even a weakPtr on the EdNode upon loading for better perf in the future
+	return UE::Dataflow::DataflowSubGraph::Private::GetNodeByType<FDataflowSubGraphInputNode>(*this);
+}
+
+FDataflowSubGraphOutputNode* UDataflowSubGraph::GetOutputNode() const
+{
+	// TODO : we should probably cache the node Guid or even a weakPtr on the EdNode upon loading for better perf in the future
+	return UE::Dataflow::DataflowSubGraph::Private::GetNodeByType<FDataflowSubGraphOutputNode>(*this);
+}
+
+void UDataflowSubGraph::PostLoad()
+{
+	Super::PostLoad();
+	
+	if (OnSubGraphLoaded.IsBound())
+	{
+		OnSubGraphLoaded.Broadcast(*this);
+	}
+}
+
+FDelegateHandle UDataflowSubGraph::AddOnSubGraphLoadedHandler(const FOnDataflowSubGraphLoaded::FDelegate& InHandler)
+{
+	return OnSubGraphLoaded.Add(InHandler);
+}
+
+void UDataflowSubGraph::RemoveOnSubGraphLoadedHandler(FDelegateHandle Handle)
+{
+	OnSubGraphLoaded.Remove(Handle);
+}
\ No newline at end of file
diff --git a/Engine/Source/Runtime/Experimental/Dataflow/Engine/Private/Dataflow/DataflowSubGraphNodes.cpp b/Engine/Source/Runtime/Experimental/Dataflow/Engine/Private/Dataflow/DataflowSubGraphNodes.cpp
new file mode 100644
--- /dev/null
+++ b/Engine/Source/Runtime/Experimental/Dataflow/Engine/Private/Dataflow/DataflowSubGraphNodes.cpp
@@ -0,0 +1,574 @@
+// Copyright Epic Games, Inc. All Rights Reserved.
+
+#include "Dataflow/DataflowSubGraphNodes.h"
+
+#include "Dataflow/DataflowNode.h"
+#include "Dataflow/DataflowEdNode.h"
+#include "Dataflow/DataflowObject.h"
+#include "Dataflow/DataflowSubGraph.h"
+#include "Dataflow/DataflowNodeFactory.h"
+#include "GraphEditAction.h"
+
+namespace UE::Dataflow::SubGraphNodes::Private
+{
+	class FContextSubGraph: public FContext
+	{
+	public:
+		DATAFLOW_CONTEXT_INTERNAL(FContext, FContextSubGraph);
+
+		FContextSubGraph(FContext& InParentContext, const ISubGraphContextCallback& InSubGraphCallback)
+			: ParentContext(InParentContext)
+			, SubGraphCallback(InSubGraphCallback)
+		{}
+
+		virtual void SetDataImpl(FContextCacheKey Key, TUniquePtr<FContextCacheElementBase>&& DataStoreEntry) override
+		{
+			DataStore.Emplace(Key, MoveTemp(DataStoreEntry));
+		}
+
+		virtual const TUniquePtr<FContextCacheElementBase>* GetDataImpl(FContextCacheKey Key) const override
+		{
+			return DataStore.Find(Key);
+		}
+
+		virtual bool HasDataImpl(FContextCacheKey Key, FTimestamp InTimestamp = FTimestamp::Invalid) const override
+		{
+			return DataStore.Contains(Key) && DataStore[Key]->GetTimestamp() >= InTimestamp;
+		}
+
+		virtual bool IsEmptyImpl() const override
+		{
+			return DataStore.IsEmpty();
+		}
+
+		virtual void Evaluate(const FDataflowNode* Node, const FDataflowOutput* Output) override
+		{
+			BeginContextEvaluation(Node, Output);
+		}
+
+		virtual bool Evaluate(const FDataflowOutput& Connection) override
+		{
+			return Connection.EvaluateImpl(*this);
+		}
+
+		bool EvaluateSubGraph(const FDataflowSubGraphOutputNode& SubGraphOutputNode, const FDataflowOutput& CallerOutput)
+		{
+			if (const FDataflowInput* InputToPull = SubGraphOutputNode.FindInput(CallerOutput.GetName()))
+			{
+				if (const FDataflowOutput* OutputToEvaluate = InputToPull->GetConnection())
+				{
+					// we need to make sure the graph output and caller output are of compatible type
+					if (UE::Dataflow::AreTypesCompatible(OutputToEvaluate->GetType(), CallerOutput.GetType()))
+					{
+						if (Evaluate(*OutputToEvaluate))
+						{
+							// copy the cache value from this cachestore to the parent one 
+							const TUniquePtr<FContextCacheElementBase>* CacheEntryToClone = GetDataImpl(OutputToEvaluate->CacheKey());
+							if (CacheEntryToClone && CacheEntryToClone->IsValid())
+							{
+								TUniquePtr<FContextCacheElementBase> ClonedCacheEntry = (*CacheEntryToClone)->Clone(*this);
+								ParentContext.SetDataImpl(CallerOutput.CacheKey(), MoveTemp(ClonedCacheEntry));
+								return true;
+							}
+						}
+					}
+				}
+			}
+			return false;
+		}
+
+		void EvaluateParentContext(const FDataflowOutput& SubGraphOutput)
+		{
+			SubGraphCallback.EvaluateParentContext(ParentContext, *this, SubGraphOutput);
+		}
+
+	private:
+		FContext& ParentContext;
+		const UE::Dataflow::ISubGraphContextCallback& SubGraphCallback;
+	};
+
+}
+
+//////////////////////////////////////////////////////////////////////////////////////////////////////////
+
+namespace UE::Dataflow
+{
+	void RegisterSubGraphNodes()
+	{
+		DATAFLOW_NODE_REGISTER_CREATION_FACTORY(FDataflowSubGraphInputNode);
+		DATAFLOW_NODE_REGISTER_CREATION_FACTORY(FDataflowSubGraphOutputNode);
+		DATAFLOW_NODE_REGISTER_CREATION_FACTORY(FDataflowCallSubGraphNode);
+	}
+}
+
+//////////////////////////////////////////////////////////////////////////////////////////////////////////
+
+FDataflowDynamicConnections::FDataflowDynamicConnections()
+	: PinDirection(UE::Dataflow::FPin::EDirection::NONE)
+	, OwnerInterface(nullptr)
+	, DataflowAssetWeakPtr(nullptr)
+{}
+
+FDataflowDynamicConnections::FDataflowDynamicConnections(UE::Dataflow::FPin::EDirection InPinDirection, IOwnerInterface* InOwnerInterface, UDataflow* InDataflowAsset)
+	: PinDirection(InPinDirection)
+	, OwnerInterface(InOwnerInterface)
+	, DataflowAssetWeakPtr(InDataflowAsset)
+{}
+
+
+FDataflowDynamicConnections::FConnectionReference FDataflowDynamicConnections::GetConnectionReference(int32 Index) const
+{
+	return { &DynamicProperties[Index], Index, &DynamicProperties };
+}
+
+void FDataflowDynamicConnections::Refresh()
+{
+	using namespace UE::Dataflow;
+
+	if (!OwnerInterface)
+	{
+		return;
+	}
+
+	FDataflowNode* OwnerNode = OwnerInterface->GetOwner(this);
+	if (!OwnerNode)
+	{
+		return;
+	}
+
+	TStrongObjectPtr<UDataflow> DataflowAsset = DataflowAssetWeakPtr.Pin();
+	if (!DataflowAsset.IsValid())
+	{
+		return;
+	}
+
+	TSharedPtr<UE::Dataflow::FGraph, ESPMode::ThreadSafe> DataflowGraph = DataflowAsset->GetDataflow();
+	if (!DataflowGraph.IsValid())
+	{
+		return;
+	}
+
+	// save the remote connections from the dynamic connections
+	using FRemoteConnections = TArray<FDataflowConnection*>;
+	TMap<FGuid, FRemoteConnections> RemoteConnectionsByPropertyId;
+
+	for (FDataflowConnection* DynamicConnection : GetDynamicConnections())
+	{
+		if (DynamicConnection)
+		{
+			if (const FGuid* PropertyId = ConnectionIdToPropertyId.Find(DynamicConnection->GetGuid()))
+			{
+				FRemoteConnections& RemoteConnections = RemoteConnectionsByPropertyId.Emplace(*PropertyId);
+				DynamicConnection->GetConnections(RemoteConnections);
+				DataflowGraph->ClearConnections(DynamicConnection);
+			}
+		}
+	}
+	// clear all the dynamic connections
+	ClearDynamicConnections();
+
+	ConnectionIdToPropertyId.Reset();
+
+	// Go through all the properties of the property bag recreate the output and reconnect them if needed
+	const TConnectionReference<FDataflowAllTypes> NullReference(nullptr);
+	const FInstancedPropertyBag& PropertyBag = OwnerInterface->GetPropertyBag(this);
+	if (const UPropertyBag* PropertyBagStruct = PropertyBag.GetPropertyBagStruct())
+	{
+		for (const FPropertyBagPropertyDesc& PropertyDesc : PropertyBagStruct->GetPropertyDescs())
+		{
+			// we can only add the output if it is a supported type
+			if (IsSupportedType(PropertyDesc))
+			{
+				const FName CppType = GetCppTypeFromPropertyDesc(PropertyDesc);
+				if (FDataflowConnection* NewConnection = AddNewConnectionFromPropertyDesc(PropertyDesc))
+				{
+					// now try to reconnect what was connected before
+					if (FRemoteConnections* RemoteConnections = RemoteConnectionsByPropertyId.Find(PropertyDesc.ID))
+					{
+						for (FDataflowConnection* RemoteConnection : *RemoteConnections)
+						{
+							DataflowGraph->Connect(NewConnection, RemoteConnection);
+						}
+					}
+				}
+			}
+		}
+		// Todo : need to also remove the one that are no longer necessary
+		// ( need to disconnect if needed ? or keep the one that are still connected ?)
+		OwnerNode->Invalidate();
+
+		DataflowAsset->RefreshEdNodeByGuid(OwnerNode->GetGuid());
+	}
+}
+
+FDataflowConnection* FDataflowDynamicConnections::AddNewConnectionFromPropertyDesc(const FPropertyBagPropertyDesc& PropertyDesc)
+{
+	using namespace UE::Dataflow;
+
+	if (FDataflowNode* OwnerNode = OwnerInterface ? OwnerInterface->GetOwner(this) : nullptr)
+	{
+		const int32 Index = DynamicProperties.AddDefaulted();
+		TConnectionReference<FDataflowAllTypes> ConnectionReference = GetConnectionReference(Index);
+		if (FDataflowConnection* NewConnection = CreateConnection(ConnectionReference))
+		{
+			SetConnectionTypeFromPropertyDesc(*NewConnection, PropertyDesc);
+			if (PinDirection == UE::Dataflow::FPin::EDirection::INPUT)
+			{
+				OwnerNode->TryRenameInput(ConnectionReference, PropertyDesc.Name);
+			}
+			else if (PinDirection == UE::Dataflow::FPin::EDirection::OUTPUT)
+			{
+				OwnerNode->TryRenameOutput(ConnectionReference, PropertyDesc.Name);
+			}
+			NewConnection->LockType();
+
+			ConnectionIdToPropertyId.Add(NewConnection->GetGuid(), PropertyDesc.ID);
+			return NewConnection;
+		}
+	}
+	return nullptr;
+}
+
+bool FDataflowDynamicConnections::IsSupportedType(const FPropertyBagPropertyDesc& Desc)
+{
+	switch (Desc.ValueType)
+	{
+	case EPropertyBagPropertyType::Bool:
+	case EPropertyBagPropertyType::Byte:
+	case EPropertyBagPropertyType::Int32:
+	case EPropertyBagPropertyType::Int64:
+	case EPropertyBagPropertyType::Float:
+	case EPropertyBagPropertyType::Double:
+	case EPropertyBagPropertyType::Name:
+	case EPropertyBagPropertyType::String:
+	case EPropertyBagPropertyType::Text:
+	case EPropertyBagPropertyType::Object:
+		return true;
+
+		// unsupported type for now 
+	case EPropertyBagPropertyType::Enum:
+	case EPropertyBagPropertyType::Struct:
+	case EPropertyBagPropertyType::SoftObject:
+	case EPropertyBagPropertyType::Class:
+	case EPropertyBagPropertyType::SoftClass:
+	case EPropertyBagPropertyType::UInt32:	// Type not fully supported at UI, will work with restrictions to type editing
+	case EPropertyBagPropertyType::UInt64: // Type not fully supported at UI, will work with restrictions to type editing
+	default:
+		return false;
+	}
+}
+
+FName FDataflowDynamicConnections::GetCppTypeFromPropertyDesc(const FPropertyBagPropertyDesc& Desc)
+{
+	FString ExtendedType;
+	const FString CPPType = Desc.CachedProperty->GetCPPType(&ExtendedType);
+	return FName(CPPType + ExtendedType);
+}
+
+bool FDataflowDynamicConnections::SetConnectionTypeFromPropertyDesc(FDataflowConnection& Connection, const FPropertyBagPropertyDesc& Desc)
+{
+	if (FDataflowNode* OwnerNode = OwnerInterface ? OwnerInterface->GetOwner(this) : nullptr)
+	{
+		FName ConcreteType(GetCppTypeFromPropertyDesc(Desc));
+
+		switch (Desc.ValueType)
+		{
+		case EPropertyBagPropertyType::Double:
+		case EPropertyBagPropertyType::Float:
+			// the UI shows float but behind the scene set a double property let only deal with float 
+			return Connection.SetConcreteType(TDataflowPolicyTypeName<float>::GetName());
+
+		case EPropertyBagPropertyType::Bool:	
+		case EPropertyBagPropertyType::Byte:
+		case EPropertyBagPropertyType::Int32:
+		case EPropertyBagPropertyType::Int64:
+		case EPropertyBagPropertyType::Name:
+		case EPropertyBagPropertyType::String:
+		case EPropertyBagPropertyType::Text:
+		case EPropertyBagPropertyType::Object:
+			return Connection.SetConcreteType(ConcreteType);
+
+		// unsupported type for now 
+		case EPropertyBagPropertyType::Enum:
+		case EPropertyBagPropertyType::Struct:
+		case EPropertyBagPropertyType::SoftObject:
+		case EPropertyBagPropertyType::Class:
+		case EPropertyBagPropertyType::SoftClass:
+		case EPropertyBagPropertyType::UInt32:	// Type not fully supported at UI, will work with restrictions to type editing
+		case EPropertyBagPropertyType::UInt64: // Type not fully supported at UI, will work with restrictions to type editing
+		default:
+			return false;
+		}
+	}
+	return false;
+}
+
+FDataflowConnection* FDataflowDynamicConnections::CreateConnection(FDataflowDynamicConnections::FConnectionReference ConnectionReference)
+{
+	if (FDataflowNode* OwnerNode = OwnerInterface ? OwnerInterface->GetOwner(this) : nullptr)
+	{
+		if (PinDirection == UE::Dataflow::FPin::EDirection::INPUT)
+		{
+			return &OwnerNode->RegisterInputArrayConnection(ConnectionReference);
+		}
+		else if (PinDirection == UE::Dataflow::FPin::EDirection::OUTPUT)
+		{
+			return &OwnerNode->RegisterOutputArrayConnection(ConnectionReference);
+		}
+	}
+	return nullptr;
+}
+
+TArray<FDataflowConnection*> FDataflowDynamicConnections::GetDynamicConnections() const
+{
+	TArray<FDataflowConnection*> Connections;
+	if (FDataflowNode* OwnerNode = OwnerInterface ? OwnerInterface->GetOwner(this) : nullptr)
+	{
+		if (PinDirection == UE::Dataflow::FPin::EDirection::INPUT)
+		{
+			Connections.Append(OwnerNode->GetInputs());
+		}
+		else if (PinDirection == UE::Dataflow::FPin::EDirection::OUTPUT)
+		{
+			Connections.Append(OwnerNode->GetOutputs());
+		}
+	}
+	return Connections;
+}
+
+void FDataflowDynamicConnections::ClearDynamicConnections()
+{
+	if (FDataflowNode* OwnerNode = OwnerInterface ? OwnerInterface->GetOwner(this) : nullptr)
+	{
+		if (PinDirection == UE::Dataflow::FPin::EDirection::INPUT)
+		{
+			OwnerNode->ClearInputs();
+		}
+		else if (PinDirection == UE::Dataflow::FPin::EDirection::OUTPUT)
+		{
+			OwnerNode->ClearOutputs();
+		}
+	}
+}
+
+//////////////////////////////////////////////////////////////////////////////////////////////////////////
+
+FDataflowSubGraphInputNode::FDataflowSubGraphInputNode(const UE::Dataflow::FNodeParameters& InParam, FGuid InGuid)
+	: FDataflowNode(InParam, InGuid)
+	, DynamicConnections(UE::Dataflow::FPin::EDirection::OUTPUT, this, Cast<UDataflow>(InParam.OwningObject))
+{
+}
+
+void FDataflowSubGraphInputNode::Evaluate(UE::Dataflow::FContext& Context, const FDataflowOutput* Out) const
+{
+	using namespace UE::Dataflow::SubGraphNodes::Private;
+
+	check(Out);
+	if (FContextSubGraph* SubGraphContext = Context.AsType<FContextSubGraph>())
+	{
+		SubGraphContext->EvaluateParentContext(*Out);
+	}
+	else
+	{
+		Out->SetNullValue(Context);
+	}
+}
+
+void FDataflowSubGraphInputNode::PostSerialize(const FArchive& Ar)
+{
+	if (Ar.IsLoading())
+	{
+		// make sure the node is up to date with the serialized data
+		DynamicConnections.Refresh();
+	}
+}
+
+void FDataflowSubGraphInputNode::OnPropertyChanged(UE::Dataflow::FContext& Context, const FPropertyChangedEvent& InPropertyChangedEvent)
+{
+	DynamicConnections.Refresh();
+}
+
+//////////////////////////////////////////////////////////////////////////////////////////////////////////
+
+FDataflowSubGraphOutputNode::FDataflowSubGraphOutputNode(const UE::Dataflow::FNodeParameters& InParam, FGuid InGuid)
+	: FDataflowNode(InParam, InGuid)
+	, DynamicConnections(UE::Dataflow::FPin::EDirection::INPUT, this, Cast<UDataflow>(InParam.OwningObject))
+{
+}
+
+void FDataflowSubGraphOutputNode::Evaluate(UE::Dataflow::FContext& Context, const FDataflowOutput* Out) const
+{
+	// nothing to do, the subgraph call node will be pull the inputs directly 
+	// and this node has no outputs
+}
+
+void FDataflowSubGraphOutputNode::PostSerialize(const FArchive& Ar)
+{
+	if (Ar.IsLoading())
+	{
+		// make sure the node is up to date with the serialized data
+		DynamicConnections.Refresh();
+	}
+}
+
+void FDataflowSubGraphOutputNode::OnPropertyChanged(UE::Dataflow::FContext& Context, const FPropertyChangedEvent& InPropertyChangedEvent)
+{
+	DynamicConnections.Refresh();
+}
+
+//////////////////////////////////////////////////////////////////////////////////////////////////////////
+
+FDataflowCallSubGraphNode::FDataflowCallSubGraphNode(const UE::Dataflow::FNodeParameters& InParam, FGuid InGuid)
+	: FDataflowNode(InParam, InGuid)
+	, DynamicInputs(UE::Dataflow::FPin::EDirection::INPUT, this, Cast<UDataflow>(InParam.OwningObject))
+	, DynamicOutputs(UE::Dataflow::FPin::EDirection::OUTPUT, this, Cast<UDataflow>(InParam.OwningObject))
+{
+	DataflowAssetWeakPtr = Cast<UDataflow>(InParam.OwningObject);
+}
+
+FDataflowCallSubGraphNode::~FDataflowCallSubGraphNode()
+{
+	if (TStrongObjectPtr<UDataflow> DataflowAsset = DataflowAssetWeakPtr.Pin())
+	{
+		if (UDataflowSubGraph* SubGraph = DataflowAsset->FindSubGraphByName(SubGraphName))
+		{
+			SubGraph->RemoveOnGraphChangedHandler(OnGraphChangedHandle);
+			SubGraph->RemoveOnSubGraphLoadedHandler(OnSubGraphLoadedHandle);
+		}
+	}
+}
+
+void FDataflowCallSubGraphNode::PostSerialize(const FArchive& Ar)
+{
+	if (Ar.IsLoading())
+	{
+		// make sure the node is up to date with the serialized data
+		SetSubGraphName(SubGraphName);
+	}
+}
+
+void FDataflowCallSubGraphNode::SetSubGraphName(FName InSubGraphName)
+{
+	if (TStrongObjectPtr<UDataflow> DataflowAsset = DataflowAssetWeakPtr.Pin())
+	{
+		if (UDataflowSubGraph* PreviousSubGraph = DataflowAsset->FindSubGraphByName(SubGraphName))
+		{
+			PreviousSubGraph->RemoveOnGraphChangedHandler(OnGraphChangedHandle);
+			PreviousSubGraph->RemoveOnSubGraphLoadedHandler(OnSubGraphLoadedHandle);
+		}
+
+		SubGraphName = InSubGraphName;
+
+		if (UDataflowSubGraph* NewSubGraph = DataflowAsset->FindSubGraphByName(SubGraphName))
+		{
+			OnSubGraphLoadedHandle = NewSubGraph->AddOnGraphChangedHandler(FOnGraphChanged::FDelegate::CreateRaw(this, &FDataflowCallSubGraphNode::OnSubGraphChanged));
+			OnGraphChangedHandle = NewSubGraph->AddOnSubGraphLoadedHandler(FOnDataflowSubGraphLoaded::FDelegate::CreateRaw(this, &FDataflowCallSubGraphNode::OnSubGraphLoaded));
+		}
+
+		// make sure we have at least an initial refresh 
+		DynamicInputs.Refresh();
+		DynamicOutputs.Refresh();
+	}
+}
+
+void FDataflowCallSubGraphNode::OnSubGraphChanged(const struct FEdGraphEditAction& InAction)
+{
+	bool bInputsChanged = false;
+	bool bOutputsChanged = false;
+
+	for (const UEdGraphNode* EdNode : InAction.Nodes)
+	{
+		// Todo : use the template function in SubGraph code ? 
+		if (const UDataflowEdNode* DataflowEdNode = Cast<const UDataflowEdNode>(EdNode))
+		{
+			if (TSharedPtr<const FDataflowNode> DataflowNode = DataflowEdNode->GetDataflowNode())
+			{
+				if (const FDataflowSubGraphInputNode* SubGraphInputNode = DataflowNode->AsType<FDataflowSubGraphInputNode>())
+				{
+					bInputsChanged = true;
+				}
+				else if (const FDataflowSubGraphOutputNode* SubGraphOutputNode = DataflowNode->AsType<FDataflowSubGraphOutputNode>())
+				{
+					bOutputsChanged = true;
+				}
+			}
+		}
+	}
+	if (bInputsChanged)
+	{
+		DynamicInputs.Refresh();
+	}
+	if (bOutputsChanged)
+	{
+		DynamicOutputs.Refresh();
+	}
+}
+
+void FDataflowCallSubGraphNode::OnSubGraphLoaded(const UDataflowSubGraph& InSubGraph)
+{
+	DynamicInputs.Refresh();
+	DynamicOutputs.Refresh();
+}
+
+const FInstancedPropertyBag& FDataflowCallSubGraphNode::GetPropertyBag(const FDataflowDynamicConnections* Caller)
+{
+	static FInstancedPropertyBag DefaultPropertyBag;
+	if (TStrongObjectPtr<UDataflow> DataflowAsset = DataflowAssetWeakPtr.Pin())
+	{
+		if (UDataflowSubGraph* SubGraphToCall = DataflowAsset->FindSubGraphByName(SubGraphName))
+		{
+			if (Caller == &DynamicInputs)
+			{
+				if (FDataflowSubGraphInputNode* SubGraphInputNode = SubGraphToCall->GetInputNode())
+				{
+					return SubGraphInputNode->GetPropertyBag(Caller);
+				}
+			}
+			else if (Caller == &DynamicOutputs)
+			{
+				if (FDataflowSubGraphOutputNode* SubGraphOutputNode = SubGraphToCall->GetOutputNode())
+				{
+					return SubGraphOutputNode->GetPropertyBag(Caller);
+				}
+			}
+		}
+	}
+	return DefaultPropertyBag;
+}
+
+void FDataflowCallSubGraphNode::EvaluateParentContext(UE::Dataflow::FContext& Context, UE::Dataflow::FContext& SubGraphContext, const FDataflowOutput& SubGraphOutput) const
+{
+	if (const FDataflowInput* Input = FindInput(SubGraphOutput.GetName()))
+	{
+		Input->PullValue(Context);
+		if (const FDataflowOutput* Output = Input->GetConnection())
+		{
+			const TUniquePtr<UE::Dataflow::FContextCacheElementBase>* CacheEntryToClone = Context.GetDataImpl(Output->CacheKey());
+			if (CacheEntryToClone && CacheEntryToClone->IsValid())
+			{
+				TUniquePtr<UE::Dataflow::FContextCacheElementBase> ClonedCacheEntry = (*CacheEntryToClone)->Clone(Context);
+				SubGraphContext.SetDataImpl(SubGraphOutput.CacheKey(), MoveTemp(ClonedCacheEntry));
+				return;
+			}
+		}
+	}
+	// default pass when everything else failed
+	SubGraphOutput.SetNullValue(SubGraphContext);
+}
+
+void FDataflowCallSubGraphNode::Evaluate(UE::Dataflow::FContext& Context, const FDataflowOutput* Out) const
+{
+	check(Out);
+	if (TStrongObjectPtr<UDataflow> DataflowAsset = DataflowAssetWeakPtr.Pin())
+	{
+		if (UDataflowSubGraph* SubGraphToCall = DataflowAsset->FindSubGraphByName(SubGraphName))
+		{
+			if (const FDataflowSubGraphOutputNode* SubGraphOutputNode = SubGraphToCall->GetOutputNode())
+			{
+				// pulling the outputs and when the graph evaluate the remote left end of its graph OnInputValueRequested will be called
+				UE::Dataflow::SubGraphNodes::Private::FContextSubGraph SubGraphContext(Context, *this);
+				SubGraphContext.EvaluateSubGraph(*SubGraphOutputNode, *Out);
+			}
+		}
+	}
+}
diff --git a/Engine/Source/Runtime/Experimental/Dataflow/Engine/Public/Dataflow/DataflowSubGraph.h b/Engine/Source/Runtime/Experimental/Dataflow/Engine/Public/Dataflow/DataflowSubGraph.h
--- a/Engine/Source/Runtime/Experimental/Dataflow/Engine/Public/Dataflow/DataflowSubGraph.h
+++ b/Engine/Source/Runtime/Experimental/Dataflow/Engine/Public/Dataflow/DataflowSubGraph.h
@@ -3,20 +3,30 @@
 #pragma once
 
 #include "EdGraph/EdGraph.h"
-#include "StructUtils/PropertyBag.h"
 
 #include "DataflowSubGraph.generated.h"
 
+struct FDataflowSubGraphInputNode;
+struct FDataflowSubGraphOutputNode;
+
+class UDataflowSubGraph;
+
+DECLARE_MULTICAST_DELEGATE_OneParam(FOnDataflowSubGraphLoaded, const UDataflowSubGraph&);
+
 UCLASS(MinimalAPI)
 class UDataflowSubGraph : public UEdGraph
 {
 	GENERATED_UCLASS_BODY()
 
 public:
-	const FInstancedPropertyBag& GetInputs() const { return Inputs; }
-	const FInstancedPropertyBag& GetOutputs() const { return Outputs; }
+	FDataflowSubGraphInputNode* GetInputNode() const;
+	FDataflowSubGraphOutputNode* GetOutputNode() const;
+
+	FDelegateHandle AddOnSubGraphLoadedHandler(const FOnDataflowSubGraphLoaded::FDelegate& InHandler);
+	void RemoveOnSubGraphLoadedHandler(FDelegateHandle Handle);
+
+private:
+	virtual void PostLoad() override;
 
-private: 
-	FInstancedPropertyBag Inputs;
-	FInstancedPropertyBag Outputs;
+	FOnDataflowSubGraphLoaded OnSubGraphLoaded;
 };
\ No newline at end of file
diff --git a/Engine/Source/Runtime/Experimental/Dataflow/Engine/Public/Dataflow/DataflowSubGraphNodes.h b/Engine/Source/Runtime/Experimental/Dataflow/Engine/Public/Dataflow/DataflowSubGraphNodes.h
new file mode 100644
--- /dev/null
+++ b/Engine/Source/Runtime/Experimental/Dataflow/Engine/Public/Dataflow/DataflowSubGraphNodes.h
@@ -0,0 +1,182 @@
+// Copyright Epic Games, Inc. All Rights Reserved.
+
+#pragma once
+
+#include "Dataflow/DataflowConnection.h"
+#include "Dataflow/DataflowNodeParameters.h"
+#include "Dataflow/DataflowNode.h"
+#include "StructUtils/PropertyBag.h"
+
+#include "DataflowSubGraphNodes.generated.h"
+
+struct FDataflowOutput;
+struct FPropertyBagPropertyDesc;
+class UDataflow;
+class UDataflowSubGraph;
+
+namespace UE::Dataflow
+{
+	struct ISubGraphContextCallback
+	{
+		virtual void EvaluateParentContext(UE::Dataflow::FContext& Context, UE::Dataflow::FContext& SubGraphContext, const FDataflowOutput& SubGraphOutput) const = 0;
+	};
+}
+
+USTRUCT()
+struct FDataflowDynamicConnections
+{
+	GENERATED_USTRUCT_BODY()
+
+	struct IOwnerInterface
+	{
+		virtual FDataflowNode* GetOwner(const FDataflowDynamicConnections* Caller) = 0;
+		virtual const FInstancedPropertyBag& GetPropertyBag(const FDataflowDynamicConnections* Caller) = 0;
+	};
+
+public:
+	FDataflowDynamicConnections();
+	FDataflowDynamicConnections(UE::Dataflow::FPin::EDirection PInDirection, IOwnerInterface* OwnerInterface, UDataflow* DataflowAsset);
+
+	void Refresh();
+
+private:
+	using FConnectionReference = UE::Dataflow::TConnectionReference<FDataflowAllTypes>;
+	FConnectionReference GetConnectionReference(int32 Index) const;
+
+	bool IsSupportedType(const FPropertyBagPropertyDesc& PropertyDesc);
+	FDataflowConnection* CreateConnection(FConnectionReference Connectionreference);
+	TArray<FDataflowConnection*> GetDynamicConnections() const;
+	void ClearDynamicConnections();
+
+	bool SetConnectionTypeFromPropertyDesc(FDataflowConnection& Connection, const FPropertyBagPropertyDesc& PropertyDesc);
+	FName GetCppTypeFromPropertyDesc(const FPropertyBagPropertyDesc& PropertyDesc);
+	FDataflowConnection* AddNewConnectionFromPropertyDesc(const FPropertyBagPropertyDesc& PropertyDesc);
+
+private:
+	UPROPERTY()
+	TArray<FDataflowAllTypes> DynamicProperties;
+
+	TMap<FGuid, FGuid> ConnectionIdToPropertyId;
+
+	UE::Dataflow::FPin::EDirection PinDirection = UE::Dataflow::FPin::EDirection::NONE;
+	IOwnerInterface* OwnerInterface;
+	TWeakObjectPtr<UDataflow> DataflowAssetWeakPtr;
+};
+
+/**
+* This node represent the inputs of a dataflow subgraph 
+*/
+USTRUCT(meta=(Icon="GraphEditor.Function_16x"))
+struct FDataflowSubGraphInputNode
+	: public FDataflowNode
+	, public FDataflowDynamicConnections::IOwnerInterface
+{
+	GENERATED_USTRUCT_BODY()
+
+	DATAFLOW_NODE_DEFINE_INTERNAL(FDataflowSubGraphInputNode, "SubGraphInput", "SubGraph", "")
+
+public:
+	FDataflowSubGraphInputNode(const UE::Dataflow::FNodeParameters& Param, FGuid InGuid = FGuid::NewGuid());
+
+	/* Begin IOwnerInterface */
+	virtual FDataflowNode* GetOwner(const FDataflowDynamicConnections* Caller) override { return this; }
+	virtual const FInstancedPropertyBag& GetPropertyBag(const FDataflowDynamicConnections* Caller) override { return PropertyBag; }
+	/* End IOwnerInterface */
+
+private:
+	virtual void Evaluate(UE::Dataflow::FContext& Context, const FDataflowOutput* Out) const override;
+	virtual void PostSerialize(const FArchive& Ar) override;
+	virtual void OnPropertyChanged(UE::Dataflow::FContext& Context, const FPropertyChangedEvent& InPropertyChangedEvent) override;
+
+	UPROPERTY()
+	FDataflowDynamicConnections DynamicConnections;
+
+	UPROPERTY(EditAnyWhere, Category = Inputs, Meta = (DisplayName = "Inputs"))
+	FInstancedPropertyBag PropertyBag;
+};
+
+
+/**
+* This node represent the Outputs of a dataflow subgraph
+*/
+USTRUCT(meta = (Icon = "GraphEditor.Function_16x"))
+struct FDataflowSubGraphOutputNode 
+	: public FDataflowNode
+	, public FDataflowDynamicConnections::IOwnerInterface
+{
+	GENERATED_USTRUCT_BODY()
+
+	DATAFLOW_NODE_DEFINE_INTERNAL(FDataflowSubGraphOutputNode, "SubGraphOutput", "SubGraph", "")
+
+public:
+	FDataflowSubGraphOutputNode(const UE::Dataflow::FNodeParameters& Param, FGuid InGuid = FGuid::NewGuid());
+
+public:
+	/* Begin IOwnerInterface */
+	virtual FDataflowNode* GetOwner(const FDataflowDynamicConnections* Caller) override { return this; }
+	virtual const FInstancedPropertyBag& GetPropertyBag(const FDataflowDynamicConnections* Caller) override { return PropertyBag; }
+	/* End IOwnerInterface */
+
+private:
+	virtual void Evaluate(UE::Dataflow::FContext& Context, const FDataflowOutput* Out) const override;
+	virtual void PostSerialize(const FArchive& Ar) override;
+	virtual void OnPropertyChanged(UE::Dataflow::FContext& Context, const FPropertyChangedEvent& InPropertyChangedEvent) override;
+
+	UPROPERTY()
+	FDataflowDynamicConnections DynamicConnections;
+
+	UPROPERTY(EditAnyWhere, Category = Inputs, Meta = (DisplayName = "Inputs"))
+	FInstancedPropertyBag PropertyBag;
+};
+
+
+USTRUCT(meta = (Icon = "GraphEditor.Function_16x"))
+struct FDataflowCallSubGraphNode
+	: public FDataflowNode
+	, public FDataflowDynamicConnections::IOwnerInterface
+	, public UE::Dataflow::ISubGraphContextCallback
+{
+	GENERATED_USTRUCT_BODY()
+
+	DATAFLOW_NODE_DEFINE_INTERNAL(FDataflowCallSubGraphNode, "CallSubGraph", "SubGraph", "")
+
+public:
+	DATAFLOWENGINE_API FDataflowCallSubGraphNode(const UE::Dataflow::FNodeParameters& Param, FGuid InGuid = FGuid::NewGuid());
+	DATAFLOWENGINE_API virtual ~FDataflowCallSubGraphNode() override;
+
+	DATAFLOWENGINE_API void SetSubGraphName(FName InSubGraphName);
+
+private:
+	virtual void Evaluate(UE::Dataflow::FContext& Context, const FDataflowOutput* Out) const override;
+	virtual void PostSerialize(const FArchive& Ar) override;
+
+	/* Begin ISubGraphContextCallback */
+	virtual void EvaluateParentContext(UE::Dataflow::FContext& Context, UE::Dataflow::FContext& SubGraphContext, const FDataflowOutput& SubGraphOutput) const override;
+	/* End ISubGraphContextCallback */
+
+	/* Begin IOwnerInterface */
+	virtual FDataflowNode* GetOwner(const FDataflowDynamicConnections* Caller) override { return this; }
+	virtual const FInstancedPropertyBag& GetPropertyBag(const FDataflowDynamicConnections* Caller) override;
+	/* End IOwnerInterface */
+
+	void OnSubGraphChanged(const struct FEdGraphEditAction& InAction);
+	void OnSubGraphLoaded(const UDataflowSubGraph& InSubGraph);
+
+	UPROPERTY()
+	FName SubGraphName; // TODO : should probably be a GUID
+
+	UPROPERTY()
+	FDataflowDynamicConnections DynamicInputs;
+
+	UPROPERTY()
+	FDataflowDynamicConnections DynamicOutputs;
+
+	TWeakObjectPtr<UDataflow> DataflowAssetWeakPtr;
+	FDelegateHandle OnGraphChangedHandle;
+	FDelegateHandle OnSubGraphLoadedHandle;
+};
+
+namespace UE::Dataflow
+{
+	void RegisterSubGraphNodes();
+}
