Add a new cvar 'ias.EnableCancelation' that allows us to disable the cancelation of IAS requests.

#rnx

- The cvar is true by default which allows cancelations to be processed as normal.
- This is intended for development use (it can have it's uses when debugging) and so will be compiled out in shipping configuration.

diff --git a/Engine/Source/Runtime/Experimental/IoStore/OnDemand/Private/OnDemandIoDispatcherBackend.cpp b/Engine/Source/Runtime/Experimental/IoStore/OnDemand/Private/OnDemandIoDispatcherBackend.cpp
--- a/Engine/Source/Runtime/Experimental/IoStore/OnDemand/Private/OnDemandIoDispatcherBackend.cpp
+++ b/Engine/Source/Runtime/Experimental/IoStore/OnDemand/Private/OnDemandIoDispatcherBackend.cpp
@@ -50,6 +50,9 @@
 /** When enabled we will run some limited testing on start up for issues that are hard to reproduce with normal gameplay */
 #define UE_ENABLE_IAS_TESTING 0 && !UE_BUILD_SHIPPING
 
+/** When enabled it is possible to disable request cancelation via the cvar '' */
+#define UE_ALLOW_DISABLE_CANCELLING 1 && !UE_BUILD_SHIPPING
+
 namespace UE::IoStore
 {
 
@@ -213,6 +216,15 @@ static FAutoConsoleVariableRef CVar_WriteOnlyDecoding(
 	TEXT("Enables the use of 'WriteOnly' flag when decoding to buffers with the 'HardwareTargetBuffer' flag")
 );
 
+#if UE_ALLOW_DISABLE_CANCELLING
+bool GIasCancelationEnabled = true;
+static FAutoConsoleVariableRef CVar_IasCancelationEnabled(
+	TEXT("ias.EnableCancelation"),
+	GIasCancelationEnabled,
+	TEXT("Allows existing IO requests to be canceled")
+);
+#endif //UE_ALLOW_DISABLE_CANCELLING
+
 #if !UE_BUILD_SHIPPING
 static int32 GIasPoisonCache = false;
 static FAutoConsoleVariableRef CVar_IasPoisonCache(
@@ -220,7 +232,7 @@ static FAutoConsoleVariableRef CVar_IasPoisonCache(
 	GIasPoisonCache,
 	TEXT("Fills all data materialized from the cache with 0x4d")
 );
-#endif
+#endif // !UE_BUILD_SHIPPING
 
 
 
@@ -804,6 +816,13 @@ class FOnDemandIoBackend final
 
 		bool Cancel(FIoRequestImpl* Request, IIasCache* TheCache)
 		{
+#if UE_ALLOW_DISABLE_CANCELLING
+			if (GIasCancelationEnabled == false)
+			{
+				return false;
+			}
+#endif //UE_ALLOW_DISABLE_CANCELLING
+
 			FScopeLock _(&Mutex);
 
 			const FBackendData* BackendData = FBackendData::Get(Request);
@@ -1966,4 +1985,6 @@ struct FTests
 
 } // namespace UE::IoStore
 
+#undef UE_ALLOW_DISABLE_CANCELLING
+#undef UE_ENABLE_IAS_TESTING
 #undef UE_IAS_DEBUG_CONSOLE_CMDS
