[NetPhysics] Fix flickering during resimulation.
- A recent change moved the FinishFrame call earlier, but that introduced flickering during resimulation, this CL moves the specified logic that needed to be called earlier to the correct place.

diff --git a/Engine/Source/Runtime/Experimental/Chaos/Private/PBDRigidsSolver.cpp b/Engine/Source/Runtime/Experimental/Chaos/Private/PBDRigidsSolver.cpp
--- a/Engine/Source/Runtime/Experimental/Chaos/Private/PBDRigidsSolver.cpp
+++ b/Engine/Source/Runtime/Experimental/Chaos/Private/PBDRigidsSolver.cpp
@@ -603,12 +603,6 @@ namespace Chaos
 
 			MSolver->UpdateStatCounters();
 
-			if (FRewindData* RewindData = MSolver->GetRewindData())
-			{
-				SCOPE_CYCLE_COUNTER(STAT_RewindFinishFrame);
-				RewindData->FinishFrame();
-			}
-
 			//Editor ticks with 0 dt. We don't want to buffer any dirty data from this since it won't be consumed
 			//TODO: handle this more gracefully
 			if(MDeltaTime > 0)
@@ -631,6 +625,12 @@ namespace Chaos
 				QUICK_SCOPE_CYCLE_COUNTER(ShrinkParticleArrays);
 				MSolver->GetParticles().ShrinkArrays(CVars::ChaosArrayCollectionMaxSlackFraction, CVars::ChaosArrayCollectionMinSlack);
 			}
+
+			if (FRewindData* RewindData = MSolver->GetRewindData())
+			{
+				SCOPE_CYCLE_COUNTER(STAT_RewindFinishFrame);
+				RewindData->FinishFrame();
+			}
 		}
 
 	protected:
@@ -2252,7 +2252,7 @@ namespace Chaos
 			}
 		}
 
-		// Move correction error data from RewindData to FPullPhysicsData for marshalling
+		// Move post-resim error correction data from RewindData to FPullPhysicsData for marshaling
 		if (FRewindData* RewindData = GetRewindData())
 		{
 			RewindData->BufferPhysicsResults(PullData->DirtyRigidErrors);
diff --git a/Engine/Source/Runtime/Experimental/Chaos/Private/RewindData.cpp b/Engine/Source/Runtime/Experimental/Chaos/Private/RewindData.cpp
--- a/Engine/Source/Runtime/Experimental/Chaos/Private/RewindData.cpp
+++ b/Engine/Source/Runtime/Experimental/Chaos/Private/RewindData.cpp
@@ -670,40 +670,6 @@ void FRewindData::DesyncIfNecessary(TDirtyInfo& Info, const FFrameAndPhase Frame
 	}
 }
 
-template<>
-void FRewindData::AccumulateErrorIfNecessary(FGeometryParticleHandle& Obj, const FFrameAndPhase FrameAndPhase)
-{
-	const FDirtyParticleInfo* DirtyInfo = DirtyParticles.Find(&Obj);
-	if (!DirtyInfo)
-	{
-		return;
-	}
-
-	// Get the error offset after a correction
-	const FVec3 ErrorX = DirtyInfo->GetHistory().PreCorrectionXR.X() - Obj.GetX();
-	FQuat ErrorR = DirtyInfo->GetHistory().PreCorrectionXR.R() * Obj.GetR().Inverse();
-	ErrorR.EnforceShortestArcWith(FQuat::Identity);
-	ErrorR.Normalize();
-
-	// Check if error is large enough to hide behind render interpolation
-	if (!ErrorX.IsNearlyZero(0.1) || !ErrorR.IsIdentity(0.02))
-	{
-		// Find or add FDirtyParticleErrorInfo for the particle that has an error
-		FDirtyParticleErrorInfo& ErrorInfo = [&]() ->FDirtyParticleErrorInfo&
-		{
-			if (FDirtyParticleErrorInfo* Found = DirtyParticleErrors.Find(&Obj))
-			{
-				return *Found;
-			}
-
-			return DirtyParticleErrors.Add(&Obj, FDirtyParticleErrorInfo(Obj));
-		}();
-
-		// Cache error for particle 
-		ErrorInfo.AccumulateError(ErrorX, ErrorR);
-	}
-}
-
 void FRewindData::FinishFrame()
 {
 	QUICK_SCOPE_CYCLE_COUNTER(RewindDataFinishFrame);
@@ -728,9 +694,6 @@ void FRewindData::FinishFrame()
 				{
 					if (IsFinalResim())
 					{
-						// Cache the correction offset after a resimulation
-						AccumulateErrorIfNecessary(Handle, FutureFrame);
-
 						//Last resim so mark as in sync
 						Handle.SetSyncState(ESyncState::InSync);
 						Handle.SetEnabledDuringResim(false);
@@ -1061,9 +1024,68 @@ void FRewindData::SpawnProxyIfNeeded(FSingleParticlePhysicsProxy& Proxy)
 	}
 }
 
-// Hand over error data per particle from RewindData to a solver collection instead which gets marshalled to GT
+template<>
+void FRewindData::AccumulateErrorIfNecessary(FGeometryParticleHandle& Obj, const FFrameAndPhase FrameAndPhase)
+{
+	const FDirtyParticleInfo* DirtyInfo = DirtyParticles.Find(&Obj);
+	if (!DirtyInfo)
+	{
+		return;
+	}
+
+	// Get the error offset after a correction
+	const FVec3 ErrorX = DirtyInfo->GetHistory().PreCorrectionXR.X() - Obj.GetX();
+	FQuat ErrorR = DirtyInfo->GetHistory().PreCorrectionXR.R() * Obj.GetR().Inverse();
+	ErrorR.EnforceShortestArcWith(FQuat::Identity);
+	ErrorR.Normalize();
+
+	// Check if error is large enough to hide behind render interpolation
+	if (!ErrorX.IsNearlyZero(0.1) || !ErrorR.IsIdentity(0.02))
+	{
+		// Find or add FDirtyParticleErrorInfo for the particle that has an error
+		FDirtyParticleErrorInfo& ErrorInfo = [&]() ->FDirtyParticleErrorInfo&
+			{
+				if (FDirtyParticleErrorInfo* Found = DirtyParticleErrors.Find(&Obj))
+				{
+					return *Found;
+				}
+
+				return DirtyParticleErrors.Add(&Obj, FDirtyParticleErrorInfo(Obj));
+			}();
+
+		// Cache error for particle 
+		ErrorInfo.AccumulateError(ErrorX, ErrorR);
+	}
+}
+
+// Move post-resim error correction data from RewindData to FPullPhysicsData for marshaling to GT where it can be used in render interpolation
 void FRewindData::BufferPhysicsResults(TMap<const IPhysicsProxyBase*, struct FDirtyRigidParticleReplicationErrorData>& DirtyRigidErrors)
 {
+	if (IsFinalResim())
+	{
+		FFrameAndPhase FutureFrame{ CurFrame + 1, FFrameAndPhase::PrePushData };
+
+		auto ErrorDataHelper = [this, FutureFrame](auto& DirtyObjs)
+			{
+				for (auto& Info : DirtyObjs)
+				{
+					if (Info.bResimAsFollower)
+					{
+						continue;
+					}
+
+					if (Info.GetObjectPtr())
+					{
+						// Cache the correction offset after a resimulation
+						AccumulateErrorIfNecessary(*Info.GetObjectPtr(), FutureFrame);
+					}
+				}
+			};
+
+		ErrorDataHelper(DirtyParticles);
+		ErrorDataHelper(DirtyJoints);
+	}
+
 	DirtyRigidErrors.Reserve(DirtyParticleErrors.Num());
 
 	for (const FDirtyParticleErrorInfo& ErrorInfo : DirtyParticleErrors)
diff --git a/Engine/Source/Runtime/Experimental/Chaos/Public/RewindData.h b/Engine/Source/Runtime/Experimental/Chaos/Public/RewindData.h
--- a/Engine/Source/Runtime/Experimental/Chaos/Public/RewindData.h
+++ b/Engine/Source/Runtime/Experimental/Chaos/Public/RewindData.h
@@ -1765,6 +1765,7 @@ public:
 	/** Rewind to state for a specified frame from rewind data */
 	void RewindStates(const int32 RewindFrame, const bool bResetSolver);
 
+	/** Move post-resim error correction data from RewindData to FPullPhysicsData for marshaling to GT where it can be used in render interpolation */
 	void BufferPhysicsResults(TMap<const IPhysicsProxyBase*, struct FDirtyRigidParticleReplicationErrorData>& DirtyRigidErrors);
 
 	/** Return the rewind data solver */
