Replace `auto`s with proper types.

Most uses of `auto` are banned by our coding standard.

diff --git a/Engine/Source/Runtime/Experimental/Iris/Core/Private/Iris/ReplicationState/ReplicationStateDescriptorBuilder.cpp b/Engine/Source/Runtime/Experimental/Iris/Core/Private/Iris/ReplicationState/ReplicationStateDescriptorBuilder.cpp
--- a/Engine/Source/Runtime/Experimental/Iris/Core/Private/Iris/ReplicationState/ReplicationStateDescriptorBuilder.cpp
+++ b/Engine/Source/Runtime/Experimental/Iris/Core/Private/Iris/ReplicationState/ReplicationStateDescriptorBuilder.cpp
@@ -575,7 +575,7 @@ TRefCountPtr<const FReplicationStateDescriptor> FPropertyReplicationStateDescrip
 	// Check registry if we already have created descriptors for this class
 	if (Parameters.DescriptorRegistry)
 	{	
-		if (auto Result = Parameters.DescriptorRegistry->Find(Property))
+		if (const FReplicationStateDescriptorRegistry::FDescriptors* Result = Parameters.DescriptorRegistry->Find(Property))
 		{
 			// For structs we do only expect a single entry
 			ensure(Result->Num() == 1);
@@ -599,7 +599,7 @@ TRefCountPtr<const FReplicationStateDescriptor> FPropertyReplicationStateDescrip
 	BuildParameters.DescriptorType = FPropertyReplicationStateDescriptorBuilder::EDescriptorType::SingleProperty;
 	BuildParameters.bIsInitState = false;
 	BuildParameters.bAllMembersAreReplicated = true;
-	auto Descriptor = Builder.Build(DescriptorName, Parameters.DescriptorRegistry, BuildParameters);
+	TRefCountPtr<const FReplicationStateDescriptor> Descriptor = Builder.Build(DescriptorName, Parameters.DescriptorRegistry, BuildParameters);
 
 	if (Parameters.DescriptorRegistry && Descriptor.IsValid())
 	{
@@ -2491,7 +2491,7 @@ TRefCountPtr<const FReplicationStateDescriptor> FReplicationStateDescriptorBuild
 	// Check registry if we already have created descriptors for this class
 	if (Parameters.DescriptorRegistry)
 	{	
-		if (auto Result = Parameters.DescriptorRegistry->Find(InStruct))
+		if (const FReplicationStateDescriptorRegistry::FDescriptors* Result = Parameters.DescriptorRegistry->Find(InStruct))
 		{
 			// For structs we do only expect a single entry
 			ensure(Result->Num() == 1);
@@ -2662,7 +2662,7 @@ TRefCountPtr<const FReplicationStateDescriptor> FReplicationStateDescriptorBuild
 	// Check registry if we already have created descriptors for this function
 	if (Parameters.DescriptorRegistry)
 	{	
-		if (auto Result = Parameters.DescriptorRegistry->Find(Function))
+		if (const FReplicationStateDescriptorRegistry::FDescriptors* Result = Parameters.DescriptorRegistry->Find(Function))
 		{
 			// For function we expect a single entry
 			ensure(Result->Num() == 1);
@@ -2701,7 +2701,7 @@ TRefCountPtr<const FReplicationStateDescriptor> FReplicationStateDescriptorBuild
 	BuildParameters.bIsInitState = false;
 	// For this use case we don't care about function parameters that aren't replicated.
 	BuildParameters.bAllMembersAreReplicated = true;
-	auto Descriptor = Builder.Build(Function->GetName(), Parameters.DescriptorRegistry, BuildParameters);
+	TRefCountPtr<const FReplicationStateDescriptor> Descriptor = Builder.Build(Function->GetName(), Parameters.DescriptorRegistry, BuildParameters);
 
 	if (Parameters.DescriptorRegistry && Descriptor.IsValid())
 	{
