[AutoRTFM] Make FReplicationStateDescriptor refcounting transactionally safe.

The refcount is now incremented immediately, but only decrements when a transaction completes. This is a pretty typical pattern for transactionally-safe refcounting.

diff --git a/Engine/Source/Runtime/Experimental/Iris/Core/Private/Iris/ReplicationState/ReplicationStateDescriptor.cpp b/Engine/Source/Runtime/Experimental/Iris/Core/Private/Iris/ReplicationState/ReplicationStateDescriptor.cpp
--- a/Engine/Source/Runtime/Experimental/Iris/Core/Private/Iris/ReplicationState/ReplicationStateDescriptor.cpp
+++ b/Engine/Source/Runtime/Experimental/Iris/Core/Private/Iris/ReplicationState/ReplicationStateDescriptor.cpp
@@ -4,6 +4,7 @@
 #include "Iris/ReplicationSystem/ReplicationOperationsInternal.h"
 #include "Iris/Serialization/InternalNetSerializationContext.h"
 #include "Iris/Serialization/NetSerializer.h"
+#include "AutoRTFM/AutoRTFM.h"
 #include "HAL/PlatformAtomics.h"
 #include "UObject/UnrealType.h"
 #include "UObject/CoreNetTypes.h"
@@ -51,7 +52,14 @@ void FReplicationStateDescriptor::AddRef() const
 {
 	if (EnumHasAnyFlags(Traits, EReplicationStateTraits::NeedsRefCount))
 	{
-		++RefCount;
+		UE_AUTORTFM_OPEN
+		{
+			++RefCount;
+		};
+		UE_AUTORTFM_ONABORT(this)
+		{
+			--RefCount;
+		};
 	}
 }
 
@@ -61,41 +69,44 @@ void FReplicationStateDescriptor::Release() const
 
 	if (EnumHasAnyFlags(Traits, EReplicationStateTraits::NeedsRefCount))
 	{
-		if (--RefCount == 0)
+		UE_AUTORTFM_ONCOMMIT(this)
 		{
-			// We must destruct default state if we have one
-			if (uint8* StateBufferToFree = const_cast<uint8*>(DefaultStateBuffer))
+			if (--RefCount == 0)
 			{
-				if (EnumHasAnyFlags(Traits, EReplicationStateTraits::HasDynamicState))
+				// We must destruct default state if we have one
+				if (uint8* StateBufferToFree = const_cast<uint8*>(DefaultStateBuffer))
 				{
-					FNetSerializationContext FreeContext;
-					FInternalNetSerializationContext InternalContext;
-					FreeContext.SetInternalContext(&InternalContext);
+					if (EnumHasAnyFlags(Traits, EReplicationStateTraits::HasDynamicState))
+					{
+						FNetSerializationContext FreeContext;
+						FInternalNetSerializationContext InternalContext;
+						FreeContext.SetInternalContext(&InternalContext);
+
+						FReplicationStateOperationsInternal::FreeDynamicState(FreeContext, StateBufferToFree, this);
+					}
+					FMemory::Free(StateBufferToFree);
+				}
 
-					FReplicationStateOperationsInternal::FreeDynamicState(FreeContext, StateBufferToFree, this);
+				// We must destruct configs pointing to generated structs
+				for (uint32 It = 0; It < MemberCount; ++It)
+				{
+					const FNetSerializerConfig* Config = MemberSerializerDescriptors[It].SerializerConfig;
+					if (EnumHasAnyFlags(Config->ConfigTraits, ENetSerializerConfigTraits::NeedDestruction))
+					{
+						Config->~FNetSerializerConfig();
+					}
 				}
-				FMemory::Free(StateBufferToFree);
-			}
 
-			// We must destruct configs pointing to generated structs
-			for (uint32 It = 0; It < MemberCount; ++It)
-			{
-				const FNetSerializerConfig* Config = MemberSerializerDescriptors[It].SerializerConfig;
-				if (EnumHasAnyFlags(Config->ConfigTraits, ENetSerializerConfigTraits::NeedDestruction))
+				// We must destruct member function descriptors
+				for (const FReplicationStateMemberFunctionDescriptor& FunctionDescriptor : MakeArrayView(MemberFunctionDescriptors, FunctionCount))
 				{
-					Config->~FNetSerializerConfig();
+					FunctionDescriptor.Descriptor->Release();
 				}
-			}
 
-			// We must destruct member function descriptors
-			for (const FReplicationStateMemberFunctionDescriptor& FunctionDescriptor : MakeArrayView(MemberFunctionDescriptors, FunctionCount))
-			{
-				FunctionDescriptor.Descriptor->Release();
+				// Allocated via FMemory::Malloc in ReplicationStateDescriptorBuilder
+				FMemory::Free(const_cast<FReplicationStateDescriptor*>(this));
 			}
-
-			// Allocated via FMemory::Malloc in ReplicationStateDescriptorBuilder
-			FMemory::Free(const_cast<FReplicationStateDescriptor*>(this));
-		}
+		};
 	}
 }
 
