Added cloth-cloth constraint springs with a basic constraint creation UI.

diff --git a/Engine/Source/Runtime/Experimental/Chaos/Private/Chaos/GraphColoring.cpp b/Engine/Source/Runtime/Experimental/Chaos/Private/Chaos/GraphColoring.cpp
--- a/Engine/Source/Runtime/Experimental/Chaos/Private/Chaos/GraphColoring.cpp
+++ b/Engine/Source/Runtime/Experimental/Chaos/Private/Chaos/GraphColoring.cpp
@@ -6,129 +6,26 @@
 #include "Containers/BitArray.h"
 #include "Chaos/SoftsSolverParticlesRange.h"
 
-template<typename DynamicParticlesType>
-static bool VerifyGraph(TArray<TArray<int32>> ColorGraph, const TArray<Chaos::TVec2<int32>>& Graph, const DynamicParticlesType& InParticles)
+template<typename DynamicParticlesType, int32 N, bool bAllDynamic>
+static bool VerifyGraph(const TArray<TArray<int32>>& ColorGraph, const TArray<Chaos::TVector<int32, N>>& Graph, const DynamicParticlesType& InParticles)
 {
-	for (int32 i = 0; i < ColorGraph.Num(); ++i)
+	for (int32 Color = 0; Color < ColorGraph.Num(); ++Color)
 	{
-		TMap<int32, int32> NodeToColorMap;
-		for (const auto& Edge : ColorGraph[i])
+		TMap<int32, int32> ColorNodesToEdges;
+		for (const int32 Edge : ColorGraph[Color])
 		{
-			int32 Node1 = Graph[Edge][0];
-			int32 Node2 = Graph[Edge][1];
-			if (NodeToColorMap.Contains(Node1))
+			for (int32 NIndex = 0; NIndex < N; ++NIndex)
 			{
-				UE_LOG(LogChaos, Error, TEXT("Color %d has duplicate Node %d"), i, Node1);
-				return false;
-			}
-			if (NodeToColorMap.Contains(Node2))
-			{
-				UE_LOG(LogChaos, Error, TEXT("Color %d has duplicate Node %d"), i, Node2);
-				return false;
-			}
-			if (InParticles.InvM(Node1) != 0)
-			{
-				NodeToColorMap.Add(Node1, i);
-			}
-			if (InParticles.InvM(Node2) != 0)
-			{
-				NodeToColorMap.Add(Node2, i);
-			}
-		}
-	}
-	return true;
-}
-
-template<typename DynamicParticlesType>
-static bool VerifyGraph(TArray<TArray<int32>> ColorGraph, const TArray<Chaos::TVec3<int32>>& Graph, const DynamicParticlesType& InParticles)
-{
-	for (int32 i = 0; i < ColorGraph.Num(); ++i)
-	{
-		TMap<int32, int32> NodeToColorMap;
-		for (const auto& Edge : ColorGraph[i])
-		{
-			int32 Node1 = Graph[Edge][0];
-			int32 Node2 = Graph[Edge][1];
-			int32 Node3 = Graph[Edge][2];
-			if (NodeToColorMap.Contains(Node1))
-			{
-				UE_LOG(LogChaos, Error, TEXT("Color %d has duplicate Node %d"), i, Node1);
-				return false;
-			}
-			if (NodeToColorMap.Contains(Node2))
-			{
-				UE_LOG(LogChaos, Error, TEXT("Color %d has duplicate Node %d"), i, Node2);
-				return false;
-			}
-			if (NodeToColorMap.Contains(Node3))
-			{
-				UE_LOG(LogChaos, Error, TEXT("Color %d has duplicate Node %d"), i, Node3);
-				return false;
-			}
-			if (InParticles.InvM(Node1) != 0)
-			{
-				NodeToColorMap.Add(Node1, i);
-			}
-			if (InParticles.InvM(Node2) != 0)
-			{
-				NodeToColorMap.Add(Node2, i);
-			}
-			if (InParticles.InvM(Node3) != 0)
-			{
-				NodeToColorMap.Add(Node3, i);
-			}
-		}
-	}
-	return true;
-}
-
-template<typename DynamicParticlesType>
-static bool VerifyGraph(TArray<TArray<int32>> ColorGraph, const TArray<Chaos::TVec4<int32>>& Graph, const DynamicParticlesType& InParticles)
-{	
-	for (int32 i = 0; i < ColorGraph.Num(); ++i)
-	{
-		TMap<int32, int32> NodeToColorMap;
-		for (const auto& Edge : ColorGraph[i])
-		{
-			int32 Node1 = Graph[Edge][0];
-			int32 Node2 = Graph[Edge][1];
-			int32 Node3 = Graph[Edge][2];
-			int32 Node4 = Graph[Edge][3];
-			if (NodeToColorMap.Contains(Node1))
-			{
-				UE_LOG(LogChaos, Error, TEXT("Color %d has duplicate Node %d"), i, Node1);
-				return false;
-			}
-			if (NodeToColorMap.Contains(Node2))
-			{
-				UE_LOG(LogChaos, Error, TEXT("Color %d has duplicate Node %d"), i, Node2);
-				return false;
-			}
-			if (NodeToColorMap.Contains(Node3))
-			{
-				UE_LOG(LogChaos, Error, TEXT("Color %d has duplicate Node %d"), i, Node3);
-				return false;
-			}
-			if (NodeToColorMap.Contains(Node4))
-			{
-				UE_LOG(LogChaos, Error, TEXT("Color %d has duplicate Node %d"), i, Node4);
-				return false;
-			}
-			if (InParticles.InvM(Node1) != 0)
-			{
-				NodeToColorMap.Add(Node1, i);
-			}
-			if (InParticles.InvM(Node2) != 0)
-			{
-				NodeToColorMap.Add(Node2, i);
-			}
-			if (InParticles.InvM(Node3) != 0)
-			{
-				NodeToColorMap.Add(Node3, i);
-			}
-			if (InParticles.InvM(Node4) != 0)
-			{
-				NodeToColorMap.Add(Node4, i);
+				const int32 Node = Graph[Edge][NIndex];
+				if (const int32* const ExistingEdge = ColorNodesToEdges.Find(Node))
+				{
+					UE_LOG(LogChaos, Error, TEXT("Color %d has duplicate Node %d. First added for Edge %d, and now found for Edge %d"), Color, Node, *ExistingEdge, Edge);
+					return false;
+				}
+				if (bAllDynamic || InParticles.InvM(Node) != 0)
+				{
+					ColorNodesToEdges.Add(Node, Edge);
+				}
 			}
 		}
 	}
@@ -136,49 +33,6 @@ static bool VerifyGraph(TArray<TArray<int32>> ColorGraph, const TArray<Chaos::TV
 }
 
 
-template<typename DynamicParticlesType>
-static bool VerifyGraphAllDynamic(TArray<TArray<int32>> ColorGraph, const TArray<Chaos::TVec4<int32>>& Graph, const DynamicParticlesType& InParticles)
-{
-	for (int32 i = 0; i < ColorGraph.Num(); ++i)
-	{
-		TMap<int32, int32> NodeToColorMap;
-		for (const auto& Edge : ColorGraph[i])
-		{
-			int32 Node1 = Graph[Edge][0];
-			int32 Node2 = Graph[Edge][1];
-			int32 Node3 = Graph[Edge][2];
-			int32 Node4 = Graph[Edge][3];
-			if (NodeToColorMap.Contains(Node1))
-			{
-				UE_LOG(LogChaos, Error, TEXT("Color %d has duplicate Node %d"), i, Node1);
-				return false;
-			}
-			if (NodeToColorMap.Contains(Node2))
-			{
-				UE_LOG(LogChaos, Error, TEXT("Color %d has duplicate Node %d"), i, Node2);
-				return false;
-			}
-			if (NodeToColorMap.Contains(Node3))
-			{
-				UE_LOG(LogChaos, Error, TEXT("Color %d has duplicate Node %d"), i, Node3);
-				return false;
-			}
-			if (NodeToColorMap.Contains(Node4))
-			{
-				UE_LOG(LogChaos, Error, TEXT("Color %d has duplicate Node %d"), i, Node4);
-				return false;
-			}
-
-			NodeToColorMap.Add(Node1, i);
-			NodeToColorMap.Add(Node2, i);
-			NodeToColorMap.Add(Node3, i);
-			NodeToColorMap.Add(Node4, i);
-			
-		}
-	}
-	return true;
-}
-
 //just verify if different element in each subcolor has intersecting nodes
 template <typename T>
 static bool VerifyGridBasedSubColoring(const TArray<TArray<int32>>& ElementsPerColor, const Chaos::TMPMGrid<T>& Grid, const TArray<TArray<int32>>& ConstraintsNodesSet, const TArray<TArray<TArray<int32>>>& ElementsPerSubColors)
@@ -539,917 +393,164 @@ static bool VerifyExtraNodalColoring(const Chaos::TDynamicParticles<T, 3>& InPar
 	return true;
 }
 
-
-template<typename DynamicParticlesType>
-TArray<TArray<int32>> Chaos::FGraphColoring::ComputeGraphColoringParticlesOrRange(const TArray<Chaos::TVec2<int32>>& Graph, const DynamicParticlesType& InParticles, const int32 GraphParticlesStart, const int32 GraphParticlesEnd)
+template<typename DynamicParticlesType, int32 N, bool bAllDynamic>
+TArray<TArray<int32>> Chaos::FGraphColoring::ComputeGraphColoringParticlesOrRange(const TArray<TVector<int32, N>>& Graph, const DynamicParticlesType& InParticles, const int32 GraphParticlesStart, const int32 GraphParticlesEnd)
 {
+	TRACE_CPUPROFILER_EVENT_SCOPE(FGraphColoring_ComputeGraphColoringN);
 	checkSlow(GraphParticlesStart <= GraphParticlesEnd);
 	checkSlow(GraphParticlesEnd <= (int32)InParticles.Size());
-	TArray<TArray<int32>> ColorGraph;
-	TArray<FGraphNode> NodesSubArray;
-	TArray<FGraphEdge> Edges;
-	NodesSubArray.SetNum(GraphParticlesEnd - GraphParticlesStart);
-	Edges.SetNum(Graph.Num());
-	TArrayView<FGraphNode> Nodes(NodesSubArray.GetData() - GraphParticlesStart, GraphParticlesEnd); // Only nodes starting with GraphParticlesStart are valid to access
-
-	int32 MaxColor = -1;
-
-	for (int32 i = 0; i < Graph.Num(); ++i)
-	{
-		const TVec2<int32>& Constraint = Graph[i];
-		Edges[i].FirstNode = Constraint[0];
-		Edges[i].SecondNode = Constraint[1];
-		checkSlow(Constraint[0] >= GraphParticlesStart);
-		checkSlow(Constraint[1] >= GraphParticlesStart);
-		Nodes[Constraint[0]].Edges.Add(i);
-		Nodes[Constraint[1]].Edges.Add(i);
-	}
-
-	TSet<int32> ProcessedNodes;
-	TArray<int32> NodesToProcess;
-
-	for (int32 ParticleNodeIndex = GraphParticlesStart; ParticleNodeIndex < GraphParticlesEnd; ++ParticleNodeIndex)
-	{
-		const bool bIsParticleDynamic = InParticles.InvM(ParticleNodeIndex) != (decltype(InParticles.InvM(ParticleNodeIndex)))0.;
-		if (ProcessedNodes.Contains(ParticleNodeIndex) || !bIsParticleDynamic)
-		{
-			continue;
-		}
-
-		NodesToProcess.Add(ParticleNodeIndex);
-
-		while (NodesToProcess.Num())
-		{
-			const int32 NodeIndex = NodesToProcess.Last();
-			FGraphNode& GraphNode = Nodes[NodeIndex];
-
-			NodesToProcess.SetNum(NodesToProcess.Num() - 1, EAllowShrinking::No);
-			ProcessedNodes.Add(NodeIndex);
-
-			for (const int32 EdgeIndex : GraphNode.Edges)
-			{
-				FGraphEdge& GraphEdge = Edges[EdgeIndex];
-
-				// If edge has been colored skip it
-				if (GraphEdge.Color >= 0)
-				{
-					continue;
-				}
-
-				// Get index to the other node on the edge
-				int32 OtherNodeIndex = INDEX_NONE;
-				if (GraphEdge.FirstNode == NodeIndex)
-				{
-					OtherNodeIndex = GraphEdge.SecondNode;
-				}
-				if (GraphEdge.SecondNode == NodeIndex)
-				{
-					OtherNodeIndex = GraphEdge.FirstNode;
-				}
-
-				// Find next color that is not used already at this node
-				while (GraphNode.UsedColors.Contains(GraphNode.NextColor))
-				{
-					GraphNode.NextColor++;
-				}
-				int32 ColorToUse = GraphNode.NextColor;
-
-				// Exclude colors used by the other node (but still allow this node to use them for other edges)
-				if (OtherNodeIndex != INDEX_NONE)
-				{
-					FGraphNode& OtherNode = Nodes[OtherNodeIndex];
-
-					const bool bIsOtherGraphNodeDynamic = InParticles.InvM(OtherNodeIndex) != (decltype(InParticles.InvM(ParticleNodeIndex)))0.;
-					if (bIsOtherGraphNodeDynamic)
-					{
-						while (OtherNode.UsedColors.Contains(ColorToUse) || GraphNode.UsedColors.Contains(ColorToUse))
-						{
-							ColorToUse++;
-						}
-					}
-				}
-
-				// Assign color and set as used at this node
-				MaxColor = FMath::Max(ColorToUse, MaxColor);
-				GraphNode.UsedColors.Add(ColorToUse);
-				GraphEdge.Color = ColorToUse;
-
-				// Bump color to use next time, but only if we weren't forced to use a different color by the other node
-				if ((ColorToUse == GraphNode.NextColor) && bIsParticleDynamic == true)
-				{
-					GraphNode.NextColor++;
-				}
-
-				if (ColorGraph.Num() <= MaxColor)
-				{
-					ColorGraph.SetNum(MaxColor + 1);
-				}
-				ColorGraph[GraphEdge.Color].Add(EdgeIndex);
-
-				if (OtherNodeIndex != INDEX_NONE)
-				{
-					FGraphNode& OtherGraphNode = Nodes[OtherNodeIndex];
-					const bool bIsOtherGraphNodeDynamic = InParticles.InvM(OtherNodeIndex) != (decltype(InParticles.InvM(ParticleNodeIndex)))0.;
-					if (bIsOtherGraphNodeDynamic)
-					{
-						// Mark other node as not allowing use of this color
-						if (bIsParticleDynamic)
-						{
-							OtherGraphNode.UsedColors.Add(GraphEdge.Color);
-						}
-
-						// Queue other node for processing
-						if (!ProcessedNodes.Contains(OtherNodeIndex))
-						{
-							NodesToProcess.Add(OtherNodeIndex);
-						}
-					}
-				}
-			}
-		}
-	}
-
-	checkSlow(VerifyGraph(ColorGraph, Graph, InParticles));
-	return ColorGraph;
-}
-
-template<typename DynamicParticlesType>
-TArray<TArray<int32>> Chaos::FGraphColoring::ComputeGraphColoringParticlesOrRange(const TArray<TVec3<int32>>& Graph, const DynamicParticlesType& InParticles, const int32 GraphParticlesStart, const int32 GraphParticlesEnd)
-{
-	using namespace Chaos;
 
-	checkSlow(GraphParticlesStart <= GraphParticlesEnd);
-	checkSlow(GraphParticlesEnd <= (int32)InParticles.Size());
 	TArray<TArray<int32>> ColorGraph;
-	TArray<FGraphNode> NodesSubArray;
-	TArray<FGraph3dEdge> Edges;
-	NodesSubArray.SetNum(GraphParticlesEnd - GraphParticlesStart);
-	Edges.SetNum(Graph.Num());
-	TArrayView<FGraphNode> Nodes(NodesSubArray.GetData() - GraphParticlesStart, GraphParticlesEnd); // Only nodes starting with GraphParticlesStart are valid to access
 
 	int32 MaxColor = -1;
-
-	for (int32 i = 0; i < Graph.Num(); ++i)
+	TArray<FColorSet> NodeUsedColorsSubArray;
+	NodeUsedColorsSubArray.SetNum(GraphParticlesEnd - GraphParticlesStart);
+	TArrayView<FColorSet> NodeUsedColors(NodeUsedColorsSubArray.GetData() - GraphParticlesStart, GraphParticlesEnd); // Only nodes starting with GraphParticlesStart are valid to access
+	for (int32 EdgeIndex = 0; EdgeIndex < Graph.Num(); ++EdgeIndex)
 	{
-		const TVec3<int32>& Constraint = Graph[i];
-		Edges[i].FirstNode = Constraint[0];
-		Edges[i].SecondNode = Constraint[1];
-		Edges[i].ThirdNode = Constraint[2];
-		checkSlow(Constraint[0] >= GraphParticlesStart);
-		checkSlow(Constraint[1] >= GraphParticlesStart);
-		checkSlow(Constraint[2] >= GraphParticlesStart);
-		Nodes[Constraint[0]].Edges.Add(i);
-		Nodes[Constraint[1]].Edges.Add(i);
-		Nodes[Constraint[2]].Edges.Add(i);
-	}
-
-	TSet<int32> ProcessedNodes;
-	TArray<int32> NodesToProcess;
-
-	for (int32 ParticleNodeIndex = GraphParticlesStart; ParticleNodeIndex < GraphParticlesEnd; ++ParticleNodeIndex)
-	{
-		const bool bIsParticleDynamic = InParticles.InvM(ParticleNodeIndex) != (decltype(InParticles.InvM(ParticleNodeIndex)))0.;
-		if (ProcessedNodes.Contains(ParticleNodeIndex) || !bIsParticleDynamic)
+		// Find color that hasn't already been assigned to a node on this edge.
+		int32 FirstFreeColor = 0;
+		while (true)
 		{
-			continue;
-		}
-
-		NodesToProcess.Add(ParticleNodeIndex);
-
-		while (NodesToProcess.Num())
-		{
-			const int32 NodeIndex = NodesToProcess.Last();
-			FGraphNode& GraphNode = Nodes[NodeIndex];
-
-			NodesToProcess.SetNum(NodesToProcess.Num() - 1, EAllowShrinking::No);
-			ProcessedNodes.Add(NodeIndex);
-
-			for (const int32 EdgeIndex : GraphNode.Edges)
+			bool bColorFound = false;
+			for (int32 NIndex = 0; NIndex < N; ++NIndex)
 			{
-				FGraph3dEdge& GraphEdge = Edges[EdgeIndex];
-
-				// If edge has been colored skip it
-				if (GraphEdge.Color >= 0)
-				{
-					continue;
-				}
-
-				// Get index to the other node on the edge
-				int32 OtherNodeIndex = INDEX_NONE;
-				int32 OtherNodeIndex2 = INDEX_NONE;
-				if (GraphEdge.FirstNode == NodeIndex)
-				{
-					OtherNodeIndex = GraphEdge.SecondNode;
-					OtherNodeIndex2 = GraphEdge.ThirdNode;
-				}
-				if (GraphEdge.SecondNode == NodeIndex)
-				{
-					OtherNodeIndex = GraphEdge.FirstNode;
-					OtherNodeIndex2 = GraphEdge.ThirdNode;
-				}
-				if (GraphEdge.ThirdNode == NodeIndex)
-				{
-					OtherNodeIndex = GraphEdge.FirstNode;
-					OtherNodeIndex2 = GraphEdge.SecondNode;
-				}
-
-				// Find next color that is not used already at this node
-				while (GraphNode.UsedColors.Contains(GraphNode.NextColor))
-				{
-					GraphNode.NextColor++;
-				}
-				int32 ColorToUse = GraphNode.NextColor;
-
-				// Exclude colors used by the other node (but still allow this node to use them for other edges)
-				if (OtherNodeIndex != INDEX_NONE)
+				const int32 NodeIndex = Graph[EdgeIndex][NIndex];
+				if constexpr (!bAllDynamic)
 				{
-					FGraphNode& OtherNode = Nodes[OtherNodeIndex];
-
-					const bool bIsOtherGraphNodeDynamic = InParticles.InvM(OtherNodeIndex) != (decltype(InParticles.InvM(ParticleNodeIndex)))0.;
-					if (bIsOtherGraphNodeDynamic)
-					{
-						while (OtherNode.UsedColors.Contains(ColorToUse) || GraphNode.UsedColors.Contains(ColorToUse))
-						{
-							ColorToUse++;
-						}
-					}
-				}
-				if (OtherNodeIndex2 != INDEX_NONE)
-				{
-					FGraphNode& OtherNode = Nodes[OtherNodeIndex2];
-
-					const bool bIsOtherGraphNodeDynamic = InParticles.InvM(OtherNodeIndex2) != (decltype(InParticles.InvM(ParticleNodeIndex)))0.;
-					if (bIsOtherGraphNodeDynamic)
+					const bool bIsParticleDynamic = InParticles.InvM(NodeIndex) != (decltype(InParticles.InvM(NodeIndex)))0.;
+					if (!bIsParticleDynamic)
 					{
-						if (OtherNodeIndex == INDEX_NONE)
-						{
-							while (OtherNode.UsedColors.Contains(ColorToUse) || GraphNode.UsedColors.Contains(ColorToUse))
-							{
-								ColorToUse++;
-							}
-						}
-						else
-						{
-							FGraphNode& PrevOtherNode = Nodes[OtherNodeIndex];
-							while (OtherNode.UsedColors.Contains(ColorToUse) || PrevOtherNode.UsedColors.Contains(ColorToUse) || GraphNode.UsedColors.Contains(ColorToUse))
-							{
-								ColorToUse++;
-							}
-						}
-					}
-				}
-
-				// Assign color and set as used at this node
-				MaxColor = FMath::Max(ColorToUse, MaxColor);
-				GraphNode.UsedColors.Add(ColorToUse);
-				GraphEdge.Color = ColorToUse;
-
-				// Bump color to use next time, but only if we weren't forced to use a different color by the other node
-				if ((ColorToUse == GraphNode.NextColor) && bIsParticleDynamic == true)
-				{
-					GraphNode.NextColor++;
-				}
-
-				if (ColorGraph.Num() <= MaxColor)
-				{
-					ColorGraph.SetNum(MaxColor + 1);
-				}
-				ColorGraph[GraphEdge.Color].Add(EdgeIndex);
-
-				if (OtherNodeIndex != INDEX_NONE)
-				{
-					FGraphNode& OtherGraphNode = Nodes[OtherNodeIndex];
-					const bool bIsOtherGraphNodeDynamic = InParticles.InvM(OtherNodeIndex) != (decltype(InParticles.InvM(ParticleNodeIndex)))0.;
-					if (bIsOtherGraphNodeDynamic)
-					{
-						// Mark other node as not allowing use of this color
-						if (bIsParticleDynamic)
-						{
-							OtherGraphNode.UsedColors.Add(GraphEdge.Color);
-						}
-
-						// Queue other node for processing
-						if (!ProcessedNodes.Contains(OtherNodeIndex))
-						{
-							NodesToProcess.Add(OtherNodeIndex);
-						}
+						continue;
 					}
 				}
-				if (OtherNodeIndex2 != INDEX_NONE)
+				if (NodeUsedColors[NodeIndex].Contains(FirstFreeColor))
 				{
-					FGraphNode& OtherGraphNode = Nodes[OtherNodeIndex2];
-					const bool bIsOtherGraphNodeDynamic = InParticles.InvM(OtherNodeIndex2) != (decltype(InParticles.InvM(ParticleNodeIndex)))0.;
-					if (bIsOtherGraphNodeDynamic)
-					{
-						// Mark other node as not allowing use of this color
-						if (bIsParticleDynamic)
-						{
-							OtherGraphNode.UsedColors.Add(GraphEdge.Color);
-						}
-
-						// Queue other node for processing
-						if (!ProcessedNodes.Contains(OtherNodeIndex2))
-						{
-							NodesToProcess.Add(OtherNodeIndex2);
-						}
-					}
+					bColorFound = true;
+					break;
 				}
 			}
-		}
-	}
-
-	checkSlow(VerifyGraph(ColorGraph, Graph, InParticles));
-	return ColorGraph;
-}
-
-template<typename DynamicParticlesType>
-TArray<TArray<int32>> Chaos::FGraphColoring::ComputeGraphColoringParticlesOrRange(const TArray<TVec4<int32>>& Graph, const DynamicParticlesType& InParticles, const int32 GraphParticlesStart, const int32 GraphParticlesEnd)
-{
-	using namespace Chaos;
-
-	checkSlow(GraphParticlesStart <= GraphParticlesEnd);
-	checkSlow(GraphParticlesEnd <= (int32)InParticles.Size());
-	TArray<TArray<int32>> ColorGraph;
-	TArray<FGraphNode> NodesSubArray;
-	TArray<FGraphTetEdge> Edges;
-	NodesSubArray.SetNum(GraphParticlesEnd - GraphParticlesStart);
-	Edges.SetNum(Graph.Num());
-	TArrayView<FGraphNode> Nodes(NodesSubArray.GetData() - GraphParticlesStart, GraphParticlesEnd); // Only nodes starting with GraphParticlesStart are valid to access
-
-	int32 MaxColor = -1;
-
-	for (int32 i = 0; i < Graph.Num(); ++i)
-	{
-		const TVec4<int32>& Constraint = Graph[i];
-		Edges[i].FirstNode = Constraint[0];
-		Edges[i].SecondNode = Constraint[1];
-		Edges[i].ThirdNode = Constraint[2];
-		Edges[i].FourthNode = Constraint[3];
-		checkSlow(Constraint[0] >= GraphParticlesStart);
-		checkSlow(Constraint[1] >= GraphParticlesStart);
-		checkSlow(Constraint[2] >= GraphParticlesStart);
-		checkSlow(Constraint[3] >= GraphParticlesStart);
-		Nodes[Constraint[0]].Edges.Add(i);
-		Nodes[Constraint[1]].Edges.Add(i);
-		Nodes[Constraint[2]].Edges.Add(i);
-		Nodes[Constraint[3]].Edges.Add(i);
-	}
-
-	TSet<int32> ProcessedNodes;
-	TArray<int32> NodesToProcess;
-
-	for (int32 ParticleNodeIndex = GraphParticlesStart; ParticleNodeIndex < GraphParticlesEnd; ++ParticleNodeIndex)
-	{
-		const bool bIsParticleDynamic = InParticles.InvM(ParticleNodeIndex) != (decltype(InParticles.InvM(ParticleNodeIndex)))0.;
-		if (ProcessedNodes.Contains(ParticleNodeIndex) || !bIsParticleDynamic)
-		{
-			continue;
-		}
-
-		NodesToProcess.Add(ParticleNodeIndex);
-
-		while (NodesToProcess.Num())
-		{
-			const int32 NodeIndex = NodesToProcess.Last();
-			FGraphNode& GraphNode = Nodes[NodeIndex];
-
-			NodesToProcess.SetNum(NodesToProcess.Num() - 1, EAllowShrinking::No);
-			ProcessedNodes.Add(NodeIndex);
-
-			for (const int32 EdgeIndex : GraphNode.Edges)
+			if (!bColorFound)
 			{
-				FGraphTetEdge& GraphEdge = Edges[EdgeIndex];
-
-				// If edge has been colored skip it
-				if (GraphEdge.Color >= 0)
-				{
-					continue;
-				}
-
-				// Get index to the other node on the edge
-				int32 OtherNodeIndex = INDEX_NONE;
-				int32 OtherNodeIndex2 = INDEX_NONE;
-				int32 OtherNodeIndex3 = INDEX_NONE;
-				if (GraphEdge.FirstNode == NodeIndex)
-				{
-					OtherNodeIndex = GraphEdge.SecondNode;
-					OtherNodeIndex2 = GraphEdge.ThirdNode;
-					OtherNodeIndex3 = GraphEdge.FourthNode;
-				}
-				if (GraphEdge.SecondNode == NodeIndex)
-				{
-					OtherNodeIndex = GraphEdge.FirstNode;
-					OtherNodeIndex2 = GraphEdge.ThirdNode;
-					OtherNodeIndex3 = GraphEdge.FourthNode;
-				}
-				if (GraphEdge.ThirdNode == NodeIndex)
-				{
-					OtherNodeIndex = GraphEdge.FirstNode;
-					OtherNodeIndex2 = GraphEdge.SecondNode;
-					OtherNodeIndex3 = GraphEdge.FourthNode;
-				}
-				if (GraphEdge.FourthNode == NodeIndex)
-				{
-					OtherNodeIndex = GraphEdge.FirstNode;
-					OtherNodeIndex2 = GraphEdge.SecondNode;
-					OtherNodeIndex3 = GraphEdge.ThirdNode;
-				}
-
-				// Find next color that is not used already at this node
-				while (GraphNode.UsedColors.Contains(GraphNode.NextColor))
-				{
-					GraphNode.NextColor++;
-				}
-				int32 ColorToUse = GraphNode.NextColor;
-
-				// Exclude colors used by the other node (but still allow this node to use them for other edges)
-				if (OtherNodeIndex != INDEX_NONE)
-				{
-					FGraphNode& OtherNode = Nodes[OtherNodeIndex];
-
-					const bool bIsOtherGraphNodeDynamic = InParticles.InvM(OtherNodeIndex) != (decltype(InParticles.InvM(ParticleNodeIndex)))0.;
-					if (bIsOtherGraphNodeDynamic)
-					{
-						while (OtherNode.UsedColors.Contains(ColorToUse) || GraphNode.UsedColors.Contains(ColorToUse))
-						{
-							ColorToUse++;
-						}
-					}
-				}
-				if (OtherNodeIndex2 != INDEX_NONE)
-				{
-					FGraphNode& OtherNode = Nodes[OtherNodeIndex2];
-
-					const bool bIsOtherGraphNodeDynamic = InParticles.InvM(OtherNodeIndex2) != (decltype(InParticles.InvM(ParticleNodeIndex)))0.;
-					if (bIsOtherGraphNodeDynamic)
-					{
-						if (OtherNodeIndex == INDEX_NONE)
-						{
-							while (OtherNode.UsedColors.Contains(ColorToUse) || GraphNode.UsedColors.Contains(ColorToUse))
-							{
-								ColorToUse++;
-							}
-						}
-						else
-						{
-							FGraphNode& PrevOtherNode = Nodes[OtherNodeIndex];
-							while (OtherNode.UsedColors.Contains(ColorToUse) || PrevOtherNode.UsedColors.Contains(ColorToUse) || GraphNode.UsedColors.Contains(ColorToUse))
-							{
-								ColorToUse++;
-							}
-						}
-					}
-				}
-				if (OtherNodeIndex3 != INDEX_NONE)
-				{
-					FGraphNode& OtherNode = Nodes[OtherNodeIndex3];
-
-					const bool bIsOtherGraphNodeDynamic = InParticles.InvM(OtherNodeIndex3) != (decltype(InParticles.InvM(ParticleNodeIndex)))0.;
-					if (bIsOtherGraphNodeDynamic)
-					{
-						if (OtherNodeIndex == INDEX_NONE && OtherNodeIndex2 == INDEX_NONE)
-						{
-							while (OtherNode.UsedColors.Contains(ColorToUse) || GraphNode.UsedColors.Contains(ColorToUse))
-							{
-								ColorToUse++;
-							}
-						}
-						else if (OtherNodeIndex == INDEX_NONE)
-						{
-							//OtherNodeIndex2 != INDEX_NONE
-							FGraphNode& PrevOtherNode = Nodes[OtherNodeIndex2];
-							while (OtherNode.UsedColors.Contains(ColorToUse) || PrevOtherNode.UsedColors.Contains(ColorToUse) || GraphNode.UsedColors.Contains(ColorToUse))
-							{
-								ColorToUse++;
-							}
-						}
-						else if (OtherNodeIndex2 == INDEX_NONE)
-						{
-							FGraphNode& PrevOtherNode = Nodes[OtherNodeIndex];
-							while (OtherNode.UsedColors.Contains(ColorToUse) || PrevOtherNode.UsedColors.Contains(ColorToUse) || GraphNode.UsedColors.Contains(ColorToUse))
-							{
-								ColorToUse++;
-							}
-						}
-						else
-						{
-							FGraphNode& PrevOtherNode = Nodes[OtherNodeIndex];
-							FGraphNode& PrevPrevOtherNode = Nodes[OtherNodeIndex2];
-							while (OtherNode.UsedColors.Contains(ColorToUse) || PrevOtherNode.UsedColors.Contains(ColorToUse) || PrevPrevOtherNode.UsedColors.Contains(ColorToUse) || GraphNode.UsedColors.Contains(ColorToUse))
-							{
-								ColorToUse++;
-							}
-						}
-					}
-				}
-
-				// Assign color and set as used at this node
-				MaxColor = FMath::Max(ColorToUse, MaxColor);
-				GraphNode.UsedColors.Add(ColorToUse);
-				GraphEdge.Color = ColorToUse;
-
-				// Bump color to use next time, but only if we weren't forced to use a different color by the other node
-				if ((ColorToUse == GraphNode.NextColor) && bIsParticleDynamic == true)
-				{
-					GraphNode.NextColor++;
-				}
-
-				if (ColorGraph.Num() <= MaxColor)
-				{
-					ColorGraph.SetNum(MaxColor + 1);
-				}
-				ColorGraph[GraphEdge.Color].Add(EdgeIndex);
-
-				if (OtherNodeIndex != INDEX_NONE)
-				{
-					FGraphNode& OtherGraphNode = Nodes[OtherNodeIndex];
-					const bool bIsOtherGraphNodeDynamic = InParticles.InvM(OtherNodeIndex) != (decltype(InParticles.InvM(ParticleNodeIndex)))0.;
-					if (bIsOtherGraphNodeDynamic)
-					{
-						// Mark other node as not allowing use of this color
-						if (bIsParticleDynamic)
-						{
-							OtherGraphNode.UsedColors.Add(GraphEdge.Color);
-						}
-
-						// Queue other node for processing
-						if (!ProcessedNodes.Contains(OtherNodeIndex))
-						{
-							NodesToProcess.Add(OtherNodeIndex);
-						}
-					}
-				}
-				if (OtherNodeIndex2 != INDEX_NONE)
-				{
-					FGraphNode& OtherGraphNode = Nodes[OtherNodeIndex2];
-					const bool bIsOtherGraphNodeDynamic = InParticles.InvM(OtherNodeIndex2) != (decltype(InParticles.InvM(ParticleNodeIndex)))0.;
-					if (bIsOtherGraphNodeDynamic)
-					{
-						// Mark other node as not allowing use of this color
-						if (bIsParticleDynamic)
-						{
-							OtherGraphNode.UsedColors.Add(GraphEdge.Color);
-						}
-
-						// Queue other node for processing
-						if (!ProcessedNodes.Contains(OtherNodeIndex2))
-						{
-							NodesToProcess.Add(OtherNodeIndex2);
-						}
-					}
-				}
-				if (OtherNodeIndex3 != INDEX_NONE)
-				{
-					FGraphNode& OtherGraphNode = Nodes[OtherNodeIndex3];
-					const bool bIsOtherGraphNodeDynamic = InParticles.InvM(OtherNodeIndex3) != (decltype(InParticles.InvM(ParticleNodeIndex)))0.;
-					if (bIsOtherGraphNodeDynamic)
-					{
-						// Mark other node as not allowing use of this color
-						if (bIsParticleDynamic)
-						{
-							OtherGraphNode.UsedColors.Add(GraphEdge.Color);
-						}
-
-						// Queue other node for processing
-						if (!ProcessedNodes.Contains(OtherNodeIndex3))
-						{
-							NodesToProcess.Add(OtherNodeIndex3);
-						}
-					}
-				}
+				break;
 			}
+			++FirstFreeColor;
 		}
-	}
-
-	checkSlow(VerifyGraph(ColorGraph, Graph, InParticles)); 
-	return ColorGraph;
-}
-
-template<typename DynamicParticlesType>
-TArray<TArray<int32>> Chaos::FGraphColoring::ComputeGraphColoringAllDynamicParticlesOrRange(const TArray<TVec4<int32>>& Graph, const DynamicParticlesType& InParticles, const int32 GraphParticlesStart, const int32 GraphParticlesEnd)
-{
-	using namespace Chaos;
-
-	checkSlow(GraphParticlesStart <= GraphParticlesEnd);
-	checkSlow(GraphParticlesEnd <= (int32)InParticles.Size());
-	TArray<TArray<int32>> ColorGraph;
-	TArray<FGraphNode> NodesSubArray;
-	TArray<FGraphTetEdge> Edges;
-	NodesSubArray.SetNum(GraphParticlesEnd - GraphParticlesStart);
-	Edges.SetNum(Graph.Num());
-	TArrayView<FGraphNode> Nodes(NodesSubArray.GetData() - GraphParticlesStart, GraphParticlesEnd); // Only nodes starting with GraphParticlesStart are valid to access
-
-	int32 MaxColor = -1;
-
-	for (int32 i = 0; i < Graph.Num(); ++i)
-	{
-		const TVec4<int32>& Constraint = Graph[i];
-		Edges[i].FirstNode = Constraint[0];
-		Edges[i].SecondNode = Constraint[1];
-		Edges[i].ThirdNode = Constraint[2];
-		Edges[i].FourthNode = Constraint[3];
-		checkSlow(Constraint[0] >= GraphParticlesStart);
-		checkSlow(Constraint[1] >= GraphParticlesStart);
-		checkSlow(Constraint[2] >= GraphParticlesStart);
-		checkSlow(Constraint[3] >= GraphParticlesStart);
-		Nodes[Constraint[0]].Edges.Add(i);
-		Nodes[Constraint[1]].Edges.Add(i);
-		Nodes[Constraint[2]].Edges.Add(i);
-		Nodes[Constraint[3]].Edges.Add(i);
-	}
-
-	TSet<int32> ProcessedNodes;
-	TArray<int32> NodesToProcess;
-
-	for (int32 ParticleNodeIndex = GraphParticlesStart; ParticleNodeIndex < GraphParticlesEnd; ++ParticleNodeIndex)
-	{
-		//const bool bIsParticleDynamic = InParticles.InvM(ParticleNodeIndex) != (T)0.;
-		if (ProcessedNodes.Contains(ParticleNodeIndex))
+		MaxColor = FMath::Max(MaxColor, FirstFreeColor);
+		if (ColorGraph.Num() <= MaxColor)
 		{
-			continue;
+			ColorGraph.SetNum(MaxColor + 1);
 		}
-
-		NodesToProcess.Add(ParticleNodeIndex);
-
-		while (NodesToProcess.Num())
+		for (int32 NIndex = 0; NIndex < N; ++NIndex)
 		{
-			const int32 NodeIndex = NodesToProcess.Last();
-			FGraphNode& GraphNode = Nodes[NodeIndex];
-
-			NodesToProcess.SetNum(NodesToProcess.Num() - 1, EAllowShrinking::No);
-			ProcessedNodes.Add(NodeIndex);
-
-			for (const int32 EdgeIndex : GraphNode.Edges)
-			{
-				FGraphTetEdge& GraphEdge = Edges[EdgeIndex];
-
-				// If edge has been colored skip it
-				if (GraphEdge.Color >= 0)
-				{
-					continue;
-				}
-
-				// Get index to the other node on the edge
-				int32 OtherNodeIndex = INDEX_NONE;
-				int32 OtherNodeIndex2 = INDEX_NONE;
-				int32 OtherNodeIndex3 = INDEX_NONE;
-				if (GraphEdge.FirstNode == NodeIndex)
-				{
-					OtherNodeIndex = GraphEdge.SecondNode;
-					OtherNodeIndex2 = GraphEdge.ThirdNode;
-					OtherNodeIndex3 = GraphEdge.FourthNode;
-				}
-				if (GraphEdge.SecondNode == NodeIndex)
-				{
-					OtherNodeIndex = GraphEdge.FirstNode;
-					OtherNodeIndex2 = GraphEdge.ThirdNode;
-					OtherNodeIndex3 = GraphEdge.FourthNode;
-				}
-				if (GraphEdge.ThirdNode == NodeIndex)
-				{
-					OtherNodeIndex = GraphEdge.FirstNode;
-					OtherNodeIndex2 = GraphEdge.SecondNode;
-					OtherNodeIndex3 = GraphEdge.FourthNode;
-				}
-				if (GraphEdge.FourthNode == NodeIndex)
-				{
-					OtherNodeIndex = GraphEdge.FirstNode;
-					OtherNodeIndex2 = GraphEdge.SecondNode;
-					OtherNodeIndex3 = GraphEdge.ThirdNode;
-				}
-
-				// Find next color that is not used already at this node
-				while (GraphNode.UsedColors.Contains(GraphNode.NextColor))
-				{
-					GraphNode.NextColor++;
-				}
-				int32 ColorToUse = GraphNode.NextColor;
-
-				// Exclude colors used by the other node (but still allow this node to use them for other edges)
-				if (OtherNodeIndex != INDEX_NONE)
-				{
-					FGraphNode& OtherNode = Nodes[OtherNodeIndex];
-
-					while (OtherNode.UsedColors.Contains(ColorToUse) || GraphNode.UsedColors.Contains(ColorToUse))
-					{
-						ColorToUse++;
-					}
-					
-				}
-				if (OtherNodeIndex2 != INDEX_NONE)
-				{
-					FGraphNode& OtherNode = Nodes[OtherNodeIndex2];
-					if (OtherNodeIndex == INDEX_NONE)
-					{
-						while (OtherNode.UsedColors.Contains(ColorToUse) || GraphNode.UsedColors.Contains(ColorToUse))
-						{
-							ColorToUse++;
-						}
-					}
-					else
-					{
-						FGraphNode& PrevOtherNode = Nodes[OtherNodeIndex];
-						while (OtherNode.UsedColors.Contains(ColorToUse) || PrevOtherNode.UsedColors.Contains(ColorToUse) || GraphNode.UsedColors.Contains(ColorToUse))
-						{
-							ColorToUse++;
-						}
-					}
-					
-				}
-				if (OtherNodeIndex3 != INDEX_NONE)
-				{
-					FGraphNode& OtherNode = Nodes[OtherNodeIndex3];
-
-					if (OtherNodeIndex == INDEX_NONE && OtherNodeIndex2 == INDEX_NONE)
-					{
-						while (OtherNode.UsedColors.Contains(ColorToUse) || GraphNode.UsedColors.Contains(ColorToUse))
-						{
-							ColorToUse++;
-						}
-					}
-					else if (OtherNodeIndex == INDEX_NONE)
-					{
-						//OtherNodeIndex2 != INDEX_NONE
-						FGraphNode& PrevOtherNode = Nodes[OtherNodeIndex2];
-						while (OtherNode.UsedColors.Contains(ColorToUse) || PrevOtherNode.UsedColors.Contains(ColorToUse) || GraphNode.UsedColors.Contains(ColorToUse))
-						{
-							ColorToUse++;
-						}
-					}
-					else if (OtherNodeIndex2 == INDEX_NONE)
-					{
-						FGraphNode& PrevOtherNode = Nodes[OtherNodeIndex];
-						while (OtherNode.UsedColors.Contains(ColorToUse) || PrevOtherNode.UsedColors.Contains(ColorToUse) || GraphNode.UsedColors.Contains(ColorToUse))
-						{
-							ColorToUse++;
-						}
-					}
-					else
-					{
-						FGraphNode& PrevOtherNode = Nodes[OtherNodeIndex];
-						FGraphNode& PrevPrevOtherNode = Nodes[OtherNodeIndex2];
-						while (OtherNode.UsedColors.Contains(ColorToUse) || PrevOtherNode.UsedColors.Contains(ColorToUse) || PrevPrevOtherNode.UsedColors.Contains(ColorToUse) || GraphNode.UsedColors.Contains(ColorToUse))
-						{
-							ColorToUse++;
-						}
-					}
-					
-				}
-
-				// Assign color and set as used at this node
-				MaxColor = FMath::Max(ColorToUse, MaxColor);
-				GraphNode.UsedColors.Add(ColorToUse);
-				GraphEdge.Color = ColorToUse;
-
-				// Bump color to use next time, but only if we weren't forced to use a different color by the other node
-				if (ColorToUse == GraphNode.NextColor) 
-				{
-					GraphNode.NextColor++;
-				}
-
-				if (ColorGraph.Num() <= MaxColor)
-				{
-					ColorGraph.SetNum(MaxColor + 1);
-				}
-				ColorGraph[GraphEdge.Color].Add(EdgeIndex);
-
-				if (OtherNodeIndex != INDEX_NONE)
-				{
-					FGraphNode& OtherGraphNode = Nodes[OtherNodeIndex];
-
-					// Mark other node as not allowing use of this color
-					OtherGraphNode.UsedColors.Add(GraphEdge.Color);
-					
-					// Queue other node for processing
-					if (!ProcessedNodes.Contains(OtherNodeIndex))
-					{
-						NodesToProcess.Add(OtherNodeIndex);
-					}
-					
-				}
-				if (OtherNodeIndex2 != INDEX_NONE)
-				{
-					FGraphNode& OtherGraphNode = Nodes[OtherNodeIndex2];
-
-					// Mark other node as not allowing use of this color
-					OtherGraphNode.UsedColors.Add(GraphEdge.Color);
-					
-					// Queue other node for processing
-					if (!ProcessedNodes.Contains(OtherNodeIndex2))
-					{
-						NodesToProcess.Add(OtherNodeIndex2);
-					}
-					
-				}
-				if (OtherNodeIndex3 != INDEX_NONE)
-				{
-					FGraphNode& OtherGraphNode = Nodes[OtherNodeIndex3];
-					// Mark other node as not allowing use of this color
-					OtherGraphNode.UsedColors.Add(GraphEdge.Color);
-						
-					// Queue other node for processing
-					if (!ProcessedNodes.Contains(OtherNodeIndex3))
-					{
-						NodesToProcess.Add(OtherNodeIndex3);
-					}
-					
-				}
-			}
+			const int32 NodeIndex = Graph[EdgeIndex][NIndex];
+			NodeUsedColors[NodeIndex].Add(FirstFreeColor);
 		}
+		ColorGraph[FirstFreeColor].Add(EdgeIndex);
 	}
-
-	checkSlow(VerifyGraphAllDynamic(ColorGraph, Graph, InParticles));
+#if DO_GUARD_SLOW
+	const bool bVerifyGraphResult = VerifyGraph<DynamicParticlesType, N, bAllDynamic>(ColorGraph, Graph, InParticles);
+	checkSlow(bVerifyGraphResult);
+#endif
 	return ColorGraph;
 }
 
 
 template<typename T>
-void Chaos::ComputeGridBasedGraphSubColoringPointer(const TArray<TArray<int32>>& ElementsPerColor, const TMPMGrid<T>& Grid, const int32 GridSize, TArray<TArray<int32>>*& PreviousColoring, const TArray<TArray<int32>>& ConstraintsNodesSet, TArray<TArray<TArray<int32>>>& ElementsPerSubColors) 
+void Chaos::ComputeGridBasedGraphSubColoringPointer(const TArray<TArray<int32>>& ElementsPerColor, const TMPMGrid<T>& Grid, const int32 GridSize, TUniquePtr<TArray<TArray<int32>>>& PreviousColoring, const TArray<TArray<int32>>& ConstraintsNodesSet, TArray<TArray<TArray<int32>>>& ElementsPerSubColors) 
 {
-	bool InitialGuess = true;
-	if (!PreviousColoring) {
-		InitialGuess = false;
-		PreviousColoring = new TArray<TArray<int32>>();
+	typedef TArray<int32, TInlineAllocator<8>> FColorSet;
+
+	bool bHaveInitialGuess = true;
+	if (!PreviousColoring) 
+	{
+		bHaveInitialGuess = false;
+		PreviousColoring = MakeUnique<TArray<TArray<int32>>>();
 		PreviousColoring->SetNum(ElementsPerColor.Num());
 	}
 	ElementsPerSubColors.SetNum(ElementsPerColor.Num());
 
 	//for (int32 i = 0; i < ElementsPerColor.Num(); i++) {
-	PhysicsParallelFor(ElementsPerColor.Num(), [&](const int32 i)
+	PhysicsParallelFor(ElementsPerColor.Num(), [bHaveInitialGuess, &ElementsPerColor, &Grid, GridSize, &PreviousColoring, &ConstraintsNodesSet, &ElementsPerSubColors](const int32 Color)
 		{
-			if (!InitialGuess) {
-				PreviousColoring->operator[](i).Init(-1, ElementsPerColor[i].Num());
+			const TArray<int32>& ColorElements = ElementsPerColor[Color];
+			TArray<int32>& PreviousColoringElements = PreviousColoring->operator[](Color);
+			if (!bHaveInitialGuess)
+			{
+				PreviousColoringElements.Init(-1, ColorElements.Num());
 			}
 			int32 NumNodes = GridSize;
 			TArray<int32> ElementSubColors;
-			ElementSubColors.Init(-1, ElementsPerColor[i].Num());
-			TArray<TSet<int32>*> UsedColors;
-			UsedColors.Init(nullptr, NumNodes);
-			for (int32 j = 0; j < ElementsPerColor[i].Num(); j++) {
+			ElementSubColors.Init(-1, ColorElements.Num());
+			TArray<FColorSet> UsedColors;
+			UsedColors.SetNum(NumNodes);
+			int32 MaxColor = -1;
+			for (int32 ElementIndex = 0; ElementIndex < ColorElements.Num(); ElementIndex++) 
+			{
 				int32 ColorToUse = 0;
-				int32 e = ElementsPerColor[i][j];
+				int32 Element = ColorElements[ElementIndex];
 				// check initial guess:
-				if (InitialGuess) {
-					ColorToUse = PreviousColoring->operator[](i)[j];
-					bool ColorFound = false;
-					for (auto node : ConstraintsNodesSet[e]) {
-						if (!UsedColors[node]) {
-							UsedColors[node] = new TSet<int32>();
-						}
-						if (UsedColors[node]->Contains(ColorToUse)) {
-							ColorFound = true;
+				if (bHaveInitialGuess)
+				{
+					ColorToUse = PreviousColoringElements[ElementIndex];
+					bool bColorFound = false;
+					for (const int32 Node : ConstraintsNodesSet[Element]) 
+					{
+						if (UsedColors[Node].Contains(ColorToUse)) 
+						{
+							bColorFound = true;
 							break;
 						}
 					}
-					if (ColorFound) {
+					if (bColorFound) 
+					{
 						ColorToUse = 0;
 					}
-					else {
-						for (auto node : ConstraintsNodesSet[e]) {
-							UsedColors[node]->Emplace(ColorToUse);
+					else 
+					{
+						for (const int32 Node : ConstraintsNodesSet[Element])
+						{
+							UsedColors[Node].Emplace(ColorToUse);
 						}
-						ElementSubColors[j] = ColorToUse;
+						ElementSubColors[ElementIndex] = ColorToUse;
 					}
 				}
-				if (ElementSubColors[j] == -1) {
-					while (true) {
-						bool ColorFound = false;
-						for (auto node : ConstraintsNodesSet[e]) {
-							if (!UsedColors[node]) {
-								UsedColors[node] = new TSet<int32>();// new std::unordered_set<int>;
-							}
-							if (UsedColors[node]->Contains(ColorToUse)) {
-								ColorFound = true;
+				if (ElementSubColors[ElementIndex] == -1)
+				{
+					while (true) 
+					{
+						bool bColorFound = false;
+						for (const int32 Node : ConstraintsNodesSet[Element]) 
+						{
+							if (UsedColors[Node].Contains(ColorToUse)) 
+							{
+								bColorFound = true;
 								break;
 							}
 						}
-						if (!ColorFound) {
+						if (!bColorFound) {
 							break;
 						}
 						ColorToUse++;
 					}
-					ElementSubColors[j] = ColorToUse;
-					for (auto node : ConstraintsNodesSet[e]) {
-						UsedColors[node]->Emplace(ColorToUse);
+					ElementSubColors[ElementIndex] = ColorToUse;
+					for (const int32 Node : ConstraintsNodesSet[Element])
+					{
+						UsedColors[Node].Emplace(ColorToUse);
 					}
 				}
 
 				// assign colors to previous guess for next timestep:
-				PreviousColoring->operator[](i)[j] = ColorToUse;
-			}
-
-			for (int ii = 0; ii < UsedColors.Num(); ii++) {
-				delete (UsedColors[ii]);
+				MaxColor = FMath::Max(MaxColor, ColorToUse);
+				PreviousColoringElements[ElementIndex] = ColorToUse;
 			}
 
-			int32 NumColors = FMath::Max<int32>(ElementSubColors);
+			ElementsPerSubColors[Color].Reset();
+			ElementsPerSubColors[Color].SetNum(MaxColor + 1);
 
-			//int32 num_colors = *std::max_element(ElementSubColors.begin(), ElementSubColors.end());
-			ElementsPerSubColors[i].SetNum(0);
-			ElementsPerSubColors[i].SetNum(NumColors + 1);
-
-			for (int32 j = 0; j < ElementsPerColor[i].Num(); j++) {
-				ElementsPerSubColors[i][ElementSubColors[j]].Emplace(ElementsPerColor[i][j]);
+			for (int32 ElementIndex = 0; ElementIndex < ColorElements.Num(); ElementIndex++)
+			{
+				ElementsPerSubColors[Color][ElementSubColors[ElementIndex]].Emplace(ColorElements[ElementIndex]);
 			}
 		}, ElementsPerColor.Num() < 20);
 	
@@ -2008,35 +1109,32 @@ void Chaos::ComputeExtraNodalColoring(const TArray<TArray<int32>>& StaticGraph,
 	checkSlow(VerifyExtraNodalColoring<T>(InParticles, StaticGraph, DynamicGraph, ExtraGraph, StaticIncidentElements, DynamicIncidentElements, ExtraIncidentElements, ParticleColors, ParticlesPerColor));
 }
 
+#define UE_COMPUTE_GRAPH_COLORING_PARTICLES_OR_RANGE_N_HELPER(N, bAllDynamic) \
+template CHAOS_API TArray<TArray<int32>> Chaos::FGraphColoring::ComputeGraphColoringParticlesOrRange<Chaos::TDynamicParticles<Chaos::FRealSingle, 3>, N, bAllDynamic>(const TArray<Chaos::TVector<int32, N>>&, const Chaos::TDynamicParticles<Chaos::FRealSingle, 3>&, const int32 GraphParticlesStart, const int32 GraphParticlesEnd); \
+template CHAOS_API TArray<TArray<int32>> Chaos::FGraphColoring::ComputeGraphColoringParticlesOrRange<Chaos::TDynamicParticles<Chaos::FRealDouble, 3>, N, bAllDynamic>(const TArray<Chaos::TVector<int32, N>>&, const Chaos::TDynamicParticles<Chaos::FRealDouble, 3>&, const int32 GraphParticlesStart, const int32 GraphParticlesEnd); \
+template CHAOS_API TArray<TArray<int32>> Chaos::FGraphColoring::ComputeGraphColoringParticlesOrRange<Chaos::Softs::FSolverParticles, N, bAllDynamic>(const TArray<Chaos::TVector<int32, N>>&, const Chaos::Softs::FSolverParticles&, const int32 GraphParticlesStart, const int32 GraphParticlesEnd); \
+template CHAOS_API TArray<TArray<int32>> Chaos::FGraphColoring::ComputeGraphColoringParticlesOrRange<Chaos::Softs::FSolverParticlesRange, N, bAllDynamic>(const TArray<Chaos::TVector<int32, N>>&, const Chaos::Softs::FSolverParticlesRange&, const int32 GraphParticlesStart, const int32 GraphParticlesEnd);
 
+UE_COMPUTE_GRAPH_COLORING_PARTICLES_OR_RANGE_N_HELPER(2, false)
+UE_COMPUTE_GRAPH_COLORING_PARTICLES_OR_RANGE_N_HELPER(3, false)
+UE_COMPUTE_GRAPH_COLORING_PARTICLES_OR_RANGE_N_HELPER(4, false)
+UE_COMPUTE_GRAPH_COLORING_PARTICLES_OR_RANGE_N_HELPER(6, false)
+UE_COMPUTE_GRAPH_COLORING_PARTICLES_OR_RANGE_N_HELPER(4, true)
 
+template CHAOS_API void Chaos::ComputeGridBasedGraphSubColoringPointer(const TArray<TArray<int32>>& ElementsPerColor, const TMPMGrid<Chaos::FRealSingle>& Grid, const int32 GridSize, TUniquePtr<TArray<TArray<int32>>>& PreviousColoring, const TArray<TArray<int32>>& ConstraintsNodesSet, TArray<TArray<TArray<int32>>>& ElementsPerSubColors);
+template CHAOS_API void Chaos::ComputeGridBasedGraphSubColoringPointer(const TArray<TArray<int32>>& ElementsPerColor, const TMPMGrid<Chaos::FRealDouble>& Grid, const int32 GridSize, TUniquePtr<TArray<TArray<int32>>>& PreviousColoring, const TArray<TArray<int32>>& ConstraintsNodesSet, TArray<TArray<TArray<int32>>>& ElementsPerSubColors);
 
-
-template CHAOS_API TArray<TArray<int32>> Chaos::FGraphColoring::ComputeGraphColoringParticlesOrRange(const TArray<Chaos::TVector<int32, 2>>&, const Chaos::TDynamicParticles<Chaos::FRealSingle, 3>&, const int32 GraphParticlesStart, const int32 GraphParticlesEnd);
-template CHAOS_API TArray<TArray<int32>> Chaos::FGraphColoring::ComputeGraphColoringParticlesOrRange(const TArray<Chaos::TVector<int32, 2>>&, const Chaos::TDynamicParticles<Chaos::FRealDouble, 3>&, const int32 GraphParticlesStart, const int32 GraphParticlesEnd);
-template CHAOS_API TArray<TArray<int32>> Chaos::FGraphColoring::ComputeGraphColoringParticlesOrRange(const TArray<Chaos::TVector<int32, 2>>&, const Chaos::Softs::FSolverParticles&, const int32 GraphParticlesStart, const int32 GraphParticlesEnd);
-template CHAOS_API TArray<TArray<int32>> Chaos::FGraphColoring::ComputeGraphColoringParticlesOrRange(const TArray<Chaos::TVector<int32, 2>>&, const Chaos::Softs::FSolverParticlesRange&, const int32 GraphParticlesStart, const int32 GraphParticlesEnd);
-template CHAOS_API TArray<TArray<int32>> Chaos::FGraphColoring::ComputeGraphColoringParticlesOrRange(const TArray<Chaos::TVector<int32, 3>>&, const Chaos::TDynamicParticles<Chaos::FRealSingle, 3>&, const int32 GraphParticlesStart, const int32 GraphParticlesEnd);
-template CHAOS_API TArray<TArray<int32>> Chaos::FGraphColoring::ComputeGraphColoringParticlesOrRange(const TArray<Chaos::TVector<int32, 3>>&, const Chaos::TDynamicParticles<Chaos::FRealDouble, 3>&, const int32 GraphParticlesStart, const int32 GraphParticlesEnd);
-template CHAOS_API TArray<TArray<int32>> Chaos::FGraphColoring::ComputeGraphColoringParticlesOrRange(const TArray<Chaos::TVector<int32, 3>>&, const Chaos::Softs::FSolverParticles&, const int32 GraphParticlesStart, const int32 GraphParticlesEnd);
-template CHAOS_API TArray<TArray<int32>> Chaos::FGraphColoring::ComputeGraphColoringParticlesOrRange(const TArray<Chaos::TVector<int32, 3>>&, const Chaos::Softs::FSolverParticlesRange&, const int32 GraphParticlesStart, const int32 GraphParticlesEnd);
-template CHAOS_API TArray<TArray<int32>> Chaos::FGraphColoring::ComputeGraphColoringParticlesOrRange(const TArray<Chaos::TVector<int32, 4>>&, const Chaos::TDynamicParticles<Chaos::FRealSingle, 3>&, const int32 GraphParticlesStart, const int32 GraphParticlesEnd);
-template CHAOS_API TArray<TArray<int32>> Chaos::FGraphColoring::ComputeGraphColoringParticlesOrRange(const TArray<Chaos::TVector<int32, 4>>&, const Chaos::TDynamicParticles<Chaos::FRealDouble, 3>&, const int32 GraphParticlesStart, const int32 GraphParticlesEnd);
-template CHAOS_API TArray<TArray<int32>> Chaos::FGraphColoring::ComputeGraphColoringParticlesOrRange(const TArray<Chaos::TVector<int32, 4>>&, const Chaos::Softs::FSolverParticles&, const int32 GraphParticlesStart, const int32 GraphParticlesEnd);
-template CHAOS_API TArray<TArray<int32>> Chaos::FGraphColoring::ComputeGraphColoringParticlesOrRange(const TArray<Chaos::TVector<int32, 4>>&, const Chaos::Softs::FSolverParticlesRange&, const int32 GraphParticlesStart, const int32 GraphParticlesEnd);
-template CHAOS_API TArray<TArray<int32>> Chaos::FGraphColoring::ComputeGraphColoringAllDynamicParticlesOrRange(const TArray<Chaos::TVector<int32, 4>>&, const Chaos::TDynamicParticles<Chaos::FRealSingle, 3>&, const int32 GraphParticlesStart, const int32 GraphParticlesEnd);
-template CHAOS_API TArray<TArray<int32>> Chaos::FGraphColoring::ComputeGraphColoringAllDynamicParticlesOrRange(const TArray<Chaos::TVector<int32, 4>>&, const Chaos::TDynamicParticles<Chaos::FRealDouble, 3>&, const int32 GraphParticlesStart, const int32 GraphParticlesEnd);
-template CHAOS_API TArray<TArray<int32>> Chaos::FGraphColoring::ComputeGraphColoringAllDynamicParticlesOrRange(const TArray<Chaos::TVector<int32, 4>>&, const Chaos::Softs::FSolverParticles&, const int32 GraphParticlesStart, const int32 GraphParticlesEnd);
-template CHAOS_API TArray<TArray<int32>> Chaos::FGraphColoring::ComputeGraphColoringAllDynamicParticlesOrRange(const TArray<Chaos::TVector<int32, 4>>&, const Chaos::Softs::FSolverParticlesRange&, const int32 GraphParticlesStart, const int32 GraphParticlesEnd);
-template CHAOS_API void Chaos::ComputeGridBasedGraphSubColoringPointer(const TArray<TArray<int32>>& ElementsPerColor, const TMPMGrid<Chaos::FRealSingle>& Grid, const int32 GridSize, TArray<TArray<int32>>*& PreviousColoring, const TArray<TArray<int32>>& ConstraintsNodesSet, TArray<TArray<TArray<int32>>>& ElementsPerSubColors);
-template CHAOS_API void Chaos::ComputeGridBasedGraphSubColoringPointer(const TArray<TArray<int32>>& ElementsPerColor, const TMPMGrid<Chaos::FRealDouble>& Grid, const int32 GridSize, TArray<TArray<int32>>*& PreviousColoring, const TArray<TArray<int32>>& ConstraintsNodesSet, TArray<TArray<TArray<int32>>>& ElementsPerSubColors);
 template CHAOS_API void Chaos::ComputeWeakConstraintsColoring<Chaos::FRealSingle>(const TArray<TArray<int32>>& Indices, const TArray<TArray<int32>>& SecondIndices, const Chaos::TDynamicParticles<Chaos::FRealSingle, 3>& InParticles, TArray<TArray<int32>>& ConstraintsPerColor);
 template CHAOS_API void Chaos::ComputeWeakConstraintsColoring<Chaos::FRealDouble>(const TArray<TArray<int32>>& Indices, const TArray<TArray<int32>>& SecondIndices, const Chaos::TDynamicParticles<Chaos::FRealDouble, 3>& InParticles, TArray<TArray<int32>>& ConstraintsPerColor);
+
 template CHAOS_API TArray<TArray<int32>> Chaos::ComputeNodalColoring<Chaos::FRealSingle>(const TArray<TVec4<int32>>& Graph, const Chaos::TDynamicParticles<Chaos::FRealSingle, 3>& InParticles, const int32 GraphParticlesStart, const int32 GraphParticlesEnd, const TArray<TArray<int32>>& IncidentElements, const TArray<TArray<int32>>& IncidentElementsLocalIndex);
 template CHAOS_API TArray<TArray<int32>> Chaos::ComputeNodalColoring<Chaos::FRealDouble>(const TArray<TVec4<int32>>& Graph, const Chaos::TDynamicParticles<Chaos::FRealDouble, 3>& InParticles, const int32 GraphParticlesStart, const int32 GraphParticlesEnd, const TArray<TArray<int32>>& IncidentElements, const TArray<TArray<int32>>& IncidentElementsLocalIndex);
+
 template CHAOS_API TArray<TArray<int32>> Chaos::ComputeNodalColoring<Chaos::FRealSingle>(const TArray<TArray<int32>>& Graph, const Chaos::TDynamicParticles<Chaos::FRealSingle, 3>& InParticles, const int32 GraphParticlesStart, const int32 GraphParticlesEnd, const TArray<TArray<int32>>& IncidentElements, const TArray<TArray<int32>>& IncidentElementsLocalIndex, const TPBDActiveView<Chaos::Softs::FSolverParticles>* InParticleActiveView, TArray<int32>* ParticleColorsOut);
 template CHAOS_API TArray<TArray<int32>> Chaos::ComputeNodalColoring<Chaos::FRealDouble>(const TArray<TArray<int32>>& Graph, const Chaos::TDynamicParticles<Chaos::FRealDouble, 3>& InParticles, const int32 GraphParticlesStart, const int32 GraphParticlesEnd, const TArray<TArray<int32>>& IncidentElements, const TArray<TArray<int32>>& IncidentElementsLocalIndex, const TPBDActiveView<Chaos::TDynamicParticles<Chaos::FRealDouble, 3>>* InParticleActiveView, TArray<int32>* ParticleColorsOut);
+
 template CHAOS_API void Chaos::ComputeExtraNodalColoring<Chaos::FRealSingle>(const TArray<TArray<int32>>& Graph, const TArray<TArray<int32>>& ExtraGraph, const Chaos::TDynamicParticles<Chaos::FRealSingle, 3>& InParticles, const TArray<TArray<int32>>& IncidentElements, const TArray<TArray<int32>>& ExtraIncidentElements, TArray<int32>& ParticleColors, TArray<TArray<int32>>& ParticlesPerColor);
 template CHAOS_API void Chaos::ComputeExtraNodalColoring<Chaos::FRealDouble>(const TArray<TArray<int32>>& Graph, const TArray<TArray<int32>>& ExtraGraph, const Chaos::TDynamicParticles<Chaos::FRealDouble, 3>& InParticles, const TArray<TArray<int32>>& IncidentElements, const TArray<TArray<int32>>& ExtraIncidentElements, TArray<int32>& ParticleColors, TArray<TArray<int32>>& ParticlesPerColor);
+
 template CHAOS_API void Chaos::ComputeExtraNodalColoring<Chaos::FRealSingle>(const TArray<TArray<int32>>& StaticGraph, const TArray<TArray<int32>>& DynamicGraph, const TArray<TArray<int32>>& ExtraGraph, const Chaos::TDynamicParticles<Chaos::FRealSingle, 3>& InParticles, const TArray<TArray<int32>>& StaticIncidentElements, const TArray<TArray<int32>>& DynamicIncidentElements, const TArray<TArray<int32>>& ExtraIncidentElements, TArray<int32>& ParticleColors, TArray<TArray<int32>>& ParticlesPerColor);
 template CHAOS_API void Chaos::ComputeExtraNodalColoring<Chaos::FRealDouble>(const TArray<TArray<int32>>& StaticGraph, const TArray<TArray<int32>>& DynamicGraph, const TArray<TArray<int32>>& ExtraGraph, const Chaos::TDynamicParticles<Chaos::FRealDouble, 3>& InParticles, const TArray<TArray<int32>>& StaticIncidentElements, const TArray<TArray<int32>>& DynamicIncidentElements, const TArray<TArray<int32>>& ExtraIncidentElements, TArray<int32>& ParticleColors, TArray<TArray<int32>>& ParticlesPerColor);
\ No newline at end of file
diff --git a/Engine/Source/Runtime/Experimental/Chaos/Private/Chaos/PBDRepulsionConstraints.cpp b/Engine/Source/Runtime/Experimental/Chaos/Private/Chaos/PBDRepulsionConstraints.cpp
new file mode 100644
--- /dev/null
+++ b/Engine/Source/Runtime/Experimental/Chaos/Private/Chaos/PBDRepulsionConstraints.cpp
@@ -0,0 +1,213 @@
+// Copyright Epic Games, Inc. All Rights Reserved.
+#include "Chaos/PBDRepulsionConstraints.h"
+#include "Chaos/Framework/Parallel.h"
+#include "Chaos/GraphColoring.h"
+#include "Chaos/Triangle.h"
+#include "ChaosStats.h"
+#include "ChaosLog.h"
+#include "HAL/IConsoleManager.h"
+
+#if UE_BUILD_SHIPPING
+static constexpr Chaos::Softs::FSolverReal RepulsionConstraintFalloffMultiplier = (Chaos::Softs::FSolverReal)1.f;
+static constexpr Chaos::Softs::FSolverReal RepulsionConstraintMaxDepthMultiplier = (Chaos::Softs::FSolverReal)10.f;
+#else
+Chaos::Softs::FSolverReal RepulsionConstraintFalloffMultiplier = (Chaos::Softs::FSolverReal)1.f;
+FAutoConsoleVariableRef CVarChaosRepulsionConstraintFalloffMultiplier(TEXT("p.Chaos.RepulsionConstraint.FalloffMultiplier"), RepulsionConstraintFalloffMultiplier, TEXT("Tangential distance away from a triangle (scaled by thickness) beyond which a point isn't considered to be repelling"));
+
+Chaos::Softs::FSolverReal RepulsionConstraintMaxDepthMultiplier = (Chaos::Softs::FSolverReal)10.f;
+FAutoConsoleVariableRef CVarChaosRepulsionConstraintMaxDepthMultiplier(TEXT("p.Chaos.RepulsionConstraint.MaxDepthMultiplier"), RepulsionConstraintMaxDepthMultiplier, TEXT("Penetration depth beyond which we ignore the repulsion (so you don't push through the wrong side)"));
+#endif
+
+namespace Chaos::Softs
+{
+	FPBDRepulsionConstraints::FPBDRepulsionConstraints(
+		const FSolverParticlesRange& Particles,
+		const FCollectionPropertyConstFacade& PropertyCollection,
+		const FEmbeddedSpringConstraintFacade& SpringConstraintFacade
+	)
+		: Stiffness(FMath::Clamp(GetVertexFaceRepulsionStiffness(PropertyCollection, DefaultStiffness), MinStiffness, MaxStiffness))
+		, FrictionCoefficient(FMath::Clamp(GetVertexFaceRepulsionFrictionCoefficient(PropertyCollection, DefaultFriction), MinFriction, MaxFriction))
+		, VertexFaceRepulsionStiffnessIndex(PropertyCollection)
+		, VertexFaceRepulsionFrictionCoefficientIndex(PropertyCollection)
+	{
+		const TConstArrayView<TArray<int32>>& InSourceIndices = SpringConstraintFacade.GetSourceIndexConst();
+		const TConstArrayView<TArray<int32>>& InTargetIndices = SpringConstraintFacade.GetTargetIndexConst();
+		const TConstArrayView<TArray<float>>& InTargetWeights = SpringConstraintFacade.GetTargetWeightsConst();
+		const TConstArrayView<float>& InSpringLengths = SpringConstraintFacade.GetSpringLengthConst();
+
+		// Validate data
+		const int32 NumConstraints = InSourceIndices.Num();
+		check(InTargetIndices.Num() == NumConstraints);
+		check(InTargetWeights.Num() == NumConstraints);
+		check(InSpringLengths.Num() == NumConstraints);
+		Constraints.Reserve(NumConstraints);
+		TargetWeights.Reserve(NumConstraints);
+		RestLengths.Reserve(NumConstraints);
+
+		auto IsKinematic = [&Particles](const TArray<int32>& Indices, const TArray<FRealSingle>& Weights, int32 ExpectedNum)
+			{
+				for (int32 Idx = 0; Idx < ExpectedNum; ++Idx)
+				{
+					if (Weights[Idx] != 0.f && Particles.InvM(Indices[Idx]) != 0.f)
+					{
+						return false;
+					}
+				}
+				return true;
+			};
+
+		TArray<float> InSourceWeight = { 1.f };
+		for (int32 ConstraintIdx = 0; ConstraintIdx < NumConstraints; ++ConstraintIdx)
+		{
+			if (IsKinematic(InSourceIndices[ConstraintIdx], InSourceWeight, 1)
+				&& IsKinematic(InTargetIndices[ConstraintIdx], InTargetWeights[ConstraintIdx], 3))
+			{
+				// Strip this constraint
+				continue;
+			}
+
+			const TVector<int32, 4> Constraint = { InSourceIndices[ConstraintIdx][0], InTargetIndices[ConstraintIdx][0], InTargetIndices[ConstraintIdx][1], InTargetIndices[ConstraintIdx][2] };
+			const TVector<FSolverReal, 3> TargetWeight = { InTargetWeights[ConstraintIdx][0], InTargetWeights[ConstraintIdx][1], InTargetWeights[ConstraintIdx][2] };
+
+			const int32 NewConstraintIndex = Constraints.Emplace(Constraint);
+			TargetWeights.Emplace(TargetWeight);
+			RestLengths.Emplace(InSpringLengths[ConstraintIdx]);
+		}
+
+		InitColor(Particles);
+
+#if CHAOS_DEBUG_DRAW
+		ConstraintIsActive.SetNum(Constraints.Num());
+#endif
+	}
+
+	void FPBDRepulsionConstraints::InitColor(const FSolverParticlesRange& Particles)
+	{
+		const int32 NumConstraints = Constraints.Num();
+		const TArray<TArray<int32>> ConstraintsPerColor = FGraphColoring::ComputeGraphColoringParticlesOrRange(Constraints, Particles, 0, Particles.Size());
+		// Reorder constraints based on color so each array in ConstraintsPerColor contains contiguous elements.
+		TArray<TVector<int32, 4>> ReorderedConstraints;
+		TArray<FSolverVec3> ReorderedTargetWeights;
+		TArray<FSolverReal> ReorderedRestLengths;
+		ReorderedConstraints.SetNumUninitialized(NumConstraints);
+		ReorderedTargetWeights.SetNumUninitialized(NumConstraints);
+		ReorderedRestLengths.SetNumUninitialized(NumConstraints);
+		ConstraintsPerColorStartIndex.Reset(ConstraintsPerColor.Num() + 1);
+		int32 ReorderedIndex = 0;
+		for (const TArray<int32>& ConstraintsBatch : ConstraintsPerColor)
+		{
+			ConstraintsPerColorStartIndex.Add(ReorderedIndex);
+			for (const int32 OrigIndex : ConstraintsBatch)
+			{
+				ReorderedConstraints[ReorderedIndex] = Constraints[OrigIndex];
+				ReorderedTargetWeights[ReorderedIndex] = TargetWeights[OrigIndex];
+				ReorderedRestLengths[ReorderedIndex] = RestLengths[OrigIndex];
+
+				++ReorderedIndex;
+			}
+		}
+		ConstraintsPerColorStartIndex.Add(ReorderedIndex);
+
+		Constraints = MoveTemp(ReorderedConstraints);
+		TargetWeights = MoveTemp(ReorderedTargetWeights);
+		RestLengths = MoveTemp(ReorderedRestLengths);
+	}
+
+	void FPBDRepulsionConstraints::SetProperties(const FCollectionPropertyConstFacade& PropertyCollection)
+	{
+		if (IsVertexFaceRepulsionStiffnessMutable(PropertyCollection))
+		{
+			Stiffness = FMath::Clamp(GetVertexFaceRepulsionStiffness(PropertyCollection), MinStiffness, MaxStiffness);
+		}
+		if (IsVertexFaceRepulsionFrictionCoefficientMutable(PropertyCollection))
+		{
+			FrictionCoefficient = FMath::Clamp(GetVertexFaceRepulsionFrictionCoefficient(PropertyCollection), MinFriction, MaxFriction);
+		}
+	}
+
+	void FPBDRepulsionConstraints::Apply(FSolverParticlesRange& Particles, const FSolverReal Dt) const
+	{
+		// TODO ISPC
+		const TVec4<int32>* const ConstraintsData = Constraints.GetData();
+		const FSolverReal* const RestLengthsData = RestLengths.GetData();
+
+		FPAndInvM* const PAndInvM = Particles.GetPAndInvM().GetData();
+		const FSolverVec3* const X = Particles.XArray().GetData();
+
+		for (int32 ConstraintIndex = 0; ConstraintIndex < Constraints.Num(); ++ConstraintIndex)
+		{
+#if CHAOS_DEBUG_DRAW
+			ConstraintIsActive[ConstraintIndex] = false;
+#endif
+			const TVec4<int32>& Constraint = ConstraintsData[ConstraintIndex];
+			const int32 Index1 = Constraint[0];
+			const int32 Index2 = Constraint[1];
+			const int32 Index3 = Constraint[2];
+			const int32 Index4 = Constraint[3];
+			FSolverVec3& P1 = PAndInvM[Index1].P;
+			FSolverVec3& P2 = PAndInvM[Index2].P;
+			FSolverVec3& P3 = PAndInvM[Index3].P;
+			FSolverVec3& P4 = PAndInvM[Index4].P;
+
+			FSolverVec3 Bary;
+			const FSolverVec3 P = FindClosestPointAndBaryOnTriangle(P2, P3, P4, P1, Bary);
+
+			const FSolverReal TrianglePointInvM =
+				PAndInvM[Index2].InvM * Bary[0] +
+				PAndInvM[Index3].InvM * Bary[1] +
+				PAndInvM[Index4].InvM * Bary[2];
+
+			const FSolverReal CombinedMass = PAndInvM[Index1].InvM + TrianglePointInvM;
+			if (CombinedMass <= (FSolverReal)1e-7)
+			{
+				continue;
+			}
+
+			const FSolverReal Height = RestLengthsData[ConstraintIndex];
+			const FSolverReal MaxDepth = -Height * RepulsionConstraintMaxDepthMultiplier;
+			const FSolverReal OneOverTangentialFalloffDist = (FSolverReal)1.f / FMath::Max(RepulsionConstraintFalloffMultiplier * Height, UE_KINDA_SMALL_NUMBER);
+			const TTriangle<FSolverReal> Triangle(P2, P3, P4);
+			const FSolverVec3 Normal = Triangle.GetNormal();
+			const FSolverVec3 Difference = P1 - P;
+			const FSolverReal NormalDifference = Difference.Dot(Normal);
+
+			if (NormalDifference >= Height || NormalDifference < MaxDepth)
+			{
+				continue;
+			}
+
+			const FSolverReal TangentialDifference = (Difference - NormalDifference * Normal).Size();
+			const FSolverReal TangentialFalloff = (FSolverReal)1.f - TangentialDifference * OneOverTangentialFalloffDist;
+			if (TangentialFalloff <= 0.f)
+			{
+				continue;
+			}
+
+			const FSolverReal NormalDelta = Height - NormalDifference;
+			const FSolverVec3 RepulsionDelta = Stiffness * TangentialFalloff * NormalDelta * Normal / CombinedMass;
+
+			FSolverVec3 FrictionDelta((FSolverReal)0.);
+			if (FrictionCoefficient > 0)
+			{
+				const FSolverVec3& X1 = X[Index1];
+				const FSolverVec3 XP = Bary[0] * X[Index2] + Bary[1] * X[Index3] + Bary[2] * X[Index4];
+				const FSolverVec3 RelativeDisplacement = (P1 - X1) - (P - XP) + (PAndInvM[Index1].InvM - TrianglePointInvM) * RepulsionDelta;
+				const FSolverVec3 RelativeDisplacementTangent = RelativeDisplacement - RelativeDisplacement.Dot(Normal) * Normal;
+				const FSolverReal RelativeDisplacementTangentLength = RelativeDisplacementTangent.Length();
+				const FSolverReal PositionCorrection = FMath::Min(NormalDelta * FrictionCoefficient, RelativeDisplacementTangentLength);
+				const FSolverReal CorrectionRatio = RelativeDisplacementTangentLength < UE_SMALL_NUMBER ? 0.f : PositionCorrection / RelativeDisplacementTangentLength;
+				FrictionDelta = -CorrectionRatio * RelativeDisplacementTangent / CombinedMass;
+			}
+			const FSolverVec3 Delta = RepulsionDelta + FrictionDelta;
+
+			P1 += PAndInvM[Index1].InvM * Delta;
+			P2 -= PAndInvM[Index2].InvM * Bary[0] * Delta;
+			P3 -= PAndInvM[Index3].InvM * Bary[1] * Delta;
+			P4 -= PAndInvM[Index4].InvM * Bary[2] * Delta;
+
+#if CHAOS_DEBUG_DRAW
+			ConstraintIsActive[ConstraintIndex] = true;
+#endif
+		}
+	}
+} // End namespace Chaos::Softs
diff --git a/Engine/Source/Runtime/Experimental/Chaos/Private/Chaos/SoftsSpring.isph b/Engine/Source/Runtime/Experimental/Chaos/Private/Chaos/SoftsSpring.isph
--- a/Engine/Source/Runtime/Experimental/Chaos/Private/Chaos/SoftsSpring.isph
+++ b/Engine/Source/Runtime/Experimental/Chaos/Private/Chaos/SoftsSpring.isph
@@ -146,4 +146,44 @@ inline FVector3f GetXPBDAxialSpringDampingDelta(const FVector3f& P1, const FVect
 
 	return DLambda * Direction;
 }
+template <typename StiffnessFloatType>
+inline FVector3f GetXPBDEmbeddedSpringDelta(const FVector3f& P1, const FVector3f& P2, const varying float InvM1, const varying float InvM2,
+	 const uniform float Dt, const varying float Dist, const varying float Lambda, const StiffnessFloatType ExtensionStiffness, const StiffnessFloatType CompressionStiffness, float& DLambda)
+{
+	FVector3f Direction = P1 - P2;
+	const float Distance = SafeNormalize(Direction);
+	const varying float Offset = Distance - Dist;
+	
+	const varying float Stiffness = select(Offset >= 0, ExtensionStiffness, CompressionStiffness);
+	const varying float AlphaInv = Stiffness * Dt * Dt;
+	
+	const float CombinedInvMass = InvM2 + InvM1;
+
+	DLambda = (-AlphaInv * Offset - Lambda) / (AlphaInv * CombinedInvMass + (uniform float)1.f);
+
+	return DLambda * Direction;
+}
+
+
+// Spring damping constraint (separate from spring stretching)
+template <typename StiffnessFloatType>
+inline FVector3f GetXPBDEmbeddedSpringDampingDelta(const FVector3f& P1, const FVector3f& P2, const FVector3f& X1, const FVector3f& X2,
+	const varying float InvM1, const varying float InvM2,
+	const uniform float Dt, const varying float Dist, const varying float Lambda, const StiffnessFloatType ExtensionStiffness, const StiffnessFloatType CompressionStiffness, const StiffnessFloatType DampingRatio, float& DLambda)
+{
+	FVector3f Direction = P1 - P2;
+	const float Distance = SafeNormalize(Direction);
+	const varying float Offset = Distance - Dist;
+	const varying float OneOverDist = select(Dist< FLOAT_SMALL_NUMBER, 1.f, 1.f / Dist);
+	
+	const float CombinedInvMass = InvM2 + InvM1;
+	const varying float Stiffness = select(Offset >= 0, ExtensionStiffness, CompressionStiffness);
+	const varying float Damping = (uniform float)2.f * DampingRatio * sqrt(Stiffness / CombinedInvMass) * OneOverDist;
+	const varying float BetaDt = Damping * Dt;
+	const varying FVector3f RelativeVelocityTimesDt = P1 - X1 - P2 + X2;
+
+	DLambda = ( -BetaDt * VectorDot(Direction, RelativeVelocityTimesDt) - Lambda) / (BetaDt * CombinedInvMass + (uniform float)1.f);
+
+	return DLambda * Direction;
+}
 #endif
diff --git a/Engine/Source/Runtime/Experimental/Chaos/Private/Chaos/TriangleMesh.cpp b/Engine/Source/Runtime/Experimental/Chaos/Private/Chaos/TriangleMesh.cpp
--- a/Engine/Source/Runtime/Experimental/Chaos/Private/Chaos/TriangleMesh.cpp
+++ b/Engine/Source/Runtime/Experimental/Chaos/Private/Chaos/TriangleMesh.cpp
@@ -178,7 +178,7 @@ void FTriangleMesh::Init(TArray<TVec3<int32>>&& Elements, const int32 StartIdx,
 	InitHelper(StartIdx, EndIdx, CullDegenerateElements);
 }
 
-void FTriangleMesh::Init(const TArray<TVec3<int32>>& Elements, const int32 StartIdx, const int32 EndIdx, const bool CullDegenerateElements)
+void FTriangleMesh::Init(const TConstArrayView<TVec3<int32>>& Elements, const int32 StartIdx, const int32 EndIdx, const bool CullDegenerateElements)
 {
 	MElements = Elements;
 	MStartIdx = 0;
diff --git a/Engine/Source/Runtime/Experimental/Chaos/Public/Chaos/CollectionEmbeddedSpringConstraintFacade.h b/Engine/Source/Runtime/Experimental/Chaos/Public/Chaos/CollectionEmbeddedSpringConstraintFacade.h
--- a/Engine/Source/Runtime/Experimental/Chaos/Public/Chaos/CollectionEmbeddedSpringConstraintFacade.h
+++ b/Engine/Source/Runtime/Experimental/Chaos/Public/Chaos/CollectionEmbeddedSpringConstraintFacade.h
@@ -19,6 +19,7 @@ namespace Chaos::Softs
 		static CHAOS_API const FName ConstraintStart;
 		static CHAOS_API const FName ConstraintEnd;
 		static CHAOS_API const FName ConstraintEndPointNumIndices;
+		static CHAOS_API const FName ConstraintName;
 
 		CHAOS_API FEmbeddedSpringBaseFacade(FManagedArrayCollection& InCollection, const FName& InVerticesGroup = FGeometryCollection::VerticesGroup);
 		CHAOS_API FEmbeddedSpringBaseFacade(const FManagedArrayCollection& InCollection, const FName& InVerticesGroup = FGeometryCollection::VerticesGroup);
@@ -35,6 +36,7 @@ namespace Chaos::Softs
 		TManagedArrayAccessor<int32> ConstraintStartAttribute;
 		TManagedArrayAccessor<int32> ConstraintEndAttribute;
 		TManagedArrayAccessor<FUintVector2> ConstraintEndPointNumIndicesAttribute;
+		TManagedArrayAccessor<FString> ConstraintNameAttribute;
 	};
 
 	/** Facade for managing groups of springs with the same embedding type (e.g., vertex-vertex or barycentric point-vertex) */
@@ -63,6 +65,13 @@ namespace Chaos::Softs
 		 *  Stiffness and Damping weights are expected have a value in [0.f, 1.f]. They will be set to constant 0.f if not supplied.
 		*/
 		CHAOS_API void Initialize(const TConstArrayView<FIntVector2>& EndPoints, const TConstArrayView<float>& InSpringLength,
+			const TConstArrayView<float>& InExtensionStiffnesWeight = TConstArrayView<float>(), const TConstArrayView<float>& InCompressionStiffnessWeight = TConstArrayView<float>(), const TConstArrayView<float>& InDampingWeight = TConstArrayView<float>(), const FString& InConstraintName = TEXT(""));
+
+		/** Append from a list of Vertex-Vertex springs. Only succeeds if existing EndPointNumIndices == (1,1)
+		 *  Spring Indices are VerticesGroup indices.
+		 *  Stiffness and Damping weights are expected have a value in [0.f, 1.f]. They will be set to constant 0.f if not supplied.
+		*/
+		CHAOS_API void Append(const TConstArrayView<FIntVector2>& EndPoints, const TConstArrayView<float>& InSpringLength,
 			const TConstArrayView<float>& InExtensionStiffnesWeight = TConstArrayView<float>(), const TConstArrayView<float>& InCompressionStiffnessWeight = TConstArrayView<float>(), const TConstArrayView<float>& InDampingWeight = TConstArrayView<float>());
 
 		/** Initialize from a list of springs. This is generic for any number of end point indices. If you have vertex-vertex constraints, use the Initialize method above which skips setting weights.
@@ -73,6 +82,18 @@ namespace Chaos::Softs
 		 *  Stiffness and Damping weights are expected have a value in [0.f, 1.f]. They will be set to constant 0.f if not supplied.
 		*/
 		CHAOS_API void Initialize(const FUintVector2& EndPointNumIndices, const TConstArrayView<TArray<int32>>& InSourceIndices, const TConstArrayView<TArray<float>>& InSourceWeights,
+			const TConstArrayView<TArray<int32>>& InTargetIndices, const TConstArrayView<TArray<float>>& InTargetWeights,
+			const TConstArrayView<float>& InSpringLength,
+			const TConstArrayView<float>& InExtensionStiffnesWeight = TConstArrayView<float>(), const TConstArrayView<float>& InCompressionStiffnessWeight = TConstArrayView<float>(), const TConstArrayView<float>& InDampingWeight = TConstArrayView<float>(), const FString& InConstraintName = TEXT(""));
+
+		/** Append from a list of springs. This is generic for any number of end point indices. If you have vertex-vertex constraints, use the Append method above which skips setting weights.
+		 *  Source Indices and Weights arrays are expected to be exactly EndPointNumIndices[0] long.
+		 *  Target Indices and Weights arrays are expected to be exactly EndPointNumIndices[1] long.
+		 *  Indices and Weights will be truncated or filled in with zeros if they aren't the expected length.
+		 *  Spring Indices are VerticesGroup indices.
+		 *  Stiffness and Damping weights are expected have a value in [0.f, 1.f]. They will be set to constant 0.f if not supplied.
+		*/
+		CHAOS_API void Append(const TConstArrayView<TArray<int32>>& InSourceIndices, const TConstArrayView<TArray<float>>& InSourceWeights,
 			const TConstArrayView<TArray<int32>>& InTargetIndices, const TConstArrayView<TArray<float>>& InTargetWeights,
 			const TConstArrayView<float>& InSpringLength,
 			const TConstArrayView<float>& InExtensionStiffnesWeight = TConstArrayView<float>(), const TConstArrayView<float>& InCompressionStiffnessWeight = TConstArrayView<float>(), const TConstArrayView<float>& InDampingWeight = TConstArrayView<float>());
@@ -83,6 +104,16 @@ namespace Chaos::Softs
 			return ConstraintEndPointNumIndicesAttribute[ConstraintIndex];
 		}
 
+		const FString& GetConstraintName() const
+		{
+			return ConstraintNameAttribute[ConstraintIndex];
+		}
+
+		void SetConstraintName(const FString& InName)
+		{
+			return ConstraintNameAttribute.ModifyAt(ConstraintIndex, InName);
+		}
+
 		//~ Group PositionTargets
 		CHAOS_API TArrayView<float> GetSpringLength();
 		CHAOS_API TConstArrayView<float> GetSpringLengthConst() const;
diff --git a/Engine/Source/Runtime/Experimental/Chaos/Public/Chaos/GraphColoring.h b/Engine/Source/Runtime/Experimental/Chaos/Public/Chaos/GraphColoring.h
--- a/Engine/Source/Runtime/Experimental/Chaos/Public/Chaos/GraphColoring.h
+++ b/Engine/Source/Runtime/Experimental/Chaos/Public/Chaos/GraphColoring.h
@@ -14,93 +14,28 @@ namespace Chaos
 class FGraphColoring
 {
 	typedef TArray<int32, TInlineAllocator<8>> FColorSet;
-	
-	struct FGraphNode
-	{
-		FGraphNode()
-			: NextColor(0)
-		{
-		}
-
-		TArray<int32, TInlineAllocator<8>> Edges;
-		int32 NextColor;
-		FColorSet UsedColors;
-	};
-
-	struct FGraphEdge
-	{
-		FGraphEdge()
-			: FirstNode(INDEX_NONE)
-			, SecondNode(INDEX_NONE)
-			, Color(INDEX_NONE)
-		{
-		}
-
-		int32 FirstNode;
-		int32 SecondNode;
-		int32 Color;
-	};
-
-	struct FGraph3dEdge : FGraphEdge
-	{
-		FGraph3dEdge()
-			: ThirdNode(INDEX_NONE)
-		{
-		}
-
-		int32 ThirdNode;
-	};
-
-	struct FGraphTetEdge : FGraph3dEdge
-	{
-		FGraphTetEdge()
-			: FourthNode(INDEX_NONE)
-		{
-		}
-
-		int32 FourthNode;
-	};
 
   public:
-	template<typename DynamicParticlesType>
-	static TArray<TArray<int32>> ComputeGraphColoringParticlesOrRange(const TArray<TVector<int32, 2>>& Graph, const DynamicParticlesType& InParticles, const int32 GraphParticlesStart, const int32 GraphParticlesEnd);
 
-	template<typename T>
-	inline static TArray<TArray<int32>> ComputeGraphColoring(const TArray<TVector<int32, 2>>& Graph, const TDynamicParticles<T, 3>& InParticles, const int32 GraphParticlesStart, const int32 GraphParticlesEnd)
-	{
-		return ComputeGraphColoringParticlesOrRange(Graph, InParticles, GraphParticlesStart, GraphParticlesEnd);
-	}
-	template<typename T>
-	inline static TArray<TArray<int32>> ComputeGraphColoring(const TArray<TVector<int32, 2>>& Graph, const TDynamicParticles<T, 3>& InParticles)
-	{
-		return ComputeGraphColoringParticlesOrRange(Graph, InParticles, 0, InParticles.Size());
-	}
-	template<typename DynamicParticlesType>
-	static TArray<TArray<int32>> ComputeGraphColoringParticlesOrRange(const TArray<TVector<int32, 3>>& Graph, const DynamicParticlesType& InParticles, const int32 GraphParticlesStart, const int32 GraphParticlesEnd);
-	template<typename T>
-	inline static TArray<TArray<int32>> ComputeGraphColoring(const TArray<TVector<int32, 3>>& Graph, const TDynamicParticles<T, 3>& InParticles, const int32 GraphParticlesStart, const int32 GraphParticlesEnd)
+	template<typename DynamicParticlesType, int32 N, bool bAllDynamic = false>
+	static CHAOS_API TArray<TArray<int32>> ComputeGraphColoringParticlesOrRange(const TArray<TVector<int32,N>>& Graph, const DynamicParticlesType& InParticles, const int32 GraphParticlesStart, const int32 GraphParticlesEnd);
+
+	template<typename T, int32 N>
+	inline static TArray<TArray<int32>> ComputeGraphColoring(const TArray<TVector<int32, N>>& Graph, const TDynamicParticles<T, 3>& InParticles, const int32 GraphParticlesStart, const int32 GraphParticlesEnd)
 	{
 		return ComputeGraphColoringParticlesOrRange(Graph, InParticles, GraphParticlesStart, GraphParticlesEnd);
 	}
-	template<typename T>
-	inline static TArray<TArray<int32>> ComputeGraphColoring(const TArray<TVector<int32, 3>>& Graph, const TDynamicParticles<T, 3>& InParticles)
+	template<typename T, int32 N>
+	inline static TArray<TArray<int32>> ComputeGraphColoring(const TArray<TVector<int32, N>>& Graph, const TDynamicParticles<T, 3>& InParticles)
 	{
 		return ComputeGraphColoringParticlesOrRange(Graph, InParticles, 0, InParticles.Size());
 	}
+
 	template<typename DynamicParticlesType>
-	static TArray<TArray<int32>> ComputeGraphColoringParticlesOrRange(const TArray<TVector<int32, 4>>& Graph, const DynamicParticlesType& InParticles, const int32 GraphParticlesStart, const int32 GraphParticlesEnd);
-	template<typename T>
-	inline static TArray<TArray<int32>> ComputeGraphColoring(const TArray<TVector<int32, 4>>& Graph, const TDynamicParticles<T, 3>& InParticles, const int32 GraphParticlesStart, const int32 GraphParticlesEnd)
-	{
-		return ComputeGraphColoringParticlesOrRange(Graph, InParticles, GraphParticlesStart, GraphParticlesEnd);
-	}
-	template<typename T>
-	inline static TArray<TArray<int32>> ComputeGraphColoring(const TArray<TVector<int32, 4>>& Graph, const TDynamicParticles<T, 3>& InParticles)
+	inline static TArray<TArray<int32>> ComputeGraphColoringAllDynamicParticlesOrRange(const TArray<TVec4<int32>>& Graph, const DynamicParticlesType& InParticles, const int32 GraphParticlesStart, const int32 GraphParticlesEnd)
 	{
-		return ComputeGraphColoringParticlesOrRange(Graph, InParticles, 0, InParticles.Size());
+		return ComputeGraphColoringParticlesOrRange<DynamicParticlesType, 4, true>(Graph, InParticles, GraphParticlesStart, GraphParticlesEnd);
 	}
-	template<typename DynamicParticlesType>
-	static TArray<TArray<int32>> ComputeGraphColoringAllDynamicParticlesOrRange(const TArray<TVec4<int32>>& Graph, const DynamicParticlesType& InParticles, const int32 GraphParticlesStart, const int32 GraphParticlesEnd);
 	template<typename T>
 	inline static TArray<TArray<int32>> ComputeGraphColoringAllDynamic(const TArray<TVec4<int32>>& Graph, const Chaos::TDynamicParticles<T, 3>& InParticles, const int32 GraphParticlesStart, const int32 GraphParticlesEnd)
 	{
@@ -114,8 +49,16 @@ class FGraphColoring
 };
 
 template<typename T> 
-CHAOS_API void ComputeGridBasedGraphSubColoringPointer(const TArray<TArray<int32>>& ElementsPerColor, const TMPMGrid<T>& Grid, const int32 GridSize, TArray<TArray<int32>>*& PreviousColoring, const TArray<TArray<int32>>& ConstraintsNodesSet, TArray<TArray<TArray<int32>>>& ElementsPerSubColors);
+UE_DEPRECATED(5.6, "Use UniquePtr version instead.")
+void ComputeGridBasedGraphSubColoringPointer(const TArray<TArray<int32>>& ElementsPerColor, const TMPMGrid<T>& Grid, const int32 GridSize, TArray<TArray<int32>>*& PreviousColoring, const TArray<TArray<int32>>& ConstraintsNodesSet, TArray<TArray<TArray<int32>>>& ElementsPerSubColors)
+{
+	TUniquePtr<TArray<TArray<int32>>> PreviousColoringUniquePtr(PreviousColoring);
+	ComputeGridBasedGraphSubColoringPointer(ElementsPerColor, Grid, GridSize, PreviousColoringUniquePtr, ConstraintsNodesSet, ElementsPerSubColors);
+	PreviousColoring = PreviousColoringUniquePtr.Release();
+}
 
+template<typename T>
+CHAOS_API void ComputeGridBasedGraphSubColoringPointer(const TArray<TArray<int32>>& ElementsPerColor, const TMPMGrid<T>& Grid, const int32 GridSize, TUniquePtr<TArray<TArray<int32>>>& PreviousColoring, const TArray<TArray<int32>>& ConstraintsNodesSet, TArray<TArray<TArray<int32>>>& ElementsPerSubColors);
 
 template<typename T>
 CHAOS_API void ComputeWeakConstraintsColoring(const TArray<TArray<int32>>& Indices, const TArray<TArray<int32>>& SecondIndices, const Chaos::TDynamicParticles<T, 3>& InParticles, TArray<TArray<int32>>& ConstraintsPerColor);
diff --git a/Engine/Source/Runtime/Experimental/Chaos/Public/Chaos/PBDFlatWeightMap.h b/Engine/Source/Runtime/Experimental/Chaos/Public/Chaos/PBDFlatWeightMap.h
--- a/Engine/Source/Runtime/Experimental/Chaos/Public/Chaos/PBDFlatWeightMap.h
+++ b/Engine/Source/Runtime/Experimental/Chaos/Public/Chaos/PBDFlatWeightMap.h
@@ -104,7 +104,12 @@ public:
 	inline explicit FPBDFlatWeightMap(
 		const FSolverVec2& InWeightedValue,
 		const TConstArrayView<FRealSingle>& Multipliers = TConstArrayView<FRealSingle>(),
-		int32 ParticleCount = 0); // A value of 0 also disables the map
+		int32 ParticleCount = 0, // A value of 0 also disables the map
+		const FSolverVec2& MultiplierClampRange = FSolverVec2(0., 1.));
+
+	inline FPBDFlatWeightMap(
+		const FSolverVec2& InWeightedValue,
+		TArray<FRealSingle>&& Multipliers);
 
 	/**
 	 * Weightmap constraint constructor. 
@@ -162,10 +167,14 @@ public:
 	FSolverReal GetValue(int32 Index) const { return MapView.GetValue(Index); }
 
 	/** Reorder Indices based on Constraint reordering. */
-	inline void ReorderIndices(const TArray<int32>& OrigToReorderedIndices);
+	inline void ReorderIndicesAndShrink(const TArray<int32>& OrigToReorderedIndices, int32 NewNum);
+	void ReorderIndices(const TArray<int32>& OrigToReorderedIndices)
+	{
+		return ReorderIndicesAndShrink(OrigToReorderedIndices, MapValues.Num());
+	}
 
 private:
-	static TArray<FSolverReal> CalculateMapValues(const TConstArrayView<FRealSingle>& Multipliers, int32 ParticleCount)
+	static TArray<FSolverReal> CalculateMapValues(const TConstArrayView<FRealSingle>& Multipliers, int32 ParticleCount, const FSolverVec2& MultiplierClampRange)
 	{
 		TArray<FSolverReal> Result;
 		if (Multipliers.Num() == ParticleCount && ParticleCount > 0)
@@ -174,7 +183,7 @@ private:
 			Result.SetNumUninitialized(ParticleCount);
 			for (int32 Index = 0; Index < ParticleCount; ++Index)
 			{
-				Result[Index] = FMath::Clamp(Multipliers[Index], (FSolverReal)0.f, (FSolverReal)1.f);
+				Result[Index] = FMath::Clamp(Multipliers[Index], MultiplierClampRange[0], MultiplierClampRange[1]);
 			}
 		}
 		return Result;
@@ -218,12 +227,20 @@ private:
 inline FPBDFlatWeightMap::FPBDFlatWeightMap(
 	const FSolverVec2& InWeightedValue,
 	const TConstArrayView<FRealSingle>& Multipliers,
-	int32 ParticleCount)
-	: MapValues(CalculateMapValues(Multipliers, ParticleCount))
+	int32 ParticleCount,
+	const FSolverVec2& MultiplierClampRange)
+	: MapValues(CalculateMapValues(Multipliers, ParticleCount, MultiplierClampRange))
 	, MapView(InWeightedValue, TConstArrayView<FSolverReal>(MapValues), ParticleCount)
 {
 }
 
+inline FPBDFlatWeightMap::FPBDFlatWeightMap(
+	const FSolverVec2& InWeightedValue,
+	TArray<FRealSingle>&& Multipliers)
+	: MapValues(MoveTemp(Multipliers))
+	, MapView(InWeightedValue, MapValues, MapValues.Num())
+{}
+
 template<int32 Valence>
 inline FPBDFlatWeightMap::FPBDFlatWeightMap(
 	const FSolverVec2& InWeightedValue,
@@ -237,16 +254,19 @@ inline FPBDFlatWeightMap::FPBDFlatWeightMap(
 {
 }
 
-inline void FPBDFlatWeightMap::ReorderIndices(const TArray<int32>& OrigToReorderedConstraintIndices)
+inline void FPBDFlatWeightMap::ReorderIndicesAndShrink(const TArray<int32>& OrigToReorderedConstraintIndices, int32 NewNum)
 {
 	if (MapValues.Num() == OrigToReorderedConstraintIndices.Num())
 	{
 		TArray<FSolverReal> ReorderedValues;
-		ReorderedValues.SetNumUninitialized(MapValues.Num());
+		ReorderedValues.SetNumUninitialized(NewNum);
 		for (int32 OrigConstraintIndex = 0; OrigConstraintIndex < OrigToReorderedConstraintIndices.Num(); ++OrigConstraintIndex)
 		{
 			const int32 ReorderedConstraintIndex = OrigToReorderedConstraintIndices[OrigConstraintIndex];
-			ReorderedValues[ReorderedConstraintIndex] = MapValues[OrigConstraintIndex];
+			if (ReorderedConstraintIndex != INDEX_NONE)
+			{
+				ReorderedValues[ReorderedConstraintIndex] = MapValues[OrigConstraintIndex];
+			}
 		}
 		MapValues = MoveTemp(ReorderedValues);
 		MapView.SetMapValues(TConstArrayView<FSolverReal>(MapValues));
diff --git a/Engine/Source/Runtime/Experimental/Chaos/Public/Chaos/PBDRepulsionConstraints.h b/Engine/Source/Runtime/Experimental/Chaos/Public/Chaos/PBDRepulsionConstraints.h
new file mode 100644
--- /dev/null
+++ b/Engine/Source/Runtime/Experimental/Chaos/Public/Chaos/PBDRepulsionConstraints.h
@@ -0,0 +1,74 @@
+// Copyright Epic Games, Inc. All Rights Reserved.
+#pragma once
+
+#include "Chaos/CollectionEmbeddedSpringConstraintFacade.h"
+#include "Chaos/CollectionPropertyFacade.h"
+#include "Chaos/Core.h"
+#include "Chaos/PBDFlatWeightMap.h"
+#include "Chaos/PBDSoftsEvolutionFwd.h"
+#include "Chaos/PBDSoftsSolverParticles.h"
+#include "Chaos/SoftsSolverParticlesRange.h"
+#include "Containers/Array.h"
+
+namespace Chaos::Softs
+{
+
+	class FPBDRepulsionConstraints final
+	{
+	public:
+		static constexpr FSolverReal MinStiffness = (FSolverReal)0;
+		static constexpr FSolverReal MaxStiffness = (FSolverReal)1.f;
+		static constexpr FSolverReal DefaultStiffness = (FSolverReal)0.5f;
+		static constexpr FSolverReal MinFriction = (FSolverReal)0;
+		static constexpr FSolverReal MaxFriction = (FSolverReal)1.f;
+		static constexpr FSolverReal DefaultFriction = (FSolverReal)0.f;
+
+		static bool IsEnabled(const FCollectionPropertyConstFacade& PropertyCollection)
+		{
+			return IsVertexFaceRepulsionStiffnessEnabled(PropertyCollection, false);
+		}
+
+		CHAOS_API FPBDRepulsionConstraints(
+			const FSolverParticlesRange& Particles,
+			const FCollectionPropertyConstFacade& PropertyCollection,
+			const FEmbeddedSpringConstraintFacade& SpringConstraintFacade);
+
+		~FPBDRepulsionConstraints() = default;
+
+		void ApplyProperties(const FSolverReal /*Dt*/, const int32 /*NumIterations*/)
+		{
+		}
+
+		CHAOS_API void SetProperties(const FCollectionPropertyConstFacade& PropertyCollection);
+		CHAOS_API void Apply(FSolverParticlesRange& Particles, const FSolverReal Dt) const;
+		
+		const TArray<TVector<int32, 4>>& GetConstraints() const { return Constraints; }
+		const TArray<FSolverReal>& GetRestLengths() const { return RestLengths; }
+
+#if CHAOS_DEBUG_DRAW
+		const TArray<bool>& GetConstraintIsActive() const { return ConstraintIsActive; }
+#endif
+
+	private:
+
+		void InitColor(const FSolverParticlesRange& Particles);
+
+		TArray<TVector<int32, 4>> Constraints; // Index 0 is Source Vertex.
+		TArray<TVector<FSolverReal, 3>> TargetWeights;
+		TArray<FSolverReal> RestLengths;
+
+		TArray<int32> ConstraintsPerColorStartIndex; // Constraints are ordered so each batch is contiguous. This is ColorNum + 1 length so it can be used as start and end.
+
+		FSolverReal Stiffness;
+		FSolverReal FrictionCoefficient;
+
+#if CHAOS_DEBUG_DRAW
+		// Used for debug draw to show which constraints were active last apply
+		mutable TArray<bool> ConstraintIsActive;
+#endif
+
+		UE_CHAOS_DECLARE_PROPERTYCOLLECTION_NAME(VertexFaceRepulsionStiffness, float);
+		UE_CHAOS_DECLARE_PROPERTYCOLLECTION_NAME(VertexFaceRepulsionFrictionCoefficient, float);
+	};
+}  // End namespace Chaos::Softs
+
diff --git a/Engine/Source/Runtime/Experimental/Chaos/Public/Chaos/PBDSoftsEvolutionFwd.h b/Engine/Source/Runtime/Experimental/Chaos/Public/Chaos/PBDSoftsEvolutionFwd.h
--- a/Engine/Source/Runtime/Experimental/Chaos/Public/Chaos/PBDSoftsEvolutionFwd.h
+++ b/Engine/Source/Runtime/Experimental/Chaos/Public/Chaos/PBDSoftsEvolutionFwd.h
@@ -79,6 +79,10 @@ namespace Chaos
 		class FPBDSelfCollisionSphereConstraints;
 		class FPBDSoftBodyCollisionConstraint;
 		class FMultiResConstraints;
+		class FXPBDVertexConstraints;
+		class FXPBDVertexFaceConstraints;
+		class FXPBDFaceConstraints;
+		class FPBDRepulsionConstraints;
 
 		// Softs solver forces
 		class FVelocityAndPressureField;
diff --git a/Engine/Source/Runtime/Experimental/Chaos/Public/Chaos/SoftsSpring.h b/Engine/Source/Runtime/Experimental/Chaos/Public/Chaos/SoftsSpring.h
--- a/Engine/Source/Runtime/Experimental/Chaos/Public/Chaos/SoftsSpring.h
+++ b/Engine/Source/Runtime/Experimental/Chaos/Public/Chaos/SoftsSpring.h
@@ -208,6 +208,68 @@ FSolverVec3 GetXPBDAxialSpringDampingDelta(const SolverParticlesOrRange& Particl
 	return Delta;
 }
 
+// Spring without damping
+template<typename SolverParticlesOrRange, int32 N>
+FSolverVec3 GetXPBDEmbeddedSpringDelta(const SolverParticlesOrRange& Particles, const FSolverReal Dt,
+	const TVector<int32, N>& Constraint, const TVector<FSolverReal, N>& Weights, const FSolverReal RestLength, FSolverReal& Lambda,
+	const FSolverReal ExtensionStiffnessValue, const FSolverReal CompressionStiffnessValue)
+{
+	FSolverReal CombinedInvMass = 0.f;
+	FSolverVec3 Direction(0.f);
+	for (int32 NIdx = 0; NIdx < N; ++NIdx)
+	{
+		const int32 Index = Constraint[NIdx];
+		CombinedInvMass += Particles.InvM(Index) * FMath::Abs(Weights[NIdx]);
+		Direction += Particles.P(Index) * Weights[NIdx];
+	}
+
+	const FSolverReal Distance = Direction.SafeNormalize();
+	const FSolverReal Offset = Distance - RestLength;
+
+	// Do we need a smoothstep?
+	const FSolverReal StiffnessValue = Offset >= 0 ? ExtensionStiffnessValue : CompressionStiffnessValue;
+	const FSolverReal AlphaInv = StiffnessValue * Dt * Dt;
+
+	const FSolverReal DLambda = (-AlphaInv * Offset - Lambda) / (AlphaInv * CombinedInvMass + (FSolverReal)1.);
+	const FSolverVec3 Delta = DLambda * Direction;
+	Lambda += DLambda;
+
+	return Delta;
+}
+
+// Spring damping constraint (separate from spring stretching)
+template<typename SolverParticlesOrRange, int32 N>
+FSolverVec3 GetXPBDEmbeddedSpringDampingDelta(const SolverParticlesOrRange& Particles, const FSolverReal Dt,
+	const TVector<int32, N>& Constraint, const TVector<FSolverReal, N>& Weights, const FSolverReal RestLength, FSolverReal& Lambda,
+	const FSolverReal ExtensionStiffnessValue, const FSolverReal CompressionStiffnessValue, const FSolverReal DampingRatioValue)
+{
+	FSolverReal CombinedInvMass = 0.f;
+	FSolverVec3 DeltaPDirection(0.f);
+	FSolverVec3 DeltaX(0.f);
+	for (int32 NIdx = 0; NIdx < N; ++NIdx)
+	{
+		const int32 Index = Constraint[NIdx];
+		CombinedInvMass += Particles.InvM(Index) * FMath::Abs(Weights[NIdx]);
+		DeltaPDirection += Particles.P(Index) * Weights[NIdx];
+		DeltaX += Particles.X(Index) * Weights[NIdx];
+	}
+	check(CombinedInvMass > (FSolverReal)0);
+
+	const FSolverVec3 RelativeVelocityTimesDt = DeltaPDirection - DeltaX;
+
+	const FSolverReal Distance = DeltaPDirection.SafeNormalize();
+	const FSolverReal Offset = Distance - RestLength;
+	const FSolverReal StiffnessValue = Offset >= 0 ? ExtensionStiffnessValue : CompressionStiffnessValue;
+	const FSolverReal Damping = DampingRatioValue * 2.f * FMath::Sqrt(StiffnessValue / CombinedInvMass) * (RestLength > UE_SMALL_NUMBER ? (FSolverReal)1. / RestLength : (FSolverReal)1.);
+
+	const FSolverReal BetaDt = Damping * Dt;
+	const FSolverReal DLambda = (-BetaDt * FSolverVec3::DotProduct(DeltaPDirection, RelativeVelocityTimesDt) - Lambda) / (BetaDt * CombinedInvMass + (FSolverReal)1.);
+	const FSolverVec3 Delta = DLambda * DeltaPDirection;
+	Lambda += DLambda;
+
+	return Delta;
+}
+
 inline void UpdateSpringLinearSystem(const FSolverParticlesRange& Particles, const FSolverReal Dt,
 	const TVec2<int32>& Constraint, const FSolverReal RestLength,
 	const FSolverReal StiffnessValue, const FSolverReal MinStiffness, const FSolverReal DampingRatioValue,
diff --git a/Engine/Source/Runtime/Experimental/Chaos/Public/Chaos/TriangleMesh.h b/Engine/Source/Runtime/Experimental/Chaos/Public/Chaos/TriangleMesh.h
--- a/Engine/Source/Runtime/Experimental/Chaos/Public/Chaos/TriangleMesh.h
+++ b/Engine/Source/Runtime/Experimental/Chaos/Public/Chaos/TriangleMesh.h
@@ -36,7 +36,7 @@ namespace Chaos
 		 *	will change the order of \c MElements.
 		 */
 		CHAOS_API void Init(TArray<TVec3<int32>>&& Elements, const int32 StartIdx = 0, const int32 EndIdx = -1, const bool CullDegenerateElements=true);
-		CHAOS_API void Init(const TArray<TVec3<int32>>& Elements, const int32 StartIdx = 0, const int32 EndIdx = -1, const bool CullDegenerateElements=true);
+		CHAOS_API void Init(const TConstArrayView<TVec3<int32>>& Elements, const int32 StartIdx = 0, const int32 EndIdx = -1, const bool CullDegenerateElements = true);
 
 		CHAOS_API void ResetAuxiliaryStructures();
 
diff --git a/Engine/Source/Runtime/Experimental/Chaos/Public/Chaos/XPBDGridBasedCorotatedConstraints.h b/Engine/Source/Runtime/Experimental/Chaos/Public/Chaos/XPBDGridBasedCorotatedConstraints.h
--- a/Engine/Source/Runtime/Experimental/Chaos/Public/Chaos/XPBDGridBasedCorotatedConstraints.h
+++ b/Engine/Source/Runtime/Experimental/Chaos/Public/Chaos/XPBDGridBasedCorotatedConstraints.h
@@ -44,7 +44,7 @@ namespace Chaos::Softs
 			const T& EMesh = (T)1000.0,
 			const T& NuMesh = (T).3
 		)
-			: Base(InParticles, InMesh, GridDx, EMesh, NuMesh), MaxDt(MaxDtIn), MinDt(MinDtIn), CFLCoeff(CFLCoeffIn), MPMGrid(GridDx), MPMTransfer(MPMGrid, InParticles.Size()), PreviousColoring(nullptr)
+			: Base(InParticles, InMesh, GridDx, EMesh, NuMesh), MaxDt(MaxDtIn), MinDt(MinDtIn), CFLCoeff(CFLCoeffIn), MPMGrid(GridDx), MPMTransfer(MPMGrid, InParticles.Size()), PreviousColoringPtr(nullptr)
 		{
 			//TMPMGrid<T> MPMGrid(GridN);
 			MPMGrid.UpdateGridFromPositions(InParticles);
@@ -136,7 +136,7 @@ namespace Chaos::Softs
 			}
 			{
 				TRACE_CPUPROFILER_EVENT_SCOPE(STAT_ChaosXPBDGridBasedCorotatedSubcoloring);
-				ComputeGridBasedGraphSubColoringPointer(ElementsPerColor, MPMTransfer.Grid, MPMTransfer.Grid.Size(), PreviousColoring, MPMTransfer.ElementGridNodesSet, ElementsPerSubColors);
+				ComputeGridBasedGraphSubColoringPointer(ElementsPerColor, MPMTransfer.Grid, MPMTransfer.Grid.Size(), PreviousColoringPtr, MPMTransfer.ElementGridNodesSet, ElementsPerSubColors);
 			}
 			{
 				TRACE_CPUPROFILER_EVENT_SCOPE(STAT_ChaosXPBDGridBasedCorotatedComputeGridPositions);
@@ -201,7 +201,7 @@ namespace Chaos::Softs
 		}
 
 		const TArray<TArray<int32>> GetElementsPerColor () const { return ElementsPerColor; }
-		const TArray<TArray<int32>> GetPreviousColoring () const { return *PreviousColoring; }
+		const TArray<TArray<int32>> GetPreviousColoring () const { return *PreviousColoringPtr; }
 		const TArray<TArray<TArray<int32>>> GetElementsPerSubColors() const { return ElementsPerSubColors; }
 		const TMPMTransfer<T> GetMPMTransfer() const { return MPMTransfer; }
 
@@ -488,7 +488,9 @@ namespace Chaos::Softs
 		TMPMGrid<T> MPMGrid;
 		TMPMTransfer<T> MPMTransfer;
 		TArray<TArray<int32>> ElementsPerColor;
-		TArray<TArray<int32>>* PreviousColoring;
+		UE_DEPRECATED(5.6, "Use PreviousColoringPtr instead")
+		TArray<TArray<int32>>* PreviousColoring = nullptr;
+		TUniquePtr<TArray<TArray<int32>>> PreviousColoringPtr;
 		TArray<TArray<TArray<int32>>> ElementsPerSubColors;
 		TArray<T> GridData;
 		TArray<TVector<T, 3>> GridPositions;
diff --git a/Engine/Source/Runtime/Experimental/Dataflow/Core/Private/Dataflow/DataflowNode.cpp b/Engine/Source/Runtime/Experimental/Dataflow/Core/Private/Dataflow/DataflowNode.cpp
--- a/Engine/Source/Runtime/Experimental/Dataflow/Core/Private/Dataflow/DataflowNode.cpp
+++ b/Engine/Source/Runtime/Experimental/Dataflow/Core/Private/Dataflow/DataflowNode.cpp
@@ -712,7 +712,14 @@ const FProperty* FDataflowNode::FindProperty(const UStruct* Struct, const FName&
 	{
 		TArray<const FProperty*> PropertyChain;
 		PropertyIt.GetPropertyChain(PropertyChain);
-		if (GetPropertyFullNameString(PropertyChain, INDEX_NONE) == PropertyFullNameStringIndexNone)
+		const FString FullName = GetPropertyFullNameString(PropertyChain, INDEX_NONE);
+		
+		if (PropertyFullNameStringIndexNone.Left(FullName.Len()) != FullName)
+		{
+			PropertyIt.SkipRecursiveProperty();
+			continue;
+		}
+		if(PropertyFullNameStringIndexNone.Len() == FullName.Len())
 		{
 			Property = PropertyIt.Key();
 			if (OutPropertyChain)
diff --git a/Engine/Source/Runtime/GeometryCore/Public/Spatial/SparseDynamicPointOctree3.h b/Engine/Source/Runtime/GeometryCore/Public/Spatial/SparseDynamicPointOctree3.h
--- a/Engine/Source/Runtime/GeometryCore/Public/Spatial/SparseDynamicPointOctree3.h
+++ b/Engine/Source/Runtime/GeometryCore/Public/Spatial/SparseDynamicPointOctree3.h
@@ -221,7 +221,7 @@ public:
 	 * @param TempBuffer Optional temporary buffer to store a queue of cells (to avoid memory allocations)
 	 * @return The ID of the closest point, or INDEX_NONE if no point was found
 	 */
-	int32 FindClosestPoint(FVector3d QueryPt, double DistanceThreshold,
+	GEOMETRYCORE_API int32 FindClosestPoint(FVector3d QueryPt, double DistanceThreshold,
 		TFunctionRef<bool(int32)> PredicateFunc, TFunctionRef<double(int32)> DistSqFunc,
 		TArray<const FSparsePointOctreeCell*>* TempBuffer = nullptr) const;
 
