VulkanRHI: Add back ShaderResource flag to fix Vulkan validation error at SRV creation.

#rnx

diff --git a/Engine/Source/Runtime/Experimental/GeometryCollectionEngine/Private/GeometryCollection/GeometryCollectionSceneProxy.h b/Engine/Source/Runtime/Experimental/GeometryCollectionEngine/Private/GeometryCollection/GeometryCollectionSceneProxy.h
--- a/Engine/Source/Runtime/Experimental/GeometryCollectionEngine/Private/GeometryCollection/GeometryCollectionSceneProxy.h
+++ b/Engine/Source/Runtime/Experimental/GeometryCollectionEngine/Private/GeometryCollection/GeometryCollectionSceneProxy.h
@@ -38,7 +38,7 @@ public:
 		// splitting out the translation.  This is to simplify transferring data at runtime as a memcopy
 		const FRHIBufferCreateDesc CreateDesc =
 			FRHIBufferCreateDesc::CreateVertex<FVector4f>(TEXT("FGeometryCollectionTransformBuffer"), NumTransforms * 4)
-			.AddUsage(EBufferUsageFlags::Dynamic)
+			.AddUsage(EBufferUsageFlags::Dynamic | EBufferUsageFlags::ShaderResource)
 			.DetermineInitialState();
 
 		VertexBufferRHI = RHICmdList.CreateBuffer(CreateDesc);
