Dataflow : Add Guid to dataflow subgraph
- also add a way to know if the subgrapoh is loaded

diff --git a/Engine/Source/Runtime/Experimental/Dataflow/Engine/Private/Dataflow/DataflowSubGraph.cpp b/Engine/Source/Runtime/Experimental/Dataflow/Engine/Private/Dataflow/DataflowSubGraph.cpp
--- a/Engine/Source/Runtime/Experimental/Dataflow/Engine/Private/Dataflow/DataflowSubGraph.cpp
+++ b/Engine/Source/Runtime/Experimental/Dataflow/Engine/Private/Dataflow/DataflowSubGraph.cpp
@@ -30,6 +30,8 @@ namespace UE::Dataflow::DataflowSubGraph::Private
 
 UDataflowSubGraph::UDataflowSubGraph(class FObjectInitializer const& ObjectInitializer)
 	: UEdGraph(ObjectInitializer)
+	, SubGraphGuid(FGuid::NewGuid())
+	, bIsLoaded(false)
 {}
 
 FDataflowSubGraphInputNode* UDataflowSubGraph::GetInputNode() const
@@ -48,6 +50,7 @@ void UDataflowSubGraph::PostLoad()
 {
 	Super::PostLoad();
 	
+	bIsLoaded = true;
 	if (OnSubGraphLoaded.IsBound())
 	{
 		OnSubGraphLoaded.Broadcast(*this);
diff --git a/Engine/Source/Runtime/Experimental/Dataflow/Engine/Public/Dataflow/DataflowSubGraph.h b/Engine/Source/Runtime/Experimental/Dataflow/Engine/Public/Dataflow/DataflowSubGraph.h
--- a/Engine/Source/Runtime/Experimental/Dataflow/Engine/Public/Dataflow/DataflowSubGraph.h
+++ b/Engine/Source/Runtime/Experimental/Dataflow/Engine/Public/Dataflow/DataflowSubGraph.h
@@ -25,8 +25,19 @@ public:
 	FDelegateHandle AddOnSubGraphLoadedHandler(const FOnDataflowSubGraphLoaded::FDelegate& InHandler);
 	void RemoveOnSubGraphLoadedHandler(FDelegateHandle Handle);
 
+	bool IsLoaded() const { return bIsLoaded; }
+	const FGuid& GetSubGraphGuid() const { return SubGraphGuid; }
+
 private:
 	virtual void PostLoad() override;
 
+	/**
+	* Guid that uniquely identify this SubGraph
+	* this allow to identify the SubGraph across renames for example 
+	*/
+	UPROPERTY()
+	FGuid SubGraphGuid; 
+
 	FOnDataflowSubGraphLoaded OnSubGraphLoaded;
+	bool bIsLoaded : 1 = false;
 };
\ No newline at end of file
