Dataflow : Add API to find Subgraphs by Guid

diff --git a/Engine/Source/Runtime/Experimental/Dataflow/Engine/Private/Dataflow/DataflowObject.cpp b/Engine/Source/Runtime/Experimental/Dataflow/Engine/Private/Dataflow/DataflowObject.cpp
--- a/Engine/Source/Runtime/Experimental/Dataflow/Engine/Private/Dataflow/DataflowObject.cpp
+++ b/Engine/Source/Runtime/Experimental/Dataflow/Engine/Private/Dataflow/DataflowObject.cpp
@@ -363,6 +363,30 @@ UDataflowSubGraph* UDataflow::FindSubGraphByName(FName Name)
 	return nullptr;
 }
 
+const UDataflowSubGraph* UDataflow::FindSubGraphByGuid(const FGuid& SubGraphGuid) const
+{
+	for (const UDataflowSubGraph* SubGraph : DataflowSubGraphs)
+	{
+		if (SubGraph->GetSubGraphGuid() == SubGraphGuid)
+		{
+			return SubGraph;
+		}
+	}
+	return nullptr;
+}
+
+UDataflowSubGraph* UDataflow::FindSubGraphByGuid(const FGuid& SubGraphGuid)
+{
+	for (UDataflowSubGraph* SubGraph : DataflowSubGraphs)
+	{
+		if (SubGraph->GetSubGraphGuid() == SubGraphGuid)
+		{
+			return SubGraph;
+		}
+	}
+	return nullptr;
+}
+
 void UDataflow::AddSubGraph(UDataflowSubGraph* SubGraph)
 {
 	if (ensure(SubGraph))
diff --git a/Engine/Source/Runtime/Experimental/Dataflow/Engine/Public/Dataflow/DataflowObject.h b/Engine/Source/Runtime/Experimental/Dataflow/Engine/Public/Dataflow/DataflowObject.h
--- a/Engine/Source/Runtime/Experimental/Dataflow/Engine/Public/Dataflow/DataflowObject.h
+++ b/Engine/Source/Runtime/Experimental/Dataflow/Engine/Public/Dataflow/DataflowObject.h
@@ -181,6 +181,10 @@ public:
 	DATAFLOWENGINE_API const UDataflowSubGraph* FindSubGraphByName(FName Name) const;
 	DATAFLOWENGINE_API UDataflowSubGraph* FindSubGraphByName(FName Name);
 
+	/** Find a SubGraph by its guid */
+	DATAFLOWENGINE_API const UDataflowSubGraph* FindSubGraphByGuid(const FGuid& SubGraphGuid) const;
+	DATAFLOWENGINE_API UDataflowSubGraph* FindSubGraphByGuid(const FGuid& SubGraphGuid);
+
 	/** Add a SubGraph to the asset */
 	DATAFLOWENGINE_API void AddSubGraph(UDataflowSubGraph* SubGraph);
 	DATAFLOWENGINE_API void RemoveSubGraph(UDataflowSubGraph* SubGraph);
