Dataflow : fix SubGraph input / output node pins not properly reconnecting properly when the type is changing
- we need to track the propertyid by connection name

diff --git a/Engine/Source/Runtime/Experimental/Dataflow/Engine/Private/Dataflow/DataflowSubGraphNodes.cpp b/Engine/Source/Runtime/Experimental/Dataflow/Engine/Private/Dataflow/DataflowSubGraphNodes.cpp
--- a/Engine/Source/Runtime/Experimental/Dataflow/Engine/Private/Dataflow/DataflowSubGraphNodes.cpp
+++ b/Engine/Source/Runtime/Experimental/Dataflow/Engine/Private/Dataflow/DataflowSubGraphNodes.cpp
@@ -156,7 +156,7 @@ void FDataflowDynamicConnections::Refresh()
 	{
 		if (DynamicConnection)
 		{
-			if (const FGuid* PropertyId = ConnectionIdToPropertyId.Find(DynamicConnection->GetGuid()))
+			if (const FGuid* PropertyId = ConnectionNameToPropertyId.Find(DynamicConnection->GetName()))
 			{
 				FRemoteConnections& RemoteConnections = RemoteConnectionsByPropertyId.Emplace(*PropertyId);
 				DynamicConnection->GetConnections(RemoteConnections);
@@ -167,7 +167,7 @@ void FDataflowDynamicConnections::Refresh()
 	// clear all the dynamic connections
 	ClearDynamicConnections();
 
-	ConnectionIdToPropertyId.Reset();
+	ConnectionNameToPropertyId.Reset();
 
 	// Go through all the properties of the property bag recreate the output and reconnect them if needed
 	const TConnectionReference<FDataflowAllTypes> NullReference(nullptr);
@@ -222,7 +222,7 @@ FDataflowConnection* FDataflowDynamicConnections::AddNewConnectionFromPropertyDe
 			}
 			NewConnection->LockType();
 
-			ConnectionIdToPropertyId.Add(NewConnection->GetGuid(), PropertyDesc.ID);
+			ConnectionNameToPropertyId.Add(NewConnection->GetName(), PropertyDesc.ID);
 			return NewConnection;
 		}
 	}
diff --git a/Engine/Source/Runtime/Experimental/Dataflow/Engine/Public/Dataflow/DataflowSubGraphNodes.h b/Engine/Source/Runtime/Experimental/Dataflow/Engine/Public/Dataflow/DataflowSubGraphNodes.h
--- a/Engine/Source/Runtime/Experimental/Dataflow/Engine/Public/Dataflow/DataflowSubGraphNodes.h
+++ b/Engine/Source/Runtime/Experimental/Dataflow/Engine/Public/Dataflow/DataflowSubGraphNodes.h
@@ -56,7 +56,7 @@ private:
 	UPROPERTY()
 	TArray<FDataflowAllTypes> DynamicProperties;
 
-	TMap<FGuid, FGuid> ConnectionIdToPropertyId;
+	TMap<FName, FGuid> ConnectionNameToPropertyId;
 
 	UE::Dataflow::FPin::EDirection PinDirection = UE::Dataflow::FPin::EDirection::NONE;
 	IOwnerInterface* OwnerInterface;
