Allow the LogCategory and Verbosity used by FMultiEndpointHttpClient to be customized when created.

#rnx

- Since we want multiple systems to use FMultiEndpointHttpClient we need to be able to redirect the log output to different categories.
- At the moment this only affects the output of FMultiEndpointHttpClient::Log which summarizes the request.
- By default no LogCategory name is set, in which case the logging is skipped.
- By default we will log to at 'Log' verbosity.
- The format in which FMultiEndpointHttpClient::Log has been slightly altered to better match the previous logging.

diff --git a/Engine/Source/Runtime/Experimental/IoStore/OnDemand/Private/OnDemandHttpClient.cpp b/Engine/Source/Runtime/Experimental/IoStore/OnDemand/Private/OnDemandHttpClient.cpp
--- a/Engine/Source/Runtime/Experimental/IoStore/OnDemand/Private/OnDemandHttpClient.cpp
+++ b/Engine/Source/Runtime/Experimental/IoStore/OnDemand/Private/OnDemandHttpClient.cpp
@@ -628,17 +628,29 @@ FMultiEndpointHttpClient::FConnection* FMultiEndpointHttpClient::FindConnection(
 	return nullptr;
 }
 
-void FMultiEndpointHttpClient::Log(const FRequest& Request)
+void FMultiEndpointHttpClient::Log(const FRequest& Request) const
 {
-	const uint64 DurationMs = (uint64)FPlatformTime::ToMilliseconds64(FPlatformTime::Cycles64() - Request.StartTime);
-
-	TStringBuilder<512> Sb;
-	Sb << TEXT("GET ")
-		<< Request.Connection.HostGroup.Host(Request.Host)
-		<< Request.RelativeUrl;
-	Sb.Appendf(TEXT(" (Status: %u, Attempt %d/%d, %llums, %.2lf KiB)"), Request.StatusCode, Request.RetryCount, Config.MaxRetryCount, DurationMs, double(Request.Body.GetSize()) / 1024.0);
+#if !NO_LOGGING
+	auto ShouldLogCategory = [](FName Name)->bool {return !Name.IsNone();}; // For LogCategoryName to resolve to a FName if UE_LOG_LAZY_CATEGORY_NAMES if defined and enabled
+	if (!ShouldLogCategory(Config.LogCategoryName))
+	{
+		return;
+	}
 
-	UE_LOG(LogIoStoreOnDemand, Log, TEXT("%s"), Sb.ToString());
+	const uint64 DurationMs = (uint64)FPlatformTime::ToMilliseconds64(FPlatformTime::Cycles64() - Request.StartTime);
+	const uint64 Size = Request.Body.GetSize() >> 10;
+
+	FMsg::Logf(__FILE__, __LINE__, Config.LogCategoryName, Config.LogVerbosity,
+		TEXT("http-%3u: %5" UINT64_FMT "ms %5" UINT64_FMT "KiB %s (Attempt %d/%d)"),
+		Request.StatusCode,
+		DurationMs,
+		Size,
+		WriteToString<512>(Request.Connection.HostGroup.Host(Request.Host), Request.RelativeUrl).ToString(),
+		Request.RetryCount,
+		Config.MaxRetryCount
+	);
+
+#endif //!NO_LOGGING
 }
 
 } // namespace UE::IoStore
diff --git a/Engine/Source/Runtime/Experimental/IoStore/OnDemand/Private/OnDemandHttpClient.h b/Engine/Source/Runtime/Experimental/IoStore/OnDemand/Private/OnDemandHttpClient.h
--- a/Engine/Source/Runtime/Experimental/IoStore/OnDemand/Private/OnDemandHttpClient.h
+++ b/Engine/Source/Runtime/Experimental/IoStore/OnDemand/Private/OnDemandHttpClient.h
@@ -85,12 +85,18 @@ private:
 
 struct FMultiEndpointHttpClientConfig
 {
-	int32			MaxConnectionCount = 4;
-	int32			ReceiveBufferSize = -1;
-	uint32			MaxRetryCount = 0;
-	int32			TimeoutMs = 0;
-	EHttpRedirects	Redirects = EHttpRedirects::Follow;
-	bool			bEnableThreadSafetyChecks = false;
+	int32				MaxConnectionCount = 4;
+	int32				ReceiveBufferSize = -1;
+	uint32				MaxRetryCount = 0;
+	int32				TimeoutMs = 0;
+	EHttpRedirects		Redirects = EHttpRedirects::Follow;
+	bool				bEnableThreadSafetyChecks = false;
+
+#if !NO_LOGGING
+	/** When not set (the default) this will disable logging, it is up to the system creating the client to assign a LogCategory */
+	FLogCategoryName	LogCategoryName;
+	ELogVerbosity::Type	LogVerbosity = ELogVerbosity::Log;
+#endif //!NO_LOGGING
 };
 
 struct FMultiEndpointHttpClientResponse
@@ -163,7 +169,7 @@ private:
 	FConnection&						GetConnection(const FOnDemandHostGroup& HostGroup);
 	FConnection*						FindConnection(const FOnDemandHostGroup& HostGroup);
 
-	void								Log(const FRequest& Request);
+	void								Log(const FRequest& Request) const;
 
 	FMultiEndpointHttpClientConfig		Config;
 	TArray<TUniquePtr<FConnection>>		Connections;
diff --git a/Engine/Source/Runtime/Experimental/IoStore/OnDemand/Private/OnDemandIoStore.cpp b/Engine/Source/Runtime/Experimental/IoStore/OnDemand/Private/OnDemandIoStore.cpp
--- a/Engine/Source/Runtime/Experimental/IoStore/OnDemand/Private/OnDemandIoStore.cpp
+++ b/Engine/Source/Runtime/Experimental/IoStore/OnDemand/Private/OnDemandIoStore.cpp
@@ -1421,6 +1421,9 @@ FOnDemandInstallResult FOnDemandIoStore::TickInstallRequest(FOnDemandInternalIns
 	FMultiEndpointHttpClientConfig HttpConfig;
 	HttpConfig.MaxConnectionCount = 8;
 	HttpConfig.MaxRetryCount = 2;
+#if !NO_LOGGING
+	HttpConfig.LogCategoryName = LogIoStoreOnDemand.GetCategoryName();
+#endif // !NO_LOGGING
 
 	TUniquePtr<FMultiEndpointHttpClient> HttpClient = FMultiEndpointHttpClient::Create(MoveTemp(HttpConfig));
 	if (HttpClient.IsValid() == false)
