State container input (similar to regular input but is not networked)

diff --git a/Engine/Plugins/Experimental/ChaosModularVehicle/Source/ChaosModularVehicleEngine/Private/ChaosModularVehicle/ModularVehicleBaseComponent.cpp b/Engine/Plugins/Experimental/ChaosModularVehicle/Source/ChaosModularVehicleEngine/Private/ChaosModularVehicle/ModularVehicleBaseComponent.cpp
--- a/Engine/Plugins/Experimental/ChaosModularVehicle/Source/ChaosModularVehicleEngine/Private/ChaosModularVehicle/ModularVehicleBaseComponent.cpp
+++ b/Engine/Plugins/Experimental/ChaosModularVehicle/Source/ChaosModularVehicleEngine/Private/ChaosModularVehicle/ModularVehicleBaseComponent.cpp
@@ -666,6 +666,7 @@ void UModularVehicleBaseComponent::BeginPlay()
 	}
 
 	InputsContainer.Initialize(CombinedInputConfiguration, InputNameMap);
+	StateInputContainer.Initialize(StateInputConfiguration, StateNameMap);
 
 	if (!bUsingNetworkPhysicsPrediction)
 	{
@@ -674,6 +675,7 @@ void UModularVehicleBaseComponent::BeginPlay()
 	// #TODO reinstate ? GenerateInputModifiers(CombinedInputConfiguration);
 
 	VehicleSimulationPT->SetInputMappings(InputNameMap);
+	VehicleSimulationPT->SetStateMappings(StateNameMap);
 
 }
 
@@ -1033,6 +1035,9 @@ void UModularVehicleBaseComponent::Update(float DeltaTime)
 		AsyncInput->PhysicsInputs.TraceParams = TraceParams;
 		AsyncInput->PhysicsInputs.TraceCollisionResponse = SuspensionTraceCollisionResponses;
 		AsyncInput->PhysicsInputs.TraceType = TraceType;
+
+		// game state container
+		AsyncInput->PhysicsInputs.StateInputs.StateInputContainer = StateInputContainer;
 	}
 
 
@@ -1568,6 +1573,36 @@ void UModularVehicleBaseComponent::SetInput(const FName& Name, const FVector& Va
 	}
 }
 
+void UModularVehicleBaseComponent::SetState(const FName& InName, const bool InValue)
+{
+	FInputInterface Inputs(StateNameMap, StateInputContainer);
+	Inputs.SetBool(InName, InValue);
+}
+
+void UModularVehicleBaseComponent::SetState(const FName& InName, const int32 InValue)
+{
+	FInputInterface Inputs(StateNameMap, StateInputContainer);
+	Inputs.SetInteger(InName, InValue);
+}
+
+void UModularVehicleBaseComponent::SetState(const FName& InName, const double InValue)
+{
+	FInputInterface Inputs(StateNameMap, StateInputContainer);
+	Inputs.SetFloat(InName, InValue);
+}
+
+void UModularVehicleBaseComponent::SetState(const FName& InName, const FVector2D& InValue)
+{
+	FInputInterface Inputs(StateNameMap, StateInputContainer);
+	Inputs.SetVector2D(InName, InValue);
+}
+
+void UModularVehicleBaseComponent::SetState(const FName& InName, const FVector& InValue)
+{
+	FInputInterface Inputs(StateNameMap, StateInputContainer);
+	Inputs.SetVector(InName, InValue);
+}
+
 void UModularVehicleBaseComponent::SetGearInput(int32 Gear)
 {
 	GearInput = Gear;
diff --git a/Engine/Plugins/Experimental/ChaosModularVehicle/Source/ChaosModularVehicleEngine/Private/ChaosModularVehicle/ModularVehicleSimulationCU.cpp b/Engine/Plugins/Experimental/ChaosModularVehicle/Source/ChaosModularVehicleEngine/Private/ChaosModularVehicle/ModularVehicleSimulationCU.cpp
--- a/Engine/Plugins/Experimental/ChaosModularVehicle/Source/ChaosModularVehicleEngine/Private/ChaosModularVehicle/ModularVehicleSimulationCU.cpp
+++ b/Engine/Plugins/Experimental/ChaosModularVehicle/Source/ChaosModularVehicleEngine/Private/ChaosModularVehicle/ModularVehicleSimulationCU.cpp
@@ -92,7 +92,11 @@ void FModularVehicleSimulation::SimulateModuleTree(UWorld* InWorld, float DeltaS
 		FModuleInputContainer Container = InputData.PhysicsInputs.NetworkInputs.VehicleInputs.Container;
 		FInputInterface InputInterface(InputNameMap, Container);
 
+		FModuleInputContainer StateInputContainer = InputData.PhysicsInputs.StateInputs.StateInputContainer;
+		FInputInterface StateInterface(StateNameMap, StateInputContainer);
+
 		SimInputData.ControlInputs = &InputInterface;
+		SimInputData.StateInputs = &StateInterface;
 		SimInputData.bKeepVehicleAwake = InputData.PhysicsInputs.NetworkInputs.VehicleInputs.KeepAwake;
 
 		PerformAdditionalSimWork(InWorld, InputData, Proxy, SimInputData);
diff --git a/Engine/Plugins/Experimental/ChaosModularVehicle/Source/ChaosModularVehicleEngine/Public/ChaosModularVehicle/ChaosSimModuleManagerAsyncCallback.h b/Engine/Plugins/Experimental/ChaosModularVehicle/Source/ChaosModularVehicleEngine/Public/ChaosModularVehicle/ChaosSimModuleManagerAsyncCallback.h
--- a/Engine/Plugins/Experimental/ChaosModularVehicle/Source/ChaosModularVehicleEngine/Public/ChaosModularVehicle/ChaosSimModuleManagerAsyncCallback.h
+++ b/Engine/Plugins/Experimental/ChaosModularVehicle/Source/ChaosModularVehicleEngine/Public/ChaosModularVehicle/ChaosSimModuleManagerAsyncCallback.h
@@ -172,8 +172,7 @@ struct CHAOSMODULARVEHICLEENGINE_API FPhysicsModularVehicleTraits
 // #TBD
 struct FGameStateInputs
 {
-	//bool IsInWater;
-	//bool IsInAir
+	FModuleInputContainer StateInputContainer;
 };
 
 UENUM()
diff --git a/Engine/Plugins/Experimental/ChaosModularVehicle/Source/ChaosModularVehicleEngine/Public/ChaosModularVehicle/ModularVehicleBaseComponent.h b/Engine/Plugins/Experimental/ChaosModularVehicle/Source/ChaosModularVehicleEngine/Public/ChaosModularVehicle/ModularVehicleBaseComponent.h
--- a/Engine/Plugins/Experimental/ChaosModularVehicle/Source/ChaosModularVehicleEngine/Public/ChaosModularVehicle/ModularVehicleBaseComponent.h
+++ b/Engine/Plugins/Experimental/ChaosModularVehicle/Source/ChaosModularVehicleEngine/Public/ChaosModularVehicle/ModularVehicleBaseComponent.h
@@ -199,11 +199,20 @@ public:
 	// CONTROLS
 	// 
 	//void SetInput(const FName& Name, const FModuleInputValue& Value);
+	// Control inputs are automatically quantized for networking
 	void SetInput(const FName& Name, const bool Value);
+	void SetInput(const FName& Name, const int32 Value);
 	void SetInput(const FName& Name, const double Value);
 	void SetInput(const FName& Name, const FVector2D& Value);
 	void SetInput(const FName& Name, const FVector& Value);
 
+	// The state inputs are not automatically quantized
+	void SetState(const FName& Name, const bool Value);
+	void SetState(const FName& Name, const int32 Value);
+	void SetState(const FName& Name, const double Value);
+	void SetState(const FName& Name, const FVector2D& Value);
+	void SetState(const FName& Name, const FVector& Value);
+
 	// Sets the input producer class and creates an input producer if one doesn't exist.
 	UFUNCTION(BlueprintCallable, Category = "Game|Components|ModularVehicle")
 	void SetInputProducerClass(TSubclassOf<UVehicleInputProducerBase> InInputProducerClass);
@@ -342,6 +351,9 @@ public:
 
 	UPROPERTY(EditAnywhere, Category = VehicleInput)
 	TArray<FModuleInputSetup> InputConfig;
+
+	UPROPERTY(EditAnywhere, Category = VehicleState)
+	TArray<FModuleInputSetup> StateInputConfiguration;
 		
 	UPROPERTY(EditAnywhere, Category = "Game|Components|ModularVehicle")
 	TEnumAsByte<ESimTreeProcessingOrder> TreeProcessingOrder = ESimTreeProcessingOrder::LeafFirst;
@@ -412,6 +424,9 @@ private:
 
 	FModuleInputContainer InputsContainer;
 
+	FInputNameMap StateNameMap;
+	FModuleInputContainer StateInputContainer;
+
 	Chaos::FPhysicsObjectHandle RootPhysicsObject = nullptr;
 	mutable IPhysicsProxyBase* CachedPhysicsProxy = nullptr;
 };
diff --git a/Engine/Plugins/Experimental/ChaosModularVehicle/Source/ChaosModularVehicleEngine/Public/ChaosModularVehicle/ModularVehicleSimulationCU.h b/Engine/Plugins/Experimental/ChaosModularVehicle/Source/ChaosModularVehicleEngine/Public/ChaosModularVehicle/ModularVehicleSimulationCU.h
--- a/Engine/Plugins/Experimental/ChaosModularVehicle/Source/ChaosModularVehicleEngine/Public/ChaosModularVehicle/ModularVehicleSimulationCU.h
+++ b/Engine/Plugins/Experimental/ChaosModularVehicle/Source/ChaosModularVehicleEngine/Public/ChaosModularVehicle/ModularVehicleSimulationCU.h
@@ -60,6 +60,12 @@ public:
 		UE::TWriteScopeLock InputConfigLock(InputConfigurationLock);
 		InputNameMap = InNameMap; 
 	}
+	void SetStateMappings(const FInputNameMap& InNameMap)
+	{
+		UE::TWriteScopeLock InputConfigLock(InputConfigurationLock);
+		StateNameMap = InNameMap;
+	}
+
 	void CacheRootParticle(IPhysicsProxyBase* Proxy);
 
 	/** Update called from Physics Thread */
@@ -90,6 +96,7 @@ public:
 	/** Current control inputs that is being used on the PT */
 	FModularVehicleInputs VehicleInputs;
 	FInputNameMap InputNameMap;
+	FInputNameMap StateNameMap;
 	FTransactionallySafeRWLock InputConfigurationLock;
 
 	int8 NetMode;
diff --git a/Engine/Source/Runtime/Experimental/ChaosVehicles/ChaosVehiclesCore/Private/SimModule/ModuleInput.cpp b/Engine/Source/Runtime/Experimental/ChaosVehicles/ChaosVehiclesCore/Private/SimModule/ModuleInput.cpp
--- a/Engine/Source/Runtime/Experimental/ChaosVehicles/ChaosVehiclesCore/Private/SimModule/ModuleInput.cpp
+++ b/Engine/Source/Runtime/Experimental/ChaosVehicles/ChaosVehiclesCore/Private/SimModule/ModuleInput.cpp
@@ -398,13 +398,13 @@ void FModuleInputContainer::Decay(const float DecayAmount)
 	}
 }
 
-void FInputInterface::SetValue(const FName& InName, const FModuleInputValue& InValue)
+void FInputInterface::SetValue(const FName& InName, const FModuleInputValue& InValue, bool Quantize /*= true*/)
 {
 	if (ValueContainer.GetNumInputs() > 0)
 	{
 		if (const int* Index = NameMap.Find(InName))
 		{
-			ValueContainer.SetValueAtIndex(*Index, InValue);
+			ValueContainer.SetValueAtIndex(*Index, InValue, Quantize);
 		}
 		else
 		{
@@ -446,6 +446,24 @@ FModuleInputValue FInputInterface::GetValue(const FName& InName) const
 	return FModuleInputValue(EModuleInputValueType::MBoolean, FVector::ZeroVector);
 }
 
+EModuleInputValueType FInputInterface::GetValueType(const FName& InName) const
+{
+	if (ValueContainer.GetNumInputs() > 0)
+	{
+		if (const int* Index = NameMap.Find(InName))
+		{
+			return ValueContainer.GetValueAtIndex(*Index).GetValueType();
+		}
+		else
+		{
+			UE_LOG(LogModularInput, Warning, TEXT("Trying to get the value type of an undefined control input %s"), *InName.ToString());
+		}
+	}
+
+	return EModuleInputValueType::MBoolean;
+}
+
+
 float FInputInterface::GetMagnitude(const FName& InName) const
 {
 	if (ValueContainer.GetNumInputs() > 0)
diff --git a/Engine/Source/Runtime/Experimental/ChaosVehicles/ChaosVehiclesCore/Public/SimModule/ModuleInput.h b/Engine/Source/Runtime/Experimental/ChaosVehicles/ChaosVehiclesCore/Public/SimModule/ModuleInput.h
--- a/Engine/Source/Runtime/Experimental/ChaosVehicles/ChaosVehiclesCore/Public/SimModule/ModuleInput.h
+++ b/Engine/Source/Runtime/Experimental/ChaosVehicles/ChaosVehiclesCore/Public/SimModule/ModuleInput.h
@@ -425,26 +425,31 @@ public:
 
 	static bool ToBool(FModuleInputValue& InValue)
 	{
+		ensure(InValue.GetValueType() == EModuleInputValueType::MBoolean);
 		return InValue.Get<bool>();
 	}
 
 	static float ToAxis1D(FModuleInputValue& InValue)
 	{
+		ensure(InValue.GetValueType() == EModuleInputValueType::MAxis1D);
 		return static_cast<float>(InValue.Get<FModuleInputValue::MAxis1D>());
 	}
 
 	static FVector2D ToAxis2D(FModuleInputValue& InValue)
 	{
+		ensure(InValue.GetValueType() == EModuleInputValueType::MAxis2D);
 		return InValue.Get<FModuleInputValue::MAxis2D>();
 	}
 
 	static FVector ToAxis3D(FModuleInputValue& InValue)
 	{
+		ensure(InValue.GetValueType() == EModuleInputValueType::MAxis3D);
 		return InValue.Get<FModuleInputValue::MAxis3D>();
 	}
 
 	static int32 ToInteger(FModuleInputValue& InValue)
 	{
+		ensure(InValue.GetValueType() == EModuleInputValueType::MInteger);
 		return InValue.Get<FModuleInputValue::MInteger>();
 	}
 
@@ -579,7 +584,17 @@ public:
 
 	int GetNumInputs() const { return InputValues.Num(); }
 	FModuleInputValue GetValueAtIndex(int Index) const { return InputValues[Index]; }
-	void SetValueAtIndex(int Index, const FModuleInputValue& InValue) { InputValues[Index] = InValue.ReturnQuantized(); }
+	void SetValueAtIndex(int Index, const FModuleInputValue& InValue, bool Quantize = true)
+	{
+		if (Quantize)
+		{
+			InputValues[Index] = InValue.ReturnQuantized();
+		}
+		else
+		{
+			InputValues[Index] = InValue;
+		}
+	}
 	void MergeValueAtIndex(int Index, const FModuleInputValue& InValue) { InputValues[Index].Merge(InValue.ReturnQuantized()); }
 
 	FModuleInputContainer& operator=(const FModuleInputContainer& Other)
@@ -632,13 +647,71 @@ public:
 	{
 	}
 
-	void SetValue(const FName& InName, const FModuleInputValue& InValue);
+	void SetValue(const FName& InName, const FModuleInputValue& InValue, bool Quantize = true);
 	void MergeValue(const FName& InName, const FModuleInputValue& InValue);
 	FModuleInputValue GetValue(const FName& InName) const;
+	EModuleInputValueType GetValueType(const FName& InName) const;
 	float GetMagnitude(const FName& InName) const;
 	int32 GetMagnitudeInt(const FName& InName) const;
 	bool InputsNonZero() const;
 
+	// Quick Access to data type
+	bool GetBool(const FName& InName) const
+	{
+		FModuleInputValue Value = GetValue(InName);
+		return FModuleInputConversion::ToBool(Value);
+	}
+	int32 GetInteger(const FName& InName) const
+	{
+		FModuleInputValue Value = GetValue(InName);
+		return FModuleInputConversion::ToInteger(Value);
+	}
+	double GetFloat(const FName& InName) const
+	{
+		FModuleInputValue Value = GetValue(InName);
+		return FModuleInputConversion::ToAxis1D(Value);
+	}
+	FVector2D GetVector2D(const FName& InName) const
+	{
+		FModuleInputValue Value = GetValue(InName);
+		return FModuleInputConversion::ToAxis2D(Value);
+	}
+	FVector GetVector(const FName& InName) const
+	{
+		FModuleInputValue Value = GetValue(InName);
+		return FModuleInputConversion::ToAxis3D(Value);
+	}
+
+	void SetBool(const FName& InName, bool InBool)
+	{
+		ensure(GetValueType(InName) == EModuleInputValueType::MBoolean);
+		FModuleInputValue Value(EModuleInputValueType::MBoolean, InBool);
+		SetValue(InName, Value);
+	}
+	void SetInteger(const FName& InName, int32 InInteger)
+	{
+		ensure(GetValueType(InName) == EModuleInputValueType::MInteger);
+		FModuleInputValue Value(InInteger);
+		SetValue(InName, Value);
+	}
+	void SetFloat(const FName& InName, double InFloat, bool Quantize = false)
+	{
+		ensure(GetValueType(InName) == EModuleInputValueType::MAxis1D);
+		FModuleInputValue Value(InFloat);
+		SetValue(InName, Value, Quantize);
+	}
+	void SetVector2D(const FName& InName, const FVector2D& InVector2D, bool Quantize = false)
+	{
+		ensure(GetValueType(InName) == EModuleInputValueType::MAxis2D);
+		SetValue(InName, FModuleInputValue(InVector2D), Quantize);
+	}
+	void SetVector(const FName& InName, const FVector& InVector, bool Quantize = false)
+	{
+		ensure(GetValueType(InName) == EModuleInputValueType::MAxis3D);
+		SetValue(InName, FModuleInputValue(InVector), Quantize);
+	}
+
+
 	const FInputNameMap& NameMap;			// per vehicle
 	FModuleInputContainer& ValueContainer;	// per vehicle instance
 };
diff --git a/Engine/Source/Runtime/Experimental/ChaosVehicles/ChaosVehiclesCore/Public/SimModule/SimulationModuleBase.h b/Engine/Source/Runtime/Experimental/ChaosVehicles/ChaosVehiclesCore/Public/SimModule/SimulationModuleBase.h
--- a/Engine/Source/Runtime/Experimental/ChaosVehicles/ChaosVehiclesCore/Public/SimModule/SimulationModuleBase.h
+++ b/Engine/Source/Runtime/Experimental/ChaosVehicles/ChaosVehiclesCore/Public/SimModule/SimulationModuleBase.h
@@ -62,10 +62,12 @@ namespace Chaos
 	struct CHAOSVEHICLESCORE_API FAllInputs
 	{
 		FInputInterface& GetControls() const { check(ControlInputs); return *ControlInputs; }
+		FInputInterface& GetState() const { check(StateInputs); return *StateInputs; }
 
 		FTransform VehicleWorldTransform;
 		TMap<int32, FModuleHitResults> HitResults;
 		FInputInterface* ControlInputs = nullptr;
+		FInputInterface* StateInputs = nullptr;
 		bool bKeepVehicleAwake = false;
 	};
 
