FOnDemandIoBackend now uses FMultiEndpointHttpClient instead of FHttpClient for it's communication.

#rnx

### New Client
- The goal for this changelist is to reach feature parity with the old client, so for now FOnDemandIoBackend only has a single FOnDemandHostGroup which acts pretty much how AvailableEps did for the single client implementation.
-- The system should function exactly the same as before, including error handling and how we reconnect to failed CDNs etc.
- Informing the stats system that a request had performed retries used to be done by FHttpClient, this has now been moved to FOnDemandIoBackend when we process the request result. This continues our move away from treating the stats system as a singleton.
- The output generated by 'ias.DisplayOnScreenStatistics=true' has been changed slightly so that we can display the current state of each FOnDemandHostGroup in future when we have more than one.
- The statistics system still records the current url of the default (first) FOnDemandHostGroup but that only works at the moment while we retain feature parity and will have to be changed in the future.

### Misc
- The LatencyTest fuunctions now take FAnsiStringView as input rather than FStringView to avoid conversions (the new client is FAnsiString based rather than FString)
-- Fix that we were logging a string pointer from a string view. In theory we know that the stringview is nullterminated at this point but we cannot rely on that not changing in the future.
- FOnDemandIoBackend::GetEndpointTestPath now returns a const ref to the endpoint path, there is no need to return a copy.

diff --git a/Engine/Source/Runtime/Experimental/IoStore/OnDemand/Private/OnDemandHostGroup.cpp b/Engine/Source/Runtime/Experimental/IoStore/OnDemand/Private/OnDemandHostGroup.cpp
--- a/Engine/Source/Runtime/Experimental/IoStore/OnDemand/Private/OnDemandHostGroup.cpp
+++ b/Engine/Source/Runtime/Experimental/IoStore/OnDemand/Private/OnDemandHostGroup.cpp
@@ -65,7 +65,7 @@ FAnsiStringView FOnDemandHostGroup::CycleHost(int32& InOutIndex) const
 
 void FOnDemandHostGroup::SetPrimaryHost(int32 Index)
 {
-	check(Index >= 0 && Index < Impl->HostUrls.Num());
+	check(Index >= 0 && Index < Impl->HostUrls.Num() || Index == INDEX_NONE);
 	Impl->PrimaryIndex = Index;
 }
 
diff --git a/Engine/Source/Runtime/Experimental/IoStore/OnDemand/Private/OnDemandIoDispatcherBackend.cpp b/Engine/Source/Runtime/Experimental/IoStore/OnDemand/Private/OnDemandIoDispatcherBackend.cpp
--- a/Engine/Source/Runtime/Experimental/IoStore/OnDemand/Private/OnDemandIoDispatcherBackend.cpp
+++ b/Engine/Source/Runtime/Experimental/IoStore/OnDemand/Private/OnDemandIoDispatcherBackend.cpp
@@ -40,10 +40,6 @@
 
 #include <atomic>
 
-#if !UE_BUILD_SHIPPING
-#include "Modules/ModuleManager.h"
-#endif 
-
 /** When enabled the IAS system can add additional debug console commands for development use */
 #define UE_IAS_DEBUG_CONSOLE_CMDS (1 && !NO_CVARS && !UE_BUILD_SHIPPING)
 
@@ -99,7 +95,7 @@ static FAutoConsoleVariableRef CVar_GIasHttpRecvBufKiB(
 	TEXT("Recv buffer size")
 );
 
-static int32 GIasHttpConcurrentRequests = 8;
+static int32 GIasHttpConcurrentRequests = 1;
 static FAutoConsoleVariableRef CVar_IasHttpConcurrentRequests(
 	TEXT("ias.HttpConcurrentRequests"),
 	GIasHttpConcurrentRequests,
@@ -289,7 +285,8 @@ FAutoConsoleVariableRef CVarOnDemandBackendThreadPriority(
 }
 
 ///////////////////////////////////////////////////////////////////////////////
-static bool LatencyTest(FStringView Url, FStringView Path)
+
+static bool LatencyTest(FAnsiStringView Url, FStringView Path)
 {
 	TRACE_CPUPROFILER_EVENT_SCOPE(IasBackend::LatencyTest);
 
@@ -299,8 +296,8 @@ static bool LatencyTest(FStringView Url, FStringView Path)
 	if (Results[0] >= 0 || Results[1] >= 0 || Results[2] >= 0 || Results[3] >= 0)
 	{
 #if !UE_BUILD_SHIPPING
-		UE_LOG(LogIas, Log, TEXT("Endpoint '%s' latency test (ms): %d %d %d %d"),
-			Url.GetData(), Results[0], Results[1], Results[2], Results[3]);
+		UE_LOGFMT(LogIas, Log, "Endpoint '{Url}' latency test (ms): {Result0} {Result1} {Result2} {Result3}",
+			Url, Results[0], Results[1], Results[2], Results[3]);
 #endif // !UE_BUILD_SHIPPING
 
 		return true;
@@ -312,7 +309,7 @@ static bool LatencyTest(FStringView Url, FStringView Path)
 }
 
 ///////////////////////////////////////////////////////////////////////////////
-static int32 LatencyTest(TConstArrayView<FString> Urls, FStringView Path, std::atomic_bool& bCancel)
+static int32 LatencyTest(TConstArrayView<FAnsiString> Urls, FStringView Path, std::atomic_bool& bCancel)
 {
 	TRACE_CPUPROFILER_EVENT_SCOPE(IasBackend::LatencyTest);
 
@@ -328,6 +325,7 @@ static int32 LatencyTest(TConstArrayView<FString> Urls, FStringView Path, std::a
 
 	return INDEX_NONE;
 }
+
 ///////////////////////////////////////////////////////////////////////////////
 struct FBitWindow
 {
@@ -942,12 +940,12 @@ public:
 private:
 	bool Resolve(FIoRequestImpl* Request);
 
-	FString GetEndpointTestPath() const;
+	const FString& GetEndpointTestPath() const;
 	void ConditionallyStartBackendThread();
 	void CompleteRequest(FChunkRequest* ChunkRequest);
 	void CompleteCacheRead(FChunkRequest* ChunkRequest);
 	bool ResolveDistributedEndpoint(const FDistributedEndpointUrl& Url);
-	void ProcessHttpRequests(FHttpClient& HttpClient, FBitWindow& HttpErrors, int32 MaxConcurrentRequests);
+	void ProcessHttpRequests(FMultiEndpointHttpClient& HttpClient, FBitWindow& HttpErrors, int32 MaxConcurrentRequests);
 	int32 WaitForCompleteRequestTasks(float WaitTimeSeconds, float PollTimeSeconds);
 	void DrainHttpRequests();
 	void UpdateThreadPriorityIfNeeded();
@@ -972,6 +970,8 @@ private:
 	FDistributedEndpointUrl DistributionUrl;
 	FEventRef DistributedEndpointEvent;
 
+	FOnDemandHostGroup DefaultHostGroup;
+
 	FString EndpointTestPath;
 	EThreadPriority CurrentThreadPriority;
 
@@ -988,7 +988,7 @@ private:
 FOnDemandIoBackend::FOnDemandIoBackend(const FOnDemandEndpointConfig& Config, FOnDemandIoStore& InIoStore, TUniquePtr<IIasCache>&& InCache)
 	: IoStore(InIoStore)
 	, Cache(MoveTemp(InCache))
-	, Stats(AvailableEps, BackendStatus)
+	, Stats(BackendStatus)
 	, CurrentThreadPriority(EThreadPriority::TPri_Num)
 {
 	EndpointTestPath = Config.TocPath;
@@ -1026,6 +1026,7 @@ FOnDemandIoBackend::FOnDemandIoBackend(const FOnDemandEndpointConfig& Config, FO
 		ECVF_Default)
 	);
 
+#if !UE_BUILD_SHIPPING
 	DynamicConsoleCommands.Emplace(
 		IConsoleManager::Get().RegisterConsoleCommand(
 			TEXT("ias.RunRequestTest"),
@@ -1033,6 +1034,7 @@ FOnDemandIoBackend::FOnDemandIoBackend(const FOnDemandEndpointConfig& Config, FO
 			FConsoleCommandWithArgsDelegate::CreateRaw(this, &FOnDemandIoBackend::RunRequestTest),
 			ECVF_Default)
 	);
+#endif // !UE_BUILD_SHIPPING
 
 #endif // UE_IAS_DEBUG_CONSOLE_CMDS
 }
@@ -1084,7 +1086,7 @@ void FOnDemandIoBackend::Shutdown()
 	BackendContext.Reset();
 }
 
-FString FOnDemandIoBackend::GetEndpointTestPath() const
+const FString& FOnDemandIoBackend::GetEndpointTestPath() const
 {
 	return EndpointTestPath;
 }
@@ -1576,11 +1578,13 @@ FOnDemandStreamingCacheUsage FOnDemandIoBackend::GetCacheUsage() const
 	return Usage; 
 }
 
-void FOnDemandIoBackend::ProcessHttpRequests(FHttpClient& HttpClient, FBitWindow& HttpErrors, int32 MaxConcurrentRequests)
+void FOnDemandIoBackend::ProcessHttpRequests(FMultiEndpointHttpClient& HttpClient, FBitWindow& HttpErrors, int32 MaxConcurrentRequests)
 {
 	int32 NumConcurrentRequests = 0;
 	FChunkRequest* NextChunkRequest = HttpRequests.Dequeue();
 
+	HttpClient.UpdateConnection(DefaultHostGroup);
+
 	while (NextChunkRequest)
 	{
 		while (NextChunkRequest)
@@ -1609,31 +1613,37 @@ void FOnDemandIoBackend::ProcessHttpRequests(FHttpClient& HttpClient, FBitWindow
 				}
 				else
 				{
-					check(HttpClient.GetEndpoint() != INDEX_NONE);
 					TAnsiStringBuilder<256> Url;
 					ChunkRequest->Params.GetUrl(Url);
 
 					NumConcurrentRequests++;
-					HttpClient.Get(Url.ToView(), ChunkRequest->Params.ChunkRange,
+
+					auto AnsiUrl = StringCast<ANSICHAR>(Url.GetData(), Url.Len());
+					HttpClient.Get(DefaultHostGroup, AnsiUrl, ChunkRequest->Params.ChunkRange,
 						[this, &NextChunkRequest, ChunkRequest, &NumConcurrentRequests, &HttpErrors]
-						(TIoStatusOr<FIoBuffer> Status, uint64 DurationMs)
+						(FMultiEndpointHttpClientResponse&& HttpResponse)
 						{
 							NumConcurrentRequests--;
-							switch (Status.Status().GetErrorCode())
+
+							if (HttpResponse.RetryCount > 0)
 							{
-							case EIoErrorCode::Ok:
+								Stats.OnHttpRetry();
+								DefaultHostGroup.SetPrimaryHost(HttpResponse.HostIndex);
+							}
+
+							if (HttpResponse.IsOk())
 							{
 								HttpErrors.Add(false);
-								ChunkRequest->Chunk = Status.ConsumeValueOrDie();
-								Stats.OnHttpGet(ChunkRequest->Chunk.DataSize(), DurationMs);
-								break;
+								ChunkRequest->Chunk = MoveTemp(HttpResponse.Body);
+								Stats.OnHttpGet(ChunkRequest->Chunk.DataSize(), HttpResponse.DurationMilliseconds);
 							}
-							case EIoErrorCode::ReadError:
-							case EIoErrorCode::NotFound:
+							else
 							{
 								Stats.OnHttpError();
 								HttpErrors.Add(true);
 
+								DefaultHostGroup.SetPrimaryHost(HttpResponse.HostIndex);
+
 								const float Average = HttpErrors.AvgSetBits();
 								const bool bAboveHighWaterMark = Average > GIasHttpErrorHighWater;
 								UE_LOG(LogIas, Log, TEXT("%.2f%% the last %d HTTP requests failed"), Average * 100.0f, GIasHttpErrorSampleCount);
@@ -1644,27 +1654,14 @@ void FOnDemandIoBackend::ProcessHttpRequests(FHttpClient& HttpClient, FBitWindow
 									UE_LOG(LogIas, Warning, TEXT("HTTP streaming disabled due to high water mark of %.2f of the last %d requests reached"),
 										GIasHttpErrorHighWater * 100.0f, GIasHttpErrorSampleCount);
 								}
-								break;
-							}
-							case EIoErrorCode::Cancelled:
-							{
-								const FIoChunkId ChunkId = ChunkRequests.GetChunkId(ChunkRequest);
-								UE_LOG(LogIas, Log, TEXT("HTTP request for chunk '%s' cancelled"), *LexToString(ChunkId));
-								break;
-							}
-							default:
-							{
-								const FIoChunkId ChunkId = ChunkRequests.GetChunkId(ChunkRequest);
-								UE_LOG(LogIas, Warning, TEXT("Unhandled HTTP response '%s' for chunk '%s'"),
-									GetIoErrorText(Status.Status().GetErrorCode()), *LexToString(ChunkId));
-								break;
-							}
 							}
 
+							UE::Tasks::ETaskPriority TaskPriority = GCompleteMaterializeTaskPriorities[FMath::Clamp(GCompleteMaterializeTaskPriority, 0, UE_ARRAY_COUNT(GCompleteMaterializeTaskPriorities) - 1)];
 							UE::Tasks::Launch(UE_SOURCE_LOCATION, [this, ChunkRequest]()
-							{
-								CompleteRequest(ChunkRequest);
-							}, GetRequestCompletionTaskPriority());
+								{
+									CompleteRequest(ChunkRequest);
+								}, TaskPriority);
+
 						});
 				}
 			}
@@ -1720,6 +1717,8 @@ uint32 FOnDemandIoBackend::Run()
 {
 	LLM_SCOPE_BYTAG(Ias);
 
+	BackendStatus.SetHttpEnabled(false); // Http is disabled until the thread starts up successfully
+
 	if (DistributionUrl.IsValid())
 	{
 		if (ResolveDistributedEndpoint(DistributionUrl) == false)
@@ -1729,8 +1728,6 @@ uint32 FOnDemandIoBackend::Run()
 			// have a fallback url but the current process is shutting down so we might as well just
 			// exist the thread early.
 			UE_LOG(LogIas, Warning, TEXT("Failed to resolve CDN endpoints from distribution URL"));
-			BackendStatus.SetHttpEnabled(false);
-
 			return 0;
 		}
 	}
@@ -1738,7 +1735,6 @@ uint32 FOnDemandIoBackend::Run()
 	if (AvailableEps.Urls.IsEmpty())
 	{
 		UE_LOG(LogIas, Error, TEXT("HTTP streaming disabled, no valid endpoint"));
-		BackendStatus.SetHttpEnabled(false);
 		return 0;
 	}
 
@@ -1755,83 +1751,88 @@ uint32 FOnDemandIoBackend::Run()
 	FBitWindow HttpErrors;
 	HttpErrors.Reset(GIasHttpErrorSampleCount);
 
-	TUniquePtr<FHttpClient> HttpClient = FHttpClient::Create(FHttpClientConfig
-	{
-		.Endpoints = AvailableEps.Urls,
-		.PrimaryEndpoint = AvailableEps.Current,
-		.MaxConnectionCount = GIasHttpConnectionCount,
-		.MaxRetryCount = FMath::Max(AvailableEps.Urls.Num(), GIasHttpRetryCount),
-		.ReceiveBufferSize = GIasHttpRecvBufKiB >= 0 ? GIasHttpRecvBufKiB << 10 : -1,
-		.TimeoutMs = GIasHttpFailTimeOutMs,
-		.bChangeEndpointAfterSuccessfulRetry = GIasHttpChangeEndpointAfterSuccessfulRetry,
-	});
-
+	TUniquePtr<FMultiEndpointHttpClient> HttpClient = FMultiEndpointHttpClient::Create(FMultiEndpointHttpClientConfig
+		{
+			.MaxConnectionCount = GIasHttpConnectionCount,
+			.ReceiveBufferSize = GIasHttpRecvBufKiB >= 0 ? GIasHttpRecvBufKiB << 10 : -1,
+			.MaxRetryCount = (uint32)FMath::Max(AvailableEps.Urls.Num(), GIasHttpRetryCount),
+			.TimeoutMs = GIasHttpFailTimeOutMs,
+			.bEnableThreadSafetyChecks = true,
+
+#if !NO_LOGGING
+			.LogCategoryName = LogIas.GetCategoryName(),
+			.LogVerbosity = ELogVerbosity::VeryVerbose
+#endif // !NO_LOGGING
+		});
 	check(HttpClient.IsValid());
-	
-	if (LatencyTest(AvailableEps.GetCurrent(), GetEndpointTestPath()))
+
+	TIoStatusOr<FOnDemandHostGroup> HostResult = FOnDemandHostGroup::Create(AvailableEps.Urls);
+	if (!HostResult.IsOk())
 	{
-		HttpClient->SetEndpoint(AvailableEps.Current);
+		UE_LOG(LogIas, Error, TEXT("HTTP streaming disabled, could not create the default host group (%s)"), *HostResult.Status().ToString());
+		return 0;
+	}
 
+	DefaultHostGroup = HostResult.ConsumeValueOrDie();
+
+	Stats.AddHostGroup(DefaultHostGroup);
+
+	const double InitStart = FPlatformTime::Seconds();
+
+	if (LatencyTest(DefaultHostGroup.PrimaryHost(), GetEndpointTestPath()))
+	{
 		BackendStatus.SetHttpEnabled(true);
 		Stats.OnHttpConnected();
 	}
 	else
 	{
 		BackendStatus.SetHttpError(true);
-
-		AvailableEps.Current = INDEX_NONE;
-		HttpClient->SetEndpoint(INDEX_NONE);
 		HttpErrors.Reset(GIasHttpErrorSampleCount);
+
+		DefaultHostGroup.SetPrimaryHost(INDEX_NONE);
 	}
 
+	const double InitTime = FPlatformTime::Seconds() - InitStart;
+	UE_LOG(LogIas, Display, TEXT("HostGroup init took %.3f seconds"), InitTime);
+
 	while (!bStopRequested)
 	{
 		UpdateThreadPriorityIfNeeded();
 
-		// Process HTTP request(s) even if the client is invalid to ensure enqueued request(s) gets completed.
 		ProcessHttpRequests(*HttpClient, HttpErrors, FMath::Min(GIasHttpConcurrentRequests, 32));
-		AvailableEps.Current = HttpClient->GetEndpoint();
 
 		if (!bStopRequested)
 		{
-			uint32 WaitTime = MAX_uint32;
 			if (BackendStatus.IsHttpError())
 			{
-				WaitTime = GIasHttpHealthCheckWaitTime;
-				if (HttpClient->GetEndpoint() != INDEX_NONE)
+				if (DefaultHostGroup.PrimaryHostIndex() != INDEX_NONE)
 				{
 					Stats.OnHttpDisconnected();
 
-					AvailableEps.Current = INDEX_NONE;
-					HttpClient->SetEndpoint(INDEX_NONE);
+					DefaultHostGroup.SetPrimaryHost(INDEX_NONE);
 					HttpErrors.Reset(GIasHttpErrorSampleCount);
 				}
 
-				UE_LOG(LogIas, Log, TEXT("Trying to reconnect to any available endpoint"));
-				const FString TestPath = GetEndpointTestPath(); 
-				if (int32 Idx = LatencyTest(AvailableEps.Urls, TestPath, bStopRequested); Idx != INDEX_NONE)
+				UE_LOG(LogIas, Log, TEXT("Trying to reconnect to any available endpoint..."));
+
+				if (const int32 Idx = LatencyTest(DefaultHostGroup.Hosts(), GetEndpointTestPath(), bStopRequested); Idx != INDEX_NONE)
 				{
 					Stats.OnHttpConnected();
 
-					AvailableEps.Current = Idx;
-					HttpClient->SetEndpoint(Idx);
+					DefaultHostGroup.SetPrimaryHost(Idx);
 					BackendStatus.SetHttpError(false);
-					UE_LOG(LogIas, Log, TEXT("Successfully reconnected to '%s'"), *AvailableEps.GetCurrent());
+					UE_LOGFMT(LogIas, Log, "Successfully reconnected to '{Url}'", DefaultHostGroup.PrimaryHost());
 				}
 			}
-			else if (HttpClient->IsUsingPrimaryEndpoint() == false)
+			else if (DefaultHostGroup.PrimaryHostIndex() != 0)
 			{
-				WaitTime = GIasHttpHealthCheckWaitTime;
-				const FString TestPath = GetEndpointTestPath(); 
-				TConstArrayView<FString> Urls = AvailableEps.Urls;
-
-				if (int32 Idx = LatencyTest(Urls.Left(1), TestPath, bStopRequested); Idx != INDEX_NONE)
+				if (const int32 Idx = LatencyTest(DefaultHostGroup.Hosts().Left(1), GetEndpointTestPath(), bStopRequested); Idx != INDEX_NONE)
 				{
-					AvailableEps.Current = Idx;
-					HttpClient->SetEndpoint(Idx);
-					UE_LOG(LogIas, Log, TEXT("Reconnected to primary endpoint '%s'"), *AvailableEps.GetCurrent());
+					DefaultHostGroup.SetPrimaryHost(Idx);
+					UE_LOGFMT(LogIas, Log, "Reconnected to primary host '{Url}'", DefaultHostGroup.PrimaryHost());
 				}
 			}
+
 			if (BackendStatus.ShouldAbandonCache())
 			{
 				BackendStatus.SetAbandonCache(false);
@@ -1842,7 +1843,7 @@ uint32 FOnDemandIoBackend::Run()
 					Cache.Release()->Abandon(); // Will delete its self
 				}
 			}
-			TickBackendEvent->Wait(WaitTime);
+			TickBackendEvent->Wait(BackendStatus.IsHttpError() ? GIasHttpHealthCheckWaitTime : MAX_uint32);
 		}
 	}
 
diff --git a/Engine/Source/Runtime/Experimental/IoStore/OnDemand/Private/Statistics.cpp b/Engine/Source/Runtime/Experimental/IoStore/OnDemand/Private/Statistics.cpp
--- a/Engine/Source/Runtime/Experimental/IoStore/OnDemand/Private/Statistics.cpp
+++ b/Engine/Source/Runtime/Experimental/IoStore/OnDemand/Private/Statistics.cpp
@@ -4,10 +4,10 @@
 
 #include "Algo/MaxElement.h"
 #include "AnalyticsEventAttribute.h"
-#include "AvailableEndPoints.h"
 #include "HAL/IConsoleManager.h"
 #include "HAL/PlatformTime.h"
 #include "IO/IoStoreOnDemand.h"
+#include "IO/OnDemandHostGroup.h"
 #include "Internationalization/Internationalization.h"
 #include "Math/UnrealMathUtility.h"
 #include "Misc/CoreDelegates.h"
@@ -375,9 +375,8 @@ static FOnDemandIoBackendStats* GStatistics = nullptr;
 static FDelegateHandle GStatisticsEndFrameDelegateHandle;
 static FDelegateHandle GStatisticsOnScreenDelegateHandle;
 
-FOnDemandIoBackendStats::FOnDemandIoBackendStats(const FAvailableEndPoints& InAvailableEndPoints, FBackendStatus& InStatus)
+FOnDemandIoBackendStats::FOnDemandIoBackendStats(FBackendStatus& InStatus)
 	: BackendStatus(InStatus)
-	, AvailableEndPoints(InAvailableEndPoints)
 {
 	static constexpr float OneOver1024 = 1.0f / 1024.0f;
 
@@ -457,20 +456,30 @@ FOnDemandIoBackendStats::FOnDemandIoBackendStats(const FAvailableEndPoints& InAv
 					BackendStatusText << TEXT("IAS Backend Status: ");
 					BackendStatus.ToString(BackendStatusText);
 
-					FString CdnUrl = AvailableEndPoints.GetCurrent();
-					if (!CdnUrl.IsEmpty())
+					FCoreDelegates::EOnScreenMessageSeverity Verbosity = BackendStatus.IsHttpError() ?	FCoreDelegates::EOnScreenMessageSeverity::Error :
+																										FCoreDelegates::EOnScreenMessageSeverity::Info;
+
+					Out.Add(Verbosity, FText::FromStringView(BackendStatusText.ToView()));
+				}
+
+				for (const FOnDemandHostGroup& HostGroup : HostGroups)
+				{
+					FCoreDelegates::EOnScreenMessageSeverity Verbosity = FCoreDelegates::EOnScreenMessageSeverity::Info;
+					TStringBuilder<256> Text;
+					Text << TEXT("IAS HostGroup Status: ");
+
+					if(HostGroup.PrimaryHostIndex() != INDEX_NONE)
 					{
-						BackendStatusText << TEXT(" | ") << CdnUrl;
+						Text << HostGroup.PrimaryHost();
+						Text << TEXT(" (") << HostGroup.PrimaryHostIndex() << TEXT("/") << HostGroup.Hosts().Num() << TEXT(")");
 					}
 					else
 					{
-						BackendStatusText << TEXT(" | Disconnected");
+						Text << TEXT("Disconnected");
+						Verbosity = FCoreDelegates::EOnScreenMessageSeverity::Error;
 					}
 
-					FCoreDelegates::EOnScreenMessageSeverity Verbosity = BackendStatus.IsHttpError() ?	FCoreDelegates::EOnScreenMessageSeverity::Error :
-																										FCoreDelegates::EOnScreenMessageSeverity::Info;
-
-					Out.Add(Verbosity, FText::FromStringView(BackendStatusText.ToView()));
+					Out.Add(Verbosity, FText::FromStringView(Text.ToView()));
 				}
 
 				TStringBuilder<256> CachingText;
@@ -527,6 +536,12 @@ FOnDemandIoBackendStats* FOnDemandIoBackendStats::Get()
 	return GStatistics;
 }
 
+void FOnDemandIoBackendStats::AddHostGroup(const FOnDemandHostGroup& NewHostGroup)
+{
+	// TODO: Maybe check for duplicates?
+	HostGroups.Add(NewHostGroup);
+}
+
 void FOnDemandIoBackendStats::ReportGeneralAnalytics(TArray<FAnalyticsEventAttribute>& OutAnalyticsArray) const
 {
 	// Note that this analytics section is not optional, if we are reporting analytics then we report this section
@@ -601,8 +616,9 @@ class FAnalyticsRecording : public IAnalyticsRecording
 {
 public:
 	FAnalyticsRecording() = delete;
-	FAnalyticsRecording(const FAvailableEndPoints& InAvailableEndPoints, const FBackendStatus& InBackendStatus)
-		: AvailableEndPoints(InAvailableEndPoints)
+
+	FAnalyticsRecording(const TArray<FOnDemandHostGroup>& InHostGroups, const FBackendStatus& InBackendStatus)
+		: HostGroups(InHostGroups)
 		, BackendStatus(InBackendStatus)
 	{
 
@@ -645,17 +661,21 @@ private:
 
 	void Report(TArray<FAnalyticsEventAttribute>& OutAnalyticsArray) const override
 	{
-		if (AvailableEndPoints.HasCurrent())
+		// Todo: For now we assume that the first host group is our 'valid' connection. This will need updating
+		if (!HostGroups.IsEmpty() && HostGroups[0].PrimaryHostIndex() != INDEX_NONE)
 		{
-			FString CdnUrl = AvailableEndPoints.GetCurrent();
+			FString CdnUrl = FString(HostGroups[0].PrimaryHost());
 
-			// Strip the prefix from the url as some analytics systems may have trouble dealing with it
-			if (!CdnUrl.RemoveFromStart(TEXT("http://")))
+			if (!CdnUrl.IsEmpty())
 			{
-				CdnUrl.RemoveFromStart(TEXT("https://"));
-			}
+				// Strip the prefix from the url as some analytics systems may have trouble dealing with it
+				if (!CdnUrl.RemoveFromStart(TEXT("http://")))
+				{
+					CdnUrl.RemoveFromStart(TEXT("https://"));
+				}
 
-			AppendAnalyticsEventAttributeArray(OutAnalyticsArray, TEXT("IasCdnUrl"), MoveTemp(CdnUrl));
+				AppendAnalyticsEventAttributeArray(OutAnalyticsArray, TEXT("IasCdnUrl"), MoveTemp(CdnUrl));
+			}
 		}
 
 		// Report if the system has ever actually managed to make a connection
@@ -786,16 +806,14 @@ private:
 		FTrackedValue<FCounterAtomicInt, GCacheRejectBytes>			RejectBytes;
 	} Cache;
 
-	const FAvailableEndPoints& AvailableEndPoints;
+	const TArray<FOnDemandHostGroup>& HostGroups;
 	const FBackendStatus& BackendStatus;
-
 	bool bRecording = true;
 };
 
 TUniquePtr<IAnalyticsRecording> FOnDemandIoBackendStats::StartAnalyticsRecording() const
 {
-	TUniquePtr<FAnalyticsRecording> Recording = MakeUnique<FAnalyticsRecording>(AvailableEndPoints, BackendStatus);
-
+	TUniquePtr<FAnalyticsRecording> Recording = MakeUnique<FAnalyticsRecording>(HostGroups, BackendStatus);
 	return Recording;
 }
 
diff --git a/Engine/Source/Runtime/Experimental/IoStore/OnDemand/Private/Statistics.h b/Engine/Source/Runtime/Experimental/IoStore/OnDemand/Private/Statistics.h
--- a/Engine/Source/Runtime/Experimental/IoStore/OnDemand/Private/Statistics.h
+++ b/Engine/Source/Runtime/Experimental/IoStore/OnDemand/Private/Statistics.h
@@ -20,6 +20,7 @@ LLM_DECLARE_TAG(Ias);
 namespace UE::IoStore
 {
 
+class FOnDemandHostGroup;
 class IAnalyticsRecording;
 struct FAvailableEndPoints;
 
@@ -34,9 +35,11 @@ class FOnDemandIoBackendStats
 public:
 	static FOnDemandIoBackendStats* Get() IAS_STATISTICS_IMPL(nullptr)
 
-	FOnDemandIoBackendStats(const FAvailableEndPoints& InAvailableEndPoints, FBackendStatus& InBackendStatus) IAS_STATISTICS_IMPL()
+	FOnDemandIoBackendStats(FBackendStatus& InBackendStatus) IAS_STATISTICS_IMPL()
 	~FOnDemandIoBackendStats() IAS_STATISTICS_IMPL()
 
+	void AddHostGroup(const FOnDemandHostGroup& HostGroup)  IAS_STATISTICS_IMPL()
+
 	/** Report analytics not directly associated with a specific endpoint */
 	void ReportGeneralAnalytics(TArray<FAnalyticsEventAttribute>& OutAnalyticsArray) const IAS_STATISTICS_IMPL()
 	/** Report analytics for the current endpoint */
@@ -75,7 +78,7 @@ private:
 
 #if IAS_WITH_STATISTICS
 	const FBackendStatus& BackendStatus;
-	const FAvailableEndPoints& AvailableEndPoints;
+	TArray<FOnDemandHostGroup> HostGroups;
 #endif //IAS_WITH_STATISTICS
 };
 
