Iris - Removed atomic refcounts and sequence numbers until we support multi-threading.

#rnx

diff --git a/Engine/Source/Runtime/Experimental/Iris/Core/Private/Iris/ReplicationSystem/NetBlob/PartialNetBlob.cpp b/Engine/Source/Runtime/Experimental/Iris/Core/Private/Iris/ReplicationSystem/NetBlob/PartialNetBlob.cpp
--- a/Engine/Source/Runtime/Experimental/Iris/Core/Private/Iris/ReplicationSystem/NetBlob/PartialNetBlob.cpp
+++ b/Engine/Source/Runtime/Experimental/Iris/Core/Private/Iris/ReplicationSystem/NetBlob/PartialNetBlob.cpp
@@ -8,7 +8,6 @@
 #include "Iris/Serialization/InternalNetSerializationContext.h"
 #include "Iris/Serialization/NetSerializationContext.h"
 #include "Net/Core/Trace/NetTrace.h"
-#include <atomic>
 
 namespace UE::Net::Private
 {
@@ -18,7 +17,7 @@ namespace UE::Net::Private
 	 * Each split NetBlob will reserve unique sequence numbers for all of its parts. If there are multiple ReplicationSystems sharing the atomic adds relatively little growing of the sequence numbers.
 	 * Splitting blobs is a special case mainly used to enable replicating very large objects.
 	 */
-	static std::atomic<uint32> PartialNetBlobGlobalSequenceNumber;
+	static uint32 PartialNetBlobGlobalSequenceNumber;
 }
 
 namespace UE::Net
@@ -278,7 +277,8 @@ void FPartialNetBlob::SplitPayload(const FPartialNetBlob::FPayloadSplitParams& S
 	OutPartialBlobs.Reserve(OutPartialBlobs.Num() + int32(PartialBlobCount));
 	
 	// Reserve sequence numbers for all parts. 
-	uint32 SequenceNumber = Private::PartialNetBlobGlobalSequenceNumber.fetch_add(PartialBlobCount, std::memory_order_relaxed);
+	uint32 SequenceNumber = Private::PartialNetBlobGlobalSequenceNumber;
+	Private::PartialNetBlobGlobalSequenceNumber += PartialBlobCount;
 
 	uint32 PayloadBitOffset = 0U;
 
diff --git a/Engine/Source/Runtime/Experimental/Iris/Core/Private/Iris/ReplicationSystem/ReplicationProtocol.h b/Engine/Source/Runtime/Experimental/Iris/Core/Private/Iris/ReplicationSystem/ReplicationProtocol.h
--- a/Engine/Source/Runtime/Experimental/Iris/Core/Private/Iris/ReplicationSystem/ReplicationProtocol.h
+++ b/Engine/Source/Runtime/Experimental/Iris/Core/Private/Iris/ReplicationSystem/ReplicationProtocol.h
@@ -4,7 +4,6 @@
 
 #include "CoreTypes.h"
 #include "Misc/EnumClassFlags.h"
-#include <atomic>
 
 namespace UE::Net
 {
@@ -107,7 +106,7 @@ struct FReplicationProtocol
 	// Selected traits from the ReplicationStateDescriptors that might be handy to cache in the protocol.
 	EReplicationProtocolTraits ProtocolTraits;
 
-	mutable std::atomic<int32> RefCount;
+	mutable int32 RefCount;
 };
 
 }
diff --git a/Engine/Source/Runtime/Experimental/Iris/Core/Public/Iris/ReplicationSystem/NetBlob/NetBlob.h b/Engine/Source/Runtime/Experimental/Iris/Core/Public/Iris/ReplicationSystem/NetBlob/NetBlob.h
--- a/Engine/Source/Runtime/Experimental/Iris/Core/Public/Iris/ReplicationSystem/NetBlob/NetBlob.h
+++ b/Engine/Source/Runtime/Experimental/Iris/Core/Public/Iris/ReplicationSystem/NetBlob/NetBlob.h
@@ -7,7 +7,6 @@
 #include "Containers/ArrayView.h"
 #include "Templates/RefCounting.h"
 #include "Templates/UniquePtr.h"
-#include <atomic>
 
 namespace UE::Net
 {
@@ -163,7 +162,7 @@ protected:
 	FQuantizedBlobState QuantizedBlobState;
 
 private:
-	mutable std::atomic<int32> RefCount;
+	mutable int32 RefCount;
 };
 
 /**
