Remove a number of IAS cvars that are no longer in use.

#rnx

diff --git a/Engine/Source/Runtime/Experimental/IoStore/OnDemand/Private/OnDemandIoDispatcherBackend.cpp b/Engine/Source/Runtime/Experimental/IoStore/OnDemand/Private/OnDemandIoDispatcherBackend.cpp
--- a/Engine/Source/Runtime/Experimental/IoStore/OnDemand/Private/OnDemandIoDispatcherBackend.cpp
+++ b/Engine/Source/Runtime/Experimental/IoStore/OnDemand/Private/OnDemandIoDispatcherBackend.cpp
@@ -67,13 +67,6 @@ static FAutoConsoleVariableRef CVar_IasHttpPrimaryEndpoint(
 	TEXT("Primary endpoint to use returned from the distribution endpoint")
 );
 
-bool GIasHttpChangeEndpointAfterSuccessfulRetry = true;
-static FAutoConsoleVariableRef CVar_IasHttpChangeEndpointAfterSuccessfulRetry(
-	TEXT("ias.HttpChangeEndpointAfterSuccessfulRetry"),
-	GIasHttpChangeEndpointAfterSuccessfulRetry,
-	TEXT("Whether to change the current endpoint after a sucessful retry")
-);
-
 int32 GIasHttpPollTimeoutMs = 17;
 static FAutoConsoleVariableRef CVar_GIasHttpPollTimeoutMs(
 	TEXT("ias.HttpPollTimeoutMs"),
@@ -142,13 +135,6 @@ static FAutoConsoleVariableRef CVar_IasHttpHealthCheckWaitTime(
 	TEXT("Number of milliseconds to wait before reconnecting to avaiable endpoint(s)")
 );
 
-int32 GIasMaxEndpointTestCountAtStartup = 1;
-static FAutoConsoleVariableRef CVar_IasMaxEndpointTestCountAtStartup(
-	TEXT("ias.MaxEndpointTestCountAtStartup"),
-	GIasMaxEndpointTestCountAtStartup,
-	TEXT("Number of endpoint(s) to test at startup")
-);
-
 int32 GIasHttpErrorSampleCount = 8;
 static FAutoConsoleVariableRef CVar_IasHttpErrorSampleCount(
 	TEXT("ias.HttpErrorSampleCount"),
@@ -230,8 +216,6 @@ static FAutoConsoleVariableRef CVar_IasPoisonCache(
 );
 #endif // !UE_BUILD_SHIPPING
 
-
-
 // These priorities are indexed using the cvar below
 static UE::Tasks::ETaskPriority GCompleteMaterializeTaskPriorities[] =
 {
