Optimize Update AABBB tree
Now we can prune the particle that won't have any effect on the leaves in the AABB tree.
Their own bounding volume will be updated in a multithreaded way, then less particles will be updated in the AABB tree.
Increase the bounding volume leaf size, to update less particles as possible.
Up to 40% gain in the swap acceleration structure in some scenario, in the worse case collision detection can get a 5% slow down.
Overall it is a big win up to 8% gain for the full physics steps.
Everything can be tuned using cvars.

diff --git a/Engine/Source/Runtime/Experimental/Chaos/Private/Chaos/AABBTree.cpp b/Engine/Source/Runtime/Experimental/Chaos/Private/Chaos/AABBTree.cpp
--- a/Engine/Source/Runtime/Experimental/Chaos/Private/Chaos/AABBTree.cpp
+++ b/Engine/Source/Runtime/Experimental/Chaos/Private/Chaos/AABBTree.cpp
@@ -31,6 +31,9 @@ FAutoConsoleVariableRef FAABBTreeCVars::CVarDynamicTreeBoundingBoxPadding(TEXT("
 int32 FAABBTreeCVars::DynamicTreeLeafCapacity = 8;
 FAutoConsoleVariableRef FAABBTreeCVars::CVarDynamicTreeLeafCapacity(TEXT("p.aabbtree.DynamicTreeLeafCapacity"), FAABBTreeCVars::DynamicTreeLeafCapacity, TEXT("Dynamic Tree Leaf Capacity"));
 
+float FAABBTreeCVars::DynamicTreeLeafEnlargePercent = 0.1f;
+FAutoConsoleVariableRef FAABBTreeCVars::CVarDynamicTreeLeafEnlargePercent(TEXT("p.aabbtree.DynamicTreeLeafEnlargePercent"), FAABBTreeCVars::DynamicTreeLeafEnlargePercent, TEXT("Additional percentage added to bounding box of the leaf to amortize update cost. Ex: set 0.05 for 5%"));
+
 bool FAABBTimeSliceCVars::bUseTimeSliceMillisecondBudget = true;
 FAutoConsoleVariableRef FAABBTimeSliceCVars::CVarUseTimeSliceByMillisecondBudget(TEXT("p.aabbtree.UseTimeSliceMillisecondBudget"), FAABBTimeSliceCVars::bUseTimeSliceMillisecondBudget, TEXT("Set to True if we want to timeslice tree generation by a milisecond budget instead of per nodes processed"));
 
@@ -41,4 +44,4 @@ int32 FAABBTimeSliceCVars::MinNodesChunkToProcessBetweenTimeChecks = 250;
 FAutoConsoleVariableRef FAABBTimeSliceCVars::CVarMinNodesChunkToProcessBetweenTimeChecks(TEXT("p.aabbtree.MinNodesChunkToProcessBetweenTimeChecks"), FAABBTimeSliceCVars::MinNodesChunkToProcessBetweenTimeChecks, TEXT("Minimum amount of nodes we want to process before checking if we are withing budget"));
 
 int32 FAABBTimeSliceCVars::MinDataChunkToProcessBetweenTimeChecks = 500;
-FAutoConsoleVariableRef FAABBTimeSliceCVars::CVarMinDataChunkToProcessBetweenTimeChecks(TEXT("p.aabbtree.MinDataChunkToProcessBetweenTimeChecks"), FAABBTimeSliceCVars::MinDataChunkToProcessBetweenTimeChecks, TEXT("Minimum amount of data elements to process before checking if we are withing budget"));
\ No newline at end of file
+FAutoConsoleVariableRef FAABBTimeSliceCVars::CVarMinDataChunkToProcessBetweenTimeChecks(TEXT("p.aabbtree.MinDataChunkToProcessBetweenTimeChecks"), FAABBTimeSliceCVars::MinDataChunkToProcessBetweenTimeChecks, TEXT("Minimum amount of data elements to process before checking if we are withing budget"));
diff --git a/Engine/Source/Runtime/Experimental/Chaos/Private/Chaos/PBDRigidsEvolution.cpp b/Engine/Source/Runtime/Experimental/Chaos/Private/Chaos/PBDRigidsEvolution.cpp
--- a/Engine/Source/Runtime/Experimental/Chaos/Private/Chaos/PBDRigidsEvolution.cpp
+++ b/Engine/Source/Runtime/Experimental/Chaos/Private/Chaos/PBDRigidsEvolution.cpp
@@ -7,6 +7,7 @@
 #include "Chaos/ParticleHandle.h"
 #include "Chaos/SpatialAccelerationCollection.h"
 #include "Chaos/PhysicsMaterialUtilities.h"
+#include "Tasks/Task.h"
 
 CSV_DECLARE_CATEGORY_EXTERN(ChaosPhysicsTimers);
 
@@ -22,6 +23,9 @@ FAutoConsoleVariableRef CVarChaosNumPushOutIterationsOverride(TEXT("p.ChaosNumPu
 int32 ChaosNumContactIterationsOverride = -1;
 FAutoConsoleVariableRef CVarChaosNumContactIterationsOverride(TEXT("p.ChaosNumContactIterationsOverride"), ChaosNumContactIterationsOverride, TEXT("Override for num contact iterations if >= 0. [def:-1]"));
 
+bool bPruningParticleToInsertInAccelerationStructure = true;
+FAutoConsoleVariableRef CVarPruningParticleToInsertInAccelerationStructure(TEXT("p.ChaosPruningParticleToInsertInAccelerationStructure"), bPruningParticleToInsertInAccelerationStructure, TEXT("Enables pruning particle before updating them in the accceleration structure. This is an optimization that reduce the dynamic acceleration structure update."));
+
 namespace Chaos
 {
 	namespace CVars
@@ -617,8 +621,63 @@ namespace Chaos
 		}
 	}
 
+	void FPBDRigidsEvolutionBase::PruneInternalPendingData(FAccelerationStructure& AccelerationStructure)
+	{
+		// Remove all particles which are not updating their leaf in the acceleration structure.
+		// The bounding volume associated to the particle will be updated in the acceleration structure.
+		QUICK_SCOPE_CYCLE_COUNTER(STAT_PruneInternalPendingData);
+
+		int32 NumWorker = !!GSingleThreadedPhysics ? 1 : FMath::Max(FMath::Min(Chaos::CVars::NumWorkerCollisionFactor * (FTaskGraphInterface::Get().GetNumWorkerThreads() - 1), Chaos::MaxNumWorkers), 1);
+		NumWorker = FMath::Max(1, NumWorker);
+
+		const int32 NumParticles = InternalAccelerationQueue.PendingData.Num();
+		const int32 ParticleByContext = FMath::Max(FMath::DivideAndRoundUp(NumParticles, NumWorker), SmallBatchSize);
+		const int32 NumBatch = FMath::DivideAndRoundUp(NumParticles, ParticleByContext);
+
+		TArray<TArray<FPendingSpatialData>> NewPendingDataArray;
+		NewPendingDataArray.SetNum(NumBatch);
+		TArray<UE::Tasks::FTask> PendingTasks;
+		PendingTasks.Reserve(NumBatch);
+		for (int32 BatchIndex = 0; BatchIndex < NumBatch; BatchIndex++)
+		{
+			const int32 FirstIndex = BatchIndex * ParticleByContext;
+			int32 LastIndex = (BatchIndex + 1) * ParticleByContext;
+			LastIndex = BatchIndex == NumBatch - 1 ? FMath::Min(NumParticles, LastIndex) : LastIndex;
+			TArray<FPendingSpatialData>& NewPendingData = NewPendingDataArray[BatchIndex];
+			UE::Tasks::FTask PendingTask = UE::Tasks::Launch(UE_SOURCE_LOCATION, [this, &AccelerationStructure, FirstIndex, LastIndex, &NewPendingData]()
+				{
+					QUICK_SCOPE_CYCLE_COUNTER(STAT_PruneInternalPendingDataTask);
+					for (int32 Index = FirstIndex; Index < LastIndex; Index++)
+					{
+						const FPendingSpatialData& PendingData = InternalAccelerationQueue.PendingData[Index];
+						FGeometryParticleHandle* UpdateParticle = PendingData.AccelerationHandle.GetGeometryParticleHandle_PhysicsThread();
+
+						if (PendingData.Operation == EPendingSpatialDataOperation::Delete ||
+							(UpdateParticle->HasBounds() && AccelerationStructure.NeedUpdateElementIn(UpdateParticle, UpdateParticle->WorldSpaceInflatedBounds(), PendingData.SpatialIdx)))
+						{
+							NewPendingData.Add(PendingData);
+						}
+					}
+				});
+			PendingTasks.Add(PendingTask);
+		}
+
+		UE::Tasks::Wait(PendingTasks);
+		// Once all threads are finished merge all data in the acceleration queue
+		// Maybe this could be multi threaded, but since it is really fast not sure if it is worth it.
+		InternalAccelerationQueue.PendingData.Reset();
+		for (int32 BatchIndex = 0; BatchIndex < NumBatch; BatchIndex++)
+		{
+			InternalAccelerationQueue.PendingData.Append(NewPendingDataArray[BatchIndex]);
+		}
+	}
+
 	void FPBDRigidsEvolutionBase::FlushInternalAccelerationQueue()
 	{
+		if (bPruningParticleToInsertInAccelerationStructure)
+		{
+			PruneInternalPendingData(*InternalAcceleration);
+		}
 		for (const FPendingSpatialData& PendingData : InternalAccelerationQueue.PendingData)
 		{
 			ApplyParticlePendingData(PendingData, *InternalAcceleration, false, true);
@@ -786,7 +845,7 @@ namespace Chaos
 				// but the AsyncAcceleration now probably have fewer dirty elements (they were just rebuilt)
 				
 				//swap acceleration structure for new one
-				std::swap(InternalAcceleration, AsyncInternalAcceleration);
+					std::swap(InternalAcceleration, AsyncInternalAcceleration);
 				// The old InternalAcceleration is now in AsyncInternalAcceleration, and the constituent structures will be reused if no changes have been detected
 
 				if (AccelerationStructureUseDynamicTree)
@@ -816,7 +875,7 @@ namespace Chaos
 			}
 			else
 			{
-				FlushInternalAccelerationQueue();
+					FlushInternalAccelerationQueue();
 
 				if (GAccelerationStructureCacheOverlappingLeaves  && AccelerationStructureUseDynamicTree)
 				{
diff --git a/Engine/Source/Runtime/Experimental/Chaos/Public/Chaos/AABBTree.h b/Engine/Source/Runtime/Experimental/Chaos/Public/Chaos/AABBTree.h
--- a/Engine/Source/Runtime/Experimental/Chaos/Public/Chaos/AABBTree.h
+++ b/Engine/Source/Runtime/Experimental/Chaos/Public/Chaos/AABBTree.h
@@ -37,6 +37,9 @@ struct FAABBTreeCVars
 	static CHAOS_API float DynamicTreeBoundingBoxPadding;
 	static CHAOS_API FAutoConsoleVariableRef CVarDynamicTreeBoundingBoxPadding;
 
+	static CHAOS_API float DynamicTreeLeafEnlargePercent;
+	static CHAOS_API FAutoConsoleVariableRef CVarDynamicTreeLeafEnlargePercent;
+
 	static CHAOS_API int32 DynamicTreeLeafCapacity;
 	static CHAOS_API FAutoConsoleVariableRef CVarDynamicTreeLeafCapacity;
 
@@ -226,6 +229,11 @@ struct TBoundsWrapperHelper<TPayloadType, T, true>
 		{
 			Bounds.GrowToInclude(Elem.Bounds);
 		}
+		if (!Elems.IsEmpty())
+		{
+			// Grow the aabb leaf of percentage of its size
+			Bounds.ThickenSymmetrically(Bounds.Extents() * FAABBTreeCVars::DynamicTreeLeafEnlargePercent*0.5);
+		}
 	}
 
 	const TAABB<T, 3>& GetBounds() const { return Bounds; }
@@ -1841,6 +1849,29 @@ public:
 		return false;
 	}
 
+	// Returns true if the payload's leaf need to be updated, false if the Payload can stay in the same leaf without changing the leaf size.
+	// In the case this function returns false, the payload's bounding volume will be updated. 
+	virtual bool NeedUpdateElement(const TPayloadType& Payload, const TAABB<T, 3>& NewBounds) override
+	{
+		ensure(bMutable);
+		FAABBTreePayloadInfo* PayloadInfo = PayloadToInfo.Find(Payload);
+		if (PayloadInfo && bDynamicTree)
+		{
+			if (PayloadInfo->LeafIdx != INDEX_NONE)
+			{
+				// The leaf node bounds can be larger than the actual leave bound
+				const TAABB<T, 3>& LeafNodeBounds = Nodes[PayloadInfo->NodeIdx].ChildrenBounds[0];
+				if (LeafNodeBounds.Contains(NewBounds.Min()) && LeafNodeBounds.Contains(NewBounds.Max()))
+				{
+					// We still need to update the constituent bounds
+					Leaves[PayloadInfo->LeafIdx].UpdateElement(Payload, NewBounds, true);
+					return false;
+				}
+			}
+		}
+		return true;
+	}
+
 	// Returns true if element was updated, or false when it was added instead
 	virtual bool UpdateElement(const TPayloadType& Payload, const TAABB<T, 3>& NewBounds, bool bInHasBounds) override
 	{
@@ -1904,7 +1935,6 @@ public:
 							{
 								// We still need to update the constituent bounds
 								Leaves[PayloadInfo->LeafIdx].UpdateElement(Payload, NewBounds, bHasBounds);
-								Leaves[PayloadInfo->LeafIdx].RecomputeBounds();
 								bModifyingTreeMultiThreadingFastCheck = false;
 								return bElementExisted;
 							}
diff --git a/Engine/Source/Runtime/Experimental/Chaos/Public/Chaos/ISpatialAcceleration.h b/Engine/Source/Runtime/Experimental/Chaos/Public/Chaos/ISpatialAcceleration.h
--- a/Engine/Source/Runtime/Experimental/Chaos/Public/Chaos/ISpatialAcceleration.h
+++ b/Engine/Source/Runtime/Experimental/Chaos/Public/Chaos/ISpatialAcceleration.h
@@ -309,6 +309,11 @@ public:
 		return true;
 	}
 
+	virtual bool NeedUpdateElement(const TPayloadType& Payload, const TAABB<T, 3>& NewBounds)
+	{
+		return true;
+	}
+
 	virtual bool UpdateElement(const TPayloadType& Payload, const TAABB<T, d>& NewBounds, bool bHasBounds)
 	{
 		check(false);
@@ -320,6 +325,11 @@ public:
 		return RemoveElement(Payload);
 	}
 
+	virtual bool NeedUpdateElementIn(const TPayloadType& Payload, const TAABB<T, d>& NewBounds, FSpatialAccelerationIdx Idx)
+	{
+		return NeedUpdateElement(Payload, NewBounds);
+	}
+
 	virtual bool UpdateElementIn(const TPayloadType& Payload, const TAABB<T, d>& NewBounds, bool bHasBounds, FSpatialAccelerationIdx Idx)
 	{
 		return UpdateElement(Payload, NewBounds, bHasBounds);
diff --git a/Engine/Source/Runtime/Experimental/Chaos/Public/Chaos/PBDRigidsEvolution.h b/Engine/Source/Runtime/Experimental/Chaos/Public/Chaos/PBDRigidsEvolution.h
--- a/Engine/Source/Runtime/Experimental/Chaos/Public/Chaos/PBDRigidsEvolution.h
+++ b/Engine/Source/Runtime/Experimental/Chaos/Public/Chaos/PBDRigidsEvolution.h
@@ -902,6 +902,7 @@ protected:
 		IslandManager.UpdateIslands();
 	}
 
+	void PruneInternalPendingData(FAccelerationStructure& AccelerationStructure);
 	void FlushInternalAccelerationQueue();
 	void FlushAsyncAccelerationQueue();
 	void WaitOnAccelerationStructure();
diff --git a/Engine/Source/Runtime/Experimental/Chaos/Public/Chaos/SpatialAccelerationCollection.h b/Engine/Source/Runtime/Experimental/Chaos/Public/Chaos/SpatialAccelerationCollection.h
--- a/Engine/Source/Runtime/Experimental/Chaos/Public/Chaos/SpatialAccelerationCollection.h
+++ b/Engine/Source/Runtime/Experimental/Chaos/Public/Chaos/SpatialAccelerationCollection.h
@@ -611,6 +611,12 @@ public:
 		return bSuccess;
 	}
 
+	virtual bool NeedUpdateElementIn(const TPayloadType& Payload, const TAABB<T, d>& NewBounds, FSpatialAccelerationIdx SpatialIdx)
+	{
+		const uint16 UseBucket = ((1 << SpatialIdx.Bucket) & this->ActiveBucketsMask) ? SpatialIdx.Bucket : 0;
+		return  Buckets[UseBucket].Objects[SpatialIdx.InnerIdx].Acceleration->NeedUpdateElement(Payload, NewBounds);
+	}
+
 	// Returns true if the element was in fact updated or if it was moved from one SpatialIdx to another during the update or false if not found at any SpatialIdx
 	virtual bool UpdateElementIn(const TPayloadType& Payload, const TAABB<T, d>& NewBounds, bool bHasBounds, FSpatialAccelerationIdx SpatialIdx)
 	{
