Improve error reporting when failing to establish an initial connection to a CDN or failing to resolve the distributed endpoint.

#rnx

- FTicketStatus::GetErrorReason (which only returns an error message) has been replaced with ::GetError which returns a FError containing both the reason for the error and the code.
- Updated all places formerly calling FTicketStatus::GetErrorReason with FTicketStatus::GetError().Reason.
- Bumped logging verbosity when resolving the distributed endpoint from log to warning, moved to structured logging and added the error code where appropriate.
- Moved the log message announcing we are testing an endpoint to the start of the function rather than after we've already tested it.
- Log errors and bad http responses when testing endpoints. We provide the url with these warnings in case we only ever receive a warning summary so we don't lose context.

diff --git a/Engine/Source/Runtime/Experimental/IoStore/HttpClient/Internal/IO/Http/Client.h b/Engine/Source/Runtime/Experimental/IoStore/HttpClient/Internal/IO/Http/Client.h
--- a/Engine/Source/Runtime/Experimental/IoStore/HttpClient/Internal/IO/Http/Client.h
+++ b/Engine/Source/Runtime/Experimental/IoStore/HttpClient/Internal/IO/Http/Client.h
@@ -182,6 +182,13 @@ class UE_API FTicketStatus
 {
 public:
 	enum class EId : uint8 { Response, Content, Cancelled, Error };
+
+	struct FError
+	{
+		const char*	Reason;
+		int32		Code;
+	};
+
 	EId					GetId() const;
 	UPTRINT				GetParam() const;
 	FTicket				GetTicket() const;
@@ -190,7 +197,7 @@ public:
 	const FIoBuffer&	GetContent() const;			// if GetId() == EId::Content
 	uint32				GetContentLength() const;	//  |
 	const FTicketPerf&	GetPerf() const;			// _|_
-	const char*			GetErrorReason() const;		// if GetId() == EId::Error
+	FError				GetError() const;			// if GetId() == EId::Error
 
 private:
 						FTicketStatus() = delete;
diff --git a/Engine/Source/Runtime/Experimental/IoStore/HttpClient/Private/Api.inl b/Engine/Source/Runtime/Experimental/IoStore/HttpClient/Private/Api.inl
--- a/Engine/Source/Runtime/Experimental/IoStore/HttpClient/Private/Api.inl
+++ b/Engine/Source/Runtime/Experimental/IoStore/HttpClient/Private/Api.inl
@@ -251,11 +251,11 @@ const FIoBuffer& FTicketStatus::GetContent() const
 }
 
 ////////////////////////////////////////////////////////////////////////////////
-const char* FTicketStatus::GetErrorReason() const
+FTicketStatus::FError FTicketStatus::GetError() const
 {
 	check(GetId() == EId::Error);
 	const auto* Activity = (FActivity*)this;
-	return Activity->ErrorReason;
+	return FError(Activity->ErrorReason, Activity->StateParam);
 }
 
 
diff --git a/Engine/Source/Runtime/Experimental/IoStore/HttpClient/Private/Test.inl b/Engine/Source/Runtime/Experimental/IoStore/HttpClient/Private/Test.inl
--- a/Engine/Source/Runtime/Experimental/IoStore/HttpClient/Private/Test.inl
+++ b/Engine/Source/Runtime/Experimental/IoStore/HttpClient/Private/Test.inl
@@ -529,7 +529,7 @@ static void ChunkedTest(const ANSICHAR* TestHost)
 			return;
 		}
 
-		FAnsiStringView Reason = Status.GetErrorReason();
+		FAnsiStringView Reason = Status.GetError().Reason;
 		ErrorMarks |= Reason.Contains("ERRTRAIL",   ESearchCase::CaseSensitive) ? 1 : 0;
 		ErrorMarks |= Reason.Contains("ERRNOCHUNK", ESearchCase::CaseSensitive) ? 2 : 0;
 	};
@@ -742,7 +742,7 @@ static void HttpTest(const ANSICHAR* TestHost, FCertRootsRef VerifyCert)
 
 			check(Status.GetId() == FTicketStatus::EId::Error);
 
-			const char* Reason = Status.GetErrorReason();
+			const char* Reason = Status.GetError().Reason;
 			bool IsFailTimeout = (FCStringAnsi::Strstr(Reason, "FailTimeout") != nullptr);
 			check(IsFailTimeout == bExpectFailTimeout);
 		};
diff --git a/Engine/Source/Runtime/Experimental/IoStore/OnDemand/Private/DistributionEndpoints.cpp b/Engine/Source/Runtime/Experimental/IoStore/OnDemand/Private/DistributionEndpoints.cpp
--- a/Engine/Source/Runtime/Experimental/IoStore/OnDemand/Private/DistributionEndpoints.cpp
+++ b/Engine/Source/Runtime/Experimental/IoStore/OnDemand/Private/DistributionEndpoints.cpp
@@ -10,9 +10,9 @@
 #include "IO/Http/Client.h"
 #include "IO/IoBuffer.h"
 #include "IO/IoStoreOnDemand.h"
+#include "Logging/StructuredLog.h"
 #include "Misc/CommandLine.h"
 #include "Misc/ConfigCacheIni.h"
-#include "Misc/ScopeRWLock.h"
 #include "Serialization/JsonReader.h"
 #include "Serialization/JsonSerializer.h"
 #include "Serialization/MemoryReader.h"
@@ -61,8 +61,9 @@ FDistributionEndpoints::EResult FDistributionEndpoints::ResolveEndpoints(const F
 	{
 		if (Status.GetId() == FTicketStatus::EId::Error)
 		{
-			const char* ErrorReason = Status.GetErrorReason();
-			UE_LOG(LogIas, Log, TEXT("Resolve endpoint error: '%S'"), ErrorReason);
+			FTicketStatus::FError Error = Status.GetError();
+			UE_LOGFMT(LogIas, Warning, "ResolveEndpoints Error: {ErrorReason} ({ErrorCode})", Error.Reason, Error.Code);
+			return;
 		}
 
 		if (Status.GetId() >= FTicketStatus::EId::Cancelled)
@@ -81,8 +82,7 @@ FDistributionEndpoints::EResult FDistributionEndpoints::ResolveEndpoints(const F
 		FResponse& Response = Status.GetResponse();
 		if (Response.GetStatus() != EStatusCodeClass::Successful)
 		{
-			uint32 StatusCode = Response.GetStatusCode();
-			UE_LOG(LogIas, Log, TEXT("Resolve HTTP status code: '%d'"), StatusCode);
+			UE_LOGFMT(LogIas, Warning, "ResolveEndpoints failed with HTTP response {ResponseCode}", Response.GetStatusCode());
 			return;
 		}
 
diff --git a/Engine/Source/Runtime/Experimental/IoStore/OnDemand/Private/LatencyTesting.cpp b/Engine/Source/Runtime/Experimental/IoStore/OnDemand/Private/LatencyTesting.cpp
--- a/Engine/Source/Runtime/Experimental/IoStore/OnDemand/Private/LatencyTesting.cpp
+++ b/Engine/Source/Runtime/Experimental/IoStore/OnDemand/Private/LatencyTesting.cpp
@@ -3,8 +3,9 @@
 #include "Containers/StringConv.h"
 #include "Containers/StringView.h"
 #include "HAL/PlatformTime.h"
-#include "IO/IoStoreOnDemand.h"
 #include "IO/Http/Client.h"
+#include "IO/IoStoreOnDemand.h"
+#include "Logging/StructuredLog.h"
 #include "Misc/StringBuilder.h"
 #include "Templates/Function.h"
 
@@ -18,8 +19,13 @@ void LatencyTest(FAnsiStringView InUrl, FStringView InPath, uint32 InTimeOutMs,
 	FConnectionPool::FParams PoolParams;
 	PoolParams.SetHostFromUrl(InUrl);
 	PoolParams.ConnectionCount = 1;
+	
 	FConnectionPool Pool(PoolParams);
 
+	TAnsiStringBuilder<512> ConnectionDesc;
+	Pool.Describe(ConnectionDesc);
+	UE_LOGFMT(LogIas, Log, "Testing endpoint {Url}", ConnectionDesc.ToString());
+
 	TAnsiStringBuilder<256> AnsiPath;
 	if (!InPath.StartsWith(TEXT('/')))
 	{
@@ -37,11 +43,28 @@ void LatencyTest(FAnsiStringView InUrl, FStringView InPath, uint32 InTimeOutMs,
 		FRequest Request = Loop.Request("HEAD", AnsiPath, Pool);
 		Loop.Send(MoveTemp(Request), [&](const FTicketStatus& Status)
 			{
-				if (Status.GetId() != FTicketStatus::EId::Response)
+				if (Status.GetId() == FTicketStatus::EId::Error)
+				{
+					FTicketStatus::FError Error = Status.GetError();
+					UE_LOGFMT(LogIas, Warning, "LatencyTest Error: {Url} {ErrorReason} ({ErrorCode})", ConnectionDesc, Error.Reason, Error.Code);
 					return;
+				}
+				else if (Status.GetId() != FTicketStatus::EId::Response)
+				{
+					return;
+				}
 
 				const FResponse& Response = Status.GetResponse();
 				Ok = (Response.GetStatus() == EStatusCodeClass::Successful);
+
+				if (Response.GetStatus() == EStatusCodeClass::Successful)
+				{
+					Ok = true;
+				}
+				else
+				{
+					UE_LOGFMT(LogIas, Warning, "LatencyTest Failed: {Url} HTTP response ({ResponseCode})", ConnectionDesc, Response.GetStatusCode());
+				}
 			});
 
 		uint64 Cycles = FPlatformTime::Cycles64();
@@ -66,10 +89,6 @@ void LatencyTest(FAnsiStringView InUrl, FStringView InPath, uint32 InTimeOutMs,
 
 		Result = int32((int64(Result) * 1000) / Freq);
 	}
-	
-	TAnsiStringBuilder<512> ConnectionDesc;
-	Pool.Describe(ConnectionDesc);
-	UE_LOG(LogIas, Log, TEXT("Testing endpoint %s"), ANSI_TO_TCHAR(ConnectionDesc.ToString()));
 }
 
 void LatencyTest(FStringView InUrl, FStringView InPath, uint32 InTimeOutMs, TArrayView<int32> OutResults)
diff --git a/Engine/Source/Runtime/Experimental/IoStore/OnDemand/Private/OnDemandHttpClient.cpp b/Engine/Source/Runtime/Experimental/IoStore/OnDemand/Private/OnDemandHttpClient.cpp
--- a/Engine/Source/Runtime/Experimental/IoStore/OnDemand/Private/OnDemandHttpClient.cpp
+++ b/Engine/Source/Runtime/Experimental/IoStore/OnDemand/Private/OnDemandHttpClient.cpp
@@ -253,15 +253,15 @@ void FHttpClient::IssueRequest(FRequestParams&& Params)
 				}
 				else
 				{
-					LogHttpResult(*GetEndpointUrl(Params.Endpoint), *Params.Url, StatusCode, DurationMs, 0, Params.Range.GetOffset(), Status.GetErrorReason());
-					Params.Callback(FIoStatus(EIoErrorCode::ReadError, FString(Status.GetErrorReason())), DurationMs);
+					LogHttpResult(*GetEndpointUrl(Params.Endpoint), *Params.Url, StatusCode, DurationMs, 0, Params.Range.GetOffset(), Status.GetError().Reason);
+					Params.Callback(FIoStatus(EIoErrorCode::ReadError, FString(Status.GetError().Reason)), DurationMs);
 				}
 				break;
 			}
 			default:
 			{
 				LogHttpResult(*GetEndpointUrl(Params.Endpoint), *Params.Url, StatusCode, DurationMs, 0, Params.Range.GetOffset(), "cancelled");
-				Params.Callback(FIoStatus(EIoErrorCode::Cancelled, FString(Status.GetErrorReason())), DurationMs);
+				Params.Callback(FIoStatus(EIoErrorCode::Cancelled, FString(Status.GetError().Reason)), DurationMs);
 				break;
 			}
 			}
@@ -357,7 +357,7 @@ TIoStatusOr<FIoBuffer> FHttpClient::Get(FAnsiStringView Url, uint32 RetryCount,
 
 			if (Status.GetId() == FTicketStatus::EId::Error)
 			{
-				Reason << FString(Status.GetErrorReason());
+				Reason << Status.GetError().Reason;
 			}
 		});
 
@@ -526,7 +526,7 @@ void FMultiEndpointHttpClient::IssueRequest(FRequest&& Request)
 
 					if (bError)
 					{
-						Response.Reason = FString(TicketStatus.GetErrorReason());
+						Response.Reason = TicketStatus.GetError().Reason;
 					}
 					else if (TicketStatus.GetId() == FTicketStatus::EId::Cancelled)
 					{
