Dataflow : consolidate all the place where dataflow compute the typename from a FProperty and make sure we trim the spaces from it
- also add support for arrays for variables

diff --git a/Engine/Source/Runtime/Experimental/Dataflow/Core/Private/Dataflow/DataflowConnection.cpp b/Engine/Source/Runtime/Experimental/Dataflow/Core/Private/Dataflow/DataflowConnection.cpp
--- a/Engine/Source/Runtime/Experimental/Dataflow/Core/Private/Dataflow/DataflowConnection.cpp
+++ b/Engine/Source/Runtime/Experimental/Dataflow/Core/Private/Dataflow/DataflowConnection.cpp
@@ -12,31 +12,35 @@ const UE::Dataflow::FConnectionKey UE::Dataflow::FConnectionKey::Invalid = { (ui
 const UE::Dataflow::FPin UE::Dataflow::FPin::InvalidPin = { UE::Dataflow::FPin::EDirection::NONE, NAME_None, NAME_None };
 
 FDataflowConnection::FDataflowConnection(UE::Dataflow::FPin::EDirection InDirection, FName InType, FName InName, FDataflowNode* InOwningNode, const FProperty* InProperty, FGuid InGuid)
-	: Type(InType)
-	, OriginalType(InType)
-	, OwningNode(InOwningNode)
+	: OwningNode(InOwningNode)
 	, Property(InProperty)
 	, Guid(InGuid)
 	, Direction(InDirection)
+	, OriginalType(InType)
 	, Name(InName)
 {
+	ensure(!OriginalType.ToString().Contains(TEXT(" ")));
+	ensure(!Type.ToString().Contains(TEXT(" ")));
 	// this cannot be in the initializer because of  warning V670 about order of initialization
 	Offset = OwningNode ? OwningNode->GetPropertyOffset(Name) : INDEX_NONE;
-
+	SetTypeInternal(InType);
 	InitFromType();
 }
 
 FDataflowConnection::FDataflowConnection(UE::Dataflow::FPin::EDirection InDirection, const UE::Dataflow::FConnectionParameters& Params)
-	: Type(Params.Type)
-	, OriginalType(Params.Type)
-	, OwningNode(Params.Owner)
+	: OwningNode(Params.Owner)
 	, Property(Params.Property)
 	, Guid(Params.Guid)
 	, Offset(Params.Offset)
 	, Direction(InDirection)
+	, OriginalType(Params.Type)
 	, Name(Params.Name)
 {
+	ensure(!OriginalType.ToString().Contains(TEXT(" ")));
+	ensure(!Type.ToString().Contains(TEXT(" ")));
+	SetTypeInternal(Params.Type);
 	InitFromType();
+
 }
 
 void FDataflowConnection::InitFromType()
@@ -59,6 +63,19 @@ void FDataflowConnection::InitFromType()
 #endif
 }
 
+FName FDataflowConnection::GetTypeNameFromProperty(const FProperty* Property)
+{
+	if (Property)
+	{
+		FString ExtendedType;
+		const FString CPPType = Property->GetCPPType(&ExtendedType);
+		ExtendedType.RemoveSpacesInline();
+		return FName(CPPType + ExtendedType);
+	}
+	return NAME_None;
+}
+
+
 bool FDataflowConnection::IsOwningNodeEnabled() const
 {
 	return (OwningNode && OwningNode->IsActive());
@@ -84,12 +101,30 @@ bool FDataflowConnection::IsAnyType(const FName& InType)
 	return UE::Dataflow::FAnyTypesRegistry::IsAnyTypeStatic(InType);
 }
 
+void FDataflowConnection::SetTypeInternal(FName NewType)
+{
+	if (Type != NewType)
+	{
+		FString TypeAsString{ NewType.ToString() };
+		if (TypeAsString.Contains(TEXT(" ")))
+		{
+			TypeAsString.RemoveSpacesInline();
+			Type = FName(TypeAsString);
+		}
+		else
+		{
+			Type = NewType;
+		}
+	}
+	ensure(!Type.ToString().Contains(TEXT(" ")));
+}
+
 void FDataflowConnection::SetAsAnyType(bool bAnyType, const FName& ConcreteType)
 {
 	bIsAnyType = bAnyType;
 	if (bIsAnyType)
 	{
-		Type = ConcreteType;
+		SetTypeInternal(ConcreteType);
 		bHasConcreteType = !IsAnyType(ConcreteType);
 	}
 }
@@ -155,7 +190,7 @@ bool FDataflowConnection::ResetToOriginalType()
 
 	if (IsSafeToTryChangingType())
 	{
-		Type = OriginalType;
+		SetTypeInternal(OriginalType);
 		bHasConcreteType = false;
 		return true;
 	}
@@ -186,7 +221,7 @@ bool FDataflowConnection::SetConcreteType(FName InType)
 		const bool bExtendedType = IsExtendedType(InType);
 		if (bExtendedType)
 		{
-			Type = InType;
+			SetTypeInternal(InType);
 			bHasConcreteType = true;
 			return true;
 		}
@@ -195,7 +230,7 @@ bool FDataflowConnection::SetConcreteType(FName InType)
 		{
 			if (ensure(SupportsType(InType)))
 			{
-				Type = InType;
+				SetTypeInternal(InType);
 				bHasConcreteType = true;
 				return true;
 			}
@@ -216,17 +251,14 @@ void FDataflowConnection::SetTypePolicy(IDataflowTypePolicy* InTypePolicy)
 void FDataflowConnection::ForceSimpleType(FName InType)
 {
 	check(Type.ToString().StartsWith(InType.ToString()));
-	Type = InType;
+	SetTypeInternal(InType);
 	bHasConcreteType = true;
 }
 
 void FDataflowConnection::FixAndPropagateType()
 {
 	check(Property);
-	FString ExtendedType;
-	const FString CPPType = Property->GetCPPType(&ExtendedType);
-	FName FixedType(CPPType + ExtendedType);
-
+	const FName FixedType = GetTypeNameFromProperty(Property);
 	FixAndPropagateType(FixedType);
 }
 
diff --git a/Engine/Source/Runtime/Experimental/Dataflow/Core/Private/Dataflow/DataflowGraph.cpp b/Engine/Source/Runtime/Experimental/Dataflow/Core/Private/Dataflow/DataflowGraph.cpp
--- a/Engine/Source/Runtime/Experimental/Dataflow/Core/Private/Dataflow/DataflowGraph.cpp
+++ b/Engine/Source/Runtime/Experimental/Dataflow/Core/Private/Dataflow/DataflowGraph.cpp
@@ -544,7 +544,12 @@ namespace UE::Dataflow
 								}
 								if (!AddTemplateTypedConnectionToBeFixed(Output, ArType))
 								{
-									check(Output->GetType() == ArType || bIsAnyType);
+									if (Output->GetType() != ArType)
+									{
+										FString NoSpaceArType = ArType.ToString();
+										NoSpaceArType.RemoveSpacesInline();
+										check(Output->GetType() == *NoSpaceArType || bIsAnyType);
+									}
 								}
 								Output->SetPinIsHidden(bIsHidden);
 								Output->SetGuid(ArGuid);
@@ -594,7 +599,12 @@ namespace UE::Dataflow
 								}
 								if (!AddTemplateTypedConnectionToBeFixed(Input, ArType))
 								{
-									check(Input->GetType() == ArType || bIsAnyType);
+									if (Input->GetType() != ArType)
+									{
+										FString NoSpaceArType = ArType.ToString();
+										NoSpaceArType.RemoveSpacesInline();
+										check(Input->GetType() == *NoSpaceArType || bIsAnyType);
+									}
 								}
 								Input->SetPinIsHidden(bIsHidden);
 								Input->SetGuid(ArGuid);
diff --git a/Engine/Source/Runtime/Experimental/Dataflow/Core/Private/Dataflow/DataflowInputOutput.cpp b/Engine/Source/Runtime/Experimental/Dataflow/Core/Private/Dataflow/DataflowInputOutput.cpp
--- a/Engine/Source/Runtime/Experimental/Dataflow/Core/Private/Dataflow/DataflowInputOutput.cpp
+++ b/Engine/Source/Runtime/Experimental/Dataflow/Core/Private/Dataflow/DataflowInputOutput.cpp
@@ -34,7 +34,7 @@ bool FDataflowInput::IsConnected() const
 
 bool FDataflowInput::AddConnection(FDataflowConnection* InOutput)
 {
-	if (Type == InOutput->GetType())
+	if (GetType() == InOutput->GetType())
 	{
 		Connection = (FDataflowOutput*)InOutput;
 		GetOwningNode()->Invalidate();
@@ -98,12 +98,12 @@ void FDataflowInput::PullValue(UE::Dataflow::FContext& Context) const
 
 void FDataflowInput::FixAndPropagateType(FName InType)
 {
-	check(InType.ToString().StartsWith(Type.ToString()));
+	check(InType.ToString().StartsWith(GetType().ToString()));
 	check(!FDataflowConnection::IsAnyType(InType));
 	
 	if (GetType() != InType)
 	{
-		Type = InType;
+		SetTypeInternal(InType);
 		bHasConcreteType = true;
 
 		// if we have a reroute node propagate through to make sure each reroute segment is properly handled
@@ -199,7 +199,7 @@ TArray< FDataflowInput*> FDataflowOutput::GetConnectedInputs()
 
 bool FDataflowOutput::AddConnection(FDataflowConnection* InOutput)
 {
-	if (Type == InOutput->GetType())
+	if (GetType() == InOutput->GetType())
 	{
 		Connections.Add((FDataflowInput*)InOutput);
 		return true;
@@ -390,12 +390,12 @@ void FDataflowOutput::SetNullValue(UE::Dataflow::FContext& Context) const
 
 void FDataflowOutput::FixAndPropagateType(FName InType)
 {
-	check(InType.ToString().StartsWith(Type.ToString()));
+	check(InType.ToString().StartsWith(GetType().ToString()));
 	check(!FDataflowConnection::IsAnyType(InType));
 
-	if (Type != InType)
+	if (GetType() != InType)
 	{
-		Type = InType;
+		SetTypeInternal(InType);
 		bHasConcreteType = true;
 
 		// if we have a reroute node propagate through to make sure each reroute segment is properly handled
diff --git a/Engine/Source/Runtime/Experimental/Dataflow/Core/Private/Dataflow/DataflowNode.cpp b/Engine/Source/Runtime/Experimental/Dataflow/Core/Private/Dataflow/DataflowNode.cpp
--- a/Engine/Source/Runtime/Experimental/Dataflow/Core/Private/Dataflow/DataflowNode.cpp
+++ b/Engine/Source/Runtime/Experimental/Dataflow/Core/Private/Dataflow/DataflowNode.cpp
@@ -883,9 +883,8 @@ void FDataflowNode::InitConnectionParametersFromPropertyReference(const FStructO
 	TArray<const FProperty*> PropertyChain;
 	const FProperty& Property = FindPropertyChecked(Struct, PropertyRef, PropertyName, &PropertyChain);
 	check(PropertyChain.Num());
-	FString ExtendedType;
-	const FString CPPType = Property.GetCPPType(&ExtendedType);
-	OutParams.Type = FName(CPPType + ExtendedType);
+
+	OutParams.Type = FDataflowConnection::GetTypeNameFromProperty(&Property);
 	OutParams.Name = GetPropertyFullName(PropertyChain);
 	OutParams.Property = &Property;
 	OutParams.Owner = this;
@@ -949,9 +948,7 @@ FDataflowInput& FDataflowNode::RegisterInputArrayConnectionInternal(const UE::Da
 	check(InputParams.Property);
 
 	PropertyChain.Append(ArrayPropertyChain);
-	FString ExtendedType;
-	const FString CPPType = InputParams.Property->GetCPPType(&ExtendedType);
-	InputParams.Type = FName(CPPType + ExtendedType);
+	InputParams.Type = FDataflowConnection::GetTypeNameFromProperty(InputParams.Property);
 	InputParams.Name = GetPropertyFullName(PropertyChain, Reference.Index);
 	InputParams.InnerOffset = UE::Dataflow::Private::GetArrayElementOffsetFromReference(InputParams.ArrayProperty, Reference);
 
@@ -1004,9 +1001,7 @@ FDataflowOutput& FDataflowNode::RegisterOutputArrayConnectionInternal(const UE::
 	check(OutputParams.Property);
 
 	PropertyChain.Append(ArrayPropertyChain);
-	FString ExtendedType;
-	const FString CPPType = OutputParams.Property->GetCPPType(&ExtendedType);
-	OutputParams.Type = FName(CPPType + ExtendedType);
+	OutputParams.Type = FDataflowConnection::GetTypeNameFromProperty(OutputParams.Property);
 	OutputParams.Name = GetPropertyFullName(PropertyChain, Reference.Index);
 	OutputParams.InnerOffset = UE::Dataflow::Private::GetArrayElementOffsetFromReference(OutputParams.ArrayProperty, Reference);
 
diff --git a/Engine/Source/Runtime/Experimental/Dataflow/Core/Private/Dataflow/DataflowNodeParameters.cpp b/Engine/Source/Runtime/Experimental/Dataflow/Core/Private/Dataflow/DataflowNodeParameters.cpp
--- a/Engine/Source/Runtime/Experimental/Dataflow/Core/Private/Dataflow/DataflowNodeParameters.cpp
+++ b/Engine/Source/Runtime/Experimental/Dataflow/Core/Private/Dataflow/DataflowNodeParameters.cpp
@@ -156,9 +156,7 @@ namespace UE::Dataflow
 				if (Elem.Value && Elem.Value->Property && Elem.Value->Type == FContextCacheElementBase::EType::CacheElementTyped)
 				{
 					FProperty* Property = (FProperty*)Elem.Value->Property;
-					FString ExtendedType;
-					const FString CPPType = Property->GetCPPType(&ExtendedType);
-					FName TypeName(CPPType + ExtendedType);
+					FName TypeName = FDataflowConnection::GetTypeNameFromProperty(Property);
 					FGuid NodeGuid = Elem.Value->NodeGuid;
 					uint32 NodeHash = Elem.Value->NodeHash;
 
diff --git a/Engine/Source/Runtime/Experimental/Dataflow/Core/Public/Dataflow/DataflowConnection.h b/Engine/Source/Runtime/Experimental/Dataflow/Core/Public/Dataflow/DataflowConnection.h
--- a/Engine/Source/Runtime/Experimental/Dataflow/Core/Public/Dataflow/DataflowConnection.h
+++ b/Engine/Source/Runtime/Experimental/Dataflow/Core/Public/Dataflow/DataflowConnection.h
@@ -132,8 +132,6 @@ struct FDataflowConnection
 	GENERATED_USTRUCT_BODY()
 
 protected:
-	FName Type;
-	FName OriginalType;
 	FDataflowNode* OwningNode = nullptr;
 	const FProperty* Property = nullptr;
 	FGuid  Guid;
@@ -146,6 +144,8 @@ protected:
 	bool bPinIsHidden:1 = false;
 
 private:
+	FName OriginalType;
+	FName Type;
 	FName Name;
 
 	/** Lock the type from being changed */
@@ -160,6 +160,8 @@ private:
 	friend class UE::Dataflow::FGraph;
 
 protected:
+		DATAFLOWCORE_API void SetTypeInternal(FName NewType);
+
 	DATAFLOWCORE_API bool IsOwningNodeEnabled() const;
 
 	/** this should only be used for serialization */
@@ -287,6 +289,13 @@ public:
 		return *this;
 	}
 
+	/**
+	* extract the type name from a FProperty 
+	* this method make sure the resulting typename is consistent with Dataflow standard where there's no space in the name
+	* for example FProperty returns : "TArray<TObjectPtr<UObject> >" but dataflow compatible typename is "TArray<TObjectPtr<UObject>>"
+	*/
+	DATAFLOWCORE_API static FName GetTypeNameFromProperty(const FProperty* Property);
+
 private:
 	void InitFromType();
 };
diff --git a/Engine/Source/Runtime/Experimental/Dataflow/Engine/Private/Dataflow/DataflowEdNode.cpp b/Engine/Source/Runtime/Experimental/Dataflow/Engine/Private/Dataflow/DataflowEdNode.cpp
--- a/Engine/Source/Runtime/Experimental/Dataflow/Engine/Private/Dataflow/DataflowEdNode.cpp
+++ b/Engine/Source/Runtime/Experimental/Dataflow/Engine/Private/Dataflow/DataflowEdNode.cpp
@@ -582,6 +582,28 @@ void UDataflowEdNode::Serialize(FArchive& Ar)
 			}
 		}
 	}
+
+	// some double level template type have unwanted spaces 
+	// now that the dataflow connection are no longer having spaces in their type name we need to fix that up on the pins
+	RemoveSpacesInAllPinTypes();
+#endif
+}
+
+void UDataflowEdNode::RemoveSpacesInAllPinTypes()
+{
+#if WITH_EDITOR
+	for (UEdGraphPin* Pin : this->GetAllPins())
+	{
+		if (Pin)
+		{
+			FString PinTypeString{ Pin->PinType.PinCategory.ToString() };
+			if (PinTypeString.Contains(TEXT(" ")))
+			{
+				PinTypeString.RemoveSpacesInline();
+				Pin->PinType.PinCategory = FName(PinTypeString);
+			}
+		}
+	}
 #endif
 }
 
diff --git a/Engine/Source/Runtime/Experimental/Dataflow/Engine/Private/Dataflow/DataflowVariableNodes.cpp b/Engine/Source/Runtime/Experimental/Dataflow/Engine/Private/Dataflow/DataflowVariableNodes.cpp
--- a/Engine/Source/Runtime/Experimental/Dataflow/Engine/Private/Dataflow/DataflowVariableNodes.cpp
+++ b/Engine/Source/Runtime/Experimental/Dataflow/Engine/Private/Dataflow/DataflowVariableNodes.cpp
@@ -110,10 +110,9 @@ void FGetDataflowVariableNode::UpdateOutputTypes(const FPropertyBagPropertyDesc&
 
 	Output->UnlockType();
 
-	FString ExtendedType;
-	const FString CPPType = Desc.CachedProperty->GetCPPType(&ExtendedType);
-	const FName ConcreteType(CPPType + ExtendedType);
+	ContainerType = Desc.ContainerTypes.GetFirstContainerType();
 
+	const FName ConcreteType = FDataflowConnection::GetTypeNameFromProperty(Desc.CachedProperty);
 	switch (Desc.ValueType)
 	{
 	case EPropertyBagPropertyType::Bool:
@@ -248,103 +247,315 @@ void FGetDataflowVariableNode::Evaluate(UE::Dataflow::FContext& Context, const F
 
 /*static*/ void FGetDataflowVariableNode::EvaluateBool(const FGetDataflowVariableNode& Node, const FInstancedPropertyBag& Variables, UE::Dataflow::FContext& Context, const FDataflowOutput& Out)
 {
-	bool OutBool = false;
-	if (const bool* Value = Variables.GetValueBool(Node.VariableName).TryGetValue())
+	switch (Node.ContainerType)
 	{
-		OutBool = *Value;
+	case EPropertyBagContainerType::None:
+	{
+		bool OutBool = false;
+		if (const bool* Value = Variables.GetValueBool(Node.VariableName).TryGetValue())
+		{
+			OutBool = *Value;
+		}
+		Out.SetValue<bool>(MoveTemp(OutBool), Context);
+		break;
+	}
+	case EPropertyBagContainerType::Array:
+	{
+		TArray<bool> OutBoolArray;
+		if (const FPropertyBagArrayRef* ArrayRef = Variables.GetArrayRef(Node.VariableName).TryGetValue())
+		{
+			OutBoolArray.SetNum(ArrayRef->Num());
+			for (int32 Index = 0; Index < ArrayRef->Num(); Index++)
+			{
+				OutBoolArray[Index] = ArrayRef->GetValueBool(Index).GetValue();
+			}
+		}
+		Out.SetValue(MoveTemp(OutBoolArray), Context);
+		break;
+	}
 	}
-	Out.SetValue<bool>(MoveTemp(OutBool), Context);
 }
 
 /*static*/ void FGetDataflowVariableNode::EvaluateByte(const FGetDataflowVariableNode& Node, const FInstancedPropertyBag& Variables, UE::Dataflow::FContext& Context, const FDataflowOutput& Out)
 {
-	uint8 OutUInt8 = 0;
-	if (const uint8* Value = Variables.GetValueByte(Node.VariableName).TryGetValue())
+	switch (Node.ContainerType)
 	{
-		OutUInt8 = *Value;
+	case EPropertyBagContainerType::None:
+	{
+		uint8 OutUInt8 = 0;
+		if (const uint8* Value = Variables.GetValueByte(Node.VariableName).TryGetValue())
+		{
+			OutUInt8 = *Value;
+		}
+		Out.SetValue<uint8>(MoveTemp(OutUInt8), Context);
+		break;
+	}
+	case EPropertyBagContainerType::Array:
+	{
+		TArray<uint8> OutInt8Array;
+		if (const FPropertyBagArrayRef* ArrayRef = Variables.GetArrayRef(Node.VariableName).TryGetValue())
+		{
+			OutInt8Array.SetNum(ArrayRef->Num());
+			for (int32 Index = 0; Index < ArrayRef->Num(); Index++)
+			{
+				OutInt8Array[Index] = ArrayRef->GetValueByte(Index).GetValue();
+			}
+		}
+		Out.SetValue(MoveTemp(OutInt8Array), Context);
+		break;
+	}
 	}
-	Out.SetValue<uint8>(MoveTemp(OutUInt8), Context);
 }
 
 /*static*/ void FGetDataflowVariableNode::EvaluateInt32(const FGetDataflowVariableNode& Node, const FInstancedPropertyBag& Variables, UE::Dataflow::FContext& Context, const FDataflowOutput& Out)
 {
-	int32 OutInt32 = 0;
-	if (const int32* Value = Variables.GetValueInt32(Node.VariableName).TryGetValue())
+	switch (Node.ContainerType)
 	{
-		OutInt32 = *Value;
+	case EPropertyBagContainerType::None:
+	{
+		int32 OutInt32 = 0;
+		if (const int32* Value = Variables.GetValueInt32(Node.VariableName).TryGetValue())
+		{
+			OutInt32 = *Value;
+		}
+		Out.SetValue<int32>(MoveTemp(OutInt32), Context);
+		break;
+	}
+	case EPropertyBagContainerType::Array:
+	{
+		TArray<int32> OutInt32Array;
+		if (const FPropertyBagArrayRef* ArrayRef = Variables.GetArrayRef(Node.VariableName).TryGetValue())
+		{
+			OutInt32Array.SetNum(ArrayRef->Num());
+			for (int32 Index = 0; Index < ArrayRef->Num(); Index++)
+			{
+				OutInt32Array[Index] = ArrayRef->GetValueInt32(Index).GetValue();
+			}
+		}
+		Out.SetValue(MoveTemp(OutInt32Array), Context);
+		break;
+	}
 	}
-	Out.SetValue<int32>(MoveTemp(OutInt32), Context);
 }
 
 /*static*/ void FGetDataflowVariableNode::EvaluateInt64(const FGetDataflowVariableNode& Node, const FInstancedPropertyBag& Variables, UE::Dataflow::FContext& Context, const FDataflowOutput& Out)
 {
-	int64 OutInt64 = 0;
-	if (const int64* Value = Variables.GetValueInt64(Node.VariableName).TryGetValue())
+	switch (Node.ContainerType)
+	{
+	case EPropertyBagContainerType::None:
 	{
-		OutInt64 = *Value;
+		int64 OutInt64 = 0;
+		if (const int64* Value = Variables.GetValueInt64(Node.VariableName).TryGetValue())
+		{
+			OutInt64 = *Value;
+		}
+		Out.SetValue<int64>(MoveTemp(OutInt64), Context);
+		break;
+	}
+	case EPropertyBagContainerType::Array:
+	{
+		TArray<int64> OutInt64Array;
+		if (const FPropertyBagArrayRef* ArrayRef = Variables.GetArrayRef(Node.VariableName).TryGetValue())
+		{
+			OutInt64Array.SetNum(ArrayRef->Num());
+			for (int32 Index = 0; Index < ArrayRef->Num(); Index++)
+			{
+				OutInt64Array[Index] = ArrayRef->GetValueInt64(Index).GetValue();
+			}
+		}
+		Out.SetValue(MoveTemp(OutInt64Array), Context);
+		break;
+	}
 	}
-	Out.SetValue<int64>(MoveTemp(OutInt64), Context);
 }
 
 /*static*/ void FGetDataflowVariableNode::EvaluateFloat(const FGetDataflowVariableNode& Node, const FInstancedPropertyBag& Variables, UE::Dataflow::FContext& Context, const FDataflowOutput& Out)
 {
-	float OutFloat = 0.0f;
-	if (const float* Value = Variables.GetValueFloat(Node.VariableName).TryGetValue())
+	switch (Node.ContainerType)
 	{
-		OutFloat = *Value;
+	case EPropertyBagContainerType::None:
+	{
+		float OutFloat = 0.0f;
+		if (const float* Value = Variables.GetValueFloat(Node.VariableName).TryGetValue())
+		{
+			OutFloat = *Value;
+		}
+		Out.SetValue<float>(MoveTemp(OutFloat), Context);
+		break;
+	}
+	case EPropertyBagContainerType::Array:
+	{
+		TArray<float> OutFloatArray;
+		if (const FPropertyBagArrayRef* ArrayRef = Variables.GetArrayRef(Node.VariableName).TryGetValue())
+		{
+			OutFloatArray.SetNum(ArrayRef->Num());
+			for (int32 Index = 0; Index < ArrayRef->Num(); Index++)
+			{
+				OutFloatArray[Index] = ArrayRef->GetValueFloat(Index).GetValue();
+			}
+		}
+		Out.SetValue(MoveTemp(OutFloatArray), Context);
+		break;
+	}
 	}
-	Out.SetValue<float>(MoveTemp(OutFloat), Context);
 }
 
 /*static*/ void FGetDataflowVariableNode::EvaluateDouble(const FGetDataflowVariableNode& Node, const FInstancedPropertyBag& Variables, UE::Dataflow::FContext& Context, const FDataflowOutput& Out)
 {
 	// UI shows type as float but internally always uses a double and never float 
 	// so let's output as float 
-	float OutFloat = 0.0;
-	if (const double* Value = Variables.GetValueDouble(Node.VariableName).TryGetValue())
+	switch (Node.ContainerType)
+	{
+	case EPropertyBagContainerType::None:
 	{
-		OutFloat = *Value;
+		float OutFloat = 0.0;
+		if (const double* Value = Variables.GetValueDouble(Node.VariableName).TryGetValue())
+		{
+			OutFloat = *Value;
+		}
+		Out.SetValue<float>(MoveTemp(OutFloat), Context);
+		break;
+	}
+	case EPropertyBagContainerType::Array:
+	{
+		TArray<float> OutFloatArray;
+		if (const FPropertyBagArrayRef* ArrayRef = Variables.GetArrayRef(Node.VariableName).TryGetValue())
+		{
+			OutFloatArray.SetNum(ArrayRef->Num());
+			for (int32 Index = 0; Index < ArrayRef->Num(); Index++)
+			{
+				OutFloatArray[Index] = static_cast<float>(ArrayRef->GetValueDouble(Index).GetValue());
+			}
+		}
+		Out.SetValue(MoveTemp(OutFloatArray), Context);
+		break;
+	}
 	}
-	Out.SetValue<float>(MoveTemp(OutFloat), Context);
 }
 
 /*static*/ void FGetDataflowVariableNode::EvaluateName(const FGetDataflowVariableNode& Node, const FInstancedPropertyBag& Variables, UE::Dataflow::FContext& Context, const FDataflowOutput& Out)
 {
-	FName OutName;
-	if (const FName* Value = Variables.GetValueName(Node.VariableName).TryGetValue())
+	switch (Node.ContainerType)
+	{
+	case EPropertyBagContainerType::None:
 	{
-		OutName = *Value;
+		FName OutName;
+		if (const FName* Value = Variables.GetValueName(Node.VariableName).TryGetValue())
+		{
+			OutName = *Value;
+		}
+		Out.SetValue<FName>(MoveTemp(OutName), Context);
+		break;
+	}
+	case EPropertyBagContainerType::Array:
+	{
+		TArray<FName> OutNameArray;
+		if (const FPropertyBagArrayRef* ArrayRef = Variables.GetArrayRef(Node.VariableName).TryGetValue())
+		{
+			OutNameArray.SetNum(ArrayRef->Num());
+			for (int32 Index = 0; Index < ArrayRef->Num(); Index++)
+			{
+				OutNameArray[Index] = ArrayRef->GetValueName(Index).GetValue();
+			}
+		}
+		Out.SetValue(MoveTemp(OutNameArray), Context);
+		break;
+	}
 	}
-	Out.SetValue<FName>(MoveTemp(OutName), Context);
 }
 
 /*static*/ void FGetDataflowVariableNode::EvaluateString(const FGetDataflowVariableNode& Node, const FInstancedPropertyBag& Variables, UE::Dataflow::FContext& Context, const FDataflowOutput& Out)
 {
-	FString OutString;
-	if (const FString* Value = Variables.GetValueString(Node.VariableName).TryGetValue())
+	switch (Node.ContainerType)
+	{
+	case EPropertyBagContainerType::None:
 	{
-		OutString = *Value;
+		FString OutString;
+		if (const FString* Value = Variables.GetValueString(Node.VariableName).TryGetValue())
+		{
+			OutString = *Value;
+		}
+		Out.SetValue<FString>(MoveTemp(OutString), Context);
+		break;
+	}
+	case EPropertyBagContainerType::Array:
+	{
+		TArray<FString> OutStringArray;
+		if (const FPropertyBagArrayRef* ArrayRef = Variables.GetArrayRef(Node.VariableName).TryGetValue())
+		{
+			OutStringArray.SetNum(ArrayRef->Num());
+			for (int32 Index = 0; Index < ArrayRef->Num(); Index++)
+			{
+				OutStringArray[Index] = ArrayRef->GetValueString(Index).GetValue();
+			}
+		}
+		Out.SetValue(MoveTemp(OutStringArray), Context);
+		break;
+	}
 	}
-	Out.SetValue<FString>(MoveTemp(OutString), Context);
 }
 
 /*static*/ void FGetDataflowVariableNode::EvaluateText(const FGetDataflowVariableNode& Node, const FInstancedPropertyBag& Variables, UE::Dataflow::FContext& Context, const FDataflowOutput& Out)
 {
-	FString OutString;
-	if (const FText* Value = Variables.GetValueText(Node.VariableName).TryGetValue())
+	switch (Node.ContainerType)
 	{
-		OutString = Value->ToString();
+	case EPropertyBagContainerType::None:
+	{
+		FString OutString;
+		if (const FText* Value = Variables.GetValueText(Node.VariableName).TryGetValue())
+		{
+			OutString = Value->ToString();
+		}
+		Out.SetValue<FString>(MoveTemp(OutString), Context);
+		break;
+	}
+	case EPropertyBagContainerType::Array:
+	{
+		TArray<FString> OutStringArray;
+		if (const FPropertyBagArrayRef* ArrayRef = Variables.GetArrayRef(Node.VariableName).TryGetValue())
+		{
+			OutStringArray.SetNum(ArrayRef->Num());
+			for (int32 Index = 0; Index < ArrayRef->Num(); Index++)
+			{
+				OutStringArray[Index] = ArrayRef->GetValueText(Index).GetValue().ToString();
+			}
+		}
+		Out.SetValue(MoveTemp(OutStringArray), Context);
+		break;
+	}
 	}
-	Out.SetValue<FString>(MoveTemp(OutString), Context);
 }
 
 /*static*/ void FGetDataflowVariableNode::EvaluateObject(const FGetDataflowVariableNode& Node, const FInstancedPropertyBag& Variables, UE::Dataflow::FContext& Context, const FDataflowOutput& Out)
 {
-	TObjectPtr<UObject> OutObject = nullptr;
-	if (UObject* const* Value = Variables.GetValueObject<UObject>(Node.VariableName).TryGetValue())
+	switch(Node.ContainerType)
+	{
+	case EPropertyBagContainerType::None:
 	{
-		OutObject = *Value;
+		TObjectPtr<UObject> OutObject = nullptr;
+		if (UObject* const* Value = Variables.GetValueObject<UObject>(Node.VariableName).TryGetValue())
+		{
+			OutObject = *Value;
+		}
+		Out.SetValue<TObjectPtr<UObject>>(MoveTemp(OutObject), Context);
+		break;
 	}
-	Out.SetValue<TObjectPtr<UObject>>(MoveTemp(OutObject), Context);
+	case EPropertyBagContainerType::Array:
+	{
+		TArray<TObjectPtr<UObject>> OutObjectArray;
+		if (const FPropertyBagArrayRef* ArrayRef = Variables.GetArrayRef(Node.VariableName).TryGetValue())
+		{
+			OutObjectArray.SetNum(ArrayRef->Num());
+			for (int32 Index = 0; Index < ArrayRef->Num(); Index++)
+			{
+				OutObjectArray[Index] = ArrayRef->GetValueObject(Index).GetValue();
+			}
+		}
+		Out.SetValue(MoveTemp(OutObjectArray), Context);
+		break;
+	}
+	}
+
+
 }
 
diff --git a/Engine/Source/Runtime/Experimental/Dataflow/Engine/Public/Dataflow/DataflowEdNode.h b/Engine/Source/Runtime/Experimental/Dataflow/Engine/Public/Dataflow/DataflowEdNode.h
--- a/Engine/Source/Runtime/Experimental/Dataflow/Engine/Public/Dataflow/DataflowEdNode.h
+++ b/Engine/Source/Runtime/Experimental/Dataflow/Engine/Public/Dataflow/DataflowEdNode.h
@@ -113,9 +113,6 @@ private:
 	UPROPERTY(Transient)
 	bool bCanEnableRenderWireframe = true;
 
-	void OnPinConnection(const UEdGraphPin& Pin, const FEdGraphPinType& Type);
-	void OnPinDisconnection(const UEdGraphPin& Pin);
-	void SetAnyTypePinsToType(const FEdGraphPinType& Type);
-
+	void RemoveSpacesInAllPinTypes();
 };
 
diff --git a/Engine/Source/Runtime/Experimental/Dataflow/Engine/Public/Dataflow/DataflowVariableNodes.h b/Engine/Source/Runtime/Experimental/Dataflow/Engine/Public/Dataflow/DataflowVariableNodes.h
--- a/Engine/Source/Runtime/Experimental/Dataflow/Engine/Public/Dataflow/DataflowVariableNodes.h
+++ b/Engine/Source/Runtime/Experimental/Dataflow/Engine/Public/Dataflow/DataflowVariableNodes.h
@@ -61,6 +61,8 @@ private:
 	using FStaticEvaluationFunctionPtr = void(*)(const FGetDataflowVariableNode&, const FInstancedPropertyBag&, UE::Dataflow::FContext&, const FDataflowOutput&);
 	FStaticEvaluationFunctionPtr EvaluateFunction = nullptr;
 
+	EPropertyBagContainerType ContainerType = EPropertyBagContainerType::None;
+
 	void OnObjectPropertyChanged(UObject* InObject, FPropertyChangedEvent& InPropertyChangedEvent);
 
 	TWeakObjectPtr<UDataflow> WeakDataflowPtr;
