Chaosflesh:
-Added SetMuscleActivationParameter dataflow node to customize per-muscle activation parameters. Supported linear activation/contraction model
-Unified DeltaTime/NumFrames in dataflow simulation, fixed scroll bar time
-Fixed preroll bug from async dataflow simulation: preroll has very small dt and Hessian determinant epsilon check failed. Now scaling Hessian before computing determinant.

diff --git a/Engine/Plugins/Experimental/ChaosCaching/Source/ChaosCaching/Private/Chaos/CacheManagerActor.cpp b/Engine/Plugins/Experimental/ChaosCaching/Source/ChaosCaching/Private/Chaos/CacheManagerActor.cpp
--- a/Engine/Plugins/Experimental/ChaosCaching/Source/ChaosCaching/Private/Chaos/CacheManagerActor.cpp
+++ b/Engine/Plugins/Experimental/ChaosCaching/Source/ChaosCaching/Private/Chaos/CacheManagerActor.cpp
@@ -734,7 +734,7 @@ void AChaosCacheManager::OnStartFrameChanged(Chaos::FReal InTime)
 				Observed.Cache = CacheCollection->FindCache(Observed.CacheName);
 			}
 
-			if (!Observed.Cache || !Observed.BestFitAdapter || Observed.Cache->GetDuration()==0.0 || !Observed.IsEnabled(CacheMode))
+			if (!Observed.Cache || !Observed.BestFitAdapter || !Observed.IsEnabled(CacheMode))
 			{
 				continue;
 			}
diff --git a/Engine/Plugins/Experimental/ChaosFlesh/Source/ChaosFleshEngine/Private/ChaosFlesh/ChaosDeformableTetrahedralComponent.cpp b/Engine/Plugins/Experimental/ChaosFlesh/Source/ChaosFleshEngine/Private/ChaosFlesh/ChaosDeformableTetrahedralComponent.cpp
--- a/Engine/Plugins/Experimental/ChaosFlesh/Source/ChaosFleshEngine/Private/ChaosFlesh/ChaosDeformableTetrahedralComponent.cpp
+++ b/Engine/Plugins/Experimental/ChaosFlesh/Source/ChaosFleshEngine/Private/ChaosFlesh/ChaosDeformableTetrahedralComponent.cpp
@@ -724,6 +724,7 @@ TOptional<TArray<int32>> UDeformableTetrahedralComponent::GetMeshImportVertexMap
 			const FSkelMeshSection& MLDSection = MLDLOD.Sections[SectionIndex];
 			MLDVertexCount += MLDSection.NumVertices;
 		}
+		int32 NumOffVertices = 0;
 		for (int32 SectionIndex = 0; SectionIndex < NumSections; ++SectionIndex)
 		{
 			const FSkelMeshSection& MLDSection = MLDLOD.Sections[SectionIndex];
@@ -740,11 +741,17 @@ TOptional<TArray<int32>> UDeformableTetrahedralComponent::GetMeshImportVertexMap
 				const FVector3f& FleshPosition = Positions[VertexIndex + VertexOffset];
 				if (!MLDPosition.Equals(FleshPosition, UE_KINDA_SMALL_NUMBER))
 				{
-					UE_LOG(LogDeformableTetrahedralComponentInternal, Warning, TEXT("SkeletalMeshAsset and FleshAsset have different vertex positions. Check if the assets have the same vertex order or some points are not embedded."));
+					NumOffVertices++;
 				}
 			}
 			VertexOffset += MLDSection.NumVertices;
 		}
+		if (NumOffVertices > 0)
+		{
+			UE_LOG(LogDeformableTetrahedralComponentInternal, Warning, 
+				TEXT("%d vertices bound in FleshAsset [%s] have different vertex positions than SkeletalMeshAsset [%s]. Check if the assets have the same vertex order or some points are not embedded.")
+				, NumOffVertices, *RestCollection->GetName(), *SkeletalMeshAsset->GetName());
+		}
 	}
 
 	return Map;
diff --git a/Engine/Plugins/Experimental/ChaosFlesh/Source/ChaosFleshNodes/Private/Dataflow/ChaosFleshComputeFiberFieldNode.cpp b/Engine/Plugins/Experimental/ChaosFlesh/Source/ChaosFleshNodes/Private/Dataflow/ChaosFleshComputeFiberFieldNode.cpp
--- a/Engine/Plugins/Experimental/ChaosFlesh/Source/ChaosFleshNodes/Private/Dataflow/ChaosFleshComputeFiberFieldNode.cpp
+++ b/Engine/Plugins/Experimental/ChaosFlesh/Source/ChaosFleshNodes/Private/Dataflow/ChaosFleshComputeFiberFieldNode.cpp
@@ -16,7 +16,10 @@ void FComputeFiberFieldNode::Evaluate(UE::Dataflow::FContext& Context, const FDa
 		//
 		// Gather inputs
 		//
-
+		if (!IsConnected(&Collection) || !IsConnected(&OriginIndices) || !IsConnected(&InsertionIndices))
+		{
+			return;
+		}
 		FManagedArrayCollection InCollection = GetValue<FManagedArrayCollection>(Context, &Collection);
 		TArray<int32> InOriginIndices = GetValue<TArray<int32>>(Context, &OriginIndices);
 		TArray<int32> InInsertionIndices = GetValue<TArray<int32>>(Context, &InsertionIndices);
diff --git a/Engine/Plugins/Experimental/ChaosFlesh/Source/ChaosFleshNodes/Private/Dataflow/ChaosFleshComputeMuscleActivationNode.cpp b/Engine/Plugins/Experimental/ChaosFlesh/Source/ChaosFleshNodes/Private/Dataflow/ChaosFleshComputeMuscleActivationNode.cpp
--- a/Engine/Plugins/Experimental/ChaosFlesh/Source/ChaosFleshNodes/Private/Dataflow/ChaosFleshComputeMuscleActivationNode.cpp
+++ b/Engine/Plugins/Experimental/ChaosFlesh/Source/ChaosFleshNodes/Private/Dataflow/ChaosFleshComputeMuscleActivationNode.cpp
@@ -10,14 +10,117 @@
 
 void FComputeMuscleActivationDataNode::Evaluate(UE::Dataflow::FContext& Context, const FDataflowOutput* Out) const
 {
-	if (Out->IsA<FManagedArrayCollection>(&Collection))
+	if (Out->IsA(&Collection))
 	{
-		FManagedArrayCollection InCollection = GetValue<FManagedArrayCollection>(Context, &Collection);
-		TArray<int32> InOriginIndices = GetValue<TArray<int32>>(Context, &OriginIndicesIn);
-		TArray<int32> InInsertionIndices = GetValue<TArray<int32>>(Context, &InsertionIndicesIn);
+		FManagedArrayCollection InCollection = GetValue(Context, &Collection);
+		const TArray<int32>& InOriginIndices = GetValue(Context, &OriginIndicesIn);
+		const TArray<int32>& InInsertionIndices = GetValue(Context, &InsertionIndicesIn);
 		
 		GeometryCollection::Facades::FMuscleActivationFacade FMuscleActivation(InCollection);
 		FMuscleActivation.SetUpMuscleActivation(InOriginIndices, InInsertionIndices, ContractionVolumeScale);
 		Out->SetValue(MoveTemp(InCollection), Context);
 	}
 }
+
+void FComputeMuscleActivationDataNode_v2::Evaluate(UE::Dataflow::FContext& Context, const FDataflowOutput* Out) const
+{
+	if (Out->IsA(&Collection))
+	{
+		FManagedArrayCollection InCollection = GetValue(Context, &Collection);
+		TArray<int32> InOriginIndices = GetValue(Context, &OriginIndicesIn);
+		TArray<int32> InInsertionIndices = GetValue(Context, &InsertionIndicesIn);
+
+		GeometryCollection::Facades::FMuscleActivationFacade FMuscleActivation(InCollection);
+		FMuscleActivation.SetUpMuscleActivation(InOriginIndices, InInsertionIndices);
+		Out->SetValue(MoveTemp(InCollection), Context);
+	}
+}
+
+FSetMuscleActivationParameterNode::FSetMuscleActivationParameterNode(const UE::Dataflow::FNodeParameters& InParam, FGuid InGuid)
+	: FDataflowNode(InParam, InGuid)
+	, ApplyGlobalParameters(
+		FDataflowFunctionProperty::FDelegate::CreateLambda([this](UE::Dataflow::FContext& Context)
+			{
+				// triggers node invalidation
+				ParameterMethod = EParameterMethod::Global;
+			}))
+	, ImportAllMuscleNames(
+		FDataflowFunctionProperty::FDelegate::CreateLambda([this](UE::Dataflow::FContext& Context)
+			{
+				const FManagedArrayCollection& InCollection = GetValue<FManagedArrayCollection>(Context, &Collection);
+				const GeometryCollection::Facades::FMuscleActivationFacade FMuscleActivation(InCollection);
+				ParameterArray.SetNum(FMuscleActivation.NumMuscles());
+				for (int32 MuscleIdx = 0; MuscleIdx < FMuscleActivation.NumMuscles(); ++MuscleIdx)
+				{
+					ParameterArray[MuscleIdx].MuscleName = FMuscleActivation.FindMuscleName(MuscleIdx);
+					ParameterArray[MuscleIdx].ContractionVolumeScale = ContractionVolumeScale;
+					ParameterArray[MuscleIdx].FiberLengthRatioAtMaxActivation = GlobalFiberLengthRatioAtMaxActivation;
+					ParameterArray[MuscleIdx].MuscleLengthRatioThresholdForMaxActivation = GlobalMuscleLengthRatioThresholdForMaxActivation;
+				}
+			}))
+	, ResetToGlobalParameters(
+		FDataflowFunctionProperty::FDelegate::CreateLambda([this](UE::Dataflow::FContext& Context)
+			{
+				for (int32 Idx = 0; Idx < ParameterArray.Num(); ++Idx)
+				{
+					ParameterArray[Idx].ContractionVolumeScale = ContractionVolumeScale;
+					ParameterArray[Idx].FiberLengthRatioAtMaxActivation = GlobalFiberLengthRatioAtMaxActivation;
+					ParameterArray[Idx].MuscleLengthRatioThresholdForMaxActivation = GlobalMuscleLengthRatioThresholdForMaxActivation;
+				}
+			}))
+	, ApplyCustomParameters(
+		FDataflowFunctionProperty::FDelegate::CreateLambda([this](UE::Dataflow::FContext& Context)
+			{
+				// triggers node invalidation
+				ParameterMethod = EParameterMethod::Custom;
+			}))
+{
+	RegisterInputConnection(&Collection);
+	RegisterOutputConnection(&Collection, &Collection);
+}
+
+void FSetMuscleActivationParameterNode::Evaluate(UE::Dataflow::FContext& Context, const FDataflowOutput* Out) const
+{
+	if (Out->IsA(&Collection))
+	{
+		FManagedArrayCollection InCollection = GetValue(Context, &Collection);
+		if (IsConnected(&Collection))
+		{
+			GeometryCollection::Facades::FMuscleActivationFacade FMuscleActivation(InCollection);
+			FMuscleActivation.UpdateGlobalMuscleActivationParameters(
+					ContractionVolumeScale,
+					GlobalFiberLengthRatioAtMaxActivation,
+					GlobalMuscleLengthRatioThresholdForMaxActivation);
+			if (ParameterMethod == EParameterMethod::Custom)
+			{
+				// match muscle names and override parameters
+				for (int32 Idx = 0; Idx < ParameterArray.Num(); ++Idx)
+				{
+					FPerMuscleParameter Params = ParameterArray[Idx];
+					int32 MuscleIndex = FMuscleActivation.FindMuscleIndexByName(Params.MuscleName);
+					if (FMuscleActivation.IsValidMuscleIndex(MuscleIndex))
+					{
+						FMuscleActivation.UpdateMuscleActivationParameters(
+							MuscleIndex,
+							Params.ContractionVolumeScale,
+							Params.FiberLengthRatioAtMaxActivation,
+							Params.MuscleLengthRatioThresholdForMaxActivation);
+					}
+				}
+			}
+		}
+		Out->SetValue(MoveTemp(InCollection), Context);
+	}
+}
+
+bool FSetMuscleActivationParameterNode::ShouldInvalidateOnPropertyChanged(const FPropertyChangedEvent& InPropertyChangedEvent) const
+{
+	const FName ChangedPropertyName = InPropertyChangedEvent.GetMemberPropertyName();
+	if (ChangedPropertyName == GET_MEMBER_NAME_CHECKED(FSetMuscleActivationParameterNode, ParameterArray) ||
+		ChangedPropertyName == GET_MEMBER_NAME_CHECKED(FSetMuscleActivationParameterNode, ImportAllMuscleNames) ||
+		ChangedPropertyName == GET_MEMBER_NAME_CHECKED(FSetMuscleActivationParameterNode, ResetToGlobalParameters))
+	{
+		return false;
+	}
+	return true;
+}
\ No newline at end of file
diff --git a/Engine/Plugins/Experimental/ChaosFlesh/Source/ChaosFleshNodes/Private/Dataflow/ChaosFleshNodesPlugin.cpp b/Engine/Plugins/Experimental/ChaosFlesh/Source/ChaosFleshNodes/Private/Dataflow/ChaosFleshNodesPlugin.cpp
--- a/Engine/Plugins/Experimental/ChaosFlesh/Source/ChaosFleshNodes/Private/Dataflow/ChaosFleshNodesPlugin.cpp
+++ b/Engine/Plugins/Experimental/ChaosFlesh/Source/ChaosFleshNodes/Private/Dataflow/ChaosFleshNodesPlugin.cpp
@@ -48,6 +48,7 @@ void IChaosFleshNodesPlugin::StartupModule()
 	DATAFLOW_NODE_REGISTER_CREATION_FACTORY(FComputeFiberStreamlineNode);
 	DATAFLOW_NODE_REGISTER_CREATION_FACTORY(FComputeIslandsNode);
 	DATAFLOW_NODE_REGISTER_CREATION_FACTORY(FComputeMuscleActivationDataNode);
+	DATAFLOW_NODE_REGISTER_CREATION_FACTORY(FComputeMuscleActivationDataNode_v2);
 	DATAFLOW_NODE_REGISTER_CREATION_FACTORY(FCreateTetrahedronDataflowNode);
 	DATAFLOW_NODE_REGISTER_CREATION_FACTORY(FDeleteFleshVerticesDataflowNode);
 	DATAFLOW_NODE_REGISTER_CREATION_FACTORY(FFleshAssetTerminalDataflowNode);
@@ -65,6 +66,7 @@ void IChaosFleshNodesPlugin::StartupModule()
 	DATAFLOW_NODE_REGISTER_CREATION_FACTORY(FRadialTetrahedronDataflowNodes);
 	DATAFLOW_NODE_REGISTER_CREATION_FACTORY(FSetFleshBonePositionTargetBindingDataflowNode);
 	DATAFLOW_NODE_REGISTER_CREATION_FACTORY(FSetFleshDefaultPropertiesNode);
+	DATAFLOW_NODE_REGISTER_CREATION_FACTORY(FSetMuscleActivationParameterNode);
 	DATAFLOW_NODE_REGISTER_CREATION_FACTORY(FSetVerticesKinematicDataflowNode);
 	DATAFLOW_NODE_REGISTER_CREATION_FACTORY(FSetVertexTetrahedraPositionTargetBindingDataflowNode);
 	DATAFLOW_NODE_REGISTER_CREATION_FACTORY(FSetVertexTrianglePositionTargetBindingDataflowNode);
diff --git a/Engine/Plugins/Experimental/ChaosFlesh/Source/ChaosFleshNodes/Public/Dataflow/ChaosFleshComputeMuscleActivationNode.h b/Engine/Plugins/Experimental/ChaosFlesh/Source/ChaosFleshNodes/Public/Dataflow/ChaosFleshComputeMuscleActivationNode.h
--- a/Engine/Plugins/Experimental/ChaosFlesh/Source/ChaosFleshNodes/Public/Dataflow/ChaosFleshComputeMuscleActivationNode.h
+++ b/Engine/Plugins/Experimental/ChaosFlesh/Source/ChaosFleshNodes/Public/Dataflow/ChaosFleshComputeMuscleActivationNode.h
@@ -5,14 +5,16 @@
 #include "CoreMinimal.h"
 #include "Dataflow/DataflowCore.h"
 #include "Dataflow/DataflowEngine.h"
+#include "Dataflow/DataflowFunctionProperty.h"
 
 #include "ChaosFleshComputeMuscleActivationNode.generated.h"
 
 /**
+* Please use new version of this node, then use SetMuscleActivationParameter node to set up muscle activation parameters.
 * Computes an orthogonal matrix for each element
 * M = [v,w,u], where v is the fiber direction of that element, w and u are chosen to be orthogonal to v and each other.
 */
-USTRUCT(meta = (DataflowFlesh))
+USTRUCT(meta = (Deprecated = "5.6"))
 struct FComputeMuscleActivationDataNode : public FDataflowNode
 {
 	GENERATED_USTRUCT_BODY()
@@ -40,6 +42,118 @@ public:
 		RegisterInputConnection(&InsertionIndicesIn);
 		RegisterOutputConnection(&Collection, &Collection);
 	}
+	
+	virtual void Evaluate(UE::Dataflow::FContext& Context, const FDataflowOutput* Out) const override;
+};
+
+/**
+* Determines muscles that are eligible for activation and computes muscle activation data.
+*/
+USTRUCT(meta = (DataflowFlesh))
+struct FComputeMuscleActivationDataNode_v2 : public FDataflowNode
+{
+	GENERATED_USTRUCT_BODY()
+	DATAFLOW_NODE_DEFINE_INTERNAL(FComputeMuscleActivationDataNode_v2, "ComputeMuscleActivationData", "Flesh", "")
+	DATAFLOW_NODE_RENDER_TYPE("SurfaceRender", FGeometryCollection::StaticType(), "Collection")
+
+public:
+	UPROPERTY(meta = (DataflowInput, DataflowOutput, DisplayName = "Collection", DataflowPassthrough = "Collection"))
+	FManagedArrayCollection Collection;
+
+	UPROPERTY(meta = (DataflowInput, DisplayName = "OriginVertexIndices"))
+	TArray<int32> OriginIndicesIn;
+
+	UPROPERTY(meta = (DataflowInput, DisplayName = "InsertionVertexIndices"))
+	TArray<int32> InsertionIndicesIn;
+
+	FComputeMuscleActivationDataNode_v2(const UE::Dataflow::FNodeParameters& InParam, FGuid InGuid = FGuid::NewGuid())
+		: FDataflowNode(InParam, InGuid)
+	{
+		RegisterInputConnection(&Collection);
+		RegisterInputConnection(&OriginIndicesIn);
+		RegisterInputConnection(&InsertionIndicesIn);
+		RegisterOutputConnection(&Collection, &Collection);
+	}
 
 	virtual void Evaluate(UE::Dataflow::FContext& Context, const FDataflowOutput* Out) const override;
 };
+
+/**
+* Struct data structure to store per-muscle activation parameters
+*/
+USTRUCT()
+struct FPerMuscleParameter
+{
+	GENERATED_USTRUCT_BODY()
+	UPROPERTY(EditAnywhere, meta = (InlineEditConditionToggle), Category = "Muscle Activation Parameters")
+	bool bCanEditMuscleName;
+	UPROPERTY(EditAnywhere, meta = (EditCondition = "bCanEditMuscleName"), Category = "Muscle Activation Parameters")
+	FString MuscleName;
+	UPROPERTY(EditAnywhere, Category = "Muscle Activation Parameters")
+	float ContractionVolumeScale;
+	UPROPERTY(EditAnywhere, Category = "Muscle Activation Parameters")
+	float FiberLengthRatioAtMaxActivation;
+	UPROPERTY(EditAnywhere, Category = "Muscle Activation Parameters")
+	float MuscleLengthRatioThresholdForMaxActivation;
+};
+
+UENUM()
+enum EParameterMethod : int
+{
+	Global	UMETA(DisplayName = "Use global parameters"),
+	Custom	UMETA(DisplayName = "Override with custom parameters"),
+};
+
+/**
+* Sets per-muscle parameters for custom muscle contraction.
+*/
+USTRUCT(meta = (DataflowFlesh))
+struct FSetMuscleActivationParameterNode : public FDataflowNode
+{
+	GENERATED_USTRUCT_BODY()
+	DATAFLOW_NODE_DEFINE_INTERNAL(FSetMuscleActivationParameterNode, "SetMuscleActivationParameter", "Flesh", "")
+	DATAFLOW_NODE_RENDER_TYPE("SurfaceRender", FGeometryCollection::StaticType(), "Collection")
+
+public:
+	UPROPERTY(meta = (DataflowInput, DataflowOutput, DisplayName = "Collection", DataflowPassthrough = "Collection"))
+	FManagedArrayCollection Collection;
+
+	UPROPERTY(EditAnywhere, Category = "Parameter Method", meta = (EditCondition = false))
+	TEnumAsByte<EParameterMethod> ParameterMethod = EParameterMethod::Global;
+
+	/** Click on this button to apply global parameters. */
+	UPROPERTY(EditAnywhere, Category = "Global Parameters", meta = (ButtonImage = "Icons.Refresh"))
+	FDataflowFunctionProperty ApplyGlobalParameters;
+
+	/* Muscles gain volume during contraction if > 1. Volume-preserving if 1. Default: 1*/
+	UPROPERTY(EditAnywhere, Category = "Global Parameters", DisplayName = "Global Contraction Volume Scale")
+	float ContractionVolumeScale = 1.f;
+
+	/* How much muscle fibers shorten at max activation 1. A smaller value means more contraction in the fiber direction. Default: 0.5 */
+	UPROPERTY(EditAnywhere, Category = "Global Parameters")
+	float GlobalFiberLengthRatioAtMaxActivation = 0.5f;
+
+	/* Muscle length ratio (defined by origin-insertion distance) below this threshold is considered to reach max activation 1. Default: 0.75 */
+	UPROPERTY(EditAnywhere, Category = "Global Parameters")
+	float GlobalMuscleLengthRatioThresholdForMaxActivation = 0.75f;
+
+	/** Click on this button to import muscle bone names (in the Transform group) from the collection. */
+	UPROPERTY(EditAnywhere, Category = "Custom Parameters", meta = (ButtonImage = "Icons.Import"))
+	FDataflowFunctionProperty ImportAllMuscleNames;
+
+	/** Click on this button to reset per-muscle parameters to global parameters. */
+	UPROPERTY(EditAnywhere, Category = "Custom Parameters", meta = (ButtonImage = "Icons.Import"))
+	FDataflowFunctionProperty ResetToGlobalParameters;
+
+	/** Click on this button to apply custom muscle parameters. */
+	UPROPERTY(EditAnywhere, Category = "Custom Parameters", meta = (ButtonImage = "Icons.Refresh"))
+	FDataflowFunctionProperty ApplyCustomParameters;
+
+	UPROPERTY(EditAnywhere, Category = "Custom Parameters", DisplayName = "Per-muscle Parameters")
+	TArray<FPerMuscleParameter> ParameterArray;
+
+	FSetMuscleActivationParameterNode(const UE::Dataflow::FNodeParameters& InParam, FGuid InGuid = FGuid::NewGuid());
+
+	virtual void Evaluate(UE::Dataflow::FContext& Context, const FDataflowOutput* Out) const override;
+	virtual bool ShouldInvalidateOnPropertyChanged(const FPropertyChangedEvent& InPropertyChangedEvent) const override;
+};
\ No newline at end of file
diff --git a/Engine/Plugins/Experimental/Dataflow/Source/DataflowEditor/Private/Dataflow/DataflowSimulationGenerator.cpp b/Engine/Plugins/Experimental/Dataflow/Source/DataflowEditor/Private/Dataflow/DataflowSimulationGenerator.cpp
--- a/Engine/Plugins/Experimental/Dataflow/Source/DataflowEditor/Private/Dataflow/DataflowSimulationGenerator.cpp
+++ b/Engine/Plugins/Experimental/Dataflow/Source/DataflowEditor/Private/Dataflow/DataflowSimulationGenerator.cpp
@@ -65,7 +65,7 @@ void FDataflowSimulationTask::DoWork()
 
 bool FDataflowTaskManager::AllocateSimulationResource(const FVector2f& TimeRange, const int32 FrameRate,
 		const TObjectPtr<UChaosCacheCollection>& CacheAsset, const TSubclassOf<AActor>& ActorClass,
-		const TObjectPtr<UDataflowBaseContent>& DataflowContent, const FTransform& BlueprintTransform, const bool bSkeletalMeshVisibility)
+		const TObjectPtr<UDataflowBaseContent>& DataflowContent, const FTransform& BlueprintTransform, const bool bSkeletalMeshVisibility, float DeltaTime)
 {
 	SimulationWorld = UWorld::CreateWorld(EWorldType::Editor, false);
 	SimulationWorld->bPostTickComponentUpdate = false;
@@ -96,13 +96,13 @@ bool FDataflowTaskManager::AllocateSimulationResource(const FVector2f& TimeRange
 	SimulationResource->NumSimulatedFrames = &NumSimulatedFrames;
 	SimulationResource->bCancelled = &bCancelled;
 	
-	NumFrames = FMath::Floor((TimeRange[1] - TimeRange[0]) * FrameRate);
+	NumFrames = FMath::Floor((TimeRange[1] - TimeRange[0] + UE_SMALL_NUMBER) * FrameRate);
 	if(SimulationTask.IsValid())
 	{
 		SimulationTask->GetTask().SimulationWorld = SimulationWorld;
 		SimulationTask->GetTask().MinTime = TimeRange[0];
 		SimulationTask->GetTask().MaxTime = TimeRange[1];
-		SimulationTask->GetTask().DeltaTime = (TimeRange[1]-TimeRange[0]) / NumFrames;
+		SimulationTask->GetTask().DeltaTime = DeltaTime;
 
 		if(SimulationTask->GetTask().bAsyncCaching)
 		{
@@ -202,7 +202,7 @@ void FDataflowSimulationGenerator::StartGenerateSimulation()
 	TaskManager->SimulationTask->GetTask().bAsyncCaching = CacheParams.bAsyncCaching;
 	TaskManager->SimulationTask->GetTask().bRestartSimulation = CacheParams.bRestartSimulation;
 	FVector2f TimeRange = CacheParams.bRestartSimulation ? CacheParams.RestartTimeRange : CacheParams.TimeRange;
-	TaskManager->AllocateSimulationResource(TimeRange, CacheParams.FrameRate, CacheAsset, BlueprintClass, DataflowContent, BlueprintTransform, bSkeletalMeshVisibility);
+	TaskManager->AllocateSimulationResource(TimeRange, CacheParams.FrameRate, CacheAsset, BlueprintClass, DataflowContent, BlueprintTransform, bSkeletalMeshVisibility, GeneratorDeltaTime);
 
 	if(CacheParams.bAsyncCaching)
 	{
@@ -251,8 +251,8 @@ void FDataflowSimulationGenerator::TickGenerateSimulation()
 			return;
 		}
 		
-		const int32 NumSimulatedFrames = TaskManager->NumSimulatedFrames.load();
-		const int32 NumTotalFrames = TaskManager->NumFrames;
+		const int32 NumSimulatedFrames = TaskManager->NumSimulatedFrames.load() + 1; //Counting the preroll
+		const int32 NumTotalFrames = TaskManager->NumFrames + 1;
 		const FText ProgressMessage = FText::FromString(FString::Printf(TEXT("Finished %d/%d, %.1f%%"),
 			NumSimulatedFrames, NumTotalFrames, 100.0 * NumSimulatedFrames / NumTotalFrames));
 		TaskManager->AsyncNotification->SetProgressText(ProgressMessage);
diff --git a/Engine/Plugins/Experimental/Dataflow/Source/DataflowEditor/Private/Dataflow/DataflowSimulationPanel.cpp b/Engine/Plugins/Experimental/Dataflow/Source/DataflowEditor/Private/Dataflow/DataflowSimulationPanel.cpp
--- a/Engine/Plugins/Experimental/Dataflow/Source/DataflowEditor/Private/Dataflow/DataflowSimulationPanel.cpp
+++ b/Engine/Plugins/Experimental/Dataflow/Source/DataflowEditor/Private/Dataflow/DataflowSimulationPanel.cpp
@@ -125,8 +125,7 @@ FReply SDataflowSimulationPanel::OnClick_Forward_Step()
 {
 	if (const TSharedPtr<FDataflowSimulationScene> PreviewScene = SimulationScene.Pin())
 	{
-		const float StepSize = GetSequenceLength() / GetNumberOfKeys();
-		PreviewScene->SimulationTime += StepSize;
+		PreviewScene->SimulationTime += PreviewScene->GetDeltaTime();
 	}
 
 	return FReply::Handled();
@@ -146,8 +145,7 @@ FReply SDataflowSimulationPanel::OnClick_Backward_Step()
 {
 	if (const TSharedPtr<FDataflowSimulationScene> PreviewScene = SimulationScene.Pin())
 	{
-		const float StepSize = GetSequenceLength() / GetNumberOfKeys();
-		PreviewScene->SimulationTime -= StepSize;
+		PreviewScene->SimulationTime -= PreviewScene->GetDeltaTime();
 	}
 
 	return FReply::Handled();
diff --git a/Engine/Plugins/Experimental/Dataflow/Source/DataflowEditor/Private/Dataflow/DataflowSimulationScene.cpp b/Engine/Plugins/Experimental/Dataflow/Source/DataflowEditor/Private/Dataflow/DataflowSimulationScene.cpp
--- a/Engine/Plugins/Experimental/Dataflow/Source/DataflowEditor/Private/Dataflow/DataflowSimulationScene.cpp
+++ b/Engine/Plugins/Experimental/Dataflow/Source/DataflowEditor/Private/Dataflow/DataflowSimulationScene.cpp
@@ -51,6 +51,8 @@ FDataflowSimulationScene::FDataflowSimulationScene(FPreviewScene::ConstructionVa
 			SceneDescription->CacheAsset = Cast<UChaosCacheCollection>(DataflowAsset->PreviewCacheAsset.LoadSynchronous());
 			SceneDescription->BlueprintClass = DataflowAsset->PreviewBlueprintClass;
 			SceneDescription->BlueprintTransform = DataflowAsset->PreviewBlueprintTransform; 
+			SceneDescription->GeometryCacheAsset = Cast<UGeometryCache>(DataflowAsset->PreviewGeometryCacheAsset.LoadSynchronous());
+			SceneDescription->EmbeddedSkeletalMesh = Cast<USkeletalMesh>(DataflowAsset->PreviewEmbeddedSkeletalMesh.LoadSynchronous());
 		}
 		if(SceneDescription->BlueprintClass == nullptr)
 		{
@@ -196,11 +198,11 @@ void FDataflowSimulationScene::CreateSimulationScene()
 		SimulationGenerator->SetDataflowContent(GetEditorContent());
 
 		TimeRange = SceneDescription->CacheParams.TimeRange;
-		NumFrames = (TimeRange[1] > TimeRange[0]) ? FMath::Floor((TimeRange[1] - TimeRange[0]) * SceneDescription->CacheParams.FrameRate) : 0;
-		
+		NumFrames = (TimeRange[1] > TimeRange[0]) ? FMath::Floor((TimeRange[1] - TimeRange[0] + UE_SMALL_NUMBER) * SceneDescription->CacheParams.FrameRate) + 1 : 1;
+		DeltaTime = (SceneDescription->CacheParams.FrameRate > 0) ? 1.f / SceneDescription->CacheParams.FrameRate: 0.f;
 		PreviewActor = UE::Dataflow::SpawnSimulatedActor(SceneDescription->BlueprintClass, Cast<AChaosCacheManager>(RootSceneActor),
 			SceneDescription->CacheAsset, false, GetEditorContent(), SceneDescription->BlueprintTransform);
-
+		SimulationGenerator->SetDeltaTime(DeltaTime);
 		// Setup all the skelmesh animations
 		UE::Dataflow::SetupSkeletonAnimation(PreviewActor, SceneDescription->bSkeletalMeshVisibility);
 		
@@ -225,8 +227,6 @@ void FDataflowSimulationScene::RecordSimulationCache()
 	else
 	{
 		bIsRecordingCache = true;
-		const float DeltaTime = (SceneDescription->CacheParams.FrameRate > 0) ?
-			(1.0f / SceneDescription->CacheParams.FrameRate) : 0.0f; 
         SimulationTime = TimeRange[0] - DeltaTime;
 		if(AChaosCacheManager* CacheManager = Cast<AChaosCacheManager>(RootSceneActor))
 		{
@@ -238,7 +238,7 @@ void FDataflowSimulationScene::RecordSimulationCache()
 	}
 }
 
-void FDataflowSimulationScene::UpdateSimulationCache(float& DeltaSeconds)
+void FDataflowSimulationScene::UpdateSimulationCache(float DeltaSeconds)
 {
 	if(RootSceneActor)
 	{
@@ -246,15 +246,15 @@ void FDataflowSimulationScene::UpdateSimulationCache(float& DeltaSeconds)
 		if(!bIsRecordingCache)
 		{
 			// Update the cached simulation at some point in time
-			if (SimulationTime != CacheManager->StartTime)
+			const float ScrollBarTime = SimulationTime - TimeRange[0];
+			if (!SimulationGenerator->IsSimulating() || ScrollBarTime != CacheManager->StartTime)
 			{
-				CacheManager->SetStartTime(SimulationTime);
+				CacheManager->SetStartTime(ScrollBarTime);
 			}
 		}
 		else
 		{
-			DeltaSeconds = 1.0f / SceneDescription->CacheParams.FrameRate;
-			SimulationTime += DeltaSeconds;
+			SimulationTime += DeltaTime;
 
 			if(SimulationTime >= TimeRange[1]+UE_SMALL_NUMBER)
 			{
@@ -324,7 +324,6 @@ FBox FDataflowSimulationScene::GetBoundingBox() const
 
 void FDataflowSimulationScene::TickDataflowScene(const float DeltaSeconds)
 {
-	float DeltaTime = DeltaSeconds;
 	if(const TObjectPtr<UDataflowBaseContent>& EditorContent = GetEditorContent())
 	{
 		if (const TObjectPtr<UDataflow> DataflowGraph = EditorContent->GetDataflowAsset())
@@ -346,13 +345,13 @@ void FDataflowSimulationScene::TickDataflowScene(const float DeltaSeconds)
 		if(SceneDescription->CacheAsset)
 		{
 			// Update the simulation cache (loading/recording)
-			UpdateSimulationCache(DeltaTime);
+			UpdateSimulationCache(DeltaSeconds);
 			
 			// Update all the skelmesh animations at the simulation time
 			UE::Dataflow::UpdateSkeletonAnimation(PreviewActor, SimulationTime);
 		}
 	}
-	GetWorld()->Tick(ELevelTick::LEVELTICK_All, DeltaTime);
+	GetWorld()->Tick(ELevelTick::LEVELTICK_All, DeltaSeconds);
 }
 
 void FDataflowSimulationScene::AddReferencedObjects(FReferenceCollector& Collector)
@@ -401,7 +400,9 @@ void FDataflowSimulationScene::SceneDescriptionPropertyChanged(const FName& Prop
 			DataflowAsset->PreviewCacheAsset = SceneDescription->CacheAsset;
 			DataflowAsset->PreviewBlueprintClass = SceneDescription->BlueprintClass;
 			DataflowAsset->PreviewBlueprintTransform = SceneDescription->BlueprintTransform;
-			
+			DataflowAsset->PreviewGeometryCacheAsset = SceneDescription->GeometryCacheAsset;
+			DataflowAsset->PreviewEmbeddedSkeletalMesh = SceneDescription->EmbeddedSkeletalMesh;
+
 			DataflowAsset->MarkPackageDirty();
 #endif
 		}
@@ -419,9 +420,8 @@ void UDataflowSimulationSceneDescription::GenerateGeometryCache()
 	SimulationScene->ResetSimulationScene();
 	SimulationScene->CreateSimulationScene();
 	const FVector2f& TimeRange = SimulationScene->GetTimeRange();
-	const int32 NumFrames = FMath::Floor((TimeRange[1] - TimeRange[0]) * CacheParams.FrameRate);
+	const int32 NumFrames = SimulationScene->GetNumFrames();
 	float Time = TimeRange[0];
-	float DeltaTime = (TimeRange[1] - TimeRange[0]) / NumFrames;
 	TObjectPtr<AActor> GetRootActor = SimulationScene->GetRootActor();
 	TObjectPtr<AActor> PreviewActor = SimulationScene->GetPreviewActor();
 	if (CacheAsset && GeometryCacheAsset && GetRootActor && EmbeddedSkeletalMesh)
@@ -451,11 +451,11 @@ void UDataflowSimulationSceneDescription::GenerateGeometryCache()
 		}
 		const TArray<int32>& Map = OptionalMap.GetValue();
 		TArray<uint32> ImportedVertexNumbers = TArray<uint32>(reinterpret_cast<const uint32*>(Map.GetData()), Map.Num());
-		for (int32 Frame = 0; Frame < SimulationScene->GetNumFrames(); ++Frame)
+		for (int32 Frame = 0; Frame < NumFrames; ++Frame)
 		{
-			Time += DeltaTime;
 			Cast<AChaosCacheManager>(GetRootActor)->SetStartTime(Time);
 			RenderPositions[Frame] = GeometryCachable->GetGeometryCachePositions(EmbeddedSkeletalMesh);
+			Time += SimulationScene->GetDeltaTime();
 		}
 		UE::DataflowSimulationGeometryCache::SaveGeometryCache(*GeometryCacheAsset, *EmbeddedSkeletalMesh, ImportedVertexNumbers, RenderPositions);
 		UE::DataflowSimulationGeometryCache::SavePackage(*GeometryCacheAsset);
diff --git a/Engine/Plugins/Experimental/Dataflow/Source/DataflowEditor/Public/Dataflow/DataflowSimulationGenerator.h b/Engine/Plugins/Experimental/Dataflow/Source/DataflowEditor/Public/Dataflow/DataflowSimulationGenerator.h
--- a/Engine/Plugins/Experimental/Dataflow/Source/DataflowEditor/Public/Dataflow/DataflowSimulationGenerator.h
+++ b/Engine/Plugins/Experimental/Dataflow/Source/DataflowEditor/Public/Dataflow/DataflowSimulationGenerator.h
@@ -109,7 +109,7 @@ namespace UE::Dataflow
 		/** Allocate the simulation resource from the properties */
         bool AllocateSimulationResource(const FVector2f& TimeRange, const int32 FrameRate,
         	const TObjectPtr<UChaosCacheCollection>& CacheAsset, const TSubclassOf<AActor>& ActorClass,
-        	const TObjectPtr<UDataflowBaseContent>& DataflowContent, const FTransform& BlueprintTransform, const bool bSkeletalMeshVisibility = true);
+        	const TObjectPtr<UDataflowBaseContent>& DataflowContent, const FTransform& BlueprintTransform, const bool bSkeletalMeshVisibility, float DeltaTime);
 
         /** Free the simulation resource */
         void FreeSimulationResource();
@@ -172,9 +172,15 @@ namespace UE::Dataflow
 		/** Set the skeletal mesh visibility */
 		void SetSkeletalMeshVisibility(const bool bInSkeletalMeshVisibility);
 
+		/** Set delta time */
+		void SetDeltaTime(float InDeltaTime) { GeneratorDeltaTime = InDeltaTime; };
+
 		/** Enqueue a generator action to be processed on the async thread */
 		void RequestGeneratorAction(EDataflowGeneratorActions Action);
 		
+		/** If the generator is running simulation */
+		bool IsSimulating() { return PendingAction != EDataflowGeneratorActions::NoAction; };
+
 	private:
 
 		/** Start generating the simulation */
@@ -204,6 +210,9 @@ namespace UE::Dataflow
 		/** Skeletal mesh visibility */
 		bool bSkeletalMeshVisibility = true;
 
+		/** Delta time for a frame */
+		float GeneratorDeltaTime;
+
 		/** Pending action to be send to the async thread */
 		EDataflowGeneratorActions PendingAction = EDataflowGeneratorActions::NoAction;
 
diff --git a/Engine/Plugins/Experimental/Dataflow/Source/DataflowEditor/Public/Dataflow/DataflowSimulationScene.h b/Engine/Plugins/Experimental/Dataflow/Source/DataflowEditor/Public/Dataflow/DataflowSimulationScene.h
--- a/Engine/Plugins/Experimental/Dataflow/Source/DataflowEditor/Public/Dataflow/DataflowSimulationScene.h
+++ b/Engine/Plugins/Experimental/Dataflow/Source/DataflowEditor/Public/Dataflow/DataflowSimulationScene.h
@@ -50,19 +50,19 @@ public:
 	FDataflowPreviewCacheParams CacheParams;
 
 	/** Geometry cache asset used to extract skeletal mesh results from simulation */
-	UPROPERTY(EditAnywhere, Category = "Geometry", DisplayName="Geometry Cache", meta=(EditCondition = "CacheAsset != nullptr"))
+	UPROPERTY(EditAnywhere, Category = "Geometry Cache", DisplayName="Geometry Cache", meta=(EditCondition = "CacheAsset != nullptr"))
 	TObjectPtr<UGeometryCache> GeometryCacheAsset = nullptr;
 
 	/** SkeletalMesh interpolated from simulation. This should match the SkeletalMesh used in GenerateSurfaceBindings node */
-	UPROPERTY(EditAnywhere, Category = "Geometry", DisplayName = "EmbeddedSkeletalMesh", meta = (EditCondition = "CacheAsset != nullptr"))
+	UPROPERTY(EditAnywhere, Category = "Geometry Cache", DisplayName = "EmbeddedSkeletalMesh", meta = (EditCondition = "CacheAsset != nullptr"))
 	TObjectPtr<USkeletalMesh> EmbeddedSkeletalMesh = nullptr;
 
 	/** Interpolates and saves geometry cache from Chaos cache */
-	UFUNCTION(CallInEditor, Category = "Geometry", meta = (EditCondition = "CacheAsset != nullptr && GeometryCacheAsset != nullptr && EmbeddedSkeletalMesh != nullptr"))
+	UFUNCTION(CallInEditor, Category = "Geometry Cache", meta = (EditCondition = "CacheAsset != nullptr && GeometryCacheAsset != nullptr && EmbeddedSkeletalMesh != nullptr"))
 	void GenerateGeometryCache();
 
 	/** Creates a new geometry cache file */
-	UFUNCTION(CallInEditor, Category = "Geometry", meta = (EditCondition = "CacheAsset != nullptr && GeometryCacheAsset != nullptr && EmbeddedSkeletalMesh != nullptr"))
+	UFUNCTION(CallInEditor, Category = "Geometry Cache", meta = (EditCondition = "CacheAsset != nullptr && GeometryCacheAsset != nullptr && EmbeddedSkeletalMesh != nullptr"))
 	void NewGeometryCache();
 
 	/** Visibility of the skeletal mesh */
@@ -143,6 +143,9 @@ public:
 	/** Get the number of frames */
 	const int32& GetNumFrames() const {return NumFrames;}
 
+	/** Get delta time */
+	float GetDeltaTime() const { return DeltaTime; }
+
 	/** Simulation time used to drive the cache loading */
 	float SimulationTime;
 	
@@ -157,7 +160,7 @@ public:
 private:
 
 	/** Update the simulation cache at some point in time (loading/recording)*/
-	void UpdateSimulationCache(float& DeltaSeconds);
+	void UpdateSimulationCache(float DeltaSeconds);
 
 	/** Bind the scene selection to the components */
 	void BindSceneSelection();
@@ -177,6 +180,9 @@ private:
 	/** Number of cache frames */
 	int32 NumFrames;
 
+	/** Delta time (1/fps) */
+	float DeltaTime;
+
 	/** Last context time stamp for which we regenerated the world */
 	UE::Dataflow::FTimestamp LastTimeStamp = UE::Dataflow::FTimestamp::Invalid;
 
diff --git a/Engine/Source/Runtime/Experimental/Chaos/Private/Chaos/Deformable/ChaosDeformableSolver.cpp b/Engine/Source/Runtime/Experimental/Chaos/Private/Chaos/Deformable/ChaosDeformableSolver.cpp
--- a/Engine/Source/Runtime/Experimental/Chaos/Private/Chaos/Deformable/ChaosDeformableSolver.cpp
+++ b/Engine/Source/Runtime/Experimental/Chaos/Private/Chaos/Deformable/ChaosDeformableSolver.cpp
@@ -1310,7 +1310,7 @@ namespace Chaos::Softs
 
 	void FDeformableSolver::InitializeKinematicConstraint()
 	{
-		auto MKineticUpdate = [this](FSolverParticles& MParticles, const FSolverReal Dt, const FSolverReal MTime, const int32 Index)
+		auto MKinematicUpdate = [this](FSolverParticles& MParticles, const FSolverReal Dt, const FSolverReal MTime, const int32 Index)
 		{
 			PERF_SCOPE(STAT_ChaosDeformableSolver_InitializeKinematicConstraint);
 
@@ -1427,7 +1427,7 @@ namespace Chaos::Softs
 
 
 		};
-		Evolution->SetKinematicUpdateFunction(MKineticUpdate);
+		Evolution->SetKinematicUpdateFunction(MKinematicUpdate);
 	}
 
 	void FDeformableSolver::InitializeSelfCollisionVariables()
@@ -1925,7 +1925,7 @@ namespace Chaos::Softs
 			{
 				int32 VertexOffset = Proxy.GetSolverParticleRange().Start;
 				int32 ElementOffset = AllElements->Num(); //InitializeMuscleActivation has to go before adding tetrahedrons
-				MuscleActivationConstraints->AddMuscles(MuscleActivationFacade, VertexOffset, ElementOffset);
+				MuscleActivationConstraints->AddMuscles(Evolution->Particles(), MuscleActivationFacade, VertexOffset, ElementOffset);
 			}
 		}
 	}
@@ -2046,13 +2046,13 @@ namespace Chaos::Softs
 		}
 
 
-		int32 NumIterations = FMath::Clamp<int32>(Property.NumSolverSubSteps, 0, INT_MAX);
-		if (bEnableSolver && NumIterations)
+		int32 NumSubsteps = FMath::Clamp<int32>(Property.NumSolverSubSteps, 0, INT_MAX);
+		if (bEnableSolver && NumSubsteps)
 		{
-			FSolverReal SubDeltaTime = DeltaTime / (FSolverReal)NumIterations;
+			FSolverReal SubDeltaTime = DeltaTime / (FSolverReal)NumSubsteps;
 			if (!FMath::IsNearlyZero(SubDeltaTime))
 			{
-				for (int i = 0; i < NumIterations; ++i)
+				for (int i = 0; i < NumSubsteps; ++i)
 				{
 					Iteration = i+1;
 					Update(SubDeltaTime);
diff --git a/Engine/Source/Runtime/Experimental/Chaos/Private/GeometryCollection/Facades/CollectionMuscleActivationFacade.cpp b/Engine/Source/Runtime/Experimental/Chaos/Private/GeometryCollection/Facades/CollectionMuscleActivationFacade.cpp
--- a/Engine/Source/Runtime/Experimental/Chaos/Private/GeometryCollection/Facades/CollectionMuscleActivationFacade.cpp
+++ b/Engine/Source/Runtime/Experimental/Chaos/Private/GeometryCollection/Facades/CollectionMuscleActivationFacade.cpp
@@ -21,6 +21,8 @@ namespace GeometryCollection::Facades
 	const FName FMuscleActivationFacade::OriginInsertionRestLength("OriginInsertionRestLength");
 	const FName FMuscleActivationFacade::FiberDirectionMatrix("FiberDirectionMatrix");
 	const FName FMuscleActivationFacade::ContractionVolumeScale("ContractionVolumeScale");
+	const FName FMuscleActivationFacade::FiberLengthRatioAtMaxActivation("FiberLengthRatioAtMaxActivation");
+	const FName FMuscleActivationFacade::MuscleLengthRatioThresholdForMaxActivation("MuscleLengthRatioThresholdForMaxActivation");
 	const FName FMuscleActivationFacade::FiberStreamline("FiberStreamline");
 	const FName FMuscleActivationFacade::FiberStreamlineRestLength("FiberStreamlineRestLength");
 
@@ -33,6 +35,8 @@ namespace GeometryCollection::Facades
 		, OriginInsertionRestLengthAttribute(InCollection, OriginInsertionRestLength, GroupName)
 		, FiberDirectionMatrixAttribute(InCollection, FiberDirectionMatrix, GroupName)
 		, ContractionVolumeScaleAttribute(InCollection, ContractionVolumeScale, GroupName)
+		, FiberLengthRatioAtMaxActivationAttribute(InCollection, FiberLengthRatioAtMaxActivation, GroupName)
+		, MuscleLengthRatioThresholdForMaxActivationAttribute(InCollection, MuscleLengthRatioThresholdForMaxActivation, GroupName)
 		, FiberStreamlineAttribute(InCollection, FiberStreamline, GroupName)
 		, FiberStreamlineRestLengthAttribute(InCollection, FiberStreamlineRestLength, GroupName)
 	{
@@ -48,6 +52,8 @@ namespace GeometryCollection::Facades
 		, OriginInsertionRestLengthAttribute(InCollection, OriginInsertionRestLength, GroupName)
 		, FiberDirectionMatrixAttribute(InCollection, FiberDirectionMatrix, GroupName)
 		, ContractionVolumeScaleAttribute(InCollection, ContractionVolumeScale, GroupName)
+		, FiberLengthRatioAtMaxActivationAttribute(InCollection, FiberLengthRatioAtMaxActivation, GroupName)
+		, MuscleLengthRatioThresholdForMaxActivationAttribute(InCollection, MuscleLengthRatioThresholdForMaxActivation, GroupName)
 		, FiberStreamlineAttribute(InCollection, FiberStreamline, GroupName)
 		, FiberStreamlineRestLengthAttribute(InCollection, FiberStreamlineRestLength, GroupName)
 	{
@@ -57,8 +63,9 @@ namespace GeometryCollection::Facades
 	bool FMuscleActivationFacade::IsValid() const
 	{
 		return GeometryGroupIndexAttribute.IsValid() && MuscleActivationElementAttribute.IsValid() && OriginInsertionPairAttribute.IsValid() &&
-			OriginInsertionRestLengthAttribute.IsValid() && FiberDirectionMatrixAttribute.IsValid() && ContractionVolumeScaleAttribute.IsValid()
-			&& FiberStreamlineAttribute.IsValid() && FiberStreamlineRestLengthAttribute.IsValid();
+			OriginInsertionRestLengthAttribute.IsValid() && FiberDirectionMatrixAttribute.IsValid() && ContractionVolumeScaleAttribute.IsValid() &&
+			FiberLengthRatioAtMaxActivationAttribute.IsValid() && MuscleLengthRatioThresholdForMaxActivationAttribute.IsValid() &&
+			FiberStreamlineAttribute.IsValid() && FiberStreamlineRestLengthAttribute.IsValid();
 	}
 
 	void FMuscleActivationFacade::DefineSchema()
@@ -70,6 +77,8 @@ namespace GeometryCollection::Facades
 		OriginInsertionRestLengthAttribute.Add();
 		FiberDirectionMatrixAttribute.Add();
 		ContractionVolumeScaleAttribute.Add();
+		FiberLengthRatioAtMaxActivationAttribute.Add(); 
+		MuscleLengthRatioThresholdForMaxActivationAttribute.Add();
 		FiberStreamlineAttribute.Add();
 		FiberStreamlineRestLengthAttribute.Add();
 	}
@@ -80,24 +89,16 @@ namespace GeometryCollection::Facades
 		if (IsValid())
 		{
 			int32 NewIndex = MuscleActivationElementAttribute.AddElements(1);
-			GeometryGroupIndexAttribute.Modify()[NewIndex] = InputData.GeometryGroupIndex;
-			MuscleActivationElementAttribute.Modify()[NewIndex] = InputData.MuscleActivationElement;
-			MuscleActivationElementAttribute.Modify()[NewIndex] = InputData.MuscleActivationElement;
-			OriginInsertionPairAttribute.Modify()[NewIndex] = InputData.OriginInsertionPair;
-			OriginInsertionRestLengthAttribute.Modify()[NewIndex] = InputData.OriginInsertionRestLength;
-			FiberDirectionMatrixAttribute.Modify()[NewIndex] = InputData.FiberDirectionMatrix;
-			ContractionVolumeScaleAttribute.Modify()[NewIndex] = InputData.ContractionVolumeScale;
-			FiberStreamlineAttribute.Modify()[NewIndex] = InputData.FiberStreamline;
-			FiberStreamlineRestLengthAttribute.Modify()[NewIndex] = InputData.FiberStreamlineRestLength;
+			UpdateMuscleActivationData(NewIndex, InputData);
 			return NewIndex;
 		}
 		return INDEX_NONE;
 	}
 
-	bool FMuscleActivationFacade::UpdateMuscleActivationData(int32 DataIndex, FMuscleActivationData& InputData)
+	bool FMuscleActivationFacade::UpdateMuscleActivationData(int32 DataIndex, const FMuscleActivationData& InputData)
 	{
 		check(!IsConst());
-		if (IsValid() && 0 <= DataIndex && DataIndex < NumMuscles())
+		if (IsValid() && IsValidMuscleIndex(DataIndex))
 		{
 			GeometryGroupIndexAttribute.Modify()[DataIndex] = InputData.GeometryGroupIndex;
 			MuscleActivationElementAttribute.Modify()[DataIndex] = InputData.MuscleActivationElement;
@@ -105,6 +106,8 @@ namespace GeometryCollection::Facades
 			OriginInsertionRestLengthAttribute.Modify()[DataIndex] = InputData.OriginInsertionRestLength;
 			FiberDirectionMatrixAttribute.Modify()[DataIndex] = InputData.FiberDirectionMatrix;
 			ContractionVolumeScaleAttribute.Modify()[DataIndex] = InputData.ContractionVolumeScale;
+			FiberLengthRatioAtMaxActivationAttribute.Modify()[DataIndex] = InputData.FiberLengthRatioAtMaxActivation;
+			MuscleLengthRatioThresholdForMaxActivationAttribute.Modify()[DataIndex] = InputData.MuscleLengthRatioThresholdForMaxActivation;
 			FiberStreamlineAttribute.Modify()[DataIndex] = InputData.FiberStreamline;
 			FiberStreamlineRestLengthAttribute.Modify()[DataIndex] = InputData.FiberStreamlineRestLength;
 			return true;
@@ -115,19 +118,18 @@ namespace GeometryCollection::Facades
 	FMuscleActivationData FMuscleActivationFacade::GetMuscleActivationData(const int32 DataIndex) const
 	{
 		FMuscleActivationData ReturnData;
-		if (IsValid())
+		if (IsValid() && IsValidMuscleIndex(DataIndex))
 		{
-			if (0 <= DataIndex && DataIndex < NumMuscles())
-			{
-				ReturnData.GeometryGroupIndex = GeometryGroupIndexAttribute.Get()[DataIndex];
-				ReturnData.MuscleActivationElement = MuscleActivationElementAttribute.Get()[DataIndex];
-				ReturnData.OriginInsertionPair = OriginInsertionPairAttribute.Get()[DataIndex];
-				ReturnData.OriginInsertionRestLength = OriginInsertionRestLengthAttribute.Get()[DataIndex];
-				ReturnData.FiberDirectionMatrix = FiberDirectionMatrixAttribute.Get()[DataIndex];
-				ReturnData.ContractionVolumeScale = ContractionVolumeScaleAttribute.Get()[DataIndex];
-				ReturnData.FiberStreamline = FiberStreamlineAttribute.Get()[DataIndex];
-				ReturnData.FiberStreamlineRestLength = FiberStreamlineRestLengthAttribute.Get()[DataIndex];
-			}
+			ReturnData.GeometryGroupIndex = GeometryGroupIndexAttribute[DataIndex];
+			ReturnData.MuscleActivationElement = MuscleActivationElementAttribute[DataIndex];
+			ReturnData.OriginInsertionPair = OriginInsertionPairAttribute[DataIndex];
+			ReturnData.OriginInsertionRestLength = OriginInsertionRestLengthAttribute[DataIndex];
+			ReturnData.FiberDirectionMatrix = FiberDirectionMatrixAttribute[DataIndex];
+			ReturnData.ContractionVolumeScale = ContractionVolumeScaleAttribute[DataIndex];
+			ReturnData.FiberLengthRatioAtMaxActivation = FiberLengthRatioAtMaxActivationAttribute[DataIndex];
+			ReturnData.MuscleLengthRatioThresholdForMaxActivation = MuscleLengthRatioThresholdForMaxActivationAttribute[DataIndex];
+			ReturnData.FiberStreamline = FiberStreamlineAttribute[DataIndex];
+			ReturnData.FiberStreamlineRestLength = FiberStreamlineRestLengthAttribute[DataIndex];
 		}
 		return ReturnData;
 	}
@@ -158,6 +160,37 @@ namespace GeometryCollection::Facades
 		}
 	}
 
+	FString FMuscleActivationFacade::FindMuscleName(const int32 MuscleIndex) const
+	{
+		const TManagedArray<FString>* BoneName = ConstCollection.FindAttributeTyped<FString>("BoneName", FTransformCollection::TransformGroup);
+		const TManagedArray<int32>* TransformIndex = ConstCollection.FindAttributeTyped<int32>("TransformIndex", FGeometryCollection::GeometryGroup);
+		int32 MuscleGeometryIndex = GeometryGroupIndexAttribute[MuscleIndex];
+		if (BoneName && TransformIndex && TransformIndex->IsValidIndex(MuscleGeometryIndex) && BoneName->IsValidIndex((*TransformIndex)[MuscleGeometryIndex]))
+		{
+			return (*BoneName)[(*TransformIndex)[MuscleGeometryIndex]];
+		}
+		return FString();
+	}
+
+	int32 FMuscleActivationFacade::FindMuscleIndexByName(const FString MuscleName) const
+	{
+		const TManagedArray<FString>* BoneName = ConstCollection.FindAttributeTyped<FString>("BoneName", FTransformCollection::TransformGroup);
+		const TManagedArray<int32>* TransformIndex = ConstCollection.FindAttributeTyped<int32>("TransformIndex", FGeometryCollection::GeometryGroup);
+		if (BoneName && TransformIndex)
+		{
+			int32 MuscleTransformIndex = BoneName->Find(MuscleName);
+			if (MuscleTransformIndex >= 0)
+			{
+				int32 MuscleGeometryIndex = TransformIndex->Find(MuscleTransformIndex);
+				if (MuscleGeometryIndex >= 0)
+				{
+					return GeometryGroupIndexAttribute.Get().Find(MuscleGeometryIndex);
+				}
+			}
+		}
+		return INDEX_NONE;
+	}
+
 	int32 FMuscleActivationFacade::RemoveInvalidMuscles()
 	{
 		check(!IsConst());
@@ -176,6 +209,112 @@ namespace GeometryCollection::Facades
 		return InvalidMuscleIndices.Num();
 	}
 
+	bool FMuscleActivationFacade::SetUpMuscleActivation(const TArray<int32>& InOrigin, const TArray<int32>& Insertion)
+	{
+		// Vertices and fiber field
+		if (!(InOrigin.Num() > 0
+			&& Insertion.Num() > 0
+			&& ConstCollection.FindAttribute<FVector3f>("Vertex", FGeometryCollection::VerticesGroup)
+			&& ConstCollection.FindAttribute<FIntVector4>("Tetrahedron", "Tetrahedral")
+			&& ConstCollection.FindAttribute<FVector3f>("FiberDirection", "Tetrahedral")))
+		{
+			return false;
+		}
+		TArray<int32> Origin = InOrigin;
+		const TArray<FVector3f>& Vertices = ConstCollection.FindAttribute<FVector3f>("Vertex", FGeometryCollection::VerticesGroup)->GetConstArray();
+		const TArray<FIntVector4>& Elements = ConstCollection.FindAttribute<FIntVector4>("Tetrahedron", "Tetrahedral")->GetConstArray();
+		const TArray<FVector3f>& FiberDirections = ConstCollection.FindAttribute<FVector3f>("FiberDirection", "Tetrahedral")->GetConstArray();
+		TArray<TArray<int32>> MuscleActivationElements;
+		TArray<TArray<int32>> ComponentOrigins; //One origin node per muscle component
+		TArray<TArray<int32>> ComponentInsertions; //One insertion node per muscle component
+		GeometryCollection::Facades::FCollectionMeshFacade MeshFacade(ConstCollection);
+		TArray<int32> ComponentIndex = MeshFacade.GetGeometryGroupIndexArray(); //Vertex index to geometry index
+		TMap<int32, int32> ComponentToMuscleIndex; //Component index to muscle index
+		TMap<int32, int32> MuscleToComponentIndex; //Muscle index to component index
+		Origin.Sort(); //For some order in muscle groups
+		for (int32 i = 0; i < Origin.Num(); i++)
+		{
+			if (!ComponentToMuscleIndex.Contains(ComponentIndex[Origin[i]]))
+			{
+				ComponentToMuscleIndex.Add(ComponentIndex[Origin[i]], ComponentOrigins.Num());
+				MuscleToComponentIndex.Add(ComponentOrigins.Num(), ComponentIndex[Origin[i]]);
+				ComponentOrigins.SetNum(ComponentOrigins.Num() + 1);
+				ComponentOrigins[ComponentOrigins.Num() - 1].Add(Origin[i]);
+			}
+			else
+			{
+				ComponentOrigins[ComponentToMuscleIndex[ComponentIndex[Origin[i]]]].Add(Origin[i]);
+			}
+		}
+		ComponentInsertions.SetNum(ComponentOrigins.Num());
+		for (int32 i = 0; i < Insertion.Num(); i++)
+		{
+			if (!ComponentToMuscleIndex.Contains(ComponentIndex[Insertion[i]]))
+			{
+				ensureMsgf(false, TEXT("No origin in this component"));
+			}
+			else
+			{
+				ComponentInsertions[ComponentToMuscleIndex[ComponentIndex[Insertion[i]]]].Add(Insertion[i]);
+			}
+		}
+		MuscleActivationElements.SetNum(ComponentOrigins.Num());
+		for (int32 ElemIdx = 0; ElemIdx < Elements.Num(); ElemIdx++)
+		{
+			if (ComponentToMuscleIndex.Contains(ComponentIndex[Elements[ElemIdx][0]]))
+			{
+				MuscleActivationElements[ComponentToMuscleIndex[ComponentIndex[Elements[ElemIdx][0]]]].Add(ElemIdx);
+			}
+		}
+		//Choose one origin-insertion pair per muscle that has largest distance apart within each muscle
+		// TODO: painted attribute directed origin-insertion pair
+		//use origin-insertion line segment length to estimate activation
+		for (int32 MuscleComponentIdx = 0; MuscleComponentIdx < ComponentOrigins.Num(); MuscleComponentIdx++)
+		{
+			if (ensureMsgf(ComponentOrigins.Num() > 0 && ComponentInsertions.Num() > 0, TEXT("Origin or Insertion missing in the muscle %d"), MuscleComponentIdx))
+			{
+				FMuscleActivationData MuscleActivationData;
+				MuscleActivationData.GeometryGroupIndex = MuscleToComponentIndex[MuscleComponentIdx];
+				MuscleActivationData.OriginInsertionRestLength = 0;
+				for (int32 OriginIdx : ComponentOrigins[MuscleComponentIdx])
+				{
+					for (int32 InsertionIdx : ComponentInsertions[MuscleComponentIdx])
+					{
+						float Dist = (Vertices[OriginIdx] - Vertices[InsertionIdx]).Size();
+						if (Dist > MuscleActivationData.OriginInsertionRestLength)
+						{
+							MuscleActivationData.OriginInsertionPair = FIntVector2(OriginIdx, InsertionIdx);
+							MuscleActivationData.OriginInsertionRestLength = Dist;
+						}
+					}
+				}
+				MuscleActivationData.MuscleActivationElement = MuscleActivationElements[MuscleComponentIdx];
+				MuscleActivationData.FiberDirectionMatrix.SetNum(MuscleActivationElements[MuscleComponentIdx].Num());
+				MuscleActivationData.ContractionVolumeScale.SetNum(MuscleActivationElements[MuscleComponentIdx].Num());
+				for (int32 LocalElemIdx = 0; LocalElemIdx < MuscleActivationElements[MuscleComponentIdx].Num(); LocalElemIdx++)
+				{
+					FVector3f V = FiberDirections[MuscleActivationElements[MuscleComponentIdx][LocalElemIdx]];
+					// QR decomposition on vvT for orthogonal directions
+					FVector3f W = V;
+					if (V.X < V.Y)
+					{
+						W.X += 1.f;
+					}
+					else
+					{
+						W.Y += 1.f;
+					}
+					FVector3f U = (V ^ W).GetSafeNormal();
+					W = (U ^ V).GetSafeNormal();
+					MuscleActivationData.FiberDirectionMatrix[LocalElemIdx] = Chaos::PMatrix33d(V, W, U);
+				}
+				MuscleActivationData.ContractionVolumeScale.Init(1.f, MuscleActivationElements[MuscleComponentIdx].Num());
+				AddMuscleActivationData(MuscleActivationData);
+			}
+		}
+		return true;
+	}
+
 	bool FMuscleActivationFacade::SetUpMuscleActivation(const TArray<int32>& InOrigin, const TArray<int32>& Insertion,
 		float InContractionVolumeScale)
 	{
@@ -284,6 +423,39 @@ namespace GeometryCollection::Facades
 		return true;
 	}
 
+	void FMuscleActivationFacade::UpdateGlobalMuscleActivationParameters(
+		float InContractionVolumeScale,
+		float InFiberLengthRatioAtMaxActivation,
+		float InMuscleLengthRatioThresholdForMaxActivation)
+	{
+		for (int32 MuscleIndex = 0; MuscleIndex < NumMuscles(); ++MuscleIndex)
+		{
+			FMuscleActivationData MuscleActivationData = GetMuscleActivationData(MuscleIndex);
+			MuscleActivationData.ContractionVolumeScale.Init(InContractionVolumeScale, MuscleActivationData.MuscleActivationElement.Num());
+			MuscleActivationData.FiberLengthRatioAtMaxActivation = InFiberLengthRatioAtMaxActivation;
+			MuscleActivationData.MuscleLengthRatioThresholdForMaxActivation = InMuscleLengthRatioThresholdForMaxActivation;
+			UpdateMuscleActivationData(MuscleIndex, MuscleActivationData);
+		}
+	}
+
+	bool FMuscleActivationFacade::UpdateMuscleActivationParameters(
+		int32 MuscleIndex,
+		float InContractionVolumeScale,
+		float InFiberLengthRatioAtMaxActivation,
+		float InMuscleLengthRatioThresholdForMaxActivation)
+	{
+		if (IsValidMuscleIndex(MuscleIndex))
+		{
+			FMuscleActivationData MuscleActivationData = GetMuscleActivationData(MuscleIndex);
+			MuscleActivationData.ContractionVolumeScale.Init(InContractionVolumeScale, MuscleActivationData.MuscleActivationElement.Num());
+			MuscleActivationData.FiberLengthRatioAtMaxActivation = InFiberLengthRatioAtMaxActivation;
+			MuscleActivationData.MuscleLengthRatioThresholdForMaxActivation = InMuscleLengthRatioThresholdForMaxActivation;
+			UpdateMuscleActivationData(MuscleIndex, MuscleActivationData);
+			return true;
+		}
+		return false;
+	}
+
 	TArray<TArray<TArray<FVector3f>>> FMuscleActivationFacade::BuildStreamlines(const TArray<int32>& Origin, const TArray<int32>& Insertion,
 		int32 NumLinesMultiplier, int32 MaxStreamlineIterations, int32 MaxPointsPerLine)
 	{
diff --git a/Engine/Source/Runtime/Experimental/Chaos/Public/Chaos/Deformable/GaussSeidelMainConstraint.h b/Engine/Source/Runtime/Experimental/Chaos/Public/Chaos/Deformable/GaussSeidelMainConstraint.h
--- a/Engine/Source/Runtime/Experimental/Chaos/Public/Chaos/Deformable/GaussSeidelMainConstraint.h
+++ b/Engine/Source/Runtime/Experimental/Chaos/Public/Chaos/Deformable/GaussSeidelMainConstraint.h
@@ -339,13 +339,14 @@ namespace Chaos::Softs
 				AddPerNodeHessian[i](p, Dt, ParticleHessian);
 			}
 
-
+			const T DtSquaredInverse =  1 / (Dt * Dt); // scale the hessian so that the determinant (~dt^6) does not fall below epsilon.
+			ParticleHessian *= DtSquaredInverse;
 			T HessianDet = ParticleHessian.Determinant();
-			if (HessianDet > UE_SMALL_NUMBER && ParticleResidual.Size() > UE_SMALL_NUMBER)
+			if (HessianDet > TMathUtilConstants<T>::Epsilon)
 			{
 				Chaos::PMatrix<T, 3, 3> HessianInv = ParticleHessian.SymmetricCofactorMatrix();
 				HessianInv *= T(1) / HessianDet;
-				Chaos::TVector<T, 3> Dx = HessianInv.GetTransposed() * (-ParticleResidual);
+				Chaos::TVector<T, 3> Dx = HessianInv.GetTransposed() * (-ParticleResidual) * DtSquaredInverse; // add back DtSquared
 
 				if (Dx.Size() < MaxDxSize)
 				{
diff --git a/Engine/Source/Runtime/Experimental/Chaos/Public/Chaos/Deformable/MuscleActivationConstraints.h b/Engine/Source/Runtime/Experimental/Chaos/Public/Chaos/Deformable/MuscleActivationConstraints.h
--- a/Engine/Source/Runtime/Experimental/Chaos/Public/Chaos/Deformable/MuscleActivationConstraints.h
+++ b/Engine/Source/Runtime/Experimental/Chaos/Public/Chaos/Deformable/MuscleActivationConstraints.h
@@ -5,6 +5,8 @@
 #include "Chaos/XPBDCorotatedConstraints.h"
 #include "GeometryCollection/Facades/CollectionMuscleActivationFacade.h"
 
+DEFINE_LOG_CATEGORY_STATIC(LogMuscleActivationConstraint, Log, All);
+
 namespace Chaos::Softs
 {
 
@@ -20,34 +22,47 @@ namespace Chaos::Softs
 
 		virtual ~FMuscleActivationConstraints() {}
 
-		void AddMuscles(const Facade& FMuscleActivation, int32 VertexOffset = 0, int32 ElementOffset = 0)
+		void AddMuscles(const ParticleType& RestParticles, const Facade& FMuscleActivation, int32 VertexOffset = 0, int32 ElementOffset = 0)
 		{
 			for (int MuscleIdx = 0; MuscleIdx < FMuscleActivation.NumMuscles(); MuscleIdx++)
 			{
 				Data MuscleActivationData = FMuscleActivation.GetMuscleActivationData(MuscleIdx);
 				if (FMuscleActivation.IsValidGeometryIndex(MuscleActivationData.GeometryGroupIndex))
 				{
-					int32 OldSize = MuscleActivationElement.Num();
-					MuscleActivationElement.SetNum(OldSize + 1);
-					FiberDirectionMatrix.SetNum(OldSize + 1);
-					ContractionVolumeScale.SetNum(OldSize + 1);
-					for (int32 i = 0; i < MuscleActivationData.MuscleActivationElement.Num(); i++)
+					if (ensureMsgf(MuscleActivationData.MuscleLengthRatioThresholdForMaxActivation > 0 &&
+						MuscleActivationData.MuscleLengthRatioThresholdForMaxActivation < 1,
+						TEXT("MuscleLengthRatioThresholdForMaxActivation %f of muscle indexed %d is out of range (0,1), please check your setup."), MuscleActivationData.MuscleLengthRatioThresholdForMaxActivation, MuscleIdx)
+						&&
+						ensureMsgf(MuscleActivationData.MuscleActivationElement.Num() == MuscleActivationData.FiberDirectionMatrix.Num(),
+							TEXT("MuscleActivationElement size %d is not equal to FiberDirectionMatrix size %d for muscle indexed %d"), MuscleActivationData.MuscleActivationElement.Num(), MuscleActivationData.FiberDirectionMatrix.Num(), MuscleIdx)
+						&&
+						ensureMsgf(MuscleActivationData.FiberLengthRatioAtMaxActivation > 0 &&
+							MuscleActivationData.FiberLengthRatioAtMaxActivation < 1,
+							TEXT("FiberLengthRatioAtMaxActivation %f of muscle indexed %d is out of range (0,1), please check your setup."), MuscleActivationData.FiberLengthRatioAtMaxActivation, MuscleIdx)
+						)
 					{
-						if (FMuscleActivation.IsValidElementIndex(MuscleActivationData.MuscleActivationElement[i]))
+						const int32 OldSize = MuscleActivationElement.AddDefaulted(1);
+						FiberDirectionMatrix.AddDefaulted(1);
+						ContractionVolumeScale.AddDefaulted(1);
+						for (int32 i = 0; i < MuscleActivationData.MuscleActivationElement.Num(); i++)
 						{
-							MuscleActivationElement[OldSize].Add(MuscleActivationData.MuscleActivationElement[i] + ElementOffset);
-							FiberDirectionMatrix[OldSize].Add(MuscleActivationData.FiberDirectionMatrix[i]);
-							ContractionVolumeScale[OldSize].Add(MuscleActivationData.ContractionVolumeScale[i]);
+							if (FMuscleActivation.IsValidElementIndex(MuscleActivationData.MuscleActivationElement[i]))
+							{
+								MuscleActivationElement[OldSize].Add(MuscleActivationData.MuscleActivationElement[i] + ElementOffset);
+								FiberDirectionMatrix[OldSize].Add(MuscleActivationData.FiberDirectionMatrix[i]);
+								ContractionVolumeScale[OldSize].Add(MuscleActivationData.ContractionVolumeScale[i]);
+							}
 						}
+						FiberLengthRatioAtMaxActivation.Add(MuscleActivationData.FiberLengthRatioAtMaxActivation);
+						MuscleLengthRatioThresholdForMaxActivation.Add(MuscleActivationData.MuscleLengthRatioThresholdForMaxActivation);
+						MuscleActivationData.OriginInsertionPair[0] += VertexOffset;
+						MuscleActivationData.OriginInsertionPair[1] += VertexOffset;
+						OriginInsertionPair.Add(MuscleActivationData.OriginInsertionPair);
+						MuscleRestLength.Add((RestParticles.X(MuscleActivationData.OriginInsertionPair[0]) - RestParticles.X(MuscleActivationData.OriginInsertionPair[1])).Size());
+						MuscleActivation.Add(1.f);
+						MuscleVertexOffset.Add(FMuscleActivation.MuscleVertexOffset(MuscleIdx));
+						MuscleVertexCount.Add(FMuscleActivation.NumMuscleVertices(MuscleIdx));
 					}
-					MuscleActivationData.OriginInsertionPair[0] += VertexOffset;
-					MuscleActivationData.OriginInsertionPair[1] += VertexOffset;
-					OriginInsertionPair.Add(MuscleActivationData.OriginInsertionPair);
-					OriginInsertionRestLength.Add(MuscleActivationData.OriginInsertionRestLength);
-					FiberLengthRatio.Add(1.f);
-
-					MuscleVertexOffset.Add(FMuscleActivation.MuscleVertexOffset(MuscleIdx));
-					MuscleVertexCount.Add(FMuscleActivation.NumMuscleVertices(MuscleIdx));
 				}
 			}
 		}
@@ -56,11 +71,22 @@ namespace Chaos::Softs
 		{
 			for (int32 MuscleIdx = 0; MuscleIdx < MuscleActivationElement.Num(); MuscleIdx++)
 			{
-				if (ensureMsgf(MuscleActivationElement[MuscleIdx].Num() == FiberDirectionMatrix[MuscleIdx].Num(),
-					TEXT("MuscleActivationElement[%d].Num() = %d, not equal to FiberDirectionMatrix[%d].Num() = %d"), MuscleIdx, MuscleActivationElement[MuscleIdx].Num(), MuscleIdx, FiberDirectionMatrix[MuscleIdx].Num()))
+				if (ensureMsgf(MuscleLengthRatioThresholdForMaxActivation[MuscleIdx] > 0 && 
+					MuscleLengthRatioThresholdForMaxActivation[MuscleIdx] < 1,
+					TEXT("MuscleLengthRatioThresholdForMaxActivation %f of muscle indexed %d is out of range (0,1), please check your setup."), MuscleLengthRatioThresholdForMaxActivation[MuscleIdx], MuscleIdx))
 				{
-					//update origin/insertion length
-					FiberLengthRatio[MuscleIdx] = (InParticles.X(OriginInsertionPair[MuscleIdx][0]) - InParticles.X(OriginInsertionPair[MuscleIdx][1])).Size() / OriginInsertionRestLength[MuscleIdx];
+					// calculate origin/insertion length
+					const float MuscleLengthRatio = (InParticles.X(OriginInsertionPair[MuscleIdx][0]) - InParticles.X(OriginInsertionPair[MuscleIdx][1])).Size() / MuscleRestLength[MuscleIdx];
+					if (MuscleLengthRatio >= 1.f)
+					{
+						// not active
+						MuscleActivation[MuscleIdx] = 0.f;
+					}
+					else
+					{
+						// linear muscle activation model: muscle reaches max activation 1 at threshold length
+						MuscleActivation[MuscleIdx] = FMath::Min(1.f, (1.f - MuscleLengthRatio) / (1 - MuscleLengthRatioThresholdForMaxActivation[MuscleIdx]));
+					}
 				}
 			}
 		}
@@ -70,11 +96,15 @@ namespace Chaos::Softs
 			for (int32 MuscleIdx = 0; MuscleIdx < MuscleActivationElement.Num(); MuscleIdx++)
 			{
 				if (ensureMsgf(MuscleActivationElement[MuscleIdx].Num() == FiberDirectionMatrix[MuscleIdx].Num(), 
-					TEXT("MuscleActivationElement[%d].Num() = %d, not equal to FiberDirectionMatrix[%d].Num() = %d"), MuscleIdx, MuscleActivationElement[MuscleIdx].Num(), MuscleIdx, FiberDirectionMatrix[MuscleIdx].Num()))
+					TEXT("MuscleActivationElement[%d].Num() = %d, not equal to FiberDirectionMatrix[%d].Num() = %d"), MuscleIdx, MuscleActivationElement[MuscleIdx].Num(), MuscleIdx, FiberDirectionMatrix[MuscleIdx].Num())
+					&& ensureMsgf(FiberLengthRatioAtMaxActivation[MuscleIdx] > 0 &&
+						FiberLengthRatioAtMaxActivation[MuscleIdx] < 1,
+						TEXT("FiberLengthRatioAtMaxActivation %f of muscle indexed %d is out of range (0,1), please check your setup."), FiberLengthRatioAtMaxActivation[MuscleIdx], MuscleIdx))
 				{
+					const float FiberLengthRatio = 1.f - MuscleActivation[MuscleIdx] * (1.f - FiberLengthRatioAtMaxActivation[MuscleIdx]);
 					for (int32 ElemIdx = 0; ElemIdx < MuscleActivationElement[MuscleIdx].Num(); ElemIdx++)
 					{
-						Constraints.ModifyDmInverseFromFiberLength(MuscleActivationElement[MuscleIdx][ElemIdx], FiberLengthRatio[MuscleIdx] * FiberLengthRatio[MuscleIdx], FiberDirectionMatrix[MuscleIdx][ElemIdx], ContractionVolumeScale[MuscleIdx][ElemIdx]);
+						Constraints.ModifyDmInverseFromMuscleLength(MuscleActivationElement[MuscleIdx][ElemIdx], FiberLengthRatio, FiberDirectionMatrix[MuscleIdx][ElemIdx], ContractionVolumeScale[MuscleIdx][ElemIdx]);
 					}
 				}
 			}
@@ -84,17 +114,18 @@ namespace Chaos::Softs
 		int32 GetMuscleVertexOffset(int32 MuscleIndex) { return MuscleVertexOffset[MuscleIndex]; };
 		int32 GetMuscleVertexCount(int32 MuscleIndex) { return MuscleVertexCount[MuscleIndex]; };
 
-		/* returns muscle activation in [0,1]. max muscle activation 1 at 0.75 fiber length ratio or less */
-		float GetMuscleActivation(int32 MuscleIndex) { return FMath::Max(FMath::Min(4.f * (1.f - FiberLengthRatio[MuscleIndex]), 1), 0); };
+		/* returns muscle activation of the specified muscle */
+		float GetMuscleActivation(int32 MuscleIndex) const { return MuscleActivation[MuscleIndex]; };
 
 	private:	
 		TArray<TArray<int32>> MuscleActivationElement;
 		TArray<FIntVector2> OriginInsertionPair;
-		TArray<float> OriginInsertionRestLength;
-		TArray<float> FiberLengthRatio;
+		TArray<float> MuscleRestLength;
+		TArray<float> MuscleActivation;
 		TArray<TArray<Chaos::PMatrix33d>> FiberDirectionMatrix;
 		TArray<TArray<float>> ContractionVolumeScale;
-
+		TArray<float> FiberLengthRatioAtMaxActivation;
+		TArray<float> MuscleLengthRatioThresholdForMaxActivation;
 		TArray<int32> MuscleVertexOffset;
 		TArray<int32> MuscleVertexCount;
 	};
diff --git a/Engine/Source/Runtime/Experimental/Chaos/Public/Chaos/XPBDCorotatedConstraints.h b/Engine/Source/Runtime/Experimental/Chaos/Public/Chaos/XPBDCorotatedConstraints.h
--- a/Engine/Source/Runtime/Experimental/Chaos/Public/Chaos/XPBDCorotatedConstraints.h
+++ b/Engine/Source/Runtime/Experimental/Chaos/Public/Chaos/XPBDCorotatedConstraints.h
@@ -406,13 +406,13 @@ namespace Chaos::Softs
 			return dC2;
 		}
 
-		void ModifyDmInverseFromFiberLength(const int32 e, const T L, const PMatrix<T, 3, 3>& MFiberDir, const T ContractionVolumeScale) const
+		void ModifyDmInverseFromMuscleLength(const int32 e, const T FiberLengthRatio, const PMatrix<T, 3, 3>& MFiberDir, const T ContractionVolumeScale) const
 		{
-			TRACE_CPUPROFILER_EVENT_SCOPE(STAT_ChaosXPBDCorotatedModifyDmInverseFromFiberLength);
+			TRACE_CPUPROFILER_EVENT_SCOPE(STAT_ChaosXPBDCorotatedModifyDmInverseFromMuscleLength);
 			PMatrix<T, 3, 3> DmInv = ElementDmInvSave(e);
-			if (L < 1)
+			if (FiberLengthRatio < 1)
 			{
-				PMatrix<T, 3, 3> D(1 / L, FMath::Pow(L, ContractionVolumeScale/T(2)), FMath::Pow(L, ContractionVolumeScale/T(2)));
+				PMatrix<T, 3, 3> D(1 / FiberLengthRatio, FMath::Pow(FiberLengthRatio, ContractionVolumeScale/T(2)), FMath::Pow(FiberLengthRatio, ContractionVolumeScale/T(2)));
 				PMatrix<T, 3, 3> Factor = MFiberDir * D * MFiberDir.GetTransposed();
 				DmInv = Factor * DmInv; //Matrix multiplication convention: AB = B*A
 			}
diff --git a/Engine/Source/Runtime/Experimental/Chaos/Public/GeometryCollection/Facades/CollectionMuscleActivationFacade.h b/Engine/Source/Runtime/Experimental/Chaos/Public/GeometryCollection/Facades/CollectionMuscleActivationFacade.h
--- a/Engine/Source/Runtime/Experimental/Chaos/Public/GeometryCollection/Facades/CollectionMuscleActivationFacade.h
+++ b/Engine/Source/Runtime/Experimental/Chaos/Public/GeometryCollection/Facades/CollectionMuscleActivationFacade.h
@@ -18,12 +18,13 @@ namespace GeometryCollection::Facades
 		FIntVector2 OriginInsertionPair; //Muscle origin point and insertion point (to determine muscle length)
 		float OriginInsertionRestLength; //Muscle origin-insertion rest length
 		TArray<Chaos::PMatrix33d> FiberDirectionMatrix; //Per-element fiber direction orthogonal matrix: [v, w1, w2], v is the fiber direction
-		TArray<float> ContractionVolumeScale; // Per-element volume scale for muscle contraction
+		TArray<float> ContractionVolumeScale; // Per-element volume scale for muscle contraction. Muscles gain volume during contraction if > 1. Volume-preserving if 1.
+		float FiberLengthRatioAtMaxActivation; // How much muscle fibers shorten at max activation 1. A smaller value means more contraction in the fiber direction.
+		float MuscleLengthRatioThresholdForMaxActivation; // Muscle length ratio (defined by origin-insertion distance) below this threshold is considered to reach max activation 1.
 		TArray<TArray<FVector3f>> FiberStreamline; //Fiber streamline(s) for inverse dynamics
 		TArray<float> FiberStreamlineRestLength; //Fiber streamline rest length(s)
 	};
-
-	/** Muscle Activation Facade */
+	
 	class FMuscleActivationFacade
 	{
 	public:
@@ -35,6 +36,8 @@ namespace GeometryCollection::Facades
 		static CHAOS_API const FName OriginInsertionRestLength;
 		static CHAOS_API const FName FiberDirectionMatrix;
 		static CHAOS_API const FName ContractionVolumeScale;
+		static CHAOS_API const FName FiberLengthRatioAtMaxActivation;
+		static CHAOS_API const FName MuscleLengthRatioThresholdForMaxActivation;
 		static CHAOS_API const FName FiberStreamline;
 		static CHAOS_API const FName FiberStreamlineRestLength;
 
@@ -54,15 +57,28 @@ namespace GeometryCollection::Facades
 		//  Skeletal Mesh Bone Bindings
 		//
 		CHAOS_API int32 AddMuscleActivationData(const FMuscleActivationData& InputData);
-		CHAOS_API bool UpdateMuscleActivationData(const int32 DataIndex, FMuscleActivationData& InputData);
+		CHAOS_API bool UpdateMuscleActivationData(const int32 DataIndex, const FMuscleActivationData& InputData);
 		CHAOS_API FMuscleActivationData GetMuscleActivationData(const int32 DataIndex) const;
 		CHAOS_API bool IsValidGeometryIndex(const int32 Index) const { return 0 <= Index && Index < ConstCollection.NumElements(FGeometryCollection::GeometryGroup); };
 		CHAOS_API bool IsValidElementIndex(const int32 Index) const { return 0 <= Index && Index < ConstCollection.NumElements("Tetrahedral"); };
 		CHAOS_API int32 NumMuscles() const { return MuscleActivationElementAttribute.Num(); }
+		CHAOS_API bool IsValidMuscleIndex(const int32 Index) const { return 0 <= Index && Index < NumMuscles(); };
 		CHAOS_API int32 MuscleVertexOffset(const int32 MuscleIndex) const;
 		CHAOS_API int32 NumMuscleVertices(const int32 MuscleIndex) const;
+		CHAOS_API FString FindMuscleName(const int32 MuscleIndex) const;
+		CHAOS_API int32 FindMuscleIndexByName(const FString MuscleName) const;
 		CHAOS_API int32 RemoveInvalidMuscles();
+		CHAOS_API bool SetUpMuscleActivation(const TArray<int32>& Origin, const TArray<int32>& Insertion);
 		CHAOS_API bool SetUpMuscleActivation(const TArray<int32>& Origin, const TArray<int32>& Insertion, float ContractionVolumeScale);
+		CHAOS_API void UpdateGlobalMuscleActivationParameters(
+			float InContractionVolumeScale,
+			float InGlobalFiberLengthRatioAtMaxActivation,
+			float InGlobalMuscleLengthRatioThresholdForMaxActivation);
+		CHAOS_API bool UpdateMuscleActivationParameters(
+			int32 MuscleIndex,
+			float InContractionVolumeScale,
+			float InGlobalFiberLengthRatioAtMaxActivation,
+			float InGlobalMuscleLengthRatioThresholdForMaxActivation);
 		CHAOS_API TArray<TArray<TArray<FVector3f>>> BuildStreamlines(const TArray<int32>& Origin, const TArray<int32>& Insertion,
 			 int32 NumLinesMultiplier, int32 MaxStreamlineIterations, int32 MaxPointsPerLine);
 	private:
@@ -74,6 +90,8 @@ namespace GeometryCollection::Facades
 		TManagedArrayAccessor<float> OriginInsertionRestLengthAttribute;
 		TManagedArrayAccessor<TArray<Chaos::PMatrix33d>> FiberDirectionMatrixAttribute;
 		TManagedArrayAccessor<TArray<float>> ContractionVolumeScaleAttribute;
+		TManagedArrayAccessor<float> FiberLengthRatioAtMaxActivationAttribute;
+		TManagedArrayAccessor<float> MuscleLengthRatioThresholdForMaxActivationAttribute;
 		TManagedArrayAccessor<TArray<TArray<FVector3f>>> FiberStreamlineAttribute;
 		TManagedArrayAccessor<TArray<float>> FiberStreamlineRestLengthAttribute;
 	};
diff --git a/Engine/Source/Runtime/Experimental/Dataflow/Engine/Public/Dataflow/DataflowObject.h b/Engine/Source/Runtime/Experimental/Dataflow/Engine/Public/Dataflow/DataflowObject.h
--- a/Engine/Source/Runtime/Experimental/Dataflow/Engine/Public/Dataflow/DataflowObject.h
+++ b/Engine/Source/Runtime/Experimental/Dataflow/Engine/Public/Dataflow/DataflowObject.h
@@ -234,6 +234,13 @@ public:
 	UPROPERTY(DuplicateTransient, AssetRegistrySearchable)
 	FTransform PreviewBlueprintTransform = FTransform::Identity;
 
+	/** Geometry cache asset used to extract skeletal mesh results from simulation */
+	UPROPERTY(DuplicateTransient, AssetRegistrySearchable)
+	TSoftObjectPtr<UObject> PreviewGeometryCacheAsset = nullptr;
+
+	/** SkeletalMesh interpolated from simulation */
+	UPROPERTY(DuplicateTransient, AssetRegistrySearchable)
+	TSoftObjectPtr<UObject> PreviewEmbeddedSkeletalMesh = nullptr;
 #endif
 
 #if WITH_EDITOR
