Dataflow SubGraph : Fix issues with subgraph call nodes
- Fixed deserialization issue where pins would not be properly restored and reconnected
- Call node now store a Guid to be resilient to renaming
- Change OnLoaded  subgraph delegate to be a global one instead of being on the object itself

diff --git a/Engine/Plugins/Experimental/Dataflow/Source/DataflowEditor/Private/Dataflow/DataflowGraphSchemaAction.cpp b/Engine/Plugins/Experimental/Dataflow/Source/DataflowEditor/Private/Dataflow/DataflowGraphSchemaAction.cpp
--- a/Engine/Plugins/Experimental/Dataflow/Source/DataflowEditor/Private/Dataflow/DataflowGraphSchemaAction.cpp
+++ b/Engine/Plugins/Experimental/Dataflow/Source/DataflowEditor/Private/Dataflow/DataflowGraphSchemaAction.cpp
@@ -6,6 +6,7 @@
 #include "Dataflow/DataflowInstance.h"
 #include "Dataflow/DataflowObject.h"
 #include "Dataflow/DataflowSNode.h"
+#include "Dataflow/DataflowSubGraph.h"
 #include "Dataflow/DataflowAssetEditUtils.h"
 #include "Dataflow/DataflowVariableNodes.h"
 #include "Dataflow/DataflowSubGraphNodes.h"
@@ -226,7 +227,7 @@ FEdGraphSchemaAction_DataflowSubGraph::FEdGraphSchemaAction_DataflowSubGraph()
 {
 }
 
-FEdGraphSchemaAction_DataflowSubGraph::FEdGraphSchemaAction_DataflowSubGraph(UDataflow* InDataflowAsset, FName InSubgraphName)
+FEdGraphSchemaAction_DataflowSubGraph::FEdGraphSchemaAction_DataflowSubGraph(UDataflow* InDataflowAsset, const FGuid& InSubGraphGuid)
 	: FEdGraphSchemaAction(
 		FText::GetEmpty(),
 		LOCTEXT("Description", "Description"),
@@ -234,16 +235,28 @@ FEdGraphSchemaAction_DataflowSubGraph::FEdGraphSchemaAction_DataflowSubGraph(UDa
 		/*InGrouping*/0,
 		/*Keywords*/FText::GetEmpty(),
 		(int32)UE::Dataflow::ESchemaActionSectionID::SUBGRAPHS)
-	, SubGraphName(InSubgraphName)
+	, SubGraphGuid(InSubGraphGuid)
 	, DataflowAssetWeakPtr(InDataflowAsset)
 {
 }
 
+const FName FEdGraphSchemaAction_DataflowSubGraph::GetSubGraphName() const
+{
+	if (TStrongObjectPtr<const UDataflow> DataflowAsset = DataflowAssetWeakPtr.Pin())
+	{
+		if (const UDataflowSubGraph* SubGraph = DataflowAsset->FindSubGraphByGuid(SubGraphGuid))
+		{
+			return SubGraph->GetFName();
+		}
+	}
+	return {};
+}
+
 UEdGraphNode* FEdGraphSchemaAction_DataflowSubGraph::PerformAction(class UEdGraph* ParentGraph, UEdGraphPin* FromPin, const FVector2D Location, bool bSelectNewNode)
 {
 	// TODO : move this code to the FEditableDataflow eventually 
 	UEdGraphNode* NewEdNode = nullptr;
-	if (TSharedPtr<FAssetSchemaAction_Dataflow_CreateNode_DataflowEdNode>  SubGraphNodeAction = FAssetSchemaAction_Dataflow_CreateNode_DataflowEdNode::CreateAction(ParentGraph, "FDataflowCallSubGraphNode", SubGraphName))
+	if (TSharedPtr<FAssetSchemaAction_Dataflow_CreateNode_DataflowEdNode>  SubGraphNodeAction = FAssetSchemaAction_Dataflow_CreateNode_DataflowEdNode::CreateAction(ParentGraph, "FDataflowCallSubGraphNode", GetSubGraphName()))
 	{
 		UDataflowEdNode* DataflowEdNode = Cast<UDataflowEdNode>(SubGraphNodeAction->PerformAction(ParentGraph, nullptr, Location, false));
 		NewEdNode = DataflowEdNode;
@@ -254,7 +267,7 @@ UEdGraphNode* FEdGraphSchemaAction_DataflowSubGraph::PerformAction(class UEdGrap
 			{
 				if (TStrongObjectPtr<UDataflow> DataflowAsset = DataflowAssetWeakPtr.Pin())
 				{
-					SubGraphCallNode->SetSubGraphName(SubGraphName);
+					SubGraphCallNode->SetSubGraphGuid(SubGraphGuid);
 					DataflowEdNode->UpdatePinsFromDataflowNode();
 					ParentGraph->NotifyNodeChanged(DataflowEdNode);
 				}
diff --git a/Engine/Plugins/Experimental/Dataflow/Source/DataflowEditor/Private/Dataflow/DataflowMembersWidget.cpp b/Engine/Plugins/Experimental/Dataflow/Source/DataflowEditor/Private/Dataflow/DataflowMembersWidget.cpp
--- a/Engine/Plugins/Experimental/Dataflow/Source/DataflowEditor/Private/Dataflow/DataflowMembersWidget.cpp
+++ b/Engine/Plugins/Experimental/Dataflow/Source/DataflowEditor/Private/Dataflow/DataflowMembersWidget.cpp
@@ -140,7 +140,7 @@ namespace UE::Dataflow::Private
 					if (ensure(SubGraph))
 					{
 						TSharedPtr<FEdGraphSchemaAction_DataflowSubGraph> NewSubGraphAction =
-							MakeShareable(new FEdGraphSchemaAction_DataflowSubGraph(DataflowAsset, SubGraph->GetFName()));
+							MakeShareable(new FEdGraphSchemaAction_DataflowSubGraph(DataflowAsset, SubGraph->GetSubGraphGuid()));
 						OutActions.Add(NewSubGraphAction);
 					}
 				}
diff --git a/Engine/Plugins/Experimental/Dataflow/Source/DataflowEditor/Public/Dataflow/DataflowGraphSchemaAction.h b/Engine/Plugins/Experimental/Dataflow/Source/DataflowEditor/Public/Dataflow/DataflowGraphSchemaAction.h
--- a/Engine/Plugins/Experimental/Dataflow/Source/DataflowEditor/Public/Dataflow/DataflowGraphSchemaAction.h
+++ b/Engine/Plugins/Experimental/Dataflow/Source/DataflowEditor/Public/Dataflow/DataflowGraphSchemaAction.h
@@ -160,16 +160,13 @@ struct FEdGraphSchemaAction_DataflowSubGraph : public FEdGraphSchemaAction
 public:
 	FEdGraphSchemaAction_DataflowSubGraph();
 
-	FEdGraphSchemaAction_DataflowSubGraph(UDataflow* InDataflowAsset, FName SubGraphName);
+	FEdGraphSchemaAction_DataflowSubGraph(UDataflow* InDataflowAsset, const FGuid& InSubGraphGuid);
 
 	// Simple type info
 	static FName StaticGetTypeId() { static FName Type("FEdGraphSchemaAction_DataflowSubGraph"); return Type; }
 	virtual FName GetTypeId() const override { return StaticGetTypeId(); }
 
-	const FName GetSubGraphName() const
-	{
-		return SubGraphName;
-	}
+	const FName GetSubGraphName() const;
 
 	// FEdGraphSchemaAction interface
 	virtual bool IsAVariable() const { return false; }
@@ -185,8 +182,8 @@ public:
 	void DuplicateItem();
 
 private:
-	/** Name of the subgraph */
-	FName SubGraphName;
+	/** Guid of the subgraph */
+	FGuid SubGraphGuid;
 
 	/** Dataflow asset associated with this action */
 	TWeakObjectPtr<UDataflow> DataflowAssetWeakPtr;
diff --git a/Engine/Source/Runtime/Experimental/Dataflow/Engine/Private/Dataflow/DataflowSubGraph.cpp b/Engine/Source/Runtime/Experimental/Dataflow/Engine/Private/Dataflow/DataflowSubGraph.cpp
--- a/Engine/Source/Runtime/Experimental/Dataflow/Engine/Private/Dataflow/DataflowSubGraph.cpp
+++ b/Engine/Source/Runtime/Experimental/Dataflow/Engine/Private/Dataflow/DataflowSubGraph.cpp
@@ -28,6 +28,10 @@ namespace UE::Dataflow::DataflowSubGraph::Private
 
 /////////////////////////////////////////////////////////////////////////////////////////////////////////////
 
+FDataflowSubGraphDelegates::FOnDataflowSubGraphLoaded FDataflowSubGraphDelegates::OnSubGraphLoaded;
+
+/////////////////////////////////////////////////////////////////////////////////////////////////////////////
+
 UDataflowSubGraph::UDataflowSubGraph(class FObjectInitializer const& ObjectInitializer)
 	: UEdGraph(ObjectInitializer)
 	, SubGraphGuid(FGuid::NewGuid())
@@ -46,23 +50,21 @@ FDataflowSubGraphOutputNode* UDataflowSubGraph::GetOutputNode() const
 	return UE::Dataflow::DataflowSubGraph::Private::GetNodeByType<FDataflowSubGraphOutputNode>(*this);
 }
 
+void UDataflowSubGraph::Serialize(FArchive& Ar)
+{
+	Super::Serialize(Ar);
+
+	UE_LOG(LogTemp, Warning, TEXT("## UDataflowSubGraph::Serialize"));
+}
+
 void UDataflowSubGraph::PostLoad()
 {
 	Super::PostLoad();
 	
 	bIsLoaded = true;
-	if (OnSubGraphLoaded.IsBound())
+	if (FDataflowSubGraphDelegates::OnSubGraphLoaded.IsBound())
 	{
-		OnSubGraphLoaded.Broadcast(*this);
+		FDataflowSubGraphDelegates::OnSubGraphLoaded.Broadcast(*this);
 	}
 }
 
-FDelegateHandle UDataflowSubGraph::AddOnSubGraphLoadedHandler(const FOnDataflowSubGraphLoaded::FDelegate& InHandler)
-{
-	return OnSubGraphLoaded.Add(InHandler);
-}
-
-void UDataflowSubGraph::RemoveOnSubGraphLoadedHandler(FDelegateHandle Handle)
-{
-	OnSubGraphLoaded.Remove(Handle);
-}
\ No newline at end of file
diff --git a/Engine/Source/Runtime/Experimental/Dataflow/Engine/Private/Dataflow/DataflowSubGraphNodes.cpp b/Engine/Source/Runtime/Experimental/Dataflow/Engine/Private/Dataflow/DataflowSubGraphNodes.cpp
--- a/Engine/Source/Runtime/Experimental/Dataflow/Engine/Private/Dataflow/DataflowSubGraphNodes.cpp
+++ b/Engine/Source/Runtime/Experimental/Dataflow/Engine/Private/Dataflow/DataflowSubGraphNodes.cpp
@@ -156,7 +156,8 @@ void FDataflowDynamicConnections::Refresh()
 	{
 		if (DynamicConnection)
 		{
-			if (const FGuid* PropertyId = ConnectionNameToPropertyId.Find(DynamicConnection->GetName()))
+			const FGuid* PropertyId = ConnectionNameToPropertyId.Find(DynamicConnection->GetName());
+			if (PropertyId)
 			{
 				FRemoteConnections& RemoteConnections = RemoteConnectionsByPropertyId.Emplace(*PropertyId);
 				DynamicConnection->GetConnections(RemoteConnections);
@@ -260,9 +261,7 @@ bool FDataflowDynamicConnections::IsSupportedType(const FPropertyBagPropertyDesc
 
 FName FDataflowDynamicConnections::GetCppTypeFromPropertyDesc(const FPropertyBagPropertyDesc& Desc)
 {
-	FString ExtendedType;
-	const FString CPPType = Desc.CachedProperty->GetCPPType(&ExtendedType);
-	return FName(CPPType + ExtendedType);
+	return FDataflowConnection::GetTypeNameFromProperty(Desc.CachedProperty);
 }
 
 bool FDataflowDynamicConnections::SetConnectionTypeFromPropertyDesc(FDataflowConnection& Connection, const FPropertyBagPropertyDesc& Desc)
@@ -418,6 +417,8 @@ void FDataflowSubGraphOutputNode::OnPropertyChanged(UE::Dataflow::FContext& Cont
 
 //////////////////////////////////////////////////////////////////////////////////////////////////////////
 
+#define NEW_METHOD 1
+
 FDataflowCallSubGraphNode::FDataflowCallSubGraphNode(const UE::Dataflow::FNodeParameters& InParam, FGuid InGuid)
 	: FDataflowNode(InParam, InGuid)
 	, DynamicInputs(UE::Dataflow::FPin::EDirection::INPUT, this, Cast<UDataflow>(InParam.OwningObject))
@@ -430,11 +431,29 @@ FDataflowCallSubGraphNode::~FDataflowCallSubGraphNode()
 {
 	if (TStrongObjectPtr<UDataflow> DataflowAsset = DataflowAssetWeakPtr.Pin())
 	{
-		if (UDataflowSubGraph* SubGraph = DataflowAsset->FindSubGraphByName(SubGraphName))
-		{
-			SubGraph->RemoveOnGraphChangedHandler(OnGraphChangedHandle);
-			SubGraph->RemoveOnSubGraphLoadedHandler(OnSubGraphLoadedHandle);
-		}
+		UnregisterHandlers(DataflowAsset->FindSubGraphByGuid(SubGraphGuid));
+	}
+}
+
+void FDataflowCallSubGraphNode::RegisterHandlers(UDataflowSubGraph* SubGraph)
+{
+	if (SubGraph)
+	{
+		OnSubGraphLoadedHandle = SubGraph->AddOnGraphChangedHandler(FOnGraphChanged::FDelegate::CreateRaw(this, &FDataflowCallSubGraphNode::OnSubGraphChanged));
+	}
+
+	FDataflowSubGraphDelegates::OnSubGraphLoaded.AddRaw(this, &FDataflowCallSubGraphNode::OnSubGraphLoaded);
+	FDataflowAssetDelegates::OnSubGraphsChanged.AddRaw(this, &FDataflowCallSubGraphNode::OnSomeSubGraphsChanged);
+}
+
+void FDataflowCallSubGraphNode::UnregisterHandlers(UDataflowSubGraph* SubGraph)
+{
+	FDataflowAssetDelegates::OnSubGraphsChanged.RemoveAll(this);
+	FDataflowSubGraphDelegates::OnSubGraphLoaded.RemoveAll(this);
+
+	if (SubGraph)
+	{
+		SubGraph->RemoveOnGraphChangedHandler(OnGraphChangedHandle);
 	}
 }
 
@@ -442,32 +461,22 @@ void FDataflowCallSubGraphNode::PostSerialize(const FArchive& Ar)
 {
 	if (Ar.IsLoading())
 	{
-		// make sure the node is up to date with the serialized data
-		SetSubGraphName(SubGraphName);
+		SetSubGraphGuid(SubGraphGuid);
 	}
 }
 
-void FDataflowCallSubGraphNode::SetSubGraphName(FName InSubGraphName)
+void FDataflowCallSubGraphNode::SetSubGraphGuid(const FGuid& InSubGraphGuid)
 {
 	if (TStrongObjectPtr<UDataflow> DataflowAsset = DataflowAssetWeakPtr.Pin())
 	{
-		if (UDataflowSubGraph* PreviousSubGraph = DataflowAsset->FindSubGraphByName(SubGraphName))
-		{
-			PreviousSubGraph->RemoveOnGraphChangedHandler(OnGraphChangedHandle);
-			PreviousSubGraph->RemoveOnSubGraphLoadedHandler(OnSubGraphLoadedHandle);
-		}
+		UnregisterHandlers(DataflowAsset->FindSubGraphByGuid(SubGraphGuid));
 
-		SubGraphName = InSubGraphName;
+		SubGraphGuid = InSubGraphGuid;
 
-		if (UDataflowSubGraph* NewSubGraph = DataflowAsset->FindSubGraphByName(SubGraphName))
-		{
-			OnSubGraphLoadedHandle = NewSubGraph->AddOnGraphChangedHandler(FOnGraphChanged::FDelegate::CreateRaw(this, &FDataflowCallSubGraphNode::OnSubGraphChanged));
-			OnGraphChangedHandle = NewSubGraph->AddOnSubGraphLoadedHandler(FOnDataflowSubGraphLoaded::FDelegate::CreateRaw(this, &FDataflowCallSubGraphNode::OnSubGraphLoaded));
-		}
+		RegisterHandlers(DataflowAsset->FindSubGraphByGuid(SubGraphGuid));
 
-		// make sure we have at least an initial refresh 
-		DynamicInputs.Refresh();
-		DynamicOutputs.Refresh();
+		SyncInputsPropertyBagWithSubGraph();
+		SyncOutputsPropertyBagWithSubGraph();
 	}
 }
 
@@ -496,43 +505,89 @@ void FDataflowCallSubGraphNode::OnSubGraphChanged(const struct FEdGraphEditActio
 	}
 	if (bInputsChanged)
 	{
-		DynamicInputs.Refresh();
+		SyncInputsPropertyBagWithSubGraph();
 	}
 	if (bOutputsChanged)
 	{
-		DynamicOutputs.Refresh();
+		SyncOutputsPropertyBagWithSubGraph();
+	}
+}
+
+void FDataflowCallSubGraphNode::OnSomeSubGraphsChanged(const UDataflow* InDataflowAsset, const FGuid& InSubGraphGuid, UE::Dataflow::ESubGraphChangedReason InReason)
+{
+	if (InSubGraphGuid == SubGraphGuid)
+	{
+		switch (InReason)
+		{
+		case UE::Dataflow::ESubGraphChangedReason::Created:
+			// nothing to do, we should not have a newly created guid
+			break;
+
+		case UE::Dataflow::ESubGraphChangedReason::Renamed:
+			// Todo(Dataflow) : we do not care about the name right now but we may if the name is displayed 
+			break;
+
+		case UE::Dataflow::ESubGraphChangedReason::Deleted:
+			SetSubGraphGuid(FGuid());
+			break;
+		}
 	}
 }
 
 void FDataflowCallSubGraphNode::OnSubGraphLoaded(const UDataflowSubGraph& InSubGraph)
 {
-	DynamicInputs.Refresh();
-	DynamicOutputs.Refresh();
+	if (InSubGraph.GetSubGraphGuid() == SubGraphGuid)
+	{
+		SetSubGraphGuid(SubGraphGuid);
+	}
 }
 
-const FInstancedPropertyBag& FDataflowCallSubGraphNode::GetPropertyBag(const FDataflowDynamicConnections* Caller)
+void FDataflowCallSubGraphNode::SyncInputsPropertyBagWithSubGraph()
 {
-	static FInstancedPropertyBag DefaultPropertyBag;
 	if (TStrongObjectPtr<UDataflow> DataflowAsset = DataflowAssetWeakPtr.Pin())
 	{
-		if (UDataflowSubGraph* SubGraphToCall = DataflowAsset->FindSubGraphByName(SubGraphName))
+		if (UDataflowSubGraph* SubGraph = DataflowAsset->FindSubGraphByGuid(SubGraphGuid))
 		{
-			if (Caller == &DynamicInputs)
+			InputsPropertyBag.Reset();
+			if (FDataflowSubGraphInputNode* SubGraphInputNode = SubGraph->GetInputNode())
 			{
-				if (FDataflowSubGraphInputNode* SubGraphInputNode = SubGraphToCall->GetInputNode())
-				{
-					return SubGraphInputNode->GetPropertyBag(Caller);
-				}
+				InputsPropertyBag = SubGraphInputNode->GetPropertyBag(&DynamicInputs);
 			}
-			else if (Caller == &DynamicOutputs)
+		}
+	}
+	DynamicInputs.Refresh();
+}
+
+void FDataflowCallSubGraphNode::SyncOutputsPropertyBagWithSubGraph()
+{
+	if (TStrongObjectPtr<UDataflow> DataflowAsset = DataflowAssetWeakPtr.Pin())
+	{
+		if (UDataflowSubGraph* SubGraph = DataflowAsset->FindSubGraphByGuid(SubGraphGuid))
+		{
+			OutputsPropertyBag.Reset();
+			if (FDataflowSubGraphOutputNode* SubGraphOutputNode = SubGraph->GetOutputNode())
 			{
-				if (FDataflowSubGraphOutputNode* SubGraphOutputNode = SubGraphToCall->GetOutputNode())
-				{
-					return SubGraphOutputNode->GetPropertyBag(Caller);
-				}
+				OutputsPropertyBag = SubGraphOutputNode->GetPropertyBag(&DynamicOutputs);
 			}
 		}
 	}
+	DynamicOutputs.Refresh();
+}
+
+const FInstancedPropertyBag& FDataflowCallSubGraphNode::GetPropertyBag(const FDataflowDynamicConnections* Caller)
+{
+	static FInstancedPropertyBag DefaultPropertyBag;
+	if (TStrongObjectPtr<UDataflow> DataflowAsset = DataflowAssetWeakPtr.Pin())
+	{
+		if (Caller == &DynamicInputs)
+		{
+			return InputsPropertyBag;
+		}
+		else if (Caller == &DynamicOutputs)
+		{
+			return OutputsPropertyBag;
+		}
+	}
 	return DefaultPropertyBag;
 }
 
@@ -561,7 +616,7 @@ void FDataflowCallSubGraphNode::Evaluate(UE::Dataflow::FContext& Context, const
 	check(Out);
 	if (TStrongObjectPtr<UDataflow> DataflowAsset = DataflowAssetWeakPtr.Pin())
 	{
-		if (UDataflowSubGraph* SubGraphToCall = DataflowAsset->FindSubGraphByName(SubGraphName))
+		if (UDataflowSubGraph* SubGraphToCall = DataflowAsset->FindSubGraphByGuid(SubGraphGuid))
 		{
 			if (const FDataflowSubGraphOutputNode* SubGraphOutputNode = SubGraphToCall->GetOutputNode())
 			{
diff --git a/Engine/Source/Runtime/Experimental/Dataflow/Engine/Public/Dataflow/DataflowSubGraph.h b/Engine/Source/Runtime/Experimental/Dataflow/Engine/Public/Dataflow/DataflowSubGraph.h
--- a/Engine/Source/Runtime/Experimental/Dataflow/Engine/Public/Dataflow/DataflowSubGraph.h
+++ b/Engine/Source/Runtime/Experimental/Dataflow/Engine/Public/Dataflow/DataflowSubGraph.h
@@ -11,7 +11,12 @@ struct FDataflowSubGraphOutputNode;
 
 class UDataflowSubGraph;
 
-DECLARE_MULTICAST_DELEGATE_OneParam(FOnDataflowSubGraphLoaded, const UDataflowSubGraph&);
+struct FDataflowSubGraphDelegates
+{
+	/** Called when variables are edited ( add, remove, changetype , setvalue ) */
+	DECLARE_MULTICAST_DELEGATE_OneParam(FOnDataflowSubGraphLoaded, const UDataflowSubGraph&);
+	static DATAFLOWENGINE_API FOnDataflowSubGraphLoaded OnSubGraphLoaded;
+};
 
 UCLASS(MinimalAPI)
 class UDataflowSubGraph : public UEdGraph
@@ -22,13 +27,11 @@ public:
 	FDataflowSubGraphInputNode* GetInputNode() const;
 	FDataflowSubGraphOutputNode* GetOutputNode() const;
 
-	FDelegateHandle AddOnSubGraphLoadedHandler(const FOnDataflowSubGraphLoaded::FDelegate& InHandler);
-	void RemoveOnSubGraphLoadedHandler(FDelegateHandle Handle);
-
 	bool IsLoaded() const { return bIsLoaded; }
 	const FGuid& GetSubGraphGuid() const { return SubGraphGuid; }
 
 private:
+	virtual void Serialize(FArchive& Ar) override;
 	virtual void PostLoad() override;
 
 	/**
@@ -38,6 +41,5 @@ private:
 	UPROPERTY()
 	FGuid SubGraphGuid; 
 
-	FOnDataflowSubGraphLoaded OnSubGraphLoaded;
 	bool bIsLoaded : 1 = false;
 };
\ No newline at end of file
diff --git a/Engine/Source/Runtime/Experimental/Dataflow/Engine/Public/Dataflow/DataflowSubGraphNodes.h b/Engine/Source/Runtime/Experimental/Dataflow/Engine/Public/Dataflow/DataflowSubGraphNodes.h
--- a/Engine/Source/Runtime/Experimental/Dataflow/Engine/Public/Dataflow/DataflowSubGraphNodes.h
+++ b/Engine/Source/Runtime/Experimental/Dataflow/Engine/Public/Dataflow/DataflowSubGraphNodes.h
@@ -16,6 +16,8 @@ class UDataflowSubGraph;
 
 namespace UE::Dataflow
 {
+	enum class ESubGraphChangedReason: uint8;
+
 	struct ISubGraphContextCallback
 	{
 		virtual void EvaluateParentContext(UE::Dataflow::FContext& Context, UE::Dataflow::FContext& SubGraphContext, const FDataflowOutput& SubGraphOutput) const = 0;
@@ -144,7 +146,7 @@ public:
 	DATAFLOWENGINE_API FDataflowCallSubGraphNode(const UE::Dataflow::FNodeParameters& Param, FGuid InGuid = FGuid::NewGuid());
 	DATAFLOWENGINE_API virtual ~FDataflowCallSubGraphNode() override;
 
-	DATAFLOWENGINE_API void SetSubGraphName(FName InSubGraphName);
+	DATAFLOWENGINE_API void SetSubGraphGuid(const FGuid& Guid);
 
 private:
 	virtual void Evaluate(UE::Dataflow::FContext& Context, const FDataflowOutput* Out) const override;
@@ -159,18 +161,30 @@ private:
 	virtual const FInstancedPropertyBag& GetPropertyBag(const FDataflowDynamicConnections* Caller) override;
 	/* End IOwnerInterface */
 
+	void UnregisterHandlers(UDataflowSubGraph* SubGraph);
+	void RegisterHandlers(UDataflowSubGraph* SubGraph);
 	void OnSubGraphChanged(const struct FEdGraphEditAction& InAction);
+	void OnSomeSubGraphsChanged(const UDataflow* InDataflowAsset, const FGuid& InSubGraphGuid, UE::Dataflow::ESubGraphChangedReason InReason);
 	void OnSubGraphLoaded(const UDataflowSubGraph& InSubGraph);
+	
+	void SyncInputsPropertyBagWithSubGraph();
+	void SyncOutputsPropertyBagWithSubGraph();
 
 	UPROPERTY()
-	FName SubGraphName; // TODO : should probably be a GUID
+	FGuid SubGraphGuid;
 
 	UPROPERTY()
 	FDataflowDynamicConnections DynamicInputs;
 
+	UPROPERTY()
+	FInstancedPropertyBag InputsPropertyBag;
+
 	UPROPERTY()
 	FDataflowDynamicConnections DynamicOutputs;
 
+	UPROPERTY()
+	FInstancedPropertyBag OutputsPropertyBag;
+
 	TWeakObjectPtr<UDataflow> DataflowAssetWeakPtr;
 	FDelegateHandle OnGraphChangedHandle;
 	FDelegateHandle OnSubGraphLoadedHandle;
