Dataflow : Remove debug code causing warning in cook builds

diff --git a/Engine/Source/Runtime/Experimental/Dataflow/Engine/Private/Dataflow/DataflowSubGraph.cpp b/Engine/Source/Runtime/Experimental/Dataflow/Engine/Private/Dataflow/DataflowSubGraph.cpp
--- a/Engine/Source/Runtime/Experimental/Dataflow/Engine/Private/Dataflow/DataflowSubGraph.cpp
+++ b/Engine/Source/Runtime/Experimental/Dataflow/Engine/Private/Dataflow/DataflowSubGraph.cpp
@@ -50,13 +50,6 @@ FDataflowSubGraphOutputNode* UDataflowSubGraph::GetOutputNode() const
 	return UE::Dataflow::DataflowSubGraph::Private::GetNodeByType<FDataflowSubGraphOutputNode>(*this);
 }
 
-void UDataflowSubGraph::Serialize(FArchive& Ar)
-{
-	Super::Serialize(Ar);
-
-	UE_LOG(LogTemp, Warning, TEXT("## UDataflowSubGraph::Serialize"));
-}
-
 void UDataflowSubGraph::PostLoad()
 {
 	Super::PostLoad();
diff --git a/Engine/Source/Runtime/Experimental/Dataflow/Engine/Public/Dataflow/DataflowSubGraph.h b/Engine/Source/Runtime/Experimental/Dataflow/Engine/Public/Dataflow/DataflowSubGraph.h
--- a/Engine/Source/Runtime/Experimental/Dataflow/Engine/Public/Dataflow/DataflowSubGraph.h
+++ b/Engine/Source/Runtime/Experimental/Dataflow/Engine/Public/Dataflow/DataflowSubGraph.h
@@ -31,7 +31,6 @@ public:
 	const FGuid& GetSubGraphGuid() const { return SubGraphGuid; }
 
 private:
-	virtual void Serialize(FArchive& Ar) override;
 	virtual void PostLoad() override;
 
 	/**
