IasTool - added support for downloading the X latest TOC file(s) instead of the default returned from S3.

- updated S3 client to support paginated list requests
- updated upload command to download the X latest TOC files
- updated S3 upload queue to retry 3 times with delay

diff --git a/Engine/Source/Developer/S3Client/Internal/S3/S3Client.h b/Engine/Source/Developer/S3Client/Internal/S3/S3Client.h
--- a/Engine/Source/Developer/S3Client/Internal/S3/S3Client.h
+++ b/Engine/Source/Developer/S3Client/Internal/S3/S3Client.h
@@ -9,6 +9,7 @@
 #include "Containers/UnrealString.h"
 #include "Memory/MemoryView.h"
 #include "Memory/SharedBuffer.h"
+#include "Misc/DateTime.h"
 #include "Misc/StringBuilder.h"
 
 #define UE_API S3CLIENT_API
@@ -91,10 +92,20 @@ struct FS3Object
 {
 	/** The object identifider. */
 	FString Key;
+	/** Date and time when this object was last modified in text format. */
+	FString LastModifiedText;
 	/** Date and time when this object was last modified. */
-	FString LastModified;
+	FDateTime LastModified;
 	/** The size of the object. */
 	uint64 Size = 0;
+
+	bool operator==(const FS3Object& RHS) const { return Key == RHS.Key; }
+	bool operator!=(const FS3Object& RHS) const { return !operator==(RHS); }
+
+	friend inline uint32 GetTypeHash(const FS3Object& Object)
+	{
+		return GetTypeHash(Object.Key);
+	}
 };
 
 /** Basic response parameters. */
@@ -187,26 +198,38 @@ struct FS3ListObjectsRequest
 	FString Prefix;
 	/** The path delimiter, i.e. '/'. */
 	TCHAR Delimiter;
+	/* Max key(s). */
+	int32 MaxKeys = INDEX_NONE;
+	/** Marker from where to list objects. */
+	FString Marker;
 };
 
 /** Response parameters when listing objects. */
 struct FS3ListObjectResponse
 	: public FS3Response
 {
+	FS3ListObjectResponse() = default;
+
 	FS3ListObjectResponse(FS3Response&& Other)
 		: FS3Response(MoveTemp(Other))
 	{}
 
-	FS3ListObjectResponse(FS3Response&& Other, FString&& InBucketName, TArray<FS3Object>&& InObjects)
+	FS3ListObjectResponse(FS3Response&& Other, FString&& InBucketName, TArray<FS3Object>&& InObjects, FString InNextMarker = FString(), bool bInIsTruncated = false)
 		: FS3Response(MoveTemp(Other))
 		, BucketName(MoveTemp(InBucketName))
 		, Objects(MoveTemp(InObjects))
+		, NextMarker(MoveTemp(InNextMarker))
+		, bIsTruncated(bInIsTruncated)
 	{}
 
 	/** The bucket name. */
 	FString BucketName;
 	/** The list of object(s). */
 	TArray<FS3Object> Objects;
+	/** Marker to use for paginated requests. */
+	FString NextMarker;
+	/** Whether the response is truncated. */
+	bool bIsTruncated = false;
 };
 
 /** Request parameters for deleting objects. */
diff --git a/Engine/Source/Developer/S3Client/Private/S3Client.cpp b/Engine/Source/Developer/S3Client/Private/S3Client.cpp
--- a/Engine/Source/Developer/S3Client/Private/S3Client.cpp
+++ b/Engine/Source/Developer/S3Client/Private/S3Client.cpp
@@ -944,22 +944,32 @@ FS3ListObjectResponse FS3Client::ListObjects(const FS3ListObjectsRequest& ListRe
 	TAnsiStringBuilder<256> Url;
 	Url << StringCast<ANSICHAR>(*Config.ServiceUrl) << "/" << StringCast<ANSICHAR>(*ListRequest.BucketName); 
 
-	if (!ListRequest.Prefix.IsEmpty() || ListRequest.Delimiter != 0)
+	char QueryParamDelim = '?';
+	auto AppendQueryDelim = [&](FAnsiStringBuilderBase& Sb) -> FAnsiStringBuilderBase&
 	{
-		Url << "?";
-		if (ListRequest.Delimiter != 0)
-		{
-			Url << "delimiter=" << char(ListRequest.Delimiter);
-			if (!ListRequest.Prefix.IsEmpty())
-			{
-				Url << "&";
-			}
-		}
+		Sb << QueryParamDelim;
+		QueryParamDelim = '&';
+		return Sb;
+	};
 
-		if (!ListRequest.Prefix.IsEmpty())
-		{
-			Url << "prefix=" << StringCast<ANSICHAR>(*ListRequest.Prefix);
-		}
+	if (ListRequest.Delimiter != 0)
+	{
+		AppendQueryDelim(Url) << "delimiter=" << char(ListRequest.Delimiter);
+	}
+
+	if (ListRequest.Marker.IsEmpty() == false)
+	{
+		AppendQueryDelim(Url) << "marker=" << ListRequest.Marker; 
+	}
+
+	if (ListRequest.MaxKeys > 0)
+	{
+		AppendQueryDelim(Url) << "max-keys=" << ListRequest.MaxKeys; 
+	}
+
+	if (!ListRequest.Prefix.IsEmpty())
+	{
+		AppendQueryDelim(Url) << "prefix=" << StringCast<ANSICHAR>(*ListRequest.Prefix);
 	}
 
 	FS3Request Request(*this);
@@ -986,6 +996,8 @@ FS3ListObjectResponse FS3Client::ListObjects(const FS3ListObjectsRequest& ListRe
 	}
 
 	FString BucketName;
+	FString NextMarker;
+	bool bIsTruncated = false;
 	TArray<FS3Object> Objects;
 
 	const TArray<FXmlNode*>& Children = Root->GetChildrenNodes();
@@ -995,6 +1007,14 @@ FS3ListObjectResponse FS3Client::ListObjects(const FS3ListObjectsRequest& ListRe
 		{
 			BucketName = Child->GetContent();
 		}
+		else if (Child->GetTag() == TEXT("IsTruncated"))
+		{
+			bIsTruncated = Child->GetContent().Contains(TEXT("True"), ESearchCase::IgnoreCase);
+		}
+		else if (Child->GetTag() == TEXT("NextMarker"))
+		{
+			NextMarker = Child->GetContent();
+		}
 		else if (Child->GetTag() == TEXT("Contents"))
 		{
 			FS3Object& S3Object = Objects.AddDefaulted_GetRef();
@@ -1012,13 +1032,19 @@ FS3ListObjectResponse FS3Client::ListObjects(const FS3ListObjectsRequest& ListRe
 				}
 				else if (InfoNode->GetTag() == TEXT("LastModified"))
 				{
-					S3Object.LastModified = InfoNode->GetContent();
+					S3Object.LastModifiedText = InfoNode->GetContent();
+					FDateTime::ParseIso8601(*S3Object.LastModifiedText, S3Object.LastModified);
 				}
 			}
 		}
 	}
 
-	return FS3ListObjectResponse{{200, 0, FSharedBuffer()}, MoveTemp(BucketName), MoveTemp(Objects)};
+	return FS3ListObjectResponse{
+		{200, 0, FSharedBuffer()},
+		MoveTemp(BucketName),
+		MoveTemp(Objects),
+		MoveTemp(NextMarker),
+		bIsTruncated};
 }
 
 FS3DeleteObjectResponse FS3Client::DeleteObject(const FS3DeleteObjectRequest& DeleteRequest)
diff --git a/Engine/Source/Runtime/Experimental/IoStore/OnDemand/Private/Tool/Upload.cpp b/Engine/Source/Runtime/Experimental/IoStore/OnDemand/Private/Tool/Upload.cpp
--- a/Engine/Source/Runtime/Experimental/IoStore/OnDemand/Private/Tool/Upload.cpp
+++ b/Engine/Source/Runtime/Experimental/IoStore/OnDemand/Private/Tool/Upload.cpp
@@ -51,6 +51,8 @@ struct FUploadParams
 	FString TargetPlatform;
 	FString EncryptionKeyName;
 	int32 MaxConcurrentUploads = 16;
+	int32 MaxTocListCount = 10000;
+	int32 MaxTocDownloadCount = 100;
 	bool bDeleteContainerFiles = true;
 	bool bDeletePakFiles = true;
 	bool bPerContainerTocs = false;
@@ -109,6 +111,10 @@ struct FUploadResult
 ////////////////////////////////////////////////////////////////////////////////
 static FUploadParams BuildUploadParams(const FContext& Context)
 {
+	static constexpr int32 MinTocListCount		= 1000;
+	static constexpr int32 MaxTocListCount		= 100000;
+	static constexpr int32 MaxTocDownloadCount	= 10000;
+
 	FUploadParams Ret;
 
 	Ret.ServiceUrl				= Context.Get<FStringView>(TEXT("-ServiceUrl"),				Ret.ServiceUrl);
@@ -125,6 +131,8 @@ static FUploadParams BuildUploadParams(const FContext& Context)
 	Ret.bWriteTocToDisk			= Context.Get<bool>(TEXT("-WriteTocToDisk"),				Ret.bWriteTocToDisk);
 	Ret.bPerContainerTocs		= Context.Get<bool>(TEXT("-PerContainerTocs"),				Ret.bPerContainerTocs);
 	Ret.MaxConcurrentUploads	= Context.Get<int32>(TEXT("-MaxConcurrentUploads"),			Ret.MaxConcurrentUploads);
+	Ret.MaxTocListCount			= Context.Get<int32>(TEXT("-MaxTocListCount"),				Ret.MaxTocListCount);
+	Ret.MaxTocDownloadCount		= Context.Get<int32>(TEXT("-MaxTocDownloadCount"),			Ret.MaxTocDownloadCount);
 
 	Ret.bDeleteContainerFiles	= !Context.Get<bool>(TEXT("-KeepContainerFiles"),			!Ret.bDeleteContainerFiles);
 	Ret.bDeletePakFiles			= !Context.Get<bool>(TEXT("-KeepPakFiles"),					!Ret.bDeletePakFiles);
@@ -145,6 +153,9 @@ static FUploadParams BuildUploadParams(const FContext& Context)
 		Context.Abort(*Reason);
 	}
 
+	Ret.MaxTocListCount		= FMath::Max(FMath::Min(Ret.MaxTocListCount, MaxTocListCount), MinTocListCount);
+	Ret.MaxTocDownloadCount = FMath::Max(FMath::Min(Ret.MaxTocDownloadCount, MaxTocDownloadCount), 1);
+
 	return Ret;
 }
 
@@ -422,22 +433,65 @@ static TIoStatusOr<FUploadResult> UploadContainerFiles(
 			TocsKey << UploadParams.BucketPrefix << "/";
 		}
 
-		UE_LOG(LogIas, Display, TEXT("	 '%s/%s/%s'"), *Client.GetConfig().ServiceUrl, *UploadParams.Bucket, TocsKey.ToString());
-		FS3ListObjectResponse Response = Client.ListObjects(FS3ListObjectsRequest
+		const int32 MaxKeys = 1000; // AWS max value
+		const int32 MaxPaginationCount = FMath::Max(1, UploadParams.MaxTocListCount / MaxKeys);
+		const int32 MaxTocDownloadCount = UploadParams.MaxTocDownloadCount;
+		TArray<FS3Object> Objects;
+		FString Marker;
+
+		for (int32 Pagination = 0; Pagination < MaxPaginationCount; ++Pagination)
 		{
-			UploadParams.Bucket,
-			TocsKey.ToString(),
-			TEXT('/')
-		});
+			UE_LOG(LogIas, Display, TEXT("List #%d/%d '%s/%s/%s' TocCount=%d, MaxTocCount=%d"),
+				Pagination + 1, MaxPaginationCount, *Client.GetConfig().ServiceUrl, *UploadParams.Bucket, TocsKey.ToString(), Objects.Num(), UploadParams.MaxTocListCount);
+			FS3ListObjectResponse Response = Client.ListObjects(FS3ListObjectsRequest
+			{
+				.BucketName = UploadParams.Bucket,
+				.Prefix = TocsKey.ToString(),
+				.Delimiter = TCHAR('/'),
+				.MaxKeys = MaxKeys,
+				.Marker = Marker
+			});
+
+			Marker = MoveTemp(Response.NextMarker);
+
+			for (FS3Object& Object : Response.Objects)
+			{
+				if (Object.Key.EndsWith(TEXT("iochunktoc")))
+				{
+					Objects.Add(MoveTemp(Object));
+				}
+			}
+
+			if (Response.IsOk() == false || Response.Objects.IsEmpty() || Response.bIsTruncated == false || Objects.Num() >= UploadParams.MaxTocListCount)
+			{
+				break;
+			}
+		}
+
+		{
+			// Just to be sure
+			const int32 TotalObjectCount = Objects.Num();
+			TSet<FS3Object> UniqueObjects(Objects);
+			Objects = UniqueObjects.Array();
+			if (Objects.Num() < TotalObjectCount)
+			{
+				UE_LOG(LogIas, Display, TEXT("Removed %d non unique TOC objects"), TotalObjectCount - Objects.Num());
+			}
+		}
+
+		Objects.Sort([](const FS3Object& LHS, const FS3Object& RHS) { return LHS.LastModified > RHS.LastModified; });
 
-		for (const FS3Object& TocInfo : Response.Objects)
+		const int32 TocDownloadCount = FMath::Min(Objects.Num(), MaxTocDownloadCount);
+		UE_LOG(LogIas, Display, TEXT("Fetching %d latest TOC file(s) ..."), TocDownloadCount);
+		for (int32 Idx = 0; Idx < TocDownloadCount; ++Idx)
 		{
+			const FS3Object& TocInfo = Objects[Idx];
 			if (TocInfo.Key.EndsWith(TEXT("iochunktoc")) == false)
 			{
 				continue;
 			}
 
-			UE_LOG(LogIas, Display, TEXT("Fetching TOC '%s/%s/%s'"), *Client.GetConfig().ServiceUrl, *UploadParams.Bucket, *TocInfo.Key);
+			UE_LOG(LogIas, Display, TEXT("Fetching TOC %d/%d '%s/%s/%s', Size=%llu, LastModified=%s"), Idx + 1, TocDownloadCount, *Client.GetConfig().ServiceUrl, *UploadParams.Bucket, *TocInfo.Key, TocInfo.Size, *TocInfo.LastModifiedText);
 			FS3GetObjectResponse TocResponse = Client.GetObject(FS3GetObjectRequest
 			{
 				UploadParams.Bucket,
@@ -885,6 +939,8 @@ static FCommand UploadCommand(
 		TArgument<bool>(TEXT("-PerContainerTocs"),		TEXT("Whether to generate TOC's for each container file(s)")),
 		TArgument<bool>(TEXT("-IgnoreContainerFlags"),	TEXT("Whether to ignore the OnDemand container flag")),
 		TArgument<int32>(TEXT("-MaxConcurrentUploads"),	TEXT("Number of simultaneous uploads")),
+		TArgument<int32>(TEXT("-MaxTocListCount"),		TEXT("Maximum number of TOC file(s) to list from the bucket")),
+		TArgument<int32>(TEXT("-MaxTocDownloadCount"),	TEXT("Maximum number of TOC file(s) to download")),
 		S3Arguments,
 	}
 );
diff --git a/Engine/Source/Runtime/Experimental/IoStore/OnDemand/Private/Tool/UploadQueue.cpp b/Engine/Source/Runtime/Experimental/IoStore/OnDemand/Private/Tool/UploadQueue.cpp
--- a/Engine/Source/Runtime/Experimental/IoStore/OnDemand/Private/Tool/UploadQueue.cpp
+++ b/Engine/Source/Runtime/Experimental/IoStore/OnDemand/Private/Tool/UploadQueue.cpp
@@ -60,6 +60,9 @@ bool FUploadQueue::Enqueue(FStringView Key, FIoBuffer Payload)
 
 void FUploadQueue::ThreadEntry()
 {
+	const int32 MaxAttempts = 3;
+	const float RetryTimeouts[] = {0.5f, 1.0f, 2.0f};
+
 	for (;;)
 	{
 		FQueueEntry Entry;
@@ -80,9 +83,15 @@ void FUploadQueue::ThreadEntry()
 		}
 
 		ConcurrentUploads++;
-		const FS3PutObjectResponse Response = Client.TryPutObject(FS3PutObjectRequest{ Bucket, Entry.Key, Entry.Payload.GetView() });
-		ConcurrentUploads--;
-		UploadCompleteEvent->Trigger();
+		FS3PutObjectResponse Response;
+		for (int32 Attempt = 0; Attempt < MaxAttempts; ++Attempt)
+		{
+			if (Response = Client.TryPutObject(FS3PutObjectRequest{ Bucket, Entry.Key, Entry.Payload.GetView() }); Response.IsOk())
+			{
+				break;
+			}
+			FPlatformProcess::Sleep(RetryTimeouts[Attempt]);
+		}
 
 		if (Response.IsOk())
 		{
@@ -95,8 +104,9 @@ void FUploadQueue::ThreadEntry()
 
 			UE_LOG(LogIoStoreOnDemand, Warning, TEXT("Failed to upload chunk '%s/%s/%s' (%s)"), *Client.GetConfig().ServiceUrl, *Bucket, *Entry.Key, ErrorResponse.ToString());
 			ErrorCount++;
-			break;
 		}
+		ConcurrentUploads--;
+		UploadCompleteEvent->Trigger();
 	}
 
 	ActiveThreadCount--;
