CIS fix for api change in CL 39476420

#rnx

- In CL 39476420, GetErrorReason was replaced by a GetError which returns a struct containing more info.

diff --git a/Engine/Source/Runtime/Experimental/IoStore/OnDemand/Private/Tool/Purl.cpp b/Engine/Source/Runtime/Experimental/IoStore/OnDemand/Private/Tool/Purl.cpp
--- a/Engine/Source/Runtime/Experimental/IoStore/OnDemand/Private/Tool/Purl.cpp
+++ b/Engine/Source/Runtime/Experimental/IoStore/OnDemand/Private/Tool/Purl.cpp
@@ -104,7 +104,7 @@ static int32 PurlCommandEntry(const FContext& Context)
 
 		if (Status.GetId() == FTicketStatus::EId::Error)
 		{
-			const char* Reason = Status.GetErrorReason();
+			const char* Reason = Status.GetError().Reason;
 			std::printf("ERROR: %s\n", Reason);
 			return;
 		}
