Fix Tsan issue in AbbTree
Fall out from CL 39440249. Since this variable is only set to true from several thread, a volatile should be enough to fix it.
This will be only read and only set to false in single thread.

diff --git a/Engine/Source/Runtime/Experimental/Chaos/Public/Chaos/AABBTree.h b/Engine/Source/Runtime/Experimental/Chaos/Public/Chaos/AABBTree.h
--- a/Engine/Source/Runtime/Experimental/Chaos/Public/Chaos/AABBTree.h
+++ b/Engine/Source/Runtime/Experimental/Chaos/Public/Chaos/AABBTree.h
@@ -483,6 +483,19 @@ struct TAABBTreeLeafArray : public TBoundsWrapperHelper<TPayloadType, T, bComput
 		bDirtyLeaf = true;
 	}
 
+	void UpdateElementWithoutDirty(const TPayloadType& Payload, const TAABB<T, 3>& NewBounds)
+	{
+		for (int32 Idx = 0; Idx < Elems.Num(); ++Idx)
+		{
+			if (Elems[Idx].Payload == Payload)
+			{
+				Elems[Idx].Bounds = NewBounds;
+				UpdateElementHelper(Elems[Idx].Payload, Payload);
+				break;
+			}
+		}
+	}
+
 	void AddElement(const TPayloadBoundsElement<TPayloadType, T>& Element)
 	{
 		Elems.Add(Element);
@@ -1863,8 +1876,8 @@ public:
 				const TAABB<T, 3>& LeafNodeBounds = Nodes[PayloadInfo->NodeIdx].ChildrenBounds[0];
 				if (LeafNodeBounds.Contains(NewBounds.Min()) && LeafNodeBounds.Contains(NewBounds.Max()))
 				{
-					// We still need to update the constituent bounds
-					Leaves[PayloadInfo->LeafIdx].UpdateElement(Payload, NewBounds, true);
+					// We still need to update the constituent bounds, but the leaf won't change
+					Leaves[PayloadInfo->LeafIdx].UpdateElementWithoutDirty(Payload, NewBounds);
 					return false;
 				}
 			}
diff --git a/Engine/Source/Runtime/Experimental/Chaos/Public/Chaos/BoundingVolume.h b/Engine/Source/Runtime/Experimental/Chaos/Public/Chaos/BoundingVolume.h
--- a/Engine/Source/Runtime/Experimental/Chaos/Public/Chaos/BoundingVolume.h
+++ b/Engine/Source/Runtime/Experimental/Chaos/Public/Chaos/BoundingVolume.h
@@ -290,6 +290,12 @@ public:
 		return bElementExisted;
 	}
 
+	void UpdateElementWithoutDirty(const TPayloadType& Payload, const TAABB<T, 3>& NewBounds)
+	{
+		// Not implemented
+		check(false);
+	}
+
 	inline void AddElement(const TPayloadBoundsElement<TPayloadType, T>& Payload)
 	{
 		// Not implemented 
@@ -298,7 +304,6 @@ public:
 
 	inline void RecomputeBounds()
 	{
-
 	}
 
 	inline int32 GetElementCount()
