Remove hack to use a  worker thread count correct for ps5
CL 39459327 fixed worker thread count on PS5
- Optimization: Use less tasks to prune particles from the aabb tree

diff --git a/Engine/Source/Runtime/Experimental/Chaos/Private/Chaos/PBDRigidsEvolution.cpp b/Engine/Source/Runtime/Experimental/Chaos/Private/Chaos/PBDRigidsEvolution.cpp
--- a/Engine/Source/Runtime/Experimental/Chaos/Private/Chaos/PBDRigidsEvolution.cpp
+++ b/Engine/Source/Runtime/Experimental/Chaos/Private/Chaos/PBDRigidsEvolution.cpp
@@ -627,9 +627,7 @@ namespace Chaos
 		// The bounding volume associated to the particle will be updated in the acceleration structure.
 		QUICK_SCOPE_CYCLE_COUNTER(STAT_PruneInternalPendingData);
 
-		int32 NumWorker = !!GSingleThreadedPhysics ? 1 : FMath::Max(FMath::Min(Chaos::CVars::NumWorkerCollisionFactor * (FTaskGraphInterface::Get().GetNumWorkerThreads() - 1), Chaos::MaxNumWorkers), 1);
-		NumWorker = FMath::Max(1, NumWorker);
-
+		const int32 NumWorker = !!GSingleThreadedPhysics ? 1 : FTaskGraphInterface::Get().GetNumWorkerThreads();
 		const int32 NumParticles = InternalAccelerationQueue.PendingData.Num();
 		const int32 ParticleByContext = FMath::Max(FMath::DivideAndRoundUp(NumParticles, NumWorker), SmallBatchSize);
 		const int32 NumBatch = FMath::DivideAndRoundUp(NumParticles, ParticleByContext);
diff --git a/Engine/Source/Runtime/Experimental/Chaos/Public/Chaos/Collision/SpatialAccelerationBroadPhase.h b/Engine/Source/Runtime/Experimental/Chaos/Public/Chaos/Collision/SpatialAccelerationBroadPhase.h
--- a/Engine/Source/Runtime/Experimental/Chaos/Public/Chaos/Collision/SpatialAccelerationBroadPhase.h
+++ b/Engine/Source/Runtime/Experimental/Chaos/Public/Chaos/Collision/SpatialAccelerationBroadPhase.h
@@ -583,8 +583,7 @@ namespace Chaos
 			}
 
 			// The number of contexts that may have new overlaps in them
-			NumActiveBroadphaseContexts = !!GSingleThreadedPhysics ? 1 : FMath::Max(FMath::Min(Chaos::CVars::NumWorkerCollisionFactor * (FTaskGraphInterface::Get().GetNumWorkerThreads() - 1), Chaos::MaxNumWorkers), 1);
-			NumActiveBroadphaseContexts = FMath::Max(1, NumActiveBroadphaseContexts);
+			NumActiveBroadphaseContexts = !!GSingleThreadedPhysics ? 1 : FMath::Max(FMath::Min(Chaos::CVars::NumWorkerCollisionFactor * FTaskGraphInterface::Get().GetNumWorkerThreads(), Chaos::MaxNumWorkers), 1);
 
 			BroadphaseContexts.SetNum(NumActiveBroadphaseContexts);
 			Allocator->SetMaxContexts(NumActiveBroadphaseContexts);
