In Aabb tree allow to enlarge leaves only for dynamic tree

diff --git a/Engine/Source/Runtime/Experimental/Chaos/Public/Chaos/AABBTree.h b/Engine/Source/Runtime/Experimental/Chaos/Public/Chaos/AABBTree.h
--- a/Engine/Source/Runtime/Experimental/Chaos/Public/Chaos/AABBTree.h
+++ b/Engine/Source/Runtime/Experimental/Chaos/Public/Chaos/AABBTree.h
@@ -221,7 +221,7 @@ struct TBoundsWrapperHelper
 template <typename TPayloadType, typename T>
 struct TBoundsWrapperHelper<TPayloadType, T, true>
 {
-	void ComputeBounds(const TArray<TPayloadBoundsElement<TPayloadType, T>>& Elems)
+	void ComputeBounds(const TArray<TPayloadBoundsElement<TPayloadType, T>>& Elems, bool bDynamicTree)
 	{
 		Bounds = TAABB<T, 3>::EmptyAABB();
 
@@ -229,7 +229,7 @@ struct TBoundsWrapperHelper<TPayloadType, T, true>
 		{
 			Bounds.GrowToInclude(Elem.Bounds);
 		}
-		if (!Elems.IsEmpty())
+		if (!Elems.IsEmpty() && bDynamicTree)
 		{
 			// Grow the aabb leaf of percentage of its size
 			Bounds.ThickenSymmetrically(Bounds.Extents() * FAABBTreeCVars::DynamicTreeLeafEnlargePercent*0.5);
@@ -245,7 +245,7 @@ private:
 template <typename TPayloadType, typename T>
 struct TBoundsWrapperHelper<TPayloadType, T, false>
 {
-	void ComputeBounds(const TArray<TPayloadBoundsElement<TPayloadType, T>>&)
+	void ComputeBounds(const TArray<TPayloadBoundsElement<TPayloadType, T>>&, bool bDynamicTree)
 	{
 	}
 
@@ -263,7 +263,7 @@ struct TAABBTreeLeafArray : public TBoundsWrapperHelper<TPayloadType, T, bComput
 	TAABBTreeLeafArray(const TArray<TPayloadBoundsElement<TPayloadType, T>>& InElems)
 		: Elems(InElems)
 	{
-		this->ComputeBounds(Elems);
+		this->ComputeBounds(Elems, false);
 	}
 
 	void GatherElements(TArray<TPayloadBoundsElement<TPayloadType, T>>& OutElements) const
@@ -282,9 +282,9 @@ struct TAABBTreeLeafArray : public TBoundsWrapperHelper<TPayloadType, T, bComput
 		return Elems.Num();
 	}
 
-	void RecomputeBounds()
+	void RecomputeBounds(bool bDynamicTree)
 	{
-		this->ComputeBounds(Elems);
+		this->ComputeBounds(Elems, bDynamicTree);
 	}
 
 	/** Check if the leaf is dirty (if one of the payload have been updated)
@@ -499,7 +499,6 @@ struct TAABBTreeLeafArray : public TBoundsWrapperHelper<TPayloadType, T, bComput
 	void AddElement(const TPayloadBoundsElement<TPayloadType, T>& Element)
 	{
 		Elems.Add(Element);
-		this->ComputeBounds(Elems);
 		bDirtyLeaf = true;
 	}
 
@@ -1371,6 +1370,7 @@ public:
 			LeafIndex = Nodes[AllocatedNodeIdx].ChildrenNodes[0]; // This is already set when it was allocated for the first time
 			FElement NewElement{ Payload, NewBounds };
 			Leaves[LeafIndex].AddElement(NewElement);
+			Leaves[LeafIndex].RecomputeBounds(bDynamicTree);
 		}
 		else
 		{
@@ -1648,7 +1648,7 @@ public:
 		{
 			const int32 LeafIdx = Nodes[BestSibling].ChildrenNodes[0];
 			Leaves[LeafIdx].AddElement(TPayloadBoundsElement<TPayloadType, T>{Payload, NewBounds});
-			Leaves[LeafIdx].RecomputeBounds();
+			Leaves[LeafIdx].RecomputeBounds(bDynamicTree);
 			TAABB<T, 3> ExpandedBounds = Leaves[LeafIdx].GetBounds();
 			ExpandedBounds.Thicken(FAABBTreeCVars::DynamicTreeBoundingBoxPadding);
 			Nodes[BestSibling].ChildrenBounds[0] = ExpandedBounds;
@@ -1751,7 +1751,7 @@ public:
 		if (Leaves[LeafIdx].GetElementCount() > 1)
 		{
 			Leaves[LeafIdx].RemoveElement(Payload);
-			Leaves[LeafIdx].RecomputeBounds();
+			Leaves[LeafIdx].RecomputeBounds(bDynamicTree);
 			TAABB<T, 3> ExpandedBounds = Leaves[LeafIdx].GetBounds();
 			ExpandedBounds.Thicken(FAABBTreeCVars::DynamicTreeBoundingBoxPadding);
 			Nodes[LeafNodeIdx].ChildrenBounds[0] = ExpandedBounds;
@@ -2038,7 +2038,7 @@ public:
 							// We still need to update the constituent bounds
 							const int32 DirtyLeafIndex = DirtyElementTree->Nodes[DirtyElementIndex].ChildrenNodes[0]; // A Leaf node's first child is the leaf index
 							DirtyElementTree->Leaves[DirtyLeafIndex].UpdateElement(Payload, NewBounds, bHasBounds);
-							DirtyElementTree->Leaves[DirtyLeafIndex].RecomputeBounds();
+							DirtyElementTree->Leaves[DirtyLeafIndex].RecomputeBounds(bDynamicTree);
 						}
 						else
 						{
diff --git a/Engine/Source/Runtime/Experimental/Chaos/Public/Chaos/BoundingVolume.h b/Engine/Source/Runtime/Experimental/Chaos/Public/Chaos/BoundingVolume.h
--- a/Engine/Source/Runtime/Experimental/Chaos/Public/Chaos/BoundingVolume.h
+++ b/Engine/Source/Runtime/Experimental/Chaos/Public/Chaos/BoundingVolume.h
@@ -302,7 +302,7 @@ public:
 		check(false);
 	}
 
-	inline void RecomputeBounds()
+	inline void RecomputeBounds(bool bDynamicTree)
 	{
 	}
 
