[AutoRTFM] Use AutoRTFM allocators in Iris serialization.

GMalloc is not transactionally aware, but FMemory::Malloc and friends are (and are otherwise very simple wrappers for GMalloc).

diff --git a/Engine/Source/Runtime/Experimental/Iris/Core/Private/Iris/ReplicationSystem/NetBlob/NetBlob.cpp b/Engine/Source/Runtime/Experimental/Iris/Core/Private/Iris/ReplicationSystem/NetBlob/NetBlob.cpp
--- a/Engine/Source/Runtime/Experimental/Iris/Core/Private/Iris/ReplicationSystem/NetBlob/NetBlob.cpp
+++ b/Engine/Source/Runtime/Experimental/Iris/Core/Private/Iris/ReplicationSystem/NetBlob/NetBlob.cpp
@@ -170,7 +170,7 @@ void FNetObjectAttachment::DeserializeSubObjectReference(FNetSerializationContex
 }
 
 FNetBlob::FQuantizedBlobState::FQuantizedBlobState(uint32 Size, uint32 Alignment)
-: StateBuffer(static_cast<uint8*>(GMalloc->Malloc(Size, Alignment)))
+: StateBuffer(static_cast<uint8*>(FMemory::Malloc(Size, Alignment)))
 {
 	if (StateBuffer)
 	{
diff --git a/Engine/Source/Runtime/Experimental/Iris/Core/Private/Iris/ReplicationSystem/ReplicationOperations.cpp b/Engine/Source/Runtime/Experimental/Iris/Core/Private/Iris/ReplicationSystem/ReplicationOperations.cpp
--- a/Engine/Source/Runtime/Experimental/Iris/Core/Private/Iris/ReplicationSystem/ReplicationOperations.cpp
+++ b/Engine/Source/Runtime/Experimental/Iris/Core/Private/Iris/ReplicationSystem/ReplicationOperations.cpp
@@ -27,6 +27,7 @@
 #include "Net/Core/Trace/NetDebugName.h"
 #include "DequantizeAndApplyHelper.h"
 #include "HAL/IConsoleManager.h"
+#include "HAL/UnrealMemory.h"
 #include "UObject/CoreNetTypes.h"
 
 #include "CoreTypes.h"
@@ -883,7 +884,7 @@ void FReplicationInstanceOperations::OutputInternalStateToString(FNetSerializati
 		if (!EnumHasAnyFlags(CurrentFragment->GetTraits(), EReplicationFragmentTraits::HasPersistentTargetStateBuffer))
 		{
 			// Allocate buffer for temporary state and construct the external state
-			uint8* StateBuffer = (uint8*)GMalloc->Malloc(CurrentDescriptor->ExternalSize, CurrentDescriptor->ExternalAlignment);
+			uint8* StateBuffer = (uint8*)FMemory::Malloc(CurrentDescriptor->ExternalSize, CurrentDescriptor->ExternalAlignment);
 			CurrentDescriptor->ConstructReplicationState(StateBuffer, CurrentDescriptor);
 	
 			// Dequantize state data
@@ -896,7 +897,7 @@ void FReplicationInstanceOperations::OutputInternalStateToString(FNetSerializati
 
 			// Free buffer
 			CurrentDescriptor->DestructReplicationState(StateBuffer, CurrentDescriptor);
-			GMalloc->Free(StateBuffer);
+			FMemory::Free(StateBuffer);
 		}
 		else
 		{
@@ -949,7 +950,7 @@ void FReplicationInstanceOperations::OutputInternalDefaultStateToString(FNetSeri
 		if (!EnumHasAnyFlags(CurrentFragment->GetTraits(), EReplicationFragmentTraits::HasPersistentTargetStateBuffer))
 		{
 			// Allocate buffer for temporary state and construct the external state
-			uint8* StateBuffer = (uint8*)GMalloc->Malloc(CurrentDescriptor->ExternalSize, CurrentDescriptor->ExternalAlignment);
+			uint8* StateBuffer = (uint8*)FMemory::Malloc(CurrentDescriptor->ExternalSize, CurrentDescriptor->ExternalAlignment);
 			CurrentDescriptor->ConstructReplicationState(StateBuffer, CurrentDescriptor);
 
 			// Dequantize state data
@@ -962,7 +963,7 @@ void FReplicationInstanceOperations::OutputInternalDefaultStateToString(FNetSeri
 
 			// Free buffer
 			CurrentDescriptor->DestructReplicationState(StateBuffer, CurrentDescriptor);
-			GMalloc->Free(StateBuffer);
+			FMemory::Free(StateBuffer);
 		}
 		else
 		{
diff --git a/Engine/Source/Runtime/Experimental/Iris/Core/Private/Iris/Serialization/InternalNetSerializationContext.cpp b/Engine/Source/Runtime/Experimental/Iris/Core/Private/Iris/Serialization/InternalNetSerializationContext.cpp
--- a/Engine/Source/Runtime/Experimental/Iris/Core/Private/Iris/Serialization/InternalNetSerializationContext.cpp
+++ b/Engine/Source/Runtime/Experimental/Iris/Core/Private/Iris/Serialization/InternalNetSerializationContext.cpp
@@ -3,7 +3,7 @@
 #include "Iris/Serialization/InternalNetSerializationContext.h"
 #include "Iris/ReplicationSystem/ReplicationSystem.h"
 #include "Iris/ReplicationSystem/ReplicationSystemInternal.h"
-#include "HAL/MemoryBase.h"
+#include "HAL/UnrealMemory.h"
 
 namespace UE::Net::Private
 {
@@ -19,17 +19,17 @@ FInternalNetSerializationContext::FInternalNetSerializationContext(UReplicationS
 
 void* FInternalNetSerializationContext::Alloc(SIZE_T Size, SIZE_T Alignment)
 {
-	return GMalloc->Malloc(Size, static_cast<uint32>(Alignment));
+	return FMemory::Malloc(Size, static_cast<uint32>(Alignment));
 }
 
 void FInternalNetSerializationContext::Free(void* Ptr)
 {
-	return GMalloc->Free(Ptr);
+	return FMemory::Free(Ptr);
 }
 
 void* FInternalNetSerializationContext::Realloc(void* PrevAddress, SIZE_T NewSize, uint32 Alignment)
 {
-	return GMalloc->Realloc(PrevAddress, NewSize, Alignment);
+	return FMemory::Realloc(PrevAddress, NewSize, Alignment);
 }
 
 void FInternalNetSerializationContext::Init(const FInitParameters& InitParams)
diff --git a/Engine/Source/Runtime/Experimental/Iris/Core/Public/Iris/ReplicationSystem/NetBlob/NetBlob.h b/Engine/Source/Runtime/Experimental/Iris/Core/Public/Iris/ReplicationSystem/NetBlob/NetBlob.h
--- a/Engine/Source/Runtime/Experimental/Iris/Core/Public/Iris/ReplicationSystem/NetBlob/NetBlob.h
+++ b/Engine/Source/Runtime/Experimental/Iris/Core/Public/Iris/ReplicationSystem/NetBlob/NetBlob.h
@@ -4,6 +4,7 @@
 #include "CoreTypes.h"
 #include "Iris/Core/NetObjectReference.h"
 #include "Iris/ReplicationState/ReplicationStateDescriptor.h"
+#include "HAL/UnrealMemory.h"
 #include "Containers/ArrayView.h"
 #include "Templates/RefCounting.h"
 #include "Templates/UniquePtr.h"
@@ -214,7 +215,7 @@ inline FNetBlob::FQuantizedBlobState::~FQuantizedBlobState()
 {
 	if (StateBuffer)
 	{
-		GMalloc->Free(StateBuffer);
+		FMemory::Free(StateBuffer);
 	}
 }
 
@@ -222,7 +223,7 @@ inline FNetBlob::FQuantizedBlobState& FNetBlob::FQuantizedBlobState::operator=(F
 {
 	if (StateBuffer)
 	{
-		GMalloc->Free(StateBuffer);
+		FMemory::Free(StateBuffer);
 	}
 
 	StateBuffer = Other.StateBuffer;
