Fix division by zero
fallout from CL 39480191

diff --git a/Engine/Source/Runtime/Experimental/Chaos/Private/Chaos/PBDRigidsEvolution.cpp b/Engine/Source/Runtime/Experimental/Chaos/Private/Chaos/PBDRigidsEvolution.cpp
--- a/Engine/Source/Runtime/Experimental/Chaos/Private/Chaos/PBDRigidsEvolution.cpp
+++ b/Engine/Source/Runtime/Experimental/Chaos/Private/Chaos/PBDRigidsEvolution.cpp
@@ -627,7 +627,7 @@ namespace Chaos
 		// The bounding volume associated to the particle will be updated in the acceleration structure.
 		QUICK_SCOPE_CYCLE_COUNTER(STAT_PruneInternalPendingData);
 
-		const int32 NumWorker = !!GSingleThreadedPhysics ? 1 : FTaskGraphInterface::Get().GetNumWorkerThreads();
+		const int32 NumWorker = !!GSingleThreadedPhysics ? 1 : FMath::Max(FTaskGraphInterface::Get().GetNumWorkerThreads(), 1);
 		const int32 NumParticles = InternalAccelerationQueue.PendingData.Num();
 		const int32 ParticleByContext = FMath::Max(FMath::DivideAndRoundUp(NumParticles, NumWorker), SmallBatchSize);
 		const int32 NumBatch = FMath::DivideAndRoundUp(NumParticles, ParticleByContext);
