Add a new IAS analytics value to track if a user has managed to resolve a distributed endpoint or not.

#rnx

- The new analytics value is named "IasHttpDistributedEndpointResolved" and will be set to true once a successful resolution of the distributed endpoint has been made. If the value is false thne either the resolution was never attempted, is in progress or failed.

diff --git a/Engine/Source/Runtime/Experimental/IoStore/OnDemand/Private/OnDemandIoDispatcherBackend.cpp b/Engine/Source/Runtime/Experimental/IoStore/OnDemand/Private/OnDemandIoDispatcherBackend.cpp
--- a/Engine/Source/Runtime/Experimental/IoStore/OnDemand/Private/OnDemandIoDispatcherBackend.cpp
+++ b/Engine/Source/Runtime/Experimental/IoStore/OnDemand/Private/OnDemandIoDispatcherBackend.cpp
@@ -1460,6 +1460,8 @@ bool FOnDemandIoBackend::ResolveDistributedEndpoint(const FDistributedEndpointUr
 		FDistributionEndpoints::EResult Result = Resolver.ResolveEndpoints(DistributedEndpointUrl.EndpointUrl, ServiceUrls, *DistributedEndpointEvent.Get());
 		if (Result == FDistributionEndpoints::EResult::Success)
 		{
+			Stats.OnHttpDistributedEndpointResolved();
+
 			FWriteScopeLock _(Lock);
 			AvailableEps.Urls.Reserve(AvailableEps.Urls.Num() + ServiceUrls.Num());
 
diff --git a/Engine/Source/Runtime/Experimental/IoStore/OnDemand/Private/Statistics.cpp b/Engine/Source/Runtime/Experimental/IoStore/OnDemand/Private/Statistics.cpp
--- a/Engine/Source/Runtime/Experimental/IoStore/OnDemand/Private/Statistics.cpp
+++ b/Engine/Source/Runtime/Experimental/IoStore/OnDemand/Private/Statistics.cpp
@@ -272,6 +272,7 @@ FCounterAtomicInt		GCachePendingBytes(TEXT("Ias/CachePendingBytes"), TraceCounte
 FCounterAtomicInt		GCacheReadBytes(TEXT("Ias/CacheReadBytes"), TraceCounterDisplayHint_Memory);
 FCounterAtomicInt		GCacheRejectBytes(TEXT("Ias/CachePutRejectBytes"), TraceCounterDisplayHint_Memory);
 // http stats
+bool					GHttpDistributedEndpointResolved = false;
 FCounterInt				GHttpConnectCount(TEXT("Ias/HttpConnectCount"), TraceCounterDisplayHint_None);
 FCounterInt				GHttpDisconnectCount(TEXT("Ias/HttpDisconnectCount"), TraceCounterDisplayHint_None);
 FCounterInt				GHttpGetCount(TEXT("Ias/HttpGetCount"), TraceCounterDisplayHint_None);
@@ -549,6 +550,7 @@ void FOnDemandIoBackendStats::ReportGeneralAnalytics(TArray<FAnalyticsEventAttri
 	// not since with the current system we are unable to specify our own analytics payload.
 
 	AppendAnalyticsEventAttributeArray(OutAnalyticsArray
+		, TEXT("IasHttpDistributedEndpointResolved"), GHttpDistributedEndpointResolved
 		, TEXT("IasHttpHasEverConnected"), GHttpConnectCount.Get() > 0 // Report if the system has ever actually managed to make a connection
 	);
 }
@@ -896,6 +898,11 @@ void FOnDemandIoBackendStats::OnCacheSetMaxBytes(uint64 TotalSize)
 	GCacheMaxBytes = TotalSize;
 }
 
+void FOnDemandIoBackendStats::OnHttpDistributedEndpointResolved()
+{
+	GHttpDistributedEndpointResolved = true;
+}
+
 void FOnDemandIoBackendStats::OnHttpConnected()
 {
 	GHttpConnectCount.Add(1);
diff --git a/Engine/Source/Runtime/Experimental/IoStore/OnDemand/Private/Statistics.h b/Engine/Source/Runtime/Experimental/IoStore/OnDemand/Private/Statistics.h
--- a/Engine/Source/Runtime/Experimental/IoStore/OnDemand/Private/Statistics.h
+++ b/Engine/Source/Runtime/Experimental/IoStore/OnDemand/Private/Statistics.h
@@ -63,6 +63,7 @@ public:
 	void OnCacheWriteBytes(uint64 WriteSize) IAS_STATISTICS_IMPL()
 	void OnCacheSetMaxBytes(uint64 TotalSize) IAS_STATISTICS_IMPL()
 
+	void OnHttpDistributedEndpointResolved() IAS_STATISTICS_IMPL()
 	void OnHttpConnected() IAS_STATISTICS_IMPL()
 	void OnHttpDisconnected() IAS_STATISTICS_IMPL()
 
