Replace raw use of GCompleteMaterializeTaskPriorities with GetRequestCompletionTaskPriority()

#rnx

- Missed this change from CL 39026970 when working on CL 39402569

diff --git a/Engine/Source/Runtime/Experimental/IoStore/OnDemand/Private/OnDemandIoDispatcherBackend.cpp b/Engine/Source/Runtime/Experimental/IoStore/OnDemand/Private/OnDemandIoDispatcherBackend.cpp
--- a/Engine/Source/Runtime/Experimental/IoStore/OnDemand/Private/OnDemandIoDispatcherBackend.cpp
+++ b/Engine/Source/Runtime/Experimental/IoStore/OnDemand/Private/OnDemandIoDispatcherBackend.cpp
@@ -1642,11 +1642,10 @@ void FOnDemandIoBackend::ProcessHttpRequests(FMultiEndpointHttpClient& HttpClien
 								}
 							}
 
-							UE::Tasks::ETaskPriority TaskPriority = GCompleteMaterializeTaskPriorities[FMath::Clamp(GCompleteMaterializeTaskPriority, 0, UE_ARRAY_COUNT(GCompleteMaterializeTaskPriorities) - 1)];
 							UE::Tasks::Launch(UE_SOURCE_LOCATION, [this, ChunkRequest]()
 								{
 									CompleteRequest(ChunkRequest);
-								}, TaskPriority);
+								}, GetRequestCompletionTaskPriority());
 
 						});
 				}
