Fix a bug where one overload of FEventLoop::Request was ignoring FRequestParams::bAllowChunked.

#rnx

- One of the request overloads was not checking FRequestParams to see if chunked transfers were allowed and was defaulting to enabled, meaning there was no way to disable it when using that overload.

diff --git a/Engine/Source/Runtime/Experimental/IoStore/HttpClient/Private/Loop.inl b/Engine/Source/Runtime/Experimental/IoStore/HttpClient/Private/Loop.inl
--- a/Engine/Source/Runtime/Experimental/IoStore/HttpClient/Private/Loop.inl
+++ b/Engine/Source/Runtime/Experimental/IoStore/HttpClient/Private/Loop.inl
@@ -1574,6 +1574,7 @@ FRequest FEventLoop::Request(
 	Activity->Host = Pool.Ptr;
 	Activity->IsKeepAlive = 1;
 	Activity->bFollow30x = (Params->bAutoRedirect == true);
+	Activity->bAllowChunked = (Params->bAllowChunked == true);
 
 	return Impl->Request(Method, Path, Activity);
 }
