Re-enable the cvar 'ias.HttpAllowChunkedXfer' for IAS requests.

#rnx

- In CL 39402569 FOnDemandIoBackend was updated to use FMultiEndpointHttpClient in place of FHttpClient but the cvar 'ias.HttpAllowChunkedXfer' was only hooked up to FHttpClient so this change potentially changed behavior.
- Moved the cvar to OnDemandIoDispatcherBackend where the other cvars for IAS live.
- It is now passed in via FMultiEndpointHttpClientConfig when creating the http client so different owning systems can request different behavior.
- FHttpClient no longer supports the cvar but as it is currently unused and likely to be removed in the near future there was no need to update it's config structure as well.

diff --git a/Engine/Source/Runtime/Experimental/IoStore/OnDemand/Private/OnDemandHttpClient.cpp b/Engine/Source/Runtime/Experimental/IoStore/OnDemand/Private/OnDemandHttpClient.cpp
--- a/Engine/Source/Runtime/Experimental/IoStore/OnDemand/Private/OnDemandHttpClient.cpp
+++ b/Engine/Source/Runtime/Experimental/IoStore/OnDemand/Private/OnDemandHttpClient.cpp
@@ -20,13 +20,6 @@
 namespace UE::IoStore
 {
 
-static bool GIasHttpAllowChunkedXfer = false;
-static FAutoConsoleVariableRef CVar_IasHttpAllowChunkedXfer(
-	TEXT("ias.HttpAllowChunkedXfer"),
-	GIasHttpAllowChunkedXfer,
-	TEXT("Enable/disable IAS' support for chunked transfer encoding")
-);
-
 FIoStatus LoadDefaultHttpCertificates(bool& bWasLoaded)
 {
 	using namespace UE::IoStore::HTTP;
@@ -270,8 +263,7 @@ void FHttpClient::IssueRequest(FRequestParams&& Params)
 	TUniquePtr<FConnectionPool>& Connection = Connections[Params.Endpoint];
 	check(Connection.IsValid());
 
-	FEventLoop::FRequestParams RequestParams = { .bAllowChunked = GIasHttpAllowChunkedXfer };
-	FRequest Request = EventLoop.Get(Url, *Connection, &RequestParams);
+	FRequest Request = EventLoop.Get(Url, *Connection);
 
 	if (Params.Range.GetOffset() > 0 || Params.Range.GetLength() > 0)
 	{
@@ -541,7 +533,8 @@ void FMultiEndpointHttpClient::IssueRequest(FRequest&& Request)
 		}
 	};
 
-	UE::IoStore::HTTP::FRequest HttpRequest = EventLoop.Get(Url, ConnectionPool);
+	FEventLoop::FRequestParams RequestParams = { .bAllowChunked = Config.bAllowChunkedTransfer };
+	UE::IoStore::HTTP::FRequest HttpRequest = EventLoop.Get(Url, ConnectionPool, &RequestParams);
 
 	if (Request.Range.GetOffset() > 0 || Request.Range.GetLength() > 0)
 	{
diff --git a/Engine/Source/Runtime/Experimental/IoStore/OnDemand/Private/OnDemandHttpClient.h b/Engine/Source/Runtime/Experimental/IoStore/OnDemand/Private/OnDemandHttpClient.h
--- a/Engine/Source/Runtime/Experimental/IoStore/OnDemand/Private/OnDemandHttpClient.h
+++ b/Engine/Source/Runtime/Experimental/IoStore/OnDemand/Private/OnDemandHttpClient.h
@@ -91,6 +91,7 @@ struct FMultiEndpointHttpClientConfig
 	int32				TimeoutMs = 0;
 	EHttpRedirects		Redirects = EHttpRedirects::Follow;
 	bool				bEnableThreadSafetyChecks = false;
+	bool				bAllowChunkedTransfer = true;
 
 #if !NO_LOGGING
 	/** When not set (the default) this will disable logging, it is up to the system creating the client to assign a LogCategory */
diff --git a/Engine/Source/Runtime/Experimental/IoStore/OnDemand/Private/OnDemandIoDispatcherBackend.cpp b/Engine/Source/Runtime/Experimental/IoStore/OnDemand/Private/OnDemandIoDispatcherBackend.cpp
--- a/Engine/Source/Runtime/Experimental/IoStore/OnDemand/Private/OnDemandIoDispatcherBackend.cpp
+++ b/Engine/Source/Runtime/Experimental/IoStore/OnDemand/Private/OnDemandIoDispatcherBackend.cpp
@@ -191,6 +191,13 @@ static FAutoConsoleVariableRef CVar_DistributedEndpointAttemptCount(
 	TEXT("Number of times we should try to resolve a distributed endpoint befor eusing the fallback url (if there is one)")
 );
 
+static bool GIasHttpAllowChunkedXfer = false;
+static FAutoConsoleVariableRef CVar_IasHttpAllowChunkedXfer(
+	TEXT("ias.HttpAllowChunkedXfer"),
+	GIasHttpAllowChunkedXfer,
+	TEXT("Enable/disable IAS' support for chunked transfer encoding")
+);
+
 bool GIasEnableWriteOnlyDecoding = false;
 static FAutoConsoleVariableRef CVar_WriteOnlyDecoding(
 	TEXT("ias.EnableWriteOnlyDecoding"),
@@ -1743,6 +1750,7 @@ uint32 FOnDemandIoBackend::Run()
 			.MaxRetryCount = (uint32)FMath::Max(AvailableEps.Urls.Num(), GIasHttpRetryCount),
 			.TimeoutMs = GIasHttpFailTimeOutMs,
 			.bEnableThreadSafetyChecks = true,
+			.bAllowChunkedTransfer = GIasHttpAllowChunkedXfer,
 
 #if !NO_LOGGING
 			.LogCategoryName = LogIas.GetCategoryName(),
