IoStoreOnDemand
- Add API options to include soft references.

diff --git a/Engine/Plugins/Runtime/GameFeatures/Source/GameFeatures/Private/GameFeaturePluginStateMachine.cpp b/Engine/Plugins/Runtime/GameFeatures/Source/GameFeatures/Private/GameFeaturePluginStateMachine.cpp
--- a/Engine/Plugins/Runtime/GameFeatures/Source/GameFeatures/Private/GameFeaturePluginStateMachine.cpp
+++ b/Engine/Plugins/Runtime/GameFeatures/Source/GameFeatures/Private/GameFeaturePluginStateMachine.cpp
@@ -1692,6 +1692,7 @@ struct FGameFeaturePluginState_Downloading : public FGameFeaturePluginState
 			UE::IoStore::FOnDemandInstallArgs InstallArgs;
 			InstallArgs.MountId = InstallBundle.ToString();
 			InstallArgs.TagSets.Emplace(TEXTVIEW("mount")); // May not be a real tagset, but this will install all the mandatory untagged chunks
+			InstallArgs.Options |= UE::IoStore::EOnDemandInstallOptions::InstallSoftReferences;
 			InstallArgs.Options |= UE::IoStore::EOnDemandInstallOptions::CallbackOnGameThread;
 			
 			InstallArgs.ContentHandle = UE::GameFeatures::GFPSharedInstallTracker.AddOnDemandContentHandle(
@@ -2887,7 +2888,8 @@ struct FGameFeaturePluginState_AssetDependencyStreaming : public FGameFeaturePlu
 			if (InstallMode == EStreamingAssetInstallMode::GfpRequiredOnly)
 			{
 				InstallArgs.TagSets.Emplace(TEXTVIEW("required"));
-			}			
+			}
+			InstallArgs.Options |= UE::IoStore::EOnDemandInstallOptions::InstallSoftReferences;
 			InstallArgs.Options |= UE::IoStore::EOnDemandInstallOptions::CallbackOnGameThread;
 
 			InstallArgs.ContentHandle = UE::GameFeatures::GFPSharedInstallTracker.AddOnDemandContentHandle(
diff --git a/Engine/Plugins/Runtime/GameFeatures/Source/GameFeatures/Private/GameFeaturesSubsystem.cpp b/Engine/Plugins/Runtime/GameFeatures/Source/GameFeatures/Private/GameFeaturesSubsystem.cpp
--- a/Engine/Plugins/Runtime/GameFeatures/Source/GameFeatures/Private/GameFeaturesSubsystem.cpp
+++ b/Engine/Plugins/Runtime/GameFeatures/Source/GameFeatures/Private/GameFeaturesSubsystem.cpp
@@ -3144,6 +3144,7 @@ struct FGameFeatureAssetStreamContext : public FGameFeaturePluginAssetStreamHand
 
 			InstallArgs.ContentHandle = AssetRef.ContentHandle;
 			InstallArgs.Options |= UE::IoStore::EOnDemandInstallOptions::CallbackOnGameThread;
+			InstallArgs.Options |= UE::IoStore::EOnDemandInstallOptions::InstallSoftReferences;
 
 			// TODO: convert Priority to IOStore priority here
 
diff --git a/Engine/Source/Runtime/Experimental/IoStore/OnDemand/Private/OnDemandIoStore.cpp b/Engine/Source/Runtime/Experimental/IoStore/OnDemand/Private/OnDemandIoStore.cpp
--- a/Engine/Source/Runtime/Experimental/IoStore/OnDemand/Private/OnDemandIoStore.cpp
+++ b/Engine/Source/Runtime/Experimental/IoStore/OnDemand/Private/OnDemandIoStore.cpp
@@ -192,9 +192,17 @@ struct FContainerInstallData
 
 using FInstallData = TMap<FSharedOnDemandContainer, FContainerInstallData>;
 
+enum class EBuildInstallDataOptions
+{
+	None = 0,
+	IncludeSoftReferences = 1 << 0,
+};
+ENUM_CLASS_FLAGS(EBuildInstallDataOptions);
+
 static FIoStatus BuildInstallData(
 	const TSet<FSharedOnDemandContainer>& Containers,
 	const TSet<FPackageId>& PackageIds,
+	const EBuildInstallDataOptions Options,
 	FInstallData& OutInstallData,
 	TSet<FPackageId>& OutMissing)
 {
@@ -312,16 +320,19 @@ static FIoStatus BuildInstallData(
 					}
 				}
 
-				// Add soft references 
-				const FIoContainerHeader& Header = *Kv.Key->Header;
-				if (const FFilePackageStoreEntrySoftReferences* SoftRefs = Data.SoftPackageReferences.FindRef(PackageId))
+				if (EnumHasAnyFlags(Options, EBuildInstallDataOptions::IncludeSoftReferences))
 				{
-					for (uint32 Index : SoftRefs->Indices)
+					// Add soft references
+					if (const FFilePackageStoreEntrySoftReferences* SoftRefs = Data.SoftPackageReferences.FindRef(PackageId))
 					{
-						const FPackageId SoftPackageReference = Header.SoftPackageReferences.PackageIds[int32(Index)];
-						if (!Visitied.Contains(SoftPackageReference))
+						const FIoContainerHeader& Header = *Kv.Key->Header;
+						for (uint32 Index : SoftRefs->Indices)
 						{
-							Queue.Add(SoftPackageReference);
+							const FPackageId SoftPackageReference = Header.SoftPackageReferences.PackageIds[int32(Index)];
+							if (!Visitied.Contains(SoftPackageReference))
+							{
+								Queue.Add(SoftPackageReference);
+							}
 						}
 					}
 				}
@@ -759,10 +770,16 @@ TIoStatusOr<uint64> FOnDemandIoStore::GetInstallSize(const FOnDemandGetInstallSi
 		return Status;
 	}
 
+	Private::EBuildInstallDataOptions BuildInstallDataOptions = Private::EBuildInstallDataOptions::None;
+	if (EnumHasAnyFlags(Args.Options, EOnDemandGetInstallSizeOptions::IncludeSoftReferences))
+	{
+		BuildInstallDataOptions |= Private::EBuildInstallDataOptions::IncludeSoftReferences;
+	}
+
 	Private::FInstallData InstallData;
 	TSet<FPackageId> Missing;
 
-	FIoStatus Status = BuildInstallData(ContainersForInstallation, PackageIdsToInstall, InstallData, Missing);
+	FIoStatus Status = Private::BuildInstallData(ContainersForInstallation, PackageIdsToInstall, BuildInstallDataOptions, InstallData, Missing);
 	if (Status.IsOk() == false)
 	{
 		return Status;
@@ -792,10 +809,16 @@ FIoStatus FOnDemandIoStore::GetInstallSizesByMountId(const FOnDemandGetInstallSi
 		return Status;
 	}
 
+	Private::EBuildInstallDataOptions BuildInstallDataOptions = Private::EBuildInstallDataOptions::None;
+	if (EnumHasAnyFlags(Args.Options, EOnDemandGetInstallSizeOptions::IncludeSoftReferences))
+	{
+		BuildInstallDataOptions |= Private::EBuildInstallDataOptions::IncludeSoftReferences;
+	}
+
 	Private::FInstallData InstallData;
 	TSet<FPackageId> Missing;
 
-	FIoStatus Status = BuildInstallData(ContainersForInstallation, PackageIdsToInstall, InstallData, Missing);
+	FIoStatus Status = Private::BuildInstallData(ContainersForInstallation, PackageIdsToInstall, BuildInstallDataOptions, InstallData, Missing);
 	if (Status.IsOk() == false)
 	{
 		return Status;
@@ -1349,10 +1372,16 @@ FOnDemandInstallResult FOnDemandIoStore::TickInstallRequest(FOnDemandInternalIns
 	// Its OK for PackageIdsToInstall to be empty at this point. Any chunks not referenced by a package must still be installed.
 
 	// Find all I/O chunks fo the specified list of packages
+	Private::EBuildInstallDataOptions BuildInstallDataOptions = Private::EBuildInstallDataOptions::None;
+	if (EnumHasAnyFlags(InstallRequest.Args.Options, EOnDemandInstallOptions::InstallSoftReferences))
+	{
+		BuildInstallDataOptions |= Private::EBuildInstallDataOptions::IncludeSoftReferences;
+	}
+
 	Private::FInstallData InstallData;
 	TSet<FPackageId> Missing;
 
-	Status = BuildInstallData(ContainersForInstallation, PackageIdsToInstall, InstallData, Missing);
+	Status = Private::BuildInstallData(ContainersForInstallation, PackageIdsToInstall, BuildInstallDataOptions, InstallData, Missing);
 	if (Status.IsOk() == false)
 	{
 		return OutResult;
diff --git a/Engine/Source/Runtime/Experimental/IoStore/OnDemand/Public/IO/IoStoreOnDemand.h b/Engine/Source/Runtime/Experimental/IoStore/OnDemand/Public/IO/IoStoreOnDemand.h
--- a/Engine/Source/Runtime/Experimental/IoStore/OnDemand/Public/IO/IoStoreOnDemand.h
+++ b/Engine/Source/Runtime/Experimental/IoStore/OnDemand/Public/IO/IoStoreOnDemand.h
@@ -177,6 +177,8 @@ enum class EOnDemandInstallOptions
 	None				 = 0,
 	/** Trigger callback on game thread. */
 	CallbackOnGameThread = 1 << 0,
+	/** Follow soft references when gathering packages to install */
+	InstallSoftReferences = 1 << 1,
 };
 ENUM_CLASS_FLAGS(EOnDemandInstallOptions);
 
@@ -325,6 +327,16 @@ struct FOnDemandPurgeResult
 /** Purge completion callback */
 using FOnDemandPurgeCompleted = TUniqueFunction<void(FOnDemandPurgeResult)>;
 
+/** Options for controlling the behavior of the install size request. */
+enum class EOnDemandGetInstallSizeOptions
+{
+	/** No additional options. */
+	None = 0,
+	/** Follow soft references when gathering packages to install */
+	IncludeSoftReferences = 1 << 0,
+};
+ENUM_CLASS_FLAGS(EOnDemandGetInstallSizeOptions);
+
 /** Arguments for getting the size of on-demand content. */
 struct FOnDemandGetInstallSizeArgs 
 {
@@ -334,6 +346,8 @@ struct FOnDemandGetInstallSizeArgs
 	TArray<FString> TagSets;
 	/** Get total intall size for the specified package IDs. */
 	TArray<FPackageId> PackageIds;
+	/** Options */
+	EOnDemandGetInstallSizeOptions Options = EOnDemandGetInstallSizeOptions::None;
 };
 
 /** Holds information about the install cache usage */
