Bug fix for rare edge case when connecting to server via command line

diff --git a/Engine/Source/Runtime/Experimental/GeometryCollectionEngine/Private/GeometryCollection/GeometryCollectionComponent.cpp b/Engine/Source/Runtime/Experimental/GeometryCollectionEngine/Private/GeometryCollection/GeometryCollectionComponent.cpp
--- a/Engine/Source/Runtime/Experimental/GeometryCollectionEngine/Private/GeometryCollection/GeometryCollectionComponent.cpp
+++ b/Engine/Source/Runtime/Experimental/GeometryCollectionEngine/Private/GeometryCollection/GeometryCollectionComponent.cpp
@@ -4552,7 +4552,10 @@ void UGeometryCollectionComponent::RequestUpdateRepData()
 			PhysScene->EnqueueAsyncPhysicsCommand(CurrentFrame, this,
 				[this]()
 				{
-					UpdateRepData();
+					if (GetIsReplicated() && GetNetMode() != ENetMode::NM_Client)
+					{
+						UpdateRepData();
+					}
 				}, false
 			);
 		}
