Fixed perf timers for collision detection
Split Broad Phase and  Narrow Phase for the main physics thread

diff --git a/Engine/Source/Runtime/Experimental/Chaos/Private/ChaosStats.cpp b/Engine/Source/Runtime/Experimental/Chaos/Private/ChaosStats.cpp
--- a/Engine/Source/Runtime/Experimental/Chaos/Private/ChaosStats.cpp
+++ b/Engine/Source/Runtime/Experimental/Chaos/Private/ChaosStats.cpp
@@ -112,7 +112,7 @@ DEFINE_STAT(STAT_ForceUpdateField_LinearImpulse);
 DEFINE_STAT(STAT_Collisions_Detect);
 DEFINE_STAT(STAT_Collisions_ParticlePairBroadPhase);
 DEFINE_STAT(STAT_Collisions_SpatialBroadPhase);
-DEFINE_STAT(STAT_Collisions_MidPhase);
+DEFINE_STAT(STAT_Collisions_BroadPhase);
 DEFINE_STAT(STAT_Collisions_AssignMidPhases);
 DEFINE_STAT(STAT_Collisions_NarrowPhase);
 DEFINE_STAT(STAT_Collisions_GenerateCollisions);
diff --git a/Engine/Source/Runtime/Experimental/Chaos/Public/Chaos/Collision/SpatialAccelerationBroadPhase.h b/Engine/Source/Runtime/Experimental/Chaos/Public/Chaos/Collision/SpatialAccelerationBroadPhase.h
--- a/Engine/Source/Runtime/Experimental/Chaos/Public/Chaos/Collision/SpatialAccelerationBroadPhase.h
+++ b/Engine/Source/Runtime/Experimental/Chaos/Public/Chaos/Collision/SpatialAccelerationBroadPhase.h
@@ -409,6 +409,8 @@ namespace Chaos
 			const FCollisionDetectorSettings& Settings,
 			IResimCacheBase* ResimCache)
 		{
+			SCOPE_CYCLE_COUNTER(STAT_Collisions_BroadPhase);
+			CSV_SCOPED_TIMING_STAT(PhysicsVerbose, DetectCollisions_BroadPhase);
 
 			if (!ensure(SpatialAcceleration))
 			{
@@ -446,11 +448,7 @@ namespace Chaos
 					check(false);  //question: do we want to support a dynamic dispatch version?
 				}
 			}
-
 			{
-				SCOPE_CYCLE_COUNTER(STAT_Collisions_MidPhase);
-				CSV_SCOPED_TIMING_STAT(PhysicsVerbose, DetectCollisions_AssignMidPhases);
-
 				FChaosVDContextWrapper CVDContext;
 				CVD_GET_WRAPPED_CURRENT_CONTEXT(CVDContext);
 
@@ -485,9 +483,6 @@ namespace Chaos
 		*/
 		void ProduceCollisions(FReal Dt)
 		{
-			SCOPE_CYCLE_COUNTER(STAT_Collisions_NarrowPhase);
-			CSV_SCOPED_TIMING_STAT(PhysicsVerbose, DetectCollisions_MidPhase);
-
 			FChaosVDContextWrapper CVDContext;
 			CVD_GET_WRAPPED_CURRENT_CONTEXT(CVDContext);
 
@@ -516,7 +511,6 @@ namespace Chaos
 
 		void GatherConstraints(bool bIsDeterministic)
 		{
-			SCOPE_CYCLE_COUNTER(STAT_Collisions_NarrowPhase);
 			if (bIsDeterministic)
 			{
 				FGraphEventArray ProcessSorting;
@@ -957,6 +951,7 @@ namespace Chaos
 		int32 NumMidPhases;
 
 		FGraphEventArray PendingTasks;
+		// Contexts indexes in which there are constraints
 		TArray<int32> ContextsWithConstraints;
 
 		template<bool bOnlyRigid, typename THandle, typename LambdaHandle, typename SpatialAccelerationType>
diff --git a/Engine/Source/Runtime/Experimental/Chaos/Public/Chaos/Collision/SpatialAccelerationCollisionDetector.h b/Engine/Source/Runtime/Experimental/Chaos/Public/Chaos/Collision/SpatialAccelerationCollisionDetector.h
--- a/Engine/Source/Runtime/Experimental/Chaos/Public/Chaos/Collision/SpatialAccelerationCollisionDetector.h
+++ b/Engine/Source/Runtime/Experimental/Chaos/Public/Chaos/Collision/SpatialAccelerationCollisionDetector.h
@@ -45,9 +45,12 @@ namespace Chaos
 			}
 
 			// Run collision detection on the output of the broadphase
-			BroadPhase.ProduceCollisions(Dt);
-
-			BroadPhase.GatherConstraints(GetCollisionContainer().GetIsDeterministic());
+			{
+				SCOPE_CYCLE_COUNTER(STAT_Collisions_NarrowPhase);
+				CSV_SCOPED_TIMING_STAT(PhysicsVerbose, DetectCollisions_NarrowPhase);
+				BroadPhase.ProduceCollisions(Dt);
+				BroadPhase.GatherConstraints(GetCollisionContainer().GetIsDeterministic());
+			}
 			GetCollisionContainer().EndDetectCollisions();
 
 			// If we have a resim cache restore and save contacts
diff --git a/Engine/Source/Runtime/Experimental/Chaos/Public/ChaosStats.h b/Engine/Source/Runtime/Experimental/Chaos/Public/ChaosStats.h
--- a/Engine/Source/Runtime/Experimental/Chaos/Public/ChaosStats.h
+++ b/Engine/Source/Runtime/Experimental/Chaos/Public/ChaosStats.h
@@ -137,7 +137,7 @@ DECLARE_CYCLE_STAT_EXTERN(TEXT("[Field Update] Impulse"), STAT_ForceUpdateField_
 DECLARE_CYCLE_STAT_EXTERN(TEXT("Collisions::Detect"), STAT_Collisions_Detect, STATGROUP_ChaosCollision, CHAOS_API);
 DECLARE_CYCLE_STAT_EXTERN(TEXT("Collisions::BroadPhase::ParticlePair"), STAT_Collisions_ParticlePairBroadPhase, STATGROUP_ChaosCollision, CHAOS_API);
 DECLARE_CYCLE_STAT_EXTERN(TEXT("Collisions::BroadPhase::Spatial"), STAT_Collisions_SpatialBroadPhase, STATGROUP_ChaosCollision, CHAOS_API);
-DECLARE_CYCLE_STAT_EXTERN(TEXT("Collisions::MidPhase"), STAT_Collisions_MidPhase, STATGROUP_ChaosCollision, CHAOS_API);
+DECLARE_CYCLE_STAT_EXTERN(TEXT("Collisions::MidPhase"), STAT_Collisions_BroadPhase, STATGROUP_ChaosCollision, CHAOS_API);
 DECLARE_CYCLE_STAT_EXTERN(TEXT("Collisions::AssignMidPhases"), STAT_Collisions_AssignMidPhases, STATGROUP_ChaosCollision, CHAOS_API);
 DECLARE_CYCLE_STAT_EXTERN(TEXT("Collisions::NarrowPhase"), STAT_Collisions_NarrowPhase, STATGROUP_ChaosCollision, CHAOS_API);
 DECLARE_CYCLE_STAT_EXTERN(TEXT("Collisions::GenerateCollisions"), STAT_Collisions_GenerateCollisions, STATGROUP_ChaosCollision, CHAOS_API);
