Dataflow:
- Changed OutputConcreteType from double to float to make sure outputs can be connected

diff --git a/Engine/Source/Runtime/Experimental/Dataflow/Core/Private/Dataflow/DataflowMathNodes.cpp b/Engine/Source/Runtime/Experimental/Dataflow/Core/Private/Dataflow/DataflowMathNodes.cpp
--- a/Engine/Source/Runtime/Experimental/Dataflow/Core/Private/Dataflow/DataflowMathNodes.cpp
+++ b/Engine/Source/Runtime/Experimental/Dataflow/Core/Private/Dataflow/DataflowMathNodes.cpp
@@ -48,10 +48,6 @@ namespace UE::Dataflow
 		DATAFLOW_NODE_REGISTER_CREATION_FACTORY(FDataflowMathArcTan2Node)
 		DATAFLOW_NODE_REGISTER_CREATION_FACTORY(FDataflowMathDegToRadNode);
 		DATAFLOW_NODE_REGISTER_CREATION_FACTORY(FDataflowMathRadToDegNode);
-
-		// Math
-		static constexpr FLinearColor CDefaultMathNodeBodyTintColor(0.f, 0.f, 0.f, 0.5f);
-		DATAFLOW_NODE_REGISTER_CREATION_FACTORY_NODE_COLORS_BY_CATEGORY("Math", FLinearColor(0.f, 0.4f, 0.8f), CDefaultMathNodeBodyTintColor);
 	}
 }
 
@@ -69,7 +65,7 @@ void FDataflowMathOneInputOperatorNode::RegisterInputsAndOutputs()
 
 	// Set the output to double for now so that it is strongly type and easy to connect to the next node
 	// Once we can change the output type from the UI, this could be removed 
-	SetOutputConcreteType(&Result, TDataflowSingleTypePolicy<double>::TypeName);
+	SetOutputConcreteType(&Result, TDataflowSingleTypePolicy<float>::TypeName);
 }
 
 void FDataflowMathOneInputOperatorNode::Evaluate(UE::Dataflow::FContext& Context, const FDataflowOutput* Out) const
@@ -97,7 +93,7 @@ void FDataflowMathTwoInputsOperatorNode::RegisterInputsAndOutputs()
 
 	// Set the output to double for now so that it is strongly type and easy to connect to the next node
 	// Once we can change the output type from the UI, this could be removed 
-	SetOutputConcreteType(&Result, TDataflowSingleTypePolicy<double>::TypeName);
+	SetOutputConcreteType(&Result, TDataflowSingleTypePolicy<float>::TypeName);
 }
 
 void FDataflowMathTwoInputsOperatorNode::Evaluate(UE::Dataflow::FContext& Context, const FDataflowOutput* Out) const
@@ -459,7 +455,7 @@ FDataflowMathConstantNode::FDataflowMathConstantNode(const UE::Dataflow::FNodePa
 
 	// Set the output to double for now so that it is strongly type and easy to connect to the next node
 	// Once we can change the output type from the UI, this could be removed 
-	SetOutputConcreteType(&Result, TDataflowSingleTypePolicy<double>::TypeName);
+	SetOutputConcreteType(&Result, TDataflowSingleTypePolicy<float>::TypeName);
 }
 
 double FDataflowMathConstantNode::GetConstant() const
