#PLAY-74040

Temporary reversion of previous change to CCD flag code.  The original change was producing unintended effects in small wheels and will need to be evaluated before being restored.

diff --git a/Engine/Source/Runtime/Experimental/Chaos/Private/Chaos/Collision/ParticlePairMidPhase.cpp b/Engine/Source/Runtime/Experimental/Chaos/Private/Chaos/Collision/ParticlePairMidPhase.cpp
--- a/Engine/Source/Runtime/Experimental/Chaos/Private/Chaos/Collision/ParticlePairMidPhase.cpp
+++ b/Engine/Source/Runtime/Experimental/Chaos/Private/Chaos/Collision/ParticlePairMidPhase.cpp
@@ -477,6 +477,10 @@ namespace Chaos
 			Constraint->SetCullDistance(CullDistance);
 			Constraint->SetRelativeMovement(RelativeMovement);
 
+			// Constraint may have been previously used with CCD enabled (e.g., a midphase modifier)
+			// so we need to make sure that the CCD flag is disabled
+			Constraint->SetCCDEnabled(false);
+
 			// If the constraint was not used last frame, it needs to be reset, otherwise we will try to reuse
 			if (!bWasUpdatedLastTick || (Constraint->GetManifoldPoints().Num() == 0))
 			{
