Fixed error in cvar description.

#rnxx

diff --git a/Engine/Source/Runtime/Experimental/IoStore/HttpClient/Private/Misc.inl b/Engine/Source/Runtime/Experimental/IoStore/HttpClient/Private/Misc.inl
--- a/Engine/Source/Runtime/Experimental/IoStore/HttpClient/Private/Misc.inl
+++ b/Engine/Source/Runtime/Experimental/IoStore/HttpClient/Private/Misc.inl
@@ -111,7 +111,7 @@ static void Trace(T Id, ETrace Action, U Param=0)
 
 ////////////////////////////////////////////////////////////////////////////////
 static IAS_CVAR(int32, RecvWorkThresholdKiB,80,		"Threshold of data remaining at which next request is sent (in KiB)");
-static IAS_CVAR(int32, IdleMs,				50'000,	"Time in seconds to close idle connections or fail waits");
+static IAS_CVAR(int32, IdleMs,				50'000,	"Time in milliseconds to close idle connections or fail waits");
 
 ////////////////////////////////////////////////////////////////////////////////
 class FOutcome
