Reduce pools' exposure so they are no longer aware of peers

#rnx
#ushell-unshelve of 39626023

diff --git a/Engine/Source/Runtime/Experimental/IoStore/HttpClient/Private/ConnectionPool.inl b/Engine/Source/Runtime/Experimental/IoStore/HttpClient/Private/ConnectionPool.inl
--- a/Engine/Source/Runtime/Experimental/IoStore/HttpClient/Private/ConnectionPool.inl
+++ b/Engine/Source/Runtime/Experimental/IoStore/HttpClient/Private/ConnectionPool.inl
@@ -25,7 +25,7 @@ public:
 					FHost(const FParams& Params);
 	void			SetBufferSize(EDirection Dir, int32 Size);
 	int32			GetBufferSize(EDirection Dir) const;
-	FOutcome		Connect(FTlsPeer& Peer);
+	FOutcome		Connect(FSocket& Socket);
 	int32			IsResolved() const;
 	FOutcome		ResolveHostName();
 	FCertRootsRef	GetVerifyCert() const		{ return VerifyCert; }
@@ -35,7 +35,6 @@ public:
 	uint32			GetPort() const				{ return Port; }
 
 private:
-	FOutcome		Connect(FSocket& Socket);
 	FCertRootsRef	VerifyCert;
 	const ANSICHAR*	HostName;
 	uint32			IpAddresses[4] = {};
@@ -142,30 +141,11 @@ int32 FHost::IsResolved() const
 	}
 }
 
-////////////////////////////////////////////////////////////////////////////////
-FOutcome FHost::Connect(FTlsPeer& Peer)
-{
-	FSocket Socket;
-	FOutcome Outcome = Connect(Socket);
-
-	if (Outcome.IsError())
-	{
-		return Outcome;
-	}
-
-	if (VerifyCert == ECertRootsRefType::None)
-	{
-		Peer = FTlsPeer(MoveTemp(Socket));
-		return Outcome;
-	}
-
-	Peer = FTlsPeer(MoveTemp(Socket), VerifyCert, HostName);
-	return Outcome;
-}
-
 ////////////////////////////////////////////////////////////////////////////////
 FOutcome FHost::Connect(FSocket& Socket)
 {
+	check(Socket.IsValid());
+
 	FOutcome Outcome = FOutcome::None();
 
 	if (IsResolved() <= 0)
@@ -177,18 +157,11 @@ FOutcome FHost::Connect(FSocket& Socket)
 	}
 
 	check(IsResolved() > 0);
-	check(!Socket.IsValid());
 
 	uint32 IpAddress = GetIpAddress();
 
-	FSocket Candidate;
-	if (!Candidate.Create())
-	{
-		return FOutcome::Error("Failed to create socket");
-	}
-
 	// Attempt a SOCKS connect
-	Outcome = MaybeConnectSocks(Candidate, IpAddress, Port);
+	Outcome = MaybeConnectSocks(Socket, IpAddress, Port);
 	if (Outcome.IsError())
 	{
 		return Outcome;
@@ -197,35 +170,34 @@ FOutcome FHost::Connect(FSocket& Socket)
 	bool bSocksConnected = (Outcome.GetResult() == 1);
 
 	// Condition the socket
-	if (!Candidate.SetBlocking(false))
+	if (!Socket.SetBlocking(false))
 	{
 		return FOutcome::Error("Unable to set socket non-blocking");
 	}
 
 	if (int32 OptValue = GetBufferSize(FHost::EDirection::Send); OptValue >= 0)
 	{
-		Candidate.SetSendBufSize(OptValue);
+		Socket.SetSendBufSize(OptValue);
 	}
 
 	if (int32 OptValue = GetBufferSize(FHost::EDirection::Recv); OptValue >= 0)
 	{
-		Candidate.SetRecvBufSize(OptValue);
+		Socket.SetRecvBufSize(OptValue);
 	}
 
 	// Socks connect in a blocking fashion so we're all set (ret=1)
 	if (bSocksConnected)
 	{
-		Socket = MoveTemp(Candidate);
+		Socket = MoveTemp(Socket);
 		return FOutcome::Ok();
 	}
 
 	// Issue the connect - this is done non-blocking so we need to wait (ret=0)
-	if (Outcome = Candidate.Connect(IpAddress, Port); Outcome.IsError())
+	if (Outcome = Socket.Connect(IpAddress, Port); Outcome.IsError())
 	{
 		return Outcome;
 	}
 
-	Socket = MoveTemp(Candidate);
 	return Outcome;
 }
 
diff --git a/Engine/Source/Runtime/Experimental/IoStore/HttpClient/Private/Loop.inl b/Engine/Source/Runtime/Experimental/IoStore/HttpClient/Private/Loop.inl
--- a/Engine/Source/Runtime/Experimental/IoStore/HttpClient/Private/Loop.inl
+++ b/Engine/Source/Runtime/Experimental/IoStore/HttpClient/Private/Loop.inl
@@ -1004,20 +1004,34 @@ void FPeerGroup::TickSend(FTickState& State, FHost& Host)
 	bool bWillBlock = false;
 	if (!Peer.IsValid())
 	{
-		IsKeepAlive = 1;
-		bNegotiating = true;
-		FOutcome Outcome = Host.Connect(Peer);
+		FOutcome Outcome = FOutcome::None();
+
+		FSocket Socket;
+		if (Socket.Create())
+		{
+			Outcome = Host.Connect(Socket);
+		}
+		else
+		{
+			Outcome = FOutcome::Error("Failed to create socket");
+		}
 
-		// We failed to connect, let's bail.
 		if (Outcome.IsError())
 		{
+			// We failed to connect, let's bail.
 			Pending->Next = Recv;
 			Recv = Pending;
 			Fail(State, Outcome.GetMessage().GetData());
 			return;
 		}
 
+		IsKeepAlive = 1;
+		bNegotiating = true;
 		bWillBlock = Outcome.IsWaiting();
+
+		FCertRootsRef VerifyCert = Host.GetVerifyCert();
+		const char* HostName = Host.GetHostName().GetData();
+		Peer = FPeerType(MoveTemp(Socket), VerifyCert, HostName);
 	}
 
 	Send = Pending;
diff --git a/Engine/Source/Runtime/Experimental/IoStore/HttpClient/Private/Socks.inl b/Engine/Source/Runtime/Experimental/IoStore/HttpClient/Private/Socks.inl
--- a/Engine/Source/Runtime/Experimental/IoStore/HttpClient/Private/Socks.inl
+++ b/Engine/Source/Runtime/Experimental/IoStore/HttpClient/Private/Socks.inl
@@ -171,6 +171,8 @@ static FOutcome MaybeConnectSocks(FSocket& Socket, uint32 IpAddress, uint32 Port
 		return FOutcome::Ok();
 	}
 
+	Socket.SetBlocking(true);
+
 	switch (GSocksVersion)
 	{
 	case 4: return ConnectSocks4(Socket, IpAddress, Port);
