Restore cvar 'ias.HttpConcurrentRequests' to it's original value that was incorrect overwritten in CL 39402569

#rnx

diff --git a/Engine/Source/Runtime/Experimental/IoStore/OnDemand/Private/OnDemandIoDispatcherBackend.cpp b/Engine/Source/Runtime/Experimental/IoStore/OnDemand/Private/OnDemandIoDispatcherBackend.cpp
--- a/Engine/Source/Runtime/Experimental/IoStore/OnDemand/Private/OnDemandIoDispatcherBackend.cpp
+++ b/Engine/Source/Runtime/Experimental/IoStore/OnDemand/Private/OnDemandIoDispatcherBackend.cpp
@@ -88,7 +88,7 @@ static FAutoConsoleVariableRef CVar_GIasHttpRecvBufKiB(
 	TEXT("Recv buffer size")
 );
 
-static int32 GIasHttpConcurrentRequests = 1;
+static int32 GIasHttpConcurrentRequests = 8;
 static FAutoConsoleVariableRef CVar_IasHttpConcurrentRequests(
 	TEXT("ias.HttpConcurrentRequests"),
 	GIasHttpConcurrentRequests,
