Remove FAvailableEndPoints

#rnx

- Since the recent refactor, FAvailableEndPoints has been almost entirely replaced by FOnDemandHostGroup and was only being used during the initial set up phase of the FOnDemandIoBackend. We don't really need any of the functionality anymore so it can be replaced by a simple array of strings.
- This also makes it easier to move more of the set up code from FOnDemandIoBackend::Run to a new ::CreateDefaultHostGroup method, making the code flow easier to follow.

diff --git a/Engine/Source/Runtime/Experimental/IoStore/OnDemand/Private/AvailableEndPoints.h b/Engine/Source/Runtime/Experimental/IoStore/OnDemand/Private/AvailableEndPoints.h
deleted file mode 100644
diff --git a/Engine/Source/Runtime/Experimental/IoStore/OnDemand/Private/OnDemandIoDispatcherBackend.cpp b/Engine/Source/Runtime/Experimental/IoStore/OnDemand/Private/OnDemandIoDispatcherBackend.cpp
--- a/Engine/Source/Runtime/Experimental/IoStore/OnDemand/Private/OnDemandIoDispatcherBackend.cpp
+++ b/Engine/Source/Runtime/Experimental/IoStore/OnDemand/Private/OnDemandIoDispatcherBackend.cpp
@@ -3,7 +3,6 @@
 #include "OnDemandIoDispatcherBackend.h"
 
 #include "AnalyticsEventAttribute.h"
-#include "AvailableEndPoints.h"
 #include "Containers/BitArray.h"
 #include "Containers/StringView.h"
 #include "DistributionEndpoints.h"
@@ -935,7 +934,8 @@ private:
 	void ConditionallyStartBackendThread();
 	void CompleteRequest(FChunkRequest* ChunkRequest);
 	void CompleteCacheRead(FChunkRequest* ChunkRequest);
-	bool ResolveDistributedEndpoint(const FDistributedEndpointUrl& Url);
+	bool CreateDefaultHostGroup();
+	bool ResolveDistributedEndpoint(const FDistributedEndpointUrl& Url, TArray<FString>& OutUrls);
 	void ProcessHttpRequests(FMultiEndpointHttpClient& HttpClient, FBitWindow& HttpErrors, int32 MaxConcurrentRequests);
 	int32 WaitForCompleteRequestTasks(float WaitTimeSeconds, float PollTimeSeconds);
 	void DrainHttpRequests();
@@ -956,7 +956,7 @@ private:
 	FChunkRequestQueue HttpRequests;
 	TArray<FChunkRequest*> PendingCacheGets;
 	FBackendStatus BackendStatus;
-	FAvailableEndPoints AvailableEps;
+	TArray<FString> DefaultUrls;
 	FOnDemandIoBackendStats Stats;
 	FDistributedEndpointUrl DistributionUrl;
 	FEventRef DistributedEndpointEvent;
@@ -989,13 +989,7 @@ FOnDemandIoBackend::FOnDemandIoBackend(const FOnDemandEndpointConfig& Config, FO
 	}
 	else
 	{
-		for (const FString& ServiceUrl : Config.ServiceUrls)
-		{
-			for (const FString& Url : Config.ServiceUrls)
-			{
-				AvailableEps.Urls.Add(Url.Replace(TEXT("https"), TEXT("http")).ToLower());
-			}
-		}
+		DefaultUrls = Config.ServiceUrls;
 	}
 
 	// Don't enable HTTP until the background thread has been started
@@ -1445,7 +1439,59 @@ const TCHAR* FOnDemandIoBackend::GetName() const
 	return TEXT("OnDemand");
 }
 
-bool FOnDemandIoBackend::ResolveDistributedEndpoint(const FDistributedEndpointUrl& DistributedEndpointUrl)
+bool FOnDemandIoBackend::CreateDefaultHostGroup()
+{
+	if (DistributionUrl.IsValid())
+	{
+		DefaultUrls.Empty(); // We don't want any pre-existing urls if we are getting them from the distributed endpoint
+
+		if (ResolveDistributedEndpoint(DistributionUrl, DefaultUrls) == false)
+		{
+			// ResolveDistributedEndpoint should spin forever until either a valid url is found or
+			// we give up and use a predetermined fallback url. If this returned false then we didn't
+			// have a fallback url but the current process is shutting down so we might as well just
+			// exist the thread early.
+			UE_LOG(LogIas, Error, TEXT("Failed to resolve CDN endpoints from distribution URL"));
+			return false;
+		}
+	}
+
+	if (DefaultUrls.IsEmpty())
+	{
+		UE_LOG(LogIas, Error, TEXT("HTTP streaming disabled, no valid urls"));
+		return false;
+	}
+
+	// Sanitize the urls
+	for (FString& Url : DefaultUrls)
+	{
+		Url.ReplaceInline(TEXT("https"), TEXT("http"));
+		Url.ToLowerInline();
+	}
+
+	if (GIasHttpPrimaryEndpoint > 0)
+	{
+		// Rotate the list of urls so that the primary endpoint is the first element
+		Algo::Rotate(DefaultUrls, GIasHttpPrimaryEndpoint);
+	}
+
+	TIoStatusOr<FOnDemandHostGroup> HostResult = FOnDemandHostGroup::Create(DefaultUrls);
+	DefaultUrls.Empty(); // No longer needed
+
+	if (!HostResult.IsOk())
+	{
+		UE_LOG(LogIas, Error, TEXT("HTTP streaming disabled, could not create the default host group (%s)"), *HostResult.Status().ToString());
+		return false;
+	}
+
+	DefaultHostGroup = HostResult.ConsumeValueOrDie();
+
+	Stats.AddHostGroup(DefaultHostGroup);
+
+	return true;
+}
+
+bool FOnDemandIoBackend::ResolveDistributedEndpoint(const FDistributedEndpointUrl& DistributedEndpointUrl, TArray<FString>& OutUrls)
 {
 	TRACE_CPUPROFILER_EVENT_SCOPE(IasBackend::ResolveDistributedEndpoint);
 
@@ -1461,22 +1507,11 @@ bool FOnDemandIoBackend::ResolveDistributedEndpoint(const FDistributedEndpointUr
 
 	while (!bStopRequested)
 	{
-		TArray<FString> ServiceUrls;
-
 		FDistributionEndpoints Resolver;
-		FDistributionEndpoints::EResult Result = Resolver.ResolveEndpoints(DistributedEndpointUrl.EndpointUrl, ServiceUrls, *DistributedEndpointEvent.Get());
+		FDistributionEndpoints::EResult Result = Resolver.ResolveEndpoints(DistributedEndpointUrl.EndpointUrl, OutUrls, *DistributedEndpointEvent.Get());
 		if (Result == FDistributionEndpoints::EResult::Success)
 		{
 			Stats.OnHttpDistributedEndpointResolved();
-
-			FWriteScopeLock _(Lock);
-			AvailableEps.Urls.Reserve(AvailableEps.Urls.Num() + ServiceUrls.Num());
-
-			for (const FString& Url : ServiceUrls)
-			{
-				AvailableEps.Urls.Add(Url.Replace(TEXT("https"), TEXT("http")));
-			}
-
 			return true;
 		}
 
@@ -1484,10 +1519,8 @@ bool FOnDemandIoBackend::ResolveDistributedEndpoint(const FDistributedEndpointUr
 		{
 			FString FallbackUrl = DistributedEndpointUrl.FallbackUrl.Replace(TEXT("https"), TEXT("http"));
 			UE_LOG(LogIas, Warning, TEXT("Failed to resolve the distributed endpoint %d times. Fallback CDN '%s' will be used instead"), GDistributedEndpointAttemptCount , *FallbackUrl);
-			
-			FWriteScopeLock _(Lock);
-			AvailableEps.Urls.Emplace(MoveTemp(FallbackUrl));
-		
+
+			OutUrls.Emplace(MoveTemp(FallbackUrl));
 			return true;
 		}
 
@@ -1534,11 +1567,10 @@ void FOnDemandIoBackend::ReportAnalytics(TArray<FAnalyticsEventAttribute>& OutAn
 
 	Stats.ReportGeneralAnalytics(OutAnalyticsArray);
 
-	if (AvailableEps.HasCurrent())
+	if (DefaultHostGroup.PrimaryHostIndex() != INDEX_NONE)
 	{
-		FString CdnUrl = AvailableEps.GetCurrent();
+		FString CdnUrl = FString(DefaultHostGroup.PrimaryHost());
 
-		// Strip the prefix from the url as some analytics systems may have trouble dealing with it
 		if (!CdnUrl.RemoveFromStart(TEXT("http://")))
 		{
 			CdnUrl.RemoveFromStart(TEXT("https://"));
@@ -1711,43 +1743,19 @@ uint32 FOnDemandIoBackend::Run()
 
 	BackendStatus.SetHttpEnabled(false); // Http is disabled until the thread starts up successfully
 
-	if (DistributionUrl.IsValid())
-	{
-		if (ResolveDistributedEndpoint(DistributionUrl) == false)
-		{
-			// ResolveDistributedEndpoint should spin forever until either a valid url is found or
-			// we give up and use a predetermined fallback url. If this returned false then we didn't
-			// have a fallback url but the current process is shutting down so we might as well just
-			// exist the thread early.
-			UE_LOG(LogIas, Warning, TEXT("Failed to resolve CDN endpoints from distribution URL"));
-			return 0;
-		}
-	}
-
-	if (AvailableEps.Urls.IsEmpty())
+	if (!CreateDefaultHostGroup())
 	{
-		UE_LOG(LogIas, Error, TEXT("HTTP streaming disabled, no valid endpoint"));
 		return 0;
 	}
 
-	if (GIasHttpPrimaryEndpoint < 0)
-	{
-		UE_LOG(LogIas, Error, TEXT("ias.HttpPrimaryEndpoint should not be set as a negative number, defaulting to 0"));
-		GIasHttpPrimaryEndpoint = 0;
-	}
-
-	// Rotate the list of urls so that the primary endpoint is the first element
-	Algo::Rotate(AvailableEps.Urls, GIasHttpPrimaryEndpoint);
-	AvailableEps.Current = 0;
-
 	FBitWindow HttpErrors;
 	HttpErrors.Reset(GIasHttpErrorSampleCount);
-
+	
 	TUniquePtr<FMultiEndpointHttpClient> HttpClient = FMultiEndpointHttpClient::Create(FMultiEndpointHttpClientConfig
 		{
 			.MaxConnectionCount = GIasHttpConnectionCount,
 			.ReceiveBufferSize = GIasHttpRecvBufKiB >= 0 ? GIasHttpRecvBufKiB << 10 : -1,
-			.MaxRetryCount = (uint32)FMath::Max(AvailableEps.Urls.Num(), GIasHttpRetryCount),
+			.MaxRetryCount = (uint32)FMath::Max(DefaultHostGroup.Hosts().Num(), GIasHttpRetryCount),
 			.TimeoutMs = GIasHttpFailTimeOutMs,
 			.bEnableThreadSafetyChecks = true,
 			.bAllowChunkedTransfer = GIasHttpAllowChunkedXfer,
@@ -1759,17 +1767,6 @@ uint32 FOnDemandIoBackend::Run()
 		});
 	check(HttpClient.IsValid());
 
-	TIoStatusOr<FOnDemandHostGroup> HostResult = FOnDemandHostGroup::Create(AvailableEps.Urls);
-	if (!HostResult.IsOk())
-	{
-		UE_LOG(LogIas, Error, TEXT("HTTP streaming disabled, could not create the default host group (%s)"), *HostResult.Status().ToString());
-		return 0;
-	}
-
-	DefaultHostGroup = HostResult.ConsumeValueOrDie();
-
-	Stats.AddHostGroup(DefaultHostGroup);
-
 	const double InitStart = FPlatformTime::Seconds();
 
 	if (LatencyTest(DefaultHostGroup.PrimaryHost(), GetEndpointTestPath()))
