Correctly wait if the send doesn't complete immediately. Don't overwrite wait if an activity is dispatched as the previous one nears its end.

#rnx
#ushell-unshelve of 39626023

diff --git a/Engine/Source/Runtime/Experimental/IoStore/HttpClient/Private/Loop.inl b/Engine/Source/Runtime/Experimental/IoStore/HttpClient/Private/Loop.inl
--- a/Engine/Source/Runtime/Experimental/IoStore/HttpClient/Private/Loop.inl
+++ b/Engine/Source/Runtime/Experimental/IoStore/HttpClient/Private/Loop.inl
@@ -864,9 +864,10 @@ void FPeerGroup::RecvInternal(FTickState& State)
 
 	IsKeepAlive &= Activity->IsKeepAlive;
 	LastUseMs = State.NowMs;
+	bWaiting |= Outcome.IsWaiting();
 
 	// If we've only a small amount left to receive we can start more work
-	if (IsKeepAlive & (Recv->Next == nullptr))
+	if (IsKeepAlive & (Recv->Next == nullptr) & (Send == nullptr))
 	{
 		uint32 Remaining = Activity_RemainingKiB(Activity);
 		if (Remaining < uint32(GRecvWorkThresholdKiB))
@@ -888,7 +889,7 @@ void FPeerGroup::RecvInternal(FTickState& State)
 	}
 
 	// If there was no data available this is far as receiving can go
-	if (bWaiting = Outcome.IsWaiting(); bWaiting)
+	if (Outcome.IsWaiting())
 	{
 		return;
 	}
@@ -934,8 +935,7 @@ void FPeerGroup::SendInternal(FTickState& State)
 
 	if (Outcome.IsWaiting())
 	{
-		// For now we'll not add the socket as a waiter. It is unlikely that we
-		// send enough to need to wait currently.
+		bWaiting = true;
 		return;
 	}
 
