Shim system calls so they can more easily be redirected

#rnx
#ushell-unshelve of 39626023

diff --git a/Engine/Source/Runtime/Experimental/IoStore/HttpClient/Private/ConnectionPool.inl b/Engine/Source/Runtime/Experimental/IoStore/HttpClient/Private/ConnectionPool.inl
--- a/Engine/Source/Runtime/Experimental/IoStore/HttpClient/Private/ConnectionPool.inl
+++ b/Engine/Source/Runtime/Experimental/IoStore/HttpClient/Private/ConnectionPool.inl
@@ -81,13 +81,13 @@ FOutcome FHost::ResolveHostName()
 	IpAddresses[0] = 1;
 
 	addrinfo* Info = nullptr;
-	ON_SCOPE_EXIT { if (Info != nullptr) freeaddrinfo(Info); };
+	ON_SCOPE_EXIT { if (Info != nullptr) Socket_FreeAddrInfo(Info); };
 
 	addrinfo Hints = {};
 	Hints.ai_family = AF_INET;
 	Hints.ai_socktype = SOCK_STREAM;
 	Hints.ai_protocol = IPPROTO_TCP;
-	auto Result = getaddrinfo(HostName, nullptr, &Hints, &Info);
+	auto Result = Socket_GetAddrInfo(HostName, nullptr, &Hints, &Info);
 	if (uint32(Result) || Info == nullptr)
 	{
 		return FOutcome::Error("Error encountered resolving");
@@ -115,7 +115,7 @@ FOutcome FHost::ResolveHostName()
 			break;
 		}
 
-		IpAddresses[AddressCount] = htonl(IpAddress);
+		IpAddresses[AddressCount] = Socket_HtoNl(IpAddress);
 		if (++AddressCount >= UE_ARRAY_COUNT(IpAddresses))
 		{
 			break;
diff --git a/Engine/Source/Runtime/Experimental/IoStore/HttpClient/Private/Http.cpp b/Engine/Source/Runtime/Experimental/IoStore/HttpClient/Private/Http.cpp
--- a/Engine/Source/Runtime/Experimental/IoStore/HttpClient/Private/Http.cpp
+++ b/Engine/Source/Runtime/Experimental/IoStore/HttpClient/Private/Http.cpp
@@ -9,6 +9,7 @@
 #include <IO/IoBuffer.h>
 #include <Math/UnrealMathUtility.h>
 #include <Memory/MemoryView.h>
+#include <Misc/ByteSwap.h>
 #include <Misc/ScopeExit.h>
 #include <Misc/ScopeLock.h>
 #include <Misc/StringBuilder.h>
diff --git a/Engine/Source/Runtime/Experimental/IoStore/HttpClient/Private/Socket.inl b/Engine/Source/Runtime/Experimental/IoStore/HttpClient/Private/Socket.inl
--- a/Engine/Source/Runtime/Experimental/IoStore/HttpClient/Private/Socket.inl
+++ b/Engine/Source/Runtime/Experimental/IoStore/HttpClient/Private/Socket.inl
@@ -4,6 +4,12 @@
 
 // {{{1 platforms ..............................................................
 
+#define IAS_HTTP_SOCKET_FORWARD(Src, Dest)			\
+	template <typename... ArgTypes>					\
+	static auto Socket_##Src(ArgTypes&&... Args) {	\
+		return Dest(Forward<ArgTypes>(Args)...);	\
+	}
+
 #if PLATFORM_MICROSOFT
 #	if !defined(NO_UE_INCLUDES)
 #	include "Windows/AllowWindowsPlatformTypes.h"
@@ -59,6 +65,35 @@
 #	include "CoreHttp/Http.inl"
 #endif
 
+#if defined(IAS_HTTP_SOCKET_FORWARD)
+	namespace UE::IoStore::HTTP {
+		IAS_HTTP_SOCKET_FORWARD(Close,			closesocket)
+		IAS_HTTP_SOCKET_FORWARD(Connect,		connect)
+		IAS_HTTP_SOCKET_FORWARD(FreeAddrInfo,	freeaddrinfo)
+		IAS_HTTP_SOCKET_FORWARD(GetAddrInfo,	getaddrinfo)
+		IAS_HTTP_SOCKET_FORWARD(GetSockOpt,		getsockopt)
+		IAS_HTTP_SOCKET_FORWARD(Recv,			recv)
+		IAS_HTTP_SOCKET_FORWARD(Select,			select)
+		IAS_HTTP_SOCKET_FORWARD(Send,			send)
+		IAS_HTTP_SOCKET_FORWARD(SetSockOpt,		setsockopt)
+		IAS_HTTP_SOCKET_FORWARD(Shutdown,		shutdown)
+		IAS_HTTP_SOCKET_FORWARD(Socket,			socket)
+
+#		if defined(IAS_HTTP_USE_POLL)
+			IAS_HTTP_SOCKET_FORWARD(Poll,			poll)
+#		endif
+	} // namespace UE::IoStore::HTTP
+#endif // IAS_HTTP_SOCKET_FORWARD
+
+namespace UE::IoStore::HTTP {
+
+// hton* can sometimes be implemented as a macro which does not work with the
+// forward approach above. Hence we have our own here. Thanks for reading!
+uint32 Socket_HtoNl(uint32 v) { return ByteSwap(v); }
+uint16 Socket_HtoNs(uint16 v) { return ByteSwap(v); }
+
+} // namespace UE::IoStore::HTTP
+
 // }}}
 
 namespace UE::IoStore::HTTP
@@ -186,7 +221,7 @@ int32 FWaiter::Wait(TArrayView<FWaiter> Waiters, int32 TimeoutMs)
 
 	// Poll the sockets
 #if defined(IAS_HTTP_USE_POLL)
-	int32 Result = poll(Selects.GetData(), Selects.Num(), TimeoutMs);
+	int32 Result = Socket_Poll(Selects.GetData(), Selects.Num(), TimeoutMs);
 	if (Result <= 0)
 	{
 		return Result;
@@ -213,7 +248,7 @@ int32 FWaiter::Wait(TArrayView<FWaiter> Waiters, int32 TimeoutMs)
 		MaxFd = FMath::Max(Select.fd, MaxFd);
 	}
 
-	int32 Result = select(int32(MaxFd + 1), &FdSetRead, &FdSetWrite, &FdSetExcept, TimeValPtr);
+	int32 Result = Socket_Select(int32(MaxFd + 1), &FdSetRead, &FdSetWrite, &FdSetExcept, TimeValPtr);
 	if (Result <= 0)
 	{
 		return Result;
@@ -298,7 +333,7 @@ void FSocket::Move(FSocket&& Rhs)
 bool FSocket::Create()
 {
 	check(!IsValid());
-	Socket = socket(AF_INET, SOCK_STREAM, IPPROTO_TCP);
+	Socket = Socket_Socket(AF_INET, SOCK_STREAM, IPPROTO_TCP);
 
 	if (!IsValid())
 	{
@@ -322,7 +357,7 @@ void FSocket::Destroy()
 
 	Trace(Socket, ETrace::SocketDestroy);
 
-	closesocket(Socket);
+	Socket_Close(Socket);
 	Socket = InvalidSocket;
 }
 
@@ -333,14 +368,14 @@ FOutcome FSocket::Connect(uint32 IpAddress, uint32 Port)
 
 	Trace(Socket, ETrace::Connect, IpAddress);
 
-	IpAddress = htonl(IpAddress);
+	IpAddress = Socket_HtoNl(IpAddress);
 
 	sockaddr_in AddrInet = { sizeof(sockaddr_in) };
 	AddrInet.sin_family = AF_INET;
-	AddrInet.sin_port = htons(uint16(Port));
+	AddrInet.sin_port = Socket_HtoNs(uint16(Port));
 	memcpy(&(AddrInet.sin_addr), &IpAddress, sizeof(IpAddress));
 
-	int32 Result = connect(Socket, &(sockaddr&)AddrInet, sizeof(AddrInet));
+	int32 Result = Socket_Connect(Socket, &(sockaddr&)AddrInet, uint32(sizeof(AddrInet)));
 
 	if (IsSocketResult(EWOULDBLOCK) | IsSocketResult(EINPROGRESS))
 	{
@@ -359,14 +394,14 @@ FOutcome FSocket::Connect(uint32 IpAddress, uint32 Port)
 void FSocket::Disconnect()
 {
 	check(IsValid());
-	shutdown(Socket, SHUT_RDWR);
+	Socket_Shutdown(Socket, SHUT_RDWR);
 }
 
 ////////////////////////////////////////////////////////////////////////////////
 FOutcome FSocket::Send(const char* Data, uint32 Size)
 {
 	Trace(Socket, ETrace::Send, -1);
-	int32 Result = send(Socket, Data, Size, MsgFlagType(MSG_NOSIGNAL));
+	int32 Result = Socket_Send(Socket, Data, Size, MsgFlagType(MSG_NOSIGNAL));
 	Trace(Socket, ETrace::Send, FMath::Max(Result, 0));
 
 	if (Result > 0)						return FOutcome::Ok(Result);
@@ -377,7 +412,7 @@ FOutcome FSocket::Send(const char* Data, uint32 Size)
 	{
 		int32 Error = 0;
 		socklen_t ErrorSize = sizeof(Error);
-		Result = getsockopt(Socket, SOL_SOCKET, SO_ERROR, (char*)&Error, &ErrorSize);
+		Result = Socket_GetSockOpt(Socket, SOL_SOCKET, SO_ERROR, (char*)&Error, &ErrorSize);
 		if (Result < 0 || Error != 0)
 		{
 			Result = (Result < 0) ? LastSocketResult() : Error;
@@ -395,7 +430,7 @@ FOutcome FSocket::Send(const char* Data, uint32 Size)
 FOutcome FSocket::Recv(char* Dest, uint32 Size)
 {
 	Trace(Socket, ETrace::Recv, -1);
-	int32 Result = recv(Socket, Dest, Size, MsgFlagType(0));
+	int32 Result = Socket_Recv(Socket, Dest, Size, MsgFlagType(0));
 	Trace(Socket, ETrace::Recv, FMath::Max(0, Result));
 
 	if (Result > 0)						return FOutcome::Ok(Result);
@@ -431,13 +466,13 @@ bool FSocket::SetBlocking(bool bBlocking)
 ////////////////////////////////////////////////////////////////////////////////
 bool FSocket::SetSendBufSize(int32 Size)
 {
-	return 0 == setsockopt(Socket, SOL_SOCKET, SO_SNDBUF, &(char&)Size, sizeof(Size));
+	return 0 == Socket_SetSockOpt(Socket, SOL_SOCKET, SO_SNDBUF, &(char&)Size, int32(sizeof(Size)));
 }
 
 ////////////////////////////////////////////////////////////////////////////////
 bool FSocket::SetRecvBufSize(int32 Size)
 {
-	return 0 == setsockopt(Socket, SOL_SOCKET, SO_RCVBUF, &(char&)Size, sizeof(Size));
+	return 0 == Socket_SetSockOpt(Socket, SOL_SOCKET, SO_RCVBUF, &(char&)Size, int32(sizeof(Size)));
 }
 
 // }}}
