Remove data being sent to IAS analytics that we were not using.

#rnx

- IasCdnUrl: This stat was never that accurate and as we move towards supporting multiple host groups it will become meaningless anyway.
- IasHttpDurationMeanAvg/IasHttpDurationStdDev: Originally intended to try and track how stable people's connections are but we haven't really used these values since IAS went live.
- IasHttpDuration?: The idea with this stat was to try and get a rough idea on the distributed of people's request time but we haven't managed to find a useful way to display this data.

diff --git a/Engine/Source/Runtime/Experimental/IoStore/OnDemand/Private/OnDemandIoDispatcherBackend.cpp b/Engine/Source/Runtime/Experimental/IoStore/OnDemand/Private/OnDemandIoDispatcherBackend.cpp
--- a/Engine/Source/Runtime/Experimental/IoStore/OnDemand/Private/OnDemandIoDispatcherBackend.cpp
+++ b/Engine/Source/Runtime/Experimental/IoStore/OnDemand/Private/OnDemandIoDispatcherBackend.cpp
@@ -1569,15 +1569,6 @@ void FOnDemandIoBackend::ReportAnalytics(TArray<FAnalyticsEventAttribute>& OutAn
 
 	if (DefaultHostGroup.PrimaryHostIndex() != INDEX_NONE)
 	{
-		FString CdnUrl = FString(DefaultHostGroup.PrimaryHost());
-
-		if (!CdnUrl.RemoveFromStart(TEXT("http://")))
-		{
-			CdnUrl.RemoveFromStart(TEXT("https://"));
-		}
-
-		AppendAnalyticsEventAttributeArray(OutAnalyticsArray, TEXT("IasCdnUrl"), MoveTemp(CdnUrl));
-
 		Stats.ReportEndPointAnalytics(OutAnalyticsArray);
 	}
 }
diff --git a/Engine/Source/Runtime/Experimental/IoStore/OnDemand/Private/Statistics.cpp b/Engine/Source/Runtime/Experimental/IoStore/OnDemand/Private/Statistics.cpp
--- a/Engine/Source/Runtime/Experimental/IoStore/OnDemand/Private/Statistics.cpp
+++ b/Engine/Source/Runtime/Experimental/IoStore/OnDemand/Private/Statistics.cpp
@@ -54,137 +54,6 @@ static FAutoConsoleVariableRef CVar_DisplayOnScreenStatistics(
 static int32 BytesToApproxMB(uint64 Bytes) { return int32(Bytes >> 20); }
 static int32 BytesToApproxKB(uint64 Bytes) { return int32(Bytes >> 10); }
 
-/**
- * Code taken from SummarizeTraceCommandlet.cpp pending discussion on moving it
- * somewhere for general use.
- * Currently not thread safe!
- */
-class FIncrementalVariance
-{
-public:
-	FIncrementalVariance()
-		: Count(0)
-		, Mean(0.0)
-		, VarianceAccumulator(0.0)
-	{
-
-	}
-
-	uint64 GetCount() const
-	{
-		return Count;
-	}
-
-	double GetMean() const
-	{
-		return Mean;
-	}
-
-	/**
-	* Compute the variance given Welford's accumulator and the overall count
-	*
-	* @return The variance in sample units squared
-	*/
-	double GetVariance() const
-	{
-		double Result = 0.0;
-
-		if (Count > 1)
-		{
-			// Welford's final step, dependent on sample count
-			Result = VarianceAccumulator / double(Count - 1);
-		}
-
-		return Result;
-	}
-
-	/**
-	* Compute the standard deviation given Welford's accumulator and the overall count
-	*
-	* @return The standard deviation in sample units
-	*/
-	double GetDeviation() const
-	{
-		double Result = 0.0;
-
-		if (Count > 1)
-		{
-			// Welford's final step, dependent on sample count
-			double DeviationSqrd = VarianceAccumulator / double(Count - 1);
-
-			// stddev is sqrt of variance, to restore to units (vs. units squared)
-			Result = sqrt(DeviationSqrd);
-		}
-
-		return Result;
-	}
-
-	/**
-	* Perform an increment of work for Welford's variance, from which we can compute variation and standard deviation
-	*
-	* @param InSample	The new sample value to operate on
-	*/
-	void Increment(const double InSample)
-	{
-		Count++;
-		const double OldMean = Mean;
-		Mean += ((InSample - Mean) / double(Count));
-		VarianceAccumulator += ((InSample - Mean) * (InSample - OldMean));
-	}
-
-	/**
-	* Merge with another IncrementalVariance series in progress
-	*
-	* @param Other	The other variance incremented from another mutually exclusive population of analogous data.
-	*/
-	void Merge(const FIncrementalVariance& Other)
-	{
-		// empty other, nothing to do
-		if (Other.Count == 0)
-		{
-			return;
-		}
-
-		// empty this, just copy other
-		if (Count == 0)
-		{
-			Count = Other.Count;
-			Mean = Other.Mean;
-			VarianceAccumulator = Other.VarianceAccumulator;
-			return;
-		}
-
-		const double TotalPopulation = static_cast<double>(Count + Other.Count);
-		const double MeanDifference = Mean - Other.Mean;
-		const double A = (double(Count - 1) * GetVariance()) + (double(Other.Count - 1) * Other.GetVariance());
-		const double B = (MeanDifference) * (MeanDifference) * (double(Count) * double(Other.Count) / TotalPopulation);
-		const double MergedVariance = (A + B) / (TotalPopulation - 1);
-
-		const uint64 NewCount = Count + Other.Count;
-		const double NewMean = ((Mean * double(Count)) + (Other.Mean * double(Other.Count))) / double(NewCount);
-		const double NewVarianceAccumulator = MergedVariance * double(NewCount - 1);
-
-		Count = NewCount;
-		Mean = NewMean;
-		VarianceAccumulator = NewVarianceAccumulator;
-	}
-
-	/**
-	* Reset state back to initialized.
-	*/
-	void Reset()
-	{
-		Count = 0;
-		Mean = 0.0;
-		VarianceAccumulator = 0.0;
-	}
-
-private:
-	uint64 Count;
-	double Mean;
-	double VarianceAccumulator;
-};
-
 class FDeltaTracking
 {
 public:
@@ -294,53 +163,6 @@ int64					GHttpHistoryBytes[GHttpHistoryCount] = {};
 int64					GHttpHistoryTotalDuration = 0;
 int64					GHttpHistoryTotalBytes = 0;
 int64 					GHttpHistoryIndex = 0;
-FIncrementalVariance	GHttpAvgDuration; // Duration of the http requests, in milliseconds
-
-////////////////////////////////////////////////////////////////////////////////
-// 
-
-static const uint64 DurationBuckets[] = { 30, 150, 400, 1000 }; // Time boundaries (in ms) for each bucket
-
-// P1907R1 is not yet implemented in some of our compilers so FTrackedValue will not work on subobjects
-// so we need each bucket to be it's own variable rather than an array.
-uint32 GHttpDurationBucket0 = 0;
-uint32 GHttpDurationBucket1 = 0;
-uint32 GHttpDurationBucket2 = 0;
-uint32 GHttpDurationBucket3 = 0;
-uint32 GHttpDurationBucket4 = 0;
-
-uint32* GHttpDurationBuckets[] =
-{
-	&GHttpDurationBucket0,
-	&GHttpDurationBucket1,
-	&GHttpDurationBucket2,
-	&GHttpDurationBucket3,
-	&GHttpDurationBucket4
-};
-
-void IncrementDurationBucket(uint64 DurationMs)
-{
-	const uint32 NumBoundries = UE_ARRAY_COUNT(DurationBuckets);
-	for (int32 Index = 0; Index < NumBoundries; ++Index)
-	{
-		if (DurationMs <= DurationBuckets[Index])
-		{
-			(*GHttpDurationBuckets[Index])++;
-			return;
-		}
-	}
-
-	(*GHttpDurationBuckets[NumBoundries])++;
-}
-
-void ResetDurationBuckets()
-{
-	const uint32 NumBuckets = UE_ARRAY_COUNT(DurationBuckets) + 1;
-	for (int32 Index = 0; Index < NumBuckets; ++Index)
-	{
-		(*GHttpDurationBuckets[Index]) = 0;
-	}
-}
 
 ////////////////////////////////////////////////////////////////////////////////
 // CSV STATS
@@ -570,19 +392,7 @@ void FOnDemandIoBackendStats::ReportEndPointAnalytics(TArray<FAnalyticsEventAttr
 			,TRACK_DELTA("IasHttpRetryCount", GHttpRetryCount.Get())
 			,TRACK_DELTA("IasHttpGetCount", GHttpGetCount.Get())
 			,TRACK_DELTA("IasHttpDownloadedBytes", GHttpDownloadedBytes.Get())
-			,TEXT("IasHttpDurationMeanAvg"), GHttpAvgDuration.GetMean()
-			,TEXT("IasHttpDurationStdDev"), GHttpAvgDuration.GetDeviation()
-
-			,TEXT("IasHttpDuration0"), GHttpDurationBucket0
-			,TEXT("IasHttpDuration1"), GHttpDurationBucket1
-			,TEXT("IasHttpDuration2"), GHttpDurationBucket2
-			,TEXT("IasHttpDuration3"), GHttpDurationBucket3
-			,TEXT("IasHttpDuration4"), GHttpDurationBucket4
 		);
-
-		// These values we can just reset as they are only being used with analytics
-		GHttpAvgDuration.Reset();
-		ResetDurationBuckets();
 	}
 
 	if (GIasReportCacheAnalyticsEnabled)
@@ -643,12 +453,6 @@ private:
 		Http.GetCount.Stop();
 		Http.DownloadedBytes.Stop();
 
-		Http.DurationBuckets0.Stop();
-		Http.DurationBuckets1.Stop();
-		Http.DurationBuckets2.Stop();
-		Http.DurationBuckets3.Stop();
-		Http.DurationBuckets4.Stop();
-
 		Cache.ErrorCount.Stop();
 		Cache.DecodeErrorCount.Stop();
 		Cache.GetCount.Stop();
@@ -663,23 +467,6 @@ private:
 
 	void Report(TArray<FAnalyticsEventAttribute>& OutAnalyticsArray) const override
 	{
-		// Todo: For now we assume that the first host group is our 'valid' connection. This will need updating
-		if (!HostGroups.IsEmpty() && HostGroups[0].PrimaryHostIndex() != INDEX_NONE)
-		{
-			FString CdnUrl = FString(HostGroups[0].PrimaryHost());
-
-			if (!CdnUrl.IsEmpty())
-			{
-				// Strip the prefix from the url as some analytics systems may have trouble dealing with it
-				if (!CdnUrl.RemoveFromStart(TEXT("http://")))
-				{
-					CdnUrl.RemoveFromStart(TEXT("https://"));
-				}
-
-				AppendAnalyticsEventAttributeArray(OutAnalyticsArray, TEXT("IasCdnUrl"), MoveTemp(CdnUrl));
-			}
-		}
-
 		// Report if the system has ever actually managed to make a connection
 		AppendAnalyticsEventAttributeArray(OutAnalyticsArray, TEXT("IasHttpHasEverConnected"), GHttpConnectCount.Get() > 0);
 
@@ -691,14 +478,6 @@ private:
 				, TEXT("IasHttpRetryCount"), Http.RetryCount.GetValue()
 				, TEXT("IasHttpGetCount"), Http.GetCount.GetValue()
 				, TEXT("IasHttpDownloadedBytes"), Http.DownloadedBytes.GetValue()
-				, TEXT("IasHttpDurationMeanAvg"), GHttpAvgDuration.GetMean()		// TODO: Reconsider
-				, TEXT("IasHttpDurationStdDev"), GHttpAvgDuration.GetDeviation()	// TODO: Reconsider
-
-				, TEXT("IasHttpDuration0"), Http.DurationBuckets0.GetValue()
-				, TEXT("IasHttpDuration1"), Http.DurationBuckets1.GetValue()
-				, TEXT("IasHttpDuration2"), Http.DurationBuckets2.GetValue()
-				, TEXT("IasHttpDuration3"), Http.DurationBuckets3.GetValue()
-				, TEXT("IasHttpDuration4"), Http.DurationBuckets4.GetValue()
 			);
 		}
 
@@ -787,12 +566,6 @@ private:
 		FTrackedValue<FCounterInt, GHttpRetryCount>				RetryCount;
 		FTrackedValue<FCounterInt, GHttpGetCount>				GetCount;
 		FTrackedValue<FCounterInt, GHttpDownloadedBytes>		DownloadedBytes;
-
-		FTrackedValue<uint32, GHttpDurationBucket0>				DurationBuckets0;
-		FTrackedValue<uint32, GHttpDurationBucket1>				DurationBuckets1;
-		FTrackedValue<uint32, GHttpDurationBucket2>				DurationBuckets2;
-		FTrackedValue<uint32, GHttpDurationBucket3>				DurationBuckets3;
-		FTrackedValue<uint32, GHttpDurationBucket4>				DurationBuckets4;
 	} Http;
 
 	/** Cache Stats */
@@ -956,10 +729,6 @@ void FOnDemandIoBackendStats::OnHttpGet(uint64 SizeBytes, uint64 DurationMs)
 	}
 
 	GHttpHistoryIndex = (GHttpHistoryIndex + 1) % GHttpHistoryCount;
-
-	GHttpAvgDuration.Increment(static_cast<double>(DurationMs));
-
-	IncrementDurationBucket(DurationMs);
 }
 
 void FOnDemandIoBackendStats::OnHttpCancel()
