Dataflow : Add support for drag and dropping a pin onto a dataflow node
- implement it for subgraph input and output node
- also added support for structure types ( support managed array collection drop )

diff --git a/Engine/Plugins/Experimental/Dataflow/Source/DataflowEditor/Private/Dataflow/DataflowSchema.cpp b/Engine/Plugins/Experimental/Dataflow/Source/DataflowEditor/Private/Dataflow/DataflowSchema.cpp
--- a/Engine/Plugins/Experimental/Dataflow/Source/DataflowEditor/Private/Dataflow/DataflowSchema.cpp
+++ b/Engine/Plugins/Experimental/Dataflow/Source/DataflowEditor/Private/Dataflow/DataflowSchema.cpp
@@ -342,6 +342,84 @@ const FPinConnectionResponse UDataflowSchema::CanCreateConnection(const UEdGraph
 	return FPinConnectionResponse(CONNECT_RESPONSE_DISALLOW, NoConnectionResponse[FMath::RandRange(0, NoConnectionResponse.Num()-1)]);
 }
 
+TSharedPtr<FDataflowNode> UDataflowSchema::GetDataflowNodeFromEdNode(UEdGraphNode* EdNode)
+{
+	if (UDataflowEdNode* DataflowEdNode = Cast<UDataflowEdNode>(EdNode))
+	{
+		return DataflowEdNode->GetDataflowNode();
+	}
+	return {};
+}
+
+FDataflowConnection* UDataflowSchema::GetConnectionFromPin(const UEdGraphPin* Pin)
+{
+	FDataflowConnection* DataflowConnection = nullptr;
+	if (Pin)
+	{
+		if (TSharedPtr<FDataflowNode> DataflowNode = GetDataflowNodeFromEdNode(Pin->GetOwningNode()))
+		{
+			if (Pin->Direction == EEdGraphPinDirection::EGPD_Input)
+			{
+				DataflowConnection = DataflowNode->FindInput(Pin->PinName);
+			}
+			else if (Pin->Direction == EEdGraphPinDirection::EGPD_Output)
+			{
+				DataflowConnection = DataflowNode->FindOutput(Pin->PinName);
+			}
+		}
+	}
+	return DataflowConnection;
+}
+
+UE::Dataflow::FPin::EDirection UDataflowSchema::GetDirectionFromPinDirection(EEdGraphPinDirection InPinDirection)
+{
+	switch (InPinDirection)
+	{
+	case EEdGraphPinDirection::EGPD_Input:
+		return UE::Dataflow::FPin::EDirection::INPUT;
+	case EEdGraphPinDirection::EGPD_Output:
+		return UE::Dataflow::FPin::EDirection::OUTPUT;
+	default:
+		return UE::Dataflow::FPin::EDirection::NONE;
+	}
+}
+
+bool UDataflowSchema::SupportsDropPinOnNode(UEdGraphNode* InTargetNode, const FEdGraphPinType& InSourcePinType, EEdGraphPinDirection InSourcePinDirection, FText& OutMessage) const
+{
+	if (TSharedPtr<FDataflowNode> TargetDataflowNode = GetDataflowNodeFromEdNode(InTargetNode))
+	{
+		const UE::Dataflow::FPin::EDirection SourceDirection = GetDirectionFromPinDirection(InSourcePinDirection);
+		if (TargetDataflowNode->SupportsDropConnectionOnNode(InSourcePinType.PinCategory, SourceDirection))
+		{
+			OutMessage = LOCTEXT("DataflowSchema_DropPinOnNodeSupported", "Add pin to this node");
+			return true;
+		}
+	}
+	OutMessage = LOCTEXT("DataflowSchema_NoDropPinOnNodeSupport", "This node does not support this pin type");
+	return false;
+}
+
+UEdGraphPin* UDataflowSchema::DropPinOnNode(UEdGraphNode* InTargetNode, const FName& InSourcePinName, const FEdGraphPinType& InSourcePinType, EEdGraphPinDirection InSourcePinDirection) const
+{
+	UEdGraphPin* ResultPin = nullptr;
+
+	if (TSharedPtr<FDataflowNode> TargetDataflowNode = GetDataflowNodeFromEdNode(InTargetNode))
+	{
+		if (const FDataflowConnection* SourceConnection = GetConnectionFromPin(PinBeingDropped))
+		{
+			if (const FDataflowConnection* NewConnection = TargetDataflowNode->OnDropConnectionOnNode(*SourceConnection))
+			{
+				if (UDataflowEdNode* TargetDataflowEdNode = Cast<UDataflowEdNode>(InTargetNode))
+				{
+					TargetDataflowEdNode->UpdatePinsFromDataflowNode();
+					ResultPin = TargetDataflowEdNode->FindPin(NewConnection->GetName());
+				}
+			}
+		}
+	}
+	return ResultPin;
+}
+
 FLinearColor UDataflowSchema::GetPinTypeColor(const FEdGraphPinType& PinType) const
 {
 	return GetTypeColor(PinType.PinCategory);
diff --git a/Engine/Plugins/Experimental/Dataflow/Source/DataflowEditor/Public/Dataflow/DataflowSchema.h b/Engine/Plugins/Experimental/Dataflow/Source/DataflowEditor/Public/Dataflow/DataflowSchema.h
--- a/Engine/Plugins/Experimental/Dataflow/Source/DataflowEditor/Public/Dataflow/DataflowSchema.h
+++ b/Engine/Plugins/Experimental/Dataflow/Source/DataflowEditor/Public/Dataflow/DataflowSchema.h
@@ -32,6 +32,10 @@ public:
 	virtual void OnPinConnectionDoubleCicked(UEdGraphPin* PinA, UEdGraphPin* PinB, const FVector2D& GraphPosition) const override;
 	virtual void BreakPinLinks(UEdGraphPin& TargetPin, bool bSendsNodeNotifcation) const override;
 	virtual bool TryCreateConnection(UEdGraphPin* PinA, UEdGraphPin* PinB) const override;
+
+	virtual void SetPinBeingDroppedOnNode(UEdGraphPin* InSourcePin) const override { PinBeingDropped = InSourcePin; }
+	virtual bool SupportsDropPinOnNode(UEdGraphNode* InTargetNode, const FEdGraphPinType& InSourcePinType, EEdGraphPinDirection InSourcePinDirection, FText& OutMessage) const override;
+	virtual UEdGraphPin* DropPinOnNode(UEdGraphNode* InTargetNode, const FName& InSourcePinName, const FEdGraphPinType& InSourcePinType, EEdGraphPinDirection InSourcePinDirection) const override;
 	//~ End UEdGraphSchema Interface
 
 	static FLinearColor GetTypeColor(const FName& Type);
@@ -41,6 +45,13 @@ public:
 private:
 	static bool CanPinBeConnectedToNode(const UEdGraphPin* Pin, const UE::Dataflow::FFactoryParameters& NodeParameters);
 	static bool CanConnectPins(const UEdGraphPin& OutputPin, const UEdGraphPin& InputPin);
+	
+	static TSharedPtr<FDataflowNode> GetDataflowNodeFromEdNode(UEdGraphNode* EdNode);
+	static FDataflowConnection* GetConnectionFromPin(const UEdGraphPin* Pin);
+	static UE::Dataflow::FPin::EDirection GetDirectionFromPinDirection(EEdGraphPinDirection InPinDirection);
+
+	// used by SupportsDropPinOnNode to know which original connection is being dropped
+	mutable UEdGraphPin* PinBeingDropped = nullptr;
 };
 
 class FDataflowConnectionDrawingPolicy : public FConnectionDrawingPolicy, public FGCObject
diff --git a/Engine/Source/Runtime/Experimental/Dataflow/Core/Public/Dataflow/DataflowNode.h b/Engine/Source/Runtime/Experimental/Dataflow/Core/Public/Dataflow/DataflowNode.h
--- a/Engine/Source/Runtime/Experimental/Dataflow/Core/Public/Dataflow/DataflowNode.h
+++ b/Engine/Source/Runtime/Experimental/Dataflow/Core/Public/Dataflow/DataflowNode.h
@@ -629,6 +629,13 @@ PRAGMA_ENABLE_DEPRECATION_WARNINGS
 	// Called when UPROPERTY members of the dataflow node have been changed in the editor
 	DATAFLOWCORE_API virtual void OnPropertyChanged(UE::Dataflow::FContext& Context, const FPropertyChangedEvent& InPropertyChangedEvent) {}
 
+	// check if the node support dropped connections
+	DATAFLOWCORE_API virtual bool SupportsDropConnectionOnNode(FName TypeName, UE::Dataflow::FPin::EDirection Direction) const { return false; }
+
+	// called when a dragged pin is dropped on the node 
+	// if the node creates a connection, returns its pointer
+	DATAFLOWCORE_API virtual const FDataflowConnection* OnDropConnectionOnNode(const FDataflowConnection& DroppedConnection) { return nullptr; }
+
 	DATAFLOWCORE_API virtual bool ShouldInvalidateOnPropertyChanged(const FPropertyChangedEvent& InPropertyChangedEvent) const { return true; };
 
 	bool HasValidConnections() const { return bHasValidConnections; }
diff --git a/Engine/Source/Runtime/Experimental/Dataflow/Engine/Private/Dataflow/DataflowSubGraphNodes.cpp b/Engine/Source/Runtime/Experimental/Dataflow/Engine/Private/Dataflow/DataflowSubGraphNodes.cpp
--- a/Engine/Source/Runtime/Experimental/Dataflow/Engine/Private/Dataflow/DataflowSubGraphNodes.cpp
+++ b/Engine/Source/Runtime/Experimental/Dataflow/Engine/Private/Dataflow/DataflowSubGraphNodes.cpp
@@ -244,11 +244,11 @@ bool FDataflowDynamicConnections::IsSupportedType(const FPropertyBagPropertyDesc
 	case EPropertyBagPropertyType::String:
 	case EPropertyBagPropertyType::Text:
 	case EPropertyBagPropertyType::Object:
+	case EPropertyBagPropertyType::Struct:
 		return true;
 
 		// unsupported type for now 
 	case EPropertyBagPropertyType::Enum:
-	case EPropertyBagPropertyType::Struct:
 	case EPropertyBagPropertyType::SoftObject:
 	case EPropertyBagPropertyType::Class:
 	case EPropertyBagPropertyType::SoftClass:
@@ -285,11 +285,11 @@ bool FDataflowDynamicConnections::SetConnectionTypeFromPropertyDesc(FDataflowCon
 		case EPropertyBagPropertyType::String:
 		case EPropertyBagPropertyType::Text:
 		case EPropertyBagPropertyType::Object:
+		case EPropertyBagPropertyType::Struct:
 			return Connection.SetConcreteType(ConcreteType);
 
 		// unsupported type for now 
 		case EPropertyBagPropertyType::Enum:
-		case EPropertyBagPropertyType::Struct:
 		case EPropertyBagPropertyType::SoftObject:
 		case EPropertyBagPropertyType::Class:
 		case EPropertyBagPropertyType::SoftClass:
@@ -387,6 +387,23 @@ void FDataflowSubGraphInputNode::OnPropertyChanged(UE::Dataflow::FContext& Conte
 	DynamicConnections.Refresh();
 }
 
+bool FDataflowSubGraphInputNode::SupportsDropConnectionOnNode(FName TypeName, UE::Dataflow::FPin::EDirection Direction) const
+{ 
+	return (Direction == UE::Dataflow::FPin::EDirection::INPUT);
+}
+
+const FDataflowConnection* FDataflowSubGraphInputNode::OnDropConnectionOnNode(const FDataflowConnection& DroppedConnection)
+{
+	if (SupportsDropConnectionOnNode(DroppedConnection.GetType(), DroppedConnection.GetDirection()))
+	{
+		const FPropertyBagPropertyDesc Desc(DroppedConnection.GetName(), DroppedConnection.GetProperty());
+		PropertyBag.AddProperties({ Desc });
+		DynamicConnections.Refresh();
+		return FindOutput(Desc.Name);
+	}
+	return nullptr;
+}
+
 //////////////////////////////////////////////////////////////////////////////////////////////////////////
 
 FDataflowSubGraphOutputNode::FDataflowSubGraphOutputNode(const UE::Dataflow::FNodeParameters& InParam, FGuid InGuid)
@@ -415,6 +432,24 @@ void FDataflowSubGraphOutputNode::OnPropertyChanged(UE::Dataflow::FContext& Cont
 	DynamicConnections.Refresh();
 }
 
+bool FDataflowSubGraphOutputNode::SupportsDropConnectionOnNode(FName TypeName, UE::Dataflow::FPin::EDirection Direction) const
+{
+	return (Direction == UE::Dataflow::FPin::EDirection::OUTPUT);
+}
+
+
+const FDataflowConnection* FDataflowSubGraphOutputNode::OnDropConnectionOnNode(const FDataflowConnection& DroppedConnection)
+{
+	if (SupportsDropConnectionOnNode(DroppedConnection.GetType(), DroppedConnection.GetDirection()))
+	{
+		const FPropertyBagPropertyDesc Desc(DroppedConnection.GetName(), DroppedConnection.GetProperty());
+		PropertyBag.AddProperties({ Desc });
+		DynamicConnections.Refresh();
+		return FindInput(Desc.Name);
+	}
+	return nullptr;
+}
+
 //////////////////////////////////////////////////////////////////////////////////////////////////////////
 
 #define NEW_METHOD 1
@@ -437,13 +472,16 @@ FDataflowCallSubGraphNode::~FDataflowCallSubGraphNode()
 
 void FDataflowCallSubGraphNode::RegisterHandlers(UDataflowSubGraph* SubGraph)
 {
-	if (SubGraph)
+	if (IsValid())
 	{
-		OnSubGraphLoadedHandle = SubGraph->AddOnGraphChangedHandler(FOnGraphChanged::FDelegate::CreateRaw(this, &FDataflowCallSubGraphNode::OnSubGraphChanged));
-	}
+		if (SubGraph)
+		{
+			OnSubGraphLoadedHandle = SubGraph->AddOnGraphChangedHandler(FOnGraphChanged::FDelegate::CreateRaw(this, &FDataflowCallSubGraphNode::OnSubGraphChanged));
+		}
 
-	FDataflowSubGraphDelegates::OnSubGraphLoaded.AddRaw(this, &FDataflowCallSubGraphNode::OnSubGraphLoaded);
-	FDataflowAssetDelegates::OnSubGraphsChanged.AddRaw(this, &FDataflowCallSubGraphNode::OnSomeSubGraphsChanged);
+		FDataflowSubGraphDelegates::OnSubGraphLoaded.AddRaw(this, &FDataflowCallSubGraphNode::OnSubGraphLoaded);
+		FDataflowAssetDelegates::OnSubGraphsChanged.AddRaw(this, &FDataflowCallSubGraphNode::OnSomeSubGraphsChanged);
+	}
 }
 
 void FDataflowCallSubGraphNode::UnregisterHandlers(UDataflowSubGraph* SubGraph)
@@ -480,6 +518,11 @@ void FDataflowCallSubGraphNode::SetSubGraphGuid(const FGuid& InSubGraphGuid)
 	}
 }
 
+bool FDataflowCallSubGraphNode::IsValid() const
+{
+	return SubGraphGuid.IsValid();
+}
+
 void FDataflowCallSubGraphNode::OnSubGraphChanged(const struct FEdGraphEditAction& InAction)
 {
 	bool bInputsChanged = false;
diff --git a/Engine/Source/Runtime/Experimental/Dataflow/Engine/Public/Dataflow/DataflowSubGraphNodes.h b/Engine/Source/Runtime/Experimental/Dataflow/Engine/Public/Dataflow/DataflowSubGraphNodes.h
--- a/Engine/Source/Runtime/Experimental/Dataflow/Engine/Public/Dataflow/DataflowSubGraphNodes.h
+++ b/Engine/Source/Runtime/Experimental/Dataflow/Engine/Public/Dataflow/DataflowSubGraphNodes.h
@@ -90,6 +90,9 @@ private:
 	virtual void PostSerialize(const FArchive& Ar) override;
 	virtual void OnPropertyChanged(UE::Dataflow::FContext& Context, const FPropertyChangedEvent& InPropertyChangedEvent) override;
 
+	virtual bool SupportsDropConnectionOnNode(FName TypeName, UE::Dataflow::FPin::EDirection Direction) const override;
+	virtual const FDataflowConnection* OnDropConnectionOnNode(const FDataflowConnection& DroppedConnection) override;
+
 	UPROPERTY()
 	FDataflowDynamicConnections DynamicConnections;
 
@@ -124,6 +127,9 @@ private:
 	virtual void PostSerialize(const FArchive& Ar) override;
 	virtual void OnPropertyChanged(UE::Dataflow::FContext& Context, const FPropertyChangedEvent& InPropertyChangedEvent) override;
 
+	virtual bool SupportsDropConnectionOnNode(FName TypeName, UE::Dataflow::FPin::EDirection Direction) const override;
+	virtual const FDataflowConnection* OnDropConnectionOnNode(const FDataflowConnection& DroppedConnection) override;
+
 	UPROPERTY()
 	FDataflowDynamicConnections DynamicConnections;
 
@@ -171,6 +177,8 @@ private:
 	void SyncInputsPropertyBagWithSubGraph();
 	void SyncOutputsPropertyBagWithSubGraph();
 
+	bool IsValid() const;
+
 	UPROPERTY()
 	FGuid SubGraphGuid;
 
