Missed a function prefixtion.

#rnx

diff --git a/Engine/Source/Runtime/Experimental/IoStore/HttpClient/Private/Socket.inl b/Engine/Source/Runtime/Experimental/IoStore/HttpClient/Private/Socket.inl
--- a/Engine/Source/Runtime/Experimental/IoStore/HttpClient/Private/Socket.inl
+++ b/Engine/Source/Runtime/Experimental/IoStore/HttpClient/Private/Socket.inl
@@ -341,7 +341,7 @@ bool FSocket::Create()
 	}
 
 	int32 Yes = 1;
-	setsockopt(Socket, IPPROTO_TCP, TCP_NODELAY, &(char&)(Yes), sizeof(Yes));
+	Socket_SetSockOpt(Socket, IPPROTO_TCP, TCP_NODELAY, &(char&)(Yes), int32(sizeof(Yes)));
 
 	Trace(Socket, ETrace::SocketCreate);
 	return true;
