Dataflow : fix crash happening when undoing a graph that contains a subgraph call node
- issue was because of the use the wrong delegate handle to register a handler

diff --git a/Engine/Source/Runtime/Experimental/Dataflow/Engine/Private/Dataflow/DataflowSubGraphNodes.cpp b/Engine/Source/Runtime/Experimental/Dataflow/Engine/Private/Dataflow/DataflowSubGraphNodes.cpp
--- a/Engine/Source/Runtime/Experimental/Dataflow/Engine/Private/Dataflow/DataflowSubGraphNodes.cpp
+++ b/Engine/Source/Runtime/Experimental/Dataflow/Engine/Private/Dataflow/DataflowSubGraphNodes.cpp
@@ -515,7 +515,7 @@ void FDataflowCallSubGraphNode::RegisterHandlers(UDataflowSubGraph* SubGraph)
 	{
 		if (SubGraph)
 		{
-			OnSubGraphLoadedHandle = SubGraph->AddOnGraphChangedHandler(FOnGraphChanged::FDelegate::CreateRaw(this, &FDataflowCallSubGraphNode::OnSubGraphChanged));
+			OnGraphChangedHandle = SubGraph->AddOnGraphChangedHandler(FOnGraphChanged::FDelegate::CreateRaw(this, &FDataflowCallSubGraphNode::OnSubGraphChanged));
 		}
 
 		FDataflowSubGraphDelegates::OnSubGraphLoaded.AddRaw(this, &FDataflowCallSubGraphNode::OnSubGraphLoaded);
diff --git a/Engine/Source/Runtime/Experimental/Dataflow/Engine/Public/Dataflow/DataflowSubGraphNodes.h b/Engine/Source/Runtime/Experimental/Dataflow/Engine/Public/Dataflow/DataflowSubGraphNodes.h
--- a/Engine/Source/Runtime/Experimental/Dataflow/Engine/Public/Dataflow/DataflowSubGraphNodes.h
+++ b/Engine/Source/Runtime/Experimental/Dataflow/Engine/Public/Dataflow/DataflowSubGraphNodes.h
@@ -196,7 +196,6 @@ private:
 
 	TWeakObjectPtr<UDataflow> DataflowAssetWeakPtr;
 	FDelegateHandle OnGraphChangedHandle;
-	FDelegateHandle OnSubGraphLoadedHandle;
 };
 
 namespace UE::Dataflow
