Dataflow: When pasting variables in a separate dataflow , make sure to create the matching variable if does not exists yet
- also fix evaluation of variable node not properly working in non bound dataflow assets

diff --git a/Engine/Plugins/Experimental/Dataflow/Source/DataflowEditor/Private/Dataflow/DataflowAssetEditUtils.cpp b/Engine/Plugins/Experimental/Dataflow/Source/DataflowEditor/Private/Dataflow/DataflowAssetEditUtils.cpp
--- a/Engine/Plugins/Experimental/Dataflow/Source/DataflowEditor/Private/Dataflow/DataflowAssetEditUtils.cpp
+++ b/Engine/Plugins/Experimental/Dataflow/Source/DataflowEditor/Private/Dataflow/DataflowAssetEditUtils.cpp
@@ -715,6 +715,13 @@ namespace UE::Dataflow
 					{
 						DataflowNode->TypedScriptStruct()->ImportText(*NodeData.Properties, DataflowNode.Get(), nullptr, EPropertyPortFlags::PPF_None, nullptr, DataflowNode->TypedScriptStruct()->GetName(), true);
 					}
+
+					// if we are pasting a dataflow variable, let's create the variable if needed
+					if (FGetDataflowVariableNode* VariableNode = DataflowNode->AsType<FGetDataflowVariableNode>())
+					{
+						VariableNode->TryAddVariableToDataflowAsset(DataflowAsset);
+					}
+
 					// Do any post-import fixup.
 					FArchive Ar;
 					Ar.SetIsLoading(true);
@@ -782,7 +789,8 @@ namespace UE::Dataflow
 			return EChangeResult::Changed;
 		};
 
-		ChangeDataflowAssetWithTransaction(DataflowAsset, TransactionName, PasteNodesInternal, NAME_None);
+		// make sure we notify that variables may have changed as a property since we may have pasted variable nodes
+		ChangeDataflowAssetWithTransaction(DataflowAsset, TransactionName, PasteNodesInternal, GET_MEMBER_NAME_CHECKED(UDataflow, Variables));
 	}
 
 	//////////////////////////////////////////////////////////////////////////////////////////////////////////
diff --git a/Engine/Source/Runtime/Experimental/Dataflow/Engine/Private/Dataflow/DataflowVariableNodes.cpp b/Engine/Source/Runtime/Experimental/Dataflow/Engine/Private/Dataflow/DataflowVariableNodes.cpp
--- a/Engine/Source/Runtime/Experimental/Dataflow/Engine/Private/Dataflow/DataflowVariableNodes.cpp
+++ b/Engine/Source/Runtime/Experimental/Dataflow/Engine/Private/Dataflow/DataflowVariableNodes.cpp
@@ -79,6 +79,21 @@ void FGetDataflowVariableNode::OnObjectPropertyChanged(UObject* InObject, FPrope
 	}
 }
 
+bool FGetDataflowVariableNode::TryAddVariableToDataflowAsset(UDataflow& DataflowAsset)
+{
+	const FPropertyBagPropertyDesc* VariableDesc = DataflowAsset.Variables.FindPropertyDescByName(VariableName);
+	if (VariableDesc == nullptr)
+	{
+		if (const FPropertyBagPropertyDesc* Desc = VariablePropertyBag.FindPropertyDescByName(VariableName))
+		{
+			DataflowAsset.Variables.AddProperties({ *Desc });
+			DataflowAsset.Variables.SetValue(VariableName, Desc->CachedProperty, VariablePropertyBag.GetValue().GetMemory());
+			return true;
+		}
+	}
+	return false;
+}
+
 void FGetDataflowVariableNode::SetVariable(UDataflow* Dataflow, FName InVariableName)
 {
 	UnregisterHandlers();
@@ -237,18 +252,21 @@ void FGetDataflowVariableNode::Evaluate(UE::Dataflow::FContext& Context, const F
 					if (DataflowInstance.GetVariableOverrides().IsVariableOverridden(VariableName))
 					{
 						EvaluateFunction(*this, DataflowInstance.GetVariableOverrides().GetVariables(), Context, *Out);
+						return;
 					}
-					else if (const UDataflow* DataflowAsset = DataflowInstance.GetDataflowAsset())
+				}
+				// no override or engine context , use the default defined in the dataflow asset
+				if (TStrongObjectPtr<UDataflow> DataflowAsset = WeakDataflowPtr.Pin())
+				{
+					if (DataflowAsset->Variables.FindPropertyDescByName(VariableName) != nullptr)
 					{
-						if (DataflowAsset->Variables.FindPropertyDescByName(VariableName) != nullptr)
-						{
-							EvaluateFunction(*this, DataflowAsset->Variables, Context, *Out);
-						}
+						EvaluateFunction(*this, DataflowAsset->Variables, Context, *Out);
+						return;
 					}
-					return;
 				}
 			}
 		}
+		// nothing worked, just write a null value that will read as default value
 		Out->SetNullValue(Context);
 	}
 }
diff --git a/Engine/Source/Runtime/Experimental/Dataflow/Engine/Public/Dataflow/DataflowVariableNodes.h b/Engine/Source/Runtime/Experimental/Dataflow/Engine/Public/Dataflow/DataflowVariableNodes.h
--- a/Engine/Source/Runtime/Experimental/Dataflow/Engine/Public/Dataflow/DataflowVariableNodes.h
+++ b/Engine/Source/Runtime/Experimental/Dataflow/Engine/Public/Dataflow/DataflowVariableNodes.h
@@ -22,6 +22,7 @@ public:
 	DATAFLOWENGINE_API ~FGetDataflowVariableNode();
 
 	DATAFLOWENGINE_API void SetVariable(UDataflow* DataflowAsset, FName VariableName);
+	DATAFLOWENGINE_API bool TryAddVariableToDataflowAsset(UDataflow& DataflowAsset);
 
 	DATAFLOWENGINE_API FName GetVariableName() const { return VariableName; }
 
