Gather CDN cache hit statistics

#rnx

diff --git a/Engine/Source/Runtime/Experimental/IoStore/OnDemand/Private/OnDemandHttpClient.cpp b/Engine/Source/Runtime/Experimental/IoStore/OnDemand/Private/OnDemandHttpClient.cpp
--- a/Engine/Source/Runtime/Experimental/IoStore/OnDemand/Private/OnDemandHttpClient.cpp
+++ b/Engine/Source/Runtime/Experimental/IoStore/OnDemand/Private/OnDemandHttpClient.cpp
@@ -20,6 +20,24 @@
 namespace UE::IoStore
 {
 
+static void TrackCdnCacheStats(const HTTP::FResponse& Response)
+{
+#if !UE_BUILD_TEST
+	for (FAnsiStringView Header : { "X-Cache", "CF-Cache-Status" })
+	{
+		FAnsiStringView Value = Response.GetHeader(Header);
+		if (Value.IsEmpty())
+		{
+			continue;
+		}
+		int32 Reply = (Value.Find("HIT", 0, ESearchCase::IgnoreCase) >= 0);
+		FOnDemandIoBackendStats::Get()->OnHttpCdnCacheReply(Reply);
+		return;
+	}
+	FOnDemandIoBackendStats::Get()->OnHttpCdnCacheReply(-1);
+#endif
+}
+
 FIoStatus LoadDefaultHttpCertificates(bool& bWasLoaded)
 {
 	using namespace UE::IoStore::HTTP;
@@ -204,6 +222,7 @@ void FHttpClient::IssueRequest(FRequestParams&& Params)
 				FResponse& Response = Status.GetResponse();
 				StatusCode = Response.GetStatusCode();
 				Response.SetDestination(&Buffer);
+				TrackCdnCacheStats(Response);
 				break;
 			}
 			case FTicketStatus::EId::Content:
@@ -490,6 +509,7 @@ void FMultiEndpointHttpClient::IssueRequest(FRequest&& Request)
 				FResponse& HttpResponse = TicketStatus.GetResponse();
 				Request.StatusCode		= HttpResponse.GetStatusCode();
 				HttpResponse.SetDestination(&Request.Body);
+				TrackCdnCacheStats(HttpResponse);
 				break;
 			}
 			case FTicketStatus::EId::Content:
@@ -536,6 +556,10 @@ void FMultiEndpointHttpClient::IssueRequest(FRequest&& Request)
 	FEventLoop::FRequestParams RequestParams = { .bAllowChunked = Config.bAllowChunkedTransfer };
 	UE::IoStore::HTTP::FRequest HttpRequest = EventLoop.Get(Url, ConnectionPool, &RequestParams);
 
+#if !UE_BUILD_TEST
+	HttpRequest.Header("Pragma", "akamai-x-cache-on");
+#endif
+
 	if (Request.Range.GetOffset() > 0 || Request.Range.GetLength() > 0)
 	{
 		HttpRequest.Header(ANSITEXTVIEW("Range"),
diff --git a/Engine/Source/Runtime/Experimental/IoStore/OnDemand/Private/Statistics.cpp b/Engine/Source/Runtime/Experimental/IoStore/OnDemand/Private/Statistics.cpp
--- a/Engine/Source/Runtime/Experimental/IoStore/OnDemand/Private/Statistics.cpp
+++ b/Engine/Source/Runtime/Experimental/IoStore/OnDemand/Private/Statistics.cpp
@@ -163,6 +163,9 @@ int64					GHttpHistoryBytes[GHttpHistoryCount] = {};
 int64					GHttpHistoryTotalDuration = 0;
 int64					GHttpHistoryTotalBytes = 0;
 int64 					GHttpHistoryIndex = 0;
+static uint32			GHttpCdnCacheHit = 0;
+static uint32			GHttpCdnCacheMiss = 0;
+static uint32			GHttpCdnCacheUnknown = 0;
 
 ////////////////////////////////////////////////////////////////////////////////
 // CSV STATS
@@ -272,76 +275,99 @@ FOnDemandIoBackendStats::FOnDemandIoBackendStats(FBackendStatus& InStatus)
 	GStatisticsOnScreenDelegateHandle = FCoreDelegates::OnGetOnScreenMessages.AddLambda(
 		[this] (FCoreDelegates::FSeverityMessageMap& Out)
 		{
-			if (GIasDisplayOnScreenStatistics)
+			if (!GIasDisplayOnScreenStatistics)
 			{
-				{
-					TStringBuilder<256> BackendStatusText;
-					BackendStatusText << TEXT("IAS Backend Status: ");
-					BackendStatus.ToString(BackendStatusText);
+				return;
+			}
 
-					FCoreDelegates::EOnScreenMessageSeverity Verbosity = BackendStatus.IsHttpError() ?	FCoreDelegates::EOnScreenMessageSeverity::Error :
-																										FCoreDelegates::EOnScreenMessageSeverity::Info;
+			{
+				TStringBuilder<256> BackendStatusText;
+				BackendStatusText << TEXT("IAS Backend Status: ");
+				BackendStatus.ToString(BackendStatusText);
 
-					Out.Add(Verbosity, FText::FromStringView(BackendStatusText.ToView()));
-				}
+				FCoreDelegates::EOnScreenMessageSeverity Verbosity = BackendStatus.IsHttpError() ?	FCoreDelegates::EOnScreenMessageSeverity::Error :
+																									FCoreDelegates::EOnScreenMessageSeverity::Info;
 
-				for (const FOnDemandHostGroup& HostGroup : HostGroups)
-				{
-					FCoreDelegates::EOnScreenMessageSeverity Verbosity = FCoreDelegates::EOnScreenMessageSeverity::Info;
-					TStringBuilder<256> Text;
-					Text << TEXT("IAS HostGroup Status: ");
-
-					if(HostGroup.PrimaryHostIndex() != INDEX_NONE)
-					{
-						Text << HostGroup.PrimaryHost();
-						Text << TEXT(" (") << HostGroup.PrimaryHostIndex() << TEXT("/") << HostGroup.Hosts().Num() << TEXT(")");
-					}
-					else
-					{
-						Text << TEXT("Disconnected");
-						Verbosity = FCoreDelegates::EOnScreenMessageSeverity::Error;
-					}
-
-					Out.Add(Verbosity, FText::FromStringView(Text.ToView()));
-				}
+				Out.Add(Verbosity, FText::FromStringView(BackendStatusText.ToView()));
+			}
 
-				TStringBuilder<256> CachingText;
-				if (BackendStatus.IsCacheEnabled())
+			for (const FOnDemandHostGroup& HostGroup : HostGroups)
+			{
+				FCoreDelegates::EOnScreenMessageSeverity Verbosity = FCoreDelegates::EOnScreenMessageSeverity::Info;
+				TStringBuilder<256> Text;
+				Text << TEXT("IAS HostGroup Status: ");
+
+				if(HostGroup.PrimaryHostIndex() != INDEX_NONE)
 				{
-					CachingText << TEXT("Cached: ") << FText::AsMemory(GCacheCachedBytes.Get()).ToString();
-					CachingText << TEXT(" | Rejected: ") << FText::AsMemory(GCacheRejectBytes.Get()).ToString();
-					CachingText << TEXT(" | Read: ") << FText::AsMemory(GCacheReadBytes.Get()).ToString();
-					CachingText << TEXT(" (") << GCacheGetCount.Get() << TEXT(")");
+					Text << HostGroup.PrimaryHost();
+					Text << TEXT(" (") << HostGroup.PrimaryHostIndex() << TEXT("/") << HostGroup.Hosts().Num() << TEXT(")");
 				}
 				else
 				{
-					CachingText = TEXTVIEW("Caching Disabled");
+					Text << TEXT("Disconnected");
+					Verbosity = FCoreDelegates::EOnScreenMessageSeverity::Error;
 				}
 
-				TStringBuilder<256> Builder;
-				Builder.Appendf(
-					TEXT("IAS Backend Stats:   %s | Downloaded: %s (%d) Avg %d ms | Retries: %d | Pending: %d"),
-					CachingText.ToString(),
-					*FText::AsMemory(GHttpDownloadedBytes.Get()).ToString(),
-					GHttpGetCount.Get(),
-					GHttpDurationMsAvg.Get(),
-					GHttpRetryCount.Get(),
-					GHttpPendingCount.Get()
-				);
+				Out.Add(Verbosity, FText::FromStringView(Text.ToView()));
+			}
+
+			TStringBuilder<256> CachingText;
+			if (BackendStatus.IsCacheEnabled())
+			{
+				CachingText << TEXT("Cached: ") << FText::AsMemory(GCacheCachedBytes.Get()).ToString();
+				CachingText << TEXT(" | Rejected: ") << FText::AsMemory(GCacheRejectBytes.Get()).ToString();
+				CachingText << TEXT(" | Read: ") << FText::AsMemory(GCacheReadBytes.Get()).ToString();
+				CachingText << TEXT(" (") << GCacheGetCount.Get() << TEXT(")");
+			}
+			else
+			{
+				CachingText = TEXTVIEW("Caching Disabled");
+			}
+
+			TStringBuilder<256> Builder;
+			Builder.Appendf(
+				TEXT("IAS Backend Stats:   %s | Downloaded: %s (%d) Avg %d ms | Retries: %d | Pending: %d"),
+				CachingText.ToString(),
+				*FText::AsMemory(GHttpDownloadedBytes.Get()).ToString(),
+				GHttpGetCount.Get(),
+				GHttpDurationMsAvg.Get(),
+				GHttpRetryCount.Get(),
+				GHttpPendingCount.Get()
+			);
+
+			Out.Add(FCoreDelegates::EOnScreenMessageSeverity::Info, FText::FromStringView(Builder.ToView()));
+
+			{
+				Builder.Reset();
+				Builder << TEXT("IAS CDN:");
+
+				uint32 Total = GHttpCdnCacheHit + GHttpCdnCacheMiss + GHttpCdnCacheUnknown;
+				auto AsPercent = [Total] (uint32 Value) { return Total ? ((Value * 100 + (Total >> 1)) / Total) : 0; };
+
+				Builder << TEXT(" Hit/Miss/NoHdr: ");
+
+				Builder << GHttpCdnCacheHit;
+				Builder << TEXT("/") << GHttpCdnCacheMiss;
+				Builder << TEXT("/") << GHttpCdnCacheUnknown;
+
+				Builder << TEXT(" - ") << AsPercent(GHttpCdnCacheHit),
+				Builder << TEXT("%/") << AsPercent(GHttpCdnCacheMiss),
+				Builder << TEXT("%/") << AsPercent(GHttpCdnCacheUnknown),
+				Builder << TEXT("%");
 
 				Out.Add(FCoreDelegates::EOnScreenMessageSeverity::Info, FText::FromStringView(Builder.ToView()));
+			}
 
-				if (GHttpDecodeErrorCount.Get() > 0 || GCacheDecodeErrorCount.Get() > 0 || GHttpErrorCount.Get() > 0)
-				{
-					TStringBuilder<256> Warning;
-					Warning.Appendf(TEXT("IAS - Cache Decode Errors %d | Http Decode Errors %d | Http Errors %d"),
-						GCacheDecodeErrorCount.Get(),
-						GHttpDecodeErrorCount.Get(),
-						GHttpErrorCount.Get()
-					);
-
-					Out.Add(FCoreDelegates::EOnScreenMessageSeverity::Error, FText::FromStringView(Warning.ToView()));
-				}
+			if (GHttpDecodeErrorCount.Get() > 0 || GCacheDecodeErrorCount.Get() > 0 || GHttpErrorCount.Get() > 0)
+			{
+				Builder.Reset();
+				Builder.Appendf(TEXT("IAS Errors: Cache Decode: %d | Http Decode: %d | Http: %d"),
+					GCacheDecodeErrorCount.Get(),
+					GHttpDecodeErrorCount.Get(),
+					GHttpErrorCount.Get()
+				);
+
+				Out.Add(FCoreDelegates::EOnScreenMessageSeverity::Error, FText::FromStringView(Builder.ToView()));
 			}
 		});
 #endif // UE_ENABLE_ONSCREEN_STATISTICS
@@ -755,6 +781,16 @@ void FOnDemandIoBackendStats::OnHttpDecodeError()
 	GHttpDecodeErrorCount.Add(1);
 }
 
+void FOnDemandIoBackendStats::OnHttpCdnCacheReply(int32 Reply)
+{
+	switch (Reply)
+	{
+	case -1:	GHttpCdnCacheUnknown += 1;	break;	
+	case 0:		GHttpCdnCacheMiss += 1;		break;
+	default:	GHttpCdnCacheHit += 1;		break;
+	}
+}
+
 } // namespace UE::IoStore
 
 #undef UE_ENABLE_ONSCREEN_STATISTICS
diff --git a/Engine/Source/Runtime/Experimental/IoStore/OnDemand/Private/Statistics.h b/Engine/Source/Runtime/Experimental/IoStore/OnDemand/Private/Statistics.h
--- a/Engine/Source/Runtime/Experimental/IoStore/OnDemand/Private/Statistics.h
+++ b/Engine/Source/Runtime/Experimental/IoStore/OnDemand/Private/Statistics.h
@@ -74,6 +74,7 @@ public:
 	void OnHttpRetry() IAS_STATISTICS_IMPL()
 	void OnHttpError() IAS_STATISTICS_IMPL()
 	void OnHttpDecodeError() IAS_STATISTICS_IMPL()
+	void OnHttpCdnCacheReply(int32 Reply) IAS_STATISTICS_IMPL()
 
 private:
 
