Slightly tweak on-screen CDN hit/miss output

#rnx

diff --git a/Engine/Source/Runtime/Experimental/IoStore/OnDemand/Private/Statistics.cpp b/Engine/Source/Runtime/Experimental/IoStore/OnDemand/Private/Statistics.cpp
--- a/Engine/Source/Runtime/Experimental/IoStore/OnDemand/Private/Statistics.cpp
+++ b/Engine/Source/Runtime/Experimental/IoStore/OnDemand/Private/Statistics.cpp
@@ -339,21 +339,20 @@ FOnDemandIoBackendStats::FOnDemandIoBackendStats(FBackendStatus& InStatus)
 
 			{
 				Builder.Reset();
-				Builder << TEXT("IAS CDN:");
-
-				uint32 Total = GHttpCdnCacheHit + GHttpCdnCacheMiss + GHttpCdnCacheUnknown;
-				auto AsPercent = [Total] (uint32 Value) { return Total ? ((Value * 100 + (Total >> 1)) / Total) : 0; };
-
-				Builder << TEXT(" Hit/Miss/NoHdr: ");
+				Builder << TEXT("IAS CDN: Hit/Miss/NoHdr: ");
 
 				Builder << GHttpCdnCacheHit;
 				Builder << TEXT("/") << GHttpCdnCacheMiss;
 				Builder << TEXT("/") << GHttpCdnCacheUnknown;
 
-				Builder << TEXT(" - ") << AsPercent(GHttpCdnCacheHit),
-				Builder << TEXT("%/") << AsPercent(GHttpCdnCacheMiss),
-				Builder << TEXT("%/") << AsPercent(GHttpCdnCacheUnknown),
-				Builder << TEXT("%");
+				if (uint32 Total = GHttpCdnCacheHit + GHttpCdnCacheMiss + GHttpCdnCacheUnknown; Total)
+				{
+					auto AsPercent = [Total] (uint32 Value) { return (Value * 100 + (Total >> 1)) / Total; };
+					Builder << TEXT(" - ") << AsPercent(GHttpCdnCacheHit);
+					Builder << TEXT("%/") << AsPercent(GHttpCdnCacheMiss);
+					Builder << TEXT("%/") << AsPercent(GHttpCdnCacheUnknown);
+					Builder << TEXT("%");
+				}
 
 				Out.Add(FCoreDelegates::EOnScreenMessageSeverity::Info, FText::FromStringView(Builder.ToView()));
 			}
