Route polling through a "poller" object to allow some context association.

#rnx

diff --git a/Engine/Source/Runtime/Experimental/IoStore/HttpClient/Private/Loop.inl b/Engine/Source/Runtime/Experimental/IoStore/HttpClient/Private/Loop.inl
--- a/Engine/Source/Runtime/Experimental/IoStore/HttpClient/Private/Loop.inl
+++ b/Engine/Source/Runtime/Experimental/IoStore/HttpClient/Private/Loop.inl
@@ -1068,6 +1068,7 @@ private:
 	TArray<FPeerGroup>			PeerGroups;
 	FWorkQueue					Work;
 	FHost&						Host;
+	FPoller						Poller;
 	uint32						BusyCount = 0;
 	int32						WaitTimeAccum = 0;
 
@@ -1116,7 +1117,7 @@ int32 FHostGroup::Wait(const FTickState& State)
 	}
 
 	// Actually do the wait
-	int32 Result = FWaiter::Wait(Waiters, PollTimeoutMs);
+	int32 Result = FWaiter::Wait(Waiters, Poller, PollTimeoutMs);
 	if (Result <= 0)
 	{
 		// If the user opts to not block then we don't accumulate wait time and
diff --git a/Engine/Source/Runtime/Experimental/IoStore/HttpClient/Private/Socket.inl b/Engine/Source/Runtime/Experimental/IoStore/HttpClient/Private/Socket.inl
--- a/Engine/Source/Runtime/Experimental/IoStore/HttpClient/Private/Socket.inl
+++ b/Engine/Source/Runtime/Experimental/IoStore/HttpClient/Private/Socket.inl
@@ -79,9 +79,17 @@
 		IAS_HTTP_SOCKET_FORWARD(Shutdown,		shutdown)
 		IAS_HTTP_SOCKET_FORWARD(Socket,			socket)
 
-#		if defined(IAS_HTTP_USE_POLL)
-			IAS_HTTP_SOCKET_FORWARD(Poll,			poll)
-#		endif
+		using PollType = UPTRINT;
+		static PollType	Poll_CreateContext()			{ return 0; }
+		static void		Poll_DestroyContext(PollType)	{}
+
+#if defined(IAS_HTTP_USE_POLL)
+		template <typename... ArgTypes>
+		static auto Poll_Wait(PollType, ArgTypes&&... Args)
+		{
+			return poll(Forward<ArgTypes>(Args)...);
+		}
+#endif
 	} // namespace UE::IoStore::HTTP
 #endif // IAS_HTTP_SOCKET_FORWARD
 
@@ -113,6 +121,122 @@ uint16 Socket_HtoNs(uint16 v) { return ByteSwap(v); }
 namespace UE::IoStore::HTTP
 {
 
+// {{{1 poll ...................................................................
+
+////////////////////////////////////////////////////////////////////////////////
+struct FPollFd
+#if defined(IAS_HTTP_USE_POLL)
+	: public pollfd
+{
+#else
+{
+	SocketType			fd;
+	int16				events;
+	int16				revents;
+	static const int16	POLLIN  = 1 << 0;
+	static const int16	POLLOUT = 1 << 1;
+	static const int16	POLLERR = 1 << 2;
+	static const int16	POLLHUP = POLLERR;
+	static const int16	POLLNVAL= POLLERR;
+#endif
+
+	static const int16	PollIn	= int16(POLLIN);
+	static const int16	PollOut	= int16(POLLOUT);
+	static const int16	PollErr	= int16(POLLERR);
+	static const int16	PollHup	= int16(POLLHUP);
+	static const int16	PollNVal= int16(POLLNVAL);
+	static const int16	PollAll = PollIn|PollOut|PollErr|PollHup|PollNVal;
+
+	FPollFd(SocketType Socket, int16 Events)
+	{
+		fd = FdType(Socket);
+		events = EventType(Events);
+		revents = EventType(0);
+	}
+
+	bool HasTrigger() const		{ return !!revents; }
+	bool HasSendTrigger() const	{ return !!(revents & EventType(PollOut)); }
+	bool HasRecvTrigger() const	{ return !!(revents & EventType(PollIn)); }
+
+private:
+	// The following is odd because POLLFD varies subtly from one platform to
+	// the next. With typedefs we can cleanly set members and not get narrowing
+	// warnings from the compiler.
+	using FdType = decltype(fd);
+	using EventType = decltype(events);
+};
+
+////////////////////////////////////////////////////////////////////////////////
+class FPoller
+{
+public:
+				FPoller()					{ Handle = Poll_CreateContext(); }
+				FPoller(FPoller&& Rhs)		{ Swap(Handle, Rhs.Handle); }
+				~FPoller()					{ Poll_DestroyContext(Handle); }
+	FPoller&	operator = (FPoller&& Rhs)	{ Swap(Handle, Rhs.Handle); return *this; }
+	int32		Wait(FPollFd* Polls, int32 Num, int32 TimeOutMs);
+
+private:
+				FPoller(const FPoller&)		= delete;
+	FPoller&	operator = (const FPoller&)	= delete;
+	PollType	Handle = 0;
+};
+
+////////////////////////////////////////////////////////////////////////////////
+int32 FPoller::Wait(FPollFd* Polls, int32 Num, int32 TimeOutMs)
+{
+#if defined(IAS_HTTP_USE_POLL)
+	static_assert(sizeof(FPollFd) == sizeof(pollfd));
+	return Poll_Wait(Handle, Polls, Num, TimeOutMs);
+#else
+	timeval TimeVal = {};
+	timeval* TimeValPtr = (TimeOutMs >= 0 ) ? &TimeVal : nullptr;
+	if (TimeOutMs > 0)
+	{
+		int32 TimeoutUs = TimeOutMs << 10;
+		TimeVal = { TimeOutMs >> 10, TimeoutUs & ((1 << 20) - 1) };
+	}
+
+	fd_set FdSetRead;	FD_ZERO(&FdSetRead);
+	fd_set FdSetWrite;	FD_ZERO(&FdSetWrite);
+	fd_set FdSetExcept; FD_ZERO(&FdSetExcept);
+
+	SocketType MaxFd = 0;
+	for (FPollFd& Poll : MakeArrayView(Polls, Num))
+	{
+		fd_set* RwSet = (Poll.events & FPollFd::PollIn) ? &FdSetRead : &FdSetWrite;
+		FD_SET(Poll.fd, RwSet);
+		FD_SET(Poll.fd, &FdSetExcept);
+		MaxFd = FMath::Max(Poll.fd, MaxFd);
+	}
+
+	int32 Result = Socket_Select(int32(MaxFd + 1), &FdSetRead, &FdSetWrite, &FdSetExcept, TimeValPtr);
+	if (Result <= 0)
+	{
+		return Result;
+	}
+
+	for (FPollFd& Poll : MakeArrayView(Polls, Num))
+	{
+		if (FD_ISSET(Poll.fd, &FdSetExcept))
+		{
+			Poll.revents = FPollFd::PollErr;
+			continue;
+		}
+
+		fd_set* RwSet = (Poll.events & FPollFd::PollIn) ? &FdSetRead : &FdSetWrite;
+		if (FD_ISSET(Poll.fd, RwSet))
+		{
+			Poll.revents = Poll.events;
+		}
+	}
+
+	return Result;
+#endif // IAS_HTTP_USE_POLL
+}
+
+
+
 // {{{1 wait ...................................................................
 
 ////////////////////////////////////////////////////////////////////////////////
@@ -154,7 +278,7 @@ public:
 		Both = Send|Recv
 	};
 
-	static int32	Wait(TArrayView<FWaiter> Waiters, int32 TimeoutMs);
+	static int32	Wait(TArrayView<FWaiter> Waiters, FPoller& Poller, int32 TimeoutMs);
 
 					FWaiter() = default;
 					FWaiter(FWaitable&& Waitable);
@@ -184,104 +308,40 @@ FWaiter::FWaiter(FWaitable&& Waitable)
 }
 
 ////////////////////////////////////////////////////////////////////////////////
-int32 FWaiter::Wait(TArrayView<FWaiter> Waiters, int32 TimeoutMs)
+int32 FWaiter::Wait(TArrayView<FWaiter> Waiters, FPoller& Poller, int32 TimeoutMs)
 {
-#if !defined(IAS_HTTP_USE_POLL)
-	struct FSelect
-	{
-		SocketType	fd;
-		int32		events;
-		int32		revents;
+	uint16 Events[] = {
+		FPollFd::PollErr,
+		FPollFd::PollOut,
+		FPollFd::PollIn,
+		FPollFd::PollIn | FPollFd::PollOut,
 	};
-	static const int32 POLLIN  = 1 << 0;
-	static const int32 POLLOUT = 1 << 1;
-	static const int32 POLLERR = 1 << 2;
-	static const int32 POLLHUP = POLLERR;
-	static const int32 POLLNVAL= POLLERR;
-#else
-	using FSelect = pollfd;
-#endif
-
-	// The following looks odd because POLLFD varies subtly from one platform
-	// to the next. To cleanly set members to zero and to not get narrowing
-	// warnings from the compiler, we list-init and don't assume POD types.
-	using PollEventType = decltype(FSelect::events);
-	PollEventType Events[] = { POLLERR, POLLOUT, POLLIN, POLLIN|POLLOUT };
-	TArray<FSelect, TFixedAllocator<64>> Selects;
+	TArray<FPollFd, TFixedAllocator<64>> Polls;
 	for (FWaiter& Waiter : Waiters)
 	{
-		Selects.Emplace_GetRef() = {
-			Waiter.Socket,
-			Events[uint32(Waiter.WaitOn)],
-			/* 0 */
-		};
+		Polls.Emplace(Waiter.Socket, Events[uint32(Waiter.WaitOn)]);
 	}
 
 	// Poll the sockets
-#if defined(IAS_HTTP_USE_POLL)
-	int32 Result = Socket_Poll(Selects.GetData(), Selects.Num(), TimeoutMs);
-	if (Result <= 0)
-	{
-		return Result;
-	}
-#else
-	timeval TimeVal = {};
-	timeval* TimeValPtr = (TimeoutMs >= 0 ) ? &TimeVal : nullptr;
-	if (TimeoutMs > 0)
-	{
-		int32 TimeoutUs = TimeoutMs << 10;
-		TimeVal = { TimeoutMs >> 10, TimeoutUs & ((1 << 20) - 1) };
-	}
-
-	fd_set FdSetRead;	FD_ZERO(&FdSetRead);
-	fd_set FdSetWrite;	FD_ZERO(&FdSetWrite);
-	fd_set FdSetExcept; FD_ZERO(&FdSetExcept);
-
-	SocketType MaxFd = 0;
-	for (FSelect& Select : Selects)
-	{
-		fd_set* RwSet = (Select.events & POLLIN) ? &FdSetRead : &FdSetWrite;
-		FD_SET(Select.fd, RwSet);
-		FD_SET(Select.fd, &FdSetExcept);
-		MaxFd = FMath::Max(Select.fd, MaxFd);
-	}
-
-	int32 Result = Socket_Select(int32(MaxFd + 1), &FdSetRead, &FdSetWrite, &FdSetExcept, TimeValPtr);
+	int32 Result = Poller.Wait(Polls.GetData(), Polls.Num(), TimeoutMs);
 	if (Result <= 0)
 	{
 		return Result;
 	}
 
-	for (FSelect& Select : Selects)
-	{
-		if (FD_ISSET(Select.fd, &FdSetExcept))
-		{
-			Select.revents = POLLERR;
-			continue;
-		}
-
-		fd_set* RwSet = (Select.events & POLLIN) ? &FdSetRead : &FdSetWrite;
-		if (FD_ISSET(Select.fd, RwSet))
-		{
-			Select.revents = Select.events;
-		}
-	}
-#endif // IAS_HTTP_USE_POLL
-
 	// Transfer poll results to the input sockets. We don't transfer across error
 	// states. Subsequent sockets ops can take care of that instead.
-	static const auto TestBits = POLLIN|POLLOUT|POLLERR|POLLHUP|POLLNVAL;
 	for (uint32 i = 0, n = Waiters.Num(); i < n; ++i)
 	{
-		auto RetEvents = Selects[i].revents;
-		if (!(RetEvents & TestBits))
+		const FPollFd& Poll = Polls[i];
+		if (!Poll.HasTrigger())
 		{
 			continue;
 		}
 
 		uint8 Value = 0;
-		if (!!(RetEvents & POLLOUT)) Value |= uint8(EWhat::Send);
-		if (!!(RetEvents & POLLIN))	 Value |= uint8(EWhat::Recv);
+		if (Poll.HasSendTrigger()) Value |= uint8(EWhat::Send);
+		if (Poll.HasRecvTrigger()) Value |= uint8(EWhat::Recv);
 		Waiters[i].Ready = Value ? Value : uint8(EWhat::Both);
 	}
 
