Fix an IAS crash if no CDN is reachable.

#rnx

- Do not try to create a connection if there is no valid current host to connect to.

diff --git a/Engine/Source/Runtime/Experimental/IoStore/OnDemand/Private/OnDemandHttpClient.cpp b/Engine/Source/Runtime/Experimental/IoStore/OnDemand/Private/OnDemandHttpClient.cpp
--- a/Engine/Source/Runtime/Experimental/IoStore/OnDemand/Private/OnDemandHttpClient.cpp
+++ b/Engine/Source/Runtime/Experimental/IoStore/OnDemand/Private/OnDemandHttpClient.cpp
@@ -479,9 +479,12 @@ void FMultiEndpointHttpClient::UpdateConnection(const FOnDemandHostGroup& HostGr
 	{
 		Connection->CurrentHost = HostGroup.PrimaryHostIndex();
 
-		if (Connection->Pools[Connection->CurrentHost].IsValid() == false)
+		if (Connection->CurrentHost != INDEX_NONE)
 		{
-			Connection->Pools[Connection->CurrentHost] = CreateConnection(HostGroup.PrimaryHost());
+			if (Connection->Pools[Connection->CurrentHost].IsValid() == false)
+			{
+				Connection->Pools[Connection->CurrentHost] = CreateConnection(HostGroup.PrimaryHost());
+			}
 		}
 	}
 }
