Adding a CVar for enabling Iris CSV stats for untracked classes.

diff --git a/Engine/Source/Runtime/Experimental/Iris/Core/Private/Iris/ReplicationSystem/ObjectReplicationBridge.cpp b/Engine/Source/Runtime/Experimental/Iris/Core/Private/Iris/ReplicationSystem/ObjectReplicationBridge.cpp
--- a/Engine/Source/Runtime/Experimental/Iris/Core/Private/Iris/ReplicationSystem/ObjectReplicationBridge.cpp
+++ b/Engine/Source/Runtime/Experimental/Iris/Core/Private/Iris/ReplicationSystem/ObjectReplicationBridge.cpp
@@ -111,6 +111,15 @@ static FAutoConsoleVariableRef CVarUseVerboseCsvStats(
 	TEXT("Whether to use verbose per-class csv stats. Default is false in Shipping, otherwise True.")
 );
 
+#if !UE_BUILD_SHIPPING
+static bool bUseVeryVerboseIrisCsvStats = false;
+static FAutoConsoleVariableRef CVarUseVeryVerboseCsvStats(
+	TEXT("net.Iris.UseVeryVerboseIrisCsvStats"),
+	bUseVeryVerboseIrisCsvStats,
+	TEXT("Whether to use very verbose per-class csv stats. This causes all untracked classes to get their own CSV stats. Will incur significant CSV bloat.")
+);
+#endif // !UE_BUILD_SHIPPING
+
 namespace UE::Net::Private
 {
 	void CallRegisterReplicationFragments(UObject* Object, FFragmentRegistrationContext& Context, EFragmentRegistrationFlags RegistrationFlags)
@@ -1841,7 +1850,9 @@ int32 UObjectReplicationBridge::GetTypeStatsIndex(const UClass* Class)
 {
 	using namespace UE::Net;
 	using namespace UE::Net::Private;
-	
+
+	const UClass* const CachedClass = Class;
+
 	if (ClassesWithTypeStats.Num() > 0)
 	{
 		FNetTypeStats& TypeStats = GetReplicationSystem()->GetReplicationSystemInternal()->GetNetTypeStats();
@@ -1854,6 +1865,18 @@ int32 UObjectReplicationBridge::GetTypeStatsIndex(const UClass* Class)
 		}
 	}
 
+#if !UE_BUILD_SHIPPING
+	if (bUseVeryVerboseIrisCsvStats)
+	{
+		FNetTypeStats& TypeStats = GetReplicationSystem()->GetReplicationSystemInternal()->GetNetTypeStats();
+		if (CachedClass)
+		{
+			const FName ClassName = CachedClass->GetFName();
+			return TypeStats.GetOrCreateTypeStats(ClassName);
+		}
+	}
+#endif // !UE_BUILD_SHIPPING
+
 	return FNetTypeStats::DefaultTypeStatsIndex;
 }
 
