Dataflow node debug drawing:
- new IDataflowDebugDrawInterface to abstract out drawing basic shapes
- new FDataflowDebugDraw concrete class that uses UDebugDrawComponent and its SceneProxy to do the actual drawing
- adding a new virtual DebugDraw function to FDataflowNode so that nodes can specify their own debug drawing
- calling the new virtual function on any selected or pinned nodes from FDataflowConstructionViewportClient::Draw()

diff --git a/Engine/Plugins/Experimental/Dataflow/Source/DataflowEditor/Private/Dataflow/DataflowEditorToolkit.cpp b/Engine/Plugins/Experimental/Dataflow/Source/DataflowEditor/Private/Dataflow/DataflowEditorToolkit.cpp
--- a/Engine/Plugins/Experimental/Dataflow/Source/DataflowEditor/Private/Dataflow/DataflowEditorToolkit.cpp
+++ b/Engine/Plugins/Experimental/Dataflow/Source/DataflowEditor/Private/Dataflow/DataflowEditorToolkit.cpp
@@ -10,6 +10,8 @@
 #include "Animation/Skeleton.h"
 #include "Dataflow/DataflowAdvancedPreviewDetailsTab.h"
 #include "Dataflow/DataflowAssetViewerSettingsCustomization.h"
+#include "Dataflow/DataflowDebugDrawComponent.h"
+#include "Dataflow/DataflowDebugDraw.h"
 #include "Dataflow/DataflowCore.h"
 #include "Dataflow/DataflowEditor.h"
 #include "Dataflow/DataflowContent.h"
@@ -331,6 +333,11 @@ FDataflowEditorToolkit::FDataflowEditorToolkit(UAssetEditor* InOwningAssetEditor
 
 FDataflowEditorToolkit::~FDataflowEditorToolkit()
 {
+	if (DebugDrawComponent)
+	{
+		DebugDrawComponent->UnregisterComponent();
+	}
+
 	if (IConsoleVariable* const ConsoleVar = IConsoleManager::Get().FindConsoleVariable(TEXT("p.Dataflow.EnableGraphEval")))
 	{
 		ConsoleVar->OnChangedDelegate().Remove(GraphEvalCVarChangedDelegateHandle);
@@ -829,6 +836,74 @@ void FDataflowEditorToolkit::CreateSimulationViewportClient()
 
 //~ End FBaseAssetToolkit overrides
 
+
+TUniquePtr<IDataflowDebugDrawInterface> FDataflowEditorToolkit::CreateDebugDrawObject()
+{
+	if (ConstructionScene && !DebugDrawComponent)
+	{
+		if (TObjectPtr<AActor> RootActor = ConstructionScene->GetRootActor())
+		{
+			if (UDataflowDebugDrawComponent* Component = Cast<UDataflowDebugDrawComponent>(RootActor->FindComponentByClass<UDataflowDebugDrawComponent>()))
+			{
+				DebugDrawComponent = Cast<UDataflowDebugDrawComponent>(Component);
+			}
+			else
+			{
+				DebugDrawComponent = NewObject<UDataflowDebugDrawComponent>(RootActor, FName(LOCTEXT("DataflowDebugDrawComponent", "Dataflow Debug Draw Component").ToString(), RF_Transient));
+				DebugDrawComponent->RegisterComponentWithWorld(ConstructionScene->GetWorld());
+			}
+		}
+	}
+
+	checkf(DebugDrawComponent, TEXT("Could not create or find a DebugDrawComponentComponent"));
+
+	FDataflowDebugRenderSceneProxy* SceneProxy = static_cast<FDataflowDebugRenderSceneProxy*>(DebugDrawComponent->GetSceneProxy());
+	checkf(SceneProxy, TEXT("Could not find a FDataflowDebugRenderSceneProxy on the DebugDrawComponent"));
+
+	return MakeUnique<FDataflowDebugDraw>(SceneProxy);
+}
+
+void FDataflowEditorToolkit::UpdateDebugDraw()
+{
+	const TUniquePtr<IDataflowDebugDrawInterface> DebugDrawObject = CreateDebugDrawObject();
+
+	if (const TObjectPtr<UDataflowBaseContent>& EditorContent = GetEditorContent())
+	{
+		if (const TSharedPtr<UE::Dataflow::FEngineContext> DataflowContext = EditorContent->GetDataflowContext())
+		{
+			if (const UDataflowEdNode* const SelectedNode = EditorContent->GetSelectedNode())
+			{
+				if (const TSharedPtr<const FDataflowNode> SelectedDataflowNode = SelectedNode->GetDataflowNode())
+				{
+					DebugDrawObject->ResetAllState();
+
+					FDataflowNode::FDebugDrawParameters DebugDrawParameters;
+					DebugDrawParameters.bNodeIsSelected = true;
+					DebugDrawParameters.bNodeIsPinned = false;
+					SelectedDataflowNode->DebugDraw(*DataflowContext, *DebugDrawObject, DebugDrawParameters);
+				}
+			}
+
+			for (const UDataflowEdNode* const PinnedNode : EditorContent->GetDataflowAsset()->GetWireframeRenderTargets())
+			{
+				if (const TSharedPtr<const FDataflowNode> PinnedDataflowNode = PinnedNode->GetDataflowNode())
+				{
+					DebugDrawObject->ResetAllState();
+
+					FDataflowNode::FDebugDrawParameters DebugDrawParameters;
+					DebugDrawParameters.bNodeIsSelected = false;
+					DebugDrawParameters.bNodeIsPinned = true;
+					PinnedDataflowNode->DebugDraw(*DataflowContext, *DebugDrawObject, DebugDrawParameters);
+				}
+			}
+		}
+	}
+
+	DebugDrawComponent->UpdateBounds();
+	DebugDrawComponent->MarkRenderTransformDirty();
+}
+
+
 namespace UE::Dataflow::Private
 {
 	static void ShowNotificationMessage(const FText& Message, const SNotificationItem::ECompletionState CompletionState)
@@ -863,6 +938,8 @@ void FDataflowEditorToolkit::OnPropertyValueChanged(const FPropertyChangedEvent&
 
 			// Refresh graph display to display nodes with warning/error
 			GraphEditor->NotifyGraphChanged();
+
+			UpdateDebugDraw();
 		}
 	}
 }
@@ -1128,6 +1205,8 @@ void FDataflowEditorToolkit::OnNodeSelectionChanged(const TSet<UObject*>& InNewS
 			{
 				SelectComponentsInView(PrimarySelection);
 			}
+
+			UpdateDebugDraw();
 		}
 
 		EditorContent->SetConstructionDirty(true);
@@ -2264,6 +2343,7 @@ void FDataflowEditorToolkit::AddReferencedObjects(FReferenceCollector& Collector
 {
 	Collector.AddReferencedObjects(SelectedDataflowNodes);
 	Collector.AddReferencedObject(PrimarySelection);
+	Collector.AddReferencedObject(DebugDrawComponent);
 }
 
 
diff --git a/Engine/Plugins/Experimental/Dataflow/Source/DataflowEditor/Public/Dataflow/DataflowEditorToolkit.h b/Engine/Plugins/Experimental/Dataflow/Source/DataflowEditor/Public/Dataflow/DataflowEditorToolkit.h
--- a/Engine/Plugins/Experimental/Dataflow/Source/DataflowEditor/Public/Dataflow/DataflowEditorToolkit.h
+++ b/Engine/Plugins/Experimental/Dataflow/Source/DataflowEditor/Public/Dataflow/DataflowEditorToolkit.h
@@ -229,6 +229,12 @@ private:
 	void RegisterContextHandlers();
 	void UnregisterContextHandlers();
 
+	/** Create a concrete DebugDrawInterface object */
+	TUniquePtr<IDataflowDebugDrawInterface> CreateDebugDrawObject();
+
+	/** Update the debug draw based on a change of currently selected or pinned nodes */
+	void UpdateDebugDraw();
+
 	// List of editor commands used  for the dataflow asset
 	TSharedPtr<FUICommandList> GraphEditorCommands;
 
@@ -299,4 +305,7 @@ private:
 
 	FDateTime GraphEvaluationBegin;
 	FDateTime GraphEvaluationFinished;
+
+	/** Persistent component used for debug drawing */
+	TObjectPtr<class UDataflowDebugDrawComponent> DebugDrawComponent;
 };
diff --git a/Engine/Source/Runtime/Experimental/Dataflow/Core/Private/Dataflow/DataflowDynamicMeshDebugDrawMesh.cpp b/Engine/Source/Runtime/Experimental/Dataflow/Core/Private/Dataflow/DataflowDynamicMeshDebugDrawMesh.cpp
new file mode 100644
--- /dev/null
+++ b/Engine/Source/Runtime/Experimental/Dataflow/Core/Private/Dataflow/DataflowDynamicMeshDebugDrawMesh.cpp
@@ -0,0 +1,56 @@
+// Copyright Epic Games, Inc. All Rights Reserved.
+
+#include "Dataflow/DataflowDynamicMeshDebugDrawMesh.h"
+#include "DynamicMesh/DynamicMesh3.h"
+
+FDynamicMeshDebugDrawMesh::FDynamicMeshDebugDrawMesh(const class UE::Geometry::FDynamicMesh3* DynamicMesh) :
+	DynamicMesh(DynamicMesh)
+{
+}
+
+int32 FDynamicMeshDebugDrawMesh::GetMaxVertexIndex() const
+{
+	if (DynamicMesh)
+	{
+		return DynamicMesh->MaxVertexID();
+	}
+	return 0;
+}
+
+bool FDynamicMeshDebugDrawMesh::IsValidVertex(int32 VertexIndex) const
+{
+	return (DynamicMesh && DynamicMesh->IsVertex(VertexIndex));
+}
+
+FVector FDynamicMeshDebugDrawMesh::GetVertexPosition(int32 VertexIndex) const
+{
+	if (IsValidVertex(VertexIndex))
+	{
+		return DynamicMesh->GetVertex(VertexIndex);
+	}
+
+	return FVector(0.0);
+}
+
+int32 FDynamicMeshDebugDrawMesh::GetMaxTriangleIndex() const
+{
+	if (DynamicMesh)
+	{
+		return DynamicMesh->MaxTriangleID();
+	}
+	return 0;
+}
+
+bool FDynamicMeshDebugDrawMesh::IsValidTriangle(int32 TriangleIndex) const
+{
+	return (DynamicMesh && DynamicMesh->IsTriangle(TriangleIndex));
+}
+
+FIntVector3 FDynamicMeshDebugDrawMesh::GetTriangle(int32 TriangleIndex) const
+{
+	if (IsValidTriangle(TriangleIndex))
+	{
+		return DynamicMesh->GetTriangle(TriangleIndex);
+	}
+	return FIntVector3(UE::Geometry::FDynamicMesh3::InvalidID);
+}
diff --git a/Engine/Source/Runtime/Experimental/Dataflow/Core/Private/Dataflow/DataflowSimpleDebugDrawMesh.cpp b/Engine/Source/Runtime/Experimental/Dataflow/Core/Private/Dataflow/DataflowSimpleDebugDrawMesh.cpp
new file mode 100644
--- /dev/null
+++ b/Engine/Source/Runtime/Experimental/Dataflow/Core/Private/Dataflow/DataflowSimpleDebugDrawMesh.cpp
@@ -0,0 +1,34 @@
+// Copyright Epic Games, Inc. All Rights Reserved.
+
+#include "Dataflow/DataflowSimpleDebugDrawMesh.h"
+
+int32 FSimpleDebugDrawMesh::GetMaxVertexIndex() const
+{
+	return Vertices.Num();
+}
+
+bool FSimpleDebugDrawMesh::IsValidVertex(int32 VertexIndex) const
+{
+	return VertexIndex > 0 && VertexIndex < Vertices.Num();
+}
+
+FVector FSimpleDebugDrawMesh::GetVertexPosition(int32 VertexIndex) const
+{
+	return Vertices[VertexIndex];
+}
+
+int32 FSimpleDebugDrawMesh::GetMaxTriangleIndex() const
+{
+	return Triangles.Num();
+}
+
+bool FSimpleDebugDrawMesh::IsValidTriangle(int32 TriangleIndex) const
+{
+	return TriangleIndex > 0 && TriangleIndex < Triangles.Num();
+}
+
+FIntVector3 FSimpleDebugDrawMesh::GetTriangle(int32 TriangleIndex) const
+{
+	return Triangles[TriangleIndex];
+}
+
diff --git a/Engine/Source/Runtime/Experimental/Dataflow/Core/Public/Dataflow/DataflowDebugDrawInterface.h b/Engine/Source/Runtime/Experimental/Dataflow/Core/Public/Dataflow/DataflowDebugDrawInterface.h
new file mode 100644
--- /dev/null
+++ b/Engine/Source/Runtime/Experimental/Dataflow/Core/Public/Dataflow/DataflowDebugDrawInterface.h
@@ -0,0 +1,47 @@
+// Copyright Epic Games, Inc. All Rights Reserved.
+
+#pragma once 
+
+#include "Math/MathFwd.h"
+#include "Containers/Array.h"
+#include "DynamicMesh/DynamicMesh3.h"
+
+class IDataflowDebugDrawInterface
+{
+public:
+
+	virtual ~IDataflowDebugDrawInterface() = default;
+
+	// State management
+	DATAFLOWCORE_API virtual void SetColor(const FLinearColor& InColor) = 0;
+	DATAFLOWCORE_API virtual void SetLineWidth(double Width) = 0;
+	DATAFLOWCORE_API virtual void SetWireframe(bool bInWireframe) = 0;
+	DATAFLOWCORE_API virtual void SetShaded(bool bInShaded) = 0;
+	DATAFLOWCORE_API virtual void SetTranslucent(bool bInShadedTranslucent) = 0;
+	DATAFLOWCORE_API virtual void SetForegroundPriority() = 0;
+	DATAFLOWCORE_API virtual void SetWorldPriority() = 0;
+	DATAFLOWCORE_API virtual void ResetAllState() = 0;
+
+	DATAFLOWCORE_API virtual void DrawLine(const FVector& Start, const FVector& End) const = 0;
+
+	struct IDebugDrawMesh
+	{
+		virtual ~IDebugDrawMesh() = default;
+
+		virtual int32 GetMaxVertexIndex() const = 0;
+		virtual bool IsValidVertex(int32 VertexIndex) const = 0;
+		virtual FVector GetVertexPosition(int32 VertexIndex) const = 0;
+
+		virtual int32 GetMaxTriangleIndex() const = 0;
+		virtual bool IsValidTriangle(int32 TriangleIndex) const = 0;
+		virtual FIntVector3 GetTriangle(int32 TriangleIndex) const = 0;
+	};
+
+	DATAFLOWCORE_API virtual void DrawMesh(const IDebugDrawMesh& Mesh) const = 0;
+
+	DATAFLOWCORE_API virtual void DrawBox(const FVector& Extents, const FQuat& Rotation, const FVector& Center, double UniformScale) const = 0;
+	DATAFLOWCORE_API virtual void DrawSphere(const FVector& Center, double Radius) const = 0;
+
+	// TODO: More draw functions
+
+};
diff --git a/Engine/Source/Runtime/Experimental/Dataflow/Core/Public/Dataflow/DataflowDynamicMeshDebugDrawMesh.h b/Engine/Source/Runtime/Experimental/Dataflow/Core/Public/Dataflow/DataflowDynamicMeshDebugDrawMesh.h
new file mode 100644
--- /dev/null
+++ b/Engine/Source/Runtime/Experimental/Dataflow/Core/Public/Dataflow/DataflowDynamicMeshDebugDrawMesh.h
@@ -0,0 +1,23 @@
+// Copyright Epic Games, Inc. All Rights Reserved.
+
+#pragma once 
+
+#include "Dataflow/DataflowDebugDrawInterface.h"
+
+struct FDynamicMeshDebugDrawMesh : public IDataflowDebugDrawInterface::IDebugDrawMesh
+{
+	FDynamicMeshDebugDrawMesh(const class UE::Geometry::FDynamicMesh3* DynamicMesh);
+
+	DATAFLOWCORE_API virtual ~FDynamicMeshDebugDrawMesh() = default;
+
+	DATAFLOWCORE_API virtual int32 GetMaxVertexIndex() const override;
+	DATAFLOWCORE_API virtual bool IsValidVertex(int32 VertexIndex) const override;
+	DATAFLOWCORE_API virtual FVector GetVertexPosition(int32 VertexIndex) const override;
+
+	DATAFLOWCORE_API virtual int32 GetMaxTriangleIndex() const override;
+	DATAFLOWCORE_API virtual bool IsValidTriangle(int32 VertexIndex) const override;
+	DATAFLOWCORE_API virtual FIntVector3 GetTriangle(int32 VertexIndex) const override;
+
+private:
+	const class UE::Geometry::FDynamicMesh3* DynamicMesh;
+};
diff --git a/Engine/Source/Runtime/Experimental/Dataflow/Core/Public/Dataflow/DataflowNode.h b/Engine/Source/Runtime/Experimental/Dataflow/Core/Public/Dataflow/DataflowNode.h
--- a/Engine/Source/Runtime/Experimental/Dataflow/Core/Public/Dataflow/DataflowNode.h
+++ b/Engine/Source/Runtime/Experimental/Dataflow/Core/Public/Dataflow/DataflowNode.h
@@ -23,6 +23,7 @@ namespace UE::Dataflow {
 	};
 	class FGraph;
 }
+class IDataflowDebugDrawInterface;
 
 /**
 * FNode
@@ -108,6 +109,16 @@ PRAGMA_ENABLE_DEPRECATION_WARNINGS
 	UE_DEPRECATED(5.4, "FDataflowNode::IsExperimental is deprecated.")
 	virtual bool IsExperimental() { return false; }
 
+	struct FDebugDrawParameters
+	{
+		bool bNodeIsSelected;
+		bool bNodeIsPinned;
+	};
+
+	DATAFLOWCORE_API virtual void DebugDraw(UE::Dataflow::FContext& Context,
+		IDataflowDebugDrawInterface& DataflowRenderingInterface,
+		const FDebugDrawParameters& DebugDrawParameters) const {}
+
 	//
 	// Connections
 	//
diff --git a/Engine/Source/Runtime/Experimental/Dataflow/Core/Public/Dataflow/DataflowSimpleDebugDrawMesh.h b/Engine/Source/Runtime/Experimental/Dataflow/Core/Public/Dataflow/DataflowSimpleDebugDrawMesh.h
new file mode 100644
--- /dev/null
+++ b/Engine/Source/Runtime/Experimental/Dataflow/Core/Public/Dataflow/DataflowSimpleDebugDrawMesh.h
@@ -0,0 +1,20 @@
+// Copyright Epic Games, Inc. All Rights Reserved.
+
+#pragma once 
+
+#include "Dataflow/DataflowDebugDrawInterface.h"
+
+struct FSimpleDebugDrawMesh : public IDataflowDebugDrawInterface::IDebugDrawMesh
+{
+	TArray<FVector> Vertices;
+	TArray<FIntVector3> Triangles;
+
+	DATAFLOWCORE_API virtual ~FSimpleDebugDrawMesh() = default;
+	DATAFLOWCORE_API virtual int32 GetMaxVertexIndex() const override;
+	DATAFLOWCORE_API virtual bool IsValidVertex(int32 VertexIndex) const override;
+	DATAFLOWCORE_API virtual FVector GetVertexPosition(int32 VertexIndex) const override;
+
+	DATAFLOWCORE_API virtual int32 GetMaxTriangleIndex() const override;
+	DATAFLOWCORE_API virtual bool IsValidTriangle(int32 VertexIndex) const override;
+	DATAFLOWCORE_API virtual FIntVector3 GetTriangle(int32 VertexIndex) const override;
+};
diff --git a/Engine/Source/Runtime/Experimental/Dataflow/Engine/Private/Dataflow/DataflowDebugDraw.cpp b/Engine/Source/Runtime/Experimental/Dataflow/Engine/Private/Dataflow/DataflowDebugDraw.cpp
new file mode 100644
--- /dev/null
+++ b/Engine/Source/Runtime/Experimental/Dataflow/Engine/Private/Dataflow/DataflowDebugDraw.cpp
@@ -0,0 +1,181 @@
+// Copyright Epic Games, Inc. All Rights Reserved.
+
+#include "Dataflow/DataflowDebugDraw.h"
+#include "Dataflow/DataflowDebugDrawComponent.h"
+#include "DynamicMesh/DynamicMesh3.h"
+
+FDataflowDebugDraw::FDataflowDebugDraw(FDataflowDebugRenderSceneProxy* DebugRenderSceneProxy) :
+	DebugRenderSceneProxy(DebugRenderSceneProxy)
+{
+	DebugRenderSceneProxy->DrawAlpha = 1;		// uint Alpha values below this are not rendered (default is 100 but we want to allow more translucency)
+	DebugRenderSceneProxy->ClearAll();
+	ResetAllState();
+}
+
+
+void FDataflowDebugDraw::SetColor(const FLinearColor& InColor)
+{
+	Color = InColor;
+	if (bTranslucent)
+	{
+		ColorWithTranslucency = Color.CopyWithNewOpacity(0.25);
+	}
+	else
+	{
+		ColorWithTranslucency = Color;
+	}
+}
+
+void FDataflowDebugDraw::SetLineWidth(double InWidth)
+{
+	LineWidth = InWidth;
+}
+
+void FDataflowDebugDraw::SetWireframe(bool bInWireframe)
+{
+	bWireframe = bInWireframe;
+
+	if (bWireframe && bShaded)
+	{
+		DrawType = FDebugRenderSceneProxy::EDrawType::SolidAndWireMeshes;
+	}
+	else if (bWireframe)
+	{
+		DrawType = FDebugRenderSceneProxy::EDrawType::WireMesh;
+	}
+	else if (bShaded)
+	{
+		DrawType = FDebugRenderSceneProxy::EDrawType::SolidMesh;
+	}
+}
+
+void FDataflowDebugDraw::SetShaded(bool bInShaded)
+{
+	bShaded = bInShaded;
+
+	if (bWireframe && bShaded)
+	{
+		DrawType = FDebugRenderSceneProxy::EDrawType::SolidAndWireMeshes;
+	}
+	else if (bWireframe)
+	{
+		DrawType = FDebugRenderSceneProxy::EDrawType::WireMesh;
+	}
+	else if (bShaded)
+	{
+		DrawType = FDebugRenderSceneProxy::EDrawType::SolidMesh;
+	}
+}
+
+void FDataflowDebugDraw::SetTranslucent(bool bInTranslucent)
+{
+	bTranslucent = bInTranslucent;
+
+	if (bTranslucent)
+	{
+		ColorWithTranslucency = Color.CopyWithNewOpacity(0.25);
+	}
+	else
+	{
+		ColorWithTranslucency = Color;
+	}
+}
+
+void FDataflowDebugDraw::SetForegroundPriority()
+{
+	// TODO: FDebugRenderSceneProxy currently only renders SDPG_World, it should handle SDPG_Foreground as well
+	PriorityGroup = SDPG_Foreground;
+}
+
+void FDataflowDebugDraw::SetWorldPriority()
+{
+	PriorityGroup = SDPG_World;
+}
+
+void FDataflowDebugDraw::ResetAllState()
+{
+	Color = FLinearColor::White;
+	LineWidth = 1.0;
+	bWireframe = true;
+	bShaded = false;
+	bTranslucent = false;
+	PriorityGroup = SDPG_World;
+	ColorWithTranslucency = Color;
+}
+
+void FDataflowDebugDraw::DrawLine(const FVector& Start, const FVector& End) const
+{
+	DebugRenderSceneProxy->Lines.Add(FDebugRenderSceneProxy::FDebugLine(Start, End, ColorWithTranslucency.ToFColor(true), LineWidth));
+}
+
+void FDataflowDebugDraw::DrawMesh(const IDebugDrawMesh& Mesh) const
+{
+	const FColor MeshColor = ColorWithTranslucency.ToFColor(true);
+
+	if (bWireframe)		// FDebugRenderSceneProxy only renders solid meshes
+	{
+		for (int32 TriangleIndex = 0; TriangleIndex < Mesh.GetMaxTriangleIndex(); ++TriangleIndex)
+		{
+			if (Mesh.IsValidTriangle(TriangleIndex))
+			{
+				const FIntVector3 Tri = Mesh.GetTriangle(TriangleIndex);
+				const FVector A = Mesh.GetVertexPosition(Tri[0]);
+				const FVector B = Mesh.GetVertexPosition(Tri[1]);
+				const FVector C = Mesh.GetVertexPosition(Tri[2]);
+				DrawLine(A, B);
+				DrawLine(B, C);
+				DrawLine(C, A);
+			}
+		}
+	}
+	
+	if (bShaded)
+	{
+		FDebugRenderSceneProxy::FMesh SceneProxyMesh;
+
+		SceneProxyMesh.Box = FBox(EForceInit::ForceInit);
+
+		SceneProxyMesh.Vertices.Reserve(Mesh.GetMaxVertexIndex());
+
+		for (int32 VertexIndex = 0; VertexIndex < Mesh.GetMaxVertexIndex(); ++VertexIndex)
+		{
+			if (Mesh.IsValidVertex(VertexIndex))
+			{
+				SceneProxyMesh.Vertices.Add(FDynamicMeshVertex(FVector3f(0.0f), FVector2f(0.f), FColor(0)));
+			}
+			else
+			{
+				const FVector Vertex = Mesh.GetVertexPosition(VertexIndex);
+				SceneProxyMesh.Vertices.Add(FDynamicMeshVertex(FVector3f(Vertex), FVector2f(0.f, 0.f), MeshColor));
+				SceneProxyMesh.Box += Vertex;
+			}
+		}
+
+		SceneProxyMesh.Indices.Reserve(3 * Mesh.GetMaxTriangleIndex());
+		for (int32 TriangleIndex = 0; TriangleIndex < Mesh.GetMaxTriangleIndex(); ++TriangleIndex)
+		{
+			if (Mesh.IsValidTriangle(TriangleIndex))
+			{
+				const FIntVector3 Tri = Mesh.GetTriangle(TriangleIndex);
+				SceneProxyMesh.Indices.Add(Tri[0]);
+				SceneProxyMesh.Indices.Add(Tri[1]);
+				SceneProxyMesh.Indices.Add(Tri[2]);
+			}
+		}
+		
+		SceneProxyMesh.Color = MeshColor;
+
+		DebugRenderSceneProxy->Meshes.Add(SceneProxyMesh);
+	}
+	
+}
+
+void FDataflowDebugDraw::DrawBox(const FVector& Extents, const FQuat& Rotation, const FVector& Center, double UniformScale) const
+{
+	DebugRenderSceneProxy->Boxes.Add(FDebugRenderSceneProxy::FDebugBox(FBox(Center - Extents, Center + Extents), ColorWithTranslucency.ToFColor(true), DrawType, LineWidth));
+}
+
+void FDataflowDebugDraw::DrawSphere(const FVector& Center, double Radius) const
+{
+	DebugRenderSceneProxy->Spheres.Add(FDebugRenderSceneProxy::FSphere(Radius, Center, ColorWithTranslucency, DrawType));
+}
diff --git a/Engine/Source/Runtime/Experimental/Dataflow/Engine/Private/Dataflow/DataflowDebugDrawComponent.cpp b/Engine/Source/Runtime/Experimental/Dataflow/Engine/Private/Dataflow/DataflowDebugDrawComponent.cpp
new file mode 100644
--- /dev/null
+++ b/Engine/Source/Runtime/Experimental/Dataflow/Engine/Private/Dataflow/DataflowDebugDrawComponent.cpp
@@ -0,0 +1,120 @@
+// Copyright Epic Games, Inc. All Rights Reserved.
+
+#include "Dataflow/DataflowDebugDrawComponent.h"
+
+
+FDebugRenderSceneProxy* UDataflowDebugDrawComponent::CreateDebugSceneProxy()
+{
+	return new FDataflowDebugRenderSceneProxy(this);
+}
+
+FBoxSphereBounds UDataflowDebugDrawComponent::CalcBounds(const FTransform& LocalToWorld) const
+{
+	FBoxSphereBounds::Builder BoundsBuilder;
+
+	BoundsBuilder += LocalToWorld.GetLocation();
+
+	if (FDataflowDebugRenderSceneProxy* DebugSceneProxy = static_cast<FDataflowDebugRenderSceneProxy*>(GetSceneProxy()))
+	{	
+		// Some of these bounds are pretty loose but whatever
+
+		for (const FDebugRenderSceneProxy::FDebugLine& Line : DebugSceneProxy->Lines)
+		{
+			BoundsBuilder += Line.Start;
+			BoundsBuilder += Line.End;
+		}
+		for (const FDebugRenderSceneProxy::FDashedLine& DashedLine : DebugSceneProxy->DashedLines)
+		{
+			BoundsBuilder += DashedLine.Start;
+			BoundsBuilder += DashedLine.End;
+		}
+		for (const FDebugRenderSceneProxy::FArrowLine& ArrowLine : DebugSceneProxy->ArrowLines)
+		{
+			FBox ArrowBox(TArray<FVector>{ ArrowLine.Start, ArrowLine.End } );
+			ArrowBox = ArrowBox.ExpandBy(ArrowLine.Mag);
+			BoundsBuilder += ArrowBox;
+		}
+		for (const FDebugRenderSceneProxy::FCircle& Circle : DebugSceneProxy->Circles)
+		{
+			BoundsBuilder += FSphere(Circle.Center, Circle.Radius);
+		}
+		for (const FDebugRenderSceneProxy::FWireCylinder& Cyl : DebugSceneProxy->Cylinders)
+		{
+			BoundsBuilder += FSphere(Cyl.Base, FMath::Max(Cyl.Radius, Cyl.HalfHeight));
+		}
+		for (const FDebugRenderSceneProxy::FWireStar& Star : DebugSceneProxy->Stars)
+		{
+			BoundsBuilder += FSphere(Star.Position, Star.Size);
+		}
+		for (const FDebugRenderSceneProxy::FDebugBox& Box : DebugSceneProxy->Boxes)
+		{
+			BoundsBuilder += Box.Box.ExpandBy(Box.Thickness);
+		}
+		for (const FDebugRenderSceneProxy::FSphere& Sphere : DebugSceneProxy->Spheres)
+		{
+			BoundsBuilder += FSphere(Sphere.Location, Sphere.Radius);
+		}
+		for (const FDebugRenderSceneProxy::FText3d& Text : DebugSceneProxy->Texts)
+		{
+			BoundsBuilder += FSphere(Text.Location, 100.0f);		// ??
+		}
+		for (const FDebugRenderSceneProxy::FCone& Cone : DebugSceneProxy->Cones)
+		{
+			FSphere ConeSphere(FVector(0,0,0), 1.0f);
+			ConeSphere.TransformBy(Cone.ConeToWorld);
+			BoundsBuilder += ConeSphere;
+		}
+		for (const FDebugRenderSceneProxy::FMesh& Mesh : DebugSceneProxy->Meshes)
+		{
+			BoundsBuilder += Mesh.Box;
+		}
+		for (const FDebugRenderSceneProxy::FCapsule& Capsule : DebugSceneProxy->Capsules)
+		{
+			BoundsBuilder += FSphere(Capsule.Base, FMath::Max(Capsule.Radius, Capsule.HalfHeight));
+		}
+		for (const FDebugRenderSceneProxy::FCoordinateSystem& CoordinateSystem : DebugSceneProxy->CoordinateSystems)
+		{
+			BoundsBuilder += FSphere(CoordinateSystem.AxisLoc, CoordinateSystem.Scale + CoordinateSystem.Thickness);
+		}
+	}
+
+	FBoxSphereBounds ReturnBounds(BoundsBuilder);
+	ReturnBounds = ReturnBounds.ExpandBy(5);
+	return ReturnBounds;
+
+}
+
+
+FDataflowDebugRenderSceneProxy::FDataflowDebugRenderSceneProxy(const UPrimitiveComponent* InComponent) :
+	FDebugRenderSceneProxy(InComponent)
+{
+}
+
+void FDataflowDebugRenderSceneProxy::ClearAll()
+{
+	Lines.Reset();
+	DashedLines.Reset();
+	ArrowLines.Reset();
+	Circles.Reset();
+	Cylinders.Reset();
+	Stars.Reset();
+	Boxes.Reset();
+	Spheres.Reset();
+	Texts.Reset();
+	Cones.Reset();
+	Meshes.Reset();
+	Capsules.Reset();
+	CoordinateSystems.Reset();
+}
+
+
+FPrimitiveViewRelevance FDataflowDebugRenderSceneProxy::GetViewRelevance(const FSceneView* View) const
+{
+	FPrimitiveViewRelevance Result;
+	Result.bDrawRelevance = IsShown(View);
+	Result.bDynamicRelevance = true;
+	// ideally the TranslucencyRelevance should be filled out by the material, here we do it conservative
+	Result.bSeparateTranslucency = Result.bNormalTranslucency = true;
+	return Result;
+}
+
diff --git a/Engine/Source/Runtime/Experimental/Dataflow/Engine/Public/Dataflow/DataflowDebugDraw.h b/Engine/Source/Runtime/Experimental/Dataflow/Engine/Public/Dataflow/DataflowDebugDraw.h
new file mode 100644
--- /dev/null
+++ b/Engine/Source/Runtime/Experimental/Dataflow/Engine/Public/Dataflow/DataflowDebugDraw.h
@@ -0,0 +1,47 @@
+// Copyright Epic Games, Inc. All Rights Reserved.
+
+#pragma once 
+
+#include "Dataflow/DataflowDebugDrawInterface.h"
+#include "DebugRenderSceneProxy.h"
+#include "Engine/EngineTypes.h"
+#include "Math/Color.h"
+
+class FDataflowDebugDraw : public IDataflowDebugDrawInterface
+{
+public:
+
+	DATAFLOWENGINE_API FDataflowDebugDraw(class FDataflowDebugRenderSceneProxy* DebugRenderSceneProxy);
+
+	virtual ~FDataflowDebugDraw() = default;
+
+private:
+
+	virtual void SetColor(const FLinearColor& InColor) override;
+	virtual void SetLineWidth(double Width) override;
+	virtual void SetWireframe(bool bInWireframe) override;
+	virtual void SetShaded(bool bInShaded) override;
+	virtual void SetTranslucent(bool bInShadedTranslucent) override;
+	virtual void SetForegroundPriority() override;
+	virtual void SetWorldPriority() override;
+	virtual void ResetAllState() override;
+
+	virtual void DrawLine(const FVector& Start, const FVector& End) const override;
+	virtual void DrawMesh(const IDebugDrawMesh& Mesh) const override;
+
+	virtual void DrawBox(const FVector& Extents, const FQuat& Rotation, const FVector& Center, double UniformScale) const override;
+	virtual void DrawSphere(const FVector& Center, double Radius) const override;
+
+	class FDataflowDebugRenderSceneProxy* DebugRenderSceneProxy = nullptr;
+
+	FLinearColor Color = FLinearColor::White;
+	double LineWidth = 1.0;
+	bool bWireframe = true;
+	bool bShaded = false;
+	bool bTranslucent = false;
+	ESceneDepthPriorityGroup PriorityGroup = SDPG_World;
+
+	FLinearColor ColorWithTranslucency = FLinearColor::White;
+	FDebugRenderSceneProxy::EDrawType DrawType = FDebugRenderSceneProxy::EDrawType::Invalid;
+};
+
diff --git a/Engine/Source/Runtime/Experimental/Dataflow/Engine/Public/Dataflow/DataflowDebugDrawComponent.h b/Engine/Source/Runtime/Experimental/Dataflow/Engine/Public/Dataflow/DataflowDebugDrawComponent.h
new file mode 100644
--- /dev/null
+++ b/Engine/Source/Runtime/Experimental/Dataflow/Engine/Public/Dataflow/DataflowDebugDrawComponent.h
@@ -0,0 +1,32 @@
+// Copyright Epic Games, Inc. All Rights Reserved.
+
+#pragma once 
+
+#include "Debug/DebugDrawComponent.h"
+#include "DataflowDebugDrawComponent.generated.h"
+
+UCLASS(MinimalAPI)
+class UDataflowDebugDrawComponent : public UDebugDrawComponent
+{
+	GENERATED_BODY()
+
+private:
+
+	virtual FDebugRenderSceneProxy* CreateDebugSceneProxy() override;
+	virtual FBoxSphereBounds CalcBounds(const FTransform& LocalToWorld) const override;
+};
+
+
+class FDataflowDebugRenderSceneProxy final : public FDebugRenderSceneProxy
+{
+public:
+
+	FDataflowDebugRenderSceneProxy(const UPrimitiveComponent* InComponent);
+
+	/** Remove all primitives stored on this proxy */
+	void ClearAll();
+
+private:
+
+	virtual FPrimitiveViewRelevance GetViewRelevance(const FSceneView* View) const override;
+};
