Fix TransformBy discarded

#rnx

diff --git a/Engine/Source/Runtime/Experimental/Dataflow/Engine/Private/Dataflow/DataflowDebugDrawComponent.cpp b/Engine/Source/Runtime/Experimental/Dataflow/Engine/Private/Dataflow/DataflowDebugDrawComponent.cpp
--- a/Engine/Source/Runtime/Experimental/Dataflow/Engine/Private/Dataflow/DataflowDebugDrawComponent.cpp
+++ b/Engine/Source/Runtime/Experimental/Dataflow/Engine/Private/Dataflow/DataflowDebugDrawComponent.cpp
@@ -61,7 +61,7 @@ FBoxSphereBounds UDataflowDebugDrawComponent::CalcBounds(const FTransform& Local
 		for (const FDebugRenderSceneProxy::FCone& Cone : DebugSceneProxy->Cones)
 		{
 			FSphere ConeSphere(FVector(0,0,0), 1.0f);
-			ConeSphere.TransformBy(Cone.ConeToWorld);
+			ConeSphere = ConeSphere.TransformBy(Cone.ConeToWorld);
 			BoundsBuilder += ConeSphere;
 		}
 		for (const FDebugRenderSceneProxy::FMesh& Mesh : DebugSceneProxy->Meshes)
