Consolidating latency testing code to a single code file.

#rnx

- Added a header for LatencyTesting and moved the function declerations there.
- Rename the latency tests that don't return latency values to 'ConnectionTest' to better describe their functionality.
-- Moved these tests from OnDemandIoDispatcherBackend.cpp to LatencyTesting.cpp.
- Added code documentation as some of the function parameters are not obvious.

diff --git a/Engine/Source/Runtime/Experimental/IoStore/OnDemand/Private/LatencyTesting.cpp b/Engine/Source/Runtime/Experimental/IoStore/OnDemand/Private/LatencyTesting.cpp
--- a/Engine/Source/Runtime/Experimental/IoStore/OnDemand/Private/LatencyTesting.cpp
+++ b/Engine/Source/Runtime/Experimental/IoStore/OnDemand/Private/LatencyTesting.cpp
@@ -1,5 +1,7 @@
 // Copyright Epic Games, Inc. All Rights Reserved.
 
+#include "LatencyTesting.h"
+
 #include "Containers/StringConv.h"
 #include "Containers/StringView.h"
 #include "HAL/PlatformTime.h"
@@ -12,7 +14,7 @@
 namespace UE::IoStore
 {
 
-void LatencyTest(FAnsiStringView InUrl, FStringView InPath, uint32 InTimeOutMs, TArrayView<int32> OutResults)
+void LatencyTest(FAnsiStringView InUrl, FStringView InPath, uint32 TimeOutMs, TArrayView<int32> OutResults)
 {
 	using namespace HTTP;
 
@@ -34,7 +36,7 @@ void LatencyTest(FAnsiStringView InUrl, FStringView InPath, uint32 InTimeOutMs,
 	AnsiPath << InPath;
 
 	FEventLoop Loop;
-	Loop.SetFailTimeout(InTimeOutMs);
+	Loop.SetFailTimeout(TimeOutMs);
 
 	for (int32& Result : OutResults)
 	{
@@ -69,7 +71,7 @@ void LatencyTest(FAnsiStringView InUrl, FStringView InPath, uint32 InTimeOutMs,
 
 		uint64 Cycles = FPlatformTime::Cycles64();
 		
-		while (Loop.Tick(InTimeOutMs) != 0)
+		while (Loop.Tick(TimeOutMs) != 0)
 		{
 
 		}
@@ -97,4 +99,43 @@ void LatencyTest(FStringView InUrl, FStringView InPath, uint32 InTimeOutMs, TArr
 	LatencyTest(AnsiUrl, InPath, InTimeOutMs, OutResults);
 }
 
+bool ConnectionTest(FAnsiStringView Url, FStringView Path, uint32 TimeoutMs)
+{
+	TRACE_CPUPROFILER_EVENT_SCOPE(IasBackend::ConnectionTest);
+
+	int32 Results[4] = {};
+	LatencyTest(Url, Path, TimeoutMs, MakeArrayView(Results));
+
+	if (Results[0] >= 0 || Results[1] >= 0 || Results[2] >= 0 || Results[3] >= 0)
+	{
+#if !UE_BUILD_SHIPPING
+		UE_LOGFMT(LogIas, Log, "Endpoint '{Url}' latency test (ms): {Result0} {Result1} {Result2} {Result3}",
+			Url, Results[0], Results[1], Results[2], Results[3]);
+#endif // !UE_BUILD_SHIPPING
+
+		return true;
+	}
+	else
+	{
+		return false;
+	}
+}
+
+int32 ConnectionTest(TConstArrayView<FAnsiString> Urls, FStringView Path, uint32 TimeoutMs, std::atomic_bool& bCancel)
+{
+	TRACE_CPUPROFILER_EVENT_SCOPE(IasBackend::ConnectionTest);
+
+	for (int32 Idx = 0; Idx < Urls.Num() && !bCancel.load(std::memory_order_relaxed); ++Idx)
+	{
+		int32 LatencyMs = -1;
+		LatencyTest(Urls[Idx], Path, TimeoutMs, MakeArrayView(&LatencyMs, 1));
+		if (LatencyMs >= 0)
+		{
+			return Idx;
+		}
+	}
+
+	return INDEX_NONE;
+}
+
 } // namespace UE::IoStore
diff --git a/Engine/Source/Runtime/Experimental/IoStore/OnDemand/Private/LatencyTesting.h b/Engine/Source/Runtime/Experimental/IoStore/OnDemand/Private/LatencyTesting.h
new file mode 100644
--- /dev/null
+++ b/Engine/Source/Runtime/Experimental/IoStore/OnDemand/Private/LatencyTesting.h
@@ -0,0 +1,47 @@
+// Copyright Epic Games, Inc. All Rights Reserved.
+
+#pragma once
+
+#include "Containers/ArrayView.h"
+#include "Containers/StringFwd.h"
+
+namespace UE::IoStore
+{
+/**
+ * Tests a single url.
+ * 
+ * @param Url			The url to be tested
+ * @param Path			The path on the CDN to be tested
+ * @param TimeoutMs		How long to wait for a response (in ms) before timing out
+ * @param OutResults	A pre-sized array to hold the results. One latency test will be made per element in this container.
+ */
+void LatencyTest(FAnsiStringView Url, FStringView Path, uint32 TimeoutMs, TArrayView<int32> OutResults);
+
+/** Convenience overload for the ansi version of LatencyTest */
+void LatencyTest(FStringView Url, FStringView Path, uint32 TimeoutMs, TArrayView<int32> OutResults);
+
+/**
+ * An easy way to test if a valid connection can be made to a CDN or not. Instead of returning the latency to the CDN this
+ * function returns if the CDN can be contacted at all.
+ * 
+ * @param Url			The url to be tested
+ * @param Path			The path on the CDN to be tested
+ * @param TimeoutMs		How long to wait for a response (in ms) before timing out
+ * 
+ * @return True if the CDN can be reached and false if it cannot.
+ */
+bool ConnectionTest(FAnsiStringView Url, FStringView Path, uint32 TimeoutMs);
+
+/**
+ * Runs latency tests on a list of urls to find the first CDN that can be reached.
+ * 
+ * @param Urls			A list of urls that should be tested
+ * @param Path			The path on the CDN(s) to be tested
+ * @param TimeoutMs		How long to wait for a response (in ms) before timing out
+ * @param bCancel		Provides a way for the caller to cancel the tests. When this bool is set to true the test will
+ *						attempt to exit as fast as possible, usually when the next latency test has completed.
+ * @return				The index of the CDN that was successfully reached. INDEX_NONE will be returned if no CDN could be reached.
+ */
+int32 ConnectionTest(TConstArrayView<FAnsiString> Urls, FStringView Path, uint32 TimeoutMs, std::atomic_bool& bCancel);
+
+} // namespace UE::IoStore
diff --git a/Engine/Source/Runtime/Experimental/IoStore/OnDemand/Private/OnDemandIoDispatcherBackend.cpp b/Engine/Source/Runtime/Experimental/IoStore/OnDemand/Private/OnDemandIoDispatcherBackend.cpp
--- a/Engine/Source/Runtime/Experimental/IoStore/OnDemand/Private/OnDemandIoDispatcherBackend.cpp
+++ b/Engine/Source/Runtime/Experimental/IoStore/OnDemand/Private/OnDemandIoDispatcherBackend.cpp
@@ -13,7 +13,6 @@
 #include "HAL/PreprocessorHelpers.h"
 #include "HAL/Runnable.h"
 #include "HAL/RunnableThread.h"
-#include "IasCache.h"
 #include "IO/IoAllocators.h"
 #include "IO/IoChunkEncoding.h"
 #include "IO/IoDispatcher.h"
@@ -21,6 +20,8 @@
 #include "IO/IoStatus.h"
 #include "IO/IoStore.h"
 #include "IO/IoStoreOnDemand.h"
+#include "IasCache.h"
+#include "LatencyTesting.h"
 #include "Logging/StructuredLog.h"
 #include "Math/NumericLimits.h"
 #include "Misc/ConfigCacheIni.h"
@@ -54,8 +55,6 @@ namespace UE::IoStore
 ///////////////////////////////////////////////////////////////////////////////
 
 extern FString GIasOnDemandTocExt;
-void LatencyTest(FStringView, FStringView, uint32, TArrayView<int32>);
-void LatencyTest(FAnsiStringView, FStringView, uint32, TArrayView<int32>);
 
 ///////////////////////////////////////////////////////////////////////////////
 /** Note that GIasHttpPrimaryEndpoint has no effect after initial start up */
@@ -274,48 +273,6 @@ FAutoConsoleVariableRef CVarOnDemandBackendThreadPriority(
 	return GCompleteMaterializeTaskPriorities[FMath::Clamp(GCompleteMaterializeTaskPriority, 0, UE_ARRAY_COUNT(GCompleteMaterializeTaskPriorities) - 1)];
 }
 
-///////////////////////////////////////////////////////////////////////////////
-
-static bool LatencyTest(FAnsiStringView Url, FStringView Path)
-{
-	TRACE_CPUPROFILER_EVENT_SCOPE(IasBackend::LatencyTest);
-
-	int32 Results[4] = {};
-	LatencyTest(Url, Path, GIasHttpTimeOutMs, MakeArrayView(Results));
-
-	if (Results[0] >= 0 || Results[1] >= 0 || Results[2] >= 0 || Results[3] >= 0)
-	{
-#if !UE_BUILD_SHIPPING
-		UE_LOGFMT(LogIas, Log, "Endpoint '{Url}' latency test (ms): {Result0} {Result1} {Result2} {Result3}",
-			Url, Results[0], Results[1], Results[2], Results[3]);
-#endif // !UE_BUILD_SHIPPING
-
-		return true;
-	}
-	else
-	{
-		return false;
-	}
-}
-
-///////////////////////////////////////////////////////////////////////////////
-static int32 LatencyTest(TConstArrayView<FAnsiString> Urls, FStringView Path, std::atomic_bool& bCancel)
-{
-	TRACE_CPUPROFILER_EVENT_SCOPE(IasBackend::LatencyTest);
-
-	for (int32 Idx = 0; Idx < Urls.Num() && !bCancel.load(std::memory_order_relaxed); ++Idx)
-	{
-		int32 LatencyMs = -1;
-		LatencyTest(Urls[Idx], Path, GIasHttpTimeOutMs, MakeArrayView(&LatencyMs, 1));
-		if (LatencyMs >= 0)
-		{
-			return Idx;
-		}
-	}
-
-	return INDEX_NONE;
-}
-
 ///////////////////////////////////////////////////////////////////////////////
 struct FBitWindow
 {
@@ -1760,7 +1717,7 @@ uint32 FOnDemandIoBackend::Run()
 
 	const double InitStart = FPlatformTime::Seconds();
 
-	if (LatencyTest(DefaultHostGroup.PrimaryHost(), GetEndpointTestPath()))
+	if (ConnectionTest(DefaultHostGroup.PrimaryHost(), GetEndpointTestPath(), GIasHttpTimeOutMs))
 	{
 		BackendStatus.SetHttpEnabled(true);
 		Stats.OnHttpConnected();
@@ -1796,7 +1753,7 @@ uint32 FOnDemandIoBackend::Run()
 
 				UE_LOG(LogIas, Log, TEXT("Trying to reconnect to any available endpoint..."));
 
-				if (const int32 Idx = LatencyTest(DefaultHostGroup.Hosts(), GetEndpointTestPath(), bStopRequested); Idx != INDEX_NONE)
+				if (const int32 Idx = ConnectionTest(DefaultHostGroup.Hosts(), GetEndpointTestPath(), GIasHttpTimeOutMs, bStopRequested); Idx != INDEX_NONE)
 				{
 					Stats.OnHttpConnected();
 
@@ -1807,7 +1764,7 @@ uint32 FOnDemandIoBackend::Run()
 			}
 			else if (DefaultHostGroup.PrimaryHostIndex() != 0)
 			{
-				if (const int32 Idx = LatencyTest(DefaultHostGroup.Hosts().Left(1), GetEndpointTestPath(), bStopRequested); Idx != INDEX_NONE)
+				if (const int32 Idx = ConnectionTest(DefaultHostGroup.Hosts().Left(1), GetEndpointTestPath(), GIasHttpTimeOutMs, bStopRequested); Idx != INDEX_NONE)
 				{
 					DefaultHostGroup.SetPrimaryHost(Idx);
 					UE_LOGFMT(LogIas, Log, "Reconnected to primary host '{Url}'", DefaultHostGroup.PrimaryHost());
