Support for opting out of a thirdparty library.

#rnx

diff --git a/Engine/Source/Runtime/Experimental/IoStore/HttpClient/IoStoreHttpClient.Build.cs b/Engine/Source/Runtime/Experimental/IoStore/HttpClient/IoStoreHttpClient.Build.cs
--- a/Engine/Source/Runtime/Experimental/IoStore/HttpClient/IoStoreHttpClient.Build.cs
+++ b/Engine/Source/Runtime/Experimental/IoStore/HttpClient/IoStoreHttpClient.Build.cs
@@ -8,11 +8,9 @@ public class IoStoreHttpClient : ModuleRules
 	{
 		PublicDependencyModuleNames.Add("Core");
 		PublicDependencyModuleNames.Add("TraceLog");
-		PrivateDependencyModuleNames.AddRange(
-			new string[] {
-				"OpenSSL",
-			}
-		);
+
+		PrivateDependencyModuleNames.Add("OpenSSL");
+		PrivateDefinitions.Add("IAS_HTTP_HAS_OPENSSL=1");
 
 		CppCompileWarningSettings.UnsafeTypeCastWarningLevel = WarningLevel.Error; 
 	}
diff --git a/Engine/Source/Runtime/Experimental/IoStore/HttpClient/Private/Http.cpp b/Engine/Source/Runtime/Experimental/IoStore/HttpClient/Private/Http.cpp
--- a/Engine/Source/Runtime/Experimental/IoStore/HttpClient/Private/Http.cpp
+++ b/Engine/Source/Runtime/Experimental/IoStore/HttpClient/Private/Http.cpp
@@ -18,10 +18,6 @@
 #include <Trace/Trace.h>
 #endif
 
-#include <openssl/engine.h>
-#include <openssl/err.h>
-#include <openssl/ssl.h>
-
 #include <atomic>
 
 #include "Misc.inl"
diff --git a/Engine/Source/Runtime/Experimental/IoStore/HttpClient/Private/Peer.inl b/Engine/Source/Runtime/Experimental/IoStore/HttpClient/Private/Peer.inl
--- a/Engine/Source/Runtime/Experimental/IoStore/HttpClient/Private/Peer.inl
+++ b/Engine/Source/Runtime/Experimental/IoStore/HttpClient/Private/Peer.inl
@@ -2,23 +2,22 @@
 
 #pragma once
 
-#if defined(IAS_HTTP_EXPLICIT_VERIFY_TIME)
-#	if !IAS_HTTP_EXPLICIT_VERIFY_TIME 
-#		error Either define this to >=1 or not at all
+#if defined(IAS_HTTP_HAS_OPENSSL)
+#	if defined(IAS_HTTP_EXPLICIT_VERIFY_TIME)
+#		if !IAS_HTTP_EXPLICIT_VERIFY_TIME 
+#			error Either define this to >=1 or not at all
+#		endif
+#		include <HAL/PlatformTime.h>
+#		include <ctime>
 #	endif
-#	include <HAL/PlatformTime.h>
-
-#	include <ctime>
+#	include <openssl/engine.h>
+#	include <openssl/err.h>
+#	include <openssl/ssl.h>
 #endif
 
 namespace UE::IoStore::HTTP
 {
 
 ////////////////////////////////////////////////////////////////////////////////
-static_assert(OPENSSL_VERSION_NUMBER >= 0x10100000L, "version supporting autoinit required");
-
-
-
-////////////////////////////////////////////////////////////////////////////////
 static FCertRoots GDefaultCertRoots;
 
@@ -28,20 +27,42 @@
 	static const FCertRootsRef	Default = ~0ull;
 };
 
+
+
 ////////////////////////////////////////////////////////////////////////////////
-FCertRoots::~FCertRoots()
+class FPeer
 {
-	if (Handle == 0)
-	{
-		return;
-	}
+public:
+				FPeer() = default;
+				FPeer(FSocket InSocket);
+	FWaitable	GetWaitable() const					{ return Socket.GetWaitable(); }
+	FOutcome	Send(const char* Data, int32 Size)	{ return Socket.Send(Data, Size); }
+	FOutcome	Recv(char* Out, int32 MaxSize)		{ return Socket.Recv(Out, MaxSize); }
+	bool		IsValid() const						{ return Socket.IsValid(); }
 
+private:
+	FSocket		Socket;
+};
+
+////////////////////////////////////////////////////////////////////////////////
+FPeer::FPeer(FSocket InSocket)
+: Socket(MoveTemp(InSocket))
+{
+}
+
+
+
+#if defined(IAS_HTTP_HAS_OPENSSL)
+
+////////////////////////////////////////////////////////////////////////////////
+static void Ssl_ContextDestroy(UPTRINT Handle)
+{
 	auto* Context = (SSL_CTX*)Handle;
 	SSL_CTX_free(Context);
 }
 
 ////////////////////////////////////////////////////////////////////////////////
-FCertRoots::FCertRoots(FMemoryView PemData)
+static UPTRINT Ssl_ContextCreate(FMemoryView PemData)
 {
 	if (static bool InitOnce = false; !InitOnce)
 	{
@@ -84,7 +105,7 @@
 	if (NumAdded == 0)
 	{
 		SSL_CTX_free(Context);
-		return;
+		return 0;
 	}
 
 #if defined(IAS_HTTP_EXPLICIT_VERIFY_TIME)
@@ -108,17 +129,12 @@
 	}
 #endif
 
-	Handle = UPTRINT(Context);
+	return UPTRINT(Context);
 }
 
 ////////////////////////////////////////////////////////////////////////////////
-int32 FCertRoots::Num() const
+static int32 Ssl_ContextCertNum(UPTRINT Handle)
 {
-	if (Handle == 0)
-	{
-		return -1;
-	}
-
 	auto* Context = (SSL_CTX*)Handle;
 	X509_STORE* Store = SSL_CTX_get_cert_store(Context);
 	STACK_OF(X509_OBJECT)* Objects = X509_STORE_get0_objects(Store);
@@ -126,89 +142,61 @@
 }
 
 ////////////////////////////////////////////////////////////////////////////////
-void FCertRoots::SetDefault(FCertRoots&& CertRoots)
+static int32 Ssl_BioWrite(BIO* Bio, const char* Data, size_t Size, size_t* BytesWritten)
 {
-	check(GDefaultCertRoots.IsValid() != CertRoots.IsValid());
-	GDefaultCertRoots = MoveTemp(CertRoots);
+	*BytesWritten = 0;
+	BIO_clear_retry_flags(Bio);
+
+	auto* Peer = (FPeer*)BIO_get_data(Bio);
+	FOutcome Outcome = Peer->Send(Data, Size);
+	if (Outcome.IsWaiting())
+	{
+		BIO_set_retry_write(Bio);
+		return 0;
+	}
+
+	if (Outcome.IsError())
+	{
+		return -1;
+	}
+
+	*BytesWritten = Outcome.GetResult();
+	return 1;
 }
 
 ////////////////////////////////////////////////////////////////////////////////
-FCertRootsRef FCertRoots::NoTls()
+static int32 Ssl_BioRead(BIO* Bio, char* Data, size_t Size, size_t* BytesRead)
 {
-	return ECertRootsRefType::None;
+	*BytesRead = 0;
+	BIO_clear_retry_flags(Bio);
+
+	auto* Peer = (FPeer*)BIO_get_data(Bio);
+	FOutcome Outcome = Peer->Recv(Data, Size);
+	if (Outcome.IsWaiting())
+	{
+		BIO_set_retry_read(Bio);
+		return 0;
+	}
+
+	if (Outcome.IsError())
+	{
+		return -1;
+	}
+
+	*BytesRead = Outcome.GetResult();
+	return 1;
 }
 
 ////////////////////////////////////////////////////////////////////////////////
-FCertRootsRef FCertRoots::Default()
+static long Ssl_BioControl(BIO*, int Cmd, long, void*)
 {
-	return ECertRootsRefType::Default;
+	return (Cmd == BIO_CTRL_FLUSH) ? 1 : 0;
 }
 
 ////////////////////////////////////////////////////////////////////////////////
-FCertRootsRef FCertRoots::Explicit(const FCertRoots& CertRoots)
+static SSL* Ssl_Create(FCertRootsRef Certs, const char* HostName=nullptr)
 {
-	check(CertRoots.IsValid());
-	return CertRoots.Handle;
-}
+	static_assert(OPENSSL_VERSION_NUMBER >= 0x10100000L, "version supporting autoinit required");
 
-
-
-////////////////////////////////////////////////////////////////////////////////
-class FPeer
-{
-public:
-				FPeer() = default;
-				FPeer(FSocket InSocket);
-	FWaitable	GetWaitable() const					{ return Socket.GetWaitable(); }
-	FOutcome	Send(const char* Data, int32 Size)	{ return Socket.Send(Data, Size); }
-	FOutcome	Recv(char* Out, int32 MaxSize)		{ return Socket.Recv(Out, MaxSize); }
-	bool		IsValid() const						{ return Socket.IsValid(); }
-
-private:
-	FSocket		Socket;
-};
-
-////////////////////////////////////////////////////////////////////////////////
-FPeer::FPeer(FSocket InSocket)
-: Socket(MoveTemp(InSocket))
-{
-}
-
-
-
-////////////////////////////////////////////////////////////////////////////////
-class FTlsPeer
-	: public FPeer
-{
-public:
-				FTlsPeer()								= default;
-				~FTlsPeer();
-				FTlsPeer(FTlsPeer&& Rhs)				{ Move(MoveTemp(Rhs)); }
-				FTlsPeer& operator = (FTlsPeer&& Rhs)	{ return Move(MoveTemp(Rhs)); }
-				FTlsPeer(FSocket InSocket, FCertRootsRef Certs=ECertRootsRefType::None, const char* HostName=nullptr);
-	FTlsPeer&	Move(FTlsPeer&& Rhs);
-	FOutcome	Handshake();
-	FOutcome	Send(const char* Data, int32 Size);
-	FOutcome	Recv(char* Out, int32 MaxSize);
-
-protected:
-	SSL*		Ssl = nullptr;
-
-private:
-	FOutcome	GetOutcome(int32 SslResult, const char* Message="tls error") const;
-	int32		BioWrite(const char* Data, size_t Size, size_t* BytesWritten, BIO* Bio);
-	int32		BioRead(char* Data, size_t Size, size_t* BytesRead, BIO* Bio);
-	long		BioControl(int Cmd, long, void*, BIO*);
-};
-
-////////////////////////////////////////////////////////////////////////////////
-FTlsPeer::FTlsPeer(FSocket InSocket, FCertRootsRef Certs, const char* HostName)
-: FPeer(MoveTemp(InSocket))
-{
-	if (Certs == ECertRootsRefType::None)
-	{
-		return;
-	}
-
 	if (Certs == ECertRootsRefType::Default)
 	{
@@ -222,23 +210,16 @@
 	{
 		int32 BioId = BIO_get_new_index() | BIO_TYPE_SOURCE_SINK;
 		BioMethod = BIO_meth_new(BioId, "IasBIO");
-
-#define METH_THUNK(x_, y_)\
-		x_(BioMethod, [] <typename... T> (BIO* b, T... t) { \
-			return (decltype(this)(BIO_get_data(b)))->y_(Forward<T>(t)..., b); \
-		})
-		METH_THUNK(BIO_meth_set_write_ex,	BioWrite);
-		METH_THUNK(BIO_meth_set_read_ex,	BioRead);
-		METH_THUNK(BIO_meth_set_ctrl,		BioControl);
-#undef METH_THUNK
+		BIO_meth_set_write_ex(BioMethod, Ssl_BioWrite);
+		BIO_meth_set_read_ex(BioMethod,	Ssl_BioRead);
+		BIO_meth_set_ctrl(BioMethod, Ssl_BioControl);
 	}
 
 	BIO* Bio = BIO_new(BioMethod);
-	BIO_set_data(Bio, this);
 
 	// SSL_MODE_ENABLE_PARTIAL_WRITE ??!!!
 
-	Ssl = SSL_new(Context);
+	SSL* Ssl = SSL_new(Context);
 	SSL_set_connect_state(Ssl);
 	SSL_set0_rbio(Ssl, Bio);
 	SSL_set0_wbio(Ssl, Bio);
@@ -248,50 +229,89 @@
 	{
 		SSL_set_tlsext_host_name(Ssl, HostName);
 	}
+
+	return Ssl;
 }
 
 ////////////////////////////////////////////////////////////////////////////////
-FTlsPeer& FTlsPeer::Move(FTlsPeer&& Rhs)
+static void Ssl_Destroy(SSL* Ssl)
 {
-	FPeer::operator = (MoveTemp(Rhs));
+	SSL_free(Ssl);
+}
 
-	Swap(Ssl, Rhs.Ssl);
+////////////////////////////////////////////////////////////////////////////////
+static void Ssl_AssociatePeer(SSL* Ssl, FPeer* Peer)
+{
+	BIO* Bio = SSL_get_rbio(Ssl);
+	check(Bio == SSL_get_wbio(Ssl));
+	BIO_set_data(Bio, Peer);
+}
 
-	auto PatchBio = [] (FTlsPeer& Peer)
+////////////////////////////////////////////////////////////////////////////////
+static void Ssl_SetupAlpn(SSL* Ssl, int32 /*MaxHttpVersion*/)
+{
+	static const uint8 AlpnProtos[] =
+		"\x08" "http/1.1"
+	//	"\x02" "h2"
+		;
+	SSL_set_alpn_protos(Ssl, AlpnProtos, sizeof(AlpnProtos) - 1);
+}
+
+////////////////////////////////////////////////////////////////////////////////
+static int32 Ssl_GetProtocolVersion(SSL* Ssl)
+{
+	int32 Proto = -1;
+
+	const char* AlpnProto = nullptr;
+	uint32 AlpnProtoLen;
+	SSL_get0_alpn_selected(Ssl, &(const uint8*&)AlpnProto, &AlpnProtoLen);
+	if (AlpnProto == nullptr)
 	{
-		if (Peer.Ssl != nullptr)
-		{
-			BIO* Bio = SSL_get_rbio(Peer.Ssl);
-			check(Bio != nullptr);
-			BIO_set_data(Bio, &Peer);
-		}
+		return Proto;
+	}
+
+	FAnsiStringView Needle(AlpnProto, AlpnProtoLen);
+	FAnsiStringView Candidates[] = {
+		"http/1.1",
+	//	"h2",
 	};
-	PatchBio(*this);
-	PatchBio(Rhs);
+	for (int32 i = 0; i < UE_ARRAY_COUNT(Candidates); ++i)
+	{
+		const FAnsiStringView& Candidate = Candidates[i];
+		if (AlpnProtoLen != Candidate.Len())
+		{
+			continue;
+		}
 
-	return *this;
+		if (Candidate != Needle)
+		{
+			continue;
+		}
+
+		Proto = i + 1;
+		break;
+	}
+
+	return Proto;
 }
 
 ////////////////////////////////////////////////////////////////////////////////
-FTlsPeer::~FTlsPeer()
+static FOutcome Ssl_GetOutcome(SSL* Ssl, int32 SslResult, const char* Message="tls error")
 {
-	if (Ssl != nullptr)
+	int32 Error = SSL_get_error(Ssl, SslResult);
+	if (Error != SSL_ERROR_WANT_READ && Error != SSL_ERROR_WANT_WRITE)
 	{
-		SSL_free(Ssl);
+		return FOutcome::Error(Message, Error);
 	}
+	return FOutcome::Waiting();
 }
 
 ////////////////////////////////////////////////////////////////////////////////
-FOutcome FTlsPeer::Handshake()
+static FOutcome Ssl_Handshake(SSL* Ssl)
 {
-	if (Ssl == nullptr)
-	{
-		return FOutcome::Ok();
-	}
-
 	int32 Result = SSL_do_handshake(Ssl);
 	if (Result == 0) return FOutcome::Error("unsuccessful tls handshake");
-	if (Result != 1) return GetOutcome(Result, "tls handshake error");
+	if (Result != 1) return Ssl_GetOutcome(Ssl, Result, "tls handshake error");
 
 	if (Result = SSL_get_verify_result(Ssl); Result != X509_V_OK)
 	{
@@ -302,6 +322,161 @@
 }
 
 ////////////////////////////////////////////////////////////////////////////////
+static FOutcome Ssl_Write(SSL* Ssl, const char* Data, int32 Size)
+{
+	int32 Result = SSL_write(Ssl, Data, Size);
+	return (Result > 0) ? FOutcome::Ok(Result) : Ssl_GetOutcome(Ssl, Result);
+}
+
+////////////////////////////////////////////////////////////////////////////////
+static FOutcome Ssl_Read(SSL* Ssl, char* Out, int32 MaxSize)
+{
+	int32 Result = SSL_read(Ssl, Out, MaxSize);
+	return (Result > 0) ? FOutcome::Ok(Result) : Ssl_GetOutcome(Ssl, Result);
+}
+
+#else
+
+struct SSL;
+static void		Ssl_ContextDestroy(...)		{}
+static UPTRINT	Ssl_ContextCreate(...)		{ return 0; }
+static int32	Ssl_ContextCertNum(...)		{ return 0; }
+static int32	Ssl_BioWrite(...)			{ return -1; }
+static int32	Ssl_BioRead(...)			{ return -1; }
+static long		Ssl_BioControl(...)			{ return 0; }
+static SSL*		Ssl_Create(...)				{ return nullptr; }
+static void		Ssl_Destroy(...)			{}
+static void		Ssl_AssociatePeer(...)		{}
+static void		Ssl_SetupAlpn(...)			{}
+static int32	Ssl_GetProtocolVersion(...)	{ return -1; }
+static FOutcome	Ssl_GetOutcome(...)			{ return FOutcome::Error("!impl"); }
+static FOutcome	Ssl_Handshake(...)			{ return FOutcome::Error("!impl"); }
+static FOutcome	Ssl_Write(...)				{ return FOutcome::Error("!impl"); }
+static FOutcome	Ssl_Read(...)				{ return FOutcome::Error("!impl"); }
+
+#endif // IAS_HTTP_HAS_OPENSSL
+
+
+
+////////////////////////////////////////////////////////////////////////////////
+FCertRoots::~FCertRoots()
+{
+	if (Handle != 0)
+	{
+		Ssl_ContextDestroy(Handle);
+	}
+}
+
+////////////////////////////////////////////////////////////////////////////////
+FCertRoots::FCertRoots(FMemoryView PemData)
+{
+	Handle = Ssl_ContextCreate(PemData);
+}
+
+////////////////////////////////////////////////////////////////////////////////
+int32 FCertRoots::Num() const
+{
+	if (Handle == 0)
+	{
+		return -1;
+	}
+
+	return Ssl_ContextCertNum(Handle);
+}
+
+////////////////////////////////////////////////////////////////////////////////
+void FCertRoots::SetDefault(FCertRoots&& CertRoots)
+{
+	check(GDefaultCertRoots.IsValid() != CertRoots.IsValid());
+	GDefaultCertRoots = MoveTemp(CertRoots);
+}
+
+////////////////////////////////////////////////////////////////////////////////
+FCertRootsRef FCertRoots::NoTls()
+{
+	return ECertRootsRefType::None;
+}
+
+////////////////////////////////////////////////////////////////////////////////
+FCertRootsRef FCertRoots::Default()
+{
+	return ECertRootsRefType::Default;
+}
+
+////////////////////////////////////////////////////////////////////////////////
+FCertRootsRef FCertRoots::Explicit(const FCertRoots& CertRoots)
+{
+	check(CertRoots.IsValid());
+	return CertRoots.Handle;
+}
+
+
+
+////////////////////////////////////////////////////////////////////////////////
+class FTlsPeer
+	: public FPeer
+{
+public:
+				FTlsPeer()								= default;
+				~FTlsPeer();
+				FTlsPeer(FTlsPeer&& Rhs)				{ Move(MoveTemp(Rhs)); }
+				FTlsPeer& operator = (FTlsPeer&& Rhs)	{ return Move(MoveTemp(Rhs)); }
+				FTlsPeer(FSocket InSocket, FCertRootsRef Certs=ECertRootsRefType::None, const char* HostName=nullptr);
+	FOutcome	Handshake();
+	FOutcome	Send(const char* Data, int32 Size);
+	FOutcome	Recv(char* Out, int32 MaxSize);
+
+protected:
+	FTlsPeer&	Move(FTlsPeer&& Rhs);
+	SSL*		Ssl = nullptr;
+};
+
+////////////////////////////////////////////////////////////////////////////////
+FTlsPeer::FTlsPeer(FSocket InSocket, FCertRootsRef Certs, const char* HostName)
+: FPeer(MoveTemp(InSocket))
+{
+	if (Certs == ECertRootsRefType::None)
+	{
+		return;
+	}
+
+	Ssl = Ssl_Create(Certs, HostName);
+	Ssl_AssociatePeer(Ssl, this);
+}
+
+////////////////////////////////////////////////////////////////////////////////
+FTlsPeer::~FTlsPeer()
+{
+	if (Ssl != nullptr)
+	{
+		Ssl_Destroy(Ssl);
+	}
+}
+
+////////////////////////////////////////////////////////////////////////////////
+FTlsPeer& FTlsPeer::Move(FTlsPeer&& Rhs)
+{
+	FPeer::operator = (MoveTemp(Rhs));
+
+	Swap(Ssl, Rhs.Ssl);
+	if (Ssl != nullptr)		Ssl_AssociatePeer(Ssl, this);
+	if (Rhs.Ssl != nullptr) Ssl_AssociatePeer(Rhs.Ssl, &Rhs);
+
+	return *this;
+}
+
+////////////////////////////////////////////////////////////////////////////////
+FOutcome FTlsPeer::Handshake()
+{
+	if (Ssl == nullptr)
+	{
+		return FOutcome::Ok();
+	}
+
+	return Ssl_Handshake(Ssl);
+}
+
+////////////////////////////////////////////////////////////////////////////////
 FOutcome FTlsPeer::Send(const char* Data, int32 Size)
 {
 	if (Ssl == nullptr)
@@ -309,8 +484,7 @@
 		return FPeer::Send(Data, Size);
 	}
 
-	int32 Result = SSL_write(Ssl, Data, Size);
-	return (Result > 0) ? FOutcome::Ok(Result) : GetOutcome(Result);
+	return Ssl_Write(Ssl, Data, Size);
 }
 
 ////////////////////////////////////////////////////////////////////////////////
@@ -321,69 +495,7 @@
 		return FPeer::Recv(Out, MaxSize);
 	}
 
-	int32 Result = SSL_read(Ssl, Out, MaxSize);
-	return (Result > 0) ? FOutcome::Ok(Result) : GetOutcome(Result);
-}
-
-////////////////////////////////////////////////////////////////////////////////
-FOutcome FTlsPeer::GetOutcome(int32 SslResult, const char* Message) const
-{
-	int32 Error = SSL_get_error(Ssl, SslResult);
-	if (Error != SSL_ERROR_WANT_READ && Error != SSL_ERROR_WANT_WRITE)
-	{
-		return FOutcome::Error(Message, Error);
-	}
-	return FOutcome::Waiting();
-}
-
-////////////////////////////////////////////////////////////////////////////////
-int32 FTlsPeer::BioWrite(const char* Data, size_t Size, size_t* BytesWritten, BIO* Bio)
-{
-	*BytesWritten = 0;
-	BIO_clear_retry_flags(Bio);
-
-	FOutcome Outcome = FPeer::Send(Data, Size);
-	if (Outcome.IsWaiting())
-	{
-		BIO_set_retry_write(Bio);
-		return 0;
-	}
-
-	if (Outcome.IsError())
-	{
-		return -1;
-	}
-
-	*BytesWritten = Outcome.GetResult();
-	return 1;
-}
-
-////////////////////////////////////////////////////////////////////////////////
-int32 FTlsPeer::BioRead(char* Data, size_t Size, size_t* BytesRead, BIO* Bio)
-{
-	*BytesRead = 0;
-	BIO_clear_retry_flags(Bio);
-
-	FOutcome Outcome = FPeer::Recv(Data, Size);
-	if (Outcome.IsWaiting())
-	{
-		BIO_set_retry_read(Bio);
-		return 0;
-	}
-
-	if (Outcome.IsError())
-	{
-		return -1;
-	}
-
-	*BytesRead = Outcome.GetResult();
-	return 1;
-}
-
-////////////////////////////////////////////////////////////////////////////////
-long FTlsPeer::BioControl(int Cmd, long, void*, BIO*)
-{
-	return (Cmd == BIO_CTRL_FLUSH) ? 1 : 0;
+	return Ssl_Read(Ssl, Out, MaxSize);
 }
 
 
@@ -406,16 +518,10 @@
 FHttpPeer::FHttpPeer(FSocket InSocket, FCertRootsRef Certs, const char* HostName)
 : FTlsPeer(MoveTemp(InSocket), Certs, HostName)
 {
-	if (Ssl == nullptr)
+	if (Ssl != nullptr)
 	{
-		return;
+		Ssl_SetupAlpn(Ssl, 1);
 	}
-
-	static const uint8 AlpnProtos[] =
-		"\x08" "http/1.1"
-	//	"\x02" "h2"
-		;
-	SSL_set_alpn_protos(Ssl, AlpnProtos, sizeof(AlpnProtos) - 1);
 }
 
 ////////////////////////////////////////////////////////////////////////////////
@@ -440,34 +546,9 @@
 		return;
 	}
 
-	const char* AlpnProto = nullptr;
-	uint32 AlpnProtoLen;
-	SSL_get0_alpn_selected(Ssl, &(const uint8*&)AlpnProto, &AlpnProtoLen);
-	if (AlpnProto == nullptr)
+	if (int32 Negotiated = Ssl_GetProtocolVersion(Ssl); Negotiated > 0)
 	{
-		return;
+		Proto = Negotiated;
 	}
-
-	FAnsiStringView Needle(AlpnProto, AlpnProtoLen);
-	FAnsiStringView Candidates[] = {
-		"http/1.1",
-	//	"h2"
-	};
-	for (int32 i = 0; i < UE_ARRAY_COUNT(Candidates); ++i)
-	{
-		const FAnsiStringView& Candidate = Candidates[i];
-		if (AlpnProtoLen != Candidate.Len())
-		{
-			continue;
-		}
-
-		if (Candidate != Needle)
-		{
-			continue;
-		}
-
-		Proto = i + 1;
-		break;
-	}
 }
 
