Remove FString overload of UE::IoStore::LatencyTest as it was unused.

#rnx

diff --git a/Engine/Source/Runtime/Experimental/IoStore/OnDemand/Private/LatencyTesting.cpp b/Engine/Source/Runtime/Experimental/IoStore/OnDemand/Private/LatencyTesting.cpp
--- a/Engine/Source/Runtime/Experimental/IoStore/OnDemand/Private/LatencyTesting.cpp
+++ b/Engine/Source/Runtime/Experimental/IoStore/OnDemand/Private/LatencyTesting.cpp
@@ -93,12 +93,6 @@ void LatencyTest(FAnsiStringView InUrl, FStringView InPath, uint32 TimeOutMs, TA
 	}
 }
 
-void LatencyTest(FStringView InUrl, FStringView InPath, uint32 InTimeOutMs, TArrayView<int32> OutResults)
-{
-	auto AnsiUrl = StringCast<ANSICHAR>(InUrl.GetData(), InUrl.Len());
-	LatencyTest(AnsiUrl, InPath, InTimeOutMs, OutResults);
-}
-
 bool ConnectionTest(FAnsiStringView Url, FStringView Path, uint32 TimeoutMs)
 {
 	TRACE_CPUPROFILER_EVENT_SCOPE(IasBackend::ConnectionTest);
diff --git a/Engine/Source/Runtime/Experimental/IoStore/OnDemand/Private/LatencyTesting.h b/Engine/Source/Runtime/Experimental/IoStore/OnDemand/Private/LatencyTesting.h
--- a/Engine/Source/Runtime/Experimental/IoStore/OnDemand/Private/LatencyTesting.h
+++ b/Engine/Source/Runtime/Experimental/IoStore/OnDemand/Private/LatencyTesting.h
@@ -17,9 +17,6 @@ namespace UE::IoStore
  */
 void LatencyTest(FAnsiStringView Url, FStringView Path, uint32 TimeoutMs, TArrayView<int32> OutResults);
 
-/** Convenience overload for the ansi version of LatencyTest */
-void LatencyTest(FStringView Url, FStringView Path, uint32 TimeoutMs, TArrayView<int32> OutResults);
-
 /**
  * An easy way to test if a valid connection can be made to a CDN or not. Instead of returning the latency to the CDN this
  * function returns if the CDN can be contacted at all.
