Missed a file when renaming functions

#rnx

diff --git a/Engine/Source/Runtime/Experimental/IoStore/HttpClient/Private/Socks.inl b/Engine/Source/Runtime/Experimental/IoStore/HttpClient/Private/Socks.inl
--- a/Engine/Source/Runtime/Experimental/IoStore/HttpClient/Private/Socks.inl
+++ b/Engine/Source/Runtime/Experimental/IoStore/HttpClient/Private/Socks.inl
@@ -81,8 +81,8 @@ static FOutcome ConnectSocks4(FSocket& Socket, uint32 IpAddress, uint32 Port)
 	}
 
 	FSocks4Request Request = {
-		.Port		= htons(uint16(Port)),
-		.IpAddress	= htonl(IpAddress),
+		.Port		= Socket_HtoNs(uint16(Port)),
+		.IpAddress	= Socket_HtoNl(IpAddress),
 	};
 	if (Outcome = Socket.Send((const char*)&Request, sizeof(Request)); Outcome.IsError())
 	{
