Added explicit flushes to test server, just in case

#rnx

diff --git a/Engine/Source/Runtime/Experimental/IoStore/HttpClient/Private/TestServer.py b/Engine/Source/Runtime/Experimental/IoStore/HttpClient/Private/TestServer.py
--- a/Engine/Source/Runtime/Experimental/IoStore/HttpClient/Private/TestServer.py
+++ b/Engine/Source/Runtime/Experimental/IoStore/HttpClient/Private/TestServer.py
@@ -216,6 +216,7 @@ def http_data(handler, payload_size=0):
     handler.send_header("Content-Type", "application/octet-stream")
     handler.end_headers()
     handler.wfile.write(payload)
+    handler.wfile.flush()
 
 def http_chunked(handler, payload_size=0, *options):
     ext_payload = b""
@@ -256,6 +257,7 @@ def http_chunked(handler, payload_size=0, *options):
         if trailer_payload and not piece:
             handler.wfile.write(trailer_payload + b": true\r\n")
         handler.wfile.write(get_crlf())
+        # handler.wfile.flush()
         if not piece:
             break
         payload = payload[chunk_size:]
