Corrected use of feature defines.

#rnx

diff --git a/Engine/Source/Runtime/Experimental/IoStore/HttpClient/Private/Peer.inl b/Engine/Source/Runtime/Experimental/IoStore/HttpClient/Private/Peer.inl
--- a/Engine/Source/Runtime/Experimental/IoStore/HttpClient/Private/Peer.inl
+++ b/Engine/Source/Runtime/Experimental/IoStore/HttpClient/Private/Peer.inl
@@ -2,7 +2,11 @@
 
 #pragma once
 
-#if defined(IAS_HTTP_HAS_OPENSSL)
+#if !defined(IAS_HTTP_HAS_OPENSSL)
+#	define IAS_HTTP_HAS_OPENSSL 0
+#endif
+
+#if IAS_HTTP_HAS_OPENSSL
 #	if defined(IAS_HTTP_EXPLICIT_VERIFY_TIME)
 #		if !IAS_HTTP_EXPLICIT_VERIFY_TIME 
 #			error Either define this to >=1 or not at all
@@ -52,7 +56,7 @@ FPeer::FPeer(FSocket InSocket)
 
 
 
-#if defined(IAS_HTTP_HAS_OPENSSL)
+#if IAS_HTTP_HAS_OPENSSL
 
 ////////////////////////////////////////////////////////////////////////////////
 static void Ssl_ContextDestroy(UPTRINT Handle)
