Added path to register fds with a poller on creation

#rnx

diff --git a/Engine/Source/Runtime/Experimental/IoStore/HttpClient/Private/Loop.inl b/Engine/Source/Runtime/Experimental/IoStore/HttpClient/Private/Loop.inl
--- a/Engine/Source/Runtime/Experimental/IoStore/HttpClient/Private/Loop.inl
+++ b/Engine/Source/Runtime/Experimental/IoStore/HttpClient/Private/Loop.inl
@@ -740,7 +740,7 @@ public:
 	void				Unwait()			{ check(bWaiting); bWaiting = false; }
 	FWaiter				GetWaiter() const;
 	bool				Tick(FTickState& State);
-	void				TickSend(FTickState& State, FHost& Host);
+	void				TickSend(FTickState& State, FHost& Host, FPoller& Poller);
 	void				Fail(FTickState& State, const char* Reason);
 
 private:
@@ -919,7 +919,7 @@ void FPeerGroup::RecvInternal(FTickState& State)
 		return;
 	}
 
-	Peer = FPeerType();
+	Peer = FHttpPeer();
 }
 
 ////////////////////////////////////////////////////////////////////////////////
@@ -980,7 +980,7 @@ bool FPeerGroup::Tick(FTickState& State)
 }
 
 ////////////////////////////////////////////////////////////////////////////////
-void FPeerGroup::TickSend(FTickState& State, FHost& Host)
+void FPeerGroup::TickSend(FTickState& State, FHost& Host, FPoller& Poller)
 {
 	// This path is only for those that are idle and have nothing to do
 	if (Send != nullptr || Recv != nullptr)
@@ -1007,7 +1007,7 @@ void FPeerGroup::TickSend(FTickState& State, FHost& Host)
 		FOutcome Outcome = FOutcome::None();
 
 		FSocket Socket;
-		if (Socket.Create())
+		if (Socket.Create(Poller))
 		{
 			Outcome = Host.Connect(Socket);
 		}
@@ -1172,7 +1172,7 @@ void FHostGroup::Tick(FTickState& State)
 				break;
 			}
 
-			Group.TickSend(State, Host);
+			Group.TickSend(State, Host, Poller);
 		}
 	}
 
diff --git a/Engine/Source/Runtime/Experimental/IoStore/HttpClient/Private/Socket.inl b/Engine/Source/Runtime/Experimental/IoStore/HttpClient/Private/Socket.inl
--- a/Engine/Source/Runtime/Experimental/IoStore/HttpClient/Private/Socket.inl
+++ b/Engine/Source/Runtime/Experimental/IoStore/HttpClient/Private/Socket.inl
@@ -82,6 +82,7 @@
 		using PollType = UPTRINT;
 		static PollType	Poll_CreateContext()			{ return 0; }
 		static void		Poll_DestroyContext(PollType)	{}
+		static bool		Poll_Register(...)				{ return true; }
 
 #if defined(IAS_HTTP_USE_POLL)
 		template <typename... ArgTypes>
@@ -174,6 +175,7 @@ public:
 				FPoller(FPoller&& Rhs)		{ Swap(Handle, Rhs.Handle); }
 				~FPoller()					{ Poll_DestroyContext(Handle); }
 	FPoller&	operator = (FPoller&& Rhs)	{ Swap(Handle, Rhs.Handle); return *this; }
+	void		Register(SocketType Socket)	{ Poll_Register(Handle, Socket); }
 	int32		Wait(FPollFd* Polls, int32 Num, int32 TimeOutMs);
 
 private:
@@ -363,6 +365,7 @@ public:
 	bool		IsValid() const				{ return Socket != InvalidSocket; }
 	FWaitable	GetWaitable() const			{ return { Socket }; }
 	bool		Create();
+	bool		Create(FPoller& Poller);
 	void		Destroy();
 	FOutcome	Connect(uint32 Ip, uint32 Port);
 	void		Disconnect();
@@ -405,6 +408,18 @@ bool FSocket::Create()
 	return true;
 }
 
+////////////////////////////////////////////////////////////////////////////////
+bool FSocket::Create(FPoller& Poller)
+{
+	if (!Create())
+	{
+		return false;
+	}
+
+	Poller.Register(Socket);
+	return true;
+}
+
 ////////////////////////////////////////////////////////////////////////////////
 void FSocket::Destroy()
 {
