Constraints: moved transform related code under Transform sub-directory for clarity

#rnx

diff --git a/Engine/Plugins/Animation/ControlRig/Source/ControlRig/Private/Constraints/ControlRigTransformableHandle.cpp b/Engine/Plugins/Animation/ControlRig/Source/ControlRig/Private/Constraints/ControlRigTransformableHandle.cpp
--- a/Engine/Plugins/Animation/ControlRig/Source/ControlRig/Private/Constraints/ControlRigTransformableHandle.cpp
+++ b/Engine/Plugins/Animation/ControlRig/Source/ControlRig/Private/Constraints/ControlRigTransformableHandle.cpp
@@ -8,7 +8,7 @@
 #include "ControlRigComponent.h"
 #include "IControlRigObjectBinding.h"
 #include "ControlRigObjectBinding.h"
-#include "TransformableHandleUtils.h"
+#include "Transform/TransformableHandleUtils.h"
 #include "Rigs/RigHierarchyElements.h"
 #include "Sequencer/MovieSceneControlRigParameterSection.h"
 #include "Sections/MovieScene3DTransformSection.h"
diff --git a/Engine/Plugins/Animation/ControlRig/Source/ControlRig/Private/ControlRigModule.cpp b/Engine/Plugins/Animation/ControlRig/Source/ControlRig/Private/ControlRigModule.cpp
--- a/Engine/Plugins/Animation/ControlRig/Source/ControlRig/Private/ControlRigModule.cpp
+++ b/Engine/Plugins/Animation/ControlRig/Source/ControlRig/Private/ControlRigModule.cpp
@@ -3,7 +3,7 @@
 #include "ControlRigModule.h"
 
 #include "ControlRigGizmoActor.h"
-#include "TransformableRegistry.h"
+#include "Transform/TransformableRegistry.h"
 #include "Modules/ModuleManager.h"
 #include "ILevelSequenceModule.h"
 #include "ControlRigObjectVersion.h"
diff --git a/Engine/Plugins/Animation/ControlRig/Source/ControlRig/Private/Sequencer/MovieSceneControlRigParameterSection.cpp b/Engine/Plugins/Animation/ControlRig/Source/ControlRig/Private/Sequencer/MovieSceneControlRigParameterSection.cpp
--- a/Engine/Plugins/Animation/ControlRig/Source/ControlRig/Private/Sequencer/MovieSceneControlRigParameterSection.cpp
+++ b/Engine/Plugins/Animation/ControlRig/Source/ControlRig/Private/Sequencer/MovieSceneControlRigParameterSection.cpp
@@ -21,8 +21,9 @@
 #include "Components/SkeletalMeshComponent.h"
 #include "Constraints/ControlRigTransformableHandle.h"
 #include "UObject/UE5MainStreamObjectVersion.h"
-#include "TransformConstraint.h"
-#include "TransformableHandle.h"
+#include "Transform/TransformConstraint.h"
+#include "Transform/TransformableHandle.h"
+#include "Transform/TransformConstraintUtil.h"
 #include "AnimationCoreLibrary.h"
 #include "UObject/ObjectSaveContext.h"
 #include "Evaluation/MovieSceneEvaluationField.h"
@@ -662,7 +663,7 @@ struct FParameterTransformChannelEditorData
 				{
 					// switch translation to constraint space if needed
 					const uint32 ControlHash = UTransformableControlHandle::ComputeHash(ControlRig, ControlElement->GetFName());
-					TOptional<FTransform> ConstraintSpaceTransform = FTransformConstraintUtils::GetRelativeTransform(ControlRig->GetWorld(), ControlHash);
+					TOptional<FTransform> ConstraintSpaceTransform = UE::TransformConstraintUtil::GetRelativeTransform(ControlRig->GetWorld(), ControlHash);
 					if (ConstraintSpaceTransform)
 					{
 						return ConstraintSpaceTransform->GetTranslation();
@@ -716,7 +717,7 @@ struct FParameterTransformChannelEditorData
 				{
 					// switch rotation to constraint space if needed
 					const uint32 ControlHash = UTransformableControlHandle::ComputeHash(ControlRig, ControlElement->GetFName());
-					TOptional<FTransform> ConstraintSpaceTransform = FTransformConstraintUtils::GetRelativeTransform(ControlRig->GetWorld(), ControlHash);
+					TOptional<FTransform> ConstraintSpaceTransform = UE::TransformConstraintUtil::GetRelativeTransform(ControlRig->GetWorld(), ControlHash);
 					if (ConstraintSpaceTransform)
 					{
 						return ConstraintSpaceTransform->GetRotation().Rotator();
@@ -742,7 +743,7 @@ struct FParameterTransformChannelEditorData
 				{
 					// switch scale to constraint space if needed
 					const uint32 ControlHash = UTransformableControlHandle::ComputeHash(ControlRig, ControlElement->GetFName());
-					TOptional<FTransform> ConstraintSpaceTransform = FTransformConstraintUtils::GetRelativeTransform(ControlRig->GetWorld(), ControlHash);
+					TOptional<FTransform> ConstraintSpaceTransform = UE::TransformConstraintUtil::GetRelativeTransform(ControlRig->GetWorld(), ControlHash);
 					if (ConstraintSpaceTransform)
 					{
 						return ConstraintSpaceTransform->GetScale3D();
diff --git a/Engine/Plugins/Animation/ControlRig/Source/ControlRig/Private/Sequencer/MovieSceneControlRigParameterTemplate.cpp b/Engine/Plugins/Animation/ControlRig/Source/ControlRig/Private/Sequencer/MovieSceneControlRigParameterTemplate.cpp
--- a/Engine/Plugins/Animation/ControlRig/Source/ControlRig/Private/Sequencer/MovieSceneControlRigParameterTemplate.cpp
+++ b/Engine/Plugins/Animation/ControlRig/Source/ControlRig/Private/Sequencer/MovieSceneControlRigParameterTemplate.cpp
@@ -22,8 +22,8 @@
 #include "Rigs/FKControlRig.h"
 #include "UObject/UObjectAnnotation.h"
 #include "Rigs/RigHierarchy.h"
-#include "TransformConstraint.h"
-#include "TransformableHandle.h"
+#include "Transform/TransformConstraint.h"
+#include "Transform/TransformableHandle.h"
 #include "AnimationCoreLibrary.h"
 #include "Constraints/ControlRigTransformableHandle.h"
 #include "Misc/ScopeLock.h"
diff --git a/Engine/Plugins/Animation/ControlRig/Source/ControlRig/Private/Sequencer/MovieSceneControlRigSystem.cpp b/Engine/Plugins/Animation/ControlRig/Source/ControlRig/Private/Sequencer/MovieSceneControlRigSystem.cpp
--- a/Engine/Plugins/Animation/ControlRig/Source/ControlRig/Private/Sequencer/MovieSceneControlRigSystem.cpp
+++ b/Engine/Plugins/Animation/ControlRig/Source/ControlRig/Private/Sequencer/MovieSceneControlRigSystem.cpp
@@ -35,8 +35,8 @@
 #include "Rigs/FKControlRig.h"
 #include "UObject/UObjectAnnotation.h"
 #include "Rigs/RigHierarchy.h"
-#include "TransformConstraint.h"
-#include "TransformableHandle.h"
+#include "Transform/TransformConstraint.h"
+#include "Transform/TransformableHandle.h"
 #include "Constraints/ControlRigTransformableHandle.h"
 
 #include UE_INLINE_GENERATED_CPP_BY_NAME(MovieSceneControlRigSystem)
diff --git a/Engine/Plugins/Animation/ControlRig/Source/ControlRig/Public/Constraints/ControlRigTransformableHandle.h b/Engine/Plugins/Animation/ControlRig/Source/ControlRig/Public/Constraints/ControlRigTransformableHandle.h
--- a/Engine/Plugins/Animation/ControlRig/Source/ControlRig/Public/Constraints/ControlRigTransformableHandle.h
+++ b/Engine/Plugins/Animation/ControlRig/Source/ControlRig/Public/Constraints/ControlRigTransformableHandle.h
@@ -2,7 +2,7 @@
 
 #pragma once
 
-#include "TransformableHandle.h"
+#include "Transform/TransformableHandle.h"
 #include "Rigs/RigHierarchyDefines.h"
 
 #include "ControlRigTransformableHandle.generated.h"
diff --git a/Engine/Plugins/Animation/ControlRig/Source/ControlRigEditor/Private/ControlRigSequencerEditorLibrary.cpp b/Engine/Plugins/Animation/ControlRig/Source/ControlRigEditor/Private/ControlRigSequencerEditorLibrary.cpp
--- a/Engine/Plugins/Animation/ControlRig/Source/ControlRigEditor/Private/ControlRigSequencerEditorLibrary.cpp
+++ b/Engine/Plugins/Animation/ControlRig/Source/ControlRigEditor/Private/ControlRigSequencerEditorLibrary.cpp
@@ -38,7 +38,8 @@
 #include "ControlRigSpaceChannelEditors.h"
 #include "LevelSequenceEditorBlueprintLibrary.h"
 #include "Tools/ConstraintBaker.h"
-#include "TransformConstraint.h"
+#include "Transform/TransformConstraint.h"
+#include "Transform/TransformConstraintUtil.h"
 #include "Rigs/FKControlRig.h"
 #include "Subsystems/UnrealEditorSubsystem.h"
 #include "EditorScriptingHelpers.h"
@@ -484,14 +485,14 @@ UTickableConstraint* UControlRigSequencerEditorLibrary::AddConstraint(UWorld* Wo
 		return nullptr;
 	}
 
-	UTickableTransformConstraint* Constraint = FTransformConstraintUtils::CreateFromType(World, InType);
+	UTickableTransformConstraint* Constraint = UE::TransformConstraintUtil::CreateFromType(World, InType);
 	if (!Constraint)
 	{
 		UE_LOG(LogControlRig, Error, TEXT("AddConstraint: Can not create constraint from type"));
 		return nullptr;
 	}
 
-	if(FTransformConstraintUtils::AddConstraint(World, InParentHandle, InChildHandle, Constraint, bMaintainOffset) == false)
+	if(UE::TransformConstraintUtil::AddConstraint(World, InParentHandle, InChildHandle, Constraint, bMaintainOffset) == false)
 	{
 		UE_LOG(LogControlRig, Error, TEXT("AddConstraint: Constraint not added"));
 		Constraint->MarkAsGarbage();
diff --git a/Engine/Plugins/Animation/ControlRig/Source/ControlRigEditor/Private/EditMode/AnimDetailsProxy.cpp b/Engine/Plugins/Animation/ControlRig/Source/ControlRigEditor/Private/EditMode/AnimDetailsProxy.cpp
--- a/Engine/Plugins/Animation/ControlRig/Source/ControlRigEditor/Private/EditMode/AnimDetailsProxy.cpp
+++ b/Engine/Plugins/Animation/ControlRig/Source/ControlRigEditor/Private/EditMode/AnimDetailsProxy.cpp
@@ -45,7 +45,7 @@
 #include "Tree/SCurveEditorTree.h"
 #include "IKeyArea.h"
 #include "SequencerAddKeyOperation.h"
-#include "TransformConstraint.h"
+#include "Transform/TransformConstraint.h"
 #include "LevelEditorViewport.h"
 #include "Units/Execution/RigUnit_BeginExecution.h"
 #include "ConstraintsManager.h"
@@ -1078,7 +1078,7 @@ static bool SetConstrainedTransform(FTransform LocalTransform, UControlRig* Cont
 	{
 		return false;
 	}
-	const int32 LastActiveIndex = FTransformConstraintUtils::GetLastActiveConstraintIndex(Constraints);
+	const int32 LastActiveIndex = UE::TransformConstraintUtil::GetLastActiveConstraintIndex(Constraints);
 	const bool bNeedsConstraintPostProcess = Constraints.IsValidIndex(LastActiveIndex);
 
 	if (!bNeedsConstraintPostProcess)
@@ -1105,7 +1105,7 @@ static bool SetConstrainedTransform(FTransform LocalTransform, UControlRig* Cont
 	const FTransform WorldTransform = GlobalTransform * ToWorldTransform;
 
 	const TOptional<FTransform> RelativeTransform =
-		FTransformConstraintUtils::GetConstraintsRelativeTransform(Constraints, LocalTransform, WorldTransform);
+		UE::TransformConstraintUtil::GetConstraintsRelativeTransform(Constraints, LocalTransform, WorldTransform);
 	if (RelativeTransform)
 	{
 		LocalTransform = *RelativeTransform;
diff --git a/Engine/Plugins/Animation/ControlRig/Source/ControlRigEditor/Private/EditMode/ControlRigEditMode.cpp b/Engine/Plugins/Animation/ControlRig/Source/ControlRigEditor/Private/EditMode/ControlRigEditMode.cpp
--- a/Engine/Plugins/Animation/ControlRig/Source/ControlRigEditor/Private/EditMode/ControlRigEditMode.cpp
+++ b/Engine/Plugins/Animation/ControlRig/Source/ControlRigEditor/Private/EditMode/ControlRigEditMode.cpp
@@ -60,7 +60,7 @@
 #include "EdModeInteractiveToolsContext.h"
 #include "Constraints/MovieSceneConstraintChannelHelper.h"
 #include "Editor/EditorPerProjectUserSettings.h"
-#include "TransformConstraint.h"
+#include "Transform/TransformConstraint.h"
 #include "Animation/DebugSkelMeshComponent.h"
 #include "Tools/ControlRigPose.h"
 #include "Materials/Material.h"
@@ -4991,7 +4991,7 @@ void FControlRigEditMode::SetControlShapeTransform(
 			FTransform LocalTransform = ControlRig->GetControlLocalTransform(InShapeActor->ControlName);
 
 			const TOptional<FTransform> RelativeTransform =
-				FTransformConstraintUtils::GetConstraintsRelativeTransform(Constraints, LocalTransform, WorldTransform);
+				UE::TransformConstraintUtil::GetConstraintsRelativeTransform(Constraints, LocalTransform, WorldTransform);
 			if (RelativeTransform)
 			{
 				LocalTransform = *RelativeTransform; 
diff --git a/Engine/Plugins/Animation/ControlRig/Source/ControlRigEditor/Private/EditMode/ControlRigEditModeUtil.cpp b/Engine/Plugins/Animation/ControlRig/Source/ControlRigEditor/Private/EditMode/ControlRigEditModeUtil.cpp
--- a/Engine/Plugins/Animation/ControlRig/Source/ControlRigEditor/Private/EditMode/ControlRigEditModeUtil.cpp
+++ b/Engine/Plugins/Animation/ControlRig/Source/ControlRigEditor/Private/EditMode/ControlRigEditModeUtil.cpp
@@ -3,9 +3,9 @@
 #include "EditMode/ControlRigEditModeUtil.h"
 
 #include "ControlRig.h"
-#include "TransformConstraint.h"
-
 #include "Rigs/RigHierarchy.h"
+#include "Transform/TransformConstraint.h"
+#include "Transform/TransformConstraintUtil.h"
 
 namespace UE::ControlRigEditMode
 {
@@ -93,7 +93,7 @@ void FExplicitRotationInteraction::Apply(
 			FTransform LocalTransform = ControlRig->GetControlLocalTransform(ControlName);
 
 			const TOptional<FTransform> RelativeTransform =
-				FTransformConstraintUtils::GetConstraintsRelativeTransform(InConstraints, LocalTransform, WorldTransform);
+				UE::TransformConstraintUtil::GetConstraintsRelativeTransform(InConstraints, LocalTransform, WorldTransform);
 			if (ensure(RelativeTransform))
 			{
 				LocalTransform = *RelativeTransform; 
diff --git a/Engine/Plugins/Animation/ControlRig/Source/ControlRigEditor/Private/EditMode/SControlRigEditModeTools.cpp b/Engine/Plugins/Animation/ControlRig/Source/ControlRigEditor/Private/EditMode/SControlRigEditModeTools.cpp
--- a/Engine/Plugins/Animation/ControlRig/Source/ControlRigEditor/Private/EditMode/SControlRigEditModeTools.cpp
+++ b/Engine/Plugins/Animation/ControlRig/Source/ControlRigEditor/Private/EditMode/SControlRigEditModeTools.cpp
@@ -37,7 +37,7 @@
 #include "Widgets/Input/SComboButton.h"
 #include "Framework/Notifications/NotificationManager.h"
 #include "ScopedTransaction.h"
-#include "TransformConstraint.h"
+#include "Transform/TransformConstraint.h"
 #include "EditMode/ControlRigEditModeToolkit.h"
 #include "EditMode/SControlRigDetails.h"
 #include "Editor/Constraints/SConstraintsWidget.h"
diff --git a/Engine/Plugins/Animation/ControlRig/Source/ControlRigEditor/Private/Editor/Constraints/SConstraintsWidget.cpp b/Engine/Plugins/Animation/ControlRig/Source/ControlRigEditor/Private/Editor/Constraints/SConstraintsWidget.cpp
--- a/Engine/Plugins/Animation/ControlRig/Source/ControlRigEditor/Private/Editor/Constraints/SConstraintsWidget.cpp
+++ b/Engine/Plugins/Animation/ControlRig/Source/ControlRigEditor/Private/Editor/Constraints/SConstraintsWidget.cpp
@@ -12,7 +12,6 @@
 #include "SlateOptMacros.h"
 #include "LevelEditor.h"
 #include "SSocketChooser.h"
-#include "TransformConstraint.h"
 #include "Modules/ModuleManager.h"
 #include "Widgets/Input/SButton.h"
 #include "Framework/MultiBox/MultiBoxBuilder.h"
@@ -28,6 +27,8 @@
 #include "FrameNumberDetailsCustomization.h"
 #include "SAssetDropTarget.h"
 #include "Sequencer/ControlRigParameterTrackEditor.h"
+#include "Transform/TransformConstraint.h"
+#include "Transform/TransformConstraintUtil.h"
 
 #define LOCTEXT_NAMESPACE "SConstraintsWidget"
 
@@ -270,7 +271,7 @@ namespace UE::Private
 			{
 				FScopedTransaction Transaction(LOCTEXT("CreateConstraintKey", "Create Constraint Key"));
 				UTickableTransformConstraint* Constraint =
-					FTransformConstraintUtils::CreateAndAddFromObjects(World, Parent, InParent.Socket, Child, NAME_None, InConstraintType, true, bUseDefault);
+					UE::TransformConstraintUtil::CreateAndAddFromObjects(World, Parent, InParent.Socket, Child, NAME_None, InConstraintType, true, bUseDefault);
 
 				if (Constraint)
 				{
@@ -1153,7 +1154,7 @@ int32 FBaseConstraintListWidget::RefreshConstraintList()
 		{
 			for (const AActor* Actor : Selection)
 			{
-				FTransformConstraintUtils::GetParentConstraints(World, Actor, Constraints);
+				UE::TransformConstraintUtil::GetParentConstraints(World, Actor, Constraints);
 			}
 			
 			//remove if not active...
diff --git a/Engine/Plugins/Animation/ControlRig/Source/ControlRigEditor/Private/Sequencer/ControlConstraintChannelInterface.cpp b/Engine/Plugins/Animation/ControlRig/Source/ControlRigEditor/Private/Sequencer/ControlConstraintChannelInterface.cpp
--- a/Engine/Plugins/Animation/ControlRig/Source/ControlRigEditor/Private/Sequencer/ControlConstraintChannelInterface.cpp
+++ b/Engine/Plugins/Animation/ControlRig/Source/ControlRigEditor/Private/Sequencer/ControlConstraintChannelInterface.cpp
@@ -10,7 +10,7 @@
 #include "LevelSequence.h"
 #include "MovieSceneToolHelpers.h"
 
-#include "TransformConstraint.h"
+#include "Transform/TransformConstraint.h"
 #include "Algo/Copy.h"
 #include "ConstraintsManager.h"
 #include "ScopedTransaction.h"
diff --git a/Engine/Plugins/Animation/ControlRig/Source/ControlRigEditor/Private/Sequencer/ControlRigParameterTrackEditor.cpp b/Engine/Plugins/Animation/ControlRig/Source/ControlRigEditor/Private/Sequencer/ControlRigParameterTrackEditor.cpp
--- a/Engine/Plugins/Animation/ControlRig/Source/ControlRigEditor/Private/Sequencer/ControlRigParameterTrackEditor.cpp
+++ b/Engine/Plugins/Animation/ControlRig/Source/ControlRigEditor/Private/Sequencer/ControlRigParameterTrackEditor.cpp
@@ -92,7 +92,8 @@
 #include "Toolkits/IToolkitHost.h"
 #include "EditMode/ControlRigEditModeSettings.h"
 #include "ControlRigSpaceChannelEditors.h"
-#include "TransformConstraint.h"
+#include "Transform/TransformConstraint.h"
+#include "Transform/TransformConstraintUtil.h"
 #include "Misc/ScopedSlowTask.h"
 #include "Misc/TransactionObjectEvent.h"
 #include "Constraints/MovieSceneConstraintChannelHelper.h"
@@ -3785,6 +3786,8 @@ void FControlRigParameterTrackEditor::HandleControlModified(UControlRig* Control
 					{
 						KeyMode = ESequencerKeyMode::ManualKeyForced;
 					}
+
+					UE_LOG(LogControlRigEditor, Warning, TEXT("ADD KEY TO %s"), *ControlElement->GetName());
 					AddControlKeys(Object, ControlRig, Name, ControlElement->GetFName(), (EControlRigContextChannelToKey)Context.KeyMask, 
 						KeyMode, Context.LocalTime);
 					ControlChangedDuringUndoBracket++;
@@ -4060,7 +4063,7 @@ void FControlRigParameterTrackEditor::GetControlRigKeys(
 				if (bInConstraintSpace)
 				{
 					const uint32 ControlHash = UTransformableControlHandle::ComputeHash(InControlRig, ControlElement->GetFName());
-					TOptional<FTransform> Transform = FTransformConstraintUtils::GetRelativeTransform(InControlRig->GetWorld(), ControlHash);
+					TOptional<FTransform> Transform = UE::TransformConstraintUtil::GetRelativeTransform(InControlRig->GetWorld(), ControlHash);
 					if (Transform)
 					{
 						Translation = Transform->GetTranslation();
diff --git a/Engine/Plugins/Animation/ControlRig/Source/ControlRigEditor/Private/Tools/ConstraintBaker.cpp b/Engine/Plugins/Animation/ControlRig/Source/ControlRigEditor/Private/Tools/ConstraintBaker.cpp
--- a/Engine/Plugins/Animation/ControlRig/Source/ControlRigEditor/Private/Tools/ConstraintBaker.cpp
+++ b/Engine/Plugins/Animation/ControlRig/Source/ControlRigEditor/Private/Tools/ConstraintBaker.cpp
@@ -2,8 +2,8 @@
 
 #include "Tools/ConstraintBaker.h"
 
-#include "TransformConstraint.h"
-#include "TransformableHandle.h"
+#include "Transform/TransformConstraint.h"
+#include "Transform/TransformableHandle.h"
 #include "ConstraintChannel.h"
 
 #include "ISequencer.h"
diff --git a/Engine/Plugins/Animation/ControlRig/Source/ControlRigEditor/Private/Tools/ControlRigSnapper.cpp b/Engine/Plugins/Animation/ControlRig/Source/ControlRigEditor/Private/Tools/ControlRigSnapper.cpp
--- a/Engine/Plugins/Animation/ControlRig/Source/ControlRigEditor/Private/Tools/ControlRigSnapper.cpp
+++ b/Engine/Plugins/Animation/ControlRig/Source/ControlRigEditor/Private/Tools/ControlRigSnapper.cpp
@@ -31,8 +31,9 @@
 #include "LevelSequenceActor.h"
 #include "Tools/BakingHelper.h"
 #include "Containers/SortedMap.h"
-#include "TransformConstraint.h"
 #include "Sequencer/ControlRigSequencerHelpers.h"
+#include "Transform/TransformConstraint.h"
+#include "Transform/TransformConstraintUtil.h"
 #include "Constraints/MovieSceneConstraintChannelHelper.inl"
 
 
@@ -291,7 +292,7 @@ struct FGuidAndActor
 				}
 				Sequencer->GetEvaluationTemplate().EvaluateSynchronousBlocking(Context);
 				Controller.EvaluateAllConstraints();
-				FTransformConstraintUtils::UpdateTransformBasedOnConstraint(CurrentTransform, SceneComponent);
+				UE::TransformConstraintUtil::UpdateTransformBasedOnConstraint(CurrentTransform, SceneComponent);
 			}
 		}
 
diff --git a/Engine/Plugins/Experimental/CineCameraRigs/Source/CineCameraRigs/Private/CineCameraAttachMount.cpp b/Engine/Plugins/Experimental/CineCameraRigs/Source/CineCameraRigs/Private/CineCameraAttachMount.cpp
--- a/Engine/Plugins/Experimental/CineCameraRigs/Source/CineCameraRigs/Private/CineCameraAttachMount.cpp
+++ b/Engine/Plugins/Experimental/CineCameraRigs/Source/CineCameraRigs/Private/CineCameraAttachMount.cpp
@@ -2,8 +2,8 @@
 
 #include "CineCameraAttachMount.h"
 
-#include "TransformConstraint.h"
-#include "TransformableHandle.h"
+#include "Transform/TransformConstraint.h"
+#include "Transform/TransformConstraintUtil.h"
 
 #include "Components/BillboardComponent.h"
 #include "Components/StaticMeshComponent.h"
@@ -120,7 +120,7 @@ void ACineCameraAttachMount::CreateConstraint()
 {
 	UWorld* World = GetWorld();
 	TArray< TWeakObjectPtr<UTickableConstraint> > Constraints;
-	FTransformConstraintUtils::GetParentConstraints(World, this, Constraints);
+	UE::TransformConstraintUtil::GetParentConstraints(World, this, Constraints);
 
 	for (const TWeakObjectPtr<UTickableConstraint>& Constraint : Constraints)
 	{
@@ -133,7 +133,7 @@ void ACineCameraAttachMount::CreateConstraint()
 		return;
 	}
 
-	UTickableTransformConstraint* Constraint = FTransformConstraintUtils::CreateAndAddFromObjects(World, TargetActor.Get(), TargetSocket, this, NAME_None, ETransformConstraintType::Parent);
+	UTickableTransformConstraint* Constraint = UE::TransformConstraintUtil::CreateAndAddFromObjects(World, TargetActor.Get(), TargetSocket, this, NAME_None, ETransformConstraintType::Parent);
 	if (Constraint)
 	{
 		Constraint->bDynamicOffset = true;
@@ -148,7 +148,7 @@ UTickableParentConstraint* ACineCameraAttachMount::GetConstraint()
 {
 	UWorld* World = GetWorld();
 	TArray< TWeakObjectPtr<UTickableConstraint> > Constraints;
-	FTransformConstraintUtils::GetParentConstraints(World, this, Constraints);
+	UE::TransformConstraintUtil::GetParentConstraints(World, this, Constraints);
 	if (Constraints.Num() == 0)
 	{
 		return nullptr;
diff --git a/Engine/Source/Editor/MovieSceneTools/Private/Constraints/ComponentConstraintChannelInterface.cpp b/Engine/Source/Editor/MovieSceneTools/Private/Constraints/ComponentConstraintChannelInterface.cpp
--- a/Engine/Source/Editor/MovieSceneTools/Private/Constraints/ComponentConstraintChannelInterface.cpp
+++ b/Engine/Source/Editor/MovieSceneTools/Private/Constraints/ComponentConstraintChannelInterface.cpp
@@ -5,8 +5,8 @@
 
 #include "Evaluation/MovieSceneEvaluationTemplateInstance.h"
 #include "MovieSceneToolHelpers.h"
-#include "TransformableHandle.h"
-#include "TransformConstraint.h"
+#include "Transform/TransformableHandle.h"
+#include "Transform/TransformConstraint.h"
 #include "Constraints/MovieSceneConstraintChannelHelper.inl"
 #include "EntitySystem/MovieSceneDecompositionQuery.h"
 #include "Evaluation/MovieSceneEvaluationTemplateInstance.h"
diff --git a/Engine/Source/Editor/MovieSceneTools/Private/Constraints/MovieSceneConstraintChannelHelper.cpp b/Engine/Source/Editor/MovieSceneTools/Private/Constraints/MovieSceneConstraintChannelHelper.cpp
--- a/Engine/Source/Editor/MovieSceneTools/Private/Constraints/MovieSceneConstraintChannelHelper.cpp
+++ b/Engine/Source/Editor/MovieSceneTools/Private/Constraints/MovieSceneConstraintChannelHelper.cpp
@@ -5,11 +5,12 @@
 #include "ISequencer.h"
 #include "Evaluation/MovieSceneEvaluationTemplateInstance.h"
 
-#include "TransformableHandle.h"
 #include "LevelSequence.h"
 #include "MovieSceneToolHelpers.h"
 
-#include "TransformConstraint.h"
+#include "Transform/TransformableHandle.h"
+#include "Transform/TransformConstraint.h"
+#include "Transform/TransformConstraintUtil.h"
 #include "Algo/Copy.h"
 
 #include "ScopedTransaction.h"
@@ -60,7 +61,7 @@ void FCompensationEvaluator::ComputeLocalTransforms(UWorld* InWorld, const FEval
 	auto GetLastActiveConstraint = [ConstraintsMinusThis]()
 	{
 		// find last active constraint in the list that is different than the one we want to compensate for
-		const int32 LastActiveIndex = FTransformConstraintUtils::GetLastActiveConstraintIndex(ConstraintsMinusThis);
+		const int32 LastActiveIndex = UE::TransformConstraintUtil::GetLastActiveConstraintIndex(ConstraintsMinusThis);
 
 		// if found, return its parent global transform
 		return LastActiveIndex > INDEX_NONE ? Cast<UTickableTransformConstraint>(ConstraintsMinusThis[LastActiveIndex]) : nullptr;
@@ -171,7 +172,7 @@ void FCompensationEvaluator::ComputeLocalTransforms(UWorld* InWorld, const FEval
 				{
 					SpaceGlobal = LastConstraint->GetParentGlobalTransform();
 					TOptional<FTransform> Relative =
-						FTransformConstraintUtils::GetConstraintsRelativeTransform(ConstraintsMinusThis, ChildLocal, ChildGlobal);
+						UE::TransformConstraintUtil::GetConstraintsRelativeTransform(ConstraintsMinusThis, ChildLocal, ChildGlobal);
 					if (Relative)
 					{
 						ChildLocal = *Relative;
@@ -185,7 +186,7 @@ void FCompensationEvaluator::ComputeLocalTransforms(UWorld* InWorld, const FEval
 			else
 			{
 				SpaceGlobal = Constraint->GetParentGlobalTransform();
-				ChildLocal = FTransformConstraintUtils::ComputeRelativeTransform(
+				ChildLocal = UE::TransformConstraintUtil::ComputeRelativeTransform(
 					ChildLocal, ChildGlobal, SpaceGlobal, Constraint);
 			}
 		}
@@ -196,7 +197,7 @@ void FCompensationEvaluator::ComputeLocalTransforms(UWorld* InWorld, const FEval
 			if (Index == 0)
 			{
 				SpaceGlobal = Constraint->GetParentGlobalTransform();
-				ChildLocal = FTransformConstraintUtils::ComputeRelativeTransform(
+				ChildLocal = UE::TransformConstraintUtil::ComputeRelativeTransform(
 					ChildLocal, ChildGlobal, SpaceGlobal, Constraint);
 			}
 			else
@@ -205,7 +206,7 @@ void FCompensationEvaluator::ComputeLocalTransforms(UWorld* InWorld, const FEval
 				{
 					SpaceGlobal = LastConstraint->GetParentGlobalTransform();
 					TOptional<FTransform> Relative =
-						FTransformConstraintUtils::GetConstraintsRelativeTransform(ConstraintsMinusThis, ChildLocal, ChildGlobal);
+						UE::TransformConstraintUtil::GetConstraintsRelativeTransform(ConstraintsMinusThis, ChildLocal, ChildGlobal);
 					if (Relative)
 					{
 						ChildLocal = *Relative;
@@ -249,7 +250,7 @@ void FCompensationEvaluator::ComputeLocalTransformsForBaking(UWorld* InWorld, co
 	auto GetLastActiveConstraint = [ConstraintsMinusThis]()
 	{
 		// find last active constraint in the list that is different than the one we want to compensate for
-		const int32 LastActiveIndex = FTransformConstraintUtils::GetLastActiveConstraintIndex(ConstraintsMinusThis);
+		const int32 LastActiveIndex = UE::TransformConstraintUtil::GetLastActiveConstraintIndex(ConstraintsMinusThis);
 
 		// if found, return its parent global transform
 		return LastActiveIndex > INDEX_NONE ? Cast<UTickableTransformConstraint>(ConstraintsMinusThis[LastActiveIndex]) : nullptr;
@@ -341,7 +342,7 @@ void FCompensationEvaluator::ComputeLocalTransformsForBaking(UWorld* InWorld, co
 		{
 			SpaceGlobal = LastConstraint->GetParentGlobalTransform();
 			TOptional<FTransform> Relative =
-				FTransformConstraintUtils::GetConstraintsRelativeTransform(ConstraintsMinusThis, ChildLocal, ChildGlobal);
+				UE::TransformConstraintUtil::GetConstraintsRelativeTransform(ConstraintsMinusThis, ChildLocal, ChildGlobal);
 			if (Relative)
 			{
 				ChildLocal = *Relative;
@@ -383,7 +384,7 @@ void FCompensationEvaluator::ComputeLocalTransformsBeforeDeletion(UWorld* InWorl
 	auto GetLastActiveConstraint = [this, ConstraintsMinusThis]()
 	{
 		// find last active constraint in the list that is different than the on we want to compensate for
-		const int32 LastActiveIndex = FTransformConstraintUtils::GetLastActiveConstraintIndex(ConstraintsMinusThis);
+		const int32 LastActiveIndex = UE::TransformConstraintUtil::GetLastActiveConstraintIndex(ConstraintsMinusThis);
 
 		// if found, return its parent global transform
 		return LastActiveIndex > INDEX_NONE ? Cast<UTickableTransformConstraint>(ConstraintsMinusThis[LastActiveIndex]) : nullptr;
@@ -474,7 +475,7 @@ void FCompensationEvaluator::ComputeLocalTransformsBeforeDeletion(UWorld* InWorl
 		{
 			SpaceGlobal = LastConstraint->GetParentGlobalTransform();
 			TOptional<FTransform> Relative =
-				FTransformConstraintUtils::GetConstraintsRelativeTransform(Constraints, ChildLocal, ChildGlobal);
+				UE::TransformConstraintUtil::GetConstraintsRelativeTransform(Constraints, ChildLocal, ChildGlobal);
 			if(Relative)
 			{
 				ChildLocal = *Relative;
@@ -503,7 +504,7 @@ void FCompensationEvaluator::ComputeCompensation(UWorld* InWorld, const TSharedP
 	auto GetLastActiveConstraint = [this, Constraints]()
 	{
 		// find last active constraint in the list that is different than the on we want to compensate for
-		const int32 LastActiveIndex = FTransformConstraintUtils::GetLastActiveConstraintIndex(Constraints);
+		const int32 LastActiveIndex = UE::TransformConstraintUtil::GetLastActiveConstraintIndex(Constraints);
 
 		// if found, return its parent global transform
 		return LastActiveIndex > INDEX_NONE ? Cast<UTickableTransformConstraint>(Constraints[LastActiveIndex]) : nullptr;
@@ -605,7 +606,7 @@ void FCompensationEvaluator::ComputeCompensation(UWorld* InWorld, const TSharedP
 	{
 		SpaceGlobals[0] = LastConstraint->GetParentGlobalTransform();
 		TOptional<FTransform> Relative =
-			FTransformConstraintUtils::GetConstraintsRelativeTransform(Constraints, ChildLocals[0], ChildGlobals[0]);
+			UE::TransformConstraintUtil::GetConstraintsRelativeTransform(Constraints, ChildLocals[0], ChildGlobals[0]);
 		if(Relative)
 		{
 			ChildLocals[0] = *Relative;
@@ -746,7 +747,7 @@ void FCompensationEvaluator::ComputeCurrentTransforms(UWorld* InWorld)
 	auto GetLastActiveConstraint = [Constraints]()
 	{
 		// find last active constraint in the list that is different than the one we want to compensate for
-		const int32 LastActiveIndex = FTransformConstraintUtils::GetLastActiveConstraintIndex(Constraints);
+		const int32 LastActiveIndex = UE::TransformConstraintUtil::GetLastActiveConstraintIndex(Constraints);
 
 		// if found, return its parent global transform
 		return LastActiveIndex > INDEX_NONE ? Cast<UTickableTransformConstraint>(Constraints[LastActiveIndex]) : nullptr;
@@ -756,7 +757,7 @@ void FCompensationEvaluator::ComputeCurrentTransforms(UWorld* InWorld)
 	{
 		SpaceGlobals[0] = LastConstraint->GetParentGlobalTransform();
 		TOptional<FTransform> Relative =
-			FTransformConstraintUtils::GetConstraintsRelativeTransform(Constraints, ChildLocals[0], ChildGlobals[0]);
+			UE::TransformConstraintUtil::GetConstraintsRelativeTransform(Constraints, ChildLocals[0], ChildGlobals[0]);
 		if (Relative)
 		{
 			ChildLocals[0] = *Relative;
diff --git a/Engine/Source/Editor/MovieSceneTools/Private/MovieSceneToolsModule.cpp b/Engine/Source/Editor/MovieSceneTools/Private/MovieSceneToolsModule.cpp
--- a/Engine/Source/Editor/MovieSceneTools/Private/MovieSceneToolsModule.cpp
+++ b/Engine/Source/Editor/MovieSceneTools/Private/MovieSceneToolsModule.cpp
@@ -119,7 +119,7 @@
 #include "EditorAnimUtils.h"
 #include "Animation/AnimSequence.h"
 
-#include "TransformableHandle.h"
+#include "Transform/TransformableHandle.h"
 #include "Constraints/ComponentConstraintChannelInterface.h"
 #include "Constraints/TransformConstraintChannelInterface.h"
 #include "Bindings/MovieSceneReplaceableDirectorBlueprintBinding.h"
diff --git a/Engine/Source/Editor/MovieSceneTools/Private/TrackEditors/TransformTrackEditor.cpp b/Engine/Source/Editor/MovieSceneTools/Private/TrackEditors/TransformTrackEditor.cpp
--- a/Engine/Source/Editor/MovieSceneTools/Private/TrackEditors/TransformTrackEditor.cpp
+++ b/Engine/Source/Editor/MovieSceneTools/Private/TrackEditors/TransformTrackEditor.cpp
@@ -42,8 +42,9 @@
 
 #include "Editor.h"
 #include "LevelEditorViewport.h"
-#include "TransformConstraint.h"
-#include "TransformableHandle.h"
+#include "Transform/TransformConstraint.h"
+#include "Transform/TransformableHandle.h"
+#include "Transform/TransformConstraintUtil.h"
 #include "Constraints/MovieSceneConstraintChannelHelper.h"
 #include "Constraints/TransformConstraintChannelInterface.h"
 #include "Misc/TransactionObjectEvent.h"
@@ -1239,7 +1240,7 @@ FTransformData F3DTransformTrackEditor::RecomposeTransform(const FTransformData&
 		CurrentTransform *= GetTransformOrigin().Inverse();
 	}
 
-	FTransformConstraintUtils::UpdateTransformBasedOnConstraint(CurrentTransform, SceneComponent);
+	UE::TransformConstraintUtil::UpdateTransformBasedOnConstraint(CurrentTransform, SceneComponent);
 	
 	return FTransformData(CurrentTransform.GetLocation(), CurrentTransform.GetRotation().Rotator(), CurrentTransform.GetScale3D());
 }
diff --git a/Engine/Source/Runtime/Experimental/Animation/Constraints/Private/ConstraintsModule.cpp b/Engine/Source/Runtime/Experimental/Animation/Constraints/Private/ConstraintsModule.cpp
--- a/Engine/Source/Runtime/Experimental/Animation/Constraints/Private/ConstraintsModule.cpp
+++ b/Engine/Source/Runtime/Experimental/Animation/Constraints/Private/ConstraintsModule.cpp
@@ -3,7 +3,7 @@
 #include "Modules/ModuleManager.h"
 #include "Modules/ModuleInterface.h"
 
-#include "TransformableRegistry.h"
+#include "Transform/TransformableRegistry.h"
 
 DECLARE_LOG_CATEGORY_EXTERN(LogConstraints, Log, All);
 DEFINE_LOG_CATEGORY(LogConstraints);
diff --git a/Engine/Source/Runtime/Experimental/Animation/Constraints/Private/ConstraintsScripting.cpp b/Engine/Source/Runtime/Experimental/Animation/Constraints/Private/ConstraintsScripting.cpp
--- a/Engine/Source/Runtime/Experimental/Animation/Constraints/Private/ConstraintsScripting.cpp
+++ b/Engine/Source/Runtime/Experimental/Animation/Constraints/Private/ConstraintsScripting.cpp
@@ -2,10 +2,11 @@
 
 
 #include "ConstraintsScripting.h"
-#include "TransformConstraint.h"
-#include "TransformableHandle.h"
 #include "ConstraintsManager.h"
-#include "TransformableRegistry.h"
+#include "Transform/TransformConstraint.h"
+#include "Transform/TransformableHandle.h"
+#include "Transform/TransformableRegistry.h"
+#include "Transform/TransformConstraintUtil.h"
 
 #include UE_INLINE_GENERATED_CPP_BY_NAME(ConstraintsScripting)
 
@@ -17,7 +18,7 @@ UConstraintsManager* UConstraintsScriptingLibrary::GetManager(UWorld* InWorld)
 UTransformableComponentHandle* UConstraintsScriptingLibrary::CreateTransformableComponentHandle(
 	UWorld* InWorld, USceneComponent* InSceneComponent, const FName& InSocketName)
 {
-	UTransformableComponentHandle* Handle = FTransformConstraintUtils::CreateHandleForSceneComponent(InSceneComponent, InSocketName);
+	UTransformableComponentHandle* Handle = UE::TransformConstraintUtil::CreateHandleForSceneComponent(InSceneComponent, InSocketName);
 	return Handle;
 }
 
@@ -52,7 +53,7 @@ UTickableTransformConstraint* UConstraintsScriptingLibrary::CreateFromType(
 	UWorld* InWorld,
 	const ETransformConstraintType InType)
 {
-	UTickableTransformConstraint* Constraint = FTransformConstraintUtils::CreateFromType(InWorld, InType);
+	UTickableTransformConstraint* Constraint = UE::TransformConstraintUtil::CreateFromType(InWorld, InType);
 	return Constraint;
 }
 
@@ -71,7 +72,7 @@ bool UConstraintsScriptingLibrary::AddConstraint(UWorld* InWorld, UTransformable
 		return false;
 	}
 	
-	const bool bAdded = FTransformConstraintUtils::AddConstraint(InWorld, InParentHandle, InChildHandle,InConstraint, bMaintainOffset);
+	const bool bAdded = UE::TransformConstraintUtil::AddConstraint(InWorld, InParentHandle, InChildHandle,InConstraint, bMaintainOffset);
 	if (!bAdded)
 	{
 		UE_LOG(LogTemp, Error, TEXT("AddConstraint: Constraint not added"));
diff --git a/Engine/Source/Runtime/Experimental/Animation/Constraints/Private/Transform/DependencyBuilder.cpp b/Engine/Source/Runtime/Experimental/Animation/Constraints/Private/Transform/DependencyBuilder.cpp
--- a/Engine/Source/Runtime/Experimental/Animation/Constraints/Private/Transform/DependencyBuilder.cpp
+++ b/Engine/Source/Runtime/Experimental/Animation/Constraints/Private/Transform/DependencyBuilder.cpp
@@ -2,11 +2,13 @@
 
 #include "Transform/DependencyBuilder.h"
 
-#include "TransformableHandle.h"
-#include "TransformableRegistry.h"
-#include "TransformConstraint.h"
 #include "ConstraintsManager.inl"
 
+#include "Transform/TransformableHandle.h"
+#include "Transform/TransformableRegistry.h"
+#include "Transform/TransformConstraint.h"
+#include "Transform/TransformConstraintUtil.h"
+
 namespace DependencyLocals
 {
 	
@@ -245,7 +247,7 @@ bool FDependencyBuilder::BuildDependencies(UWorld* InWorld, UTickableTransformCo
 	static const TCHAR* ChildDependencyDesc = TEXT("Child Dependency");
 
 	TArray<ConstraintWeakPtr> ChildChildConstraints;
-	FTransformConstraintUtils::GetChildrenConstraints(InWorld, InConstraint, ChildChildConstraints, bIncludeTarget && !bSelf);
+	UE::TransformConstraintUtil::GetChildrenConstraints(InWorld, InConstraint, ChildChildConstraints, bIncludeTarget && !bSelf);
 	for (const ConstraintWeakPtr& ChildConstraint: ChildChildConstraints)
 	{
 		Controller.SetConstraintsDependencies(InConstraint->ConstraintID, ChildConstraint->ConstraintID);
diff --git a/Engine/Source/Runtime/Experimental/Animation/Constraints/Private/Transform/TransformConstraint.cpp b/Engine/Source/Runtime/Experimental/Animation/Constraints/Private/Transform/TransformConstraint.cpp
new file mode 100644
--- /dev/null
+++ b/Engine/Source/Runtime/Experimental/Animation/Constraints/Private/Transform/TransformConstraint.cpp
@@ -0,0 +1,1511 @@
+﻿// Copyright Epic Games, Inc. All Rights Reserved.
+
+#include "Transform/TransformConstraint.h"
+
+#include "Components/SkeletalMeshComponent.h"
+#include "ConstraintsManager.inl"
+#include "Engine/World.h"
+#include "GameFramework/Actor.h"
+#include "MovieSceneSection.h"
+#include "Transform/TransformableHandle.h"
+#include "Transform/TransformableRegistry.h"
+#include "Transform/DependencyBuilder.h"
+#include "Transform/TransformConstraintUtil.h"
+#include "UObject/Package.h"
+
+#include UE_INLINE_GENERATED_CPP_BY_NAME(TransformConstraint)
+
+namespace ConstraintLocals
+{
+
+static bool	bPreEvaluateChild = true;
+static FAutoConsoleVariableRef CVarPreEvaluateChild(
+	TEXT("Constraints.PreEvaluateChild"),
+	bPreEvaluateChild,
+	TEXT("Force child evaluation before constraint computation.")
+	);
+	
+static bool	bPreTickChild = false;
+static FAutoConsoleVariableRef CVarPreTickChild(
+	TEXT("Constraints.PreTickChild"),
+	bPreTickChild,
+	TEXT("Force child ticking before constraint computation.")
+	);
+
+static bool	bForceChildDependency = false;
+static FAutoConsoleVariableRef CVarForceChildDependency(
+	TEXT("Constraints.ForceChildDependency"),
+	bForceChildDependency,
+	TEXT("Force child dependency on the current constraint.")
+	);
+	
+void PreEvaluateHandle(const TObjectPtr<UTransformableHandle>& InHandle) 
+{
+	if ((bPreEvaluateChild || bPreTickChild) && ::IsValid(InHandle))
+	{
+		InHandle->PreEvaluate(bPreTickChild);
+	}
+}
+
+UObject* GetHandleTarget(const TObjectPtr<UTransformableHandle>& InHandle)
+{
+	return IsValid(InHandle) ? InHandle->GetTarget().Get() : nullptr; 
+}
+
+UTransformableHandle* GetHandle(UObject* InObject, const FName& InSocketName)
+{
+	// look for customized transform handle
+	const FTransformableRegistry& Registry = FTransformableRegistry::Get();
+	if (const FTransformableRegistry::CreateHandleFuncT CreateFunction = Registry.GetCreateFunction(InObject->GetClass()))
+	{
+		return CreateFunction(InObject, InSocketName);
+	}
+
+	return nullptr;
+}
+	
+uint32 GetConstrainableHash(const UObject* InObject)
+{
+	// look for customized hash function
+	const FTransformableRegistry& Registry = FTransformableRegistry::Get();
+	if (const FTransformableRegistry::GetHashFuncT HashFunction = Registry.GetHashFunction(InObject->GetClass()))
+	{
+		return HashFunction(InObject, NAME_None);
+	}
+
+	return 0;
+}
+
+}
+
+/** 
+ * UTickableTransformConstraint
+ **/
+
+#if WITH_EDITOR
+UTickableTransformConstraint::FOnConstraintChanged UTickableTransformConstraint::OnConstraintChanged;
+#endif // WITH_EDITOR
+
+int64 UTickableTransformConstraint::GetType() const
+{
+	return static_cast<int64>(Type);
+}
+EMovieSceneTransformChannel UTickableTransformConstraint::GetChannelsToKey() const
+{
+	static const TMap< ETransformConstraintType, EMovieSceneTransformChannel > ConstraintToChannels({
+	{ETransformConstraintType::Translation, EMovieSceneTransformChannel::Translation},
+	{ETransformConstraintType::Rotation, EMovieSceneTransformChannel::Rotation},
+	{ETransformConstraintType::Scale, EMovieSceneTransformChannel::Scale},
+	{ETransformConstraintType::Parent, EMovieSceneTransformChannel::AllTransform},
+	{ETransformConstraintType::LookAt, EMovieSceneTransformChannel::Rotation}
+		});
+
+	const ETransformConstraintType ConstType = static_cast<ETransformConstraintType>(GetType());
+	if (const EMovieSceneTransformChannel* Channel = ConstraintToChannels.Find(ConstType))
+	{
+		return *Channel;;
+	}
+
+	return EMovieSceneTransformChannel::AllTransform;
+}
+
+UTickableConstraint* UTickableTransformConstraint::Duplicate(UObject* NewOuter) const
+{
+	UTickableTransformConstraint* Dup = DuplicateObject<UTickableTransformConstraint>(this, NewOuter);
+	if (ChildTRSHandle)
+	{
+		UTransformableHandle* HandleCopy = ChildTRSHandle->Duplicate(Dup);
+		Dup->ChildTRSHandle = HandleCopy;
+	}
+	if (ParentTRSHandle)
+	{
+		UTransformableHandle* HandleCopy = ParentTRSHandle->Duplicate(Dup);
+		Dup->ParentTRSHandle = HandleCopy;
+	}
+	for (TPair<TWeakObjectPtr<ULevel>, FConstraintTickFunction>& Pair : ConstraintTicks)
+	{
+		if (Pair.Key.IsValid())
+		{
+			UWorld* World = Pair.Key->GetTypedOuter<UWorld>();
+		}
+	}
+	return Dup;
+}
+
+#if WITH_EDITOR
+
+FString UTickableTransformConstraint::GetLabel() const
+{
+	if (!ChildTRSHandle || !ChildTRSHandle->IsValid())
+	{
+		static const FString DummyLabel;
+		return DummyLabel;
+	}
+	
+	if (ParentTRSHandle && ParentTRSHandle->IsValid())
+	{
+		return FString::Printf(TEXT("%s.%s"), *ParentTRSHandle->GetLabel(), *ChildTRSHandle->GetLabel() );		
+	}
+
+	return ChildTRSHandle->GetLabel();
+}
+
+FString UTickableTransformConstraint::GetFullLabel() const
+{
+	if (!ChildTRSHandle || !ChildTRSHandle->IsValid())
+	{
+		static const FString DummyLabel;
+		return DummyLabel;
+	}
+	
+	if (ParentTRSHandle && ParentTRSHandle->IsValid())
+	{
+		return FString::Printf(TEXT("%s.%s"), *ParentTRSHandle->GetFullLabel(), *ChildTRSHandle->GetFullLabel() );		
+	}
+
+	return ChildTRSHandle->GetLabel();
+}
+
+FString UTickableTransformConstraint::GetTypeLabel() const
+{
+	static const UEnum* TypeEnum = StaticEnum<ETransformConstraintType>();
+	if (TypeEnum->IsValidEnumValue(GetType()))
+	{
+		return TypeEnum->GetNameStringByValue(GetType());
+	}
+
+	return Super::GetTypeLabel();
+}
+
+void UTickableTransformConstraint::PostEditChangeProperty(FPropertyChangedEvent& PropertyChangedEvent)
+{
+	Super::PostEditChangeProperty(PropertyChangedEvent);
+
+	if (HasAnyFlags(RF_ClassDefaultObject))
+	{
+		return;
+	}
+	
+	const FName PropertyName = PropertyChangedEvent.GetPropertyName();
+	if (PropertyName == GET_MEMBER_NAME_CHECKED(UTickableTransformConstraint, bMaintainOffset))
+	{
+		Evaluate();
+		return;
+	}
+
+	if (const FProperty* MemberProperty = PropertyChangedEvent.MemberProperty)
+	{
+		const FName MemberPropertyName = MemberProperty->GetFName();
+		if (MemberPropertyName == GET_MEMBER_NAME_CHECKED(UTickableTranslationConstraint, OffsetTranslation) ||
+			MemberPropertyName == GET_MEMBER_NAME_CHECKED(UTickableRotationConstraint, OffsetRotation) ||
+			MemberPropertyName == GET_MEMBER_NAME_CHECKED(UTickableScaleConstraint, OffsetScale) ||
+			MemberPropertyName == GET_MEMBER_NAME_CHECKED(UTickableParentConstraint, OffsetTransform))
+		{
+			OnConstraintChanged.Broadcast(this, PropertyChangedEvent);
+			Evaluate();
+		}
+	}
+}
+
+void UTickableTransformConstraint::PostEditUndo()
+{
+	Super::PostEditUndo();
+}
+
+UTickableTransformConstraint::FOnConstraintChanged& UTickableTransformConstraint::GetOnConstraintChanged()
+{
+	return OnConstraintChanged;	
+}
+
+#endif
+
+void UTickableTransformConstraint::UnregisterDelegates() const
+{
+	if (ChildTRSHandle)
+	{
+		ChildTRSHandle->HandleModified().RemoveAll(this);
+	}
+	if (ParentTRSHandle)
+	{
+		ParentTRSHandle->HandleModified().RemoveAll(this);
+	}
+}
+
+void UTickableTransformConstraint::RegisterDelegates()
+{
+	UnregisterDelegates();
+
+	if (ChildTRSHandle)
+	{
+		ChildTRSHandle->HandleModified().AddUObject(this, &UTickableTransformConstraint::OnHandleModified);
+	}
+	if (ParentTRSHandle)
+	{
+		ParentTRSHandle->HandleModified().AddUObject(this, &UTickableTransformConstraint::OnHandleModified);
+	}	
+}
+
+void UTickableTransformConstraint::Setup()
+{
+	if (!ParentTRSHandle || !ChildTRSHandle || !ParentTRSHandle->IsValid() || !ChildTRSHandle->IsValid())
+	{
+		// handle error
+		return;
+	}
+	
+	ComputeOffset();
+
+}
+
+void UTickableTransformConstraint::SetupDependencies(const UWorld* InWorld)
+{
+	//we may not be the outer for old files so move it over
+	if (ParentTRSHandle)
+	{
+		UTickableTransformConstraint* PTRS = ParentTRSHandle->GetTypedOuter<UTickableTransformConstraint>();
+		if(PTRS != this)
+		{
+			ParentTRSHandle->Rename(nullptr, this, REN_DontCreateRedirectors);
+		}
+	}
+	if (ChildTRSHandle)
+	{
+		UTickableTransformConstraint* CTRS = ChildTRSHandle->GetTypedOuter<UTickableTransformConstraint>();
+		if (CTRS != this)
+		{
+			ChildTRSHandle->Rename(nullptr, this, REN_DontCreateRedirectors);
+		}
+	}
+
+	if (!ensure(InWorld))
+	{
+		return;
+	}
+	
+	FTickFunction* ParentTickFunction = GetParentHandleTickFunction();
+	FTickFunction* ChildTickFunction = GetChildHandleTickFunction();
+	
+	if (ParentTickFunction && (ChildTickFunction != ParentTickFunction))
+	{
+		// manage dependencies
+		// force ConstraintTickFunction to tick after InParent does.
+		// Note that this might not register anything if the parent can't tick (static meshes for instance)
+		FConstraintTickFunction& ConstraintTick = ConstraintTicks.FindOrAdd(InWorld->GetCurrentLevel());
+		ConstraintTick.AddPrerequisite(ParentTRSHandle->GetPrerequisiteObject(), *ParentTickFunction);
+	}
+	
+	// TODO also check for cycle dependencies here
+	if (ChildTickFunction)
+	{
+		USkeletalMeshComponent* ChildOwner = Cast<USkeletalMeshComponent>(ChildTRSHandle->GetTarget().Get());
+		if (ChildOwner == nullptr || ConstraintLocals::bForceChildDependency)
+		{
+			// force InChild to tick after ConstraintTickFunction does.
+			// Note that this might not register anything if the child can't tick (static meshes for instance)
+			// Also don't do if it's a skel mesh will cause cycle.
+			FConstraintTickFunction& ConstraintTick = ConstraintTicks.FindOrAdd(InWorld->GetCurrentLevel());
+			ChildTickFunction->AddPrerequisite(this, ConstraintTick);
+		}
+	}
+}
+
+void UTickableTransformConstraint::EnsurePrimaryDependency(const UWorld* InWorld) const
+{
+	if (!ensure(InWorld))
+	{
+		return;
+	}
+	
+	const FTickFunction* ParentTickFunction = GetParentHandleTickFunction();
+	const FTickFunction* ChildTickFunction = GetChildHandleTickFunction();
+	if (ParentTickFunction && (ChildTickFunction != ParentTickFunction))
+	{
+		FConstraintTickFunction& ConstraintTick = ConstraintTicks.FindOrAdd(InWorld->GetCurrentLevel());
+		const TArray<FTickPrerequisite>& ParentPrerequisites = ConstraintTick.GetPrerequisites();
+		if (ParentPrerequisites.IsEmpty())
+		{
+			// if the constraint has no prerex at this stage, this means that the parent tick function
+			// is not registered or can't tick (static meshes for instance) so look for the first parent tick function if any.
+			// In a context of adding several constraints, we want to make sure that the evaluation order is the right one
+			FTickPrerequisite PrimaryPrerex = ParentTRSHandle->GetPrimaryPrerequisite();
+			if (FTickFunction* PotentialFunction = PrimaryPrerex.Get())
+			{
+				UObject* Target = PrimaryPrerex.PrerequisiteObject.Get();
+				ConstraintTick.AddPrerequisite(Target, *PotentialFunction);
+			}
+		}
+	}
+}
+
+void UTickableTransformConstraint::OnActiveStateChanged() const
+{
+	// tick dependencies might need to be updated when dealing with cycles (than can be created between two controls for example) 
+	FConstraintCycleChecker::CheckAndFixCycles(this);
+}
+
+void UTickableTransformConstraint::PostLoad()
+{
+	Super::PostLoad();
+}
+
+void UTickableTransformConstraint::PostDuplicate(bool bDuplicateForPIE)
+{
+	Super::PostDuplicate(bDuplicateForPIE);
+
+}
+
+uint32 UTickableTransformConstraint::GetTargetHash() const
+{
+	return (ChildTRSHandle && ChildTRSHandle->IsValid()) ? ChildTRSHandle->GetHash() : 0;
+}
+
+bool UTickableTransformConstraint::ReferencesObject(TWeakObjectPtr<UObject> InObject) const
+{
+	const TWeakObjectPtr<UObject> ChildTarget = (ChildTRSHandle && ChildTRSHandle->IsValid())? ChildTRSHandle->GetTarget() : nullptr;
+	if (ChildTarget == InObject)
+	{
+		return true;
+	}
+
+	const TWeakObjectPtr<UObject> ParentTarget = (ParentTRSHandle && ParentTRSHandle->IsValid()) ? ParentTRSHandle->GetTarget() : nullptr;
+	if (ParentTarget == InObject)
+	{
+		return true;	
+	}
+	
+	return false;
+}
+
+bool UTickableTransformConstraint::HasBoundObjects() const
+{
+	if (ChildTRSHandle && ChildTRSHandle->HasBoundObjects())
+	{
+		return true;
+	}
+	if (ParentTRSHandle && ParentTRSHandle->HasBoundObjects())
+	{
+		return true;
+	}
+	return false;
+}
+
+void UTickableTransformConstraint::ResolveBoundObjects(FMovieSceneSequenceID LocalSequenceID, TSharedRef<UE::MovieScene::FSharedPlaybackState> SharedPlaybackState, UObject* SubObject)
+{
+	// update dependencies if the constraint becomes valid once resolved 
+	FConstraintDependencyScope Scope(this);
+	
+	if (ChildTRSHandle && ChildTRSHandle->HasBoundObjects())
+	{
+		ChildTRSHandle->ResolveBoundObjects(LocalSequenceID, SharedPlaybackState, SubObject);
+	}
+	if (ParentTRSHandle && ParentTRSHandle->HasBoundObjects())
+	{
+		ParentTRSHandle->ResolveBoundObjects(LocalSequenceID, SharedPlaybackState, SubObject);
+	}
+}
+
+void UTickableTransformConstraint::Evaluate(bool bTickHandlesAlso) const
+{
+	if (IsFullyActive())
+	{
+		if (bTickHandlesAlso)
+		{
+			ParentTRSHandle->TickTarget();
+		}
+		Super::Evaluate();
+	}
+}
+
+void UTickableTransformConstraint::SetActive(const bool bIsActive)
+{
+	const bool bNeedsUpdate = (Active != bIsActive);
+	Super::SetActive(bIsActive);
+	
+	if (bNeedsUpdate)
+	{
+		OnActiveStateChanged();
+	}
+}
+
+void UTickableTransformConstraint::SetChildGlobalTransform(const FTransform& InGlobal) const
+{
+	if(ChildTRSHandle && ChildTRSHandle->IsValid())
+	{
+		ChildTRSHandle->SetGlobalTransform(InGlobal);
+	}
+}
+
+void UTickableTransformConstraint::SetChildLocalTransform(const FTransform& InLocal) const
+{
+	if(ChildTRSHandle && ChildTRSHandle->IsValid())
+	{
+		ChildTRSHandle->SetLocalTransform(InLocal);
+	}
+}
+
+FTransform UTickableTransformConstraint::GetChildGlobalTransform() const
+{
+	return (ChildTRSHandle && ChildTRSHandle->IsValid()) ? ChildTRSHandle->GetGlobalTransform() : FTransform::Identity;
+}
+
+FTransform UTickableTransformConstraint::GetChildLocalTransform() const
+{
+	return (ChildTRSHandle && ChildTRSHandle->IsValid()) ? ChildTRSHandle->GetLocalTransform() : FTransform::Identity;
+}
+
+FTransform UTickableTransformConstraint::GetParentGlobalTransform() const
+{
+	return (ParentTRSHandle && ParentTRSHandle->IsValid()) ? ParentTRSHandle->GetGlobalTransform() : FTransform::Identity;
+}
+
+FTransform UTickableTransformConstraint::GetParentLocalTransform() const
+{
+	return (ParentTRSHandle && ParentTRSHandle->IsValid()) ? ParentTRSHandle->GetLocalTransform() : FTransform::Identity;
+}
+
+void UTickableTransformConstraint::OnHandleModified(UTransformableHandle* InHandle, EHandleEvent InNotification)
+{
+	if (!InHandle)
+	{
+		return;
+	}
+
+	const UObject* Target = InHandle->GetTarget().Get();
+	UWorld* World = Target != nullptr ? Target->GetWorld() : nullptr;
+	if (!World)
+	{
+		return;
+	}
+	
+	if(InHandle == ChildTRSHandle || InHandle == ParentTRSHandle)
+	{
+		if (InNotification == EHandleEvent::ComponentUpdated)
+		{
+			SetupDependencies(World);
+			// update dependencies now the component has been updated. 
+			FDependencyBuilder::BuildDependencies(World, this);
+			return;
+		}
+	}
+
+	auto MarkForEvaluation = [this, World]()
+	{
+		const FConstraintsManagerController& Controller = FConstraintsManagerController::Get(World);
+		Controller.MarkConstraintForEvaluation(this);
+	};
+
+	auto DirectEvaluation = [this, World]()
+	{
+		const FConstraintTickFunction* ConstraintTick = ConstraintTicks.Find(World->GetCurrentLevel());
+		if (IsFullyActive() && ConstraintTick->IsTickFunctionRegistered() && ConstraintTick->IsTickFunctionEnabled())
+		{
+			Evaluate(true);
+		}
+	};
+	
+	if (InHandle == ChildTRSHandle)
+	{
+		if (InNotification == EHandleEvent::UpperDependencyUpdated)
+		{
+			const UObject* ParentTarget = ConstraintLocals::GetHandleTarget(ParentTRSHandle);
+			if (ParentTarget && ParentTarget != Target)
+			{
+				if (FConstraintsEvaluationGraph::UseEvaluationGraph())
+				{
+					MarkForEvaluation();
+				}
+				else
+				{
+					DirectEvaluation();
+				}
+				return;
+			}
+		}
+	}
+
+	// update the constraint if the parent's transform has been modified 
+	if (InHandle == ParentTRSHandle && InNotification == EHandleEvent::GlobalTransformUpdated)
+	{
+		if (FConstraintsEvaluationGraph::UseEvaluationGraph())
+		{
+			MarkForEvaluation();
+		}
+		else
+		{
+			DirectEvaluation();
+		}
+		return;
+	}
+
+	if (InHandle == ParentTRSHandle && InNotification == EHandleEvent::UpperDependencyUpdated)
+	{
+		if (FConstraintsEvaluationGraph::UseEvaluationGraph())
+		{
+			MarkForEvaluation();
+		}
+		else
+		{
+			DirectEvaluation();
+		}
+		return;
+	}
+
+	if (InHandle == ChildTRSHandle && InNotification == EHandleEvent::GlobalTransformUpdated)
+	{
+		if (FConstraintsEvaluationGraph::UseEvaluationGraph())
+		{
+			return MarkForEvaluation();
+		}
+	}
+}
+
+bool UTickableTransformConstraint::IsValid(const bool bDeepCheck) const
+{
+	const bool bAreHandlesValid =
+		::IsValid(ChildTRSHandle) && ChildTRSHandle->IsValid(bDeepCheck) &&
+		::IsValid(ParentTRSHandle) && ParentTRSHandle->IsValid(bDeepCheck); 
+	
+	return bDeepCheck ? bAreHandlesValid && bValid : bAreHandlesValid;
+}
+
+bool UTickableTransformConstraint::IsFullyActive() const
+{
+	return (Active && IsValid());
+}
+
+bool UTickableTransformConstraint::NeedsCompensation() const
+{
+	// NOTE: this can be extended to something more complex if needed  
+	return true;
+}
+
+FTickFunction* UTickableTransformConstraint::GetChildHandleTickFunction() const 
+{
+	return GetHandleTickFunction(ChildTRSHandle); 
+}
+
+FTickFunction* UTickableTransformConstraint::GetParentHandleTickFunction() const 
+{
+	return GetHandleTickFunction(ParentTRSHandle); 
+}
+
+FTickFunction* UTickableTransformConstraint::GetHandleTickFunction(const TObjectPtr<UTransformableHandle>& InHandle) const 
+{
+	if (!::IsValid(InHandle) || !InHandle->IsValid())
+	{
+		return nullptr;
+	}
+		
+	return InHandle->GetTickFunction();
+}
+
+void UTickableTransformConstraint::InitConstraint(UWorld *InWorld)
+{
+	if (!ensure(InWorld))
+	{
+		return;
+	}
+	
+	FConstraintTickFunction& ConstraintTick = ConstraintTicks.FindOrAdd(InWorld->GetCurrentLevel());
+
+	if (ConstraintTick.ConstraintFunctions.IsEmpty())
+	{
+		ConstraintTick.RegisterFunction(GetFunction());
+	}
+	ConstraintTick.RegisterTickFunction(InWorld->GetCurrentLevel());
+	ConstraintTick.Constraint = this;
+	SetupDependencies(InWorld);
+	RegisterDelegates();
+	bValid = true;
+}
+
+void UTickableTransformConstraint::TeardownConstraint(UWorld* InWorld)
+{
+	if (!ensure(InWorld))
+	{
+		return;
+	}
+	
+	FConstraintTickFunction& ConstraintTick = ConstraintTicks.FindOrAdd(InWorld->GetCurrentLevel());
+	ConstraintTick.SetTickFunctionEnable(false);
+	ConstraintTick.UnRegisterTickFunction();
+	if (FTickFunction* ChildTickFunction = GetChildHandleTickFunction())
+	{
+		ChildTickFunction->RemovePrerequisite(this, ConstraintTick);
+	}
+	ConstraintTicks.Remove(InWorld->GetCurrentLevel());
+
+	// unregister delegates. should handles delegates be unregistered as well?
+	UnregisterDelegates();
+}
+
+void UTickableTransformConstraint::AddedToWorld(UWorld* InWorld)
+{
+	if (InWorld)
+	{
+		// update dependencies once added to the sub-system 
+		FDependencyBuilder::BuildDependencies(InWorld, this);
+	}
+}
+
+/** 
+ * UTickableTranslationConstraint
+ **/
+
+UTickableTranslationConstraint::UTickableTranslationConstraint()
+{
+	Type = ETransformConstraintType::Translation;
+}
+
+#if WITH_EDITOR
+
+void UTickableTranslationConstraint::PostEditChangeProperty(FPropertyChangedEvent& PropertyChangedEvent)
+{
+	Super::PostEditChangeProperty(PropertyChangedEvent);
+
+	if (HasAnyFlags(RF_ClassDefaultObject))
+	{
+		return;
+	}
+
+	const FName PropertyName = PropertyChangedEvent.GetPropertyName();
+	if (PropertyName == GET_MEMBER_NAME_CHECKED(UTickableTransformConstraint, bDynamicOffset))
+	{
+		if (bDynamicOffset)
+		{
+			Cache.CachedInputHash = CalculateInputHash();
+			
+			const FTransform ChildGlobalTransform = GetChildGlobalTransform();
+			const FTransform ParentWorldTransform = GetParentGlobalTransform();
+			OffsetTranslation = ChildGlobalTransform.GetLocation() - ParentWorldTransform.GetLocation();
+			
+			Evaluate();
+		}
+		return;
+	}
+}
+
+#endif
+
+void UTickableTranslationConstraint::ComputeOffset()
+{
+	const FTransform InitParentTransform = GetParentGlobalTransform();
+	const FTransform InitChildTransform = GetChildGlobalTransform();
+	
+	OffsetTranslation = FVector::ZeroVector;
+	if (bMaintainOffset || bDynamicOffset)
+	{
+		OffsetTranslation = InitChildTransform.GetLocation() - InitParentTransform.GetLocation();
+	}	
+}
+
+FConstraintTickFunction::ConstraintFunction UTickableTranslationConstraint::GetFunction() const
+{
+	return [this]()
+	{
+		if (!IsFullyActive())
+		{
+			return;
+		}
+		
+		const float ClampedWeight = FMath::Clamp<float>(Weight, 0.f, 1.f);
+		if (ClampedWeight < KINDA_SMALL_NUMBER)
+		{
+			return;
+		}
+
+		ConstraintLocals::PreEvaluateHandle(ChildTRSHandle);
+		
+		const FVector ParentTranslation = GetParentGlobalTransform().GetLocation();
+		FTransform Transform = GetChildGlobalTransform();
+		const FVector ChildTranslation = Transform.GetLocation();
+		
+		FVector NewTranslation = (!bMaintainOffset) ? ParentTranslation : ParentTranslation + OffsetTranslation;
+		if (ClampedWeight < 1.0f - KINDA_SMALL_NUMBER)
+		{
+			NewTranslation = FMath::Lerp<FVector>(ChildTranslation, NewTranslation, ClampedWeight);
+		}
+
+		AxisFilter.FilterVector(NewTranslation, ChildTranslation);
+
+		Transform.SetLocation(NewTranslation);
+			
+		SetChildGlobalTransform(Transform);
+		if (ChildTRSHandle)
+		{
+			ChildTRSHandle->TickTarget();
+		}
+	};
+}
+
+void UTickableTranslationConstraint::OnHandleModified(UTransformableHandle* InHandle, EHandleEvent InEvent)
+{
+	Super::OnHandleModified(InHandle, InEvent);
+	
+	if (!IsFullyActive() || !bDynamicOffset)
+	{
+		return;
+	}
+	
+	if (!InHandle || InHandle != ChildTRSHandle)
+	{
+		return;
+	}
+
+	const bool bUpdateFromGlobal = (InEvent == EHandleEvent::GlobalTransformUpdated);
+	const bool bUpdateTransform = InEvent == EHandleEvent::LocalTransformUpdated || bUpdateFromGlobal;
+	if (!bUpdateTransform)
+	{
+		return;
+	}
+
+	const uint32 InputHash = CalculateInputHash();
+	
+	// update dynamic offset
+	if (InputHash != Cache.CachedInputHash)
+	{
+		Cache.CachedInputHash = InputHash;
+
+		if (bUpdateFromGlobal)
+		{
+			const FTransform ParentWorldTransform = GetParentGlobalTransform();
+			const FTransform ChildGlobalTransform = GetChildGlobalTransform();
+			OffsetTranslation = ChildGlobalTransform.GetLocation() - ParentWorldTransform.GetLocation();
+		}
+		else
+		{
+			const FTransform ChildLocalTransform = GetChildLocalTransform();
+			OffsetTranslation = ChildLocalTransform.GetTranslation();
+		}
+	}
+}
+
+uint32 UTickableTranslationConstraint::CalculateInputHash() const
+{
+	uint32 Hash = 0;
+
+	// local location hash
+	const FTransform ChildLocalTransform = GetChildLocalTransform();
+	Hash = HashCombine(Hash, GetTypeHash(ChildLocalTransform.GetTranslation() ));
+
+	// global location hash
+	const FTransform ChildGlobalTransform = GetChildGlobalTransform();
+	Hash = HashCombine(Hash, GetTypeHash(ChildGlobalTransform.GetTranslation() ));
+	
+	return Hash;
+}
+
+/** 
+ * UTickableRotationConstraint
+ **/
+
+UTickableRotationConstraint::UTickableRotationConstraint()
+{
+	Type = ETransformConstraintType::Rotation;
+}
+
+#if WITH_EDITOR
+
+void UTickableRotationConstraint::PostEditChangeProperty(FPropertyChangedEvent& PropertyChangedEvent)
+{
+	Super::PostEditChangeProperty(PropertyChangedEvent);
+
+	if (HasAnyFlags(RF_ClassDefaultObject))
+	{
+		return;
+	}
+	
+	const FName PropertyName = PropertyChangedEvent.GetPropertyName();
+	if (PropertyName == GET_MEMBER_NAME_CHECKED(UTickableTransformConstraint, bDynamicOffset))
+	{
+		if (bDynamicOffset)
+		{
+			Cache.CachedInputHash = CalculateInputHash();
+			
+			const FTransform ChildGlobalTransform = GetChildGlobalTransform();
+			const FTransform ParentWorldTransform = GetParentGlobalTransform();
+			OffsetRotation = ParentWorldTransform.GetRotation().Inverse() * ChildGlobalTransform.GetRotation();
+			
+			Evaluate();
+		}
+		return;
+	}
+}
+
+#endif
+
+void UTickableRotationConstraint::ComputeOffset()
+{
+	const FTransform InitParentTransform = GetParentGlobalTransform();
+	const FTransform InitChildTransform = GetChildGlobalTransform();
+	
+	OffsetRotation = FQuat::Identity;
+	if (bMaintainOffset || bDynamicOffset)
+	{
+		OffsetRotation = InitParentTransform.GetRotation().Inverse() * InitChildTransform.GetRotation();
+		OffsetRotation.Normalize(); 
+	}
+}
+
+FConstraintTickFunction::ConstraintFunction UTickableRotationConstraint::GetFunction() const
+{
+	return [this]()
+	{
+		if (!IsFullyActive())
+		{
+			return;
+		}
+		
+		const float ClampedWeight = FMath::Clamp<float>(Weight, 0.f, 1.f);
+		if (ClampedWeight < KINDA_SMALL_NUMBER)
+		{
+			return;
+		}
+
+		ConstraintLocals::PreEvaluateHandle(ChildTRSHandle);
+		
+		const FQuat ParentRotation = GetParentGlobalTransform().GetRotation();
+		FTransform Transform = GetChildGlobalTransform();
+		const FQuat ChildRotation = Transform.GetRotation();
+
+		FQuat NewRotation = (!bMaintainOffset) ? ParentRotation : ParentRotation * OffsetRotation;
+		if (ClampedWeight < 1.0f - KINDA_SMALL_NUMBER)
+		{
+			NewRotation = FQuat::Slerp(ChildRotation, NewRotation, ClampedWeight);
+		}
+
+		AxisFilter.FilterQuat(NewRotation, ChildRotation);
+
+
+		Transform.SetRotation(NewRotation);
+		
+		SetChildGlobalTransform(Transform);
+
+		if (ChildTRSHandle)
+		{
+			ChildTRSHandle->TickTarget();
+		}
+	};
+}
+
+void UTickableRotationConstraint::OnHandleModified(UTransformableHandle* InHandle, EHandleEvent InEvent)
+{
+	Super::OnHandleModified(InHandle, InEvent);
+	
+	if (!IsFullyActive() || !bDynamicOffset)
+	{
+		return;
+	}
+	
+	if (!InHandle || InHandle != ChildTRSHandle)
+	{
+		return;
+	}
+
+	const bool bUpdateFromGlobal = (InEvent == EHandleEvent::GlobalTransformUpdated);
+	const bool bUpdateTransform = (InEvent == EHandleEvent::LocalTransformUpdated) || bUpdateFromGlobal;
+	if (!bUpdateTransform)
+	{
+		return;
+	}
+
+	const uint32 InputHash = CalculateInputHash();
+	
+	// update dynamic offset
+	if (InputHash != Cache.CachedInputHash)
+	{
+		Cache.CachedInputHash = InputHash;
+
+		if (bUpdateFromGlobal)
+		{
+			const FTransform ParentWorldTransform = GetParentGlobalTransform();
+			const FTransform ChildGlobalTransform = GetChildGlobalTransform();
+			OffsetRotation = ParentWorldTransform.GetRotation().Inverse() * ChildGlobalTransform.GetRotation();
+		}
+		else
+		{
+			OffsetRotation = GetChildLocalTransform().GetRotation();
+		}
+	}
+}
+
+uint32 UTickableRotationConstraint::CalculateInputHash() const
+{
+	uint32 Hash = 0;
+
+	// local rotation hash
+	const FTransform ChildLocalTransform = GetChildLocalTransform();
+	Hash = HashCombine(Hash, GetTypeHash(ChildLocalTransform.GetRotation().Euler() ));
+
+	// global rotation hash
+	const FTransform ChildGlobalTransform = GetChildGlobalTransform();
+	Hash = HashCombine(Hash, GetTypeHash(ChildGlobalTransform.GetRotation().Euler() ));
+	
+	return Hash;
+}
+
+/** 
+ * UTickableScaleConstraint
+ **/
+
+UTickableScaleConstraint::UTickableScaleConstraint()
+{
+	Type = ETransformConstraintType::Scale;
+}
+
+#if WITH_EDITOR
+
+void UTickableScaleConstraint::PostEditChangeProperty(FPropertyChangedEvent& PropertyChangedEvent)
+{
+	Super::PostEditChangeProperty(PropertyChangedEvent);
+
+	if (HasAnyFlags(RF_ClassDefaultObject))
+	{
+		return;
+	}
+	
+	const FName PropertyName = PropertyChangedEvent.GetPropertyName();
+	if (PropertyName == GET_MEMBER_NAME_CHECKED(UTickableTransformConstraint, bDynamicOffset))
+	{
+		if (bDynamicOffset)
+		{
+			Cache.CachedInputHash = CalculateInputHash();
+			
+			const FVector ParentScale = GetParentGlobalTransform().GetScale3D();
+			OffsetScale = GetChildGlobalTransform().GetScale3D();
+			OffsetScale[0] = FMath::Abs(ParentScale[0]) > KINDA_SMALL_NUMBER ? OffsetScale[0] / ParentScale[0] : 0.f;
+			OffsetScale[1] = FMath::Abs(ParentScale[1]) > KINDA_SMALL_NUMBER ? OffsetScale[1] / ParentScale[1] : 0.f;
+			OffsetScale[2] = FMath::Abs(ParentScale[2]) > KINDA_SMALL_NUMBER ? OffsetScale[2] / ParentScale[2] : 0.f;
+			
+			Evaluate();
+		}
+		return;
+	}
+}
+
+#endif
+
+void UTickableScaleConstraint::ComputeOffset()
+{
+	const FTransform InitParentTransform = GetParentGlobalTransform();
+	const FTransform InitChildTransform = GetChildGlobalTransform();
+	
+	OffsetScale = FVector::OneVector;
+	if (bMaintainOffset)
+	{
+		const FVector InitParentScale = InitParentTransform.GetScale3D();
+		OffsetScale = InitChildTransform.GetScale3D();
+		OffsetScale[0] = FMath::Abs(InitParentScale[0]) > KINDA_SMALL_NUMBER ? OffsetScale[0] / InitParentScale[0] : 0.f;
+		OffsetScale[1] = FMath::Abs(InitParentScale[1]) > KINDA_SMALL_NUMBER ? OffsetScale[1] / InitParentScale[1] : 0.f;
+		OffsetScale[2] = FMath::Abs(InitParentScale[2]) > KINDA_SMALL_NUMBER ? OffsetScale[2] / InitParentScale[2] : 0.f;
+	}
+}
+
+FConstraintTickFunction::ConstraintFunction UTickableScaleConstraint::GetFunction() const
+{
+	return [this]()
+	{
+		if (!IsFullyActive())
+		{
+			return;
+		}
+		
+		const float ClampedWeight = FMath::Clamp<float>(Weight, 0.f, 1.f);
+		if (ClampedWeight < KINDA_SMALL_NUMBER)
+		{
+			return;
+		}
+
+		ConstraintLocals::PreEvaluateHandle(ChildTRSHandle);
+		
+		const FVector ParentScale = GetParentGlobalTransform().GetScale3D();
+		FTransform Transform = GetChildGlobalTransform();
+		const FVector ChildScale = Transform.GetScale3D();
+		
+		FVector NewScale = (!bMaintainOffset) ? ParentScale : ParentScale * OffsetScale;
+		if (ClampedWeight < 1.0f - KINDA_SMALL_NUMBER)
+		{
+			NewScale = FMath::Lerp<FVector>(ChildScale, NewScale, ClampedWeight);
+		}
+
+		AxisFilter.FilterVector(NewScale, ChildScale);
+
+		Transform.SetScale3D(NewScale);
+		
+		SetChildGlobalTransform(Transform);
+		if (ChildTRSHandle)
+		{
+			ChildTRSHandle->TickTarget();
+		}
+	};
+}
+
+void UTickableScaleConstraint::OnHandleModified(UTransformableHandle* InHandle, EHandleEvent InEvent)
+{
+	Super::OnHandleModified(InHandle, InEvent);
+	
+	if (!IsFullyActive() || !bDynamicOffset)
+	{
+		return;
+	}
+	
+	if (!InHandle || InHandle != ChildTRSHandle)
+	{
+		return;
+	}
+
+	const bool bUpdateFromGlobal = (InEvent == EHandleEvent::GlobalTransformUpdated);
+	const bool bUpdateTransform = InEvent == EHandleEvent::LocalTransformUpdated || bUpdateFromGlobal;
+	if (!bUpdateTransform)
+	{
+		return;
+	}
+
+	const uint32 InputHash = CalculateInputHash();
+	
+	// update dynamic offset
+	if (InputHash != Cache.CachedInputHash)
+	{
+		Cache.CachedInputHash = InputHash;
+
+		if (bUpdateFromGlobal)
+		{
+			const FVector ParentScale = GetParentGlobalTransform().GetScale3D();
+			OffsetScale = GetChildGlobalTransform().GetScale3D();
+			OffsetScale[0] = FMath::Abs(ParentScale[0]) > KINDA_SMALL_NUMBER ? OffsetScale[0] / ParentScale[0] : 0.f;
+			OffsetScale[1] = FMath::Abs(ParentScale[1]) > KINDA_SMALL_NUMBER ? OffsetScale[1] / ParentScale[1] : 0.f;
+			OffsetScale[2] = FMath::Abs(ParentScale[2]) > KINDA_SMALL_NUMBER ? OffsetScale[2] / ParentScale[2] : 0.f;
+		}
+		else
+		{
+			const FTransform ChildLocalTransform = GetChildLocalTransform();
+			OffsetScale = ChildLocalTransform.GetScale3D();
+		}
+	}
+}
+
+uint32 UTickableScaleConstraint::CalculateInputHash() const
+{
+	uint32 Hash = 0;
+
+	// local scale hash
+	const FTransform ChildLocalTransform = GetChildLocalTransform();
+	Hash = HashCombine(Hash, GetTypeHash(ChildLocalTransform.GetScale3D() ));
+
+	// global scale hash
+	const FTransform ChildGlobalTransform = GetChildGlobalTransform();
+	Hash = HashCombine(Hash, GetTypeHash(ChildGlobalTransform.GetScale3D() ));
+	
+	return Hash;
+}
+
+/** 
+ * UTickableParentConstraint
+ **/
+
+UTickableParentConstraint::UTickableParentConstraint()
+{
+	Type = ETransformConstraintType::Parent;
+}
+
+void UTickableParentConstraint::ComputeOffset()
+{
+	const FTransform InitParentTransform = GetParentGlobalTransform();
+	FTransform InitChildTransform = GetChildGlobalTransform();
+	
+	OffsetTransform = FTransform::Identity;
+	if (bMaintainOffset || bDynamicOffset)
+	{
+		if (!bScaling)
+		{
+			InitChildTransform.RemoveScaling();
+		}
+		OffsetTransform = InitChildTransform.GetRelativeTransform(InitParentTransform); 
+	}
+}
+
+uint32 UTickableParentConstraint::CalculateInputHash() const
+{
+	uint32 Hash = 0;
+	
+	const FTransform ChildLocalTransform = GetChildLocalTransform();
+	Hash = HashCombine(Hash, GetTypeHash(ChildLocalTransform.GetTranslation()));
+	Hash = HashCombine(Hash, GetTypeHash(ChildLocalTransform.GetRotation().Euler() ));
+	Hash = HashCombine(Hash, GetTypeHash(ChildLocalTransform.GetScale3D()));
+	
+	const FTransform ChildGlobalTransform = GetChildGlobalTransform();
+	Hash = HashCombine(Hash, GetTypeHash(ChildGlobalTransform.GetTranslation()));
+	Hash = HashCombine(Hash, GetTypeHash(ChildGlobalTransform.GetRotation().Euler() ));
+	Hash = HashCombine(Hash, GetTypeHash(ChildGlobalTransform.GetScale3D()));
+	
+	return Hash;
+}
+
+FConstraintTickFunction::ConstraintFunction UTickableParentConstraint::GetFunction() const
+{
+	return [this]()
+	{
+		if (!IsFullyActive())
+		{
+			return;
+		}
+		
+		const float ClampedWeight = FMath::Clamp<float>(Weight, 0.f, 1.f);
+		if (ClampedWeight < KINDA_SMALL_NUMBER)
+		{
+			return;
+		}
+
+		auto LerpAndFilterTransform = [ClampedWeight](
+			const FTransform& InTransform, FTransform& OutTransform, const FTransformFilter& InFilter)
+		{
+			const FVector Location = InTransform.GetLocation();
+			const FQuat Rotation = InTransform.GetRotation();
+			const FVector Scale = InTransform.GetScale3D();
+
+			FVector NewLocation = OutTransform.GetLocation();
+			FQuat NewRotation = OutTransform.GetRotation();
+			FVector NewScale = OutTransform.GetScale3D();
+			
+			if (ClampedWeight < 1.0f - KINDA_SMALL_NUMBER)
+			{
+				NewLocation = FMath::Lerp<FVector>(Location, NewLocation, ClampedWeight);
+				NewRotation = FQuat::Slerp(Rotation, NewRotation, ClampedWeight);
+				NewScale = FMath::Lerp<FVector>(Scale, NewScale, ClampedWeight);
+			}
+
+			InFilter.TranslationFilter.FilterVector(NewLocation, Location);
+			InFilter.RotationFilter.FilterQuat(NewRotation, Rotation);
+			InFilter.ScaleFilter.FilterVector(NewScale, Scale);
+
+			OutTransform.SetLocation(NewLocation);
+			OutTransform.SetRotation(NewRotation);
+			OutTransform.SetScale3D(NewScale);
+		};
+
+		ConstraintLocals::PreEvaluateHandle(ChildTRSHandle);
+		
+		const FTransform ParentTransform = GetParentGlobalTransform();
+		
+		FTransform TargetTransform = (!bMaintainOffset) ? ParentTransform : OffsetTransform * ParentTransform;
+		//apply weight if needed
+		const FTransform ChildGlobalTransform = GetChildGlobalTransform();
+		LerpAndFilterTransform(ChildGlobalTransform, TargetTransform, TransformFilter);
+
+		//remove scale?
+		if (!bScaling)
+		{
+			TargetTransform.SetScale3D(ChildGlobalTransform.GetScale3D());
+		}
+
+		SetChildGlobalTransform(TargetTransform);
+		if (ChildTRSHandle)
+		{
+			ChildTRSHandle->TickTarget();
+		}
+	};
+}
+
+void UTickableParentConstraint::OnHandleModified(UTransformableHandle* InHandle, EHandleEvent InEvent)
+{
+	Super::OnHandleModified(InHandle, InEvent);
+	
+	if (!IsFullyActive() || !bDynamicOffset)
+	{
+		return;
+	}
+	
+	if (!InHandle || InHandle != ChildTRSHandle)
+	{
+		return;
+	}
+
+	const bool bUpdateFromGlobal = (InEvent == EHandleEvent::GlobalTransformUpdated);
+	const bool bUpdateTransform = (InEvent == EHandleEvent::LocalTransformUpdated) || bUpdateFromGlobal;
+	if (!bUpdateTransform)
+	{
+		return;
+	}
+	
+	const uint32 InputHash = CalculateInputHash();
+
+	// update dynamic offset
+	if (InputHash != Cache.CachedInputHash)
+	{
+		Cache.CachedInputHash = InputHash;
+
+		if (bUpdateFromGlobal)
+		{
+			const FTransform ParentWorldTransform = GetParentGlobalTransform();
+			FTransform ChildGlobalTransform = GetChildGlobalTransform();
+			if (!bScaling)
+			{
+				ChildGlobalTransform.RemoveScaling();
+			}
+			OffsetTransform = ChildGlobalTransform.GetRelativeTransform(ParentWorldTransform);
+		}
+		else
+		{
+			OffsetTransform = GetChildLocalTransform();
+		}
+	}
+}
+
+#if WITH_EDITOR
+
+void UTickableParentConstraint::PostEditChangeProperty(FPropertyChangedEvent& PropertyChangedEvent)
+{
+	Super::PostEditChangeProperty(PropertyChangedEvent);
+
+	if (HasAnyFlags(RF_ClassDefaultObject))
+	{
+		return;
+	}
+	
+	auto UpdateOffset = [&]()
+	{
+		FTransform ChildGlobalTransform = GetChildGlobalTransform();
+		if (!bScaling)
+		{
+			ChildGlobalTransform.RemoveScaling();
+		}
+		const FTransform ParentWorldTransform = GetParentGlobalTransform();
+		OffsetTransform = ChildGlobalTransform.GetRelativeTransform(ParentWorldTransform);
+	};
+	
+	const FName PropertyName = PropertyChangedEvent.GetPropertyName();
+	if (PropertyName == GET_MEMBER_NAME_CHECKED(UTickableTransformConstraint, bDynamicOffset))
+	{
+		if(bDynamicOffset)
+		{
+			Cache.CachedInputHash = CalculateInputHash();
+			UpdateOffset();			
+			Evaluate();
+		}
+		return;
+	}
+
+	if (PropertyName == GET_MEMBER_NAME_CHECKED(UTickableParentConstraint, bScaling))
+	{
+		// notify scale change. note that this is currently the only property change we monitor but this broadcast
+		// call could be made higher in the hierarchy to monitor other changes.
+		OnConstraintChanged.Broadcast(this, PropertyChangedEvent);
+		
+		if (bMaintainOffset || bDynamicOffset)
+		{
+			Cache.CachedInputHash = CalculateInputHash();
+			UpdateOffset();
+			Evaluate();
+			return;
+		}
+	}
+}
+
+#endif
+
+/** 
+ * UTickableLookAtConstraint
+ **/
+
+UTickableLookAtConstraint::UTickableLookAtConstraint()
+{
+	bMaintainOffset = false;
+	bDynamicOffset = false;
+	Type = ETransformConstraintType::LookAt;
+}
+
+void UTickableLookAtConstraint::ComputeOffset()
+{
+	bMaintainOffset = false;
+	bDynamicOffset = false;
+
+	const FVector InitParentLocation = GetParentGlobalTransform().GetLocation();
+	const FTransform InitChildTransform = GetChildGlobalTransform();
+	const FVector InitLookAtDir = (InitParentLocation - InitChildTransform.GetLocation()).GetSafeNormal();
+
+	if (!InitLookAtDir.IsNearlyZero())
+	{
+		Axis = InitChildTransform.InverseTransformVectorNoScale(InitLookAtDir).GetSafeNormal();
+	}
+}
+
+FConstraintTickFunction::ConstraintFunction UTickableLookAtConstraint::GetFunction() const
+{
+	return [this]()
+	{
+		if (!IsFullyActive())
+		{
+			return;
+		}
+
+		ConstraintLocals::PreEvaluateHandle(ChildTRSHandle);
+		
+		const FTransform ParentTransform = GetParentGlobalTransform();
+		const FTransform ChildTransform = GetChildGlobalTransform();
+		
+		const FVector LookAtDir = (ParentTransform.GetLocation() - ChildTransform.GetLocation()).GetSafeNormal();
+
+		if (!LookAtDir.IsNearlyZero() && !Axis.IsNearlyZero())
+		{
+			const FVector AxisToOrient = ChildTransform.TransformVectorNoScale(Axis).GetSafeNormal();
+		
+			FQuat Rotation = FindQuatBetweenNormals(AxisToOrient, LookAtDir);
+			const bool bNeedsToBeRotated = !Rotation.IsIdentity();
+			if (bNeedsToBeRotated)
+			{
+				Rotation = Rotation * ChildTransform.GetRotation();
+
+				FTransform Transform = ChildTransform;
+				Transform.SetRotation(Rotation.GetNormalized());
+				SetChildGlobalTransform(Transform);
+				if (ChildTRSHandle)
+				{
+					ChildTRSHandle->TickTarget();
+				}
+			}
+		}
+	};
+}
+
+bool UTickableLookAtConstraint::NeedsCompensation() const
+{
+	return false;
+}
+
+FQuat UTickableLookAtConstraint::FindQuatBetweenNormals(const FVector& A, const FVector& B)
+{
+	const FQuat::FReal Dot = FVector::DotProduct(A, B);
+	FQuat::FReal W = 1 + Dot;
+	FQuat Result;
+
+	if (W < SMALL_NUMBER)
+	{
+		// A and B point in opposite directions
+		W = 2 - W;
+		Result = FQuat( -A.Y * B.Z + A.Z * B.Y, -A.Z * B.X + A.X * B.Z, -A.X * B.Y + A.Y * B.X, W).GetNormalized();
+
+		const FVector Normal = FMath::Abs(A.X) > FMath::Abs(A.Y) ? FVector::YAxisVector : FVector::XAxisVector;
+		const FVector BiNormal = FVector::CrossProduct(A, Normal);
+		const FVector TauNormal = FVector::CrossProduct(A, BiNormal);
+		Result = Result * FQuat(TauNormal, PI);
+	}
+	else
+	{
+		//Axis = FVector::CrossProduct(A, B);
+		Result = FQuat( A.Y * B.Z - A.Z * B.Y, A.Z * B.X - A.X * B.Z, A.X * B.Y - A.Y * B.X, W);
+	}
+
+	Result.Normalize();
+	return Result;
+}
+
+#if WITH_EDITOR
+
+void UTickableLookAtConstraint::PostEditChangeProperty(FPropertyChangedEvent& PropertyChangedEvent)
+{
+	Super::PostEditChangeProperty(PropertyChangedEvent);
+
+	if (HasAnyFlags(RF_ClassDefaultObject))
+	{
+		return;
+	}
+
+	const FName PropertyName = PropertyChangedEvent.GetMemberPropertyName();
+	if (PropertyName == GET_MEMBER_NAME_CHECKED(UTickableLookAtConstraint, Axis))
+	{
+		Evaluate();
+		return;
+	}
+}
+
+#endif
+
+/** 
+ * DEPRECATED FTransformConstraintUtils
+ **/
+
+UTransformableComponentHandle* FTransformConstraintUtils::CreateHandleForSceneComponent(
+	USceneComponent* InSceneComponent,
+	const FName& InSocketName)
+{
+	return UE::TransformConstraintUtil::CreateHandleForSceneComponent(InSceneComponent, InSocketName);
+}
+
+void FTransformConstraintUtils::GetParentConstraints(
+	UWorld* InWorld,
+	const AActor* InChild,
+	TArray< TWeakObjectPtr<UTickableConstraint> >& OutConstraints)
+{
+	return UE::TransformConstraintUtil::GetParentConstraints(InWorld, InChild, OutConstraints);
+}
+
+UTickableTransformConstraint* FTransformConstraintUtils::CreateFromType(
+	UWorld* InWorld,
+	const ETransformConstraintType InType,
+	const bool bUseDefault)
+{
+	return UE::TransformConstraintUtil::CreateFromType(InWorld, InType, bUseDefault);
+}
+
+UTickableTransformConstraint* FTransformConstraintUtils::CreateAndAddFromObjects(
+	UWorld* InWorld,
+	UObject* InParent, const FName& InParentSocketName,
+	UObject* InChild, const FName& InChildSocketName,
+	const ETransformConstraintType InType,
+	const bool bMaintainOffset,
+	const bool bUseDefault)
+{
+	return UE::TransformConstraintUtil::CreateAndAddFromObjects(
+		InWorld, InParent, InParentSocketName, InChild, InChildSocketName, InType, bMaintainOffset, bUseDefault);
+}
+
+bool FTransformConstraintUtils::AddConstraint(
+	UWorld* InWorld,
+	UTransformableHandle* InParentHandle,
+	UTransformableHandle* InChildHandle,
+	UTickableTransformConstraint* InNewConstraint,
+	const bool bMaintainOffset,
+	const bool bUseDefault)
+{
+	return UE::TransformConstraintUtil::AddConstraint(InWorld, InParentHandle, InChildHandle, InNewConstraint, bMaintainOffset, bUseDefault);
+}
+
+void FTransformConstraintUtils::UpdateTransformBasedOnConstraint(FTransform& CurrentTransform, USceneComponent* SceneComponent)
+{
+	return UE::TransformConstraintUtil::UpdateTransformBasedOnConstraint(CurrentTransform, SceneComponent);
+}
+
+FTransform FTransformConstraintUtils::ComputeRelativeTransform(
+	const FTransform& InChildLocal,
+	const FTransform& InChildWorld,
+	const FTransform& InSpaceWorld,
+	const UTickableTransformConstraint* InConstraint)
+{
+	return UE::TransformConstraintUtil::ComputeRelativeTransform(InChildLocal, InChildWorld, InSpaceWorld, InConstraint);
+}
+
+TOptional<FTransform> FTransformConstraintUtils::GetRelativeTransform(UWorld* InWorld, const uint32 InHandleHash)
+{
+	return UE::TransformConstraintUtil::GetRelativeTransform(InWorld, InHandleHash);
+}
+
+TOptional<FTransform> FTransformConstraintUtils::GetConstraintsRelativeTransform(
+	const TArray< TWeakObjectPtr<UTickableConstraint> >& InConstraints,
+	const FTransform& InChildLocal, const FTransform& InChildWorld)
+{
+	return UE::TransformConstraintUtil::GetConstraintsRelativeTransform(InConstraints, InChildLocal, InChildWorld);
+}
+
+int32 FTransformConstraintUtils::GetLastActiveConstraintIndex(const TArray< TWeakObjectPtr<UTickableConstraint> >& InConstraints)
+{
+	return UE::TransformConstraintUtil::GetLastActiveConstraintIndex(InConstraints);
+}
+
+void FTransformConstraintUtils::GetChildrenConstraints(
+	UWorld* World,
+	const UTickableTransformConstraint* InConstraint,
+	TArray< TWeakObjectPtr<UTickableConstraint> >& OutConstraints,
+	const bool bIncludeTarget)
+{
+	return UE::TransformConstraintUtil::GetChildrenConstraints(World, InConstraint, OutConstraints, bIncludeTarget);
+}
diff --git a/Engine/Source/Runtime/Experimental/Animation/Constraints/Private/Transform/TransformConstraintUtil.cpp b/Engine/Source/Runtime/Experimental/Animation/Constraints/Private/Transform/TransformConstraintUtil.cpp
--- a/Engine/Source/Runtime/Experimental/Animation/Constraints/Private/Transform/TransformConstraintUtil.cpp
+++ b/Engine/Source/Runtime/Experimental/Animation/Constraints/Private/Transform/TransformConstraintUtil.cpp
@@ -2,12 +2,16 @@
 
 #include "Transform/TransformConstraintUtil.h"
 
-#include "ConstraintsManager.h"
-#include "TransformableHandle.h"
-#include "TransformableRegistry.h"
-#include "TransformConstraint.h"
-
+#include "ConstraintsManager.inl"
+#include "Engine/World.h"
+#include "GameFramework/Actor.h"
+#include "MovieSceneSection.h"
+#include "Transform/TransformableHandle.h"
+#include "Transform/TransformableRegistry.h"
+#include "Transform/TransformConstraint.h"
+#include "Transform/DependencyBuilder.h"
 #include "UObject/Object.h"
+#include "UObject/Package.h"
 
 namespace UE::TransformConstraintUtil
 {
@@ -15,7 +19,7 @@ namespace UE::TransformConstraintUtil
 namespace Private
 {
 
-uint32 GetHandleHash(const UObject* InObject, const FName& InAttachmentName)
+uint32 GetHandleHash(const UObject* InObject, const FName& InAttachmentName = NAME_None)
 {
 	if (InObject)
 	{
@@ -29,6 +33,18 @@ uint32 GetHandleHash(const UObject* InObject, const FName& InAttachmentName)
 	return 0;
 }
 
+UTransformableHandle* GetHandle(UObject* InObject, const FName& InSocketName)
+{
+	// look for customized transform handle
+	const FTransformableRegistry& Registry = FTransformableRegistry::Get();
+	if (const FTransformableRegistry::CreateHandleFuncT CreateFunction = Registry.GetCreateFunction(InObject->GetClass()))
+	{
+		return CreateFunction(InObject, InSocketName);
+	}
+
+	return nullptr;
+}
+	
 }
 
 const TArray< TWeakObjectPtr<UTickableConstraint> >& FConstraintsInteractionCache::Get(const UObject* InObject, const FName& InAttachmentName)
@@ -94,13 +110,13 @@ const TArray< TWeakObjectPtr<UTickableConstraint> >& FConstraintsInteractionCach
 bool FConstraintsInteractionCache::HasAnyActiveConstraint(const uint32 InHandleHash, UWorld* InWorld)
 {
 	const TArray< TWeakObjectPtr<UTickableConstraint> >& TransformConstraints = Get(InHandleHash, InWorld);
-	return FTransformConstraintUtils::GetLastActiveConstraintIndex(TransformConstraints) != INDEX_NONE;
+	return GetLastActiveConstraintIndex(TransformConstraints) != INDEX_NONE;
 }
 
 TOptional<FTransform> FConstraintsInteractionCache::GetParentTransform(const uint32 InHandleHash, UWorld* InWorld)
 {
 	const TArray< TWeakObjectPtr<UTickableConstraint> >& TransformConstraints = Get(InHandleHash, InWorld);
-	const int32 ConstraintsIndex = FTransformConstraintUtils::GetLastActiveConstraintIndex(TransformConstraints);
+	const int32 ConstraintsIndex = GetLastActiveConstraintIndex(TransformConstraints);
 	if (ConstraintsIndex != INDEX_NONE)
 	{
 		UTickableTransformConstraint* Constraint = Cast<UTickableTransformConstraint>(TransformConstraints[ConstraintsIndex].Get());
@@ -141,5 +157,582 @@ void FConstraintsInteractionCache::UnregisterNotifications(UWorld* InWorld)
 		ConstraintsNotificationHandle.Reset();
 	}
 }
+
+UTransformableComponentHandle* CreateHandleForSceneComponent(USceneComponent* InSceneComponent, const FName& InSocketName)
+{
+	UTransformableComponentHandle* ComponentHandle = nullptr;
+	if (InSceneComponent)
+	{
+		ComponentHandle = NewObject<UTransformableComponentHandle>(GetTransientPackage(), NAME_None, RF_Transactional);
+		ComponentHandle->Component = InSceneComponent;
+		ComponentHandle->SocketName = InSocketName;
+		InSceneComponent->SetMobility(EComponentMobility::Movable);
+		ComponentHandle->RegisterDelegates();
+	}
+	return ComponentHandle;
+}
+
+void GetParentConstraints( UWorld* InWorld, const AActor* InChild, TArray< TWeakObjectPtr<UTickableConstraint> >& OutConstraints)
+{
+	if (!InWorld || !InChild)
+	{
+		return;
+	}
+
+	const uint32 ChildHash = Private::GetHandleHash(InChild);
+	if (ChildHash == 0)
+	{
+		return;
+	}
+	
+	static constexpr bool bSorted = true;
+	const FConstraintsManagerController& Controller = FConstraintsManagerController::Get(InWorld);
+	OutConstraints.Append(Controller.GetParentConstraints(ChildHash, bSorted));
+}
+
+UTickableTransformConstraint* CreateFromType(UWorld* InWorld, const ETransformConstraintType InType, const bool bUseDefault)
+{
+	if (!InWorld)
+	{
+		UE_LOG(LogTemp, Error, TEXT("TransformConstraintUtil::CreateFromType sanity check failed."));
+		return nullptr;
+	}
+	const UEnum* ETransformConstraintTypeEnum = StaticEnum<ETransformConstraintType>();
+	if (!ETransformConstraintTypeEnum->IsValidEnumValue(static_cast<int64>(InType)))
+	{
+		UE_LOG(LogTemp, Error, TEXT("Constraint Type %d not recognized"), int(InType));
+		return nullptr;
+	}
+
+
+	// unique name (we may want to use another approach here to manage uniqueness)
+	const FString ConstraintTypeStr = ETransformConstraintTypeEnum->GetNameStringByValue((uint8)InType);
+	const FName BaseName(*FString::Printf(TEXT("%sConstraint"), *ConstraintTypeStr));
+
+	const FConstraintsManagerController& Controller = FConstraintsManagerController::Get(InWorld);
+	UTickableTransformConstraint* Constraint = nullptr;
+
+	switch (InType)
+	{
+	case ETransformConstraintType::Translation:
+		Constraint = Controller.AllocateConstraintT<UTickableTranslationConstraint>(BaseName, bUseDefault);
+		break;
+	case ETransformConstraintType::Rotation:
+		Constraint = Controller.AllocateConstraintT<UTickableRotationConstraint>(BaseName, bUseDefault);
+		break;
+	case ETransformConstraintType::Scale:
+		Constraint = Controller.AllocateConstraintT<UTickableScaleConstraint>(BaseName, bUseDefault);
+		break;
+	case ETransformConstraintType::Parent:
+		Constraint = Controller.AllocateConstraintT<UTickableParentConstraint>(BaseName, bUseDefault);
+		break;
+	case ETransformConstraintType::LookAt:
+		Constraint = Controller.AllocateConstraintT<UTickableLookAtConstraint>(BaseName, bUseDefault);
+		break;
+	default:
+		ensure(false);
+		break;
+	}
+	return Constraint;
+}
+
+namespace Private
+{
+
+// we suppose that both InParentHandle and InChildHandle are safe to use
+bool HasConstraintDependencyWith(UWorld* InWorld, const UTransformableHandle* InParentHandle, const UTransformableHandle* InChildHandle)
+{
+	using ConstraintPtr = TWeakObjectPtr<UTickableConstraint>;
+	using HandlePtr = TObjectPtr<UTransformableHandle>;
+
+	static constexpr bool bSorted = false;
+	const FConstraintsManagerController& Controller = FConstraintsManagerController::Get(InWorld);
+	const TArray<ConstraintPtr> Constraints = Controller.GetParentConstraints(InParentHandle->GetHash(), bSorted);
+
+	// get parent handles
+	TArray< HandlePtr > ParentHandles;
+	for (const ConstraintPtr& Constraint: Constraints)
+	{
+		if (const UTickableTransformConstraint* TransformConstraint = Cast<UTickableTransformConstraint>(Constraint.Get()))
+		{
+			if (IsValid(TransformConstraint->ParentTRSHandle))
+			{
+				ParentHandles.Add(TransformConstraint->ParentTRSHandle);
+			}
+		}
+	}
+
+	// check if InChildHandle is one of them
+	const uint32 ChildHash = InChildHandle->GetHash();
+	const bool bIsParentADependency = ParentHandles.ContainsByPredicate([ChildHash](const HandlePtr& InHandle)
+	{
+		return InHandle->GetHash() == ChildHash;
+	});
+
+	if (bIsParentADependency)
+	{
+		return true;
+	}
+
+	// if not, recurse
+	for (const HandlePtr& ParentHandle: ParentHandles)
+	{
+		if (HasConstraintDependencyWith(InWorld, ParentHandle, InChildHandle))
+		{
+			return true;
+		}
+	}
+
+	return false;
+}
+	
+bool AreHandlesConstrainable( UWorld* InWorld, UTransformableHandle* InParentHandle, UTransformableHandle* InChildHandle)
+{
+	static const TCHAR* ErrorPrefix = TEXT("Dependency error:");
+
+	if (InChildHandle->GetHash() == InParentHandle->GetHash())
+	{
+		UE_LOG(LogTemp, Error, TEXT("%s handles are pointing at the same object."), ErrorPrefix);
+		return false;
+	}
+
+	// check for direct transform dependencies (ei hierarchy)
+	if (InParentHandle->HasDirectDependencyWith(*InChildHandle))
+	{
+#if WITH_EDITOR
+		UE_LOG(LogTemp, Error, TEXT("%s: %s has a direct dependency with %s."),
+			ErrorPrefix, *InParentHandle->GetLabel(), *InChildHandle->GetLabel());
+#endif
+		return false;
+	}
+
+	// check for indirect transform dependencies (ei constraint chain)
+	if (HasConstraintDependencyWith(InWorld, InParentHandle, InChildHandle))
+	{
+#if WITH_EDITOR
+		UE_LOG(LogTemp, Error, TEXT("%s: %s has an indirect dependency with %s."),
+			ErrorPrefix, *InParentHandle->GetLabel(), *InChildHandle->GetLabel());
+#endif
+		return false;
+	}
+
+	return true;
+}
+
+}
+
+UTickableTransformConstraint* CreateAndAddFromObjects(
+	UWorld* InWorld,
+	UObject* InParent, const FName& InParentSocketName,
+	UObject* InChild, const FName& InChildSocketName,
+	const ETransformConstraintType InType,
+	const bool bMaintainOffset,
+	const bool bUseDefault)
+{
+	static const TCHAR* ErrorPrefix = TEXT("TransformConstraintUtil::CreateAndAddFromActors");
+	
+	// SANITY CHECK
+	if (!InWorld || !InParent || !InChild)
+	{
+		UE_LOG(LogTemp, Error, TEXT("%s sanity check failed."), ErrorPrefix);
+		return nullptr;
+	}
+
+	UTransformableHandle* ParentHandle = Private::GetHandle(InParent, InParentSocketName);
+	if (!ParentHandle)
+	{
+		return nullptr;
+	}
+	
+	UTransformableHandle* ChildHandle = Private::GetHandle(InChild, InChildSocketName);
+	if (!ChildHandle)
+	{
+		return nullptr;
+	}
+
+	const bool bCanConstrain = Private::AreHandlesConstrainable(InWorld, ParentHandle, ChildHandle);
+	if (!bCanConstrain)
+	{
+		ChildHandle->MarkAsGarbage();
+		ParentHandle->MarkAsGarbage();
+		return nullptr;
+	}
+	
+	UTickableTransformConstraint* Constraint = CreateFromType(InWorld, InType, bUseDefault);
+	if (Constraint && (ParentHandle->IsValid() && ChildHandle->IsValid()))
+	{
+		if (AddConstraint(InWorld, ParentHandle, ChildHandle, Constraint, bMaintainOffset, bUseDefault) == false)
+		{
+			Constraint->MarkAsGarbage();
+			Constraint = nullptr;
+		}
+	}
+	return Constraint;
+}
+
+bool AddConstraint(UWorld* InWorld, UTransformableHandle* InParentHandle, UTransformableHandle* InChildHandle,
+	UTickableTransformConstraint* InNewConstraint, const bool bMaintainOffset, const bool bUseDefault)
+{
+	const bool bIsValidParent = InParentHandle && InParentHandle->IsValid();
+	const bool bIsValidChild = InChildHandle && InChildHandle->IsValid();
+	if (!bIsValidParent || !bIsValidChild)
+	{
+		UE_LOG(LogTemp, Error, TEXT("TransformConstraintUtil::AddConst error adding constraint"));
+		return false;
+	}
+
+	if (InNewConstraint == nullptr)
+	{
+		UE_LOG(LogTemp, Error, TEXT("TransformConstraintUtil::AddConst error creating constraint"));
+		return false;
+	}
+
+	// set handles before calling AddConstraint as it will build dependencies based on them
+	InNewConstraint->ParentTRSHandle = InParentHandle;
+	InNewConstraint->ChildTRSHandle = InChildHandle;
+
+	// add the new one
+	const FConstraintsManagerController& Controller = FConstraintsManagerController::Get(InWorld);
+	const bool bConstraintAdded = Controller.AddConstraint(InNewConstraint);
+
+	if (!bConstraintAdded)
+	{
+		InNewConstraint->ParentTRSHandle = nullptr;
+		InNewConstraint->ChildTRSHandle = nullptr;
+		return false;
+	}
+
+	if (!bUseDefault)
+	{
+		InNewConstraint->bMaintainOffset = bMaintainOffset;
+		InNewConstraint->Setup();
+	}
+	InNewConstraint->InitConstraint(InWorld);
+	
+	return true;
+}
+
+void UpdateTransformBasedOnConstraint(FTransform& InOutCurrentTransform, const USceneComponent* InSceneComponent)
+{
+	if (AActor* Actor = InSceneComponent->GetTypedOuter<AActor>())
+	{
+		TArray< TWeakObjectPtr<UTickableConstraint> > Constraints;
+		GetParentConstraints(InSceneComponent->GetWorld(), Actor, Constraints);
+
+		const int32 LastActiveIndex = GetLastActiveConstraintIndex(Constraints);
+		if (Constraints.IsValidIndex(LastActiveIndex))
+		{
+			// switch to constraint space
+			const FTransform WorldTransform = InSceneComponent->GetSocketTransform(InSceneComponent->GetAttachSocketName());
+			const TOptional<FTransform> RelativeTransform = GetConstraintsRelativeTransform(Constraints, InOutCurrentTransform, WorldTransform);
+			if (RelativeTransform)
+			{
+				InOutCurrentTransform = *RelativeTransform;
+			}
+		}
+	}
+}
+
+FTransform ComputeRelativeTransform( const FTransform& InChildLocal, const FTransform& InChildWorld, const FTransform& InSpaceWorld,
+	const UTickableTransformConstraint* InConstraint)
+{
+	if (!InConstraint)
+	{
+		return InChildWorld.GetRelativeTransform(InSpaceWorld);
+	}
+	
+	const ETransformConstraintType ConstraintType = static_cast<ETransformConstraintType>(InConstraint->GetType());
+	switch (ConstraintType)
+	{
+	case ETransformConstraintType::Translation:
+		{
+			FTransform RelativeTransform = InChildLocal;
+			FVector RelativeTranslation = InChildWorld.GetLocation() - InSpaceWorld.GetLocation();
+			if (const UTickableTranslationConstraint* TranslationConstraint = Cast<UTickableTranslationConstraint>(InConstraint))
+			{
+				TranslationConstraint->AxisFilter.FilterVector(RelativeTranslation, InChildLocal.GetTranslation());
+			}
+			RelativeTransform.SetLocation(RelativeTranslation);
+			return RelativeTransform;
+		}
+	case ETransformConstraintType::Rotation:
+		{
+			FTransform RelativeTransform = InChildLocal;
+			FQuat RelativeRotation = InSpaceWorld.GetRotation().Inverse() * InChildWorld.GetRotation();
+			RelativeRotation.Normalize();
+			if (const UTickableRotationConstraint* RotationConstraint = Cast<UTickableRotationConstraint>(InConstraint))
+			{
+				RotationConstraint->AxisFilter.FilterQuat(RelativeRotation, InChildLocal.GetRotation());
+			}
+			RelativeTransform.SetRotation(RelativeRotation);
+			return RelativeTransform;
+		}
+	case ETransformConstraintType::Scale:
+		{
+			FTransform RelativeTransform = InChildLocal;
+			const FVector SpaceScale = InSpaceWorld.GetScale3D();
+			FVector RelativeScale = InChildWorld.GetScale3D();
+			RelativeScale[0] = FMath::Abs(SpaceScale[0]) > KINDA_SMALL_NUMBER ? RelativeScale[0] / SpaceScale[0] : 0.f;
+			RelativeScale[1] = FMath::Abs(SpaceScale[1]) > KINDA_SMALL_NUMBER ? RelativeScale[1] / SpaceScale[1] : 0.f;
+			RelativeScale[2] = FMath::Abs(SpaceScale[2]) > KINDA_SMALL_NUMBER ? RelativeScale[2] / SpaceScale[2] : 0.f;
+			if (const UTickableScaleConstraint* ScaleConstraint = Cast<UTickableScaleConstraint>(InConstraint))
+			{
+				ScaleConstraint->AxisFilter.FilterVector(RelativeScale, InChildLocal.GetScale3D());
+			}
+			RelativeTransform.SetScale3D(RelativeScale);
+			return RelativeTransform;
+		}
+	case ETransformConstraintType::Parent:
+		{
+			const UTickableParentConstraint* ParentConstraint = Cast<UTickableParentConstraint>(InConstraint);
+			const bool bScale = ParentConstraint ? ParentConstraint->IsScalingEnabled() : true;
+			
+			FTransform ChildTransform = InChildWorld;
+			if (!bScale)
+			{
+				ChildTransform.RemoveScaling();
+			}
+
+			FTransform RelativeTransform = ChildTransform.GetRelativeTransform(InSpaceWorld);
+
+			if (ParentConstraint && !ParentConstraint->TransformFilter.TranslationFilter.HasNoEffect())
+			{
+				FVector RelativeLocation = RelativeTransform.GetLocation();
+				ParentConstraint->TransformFilter.TranslationFilter.FilterVector(RelativeLocation, InChildLocal.GetLocation());
+				RelativeTransform.SetLocation(RelativeLocation);
+			}
+
+			if (ParentConstraint && !ParentConstraint->TransformFilter.RotationFilter.HasNoEffect())
+			{
+				FQuat RelativeRotation = RelativeTransform.GetRotation();
+				ParentConstraint->TransformFilter.RotationFilter.FilterQuat(RelativeRotation, InChildLocal.GetRotation());
+				RelativeTransform.SetRotation(RelativeRotation);
+			}
+
+			if (ParentConstraint && !ParentConstraint->TransformFilter.ScaleFilter.HasNoEffect())
+			{
+				FVector RelativeScale = RelativeTransform.GetScale3D();
+				ParentConstraint->TransformFilter.ScaleFilter.FilterVector(RelativeScale, InChildLocal.GetScale3D());
+				RelativeTransform.SetScale3D(RelativeScale);
+			}
+			
+			if (!bScale)
+			{
+				RelativeTransform.SetScale3D(InChildLocal.GetScale3D());
+			}
+			return RelativeTransform;
+		}
+	case ETransformConstraintType::LookAt:
+		return InChildLocal;
+	default:
+		break;
+	}
+	
+	return InChildWorld.GetRelativeTransform(InSpaceWorld);
+}
+
+TOptional<FTransform> GetRelativeTransform(UWorld* InWorld, const uint32 InHandleHash)
+{
+	if (!InWorld || InHandleHash <= 0)
+	{
+		return TOptional<FTransform>();
+	}
+
+	static constexpr bool bSorted = true;
+	const FConstraintsManagerController& Controller = FConstraintsManagerController::Get(InWorld);
+
+	const TArray< TWeakObjectPtr<UTickableConstraint> > Constraints = Controller.GetParentConstraints(InHandleHash, bSorted);
+	if (Constraints.IsEmpty())
+	{
+		return TOptional<FTransform>();
+	}
+
+	// get current active transform constraint 	
+	const int32 LastActiveIndex = GetLastActiveConstraintIndex(Constraints);
+	if (!Constraints.IsValidIndex(LastActiveIndex))
+	{
+		return TOptional<FTransform>();
+	}
+
+	// get relative transform
+	const UTickableTransformConstraint* Constraint = Cast<UTickableTransformConstraint>(Constraints[LastActiveIndex]);
+	const FTransform ChildLocal = Constraint->GetChildLocalTransform();
+	const FTransform ChildGlobal = Constraint->GetChildGlobalTransform();
+	
+	return GetConstraintsRelativeTransform(Constraints, ChildLocal, ChildGlobal);
+}
+
+TOptional<FTransform> GetConstraintsRelativeTransform(
+	const TArray< TWeakObjectPtr<UTickableConstraint> >& InConstraints,
+	const FTransform& InChildLocal, const FTransform& InChildWorld)
+{
+	if (InConstraints.IsEmpty())
+	{
+		return TOptional<FTransform>();
+	}
+
+	// get current active transform constraint
+	const int32 LastActiveIndex = GetLastActiveConstraintIndex(InConstraints);
+	if (!InConstraints.IsValidIndex(LastActiveIndex))
+	{
+		return TOptional<FTransform>();
+	}
+
+	// get relative transform
+	// if that constraint handles the entire transform then return the relative transform directly
+	const UTickableTransformConstraint* Constraint = Cast<UTickableTransformConstraint>(InConstraints[LastActiveIndex]);
+	if (EnumHasAllFlags(Constraint->GetChannelsToKey(), EMovieSceneTransformChannel::AllTransform))
+	{
+		const FTransform ParentGlobal = Constraint->GetParentGlobalTransform();
+		return ComputeRelativeTransform(InChildLocal, InChildWorld, ParentGlobal, Constraint);
+	}
+
+	// otherwise, we need to look for constraints on a sub-transform basis so we compute the relative transform for each of them
+	using ConstraintPtr = TWeakObjectPtr<UTickableConstraint>;
+	auto GetLastSubTransformIndex = [InConstraints](const EMovieSceneTransformChannel& InChannel)
+	{
+		return InConstraints.FindLastByPredicate([InChannel](const ConstraintPtr& InConstraint)
+		{
+			if (const UTickableTransformConstraint* TransformConstraint = Cast<UTickableTransformConstraint>(InConstraint.Get()))
+			{
+				const bool bHasChannelFlag = EnumHasAllFlags(TransformConstraint->GetChannelsToKey(), InChannel);
+				return InConstraint->IsFullyActive() && TransformConstraint->bDynamicOffset && bHasChannelFlag;
+			}
+			return false;
+		});
+	};
+
+	// look for last constraint index for each channel 
+	static constexpr EMovieSceneTransformChannel SubChannels[3] = {	EMovieSceneTransformChannel::Translation,
+																	EMovieSceneTransformChannel::Rotation,
+																	EMovieSceneTransformChannel::Scale };
+	TArray<int32> SubTransformIndices;
+	for (int32 Index = 0; Index < 3; Index++)
+	{
+		const int32 LastSubIndex = GetLastSubTransformIndex(SubChannels[Index]);
+    	if (InConstraints.IsValidIndex(LastSubIndex))  
+    	{
+    		SubTransformIndices.AddUnique(LastSubIndex);
+    	}
+	}
+	SubTransformIndices.Sort();
+
+	// if none then return 
+	if (SubTransformIndices.IsEmpty())
+	{
+		return TOptional<FTransform>();
+	}
+
+	// iterate thru constraints to compute the relative transform in each of them 
+	FTransform ChildLocal = InChildLocal;
+	for (const int32 SubConstraintIndex: SubTransformIndices)
+	{
+		const UTickableTransformConstraint* SubConstraint = Cast<UTickableTransformConstraint>(InConstraints[SubConstraintIndex]);
+		const FTransform ParentGlobal = SubConstraint->GetParentGlobalTransform();
+		ChildLocal = ComputeRelativeTransform(ChildLocal, InChildWorld, ParentGlobal, SubConstraint);
+	}
+	return ChildLocal;
+}
+
+int32 GetLastActiveConstraintIndex(const TArray< TWeakObjectPtr<UTickableConstraint> >& InConstraints)
+{
+	return InConstraints.FindLastByPredicate([](const TWeakObjectPtr<UTickableConstraint>& InConstraint)
+	{
+	   if (const UTickableTransformConstraint* TransformConstraint = Cast<UTickableTransformConstraint>(InConstraint.Get()))
+	   {
+		   return InConstraint->Active && TransformConstraint->bDynamicOffset;
+	   }
+	   return false;
+	});
+}
+
+void GetChildrenConstraints( UWorld* World, const UTickableTransformConstraint* InConstraint,
+	TArray< TWeakObjectPtr<UTickableConstraint> >& OutConstraints,
+	const bool bIncludeTarget)
+{
+	if (!InConstraint || !InConstraint->IsValid())
+	{
+		// this probably has been checked before but we want to make sure the data is safe to use
+		return;
+	}
+	
+	const UTransformableHandle* Handle = InConstraint->ChildTRSHandle.Get();
+	
+	using ConstraintPtr = TWeakObjectPtr<UTickableConstraint>;
+	
+	// filter for transform constraints where the InHandle is the parent (based on its hash value)
+	// and also has the same target if bIncludeTarget is true
+	const uint32 ParentHash = Handle->GetHash();
+	const UObject* ParentTarget = Handle->GetTarget().Get();
+	auto Predicate = [InConstraint, Handle, ParentHash, bIncludeTarget, ParentTarget, World](const ConstraintPtr& Constraint)
+	{
+		const UTickableTransformConstraint* TransformConstraint = Cast<UTickableTransformConstraint>(Constraint.Get());
+		if (!TransformConstraint || TransformConstraint == InConstraint)
+		{
+			return false;
+		}
+
+		if (TransformConstraint->ParentTRSHandle)
+		{
+			const UTransformableHandle* OtherParentHandle = TransformConstraint->ParentTRSHandle;
+			if (OtherParentHandle->GetHash() == ParentHash)
+			{
+				return true;
+			}
+
+			if (bIncludeTarget && ParentTarget)
+			{
+				const UObject* Target = OtherParentHandle->GetTarget().Get();
+				if (Target == ParentTarget)
+				{
+					// check direct dependencies to avoid evaluation order issues
+					if (Handle->HasDirectDependencyWith(*OtherParentHandle))
+					{
+						return false;
+					}
+
+					// check constraints dependencies to avoid cycles
+					if (Private::HasConstraintDependencyWith(World, OtherParentHandle, Handle))
+					{
+						return false;
+					}
+
+					// check TransformConstraint's ChildHandle
+					if (const UTransformableHandle* OtherChildHandle = TransformConstraint->ChildTRSHandle)
+					{
+						const UTransformableHandle* ParentHandle = InConstraint->ParentTRSHandle.Get();
+						
+						// if TransformConstraint's child is InConstraint's parent then avoid cycles
+						if (OtherChildHandle->GetHash() == ParentHandle->GetHash())
+						{
+							return false;
+						}
+
+						// check dependencies with OtherChildHandle to avoid cycles
+						const FHandleDependencyChecker Checker(World);
+						if (Checker.HasDependency(*ParentHandle, *OtherChildHandle))
+						{
+							return false;
+						}
+						
+						// check dependencies with OtherParentHandle to avoid cycles
+						if (Checker.HasDependency(*ParentHandle, *OtherParentHandle))
+						{
+							return false;
+						}
+					}
+					
+					return true;
+				}
+			}
+		}
+
+		return false;
+	};
+
+	const FConstraintsManagerController& Controller = FConstraintsManagerController::Get(World);
+	const TArray<ConstraintPtr> FilteredConstraints = Controller.GetConstraintsByPredicate(Predicate);
+	OutConstraints.Append(FilteredConstraints);
+}
 	
 }
diff --git a/Engine/Source/Runtime/Experimental/Animation/Constraints/Private/Transform/TransformableHandle.cpp b/Engine/Source/Runtime/Experimental/Animation/Constraints/Private/Transform/TransformableHandle.cpp
new file mode 100644
--- /dev/null
+++ b/Engine/Source/Runtime/Experimental/Animation/Constraints/Private/Transform/TransformableHandle.cpp
@@ -0,0 +1,644 @@
+﻿// Copyright Epic Games, Inc. All Rights Reserved.
+
+
+#include "Transform/TransformableHandle.h"
+#include "ConstraintsManager.h"
+#include "Components/SceneComponent.h"
+#include "GameFramework/Actor.h"
+#include "Engine/Engine.h"
+#include "EntitySystem/MovieSceneSharedPlaybackState.h"
+#include "MovieSceneSection.h"
+#include "Transform/TransformableHandleUtils.h"
+#include "Channels/MovieSceneFloatChannel.h"
+#include "Channels/MovieSceneDoubleChannel.h"
+#include "Channels/MovieSceneChannelProxy.h"
+#include "Components/SkeletalMeshComponent.h"
+
+#include UE_INLINE_GENERATED_CPP_BY_NAME(TransformableHandle)
+
+/**
+ * UTransformableHandle
+ */
+
+UTransformableHandle::~UTransformableHandle()
+{
+	OnHandleModified.Clear();
+}
+
+UTransformableHandle::FHandleModifiedEvent& UTransformableHandle::HandleModified()
+{
+	return OnHandleModified;
+}
+
+void UTransformableHandle::Notify(EHandleEvent InEvent, const bool bPreTickTarget) const
+{
+	if (!bNotifying && OnHandleModified.IsBound())
+	{
+		TGuardValue<bool> ReentrantGuardSelf(bNotifying, true);
+		if (bPreTickTarget)
+		{
+			TickTarget();
+		}
+		
+		OnHandleModified.Broadcast(const_cast<UTransformableHandle*>(this), InEvent);
+	}
+}
+
+bool UTransformableHandle::HasBoundObjects() const
+{
+	return ConstraintBindingID.IsValid();
+}
+
+void UTransformableHandle::OnBindingIDsUpdated(const TMap<UE::MovieScene::FFixedObjectBindingID, UE::MovieScene::FFixedObjectBindingID>& OldFixedToNewFixedMap, FMovieSceneSequenceID LocalSequenceID, TSharedRef<UE::MovieScene::FSharedPlaybackState> SharedPlaybackState)
+{
+	if (ConstraintBindingID.IsValid())
+	{
+		UE::MovieScene::FFixedObjectBindingID FixedBindingID = ConstraintBindingID.ResolveToFixed(LocalSequenceID, SharedPlaybackState);
+		if (OldFixedToNewFixedMap.Contains(FixedBindingID))
+		{
+			Modify();
+
+			const FMovieSceneSequenceHierarchy* Hierarchy = SharedPlaybackState->GetHierarchy();
+			ConstraintBindingID = OldFixedToNewFixedMap[FixedBindingID].ConvertToRelative(LocalSequenceID, Hierarchy);
+		}
+	}
+}
+
+void UTransformableHandle::PreEvaluate(const bool bTick) const
+{
+	if (bTick)
+	{
+		TickTarget();
+	}
+}
+
+/**
+ * UTransformableComponentHandle
+ */
+
+UTransformableComponentHandle::~UTransformableComponentHandle()
+{
+	UnregisterDelegates();
+}
+
+void UTransformableComponentHandle::PostLoad()
+{
+	Super::PostLoad();
+	RegisterDelegates();
+}
+
+bool UTransformableComponentHandle::IsValid(const bool bDeepCheck) const
+{
+	return Component.IsValid();
+}
+
+//need to tick any skelmesh component, sibling or parent
+void UTransformableComponentHandle::TickTarget() const
+{
+	if (!Component.IsValid())
+	{
+		return;
+	}
+	TransformableHandleUtils::TickDependantComponents(Component.Get());
+}
+
+void UTransformableComponentHandle::SetGlobalTransform(const FTransform& InGlobal) const
+{
+	if (Component.IsValid())
+	{
+		const FVector OldScale = GetGlobalTransform().GetScale3D();
+		const FVector NewScale = InGlobal.GetScale3D();
+		const bool bMarkRenderDirty =
+			FMath::Sign(OldScale[0]) != FMath::Sign(NewScale[0]) ||
+			FMath::Sign(OldScale[1]) != FMath::Sign(NewScale[1]) ||
+			FMath::Sign(OldScale[2]) != FMath::Sign(NewScale[2]);
+		
+		Component->SetWorldTransform(InGlobal);
+		
+		UWorld* World = Component->GetWorld();
+		if (World)
+		{
+			FConstraintsManagerController& Controller = FConstraintsManagerController::Get(World);
+			Controller.OnSceneComponentConstrained().Broadcast(Component.Get());
+		}
+
+		if (bMarkRenderDirty)
+		{
+			Component->MarkRenderStateDirty();
+		}
+	}
+}
+
+void UTransformableComponentHandle::SetLocalTransform(const FTransform& InLocal) const
+{
+	if(Component.IsValid())
+	{
+		Component->SetRelativeTransform(InLocal);
+		UWorld* World = Component->GetWorld();
+		if (World)
+		{
+			FConstraintsManagerController& Controller = FConstraintsManagerController::Get(World);
+			Controller.OnSceneComponentConstrained().Broadcast(Component.Get());
+		}
+	}
+}
+
+FTransform UTransformableComponentHandle::GetLocalTransform() const
+{
+	return Component.IsValid() ? Component->GetRelativeTransform() : FTransform::Identity;
+}
+
+FTransform UTransformableComponentHandle::GetGlobalTransform() const
+{
+	return Component.IsValid() ? Component->GetSocketTransform(SocketName) : FTransform::Identity;
+}
+
+UObject* UTransformableComponentHandle::GetPrerequisiteObject() const
+{
+	return Component.Get(); 
+}
+
+FTickFunction* UTransformableComponentHandle::GetTickFunction() const
+{
+	return Component.IsValid() ? &Component->PrimaryComponentTick : nullptr;
+}
+
+uint32 UTransformableComponentHandle::GetHash() const
+{
+	return Component.IsValid() ? GetTypeHash(Component.Get()) : 0;
+}
+
+TWeakObjectPtr<UObject> UTransformableComponentHandle::GetTarget() const
+{
+	return Component;
+}
+
+bool UTransformableComponentHandle::HasDirectDependencyWith(const UTransformableHandle& InOther) const
+{
+	const uint32 OtherHash = InOther.GetHash();
+	if (OtherHash == 0)
+	{
+		return false;
+	}
+
+	if (Component.IsValid())
+	{
+		// check whether the other handle is one of the component's parent
+		for (const USceneComponent* Comp=Component->GetAttachParent(); Comp!=nullptr; Comp=Comp->GetAttachParent() )
+		{
+			const uint32 AttachParentHash = GetTypeHash(Comp);
+			if (AttachParentHash == OtherHash)
+			{
+				return true;
+			}
+		}
+	}
+
+	return false;
+}
+
+namespace
+{
+	
+FTickPrerequisite LookForPrimaryPrerequisite(USceneComponent* Component)
+{
+	static const FTickPrerequisite DummyPrerex;
+	
+	if (!Component)
+	{
+		return DummyPrerex;
+	}
+
+	auto IsValidTickFunction = [](const FTickFunction* InTickFunction)
+	{
+		if (InTickFunction)
+		{
+			return (InTickFunction->bCanEverTick || InTickFunction->IsTickFunctionRegistered());
+		};
+
+		return false;
+	};
+
+	// check if this can tick
+	if ( IsValidTickFunction(&Component->PrimaryComponentTick) )
+	{
+		return FTickPrerequisite(Component, Component->PrimaryComponentTick);
+	}
+
+	// check if a parent constraint can tick
+	if (UWorld* World = Component->GetWorld())
+	{
+		static constexpr bool bSorted = true;
+		
+		const FConstraintsManagerController& Controller = FConstraintsManagerController::Get(World);
+		const TArray<TWeakObjectPtr<UTickableConstraint>> ParentConstraints =
+			Controller.GetParentConstraints(GetTypeHash(Component), bSorted);
+		
+		for (int Index = ParentConstraints.Num()-1; Index >= 0; Index--)
+		{
+			if (UTickableConstraint* Constraint = ParentConstraints[Index].Get())
+			{
+				if(IsValidTickFunction(&Constraint->GetTickFunction(World)))
+				{
+					return FTickPrerequisite(Constraint->GetOuter(), Constraint->GetTickFunction(World));
+				}
+			}
+		}	
+	}
+
+	// check if parent ticks
+	if (USceneComponent* Parent = Component->GetAttachParent())
+	{
+		FTickFunction& ParentTickFunction = Parent->PrimaryComponentTick;
+		if (IsValidTickFunction(&ParentTickFunction))
+		{
+			return FTickPrerequisite(Parent, ParentTickFunction);
+		}
+
+		// get up the hierarchy if not found
+		return LookForPrimaryPrerequisite(Parent);
+	}
+
+	return DummyPrerex;
+}
+	
+}
+
+FTickPrerequisite UTransformableComponentHandle::GetPrimaryPrerequisite() const
+{
+	if (Component.IsValid())
+	{
+		return LookForPrimaryPrerequisite(Component.Get());
+	}
+
+	static const FTickPrerequisite DummyPrerex;
+	return DummyPrerex;
+}
+
+void UTransformableComponentHandle::UnregisterDelegates() const
+{
+#if WITH_EDITOR
+	FCoreUObjectDelegates::OnObjectPropertyChanged.RemoveAll(this);
+	if (GEngine)
+	{
+		GEngine->OnActorMoving().RemoveAll(this);
+	}
+
+	FCoreUObjectDelegates::OnObjectsReplaced.RemoveAll(this);
+#endif
+}
+
+void UTransformableComponentHandle::RegisterDelegates()
+{
+	UnregisterDelegates();
+
+#if WITH_EDITOR
+	if (GEngine)
+	{
+		// NOTE BINDER: this has to be done before binding UTransformableComponentHandle::OnActorMoving
+		if (!GEngine->OnActorMoving().IsBoundToObject(&GetEvaluationBinding()))
+		{
+			GEngine->OnActorMoving().AddRaw(&GetEvaluationBinding(), &FComponentEvaluationGraphBinding::OnActorMoving);
+		}
+		
+		GEngine->OnActorMoving().AddUObject(this, &UTransformableComponentHandle::OnActorMoving);
+	}
+	FCoreUObjectDelegates::OnObjectPropertyChanged.AddUObject(this, &UTransformableComponentHandle::OnPostPropertyChanged);
+	FCoreUObjectDelegates::OnObjectsReplaced.AddUObject(this, &UTransformableComponentHandle::OnObjectsReplaced);
+#endif
+}
+
+#if WITH_EDITOR
+void UTransformableComponentHandle::OnActorMoving(AActor* InActor)
+{
+	if (!Component.IsValid())
+	{
+		return;
+	}
+
+	const TInlineComponentArray<USceneComponent*> Components(InActor);
+	if (!Components.Contains(Component))
+	{
+		return;
+	}
+
+	GetEvaluationBinding().bPendingFlush = true;
+	
+	Notify(EHandleEvent::GlobalTransformUpdated);
+}
+
+void UTransformableComponentHandle::OnPostPropertyChanged(
+	UObject* InObject,
+	FPropertyChangedEvent& InPropertyChangedEvent)
+{
+	if (!Component.IsValid())
+	{
+		return;
+	}
+
+	USceneComponent* SceneComponent = Cast<USceneComponent>(InObject);
+	if (!SceneComponent)
+	{
+		if (const AActor* Actor = Cast<AActor>(InObject))
+		{
+			const TInlineComponentArray<USceneComponent*> Components(Actor);
+			const int32 Index = Components.IndexOfByKey(Component.Get());
+			if (Index != INDEX_NONE)
+			{
+				SceneComponent = Components[Index];
+			}
+		}
+	}
+	
+	if (SceneComponent != Component)
+	{
+		return;
+	}
+
+	const FProperty* MemberProperty = InPropertyChangedEvent.MemberProperty;
+	if (!MemberProperty)
+	{
+		return;
+	}
+	
+	const FName MemberPropertyName = MemberProperty->GetFName();
+	const bool bTransformationChanged =
+		(MemberPropertyName == USceneComponent::GetRelativeLocationPropertyName() ||
+			MemberPropertyName == USceneComponent::GetRelativeRotationPropertyName() ||
+			MemberPropertyName == USceneComponent::GetRelativeScale3DPropertyName());
+	if (!bTransformationChanged)
+	{
+		return;
+	}
+
+	Notify(EHandleEvent::GlobalTransformUpdated);
+}
+
+void UTransformableComponentHandle::OnObjectsReplaced(const TMap<UObject*, UObject*>& InOldToNewInstances)
+{
+	// in the context of blueprints being recompiled (cf. AActor::RerunConstructionScripts()), the component has to
+	// be updated. as this is called after it has been destroyed, we get the component even if it's pending kill
+	// otherwise Get() will return a nullptr.
+	static constexpr bool bEvenIfPendingKill = true;
+	if (UObject* NewObject = InOldToNewInstances.FindRef(Component.Get(bEvenIfPendingKill)))
+	{
+		if (USceneComponent* NewSceneComponent= Cast<USceneComponent>(NewObject))
+		{
+			Component = NewSceneComponent;
+			RegisterDelegates();
+		}
+	}
+}
+#endif
+
+TArrayView<FMovieSceneFloatChannel*>  UTransformableComponentHandle::GetFloatChannels(const UMovieSceneSection* InSection) const
+{
+	// no floats for transform sections
+	static const TArrayView<FMovieSceneFloatChannel*> EmptyChannelsView;
+	return EmptyChannelsView;
+}
+
+TArrayView<FMovieSceneDoubleChannel*>  UTransformableComponentHandle::GetDoubleChannels(const UMovieSceneSection* InSection) const
+{
+	static const TArrayView<FMovieSceneDoubleChannel*> EmptyChannelsView;
+	if (InSection)
+	{
+		return InSection->GetChannelProxy().GetChannels<FMovieSceneDoubleChannel>();
+	}
+	return EmptyChannelsView;
+}
+
+bool UTransformableComponentHandle::AddTransformKeys(const TArray<FFrameNumber>& Frames,
+	const TArray<FTransform>& InLocalTransforms,
+	const EMovieSceneTransformChannel& InChannels,
+	const FFrameRate& InTickResolution,
+	UMovieSceneSection* InTransformSection,
+	const bool bLocal) const
+{
+	//todo 
+	//note this is the same as MOvieSceneToolHelpers::AddTransformKeys but that's in the editor code, need to find a place in the runtime someewhere
+	if (!InTransformSection)
+	{
+		return false;
+	}
+
+	if (Frames.IsEmpty() || Frames.Num() != InLocalTransforms.Num())
+	{
+		return false;
+	}
+
+	auto GetValue = [](const uint32 Index, const FVector& InLocation, const FRotator& InRotation, const FVector& InScale)
+	{
+		switch (Index)
+		{
+		case 0:
+			return InLocation.X;
+		case 1:
+			return InLocation.Y;
+		case 2:
+			return InLocation.Z;
+		case 3:
+			return InRotation.Roll;
+		case 4:
+			return InRotation.Pitch;
+		case 5:
+			return InRotation.Yaw;
+		case 6:
+			return InScale.X;
+		case 7:
+			return InScale.Y;
+		case 8:
+			return InScale.Z;
+		default:
+			ensure(false);
+			break;
+		}
+		return 0.0;
+	};
+
+	const bool bKeyTranslation = EnumHasAllFlags(InChannels, EMovieSceneTransformChannel::Translation);
+	const bool bKeyRotation = EnumHasAllFlags(InChannels, EMovieSceneTransformChannel::Rotation);
+	const bool bKeyScale = EnumHasAllFlags(InChannels, EMovieSceneTransformChannel::Scale);
+
+	TArray<uint32> ChannelsIndexToKey;
+	if (bKeyTranslation)
+	{
+		ChannelsIndexToKey.Append({ 0,1,2 });
+	}
+	if (bKeyRotation)
+	{
+		ChannelsIndexToKey.Append({ 3,4,5 });
+	}
+	if (bKeyScale)
+	{
+		ChannelsIndexToKey.Append({ 6,7,8 });
+	}
+
+	const TArrayView<FMovieSceneDoubleChannel*> Channels =
+		InTransformSection->GetChannelProxy().GetChannels<FMovieSceneDoubleChannel>();
+
+	// set default
+	const FTransform& LocalTransform0 = InLocalTransforms[0];
+	const FVector Location0 = LocalTransform0.GetLocation();
+	const FRotator Rotation0 = LocalTransform0.GetRotation().Rotator();
+	const FVector Scale3D0 = LocalTransform0.GetScale3D();
+
+	for (int32 ChannelIndex = 0; ChannelIndex < 9; ChannelIndex++)
+	{
+		if (!Channels[ChannelIndex]->GetDefault().IsSet())
+		{
+			const double Value = GetValue(ChannelIndex, Location0, Rotation0, Scale3D0);
+			Channels[ChannelIndex]->SetDefault(Value);
+		}
+	}
+
+	// add keys
+	for (int32 Index = 0; Index < Frames.Num(); ++Index)
+	{
+		const FFrameNumber& Frame = Frames[Index];
+		const FTransform& LocalTransform = InLocalTransforms[Index];
+
+		const FVector Location = LocalTransform.GetLocation();
+		const FRotator Rotation = LocalTransform.GetRotation().Rotator();
+		const FVector Scale3D = LocalTransform.GetScale3D();
+
+		for (const int32 ChannelIndex : ChannelsIndexToKey)
+		{
+			const double Value = GetValue(ChannelIndex, Location, Rotation, Scale3D);
+			TMovieSceneChannelData<FMovieSceneDoubleValue> ChannelData = Channels[ChannelIndex]->GetData();
+			{
+				int32 ExistingIndex = ChannelData.FindKey(Frame);
+				if (ExistingIndex != INDEX_NONE)
+				{
+					FMovieSceneDoubleValue& DoubleValue = ChannelData.GetValues()[ExistingIndex]; //-V758
+					DoubleValue.Value = Value;
+				}
+				else
+				{
+					FMovieSceneDoubleValue NewKey(Value);
+					ERichCurveTangentWeightMode WeightedMode = RCTWM_WeightedNone;
+					NewKey.InterpMode = ERichCurveInterpMode::RCIM_Cubic;
+					NewKey.TangentMode = ERichCurveTangentMode::RCTM_Auto;
+					NewKey.Tangent.ArriveTangent = 0.0f;
+					NewKey.Tangent.LeaveTangent = 0.0f;
+					NewKey.Tangent.TangentWeightMode = WeightedMode;
+					NewKey.Tangent.ArriveTangentWeight = 0.0f;
+					NewKey.Tangent.LeaveTangentWeight = 0.0f;
+					ChannelData.AddKey(Frame, NewKey);
+				}
+			}
+		}
+	}
+
+	//now we need to set auto tangents
+	for (const int32 ChannelIndex : ChannelsIndexToKey)
+	{
+		Channels[ChannelIndex]->AutoSetTangents();
+	}
+	return true;
+}
+
+void UTransformableComponentHandle::ResolveBoundObjects(FMovieSceneSequenceID LocalSequenceID, TSharedRef<UE::MovieScene::FSharedPlaybackState> SharedPlaybackState, UObject* SubObject)
+{
+	// in the context of blueprints being recompiled (cf. AActor::RerunConstructionScripts()) or with spawnable objects, the component has to be updated.
+	// as this can be called after it has been destroyed, we get the component even if it's pending kill otherwise Get() will return a nullptr.
+	static constexpr bool bEvenIfPendingKill = true;
+	const USceneComponent* ComponentEvenIfPendingKill = Component.Get(bEvenIfPendingKill);
+
+	auto ReplaceComponent = [this](USceneComponent* InNewComponent)
+	{
+		if (InNewComponent != Component)
+		{
+			Component = InNewComponent;
+			RegisterDelegates();
+		}
+	};
+	
+	for (const TWeakObjectPtr<>& ParentObject : ConstraintBindingID.ResolveBoundObjects(LocalSequenceID, SharedPlaybackState))
+	{
+		if (const AActor* Actor = Cast<AActor>(ParentObject.Get()))
+		{
+			const TInlineComponentArray<USceneComponent*> Components(Actor);
+			const int32 Index = Components.IndexOfByPredicate([this, ComponentEvenIfPendingKill](const USceneComponent* SubComponent)
+			{
+				return ComponentEvenIfPendingKill && (SubComponent->GetFName() == ComponentEvenIfPendingKill->GetFName());
+			});
+
+			ReplaceComponent(Index != INDEX_NONE ? Components[Index] : Actor->GetRootComponent());
+		}
+		else if (USceneComponent* Comp = Cast<USceneComponent>(ParentObject.Get()))
+		{
+			ReplaceComponent(Comp);
+		}
+		break; //just do one
+	}
+}
+
+UTransformableHandle* UTransformableComponentHandle::Duplicate(UObject* NewOuter) const
+{
+	UTransformableComponentHandle* HandleCopy = DuplicateObject<UTransformableComponentHandle>(this, NewOuter);
+	HandleCopy->Component = Component;
+	HandleCopy->SocketName = SocketName;
+	return HandleCopy;
+}
+
+#if WITH_EDITOR
+FString UTransformableComponentHandle::GetLabel() const
+{
+	if (!Component.IsValid())
+	{
+		static const FString DummyLabel;
+		return DummyLabel;
+	}
+
+	if (SocketName != NAME_None)
+	{
+		return SocketName.ToString(); 
+	}
+
+	const AActor* Actor = Component->GetOwner();
+	return Actor ? Actor->GetActorLabel() : Component->GetName();
+}
+
+FString UTransformableComponentHandle::GetFullLabel() const
+{
+	if (!Component.IsValid())
+	{
+		static const FString DummyLabel;
+		return DummyLabel;
+	}
+
+	const AActor* Actor = Component->GetOwner();
+	const FString ComponentLabel = Actor ? Actor->GetActorLabel() : Component->GetName();
+	
+	if (SocketName == NAME_None)
+	{
+		return ComponentLabel;
+	}
+
+	return FString::Printf(TEXT("%s/%s"), *ComponentLabel, *SocketName.ToString() );
+};
+
+#endif
+
+FComponentEvaluationGraphBinding& UTransformableComponentHandle::GetEvaluationBinding()
+{
+	static FComponentEvaluationGraphBinding EvaluationBinding;
+	return EvaluationBinding;
+}
+
+void FComponentEvaluationGraphBinding::OnActorMoving(AActor* InActor)
+{
+	if (!bPendingFlush)
+	{
+		return;
+	}
+	
+	if (UWorld* World = InActor ? InActor->GetWorld() : nullptr)
+	{
+		// flush all pending evaluations if any
+		const FConstraintsManagerController& Controller = FConstraintsManagerController::Get(World);
+		Controller.FlushEvaluationGraph();
+	}
+	bPendingFlush = false;
+}
diff --git a/Engine/Source/Runtime/Experimental/Animation/Constraints/Private/Transform/TransformableHandleUtils.cpp b/Engine/Source/Runtime/Experimental/Animation/Constraints/Private/Transform/TransformableHandleUtils.cpp
new file mode 100644
--- /dev/null
+++ b/Engine/Source/Runtime/Experimental/Animation/Constraints/Private/Transform/TransformableHandleUtils.cpp
@@ -0,0 +1,78 @@
+﻿// Copyright Epic Games, Inc. All Rights Reserved.
+
+#include "Transform/TransformableHandleUtils.h"
+
+#include "GameFramework/Actor.h"
+#include "Components/SceneComponent.h"
+#include "Components/SkeletalMeshComponent.h"
+
+namespace TransformableHandleUtils
+{
+
+static int32 SkeletalMeshTickingMode = 2;
+static FAutoConsoleVariableRef CVarSkeletalMeshTickingMode(
+	TEXT("Constraints.SkeletalMesh.TickingMode"),
+	SkeletalMeshTickingMode,
+	TEXT( "Constraint skeletal mesh ticking mode ([0, 2] - default: 2):\n" )
+		TEXT( "0 - do not tick any related skeletal mesh\n" )
+		TEXT( "1 - only tick the constrained skeletal mesh\n")
+		TEXT( "2 - tick all the skeletal meshes attached to the constrained actor\n"),
+		FConsoleVariableDelegate::CreateLambda([](IConsoleVariable*)
+	{
+		SkeletalMeshTickingMode = FMath::Clamp(SkeletalMeshTickingMode, 0, 2);
+	})
+);
+	
+void TickDependantComponents(USceneComponent* InComponent)
+{
+	if (SkeletalMeshTickingMode <= 0)
+	{
+		return;
+	}
+	
+	if (!InComponent)
+	{
+		return;
+	}
+
+	if (SkeletalMeshTickingMode == 1)
+	{
+		return TickSkeletalMeshComponent(Cast<USkeletalMeshComponent>(InComponent));
+	}
+
+	static constexpr bool bIncludeFromChildActors = true;
+
+	const AActor* Parent = InComponent->GetOwner();
+	while (Parent)
+	{
+		Parent->ForEachComponent<USkeletalMeshComponent>(bIncludeFromChildActors, &TickSkeletalMeshComponent);
+		Parent = Parent->GetAttachParentActor();
+	}
+}
+
+void TickSkeletalMeshComponent(USkeletalMeshComponent* InSkeletalMeshComponent)
+{
+	if (!InSkeletalMeshComponent)
+	{
+		return;
+	}
+
+	// avoid re-entrant animation evaluation
+	if (InSkeletalMeshComponent->IsPostEvaluatingAnimation())
+	{
+		return;
+	}
+
+	static constexpr float DeltaTime = 0.03f;
+	static constexpr bool bNeedsValidRootMotion = false;
+	
+	InSkeletalMeshComponent->TickAnimation(DeltaTime, bNeedsValidRootMotion);
+	InSkeletalMeshComponent->RefreshBoneTransforms();
+	InSkeletalMeshComponent->RefreshFollowerComponents();
+	InSkeletalMeshComponent->UpdateComponentToWorld();
+	InSkeletalMeshComponent->FinalizeBoneTransform();
+	InSkeletalMeshComponent->MarkRenderTransformDirty();
+	InSkeletalMeshComponent->MarkRenderDynamicDataDirty();
+}
+	
+}
\ No newline at end of file
diff --git a/Engine/Source/Runtime/Experimental/Animation/Constraints/Private/Transform/TransformableRegistry.cpp b/Engine/Source/Runtime/Experimental/Animation/Constraints/Private/Transform/TransformableRegistry.cpp
new file mode 100644
--- /dev/null
+++ b/Engine/Source/Runtime/Experimental/Animation/Constraints/Private/Transform/TransformableRegistry.cpp
@@ -0,0 +1,114 @@
+﻿// Copyright Epic Games, Inc. All Rights Reserved.
+
+#include "Transform/TransformableRegistry.h"
+
+#include "Transform/TransformableHandle.h"
+#include "Transform/TransformConstraintUtil.h"
+#include "GameFramework/Actor.h"
+#include "Components/SceneComponent.h"
+#include "UObject/Class.h"
+
+FTransformableRegistry::~FTransformableRegistry() = default;
+
+FTransformableRegistry& FTransformableRegistry::Get()
+{
+	static FTransformableRegistry Singleton;
+	return Singleton;
+}
+
+void FTransformableRegistry::Register(UClass* InClass, CreateHandleFuncT InHandleFunc, GetHashFuncT InHashFunc)
+{
+	const FTransformableInfo Info { InHandleFunc, InHashFunc};
+	Transformables.Emplace(InClass, Info);	
+}
+
+FTransformableRegistry::GetHashFuncT FTransformableRegistry::GetHashFunction(const UClass* InClass) const
+{
+	if (InClass)
+	{
+		if (const FTransformableInfo* Info = FindInfo(InClass))
+		{
+			return Info->GetHashFunc; 
+		}
+	}
+	return {};
+}
+
+FTransformableRegistry::CreateHandleFuncT FTransformableRegistry::GetCreateFunction(const UClass* InClass) const
+{
+	if (InClass)
+	{
+		if(const FTransformableInfo* Info = FindInfo(InClass))
+		{
+			return Info->CreateHandleFunc; 
+		}
+	}
+	return {};
+}
+
+const FTransformableRegistry::FTransformableInfo* FTransformableRegistry::FindInfo(const UClass* InClass) const
+{
+	if (!InClass)
+	{
+		return nullptr;
+	}
+	
+	// look for registered class
+	if (const FTransformableInfo* Info = Transformables.Find(InClass))
+	{
+		return Info;
+	}
+
+	// if not directly registered, look for super class registration
+	return FindInfo(InClass->GetSuperClass());
+}
+
+void FTransformableRegistry::RegisterBaseObjects()
+{
+    // register USceneComponent and AActor
+    auto CreateComponentHandle = [](UObject* InObject, const FName& InSocketName)->UTransformableHandle*
+    {
+    	if (USceneComponent* Component = Cast<USceneComponent>(InObject))
+    	{
+    		return UE::TransformConstraintUtil::CreateHandleForSceneComponent(Component, InSocketName);
+    	}
+    	return nullptr;
+    };
+
+    auto GetComponentHash = [](const UObject* InObject, const FName&)->uint32
+    {
+    	if (const USceneComponent* Component = Cast<USceneComponent>(InObject))
+    	{
+    		return GetTypeHash(Component);
+    	}
+    	return 0;
+    };
+    
+    auto CreateComponentHandleFromActor = [CreateComponentHandle](UObject* InObject, const FName& InSocketName)->UTransformableHandle*
+    {
+    	if (const AActor* Actor = Cast<AActor>(InObject))
+    	{
+    		return CreateComponentHandle(Actor->GetRootComponent(), InSocketName);
+    	}
+    	return nullptr;
+    };
+
+    auto GetComponentHashFromActor = [GetComponentHash](const UObject* InObject, const FName& InSocketName)->uint32
+    {
+    	if (const AActor* Actor = Cast<AActor>(InObject))
+    	{
+    		return GetComponentHash(Actor->GetRootComponent(), InSocketName);
+    	}
+    	return 0;
+    };
+
+	FTransformableRegistry& Registry = Get();
+    Registry.Register(USceneComponent::StaticClass(), CreateComponentHandle, GetComponentHash);
+    Registry.Register(AActor::StaticClass(), CreateComponentHandleFromActor, GetComponentHashFromActor);
+}
+
+void FTransformableRegistry::UnregisterAllObjects()
+{
+	FTransformableRegistry& Registry = Get();
+	Registry.Transformables.Reset();
+}
\ No newline at end of file
diff --git a/Engine/Source/Runtime/Experimental/Animation/Constraints/Private/TransformConstraint.cpp b/Engine/Source/Runtime/Experimental/Animation/Constraints/Private/TransformConstraint.cpp
deleted file mode 100644
diff --git a/Engine/Source/Runtime/Experimental/Animation/Constraints/Private/TransformableHandle.cpp b/Engine/Source/Runtime/Experimental/Animation/Constraints/Private/TransformableHandle.cpp
deleted file mode 100644
diff --git a/Engine/Source/Runtime/Experimental/Animation/Constraints/Private/TransformableHandleUtils.cpp b/Engine/Source/Runtime/Experimental/Animation/Constraints/Private/TransformableHandleUtils.cpp
deleted file mode 100644
diff --git a/Engine/Source/Runtime/Experimental/Animation/Constraints/Private/TransformableRegistry.cpp b/Engine/Source/Runtime/Experimental/Animation/Constraints/Private/TransformableRegistry.cpp
deleted file mode 100644
diff --git a/Engine/Source/Runtime/Experimental/Animation/Constraints/Public/Transform/TransformConstraint.h b/Engine/Source/Runtime/Experimental/Animation/Constraints/Public/Transform/TransformConstraint.h
new file mode 100644
--- /dev/null
+++ b/Engine/Source/Runtime/Experimental/Animation/Constraints/Public/Transform/TransformConstraint.h
@@ -0,0 +1,499 @@
+﻿// Copyright Epic Games, Inc. All Rights Reserved.
+
+#pragma once
+
+#include "Constraint.h"
+#include "ConstraintsManager.h"
+
+#include "Kismet/BlueprintFunctionLibrary.h"
+
+#include "Delegates/DelegateCombinations.h"
+#include "Delegates/Delegate.h"
+
+#include "TransformConstraint.generated.h"
+
+enum class EHandleEvent : uint8;
+class UTransformableHandle;
+class UTransformableComponentHandle;
+class USceneComponent;
+enum class EMovieSceneTransformChannel : uint32;
+class UWorld;
+/** 
+ * UTickableTransformConstraint
+ **/
+
+UCLASS(Abstract, Blueprintable, MinimalAPI)
+class UTickableTransformConstraint : public UTickableConstraint
+{
+	GENERATED_BODY()
+
+public:
+
+	/** Sets up the constraint so that the initial offset is set and dependencies and handles managed. */
+	CONSTRAINTS_API void Setup();
+	
+	/** UObjects overrides. */
+	CONSTRAINTS_API virtual void PostLoad() override;
+	CONSTRAINTS_API virtual void PostDuplicate(bool bDuplicateForPIE) override;
+
+	/** Returns the target hash value (i.e. the child handle's hash). */
+	CONSTRAINTS_API virtual uint32 GetTargetHash() const override;
+	
+	/** Test whether an InObject is referenced by that constraint. (i.e. is it's parent or child). */
+	CONSTRAINTS_API virtual bool ReferencesObject(TWeakObjectPtr<UObject> InObject) const override;
+	/** If true it contains objects bound to an external system, like sequencer so we don't do certain things, like remove constraints when they don't resolve*/
+	CONSTRAINTS_API virtual bool HasBoundObjects() const override;
+	
+	/** Resolve the bound objects so that any object it references are resovled and correctly set up*/
+	CONSTRAINTS_API virtual void ResolveBoundObjects(FMovieSceneSequenceID LocalSequenceID, TSharedRef<UE::MovieScene::FSharedPlaybackState> SharedPlaybackState, UObject* SubObject = nullptr) override;
+
+	/** Whether or not it's valid for example it may not be fully loaded*/
+	virtual bool IsValid(const bool bDeepCheck = true) const override;
+
+	/** If Active and the handles and targets are valid, and tick function is registered*/
+	CONSTRAINTS_API virtual bool IsFullyActive() const override;
+
+	/** If that constraint needs to be handled by the compensation system. */
+	CONSTRAINTS_API virtual bool NeedsCompensation() const;
+
+	/** Create duplicate with new Outer*/
+	CONSTRAINTS_API virtual UTickableConstraint* Duplicate(UObject* NewOuter) const override;
+
+	/** Override the evaluate so we can tick our handles*/
+	CONSTRAINTS_API virtual void Evaluate(bool bTickHandlesAlso = false) const override;
+
+	/** Sets the Active value and enable/disable the tick function. */
+	CONSTRAINTS_API virtual void SetActive(const bool bIsActive) override;
+	
+	/** The transformable handle representing the parent of that constraint. */
+	UPROPERTY(BlueprintReadWrite, Category = "Handle")
+	TObjectPtr<UTransformableHandle> ParentTRSHandle;
+	
+	/** The transformable handle representing the child of that constraint. */
+	UPROPERTY(BlueprintReadWrite, Category = "Handle")
+	TObjectPtr<UTransformableHandle> ChildTRSHandle;
+
+	/** Should that constraint maintain the default offset.  */
+	UPROPERTY(EditAnywhere, BlueprintReadWrite, Category="Offset", meta=(EditCondition="Type!=ETransformConstraintType::LookAt"))
+	bool bMaintainOffset = true;
+
+	/** Defines how much the constraint will be applied. */
+	// UPROPERTY(EditAnywhere, BlueprintReadWrite,Category="Weight", meta = (Input, ClampMin = "0.0", ClampMax = "1.0", UIMin = "0.0", UIMax = "1.0"))
+	//@benoit when not EditAnywhere?
+	UPROPERTY(meta = (Input, ClampMin = "0.0", ClampMax = "1.0", UIMin = "0.0", UIMax = "1.0"))
+	float Weight = 1.f;
+
+	/** Should the child be able to change it's offset dynamically. */
+	UPROPERTY(EditAnywhere, BlueprintReadWrite, Category="Offset", meta=(EditCondition="(Type!=ETransformConstraintType::LookAt) && bMaintainOffset"))
+	bool bDynamicOffset = false;
+
+	/** Returns the constraint type (Position, Parent, Aim...). */
+	CONSTRAINTS_API int64 GetType() const;
+
+	/** Get the current child's global transform. */
+	CONSTRAINTS_API FTransform GetChildGlobalTransform() const;
+	
+	/** Get the current child's local transform. */
+	CONSTRAINTS_API FTransform GetChildLocalTransform() const;
+
+	/** Get the current parent's global transform. */
+	CONSTRAINTS_API FTransform GetParentGlobalTransform() const;
+	
+	/** Get the current parent's local transform. */
+	CONSTRAINTS_API FTransform GetParentLocalTransform() const;
+
+	/** Returns the channels to key based on the constraint's type. */
+	CONSTRAINTS_API EMovieSceneTransformChannel GetChannelsToKey() const;
+
+	/**
+	 * Manages changes on the child/parent transformable handle. This can be used to update internal data (i.e. offset)
+	 * when transform changes outside of that system and need to trigger changes within the constraint itself.   
+	*/
+	CONSTRAINTS_API virtual void OnHandleModified(UTransformableHandle* InHandle, EHandleEvent InEvent);
+
+	/** Returns the handles tick function (ensuring it lives in the same world). */
+	CONSTRAINTS_API FTickFunction* GetChildHandleTickFunction() const;
+	CONSTRAINTS_API FTickFunction* GetParentHandleTickFunction() const;
+
+	/**
+	* Sets up dependencies with the first primary prerequisite available if the parent does not tick.   
+	*/
+	CONSTRAINTS_API void EnsurePrimaryDependency(const UWorld* InWorld) const;
+	
+protected:
+
+	/** Registers/Unregisters useful delegates for both child and parent handles. */
+	CONSTRAINTS_API void UnregisterDelegates() const;
+	CONSTRAINTS_API void RegisterDelegates();
+
+	
+	/**
+	 * Computes the initial offset that is needed to keep the child's global transform unchanged when creating the
+	 * constraint. This can be called whenever necessary to update the current offset.
+	*/
+	CONSTRAINTS_API virtual void ComputeOffset() PURE_VIRTUAL(ComputeOffset, return;);
+	
+	/**
+	 * Sets up dependencies between the parent, the constraint and the child using their respective tick functions.
+	 * It creates a dependency graph between them so that they tick in the right order when evaluated.   
+	*/
+	CONSTRAINTS_API void SetupDependencies(const UWorld* InWorld);
+
+	/** Set the current child's global transform. */
+	CONSTRAINTS_API void SetChildGlobalTransform(const FTransform& InGlobal) const;
+	
+	/** Set the current child's local transform. */
+	CONSTRAINTS_API void SetChildLocalTransform(const FTransform& InLocal) const;
+
+	/** Defines the constraint's type (Position, Parent, Aim...). */
+	UPROPERTY()
+	ETransformConstraintType Type = ETransformConstraintType::Parent;
+
+	/** Handle active state modification if needed. */
+	CONSTRAINTS_API void OnActiveStateChanged() const;
+
+	/** Returns the handle's tick function (ensuring it lives in the same world). */
+	CONSTRAINTS_API FTickFunction* GetHandleTickFunction(const TObjectPtr<UTransformableHandle>& InHandle) const;
+
+public:
+	/** (Re-)Registers the constraint function and (re-)binds the required delegates*/
+	CONSTRAINTS_API virtual void InitConstraint(UWorld * InWorld) override;
+	CONSTRAINTS_API virtual void TeardownConstraint(UWorld * InWorld) override;
+	CONSTRAINTS_API virtual void AddedToWorld(UWorld* InWorld) override;
+
+#if WITH_EDITOR
+public:
+	/** Returns the constraint's label used for UI. */
+	CONSTRAINTS_API virtual FString GetLabel() const override;
+	CONSTRAINTS_API virtual FString GetFullLabel() const override;
+
+	/** Returns the constraint's type label used for UI. */
+	CONSTRAINTS_API virtual FString GetTypeLabel() const override;
+
+	// UObject interface
+	CONSTRAINTS_API virtual void PostEditChangeProperty(FPropertyChangedEvent& PropertyChangedEvent) override;
+	CONSTRAINTS_API virtual void PostEditUndo() override;
+	// End of UObject interface
+
+	/** Returns a delegate that can be used to monitor for property changes. This can be used to monitor constraints changes
+	 * only instead of using FCoreUObjectDelegates::OnObjectPropertyChanged that is listening to every objects.
+	 * Note that bScaling is currently the only property change we monitor but this can be used for other properties.
+	 */
+	DECLARE_MULTICAST_DELEGATE_TwoParams(FOnConstraintChanged, UTickableTransformConstraint*, const FPropertyChangedEvent&);
+	static CONSTRAINTS_API FOnConstraintChanged& GetOnConstraintChanged();
+
+protected:
+	static CONSTRAINTS_API FOnConstraintChanged OnConstraintChanged;
+#endif
+};
+
+/** 
+ * UTickableTranslationConstraint
+ **/
+
+UCLASS(Blueprintable, MinimalAPI)
+class UTickableTranslationConstraint : public UTickableTransformConstraint
+{
+	GENERATED_BODY()
+
+public:
+	CONSTRAINTS_API UTickableTranslationConstraint();
+
+	/** Returns the position constraint function that the tick function will evaluate. */
+	CONSTRAINTS_API virtual FConstraintTickFunction::ConstraintFunction GetFunction() const override;
+
+	/** Updates the dynamic offset based on external child's transform changes. */
+	CONSTRAINTS_API virtual void OnHandleModified(UTransformableHandle* InHandle, EHandleEvent InEvent) override;
+
+protected:
+
+	/** Cache data structure to store last child local/global location. */
+	struct FDynamicCache
+	{
+		uint32 CachedInputHash = 0;
+	};
+	mutable FDynamicCache Cache;
+	CONSTRAINTS_API uint32 CalculateInputHash() const;
+	
+	/** Computes the child's local translation offset in the parent space. */
+	CONSTRAINTS_API virtual void ComputeOffset() override;
+
+public:
+	/** Defines the local child's translation offset in the parent space. */
+	UPROPERTY(EditAnywhere, BlueprintReadWrite, Category="Offset", meta=(EditCondition="bMaintainOffset"))
+	FVector OffsetTranslation = FVector::ZeroVector;
+
+	/** Defines which translation axis is constrained. */
+	UPROPERTY(EditDefaultsOnly, BlueprintReadWrite, Category = "Axis Filter")
+	FFilterOptionPerAxis AxisFilter;
+
+#if WITH_EDITOR
+public:
+	// UObject interface
+	CONSTRAINTS_API virtual void PostEditChangeProperty(FPropertyChangedEvent& PropertyChangedEvent) override;
+	// End of UObject interface
+#endif
+};
+
+/** 
+ * UTickableRotationConstraint
+ **/
+
+UCLASS(Blueprintable, MinimalAPI)
+class UTickableRotationConstraint : public UTickableTransformConstraint
+{
+	GENERATED_BODY()
+
+public:
+	CONSTRAINTS_API UTickableRotationConstraint();
+
+	/** Returns the rotation constraint function that the tick function will evaluate. */
+	CONSTRAINTS_API virtual FConstraintTickFunction::ConstraintFunction GetFunction() const override;
+
+	/** Updates the dynamic offset based on external child's transform changes. */
+	CONSTRAINTS_API virtual void OnHandleModified(UTransformableHandle* InHandle, EHandleEvent InEvent) override;
+
+protected:
+
+	/** Cache data structure to store last child local/global rotation. */
+	struct FDynamicCache
+	{
+		uint32 CachedInputHash = 0;
+	};
+	mutable FDynamicCache Cache;
+	CONSTRAINTS_API uint32 CalculateInputHash() const;
+	
+	/** Computes the child's local rotation offset in the parent space. */
+	CONSTRAINTS_API virtual void ComputeOffset() override;
+
+public:
+	/** Defines the local child's rotation offset in the parent space. */
+	UPROPERTY(EditAnywhere, BlueprintReadWrite, Category="Offset", meta=(EditCondition="bMaintainOffset"))
+	FQuat OffsetRotation = FQuat::Identity;
+
+	/** Defines which rotation axis is constrained. */
+	UPROPERTY(EditDefaultsOnly, BlueprintReadWrite, Category = "Axis Filter")
+	FFilterOptionPerAxis AxisFilter;
+
+#if WITH_EDITOR
+public:
+	// UObject interface
+	CONSTRAINTS_API virtual void PostEditChangeProperty(FPropertyChangedEvent& PropertyChangedEvent) override;
+	// End of UObject interface
+#endif
+};
+
+/** 
+ * UTickableScaleConstraint
+ **/
+
+UCLASS(Blueprintable, MinimalAPI)
+class UTickableScaleConstraint : public UTickableTransformConstraint
+{
+	GENERATED_BODY()
+
+public:
+	CONSTRAINTS_API UTickableScaleConstraint();
+	
+	/** Returns the scale constraint function that the tick function will evaluate. */
+	CONSTRAINTS_API virtual FConstraintTickFunction::ConstraintFunction GetFunction() const override;
+
+	/** Updates the dynamic offset based on external child's transform changes. */
+	CONSTRAINTS_API virtual void OnHandleModified(UTransformableHandle* InHandle, EHandleEvent InEvent) override;
+	
+protected:
+	/** Computes the child's local scale offset in the parent space. */
+	CONSTRAINTS_API virtual void ComputeOffset() override;
+
+	/** Cache data structure to store last child local/global transform. */
+	struct FDynamicCache
+	{
+		uint32 CachedInputHash = 0;
+	};
+	mutable FDynamicCache Cache;
+	CONSTRAINTS_API uint32 CalculateInputHash() const;
+
+public:
+	/** Defines the local child's scale offset in the parent space. */
+	UPROPERTY(EditAnywhere, BlueprintReadWrite, Category="Offset", meta=(EditCondition="bMaintainOffset"))
+	FVector OffsetScale = FVector::OneVector;
+
+	/** Defines which scale axis is constrained. */
+	UPROPERTY(EditDefaultsOnly, BlueprintReadWrite, Category = "Axis Filter")
+	FFilterOptionPerAxis AxisFilter;
+
+#if WITH_EDITOR
+public:
+	// UObject interface
+	CONSTRAINTS_API virtual void PostEditChangeProperty(FPropertyChangedEvent& PropertyChangedEvent) override;
+	// End of UObject interface
+#endif
+};
+
+/** 
+ * UTickableParentConstraint
+ **/
+
+UCLASS(Blueprintable, MinimalAPI)
+class UTickableParentConstraint : public UTickableTransformConstraint
+{
+	GENERATED_BODY()
+
+public:
+	CONSTRAINTS_API UTickableParentConstraint();
+	
+	/** Returns the transform constraint function that the tick function will evaluate. */
+	CONSTRAINTS_API virtual FConstraintTickFunction::ConstraintFunction GetFunction() const override;
+
+	bool IsScalingEnabled() const
+	{
+		return bScaling;
+	}
+	
+	void SetScaling(const bool bInScale)
+	{
+		bScaling = bInScale;
+	}
+
+	static FName GetScalingPropertyName() { return GET_MEMBER_NAME_CHECKED(UTickableParentConstraint, bScaling); }
+
+	/** Updates the dynamic offset based on external child's transform changes. */
+	CONSTRAINTS_API virtual void OnHandleModified(UTransformableHandle* InHandle, EHandleEvent InEvent) override;
+
+protected:
+	/** Cache data structure to store last child local/global transform. */
+	struct FDynamicCache
+	{
+		uint32 CachedInputHash = 0;
+	};
+	mutable FDynamicCache Cache;
+	CONSTRAINTS_API uint32 CalculateInputHash() const;
+
+	/** Computes the child's local transform offset in the parent space. */
+	CONSTRAINTS_API virtual void ComputeOffset() override;
+
+public:
+	/** Defines the local child's transform offset in the parent space. */
+	UPROPERTY(EditAnywhere, BlueprintReadWrite, Category="Offset", meta=(EditCondition="bMaintainOffset"))
+	FTransform OffsetTransform = FTransform::Identity;
+
+	/** Defines whether we propagate the parent scale. */
+	UPROPERTY(EditAnywhere, BlueprintReadWrite, Category="Properties")
+	bool bScaling = false;
+
+	/** Defines which translation/rotation/scale axis are constrained. */
+	UPROPERTY(EditDefaultsOnly, BlueprintReadWrite, Category = "Axis Filter") 
+	FTransformFilter TransformFilter;
+
+#if WITH_EDITOR
+public:
+	// UObject interface
+	CONSTRAINTS_API virtual void PostEditChangeProperty(FPropertyChangedEvent& PropertyChangedEvent) override;
+	// End of UObject interface
+#endif
+};
+
+/** 
+ * UTickableLookAtConstraint
+ **/
+
+UCLASS(Blueprintable, MinimalAPI)
+class UTickableLookAtConstraint : public UTickableTransformConstraint
+{
+	GENERATED_BODY()
+
+public:
+	CONSTRAINTS_API UTickableLookAtConstraint();
+	
+	/** Returns the look at constraint function that the tick function will evaluate. */
+	CONSTRAINTS_API virtual FConstraintTickFunction::ConstraintFunction GetFunction() const override;
+
+	/** If that constraint needs to be handled by the compensation system. */
+	CONSTRAINTS_API virtual bool NeedsCompensation() const override;
+
+protected:
+	/**
+     * Computes the initial axis that is needed to keep the child's orientation unchanged when creating the constraint.
+    */
+	CONSTRAINTS_API virtual void ComputeOffset() override;
+
+	/** Defines the aiming axis. */
+	UPROPERTY(EditAnywhere, BlueprintReadWrite, Category="Axis")
+	FVector Axis = FVector::XAxisVector;
+
+private:
+
+	/** Computes the shortest quaternion between A and B. */
+	static FQuat FindQuatBetweenNormals(const FVector& A, const FVector& B);
+
+#if WITH_EDITOR
+public:
+	// UObject interface
+	CONSTRAINTS_API virtual void PostEditChangeProperty(FPropertyChangedEvent& PropertyChangedEvent) override;
+	// End of UObject interface
+#endif
+};
+
+struct FTransformConstraintUtils
+{
+	UE_DEPRECATED(5.6, "This function was moved to the UE::TransformConstraintUtil namespace.")
+	static CONSTRAINTS_API void GetParentConstraints(
+		UWorld* World,
+		const AActor* InChild,
+		TArray< TWeakObjectPtr<UTickableConstraint> >& OutConstraints);
+
+	UE_DEPRECATED(5.6, "This function was moved to the UE::TransformConstraintUtil namespace.")
+	static CONSTRAINTS_API UTransformableComponentHandle* CreateHandleForSceneComponent(
+		USceneComponent* InSceneComponent,
+		const FName& InSocketName);
+
+	UE_DEPRECATED(5.6, "This function was moved to the UE::TransformConstraintUtil namespace.")
+	static CONSTRAINTS_API UTickableTransformConstraint* CreateFromType(
+		UWorld* InWorld,
+		const ETransformConstraintType InType,
+		const bool bUseDefault = false);
+
+	UE_DEPRECATED(5.6, "This function was moved to the UE::TransformConstraintUtil namespace.")	
+	static CONSTRAINTS_API UTickableTransformConstraint* CreateAndAddFromObjects(
+		UWorld* InWorld,
+		UObject* InParent, const FName& InParentSocketName,
+		UObject* InChild, const FName& InChildSocketName,
+		const ETransformConstraintType InType,
+		const bool bMaintainOffset = true,
+		const bool bUseDefault = false);
+
+	UE_DEPRECATED(5.6, "This function was moved to the UE::TransformConstraintUtil namespace.")	
+	static CONSTRAINTS_API bool AddConstraint(
+		UWorld* InWorld,
+		UTransformableHandle* InParentHandle,
+		UTransformableHandle* InChildHandle,
+		UTickableTransformConstraint* Constraint,
+		const bool bMaintainOffset = true,
+		const bool bUseDefault = false);
+
+	UE_DEPRECATED(5.6, "This function was moved to the UE::TransformConstraintUtil namespace.")
+	static CONSTRAINTS_API FTransform ComputeRelativeTransform(
+		const FTransform& InChildLocal,
+		const FTransform& InChildWorld,
+		const FTransform& InSpaceWorld,
+		const UTickableTransformConstraint* InConstraint);
+
+	UE_DEPRECATED(5.6, "This function was moved to the UE::TransformConstraintUtil namespace.")
+	static CONSTRAINTS_API TOptional<FTransform> GetRelativeTransform(UWorld* InWorld, const uint32 InHandleHash);
+	static CONSTRAINTS_API TOptional<FTransform> GetConstraintsRelativeTransform(
+		const TArray< TWeakObjectPtr<UTickableConstraint> >& InConstraints,
+		const FTransform& InChildLocal, const FTransform& InChildWorld);
+
+	UE_DEPRECATED(5.6, "This function was moved to the UE::TransformConstraintUtil namespace.")
+	static CONSTRAINTS_API int32 GetLastActiveConstraintIndex(const TArray< TWeakObjectPtr<UTickableConstraint> >& InConstraints);
+
+	UE_DEPRECATED(5.6, "This function was moved to the UE::TransformConstraintUtil namespace.")
+	static CONSTRAINTS_API void GetChildrenConstraints(
+		UWorld* World,
+		const UTickableTransformConstraint* InConstraint,
+		TArray< TWeakObjectPtr<UTickableConstraint> >& OutConstraints,
+		const bool bIncludeTarget = false);
+	
+	UE_DEPRECATED(5.6, "This function was moved to the UE::TransformConstraintUtil namespace.")
+	static CONSTRAINTS_API void UpdateTransformBasedOnConstraint(FTransform& CurrentTransform, USceneComponent* SceneComponent);
+};
diff --git a/Engine/Source/Runtime/Experimental/Animation/Constraints/Public/Transform/TransformConstraintUtil.h b/Engine/Source/Runtime/Experimental/Animation/Constraints/Public/Transform/TransformConstraintUtil.h
--- a/Engine/Source/Runtime/Experimental/Animation/Constraints/Public/Transform/TransformConstraintUtil.h
+++ b/Engine/Source/Runtime/Experimental/Animation/Constraints/Public/Transform/TransformConstraintUtil.h
@@ -2,6 +2,8 @@
 
 #pragma once
 
+#include "Constraint.h"
+
 #include "Containers/Array.h"
 #include "Containers/Map.h"
 
@@ -11,6 +13,9 @@
 
 class UWorld;
 class UTickableConstraint;
+class UTransformableHandle;
+class UTickableTransformConstraint;
+class UTransformableComponentHandle;
 
 namespace UE::TransformConstraintUtil
 {
@@ -42,5 +47,75 @@ private:
 	TMap<uint32, TArray< TWeakObjectPtr<UTickableConstraint> >> PerHandleConstraints;
 	FDelegateHandle ConstraintsNotificationHandle;
 };
-	
+
 }
+
+
+namespace UE::TransformConstraintUtil
+{
+
+	/** Fills a sorted constraint array that InChild actor is the child of. */
+	CONSTRAINTS_API void GetParentConstraints(
+		UWorld* World,
+		const AActor* InChild,
+		TArray< TWeakObjectPtr<UTickableConstraint> >& OutConstraints);
+
+	/** Create a handle for the scene component.*/
+	CONSTRAINTS_API UTransformableComponentHandle* CreateHandleForSceneComponent(
+		USceneComponent* InSceneComponent,
+		const FName& InSocketName);
+
+	/** Creates a new transform constraint based on the InType. */
+	CONSTRAINTS_API UTickableTransformConstraint* CreateFromType(
+		UWorld* InWorld,
+		const ETransformConstraintType InType,
+		const bool bUseDefault = false);
+
+	/** Creates respective handles and creates a new InType transform constraint. */	
+	CONSTRAINTS_API UTickableTransformConstraint* CreateAndAddFromObjects(
+		UWorld* InWorld,
+		UObject* InParent, const FName& InParentSocketName,
+		UObject* InChild, const FName& InChildSocketName,
+		const ETransformConstraintType InType,
+		const bool bMaintainOffset = true,
+		const bool bUseDefault = false);
+
+	/** Registers a new transform constraint within the constraints manager. */	
+	CONSTRAINTS_API bool AddConstraint(
+		UWorld* InWorld,
+		UTransformableHandle* InParentHandle,
+		UTransformableHandle* InChildHandle,
+		UTickableTransformConstraint* Constraint,
+		const bool bMaintainOffset = true,
+		const bool bUseDefault = false);
+
+	/** Computes the relative transform between both transform based on the constraint's InType. */
+	CONSTRAINTS_API FTransform ComputeRelativeTransform(
+		const FTransform& InChildLocal,
+		const FTransform& InChildWorld,
+		const FTransform& InSpaceWorld,
+		const UTickableTransformConstraint* InConstraint);
+
+	/** Computes the current constraint space local transform. */
+	CONSTRAINTS_API TOptional<FTransform> GetRelativeTransform(UWorld* InWorld, const uint32 InHandleHash);
+	CONSTRAINTS_API TOptional<FTransform> GetConstraintsRelativeTransform(
+		const TArray< TWeakObjectPtr<UTickableConstraint> >& InConstraints,
+		const FTransform& InChildLocal, const FTransform& InChildWorld);
+
+	/** Get the last active constraint that has dynamic offset. */
+	CONSTRAINTS_API int32 GetLastActiveConstraintIndex(const TArray< TWeakObjectPtr<UTickableConstraint> >& InConstraints);
+
+	/**
+	 * Fills a constraint array that InConstraint->ChildHandle is the parent of.
+	 * If bIncludeTarget is true, we also get the other constraints that act on the same target.
+	 */
+	CONSTRAINTS_API void GetChildrenConstraints(
+		UWorld* World,
+		const UTickableTransformConstraint* InConstraint,
+		TArray< TWeakObjectPtr<UTickableConstraint> >& OutConstraints,
+		const bool bIncludeTarget = false);
+	
+	/** Adjust the transform on a scene component so it's effected by the constraint*/
+	CONSTRAINTS_API void UpdateTransformBasedOnConstraint(FTransform& InOutCurrentTransform, const USceneComponent* InSceneComponent);
+
+}
\ No newline at end of file
diff --git a/Engine/Source/Runtime/Experimental/Animation/Constraints/Public/Transform/TransformableHandle.h b/Engine/Source/Runtime/Experimental/Animation/Constraints/Public/Transform/TransformableHandle.h
new file mode 100644
--- /dev/null
+++ b/Engine/Source/Runtime/Experimental/Animation/Constraints/Public/Transform/TransformableHandle.h
@@ -0,0 +1,241 @@
+﻿// Copyright Epic Games, Inc. All Rights Reserved.
+
+#pragma once
+
+#include "UObject/Object.h"
+#include "Engine/EngineBaseTypes.h"
+#include "MovieSceneObjectBindingID.h"
+#include "TransformableHandle.generated.h"
+
+class UMovieSceneSection;
+class USceneComponent;
+enum class EMovieSceneTransformChannel : uint32;
+struct FFRameNumber;
+struct FFrameRate;
+struct FMovieSceneDoubleChannel;
+struct FMovieSceneFloatChannel;
+struct FMovieSceneSequenceHierarchy;
+struct FMovieSceneSequenceID;
+struct FTickFunction;
+
+namespace UE::MovieScene
+{
+	struct FFixedObjectBindingID;
+	struct FSharedPlaybackState;
+}
+
+UENUM()
+enum class EHandleEvent : uint8
+{
+	LocalTransformUpdated,
+	GlobalTransformUpdated,
+	ComponentUpdated,
+	UpperDependencyUpdated,
+
+	/** MAX - invalid */
+	Max UMETA(Hidden),
+};
+
+/**
+ * UTransformableHandle
+ */
+
+UCLASS(Abstract, Blueprintable, MinimalAPI)
+class UTransformableHandle : public UObject 
+{
+	GENERATED_BODY()
+	
+public:
+
+	DECLARE_EVENT_TwoParams(UTransformableHandle, FHandleModifiedEvent, UTransformableHandle*, EHandleEvent);
+	
+	CONSTRAINTS_API virtual ~UTransformableHandle();
+	
+	/** Sanity check to ensure the handle is safe to use.
+	 * @param bDeepCheck to check that the transformable object it wraps is valid AND can be transformed. Default is true.
+	 * Some handles, such as control handles, will only be bound to a skeletal mesh once the level sequence has been opened, but their control rig
+	 * pointer and control name are valid, so they must be fully loaded so that the constraint can be updated later. bDeepCheck = false will be used
+	 * for that purpose.
+	 */
+	CONSTRAINTS_API virtual bool IsValid(const bool bDeepCheck = true) const PURE_VIRTUAL(IsValid, return false;);
+	
+	/** Sets the global transform of the underlying transformable object. */
+	CONSTRAINTS_API virtual void SetGlobalTransform(const FTransform& InGlobal) const PURE_VIRTUAL(SetGlobalTransform, );
+	/** Sets the local transform of the underlying transformable object in it's parent space. */
+	CONSTRAINTS_API virtual void SetLocalTransform(const FTransform& InLocal) const PURE_VIRTUAL(SetLocalTransform, );
+	/** Gets the global transform of the underlying transformable object. */
+	CONSTRAINTS_API virtual FTransform GetGlobalTransform() const PURE_VIRTUAL(GetGlobalTransform, return FTransform::Identity;);
+	/** Gets the local transform of the underlying transformable object in it's parent space. */
+	CONSTRAINTS_API virtual FTransform GetLocalTransform() const PURE_VIRTUAL(GetLocalTransform, return FTransform::Identity;);
+
+	/** If true it contains objects bound to an external system, like sequencer so we don't do certain things, like remove constraints when they don't resolve*/
+	CONSTRAINTS_API virtual bool HasBoundObjects() const;
+
+	/** Resolve the bound objects so that any object it references are resolved and correctly set up*/
+	CONSTRAINTS_API virtual void ResolveBoundObjects(FMovieSceneSequenceID LocalSequenceID, TSharedRef<UE::MovieScene::FSharedPlaybackState> SharedPlaybackState, UObject* SubObject = nullptr) PURE_VIRTUAL(ResolveBoundObjects);
+
+	/** Make a duplicate of myself with this outer*/
+	CONSTRAINTS_API virtual UTransformableHandle* Duplicate(UObject* NewOuter) const PURE_VIRTUAL(Duplicate, return nullptr;);
+
+
+	/** Fix up Binding in case it has changed*/
+	CONSTRAINTS_API void OnBindingIDsUpdated(const TMap<UE::MovieScene::FFixedObjectBindingID, UE::MovieScene::FFixedObjectBindingID>& OldFixedToNewFixedMap, FMovieSceneSequenceID LocalSequenceID, TSharedRef<UE::MovieScene::FSharedPlaybackState> SharedPlaybackState);
+	
+	/** Perform any special ticking needed for this handle, by default it does nothing, todo need to see if we need to tick control rig also*/
+	virtual void TickTarget() const {};
+
+	/**
+	 * Perform any pre-evaluation of the handle to ensure that the transform data are up to date.
+	 * @param bTick to force any pre-evaluation ticking. Default is false.
+	*/
+	virtual void PreEvaluate(const bool bTick = false) const;
+	
+	/** Get the array of float channels for the specified section*/
+	CONSTRAINTS_API virtual TArrayView<FMovieSceneFloatChannel*>  GetFloatChannels(const UMovieSceneSection* InSection) const PURE_VIRTUAL(GetFloatChannels, return TArrayView<FMovieSceneFloatChannel*>(); );
+	/** Get the array of double channels for the specified section*/
+	CONSTRAINTS_API virtual TArrayView<FMovieSceneDoubleChannel*>  GetDoubleChannels(const UMovieSceneSection * InSection) const PURE_VIRTUAL(GetDoubleChannels, return TArrayView<FMovieSceneDoubleChannel*>(); );
+	/** Add Transform Keys at the specified times*/
+	CONSTRAINTS_API virtual bool AddTransformKeys(const TArray<FFrameNumber>& InFrames,
+		const TArray<FTransform>& InTransforms,
+		const EMovieSceneTransformChannel& InChannels,
+		const FFrameRate& InTickResolution,
+		UMovieSceneSection* InSection,
+		const bool bLocal = true) const PURE_VIRTUAL(AddTransformKeys, return false;);
+
+	/**
+	 * Returns the target object containing the tick function (returned in GetTickFunction).
+	 * See FTickFunction::AddPrerequisite for details.
+	 **/
+	CONSTRAINTS_API virtual UObject* GetPrerequisiteObject() const PURE_VIRTUAL(GetPrerequisiteObject, return nullptr;);
+	/**
+	 * Returns the tick function of the underlying transformable object.
+	 * This is used to set dependencies with the constraint.
+	**/
+	CONSTRAINTS_API virtual FTickFunction* GetTickFunction() const PURE_VIRTUAL(GetTickFunction, return nullptr;);
+
+	/** Generates a hash value of the underlying transformable object. */
+	CONSTRAINTS_API virtual uint32 GetHash() const PURE_VIRTUAL(GetHash, return 0;);
+
+	/** Returns the underlying targeted object. */
+	CONSTRAINTS_API virtual TWeakObjectPtr<UObject> GetTarget() const PURE_VIRTUAL(GetTarget, return nullptr;);
+
+	/** Check for direct dependencies with InOther. */
+	CONSTRAINTS_API virtual bool HasDirectDependencyWith(const UTransformableHandle& InOther) const PURE_VIRTUAL(HasDirectDependencyWith, return false;);
+
+	/** Look for a possible tick function that can be used as a prerequisite. */
+	CONSTRAINTS_API virtual FTickPrerequisite GetPrimaryPrerequisite() const PURE_VIRTUAL(GetPrimaryPrerequisite, return FTickPrerequisite(););
+	
+	CONSTRAINTS_API FHandleModifiedEvent& HandleModified();
+	CONSTRAINTS_API void Notify(EHandleEvent InEvent, const bool bPreTickTarget = false) const;
+	mutable bool bNotifying = false;
+
+#if WITH_EDITOR
+	CONSTRAINTS_API virtual FString GetLabel() const PURE_VIRTUAL(GetLabel, return FString(););
+	CONSTRAINTS_API virtual FString GetFullLabel() const PURE_VIRTUAL(GetFullLabel, return FString(););
+#endif
+
+	//possible bindingID
+	UPROPERTY(BlueprintReadOnly, Category = "Binding")
+	FMovieSceneObjectBindingID ConstraintBindingID;
+
+private:
+	FHandleModifiedEvent OnHandleModified;
+};
+
+/**
+ * UTransformableComponentHandle
+ */
+
+struct FComponentEvaluationGraphBinding
+{
+	void OnActorMoving(AActor* InActor);
+	bool bPendingFlush = false;
+};
+
+UCLASS(Blueprintable, MinimalAPI)
+class UTransformableComponentHandle : public UTransformableHandle 
+{
+	GENERATED_BODY()
+	
+public:
+	
+	CONSTRAINTS_API virtual ~UTransformableComponentHandle();
+
+	CONSTRAINTS_API virtual void PostLoad() override;
+	
+	/** Sanity check to ensure that Component. */
+	CONSTRAINTS_API virtual bool IsValid(const bool bDeepCheck = true) const override;
+	
+	/** Sets the global transform of Component. */
+	CONSTRAINTS_API virtual void SetGlobalTransform(const FTransform& InGlobal) const override;
+	/** Sets the local transform of Component in it's attachment. */
+	CONSTRAINTS_API virtual void SetLocalTransform(const FTransform& InLocal) const override;
+	/** Gets the global transform of Component. */
+	CONSTRAINTS_API virtual FTransform GetGlobalTransform() const override;
+	/** Gets the local transform of Component in it's attachment. */
+	CONSTRAINTS_API virtual FTransform GetLocalTransform() const override;
+	/** Tick any skeletal mesh related to the component. */
+	CONSTRAINTS_API virtual void TickTarget() const override;
+	/** Returns the target object containing the tick function (e.i. Component). */
+	CONSTRAINTS_API virtual UObject* GetPrerequisiteObject() const override;
+	/** Returns Component's tick function. */
+	CONSTRAINTS_API virtual FTickFunction* GetTickFunction() const override;
+
+	/** Generates a hash value of Component. */
+	CONSTRAINTS_API virtual uint32 GetHash() const override;
+
+	/** Returns the underlying targeted object. */
+	CONSTRAINTS_API virtual TWeakObjectPtr<UObject> GetTarget() const override;
+
+	/** Check for direct dependencies (ie hierarchy) with InOther. */
+	CONSTRAINTS_API virtual bool HasDirectDependencyWith(const UTransformableHandle& InOther) const override;
+
+	/** Look for a possible tick function that can be used as a prerequisite. */
+	CONSTRAINTS_API virtual FTickPrerequisite GetPrimaryPrerequisite() const override;
+
+	/** Get the array of float channels for the specified section*/
+	CONSTRAINTS_API virtual TArrayView<FMovieSceneFloatChannel*>  GetFloatChannels(const UMovieSceneSection* InSection) const override;
+	/** Get the array of double channels for the specified section*/
+	CONSTRAINTS_API virtual TArrayView<FMovieSceneDoubleChannel*>  GetDoubleChannels(const UMovieSceneSection* InSection) const override;
+	/** Add Transform Keys at the specified times*/
+	CONSTRAINTS_API virtual bool AddTransformKeys(const TArray<FFrameNumber>& InFrames,
+		const TArray<FTransform>& InTransforms,
+		const EMovieSceneTransformChannel& InChannels,
+		const FFrameRate& InTickResolution,
+		UMovieSceneSection* InSection,
+		const bool bLocal = true) const override;
+
+	/** Resolve the bound objects so that any object it references are resovled and correctly set up*/
+	CONSTRAINTS_API virtual void ResolveBoundObjects(FMovieSceneSequenceID LocalSequenceID, TSharedRef<UE::MovieScene::FSharedPlaybackState> SharedPlaybackState, UObject* SubObject = nullptr) override;
+
+	/** Make a duplicate of myself with this outer*/
+	CONSTRAINTS_API virtual UTransformableHandle* Duplicate(UObject* NewOuter) const override;
+
+#if WITH_EDITOR
+	/** Returns labels used for UI. */
+	CONSTRAINTS_API virtual FString GetLabel() const override;
+	CONSTRAINTS_API virtual FString GetFullLabel() const override;
+#endif
+	
+	/** The Component that this handle is pointing at. */
+	UPROPERTY(BlueprintReadOnly, Category = "Object")
+	TWeakObjectPtr<USceneComponent> Component;
+
+	/** Optional socket name on Component. */
+	UPROPERTY(BlueprintReadOnly, Category = "Object")
+	FName SocketName = NAME_None;
+
+	/** Registers/Unregisters useful delegates to track changes in the Component's transform. */
+	CONSTRAINTS_API void UnregisterDelegates() const;
+	CONSTRAINTS_API void RegisterDelegates();
+	
+#if WITH_EDITOR
+	CONSTRAINTS_API void OnActorMoving(AActor* InActor);
+	CONSTRAINTS_API void OnPostPropertyChanged(UObject* InObject, FPropertyChangedEvent& InPropertyChangedEvent);
+	CONSTRAINTS_API void OnObjectsReplaced(const TMap<UObject*, UObject*>& InOldToNewInstances);
+#endif
+
+protected:
+
+	static FComponentEvaluationGraphBinding& GetEvaluationBinding();
+};
diff --git a/Engine/Source/Runtime/Experimental/Animation/Constraints/Public/Transform/TransformableHandleUtils.h b/Engine/Source/Runtime/Experimental/Animation/Constraints/Public/Transform/TransformableHandleUtils.h
new file mode 100644
--- /dev/null
+++ b/Engine/Source/Runtime/Experimental/Animation/Constraints/Public/Transform/TransformableHandleUtils.h
@@ -0,0 +1,14 @@
+﻿// Copyright Epic Games, Inc. All Rights Reserved.
+
+#pragma once
+
+#include "HAL/Platform.h"
+
+class USceneComponent;
+class USkeletalMeshComponent;
+
+namespace TransformableHandleUtils
+{
+	CONSTRAINTS_API void TickDependantComponents(USceneComponent* InComponent);
+	CONSTRAINTS_API void TickSkeletalMeshComponent(USkeletalMeshComponent* InSkeletalMeshComponent);
+}
diff --git a/Engine/Source/Runtime/Experimental/Animation/Constraints/Public/Transform/TransformableRegistry.h b/Engine/Source/Runtime/Experimental/Animation/Constraints/Public/Transform/TransformableRegistry.h
new file mode 100644
--- /dev/null
+++ b/Engine/Source/Runtime/Experimental/Animation/Constraints/Public/Transform/TransformableRegistry.h
@@ -0,0 +1,63 @@
+﻿// Copyright Epic Games, Inc. All Rights Reserved.
+
+#pragma once
+
+#include "Containers/Map.h"
+#include "HAL/Platform.h"
+#include "Templates/Function.h"
+
+class UClass;
+class UObject;
+class UTransformableHandle;
+
+/**
+ * TransformableRegistry
+ * Can be used to register customization transformable handle factory for specific objects. 
+ */
+
+class FTransformableRegistry
+{
+public:
+
+	/** Defines the function that will allocate a new transformable handle to wrap the object. */
+	using CreateHandleFuncT = TFunction<UTransformableHandle*(UObject*, const FName&)>;
+	/** Defines the function that return a hash value from the object and the attachment name. */
+	using GetHashFuncT = TFunction<uint32(const UObject*, const FName&)>;
+	
+	CONSTRAINTS_API ~FTransformableRegistry();
+
+	/** Get registry singleton */
+	static CONSTRAINTS_API FTransformableRegistry& Get();
+
+	/** Registers a InClass as transformable object. */
+	CONSTRAINTS_API void Register(UClass* InClass, CreateHandleFuncT InHandleFunc, GetHashFuncT InHashFunc);
+
+	/** Get the associated hash function for that specific class. */
+	CONSTRAINTS_API GetHashFuncT GetHashFunction(const UClass* InClass) const;
+	/** Get the associated handle creation function for that specific class. */
+	CONSTRAINTS_API CreateHandleFuncT GetCreateFunction(const UClass* InClass) const;
+
+protected:
+	friend class FConstraintsModule;
+
+	/** Registers the basic transformable objects. (called when starting the constraints module). */
+	static CONSTRAINTS_API void RegisterBaseObjects();
+	/** Registers the basic transformable objects. (called when shutting down the constraints module). */
+	static CONSTRAINTS_API void UnregisterAllObjects();
+	
+private:
+	FTransformableRegistry() = default;
+
+	/** Per class information that need to be register. */
+	struct FTransformableInfo
+	{
+		CreateHandleFuncT CreateHandleFunc;
+		GetHashFuncT GetHashFunc;
+	};
+
+	/** Find the associated info for that specific class. */
+	CONSTRAINTS_API const FTransformableInfo* FindInfo(const UClass* InClass) const;
+
+	/** List of all registered transformable objects. */
+	TMap<UClass*, FTransformableInfo> Transformables;
+};
diff --git a/Engine/Source/Runtime/Experimental/Animation/Constraints/Public/TransformConstraint.h b/Engine/Source/Runtime/Experimental/Animation/Constraints/Public/TransformConstraint.h
deleted file mode 100644
diff --git a/Engine/Source/Runtime/Experimental/Animation/Constraints/Public/TransformableHandle.h b/Engine/Source/Runtime/Experimental/Animation/Constraints/Public/TransformableHandle.h
deleted file mode 100644
diff --git a/Engine/Source/Runtime/Experimental/Animation/Constraints/Public/TransformableHandleUtils.h b/Engine/Source/Runtime/Experimental/Animation/Constraints/Public/TransformableHandleUtils.h
deleted file mode 100644
diff --git a/Engine/Source/Runtime/Experimental/Animation/Constraints/Public/TransformableRegistry.h b/Engine/Source/Runtime/Experimental/Animation/Constraints/Public/TransformableRegistry.h
deleted file mode 100644
diff --git a/Engine/Source/Runtime/MovieSceneTracks/Private/Sections/MovieScene3DTransformSection.cpp b/Engine/Source/Runtime/MovieSceneTracks/Private/Sections/MovieScene3DTransformSection.cpp
--- a/Engine/Source/Runtime/MovieSceneTracks/Private/Sections/MovieScene3DTransformSection.cpp
+++ b/Engine/Source/Runtime/MovieSceneTracks/Private/Sections/MovieScene3DTransformSection.cpp
@@ -15,9 +15,8 @@
 #include "EntitySystem/MovieSceneEntitySystemLinker.h"
 #include "EntitySystem/MovieSceneEntityBuilder.h"
 #include "MovieSceneTracksComponentTypes.h"
-#include "Algo/AnyOf.h"
-#include "TransformConstraint.h"
-#include "TransformableHandle.h"
+#include "Transform/TransformConstraint.h"
+#include "Transform/TransformableHandle.h"
 #include "UObject/ObjectSaveContext.h"
 #include "Misc/AxisDisplayInfo.h"
 
diff --git a/Engine/Source/Runtime/MovieSceneTracks/Private/Systems/MovieSceneConstraintSystem.cpp b/Engine/Source/Runtime/MovieSceneTracks/Private/Systems/MovieSceneConstraintSystem.cpp
--- a/Engine/Source/Runtime/MovieSceneTracks/Private/Systems/MovieSceneConstraintSystem.cpp
+++ b/Engine/Source/Runtime/MovieSceneTracks/Private/Systems/MovieSceneConstraintSystem.cpp
@@ -14,8 +14,8 @@
 #include "EntitySystem/BuiltInComponentTypes.h"
 #include "MovieSceneTracksComponentTypes.h"
 #include "Systems/MovieSceneComponentTransformSystem.h"
-#include "TransformConstraint.h"
-#include "TransformableHandle.h"
+#include "Transform/TransformConstraint.h"
+#include "Transform/TransformableHandle.h"
 #include "Sections/MovieSceneConstrainedSection.h"
 
 #include UE_INLINE_GENERATED_CPP_BY_NAME(MovieSceneConstraintSystem)
