[Chaos] Initial DSTM Migration support for Particles and Joints

This CL:
- Adds new methods to handle the Serialize calls done when an object migration is happening in a UPrimitiveComponent or a USkeletalMeshComponent (other variants like Cluster unions and Geometry Collections need to override the new methods to handle the migration correctly)
- Created a new fast struct serializer that is able to serialize ustructs in an unversioned way without the baggage of the existing ustrcut binary serializer. Test shows it to be as fast a manual custom serialization via FArchive.
- Created a new Solver Serializer object that handles serializing particle and (for now) Joint Constraint data using CVD's structs and the new fast serializer. This serialized data is then used by the object migration process and sent over the network. At later point in the remote server that received the data, this new serializer object will update existing particle/joints instances with the serialized data.
The idea is to use this serializer to also serialize an entire solver scene.
- Added a way to hold newly serialized data in the Physics Scene Object until a Physics proxy that can consume it is created (this is needed because during the object migration process we still don't have a valid physics proxy).
- Added a way to hold newly serialized data in the new solver serializer object. This is needed because there is a period of time where we might have a physics proxy ready, but no representation in the Physics Thread yet.

diff --git a/Engine/Source/Programs/UnrealBuildTool/Configuration/ModuleRules.cs b/Engine/Source/Programs/UnrealBuildTool/Configuration/ModuleRules.cs
--- a/Engine/Source/Programs/UnrealBuildTool/Configuration/ModuleRules.cs
+++ b/Engine/Source/Programs/UnrealBuildTool/Configuration/ModuleRules.cs
@@ -1713,14 +1713,21 @@ namespace UnrealBuildTool
 			}
 		}
 
+		/// <summary>
+		/// Returns true if the Chaos Visual Debugger is supported by the provided target
+		/// </summary>
+		protected bool IsChaosVisualDebuggerSupported(ReadOnlyTargetRules target)
+		{
+			return target.bCompileChaosVisualDebuggerSupport && target.Configuration != UnrealTargetConfiguration.Shipping && target.bEnableTrace;
+		}
+
 		/// <summary>
 		/// Setup this module for Chaos Visual Debugger support (Required for recording debug data that will be visualized in the Chaos Visual Debugger tool)
 		/// </summary>
 		public void SetupModuleChaosVisualDebuggerSupport(ReadOnlyTargetRules target)
 		{
 			//TODO: Emit some form of warning if CVD support is explicitly enabled, but UE Trace is disabled
-			bool bHasChaosVisualDebuggerSupport = target.bCompileChaosVisualDebuggerSupport && target.Configuration != UnrealTargetConfiguration.Shipping && target.bEnableTrace;
-			if (bHasChaosVisualDebuggerSupport)
+			if (IsChaosVisualDebuggerSupported(target))
 			{
 				PublicDependencyModuleNames.Add("ChaosVDRuntime");
 
diff --git a/Engine/Source/Runtime/Engine/Classes/Components/PrimitiveComponent.h b/Engine/Source/Runtime/Engine/Classes/Components/PrimitiveComponent.h
--- a/Engine/Source/Runtime/Engine/Classes/Components/PrimitiveComponent.h
+++ b/Engine/Source/Runtime/Engine/Classes/Components/PrimitiveComponent.h
@@ -3151,6 +3151,17 @@ public:
 	/** Whether the physics is static for this component. */
 	FORCEINLINE bool IsStaticPhysics() const { return Mobility != EComponentMobility::Movable && bStaticWhenNotMoveable; }
 
+protected:
+
+	UFUNCTION()
+	void HandlePostMigrationPhysicsStateChanged(UPrimitiveComponent* ChangedComponent, EComponentPhysicsStateChange StateChange);
+
+#if UE_WITH_REMOTE_OBJECT_HANDLE
+	ENGINE_API virtual void MigratePhysicsState(FArchive& Ar) override;
+	
+	ENGINE_API virtual void ApplyMigratedPhysicsState() override;
+#endif
+
 protected:
 	void AddNavigationModifier(FNavigationRelevantData& OutData) const;
 };
diff --git a/Engine/Source/Runtime/Engine/Classes/Components/SkeletalMeshComponent.h b/Engine/Source/Runtime/Engine/Classes/Components/SkeletalMeshComponent.h
--- a/Engine/Source/Runtime/Engine/Classes/Components/SkeletalMeshComponent.h
+++ b/Engine/Source/Runtime/Engine/Classes/Components/SkeletalMeshComponent.h
@@ -20,6 +20,7 @@
 #include "ClothingSimulationInterface.h"
 #include "ClothingSimulationFactory.h"
 #include "Animation/AttributesRuntime.h"
+#include "Chaos/Serialization/SerializedDataBuffer.h"
 #if WITH_ENGINE
 #include "Engine/PoseWatchRenderData.h"
 #endif
@@ -298,6 +299,11 @@ struct FClosestPointOnPhysicsAsset
 	}
 };
 
+namespace Chaos
+{
+	struct FSerializedDataBuffer;
+}
+
 /**
  * SkeletalMeshComponent is used to create an instance of an animated SkeletalMesh asset.
  *
@@ -1977,6 +1983,40 @@ protected:
 	 * TransformToRoot offset. Otherwise, it returns the component transform.
 	 */
 	ENGINE_API virtual FTransform GetComponentTransformFromBodyInstance(FBodyInstance* UseBI) override;
+
+private:
+	struct UE_INTERNAL FMigratedPhysicsState
+	{
+		int32 NumBodies = 0;
+		int32 NumConstraints = 0;
+
+		void ReadElementDataIntoBuffer(Chaos::FSerializedDataBuffer& TargetBuffer);
+
+		template<typename ElementType>
+		void AppendDataFromElement(ElementType& InElement)
+		{
+			DataOffsetsPerElementIndex.Emplace(MigratedStateAsBytes.GetDataAsByteArrayRef().Num());
+			InElement.SerializeState(MigratedStateAsBytes);
+		}
+
+		void Serialize(FArchive& Ar);
+
+	private:
+		int32 GetDataSizeForElementIndex(int32 ElementIndex) const;
+		int32 GetDataOffset(int32 ElementIndex) const;
+
+		Chaos::FSerializedDataBuffer MigratedStateAsBytes;
+		TArray<int32> DataOffsetsPerElementIndex;
+		int32 CurrentReadElementIndex = 0;
+	};
+
+protected:
+
+#if UE_WITH_REMOTE_OBJECT_HANDLE
+	ENGINE_API virtual void MigratePhysicsState(FArchive& Ar) override;
+	ENGINE_API virtual void ApplyMigratedPhysicsState() override;
+#endif
+
 	//~ End UPrimitiveComponent Interface.
 
 public:
diff --git a/Engine/Source/Runtime/Engine/Classes/Interfaces/IPhysicsComponent.h b/Engine/Source/Runtime/Engine/Classes/Interfaces/IPhysicsComponent.h
--- a/Engine/Source/Runtime/Engine/Classes/Interfaces/IPhysicsComponent.h
+++ b/Engine/Source/Runtime/Engine/Classes/Interfaces/IPhysicsComponent.h
@@ -26,6 +26,15 @@ public:
 	virtual TArray<Chaos::FPhysicsObject*> GetAllPhysicsObjects() const = 0;
 	
 	virtual Chaos::FPhysicsObjectId GetIdFromGTParticle(Chaos::FGeometryParticle* Particle) const = 0;
+
+#if UE_WITH_REMOTE_OBJECT_HANDLE
+
+protected:
+	virtual void MigratePhysicsState(FArchive& Ar) = 0;
+	virtual void ApplyMigratedPhysicsState() = 0;
+public:
+
+#endif
 };
 
 template<typename TId>
diff --git a/Engine/Source/Runtime/Engine/Classes/PhysicsEngine/BodyInstance.h b/Engine/Source/Runtime/Engine/Classes/PhysicsEngine/BodyInstance.h
--- a/Engine/Source/Runtime/Engine/Classes/PhysicsEngine/BodyInstance.h
+++ b/Engine/Source/Runtime/Engine/Classes/PhysicsEngine/BodyInstance.h
@@ -78,6 +78,11 @@ struct FInitBodySpawnParams
 	FPhysicsAggregateHandle Aggregate;
 };
 
+namespace Chaos
+{
+	struct FSerializedDataBuffer;
+}
+
 namespace Chaos::Private
 {
 	// Private base class necessary to keep friend access to FBodyInstance
@@ -1453,6 +1458,14 @@ private:
 	friend class FBodyInstanceCustomizationHelper;
 	friend class FFoliageTypeCustomizationHelpers;
 
+public:
+
+	/** Serialized out the physics thread state of this body as an array of bytes */
+	void SerializeState(Chaos::FSerializedDataBuffer& OutSerializedState);
+
+	/** Deserializes the provided array of bytes representing a saved physics thread state and applies it to this body */
+	void ApplySerializedState(Chaos::FSerializedDataBufferPtr&& InSerializedState);
+
 private:
 
 	ENGINE_API void UpdateDebugRendering();
diff --git a/Engine/Source/Runtime/Engine/Classes/PhysicsEngine/ConstraintInstance.h b/Engine/Source/Runtime/Engine/Classes/PhysicsEngine/ConstraintInstance.h
--- a/Engine/Source/Runtime/Engine/Classes/PhysicsEngine/ConstraintInstance.h
+++ b/Engine/Source/Runtime/Engine/Classes/PhysicsEngine/ConstraintInstance.h
@@ -24,6 +24,11 @@ class FMaterialRenderProxy;
 class UPhysicsAsset;
 struct FReferenceSkeleton;
 
+namespace Chaos
+{
+	struct FSerializedDataBuffer;
+}
+
 UENUM()
 enum class EConstraintTransformComponentFlags : uint8
 {
@@ -326,6 +331,12 @@ public:
 	/** Get underlying physics engine constraint */
 	ENGINE_API const FPhysicsConstraintHandle& GetPhysicsConstraintRef() const;
 
+	/** Serialized out the physics thread state of this constraint as an array of bytes */
+	void SerializeState(Chaos::FSerializedDataBuffer& OutSerializedState);
+
+	/** Deserializes the provided array of bytes representing a saved physics thread state and applies it to this constraint */
+	void ApplySerializedState(Chaos::FSerializedDataBufferPtr&& InSerializedState);
+
 	FChaosUserData UserData;
 
 	/** Constructor **/
diff --git a/Engine/Source/Runtime/Engine/Private/Components/PrimitiveComponent.cpp b/Engine/Source/Runtime/Engine/Private/Components/PrimitiveComponent.cpp
--- a/Engine/Source/Runtime/Engine/Private/Components/PrimitiveComponent.cpp
+++ b/Engine/Source/Runtime/Engine/Private/Components/PrimitiveComponent.cpp
@@ -54,6 +54,9 @@
 #include "MaterialCachedData.h"
 #include "MaterialShared.h"
 #include "MarkActorRenderStateDirtyTask.h"
+#include "Chaos/Serialization/SerializationUtils.h"
+#include "Physics/Experimental/PhysScene_Chaos.h"
+#include "PhysicsProxy/SingleParticlePhysicsProxy.h"
 
 #if WITH_EDITOR
 #include "Engine/LODActor.h"
@@ -63,12 +66,15 @@
 #include "Components/InstancedStaticMeshComponent.h"
 #endif // WITH_EDITOR
 
+#include "Chaos/Serialization/SerializedDataBuffer.h"
+
 #if DO_CHECK
 #include "Engine/StaticMesh.h"
 #endif
 
 #define LOCTEXT_NAMESPACE "PrimitiveComponent"
 
+
 CSV_DEFINE_CATEGORY(PrimitiveComponent, false);
 
 //////////////////////////////////////////////////////////////////////////
@@ -1234,6 +1240,8 @@ void UPrimitiveComponent::Serialize(FArchive& Ar)
 		}
 		
 		Ar << MoveIgnoreMask;
+		
+		MigratePhysicsState(Ar);
 	}
 #endif
 
@@ -1261,6 +1269,84 @@ void UPrimitiveComponent::Serialize(FArchive& Ar)
 #endif
 }
 
+#if UE_WITH_REMOTE_OBJECT_HANDLE
+void UPrimitiveComponent::MigratePhysicsState(FArchive& Ar)
+{
+	if (Ar.IsLoading())
+	{
+		Chaos::FSerializedDataBufferPtr MigratedStateDataBuffer = MakeUnique<Chaos::FSerializedDataBuffer>();
+		
+		Ar << (*MigratedStateDataBuffer);
+
+		if (MigratedStateDataBuffer->GetDataAsByteArrayRef().IsEmpty())
+		{
+			UE_LOG(LogPrimitiveComponent, Warning, TEXT("Migrated an empty physics state | [%s] will use its default physics state."), *GetPathNameSafe(this))
+			return;
+		}
+
+		if (FPhysScene* PhysicsScene = FPhysicsInterface::GetCurrentSceneForComponent(this))
+		{
+			PhysicsScene->PushPendingSerializedStateForObject(this, MoveTemp(MigratedStateDataBuffer));
+
+			OnComponentPhysicsStateChanged.AddDynamic(this, &UPrimitiveComponent::HandlePostMigrationPhysicsStateChanged);
+
+			if (HasValidPhysicsState())
+			{
+				RecreatePhysicsState();
+			}
+		}
+	}
+	else if (Ar.IsSaving())
+	{
+		if (HasValidPhysicsState())
+		{
+			Chaos::FSerializedDataBuffer SerializedStateData;
+			BodyInstance.SerializeState(SerializedStateData);
+
+			Ar << SerializedStateData;
+		}
+		else
+		{
+			// This object might be being migrated before we had a chance to process a previous migration
+			FPhysScene* PhysicsScene = FPhysicsInterface::GetCurrentSceneForComponent(this);
+			if (Chaos::FSerializedDataBufferPtr RawMigratedData = PhysicsScene ? PhysicsScene->PopPendingSerializedStateForObject(this) : nullptr)
+			{
+				Ar << (*RawMigratedData);
+			}
+			else
+			{
+				// We need to serialize an empty buffer as array, so we can handle this case during the load phase post migration
+				Chaos::FSerializedDataBuffer EmptyPhysicsState;
+				Ar << EmptyPhysicsState;
+			}
+		}
+	}
+}
+
+void UPrimitiveComponent::ApplyMigratedPhysicsState()
+{
+	if (ensure(HasValidPhysicsState()))
+	{
+		FPhysScene* PhysicsScene = FPhysicsInterface::GetCurrentSceneForComponent(this);
+		if (Chaos::FSerializedDataBufferPtr RawMigratedData = PhysicsScene ? PhysicsScene->PopPendingSerializedStateForObject(this) : nullptr)
+		{
+			BodyInstance.ApplySerializedState(MoveTemp(RawMigratedData));
+		}
+	}
+}
+#endif
+
+void UPrimitiveComponent::HandlePostMigrationPhysicsStateChanged(UPrimitiveComponent* ChangedComponent, EComponentPhysicsStateChange StateChange)
+{
+#if UE_WITH_REMOTE_OBJECT_HANDLE
+	if (StateChange == EComponentPhysicsStateChange::Created)
+	{
+		ApplyMigratedPhysicsState();
+		OnComponentPhysicsStateChanged.RemoveDynamic(this, &UPrimitiveComponent::HandlePostMigrationPhysicsStateChanged);
+	}
+#endif
+}
+
 #if WITH_EDITOR
 void UPrimitiveComponent::PostEditChangeProperty(FPropertyChangedEvent& PropertyChangedEvent)
 {
@@ -1687,6 +1773,8 @@ void UPrimitiveComponent::BeginDestroy()
 		IStreamingManager::Get().NotifyPrimitiveDetached(this);
 	}
 
+	OnComponentPhysicsStateChanged.RemoveDynamic(this, &UPrimitiveComponent::HandlePostMigrationPhysicsStateChanged);
+
 	Super::BeginDestroy();
 
 	// Use a fence to keep track of when the rendering thread executes this scene detachment.
diff --git a/Engine/Source/Runtime/Engine/Private/Components/SkeletalMeshComponent.cpp b/Engine/Source/Runtime/Engine/Private/Components/SkeletalMeshComponent.cpp
--- a/Engine/Source/Runtime/Engine/Private/Components/SkeletalMeshComponent.cpp
+++ b/Engine/Source/Runtime/Engine/Private/Components/SkeletalMeshComponent.cpp
@@ -56,6 +56,8 @@
 #include "Kismet2/BlueprintEditorUtils.h"
 #endif
 
+#include "Chaos/Serialization/SerializedDataBuffer.h"
+
 #include UE_INLINE_GENERATED_CPP_BY_NAME(SkeletalMeshComponent)
 
 LLM_DEFINE_TAG(SkeletalMesh_TransformData);
@@ -3791,6 +3793,167 @@ void USkeletalMeshComponent::GetResourceSizeEx(FResourceSizeEx& CumulativeResour
 	}
 }
 
+void USkeletalMeshComponent::FMigratedPhysicsState::ReadElementDataIntoBuffer(Chaos::FSerializedDataBuffer& TargetBuffer)
+{
+	const int32 SizeToCopy = GetDataSizeForElementIndex(CurrentReadElementIndex);
+	const int32 StartPos = GetDataOffset(CurrentReadElementIndex);
+
+	TArray<uint8>& AsByteArray = MigratedStateAsBytes.GetDataAsByteArrayRef();
+
+	if (StartPos + SizeToCopy < AsByteArray.Num())
+	{
+		TargetBuffer.GetDataAsByteArrayRef().Append(&AsByteArray[StartPos], SizeToCopy);
+	}
+	
+	CurrentReadElementIndex++;
+}
+
+int32 USkeletalMeshComponent::FMigratedPhysicsState::GetDataSizeForElementIndex(int32 ElementIndex) const
+{
+	if (!DataOffsetsPerElementIndex.IsValidIndex(ElementIndex))
+	{
+		return 0;
+	}
+
+	int32 NextIndex = ElementIndex + 1;
+	return DataOffsetsPerElementIndex.IsValidIndex(NextIndex) ? DataOffsetsPerElementIndex[NextIndex] - DataOffsetsPerElementIndex[ElementIndex] : 0;
+}
+
+int32 USkeletalMeshComponent::FMigratedPhysicsState::GetDataOffset(int32 ElementIndex) const
+{
+	return DataOffsetsPerElementIndex.IsValidIndex(ElementIndex) ? DataOffsetsPerElementIndex[ElementIndex] : 0;
+}
+
+void USkeletalMeshComponent::FMigratedPhysicsState::Serialize(FArchive& Ar)
+{
+	Ar << NumBodies;
+	Ar << NumConstraints;
+	Ar << DataOffsetsPerElementIndex;
+	Ar << MigratedStateAsBytes;
+}
+
+#if UE_WITH_REMOTE_OBJECT_HANDLE
+
+void USkeletalMeshComponent::MigratePhysicsState(FArchive& Ar)
+{
+	if (Ar.IsLoading())
+	{
+		Chaos::FSerializedDataBufferPtr MigratedStateData = MakeUnique<Chaos::FSerializedDataBuffer>();
+
+		Ar << (*MigratedStateData);
+
+		if (MigratedStateData->GetDataAsByteArrayRef().IsEmpty())
+		{
+			UE_LOG(LogSkeletalMesh, Warning, TEXT("Migrated an empty physics state | [%s] will use its default physics state."), *GetPathNameSafe(this))
+			return;
+		}
+
+		if (FPhysScene* PhysicsScene = FPhysicsInterface::GetCurrentSceneForComponent(this))
+		{
+			PhysicsScene->PushPendingSerializedStateForObject(this, MoveTemp(MigratedStateData));
+
+			RegisterOnPhysicsCreatedDelegate(FOnSkelMeshPhysicsCreated::CreateLambda([WeakThis = TWeakObjectPtr<USkeletalMeshComponent>(this)]()
+			{
+				if (USkeletalMeshComponent* SKComponent = WeakThis.Get())
+				{
+					SKComponent->ApplyMigratedPhysicsState();
+				}
+			}));
+
+			if (HasValidPhysicsState())
+			{
+				RecreatePhysicsState();
+			}
+		}
+	}
+	else if (Ar.IsSaving())
+	{
+		if (HasValidPhysicsState())
+		{
+			FMigratedPhysicsState MigratedState;
+			MigratedState.NumBodies = Bodies.Num();
+			MigratedState.NumConstraints = Constraints.Num();
+
+			for (FBodyInstance* PhysicsBody : Bodies)
+			{
+				check(PhysicsBody);
+				MigratedState.AppendDataFromElement(*PhysicsBody);
+			}
+			
+			for (FConstraintInstance* ConstraintInstance : Constraints)
+			{
+				check(ConstraintInstance);
+				MigratedState.AppendDataFromElement(*ConstraintInstance);
+			}
+
+			// We need re-serialize the state as a chaos data buffer, so we can use it as is on de-serialization
+			Chaos::FSerializedDataBuffer MigratedStateAsBytes;
+			FMemoryWriter MigratedStateWriter(MigratedStateAsBytes.GetDataAsByteArrayRef());
+			MigratedState.Serialize(MigratedStateWriter);
+
+			Ar << MigratedStateAsBytes;
+		}
+		else
+		{
+			// This object might be being migrated before we had a chance to process a previous migration
+			FPhysScene* PhysicsScene = FPhysicsInterface::GetCurrentSceneForComponent(this);
+			if (Chaos::FSerializedDataBufferPtr RawMigratedData = PhysicsScene ? PhysicsScene->PopPendingSerializedStateForObject(this) : nullptr)
+			{
+				Ar << (*RawMigratedData);
+			}
+			else
+			{
+				// We need to serialize an empty buffer as array, so we can handle this case during the load phase post migration
+				Chaos::FSerializedDataBuffer EmptyPhysicsState;
+				Ar << EmptyPhysicsState;
+			}
+		}
+	}
+}
+
+void USkeletalMeshComponent::ApplyMigratedPhysicsState()
+{
+	FPhysScene* PhysicsScene = FPhysicsInterface::GetCurrentSceneForComponent(this);
+	
+	if (Chaos::FSerializedDataBufferPtr RawMigratedData = PhysicsScene ? PhysicsScene->PopPendingSerializedStateForObject(this) : nullptr)
+	{
+		FMemoryReader MultiBodyStateReader(RawMigratedData->GetDataAsByteArrayRef());
+
+		FMigratedPhysicsState MigratedState;
+		MigratedState.Serialize(MultiBodyStateReader);
+
+		const bool bIsValidMigratedState = MigratedState.NumBodies == Bodies.Num() || MigratedState.NumConstraints == Constraints.Num();
+		if (!ensure(bIsValidMigratedState))
+		{
+			return;
+		}
+
+		for (FBodyInstance* BodyInstance : Bodies)
+		{
+			Chaos::FSerializedDataBufferPtr DataBuffer = MakeUnique<Chaos::FSerializedDataBuffer>();
+			MigratedState.ReadElementDataIntoBuffer(*DataBuffer);
+
+			if (ensure(BodyInstance))
+			{
+				BodyInstance->ApplySerializedState(MoveTemp(DataBuffer));
+			}
+		}
+
+		for (FConstraintInstance* ConstraintInstance : Constraints)
+		{
+			Chaos::FSerializedDataBufferPtr DataBuffer = MakeUnique<Chaos::FSerializedDataBuffer>();
+			MigratedState.ReadElementDataIntoBuffer(*DataBuffer);
+
+			if (ensure(ConstraintInstance))
+			{
+				ConstraintInstance->ApplySerializedState(MoveTemp(DataBuffer));
+			}
+		}
+	}
+}
+
+#endif
+
 void USkeletalMeshComponent::SetAnimationMode(EAnimationMode::Type InAnimationMode, bool bForceInitAnimScriptInstance)
 {
 	if(!bEnableAnimation)
@@ -5144,7 +5307,6 @@ bool USkeletalMeshComponent::FindAttributeChecked(const FName& BoneName, const F
 	return bFound;
 }
 
-
 const FName& USkeletalMeshComponent::GetAnimationModePropertyNameChecked()
 {
 	static FName Name = GET_MEMBER_NAME_CHECKED(USkeletalMeshComponent, AnimationMode);
diff --git a/Engine/Source/Runtime/Engine/Private/PhysicsEngine/BodyInstance.cpp b/Engine/Source/Runtime/Engine/Private/PhysicsEngine/BodyInstance.cpp
--- a/Engine/Source/Runtime/Engine/Private/PhysicsEngine/BodyInstance.cpp
+++ b/Engine/Source/Runtime/Engine/Private/PhysicsEngine/BodyInstance.cpp
@@ -3421,6 +3421,22 @@ void FBodyInstance::UpdateDebugRendering()
 #endif
 }
 
+void FBodyInstance::ApplySerializedState(Chaos::FSerializedDataBufferPtr&& InSerializedState)
+{
+	FPhysicsCommand::ExecuteWrite(ActorHandle, [&InSerializedState](const FPhysicsActorHandle& Actor)
+	{
+		FChaosEngineInterface::ApplySerializedStateToParticle(Actor, MoveTemp(InSerializedState));
+	});
+}
+
+void FBodyInstance::SerializeState(Chaos::FSerializedDataBuffer& OutSerializedState)
+{
+	FPhysicsCommand::ExecuteRead(ActorHandle, [&OutSerializedState](const FPhysicsActorHandle& Actor)
+	{
+		FChaosEngineInterface::SerializeParticleStateToBuffer(Actor, OutSerializedState);
+	});
+}
+
 void FBodyInstance::UpdateDampingProperties()
 {
 	FPhysicsCommand::ExecuteWrite(GetPhysicsActor(), [&](FPhysicsActorHandle Actor)
diff --git a/Engine/Source/Runtime/Engine/Private/PhysicsEngine/ConstraintInstance.cpp b/Engine/Source/Runtime/Engine/Private/PhysicsEngine/ConstraintInstance.cpp
--- a/Engine/Source/Runtime/Engine/Private/PhysicsEngine/ConstraintInstance.cpp
+++ b/Engine/Source/Runtime/Engine/Private/PhysicsEngine/ConstraintInstance.cpp
@@ -1,6 +1,7 @@
 // Copyright Epic Games, Inc. All Rights Reserved.
 
 #include "PhysicsEngine/ConstraintInstance.h"
+
 #include "Engine/SkeletalMesh.h"
 #include "UObject/FrameworkObjectVersion.h"
 #include "Components/SkeletalMeshComponent.h"
@@ -11,6 +12,7 @@
 
 #include "Logging/MessageLog.h"
 #include "Misc/UObjectToken.h"
+#include "Physics/Experimental/PhysScene_Chaos.h"
 
 #include UE_INLINE_GENERATED_CPP_BY_NAME(ConstraintInstance)
 
@@ -711,6 +713,21 @@ const FPhysicsConstraintHandle& FConstraintInstance::GetPhysicsConstraintRef() c
 	return ConstraintHandle;
 }
 
+void FConstraintInstance::SerializeState(Chaos::FSerializedDataBuffer& OutSerializedState)
+{
+	FPhysicsCommand::ExecuteWrite(ConstraintHandle, [&OutSerializedState](const FPhysicsConstraintHandle& Constraint)
+	{
+		FChaosEngineInterface::SerializeConstraintStateToBuffer(Constraint, OutSerializedState);
+	});
+}
+
+void FConstraintInstance::ApplySerializedState(Chaos::FSerializedDataBufferPtr&& InSerializedState)
+{
+	FPhysicsCommand::ExecuteRead(ConstraintHandle, [&InSerializedState, this](const FPhysicsConstraintHandle& Constraint)
+	{
+		FChaosEngineInterface::ApplySerializedStateToConstraint(Constraint, MoveTemp(InSerializedState));
+	});
+}
 
 void FConstraintInstance::CopyConstraintGeometryFrom(const FConstraintInstance* FromInstance)
 {
diff --git a/Engine/Source/Runtime/Engine/Private/PhysicsEngine/Experimental/PhysInterface_Chaos.cpp b/Engine/Source/Runtime/Engine/Private/PhysicsEngine/Experimental/PhysInterface_Chaos.cpp
--- a/Engine/Source/Runtime/Engine/Private/PhysicsEngine/Experimental/PhysInterface_Chaos.cpp
+++ b/Engine/Source/Runtime/Engine/Private/PhysicsEngine/Experimental/PhysInterface_Chaos.cpp
@@ -18,6 +18,7 @@
 #include "Components/PrimitiveComponent.h"
 #include "Materials/MaterialInterface.h"
 #include "PBDRigidsSolver.h"
+#include "Engine/World.h"
 #include "PhysicsEngine/BodySetup.h"
 #include "Physics/Experimental/ChaosScopedSceneLock.h"
 #include "Physics/PhysicsFiltering.h"
@@ -197,6 +198,27 @@ void FPhysInterface_Chaos::UpdateMaterialMask(FPhysicsMaterialMaskHandle& InHand
 	Chaos::FPhysicalMaterialManager::Get().UpdateMaterialMask(InHandle);
 }
 
+FPhysScene* FPhysInterface_Chaos::GetCurrentSceneForComponent(UPrimitiveComponent* InComponent)
+{
+	if (!InComponent)
+	{
+		return nullptr;
+	}
+
+	FBodyInstance* BodyInstance = InComponent->GetBodyInstance();
+	if (BodyInstance && BodyInstance->GetPhysicsActorHandle())
+	{
+		return static_cast<FPhysScene*>(GetCurrentScene(BodyInstance->GetPhysicsActorHandle()));
+	}
+
+	if (UWorld* ComponentWorld = InComponent->GetWorld())
+	{
+		return ComponentWorld->GetPhysicsScene();
+	}
+
+	return nullptr;
+}
+
 bool FPhysInterface_Chaos::IsInScene(const FPhysicsActorHandle& InActorReference)
 {
 	return (GetCurrentScene(InActorReference) != nullptr);
diff --git a/Engine/Source/Runtime/Engine/Private/PhysicsEngine/Experimental/PhysScene_Chaos.cpp b/Engine/Source/Runtime/Engine/Private/PhysicsEngine/Experimental/PhysScene_Chaos.cpp
--- a/Engine/Source/Runtime/Engine/Private/PhysicsEngine/Experimental/PhysScene_Chaos.cpp
+++ b/Engine/Source/Runtime/Engine/Private/PhysicsEngine/Experimental/PhysScene_Chaos.cpp
@@ -31,6 +31,7 @@
 #include "Chaos/PendingSpatialData.h"
 #include "Chaos/PhysicsSolverBaseImpl.h"
 #include "Chaos/AsyncInitBodyHelper.h"
+#include "Chaos/Serialization/SerializedDataBuffer.h"
 #include "Misc/CoreMisc.h"
 
 #if WITH_EDITOR
@@ -841,6 +842,27 @@ void FPhysScene_Chaos::RemoveObject(Chaos::FClusterUnionPhysicsProxy* InObject)
 	RemoveFromComponentMaps(InObject);
 }
 
+void FPhysScene_Chaos::PushPendingSerializedStateForObject(UObject* InComponent, Chaos::FSerializedDataBufferPtr&& InState)
+{
+	if (InComponent && ensure(InState))
+	{
+		PendingMigratedPhysicsStateByComponent.Emplace(InComponent, MoveTemp(InState));
+	}
+}
+
+Chaos::FSerializedDataBufferPtr FPhysScene_Chaos::PopPendingSerializedStateForObject(UObject* InComponent)
+{
+	if (PendingMigratedPhysicsStateByComponent.IsEmpty())
+	{
+		return nullptr;
+	}
+
+	Chaos::FSerializedDataBufferPtr FoundStateRef = nullptr;
+	PendingMigratedPhysicsStateByComponent.RemoveAndCopyValue(InComponent, FoundStateRef);
+
+	return MoveTemp(FoundStateRef);
+}
+
 IPhysicsReplication* FPhysScene_Chaos::GetPhysicsReplication()
 {
 	return PhysicsReplication.Get();
@@ -1411,9 +1433,11 @@ void FPhysScene_Chaos::CreatePhysicsReplicationCache()
 void FPhysScene_Chaos::RemoveFromComponentMaps(IPhysicsProxyBase* InObject)
 {
 	UE_CHAOS_ASYNC_INITBODY_WRITESCOPELOCK(PhysicsProxyComponentMapsLock);
-	auto* const Component = PhysicsProxyToComponentMap.Find(InObject);
+	TObjectPtr<UPrimitiveComponent>* const Component = PhysicsProxyToComponentMap.Find(InObject);
 	if (Component)
 	{
+		PendingMigratedPhysicsStateByComponent.Remove(*Component);
+
 		TArray<IPhysicsProxyBase*>* ProxyArray = ComponentToPhysicsProxyMap.Find(*Component);
 		if (ProxyArray)
 		{
diff --git a/Engine/Source/Runtime/Engine/Public/Physics/Experimental/PhysInterface_Chaos.h b/Engine/Source/Runtime/Engine/Public/Physics/Experimental/PhysInterface_Chaos.h
--- a/Engine/Source/Runtime/Engine/Public/Physics/Experimental/PhysInterface_Chaos.h
+++ b/Engine/Source/Runtime/Engine/Public/Physics/Experimental/PhysInterface_Chaos.h
@@ -44,7 +44,9 @@ public:
 	static ENGINE_API void UpdateMaterialMask(FPhysicsMaterialMaskHandle& InHandle, const UPhysicalMaterialMask* InMaterialMask);
 
 	// Actor interface functions
-	
+
+	static ENGINE_API FPhysScene* GetCurrentSceneForComponent(UPrimitiveComponent* InComponent);
+
 	static ENGINE_API void FlushScene(FPhysScene* InScene);
 	static ENGINE_API bool IsInScene(const FPhysicsActorHandle& InActorReference);
 
diff --git a/Engine/Source/Runtime/Engine/Public/Physics/Experimental/PhysScene_Chaos.h b/Engine/Source/Runtime/Engine/Public/Physics/Experimental/PhysScene_Chaos.h
--- a/Engine/Source/Runtime/Engine/Public/Physics/Experimental/PhysScene_Chaos.h
+++ b/Engine/Source/Runtime/Engine/Public/Physics/Experimental/PhysScene_Chaos.h
@@ -72,6 +72,7 @@ namespace Chaos
 	template <typename T>
 	class TArrayCollectionArray;
 
+	struct FSerializedDataBuffer;
 }
 
 extern int32 GEnableKinematicDeferralStartPhysicsCondition;
@@ -166,6 +167,9 @@ public:
 	ENGINE_API void RemoveObject(FGeometryCollectionPhysicsProxy* InObject);
 	ENGINE_API void RemoveObject(Chaos::FClusterUnionPhysicsProxy* InObject);
 
+	void PushPendingSerializedStateForObject(UObject* InComponent, Chaos::FSerializedDataBufferPtr&& InState);
+	Chaos::FSerializedDataBufferPtr PopPendingSerializedStateForObject(UObject* InComponent);
+
 	ENGINE_API IPhysicsReplication* GetPhysicsReplication();
 
 	ENGINE_API IPhysicsReplication* CreatePhysicsReplication();
@@ -349,6 +353,9 @@ public:
 	};
 
 private:
+	
+	TMap<UObject*, Chaos::FSerializedDataBufferPtr> PendingMigratedPhysicsStateByComponent;
+
 	TUniquePtr<FPhysicsReplicationCache> PhysicsReplicationCache;
 
 	TSet<UPrimitiveComponent*> CollisionEventRegistrations;
@@ -509,4 +516,3 @@ private:
 	friend struct ChaosInterface::FScopedSceneReadLock;
 	friend struct FScopedSceneLock_Chaos;
 };
-
diff --git a/Engine/Source/Runtime/Experimental/Chaos/Chaos.Build.cs b/Engine/Source/Runtime/Experimental/Chaos/Chaos.Build.cs
--- a/Engine/Source/Runtime/Experimental/Chaos/Chaos.Build.cs
+++ b/Engine/Source/Runtime/Experimental/Chaos/Chaos.Build.cs
@@ -20,6 +20,7 @@ namespace UnrealBuildTool.Rules
 				"TraceLog",
                 "Voronoi",
 				"GeometryCore",
+				"ChaosVDRuntime",
 				"NNE",
 				}
 			);
diff --git a/Engine/Source/Runtime/Experimental/Chaos/Private/Chaos/PBDJointConstraints.cpp b/Engine/Source/Runtime/Experimental/Chaos/Private/Chaos/PBDJointConstraints.cpp
--- a/Engine/Source/Runtime/Experimental/Chaos/Private/Chaos/PBDJointConstraints.cpp
+++ b/Engine/Source/Runtime/Experimental/Chaos/Private/Chaos/PBDJointConstraints.cpp
@@ -285,6 +285,12 @@ namespace Chaos
 		return ConcreteContainer()->GetConstraintSettings(ConstraintIndex);
 	}
 
+
+	void FPBDJointConstraintHandle::SetJointState(const FPBDJointState& NewState)
+	{
+		ConcreteContainer()->SetConstraintState(ConstraintIndex, NewState);
+	}
+
 	void FPBDJointConstraintHandle::SetSettings(const FPBDJointSettings& InSettings)
 	{
 		// buffer the previous targets so plasticity can reuse them
@@ -655,7 +661,11 @@ namespace Chaos
 		return Handles.Last();
 	}
 
-	
+	void FPBDJointConstraints::SetConstraintState(int32 ConstraintIndex, const FPBDJointState& InState)
+	{
+		ConstraintStates[ConstraintIndex] = InState;
+	}
+
 	void FPBDJointConstraints::RemoveConstraint(int ConstraintIndex)
 	{
 		bJointsDirty = true;
diff --git a/Engine/Source/Runtime/Experimental/Chaos/Private/Chaos/Serialization/SerializationUtils.cpp b/Engine/Source/Runtime/Experimental/Chaos/Private/Chaos/Serialization/SerializationUtils.cpp
new file mode 100644
--- /dev/null
+++ b/Engine/Source/Runtime/Experimental/Chaos/Private/Chaos/Serialization/SerializationUtils.cpp
@@ -0,0 +1,98 @@
+// Copyright Epic Games, Inc. All Rights Reserved.
+
+#include "Chaos/Serialization/SerializationUtils.h"
+
+#include "Chaos/ConstraintHandle.h"
+#include "Chaos/ParticleHandle.h"
+#include "Chaos/PBDJointConstraints.h"
+#include "ChaosVisualDebugger/ChaosVDDataWrapperUtils.h"
+#include "DataWrappers/ChaosVDJointDataWrappers.h"
+#include "DataWrappers/ChaosVDParticleDataWrapper.h"
+#include "Serialization/MemoryReader.h"
+#include "Serialization/MemoryWriter.h"
+#include "HAL/Platform.h"
+
+namespace Chaos
+{
+	namespace Serialization::Private
+	{
+		static bool CanBulkSerialize(FProperty* Property)
+		{
+	#if PLATFORM_LITTLE_ENDIAN
+			// All numeric properties except TEnumAsByte
+			uint64 CastFlags = Property->GetClass()->GetCastFlags();
+			if (!!(CastFlags & CASTCLASS_FNumericProperty))
+			{
+				bool bEnumAsByte = (CastFlags & CASTCLASS_FByteProperty) != 0 && static_cast<FByteProperty*>(Property)->Enum;
+				return !bEnumAsByte;
+			}
+	#endif
+
+			return false;
+		}
+			
+		void SerializeProperty(FProperty* Property, void* PropertyData, FArchive& Ar)
+		{
+			if (FStructProperty* AsStruct = CastField<FStructProperty>(Property))
+			{
+				FastStructSerialize(AsStruct->Struct, PropertyData, Ar);
+			}
+			else if (FArrayProperty* AsArray = CastField<FArrayProperty>(Property))
+			{
+				if (CanBulkSerialize(AsArray->Inner))
+				{
+					FScriptArrayHelper ArrayHelper(AsArray, PropertyData);
+					int32 ElementCount = 0;		
+					if (Ar.IsLoading())
+					{
+						Ar << ElementCount;
+						ArrayHelper.EmptyAndAddUninitializedValues(ElementCount);
+					}
+					else
+					{
+						ElementCount = ArrayHelper.Num();
+						Ar << ElementCount;
+					}
+						
+					Ar.Serialize(ArrayHelper.GetRawPtr(), AsArray->Inner->GetElementSize() * ElementCount);
+				}
+				else
+				{
+					FScriptArrayHelper ArrayHelper(AsArray, PropertyData);
+					int32 ElementCount = 0;			
+					if (Ar.IsLoading())
+					{
+						Ar << ElementCount;
+						ArrayHelper.EmptyAndAddValues(ElementCount);
+					}
+					else
+					{
+						ElementCount = ArrayHelper.Num();
+						Ar << ElementCount;
+					}
+
+					for (int32 i = 0; i < ElementCount; ++i)
+					{
+						SerializeProperty(AsArray->Inner, ArrayHelper.GetRawPtr(i), Ar);
+					}
+				}
+			}
+			else
+			{
+				Ar.Serialize(PropertyData, Property->GetElementSize());	
+			}
+		}
+
+		void FastStructSerialize(UScriptStruct* Struct, void* SourceData, FArchive& Ar)
+		{
+			for (FProperty* Property = Struct->PropertyLink; Property != nullptr; Property = Property->PropertyLinkNext)
+			{
+				for (int32 Idx = 0; Idx < Property->ArrayDim; Idx++)
+				{
+					Private::SerializeProperty(Property, Property->ContainerPtrToValuePtr<void>(SourceData, Idx), Ar);
+				}
+			}
+		}
+	}
+}
+
diff --git a/Engine/Source/Runtime/Experimental/Chaos/Private/Chaos/Serialization/SolverSerializer.cpp b/Engine/Source/Runtime/Experimental/Chaos/Private/Chaos/Serialization/SolverSerializer.cpp
new file mode 100644
--- /dev/null
+++ b/Engine/Source/Runtime/Experimental/Chaos/Private/Chaos/Serialization/SolverSerializer.cpp
@@ -0,0 +1,182 @@
+// Copyright Epic Games, Inc. All Rights Reserved.
+
+#include "Chaos/Serialization/SolverSerializer.h"
+
+#include "PBDRigidsSolver.h"
+#include "Chaos/PBDJointConstraints.h"
+#include "Chaos/Character/CharacterGroundConstraintContainer.h"
+#include "Chaos/Serialization/SerializationUtils.h"
+#include "Chaos/Serialization/SerializedDataBuffer.h"
+#include "ChaosVisualDebugger/ChaosVDDataWrapperUtils.h"
+#include "DataWrappers/ChaosVDCharacterGroundConstraintDataWrappers.h"
+#include "DataWrappers/ChaosVDJointDataWrappers.h"
+
+namespace Chaos
+{
+	FSolverSerializer::FSolverSerializer(FPBDRigidsSolver* InSolver) : SolverInstance(InSolver)
+	{
+	}
+
+	void FSolverSerializer::SerializeToBuffer(FSerializedDataBuffer& OutSerializedData)
+	{
+		//TODO: Implement
+	}
+
+	void FSolverSerializer::PopulateFromSerializedBuffer(const FSerializedDataBuffer& InSerializedData)
+	{
+		//TODO: Implement
+	}
+
+	void FSolverSerializer::SerializeParticleStateToBuffer(FGeometryParticleHandle* InParticleHandle, FSerializedDataBuffer& OutSerializedData)
+	{
+		using namespace Chaos::Serialization::Private;
+
+		if (!ensure(InParticleHandle))
+		{
+			return;
+		}
+
+		// This allows us to serialize multiple particles in a single buffer, if the provided buffer was already used
+		constexpr bool bSetOffset = true;
+		constexpr bool bIsPersistent = false;
+
+		FMemoryWriter MemWriter(OutSerializedData.GetDataAsByteArrayRef(), bIsPersistent, bSetOffset);
+		MemWriter.SetShouldSkipUpdateCustomVersion(true);
+
+		// TODO: Handle Geometry and shape data serialization for the case that need it
+		FChaosVDParticleDataWrapper ParticleDataWrapper = FChaosVDDataWrapperUtils::BuildParticleDataWrapperFromParticle(InParticleHandle);
+
+		FastStructSerialize(MemWriter, &ParticleDataWrapper);
+	}
+
+	void FSolverSerializer::SerializeConstraintStateToBuffer(FConstraintHandle* ConstraintHandlePtr, FSerializedDataBuffer& OutSerializedData)
+	{
+		using namespace Chaos::Serialization::Private;
+
+		if (!ensure(ConstraintHandlePtr))
+		{
+			return;
+		}
+
+		// This allows us to serialize multiple particles in a single buffer, if the provided buffer was already used
+		constexpr bool bSetOffset = true;
+		constexpr bool bIsPersistent = false;
+
+		FMemoryWriter MemWriter(OutSerializedData.GetDataAsByteArrayRef(), bIsPersistent, bSetOffset);
+		MemWriter.SetShouldSkipUpdateCustomVersion(true);
+			
+		if (ConstraintHandlePtr->GetType().IsA(FPBDJointConstraintHandle::StaticType()))
+		{
+			const FPBDJointConstraintHandle* AsJointHandle = static_cast<FPBDJointConstraintHandle*>(ConstraintHandlePtr);	
+			FChaosVDJointConstraint JointData = FChaosVDDataWrapperUtils::BuildJointDataWrapper(AsJointHandle);
+
+			FastStructSerialize(MemWriter, &JointData);
+		}
+		else if (ConstraintHandlePtr->GetType().IsA(FCharacterGroundConstraintHandle::StaticType()))
+		{
+			const FCharacterGroundConstraintHandle* AsCharacterGroundHandle = static_cast<FCharacterGroundConstraintHandle*>(ConstraintHandlePtr);
+			FChaosVDCharacterGroundConstraint CharacterGroundConstraintData = FChaosVDDataWrapperUtils::BuildCharacterGroundConstraintDataWrapper(AsCharacterGroundHandle);
+			
+			FastStructSerialize(MemWriter, &CharacterGroundConstraintData);
+		}
+	}
+
+	void FSolverSerializer::ApplySerializedStateToParticle(FGeometryParticleHandle* InParticleHandle, FSerializedDataBuffer& InSerializedData)
+	{
+		using namespace Chaos::Serialization::Private;
+		if (!ensure(!InSerializedData.GetDataAsByteArrayRef().IsEmpty()))
+		{
+			return;
+		}
+
+		FMemoryReader MemReader(InSerializedData.GetDataAsByteArrayRef());
+		MemReader.SetShouldSkipUpdateCustomVersion(true);
+		
+		FChaosVDParticleDataWrapper ParticleDataWrapper;
+
+		FastStructSerialize(MemReader, &ParticleDataWrapper);
+
+		ApplySerializedStateToParticle(InParticleHandle, ParticleDataWrapper);
+	}
+
+	void FSolverSerializer::ApplySerializedStateToParticle(FGeometryParticleHandle* InParticleHandle, const FChaosVDParticleDataWrapper& InParticleState)
+	{
+		if (!ensure(InParticleHandle->Type == static_cast<EParticleType>(InParticleState.Type)))
+		{
+			return;
+		}
+		
+		FPBDRigidsEvolutionGBF& Evolution = *SolverInstance->GetEvolution();
+
+		constexpr bool bIsTeleport = true;
+		constexpr bool bWakeUp = false;
+		Evolution.SetParticleTransform(InParticleHandle, InParticleState.ParticlePositionRotation.MX, InParticleState.ParticlePositionRotation.MR, bIsTeleport, bWakeUp);
+
+		if (InParticleState.Type == EChaosVDParticleType::Static)
+		{
+			const FVec3 BoundsExpansion = FVec3(0);
+			InParticleHandle->UpdateWorldSpaceState(InParticleHandle->GetTransformXR(), BoundsExpansion);
+		}
+
+		Evolution.SetParticleVelocities(InParticleHandle, InParticleState.ParticleVelocities.MV, InParticleState.ParticleVelocities.MW);
+
+		if (TPBDRigidParticleHandleImp<FReal, 3, true>* RigidParticle = InParticleHandle->CastToRigidParticle())
+		{
+			InParticleState.ParticleDynamics.CopyTo(*RigidParticle);
+			InParticleState.ParticleMassProps.CopyTo(*RigidParticle);
+			InParticleState.ParticleDynamicsMisc.CopyTo<TPBDRigidParticleHandleImp<FReal, 3, true>, FRigidParticleControlFlags, EObjectStateType, ESleepType>(*RigidParticle);
+
+			// Make sure the particle ends up in the correct SoA view
+			Evolution.SetParticleObjectState(RigidParticle, RigidParticle->ObjectState());
+			Evolution.SetParticleSleepType(RigidParticle, RigidParticle->SleepType());
+		}
+
+		if (TPBDRigidClusteredParticleHandleImp<FReal, 3, true>* ClusteredParticle = InParticleHandle->CastToClustered())
+		{
+			InParticleState.ParticleCluster.CopyTo(*ClusteredParticle);
+		}
+
+		if (InParticleState.ParticleDynamicsMisc.bDisabled)
+		{
+			Evolution.DisableParticle(InParticleHandle);
+		}
+		else
+		{
+			Evolution.EnableParticle(InParticleHandle);
+		}
+	}
+
+	void FSolverSerializer::ApplySerializedStateToJointConstraint(Chaos::FPBDJointConstraintHandle* ConstraintHandlePtr, FSerializedDataBuffer& InSerializedData)
+	{
+		using namespace Chaos::Serialization::Private;
+
+		if (!ConstraintHandlePtr)
+		{
+			return;
+		}
+
+		FMemoryReader Reader(InSerializedData.GetDataAsByteArrayRef());
+		Reader.SetShouldSkipUpdateCustomVersion(true);
+		FChaosVDJointConstraint JointData;
+		FastStructSerialize(Reader, &JointData);
+	
+		FChaosVDDataWrapperUtils::ApplyJointDataWrapperToHandle(ConstraintHandlePtr, JointData);	
+	}
+
+	void FSolverSerializer::PushPendingSerializedStateForProxy(IPhysicsProxyBase* Proxy, FSerializedDataBufferPtr&& InState)
+	{
+		PendingMigratedPhysicsStateByProxy.Emplace(Proxy, MoveTemp(InState));
+	}
+
+	FSerializedDataBufferPtr FSolverSerializer::PopPendingSerializedStateForProxy(IPhysicsProxyBase* Proxy)
+	{
+		if (PendingMigratedPhysicsStateByProxy.IsEmpty())
+		{
+			return nullptr;
+		}
+
+		FSerializedDataBufferPtr FoundBufferRef = nullptr;
+		PendingMigratedPhysicsStateByProxy.RemoveAndCopyValue(Proxy, FoundBufferRef);
+		return MoveTemp(FoundBufferRef);
+	}
+}
\ No newline at end of file
diff --git a/Engine/Source/Runtime/Experimental/Chaos/Private/ChaosVisualDebugger/ChaosVDDataWrapperUtils.cpp b/Engine/Source/Runtime/Experimental/Chaos/Private/ChaosVisualDebugger/ChaosVDDataWrapperUtils.cpp
--- a/Engine/Source/Runtime/Experimental/Chaos/Private/ChaosVisualDebugger/ChaosVDDataWrapperUtils.cpp
+++ b/Engine/Source/Runtime/Experimental/Chaos/Private/ChaosVisualDebugger/ChaosVDDataWrapperUtils.cpp
@@ -4,8 +4,6 @@
 
 #include "Chaos/SpatialAccelerationCollection.h"
 
-#if WITH_CHAOS_VISUAL_DEBUGGER
-
 #include "Chaos/PBDJointConstraints.h"
 #include "DataWrappers/ChaosVDCharacterGroundConstraintDataWrappers.h"
 #include "DataWrappers/ChaosVDJointDataWrappers.h"
@@ -87,7 +85,7 @@ void FChaosVDDataWrapperUtils::CopyCollisionMaterialToDataWrapper(const Chaos::F
 	OutCopyTo.InvInertiaScale1 = InCopyFrom.InvInertiaScale1;
 }
 
-FChaosVDParticleDataWrapper FChaosVDDataWrapperUtils::BuildParticleDataWrapperFromParticle(const Chaos::FGeometryParticleHandle* ParticleHandlePtr, const TSharedRef<Chaos::VisualDebugger::FChaosVDSerializableNameTable>& InNameTableInstance)
+FChaosVDParticleDataWrapper FChaosVDDataWrapperUtils::BuildParticleDataWrapperFromParticle(const Chaos::FGeometryParticleHandle* ParticleHandlePtr)
 {
 	check(ParticleHandlePtr);
 
@@ -380,6 +378,118 @@ FChaosVDJointConstraint FChaosVDDataWrapperUtils::BuildJointDataWrapper(const Ch
 	return MoveTemp(WrappedJointData);
 }
 
+void FChaosVDDataWrapperUtils::ApplyJointDataWrapperToHandle(Chaos::FPBDJointConstraintHandle* ConstaintHanlde, const FChaosVDJointConstraint& WrappedJointData)
+{
+	if (!ensure(ConstaintHanlde))
+	{
+		return;
+	}
+
+	using namespace Chaos;
+
+	FPBDJointState NewState;
+	FPBDJointSettings NewSettings;
+
+	NewState.bBroken = WrappedJointData.JointState.bBroken;
+	NewState.bBreaking = WrappedJointData.JointState.bBreaking;
+	NewState.bDisabled = WrappedJointData.JointState.bDisabled;
+	NewState.bDriveTargetChanged = WrappedJointData.JointState.bDriveTargetChanged;
+
+	NewState.bEnabledDuringResim = WrappedJointData.JointState.bEnabledDuringResim;
+	NewState.AngularImpulse = WrappedJointData.JointState.AngularImpulse;
+	NewState.LinearImpulse = WrappedJointData.JointState.LinearImpulse;
+
+	NewState.LinearImpulse = WrappedJointData.JointState.LinearImpulse;
+
+	CVD_COPY_FIELD_TO_WRAPPER_WITH_CAST(WrappedJointData.JointState, NewState, ResimType, Chaos::EResimType);
+	CVD_COPY_FIELD_TO_WRAPPER_WITH_CAST(WrappedJointData.JointState, NewState, SyncState, Chaos::ESyncState);
+
+	CVD_COPY_FIELD_TO_WRAPPER(WrappedJointData.JointSettings, NewSettings, Stiffness);
+	CVD_COPY_FIELD_TO_WRAPPER(WrappedJointData.JointSettings, NewSettings, LinearProjection);
+	CVD_COPY_FIELD_TO_WRAPPER(WrappedJointData.JointSettings, NewSettings, AngularProjection);
+	CVD_COPY_FIELD_TO_WRAPPER(WrappedJointData.JointSettings, NewSettings, AngularProjection);
+	CVD_COPY_FIELD_TO_WRAPPER(WrappedJointData.JointSettings, NewSettings, TeleportDistance);
+	CVD_COPY_FIELD_TO_WRAPPER(WrappedJointData.JointSettings, NewSettings, TeleportAngle);
+	CVD_COPY_FIELD_TO_WRAPPER(WrappedJointData.JointSettings, NewSettings, ParentInvMassScale);
+	CVD_COPY_FIELD_TO_WRAPPER(WrappedJointData.JointSettings, NewSettings, bCollisionEnabled);
+	CVD_COPY_FIELD_TO_WRAPPER(WrappedJointData.JointSettings, NewSettings, bMassConditioningEnabled);
+	CVD_COPY_FIELD_TO_WRAPPER(WrappedJointData.JointSettings, NewSettings, bUseLinearSolver);
+	CVD_COPY_FIELD_TO_WRAPPER(WrappedJointData.JointSettings, NewSettings, bSoftLinearLimitsEnabled);
+	CVD_COPY_FIELD_TO_WRAPPER(WrappedJointData.JointSettings, NewSettings, bSoftTwistLimitsEnabled);
+	CVD_COPY_FIELD_TO_WRAPPER(WrappedJointData.JointSettings, NewSettings, bAngularSLerpPositionDriveEnabled);
+	CVD_COPY_FIELD_TO_WRAPPER(WrappedJointData.JointSettings, NewSettings, bAngularSLerpVelocityDriveEnabled);
+	CVD_COPY_FIELD_TO_WRAPPER(WrappedJointData.JointSettings, NewSettings, bAngularTwistPositionDriveEnabled);
+	CVD_COPY_FIELD_TO_WRAPPER(WrappedJointData.JointSettings, NewSettings, bAngularTwistVelocityDriveEnabled);
+	CVD_COPY_FIELD_TO_WRAPPER(WrappedJointData.JointSettings, NewSettings, bAngularSwingPositionDriveEnabled);
+	CVD_COPY_FIELD_TO_WRAPPER(WrappedJointData.JointSettings, NewSettings, bAngularSwingVelocityDriveEnabled);
+
+	NewSettings.LinearMotionTypes[0] = static_cast<EJointMotionType>(WrappedJointData.JointSettings.LinearMotionTypes[0]);
+	NewSettings.LinearMotionTypes[1] = static_cast<EJointMotionType>(WrappedJointData.JointSettings.LinearMotionTypes[1]);
+	NewSettings.LinearMotionTypes[2] = static_cast<EJointMotionType>(WrappedJointData.JointSettings.LinearMotionTypes[2]);
+
+	CVD_COPY_FIELD_TO_WRAPPER(WrappedJointData.JointSettings, NewSettings, LinearLimit);
+
+	NewSettings.AngularMotionTypes[0] = static_cast<EJointMotionType>(WrappedJointData.JointSettings.AngularMotionTypes[0]);
+	NewSettings.AngularMotionTypes[1] = static_cast<EJointMotionType>(WrappedJointData.JointSettings.AngularMotionTypes[1]);
+	NewSettings.AngularMotionTypes[2] = static_cast<EJointMotionType>(WrappedJointData.JointSettings.AngularMotionTypes[2]);
+
+	CVD_COPY_FIELD_TO_WRAPPER(WrappedJointData.JointSettings, NewSettings, AngularLimits);
+	CVD_COPY_FIELD_TO_WRAPPER_WITH_CAST(WrappedJointData.JointSettings, NewSettings, LinearSoftForceMode, EJointForceMode);
+	CVD_COPY_FIELD_TO_WRAPPER_WITH_CAST(WrappedJointData.JointSettings, NewSettings, AngularSoftForceMode, EJointForceMode);
+	CVD_COPY_FIELD_TO_WRAPPER(WrappedJointData.JointSettings, NewSettings, SoftLinearStiffness);
+	CVD_COPY_FIELD_TO_WRAPPER(WrappedJointData.JointSettings, NewSettings, SoftLinearDamping);
+	CVD_COPY_FIELD_TO_WRAPPER(WrappedJointData.JointSettings, NewSettings, SoftTwistStiffness);
+	CVD_COPY_FIELD_TO_WRAPPER(WrappedJointData.JointSettings, NewSettings, SoftTwistDamping);
+	CVD_COPY_FIELD_TO_WRAPPER(WrappedJointData.JointSettings, NewSettings, SoftSwingStiffness);
+	CVD_COPY_FIELD_TO_WRAPPER(WrappedJointData.JointSettings, NewSettings, SoftSwingDamping);
+	CVD_COPY_FIELD_TO_WRAPPER(WrappedJointData.JointSettings, NewSettings, LinearRestitution);
+	CVD_COPY_FIELD_TO_WRAPPER(WrappedJointData.JointSettings, NewSettings, TwistRestitution);
+	CVD_COPY_FIELD_TO_WRAPPER(WrappedJointData.JointSettings, NewSettings, TwistRestitution);
+	CVD_COPY_FIELD_TO_WRAPPER(WrappedJointData.JointSettings, NewSettings, SwingRestitution);
+	CVD_COPY_FIELD_TO_WRAPPER(WrappedJointData.JointSettings, NewSettings, LinearContactDistance);
+	CVD_COPY_FIELD_TO_WRAPPER(WrappedJointData.JointSettings, NewSettings, TwistContactDistance);
+	CVD_COPY_FIELD_TO_WRAPPER(WrappedJointData.JointSettings, NewSettings, SwingContactDistance);
+
+	CVD_COPY_VECTOR_FIELD_TO_WRAPPER(WrappedJointData.JointSettings, NewSettings, LinearDrivePositionTarget);
+	CVD_COPY_VECTOR_FIELD_TO_WRAPPER(WrappedJointData.JointSettings, NewSettings, LinearDriveVelocityTarget);
+	
+	NewSettings.bLinearPositionDriveEnabled[0] = WrappedJointData.JointSettings.bLinearPositionDriveEnabled0;
+	NewSettings.bLinearPositionDriveEnabled[1] = WrappedJointData.JointSettings.bLinearPositionDriveEnabled1;
+	NewSettings.bLinearPositionDriveEnabled[2] = WrappedJointData.JointSettings.bLinearPositionDriveEnabled2;
+
+	NewSettings.bLinearVelocityDriveEnabled[0] = WrappedJointData.JointSettings.bLinearVelocityDriveEnabled0;
+	NewSettings.bLinearVelocityDriveEnabled[1] = WrappedJointData.JointSettings.bLinearVelocityDriveEnabled1;
+	NewSettings.bLinearVelocityDriveEnabled[2] = WrappedJointData.JointSettings.bLinearVelocityDriveEnabled2;
+
+	CVD_COPY_FIELD_TO_WRAPPER_WITH_CAST(WrappedJointData.JointSettings, NewSettings, LinearDriveForceMode, EJointForceMode);
+
+	CVD_COPY_VECTOR_FIELD_TO_WRAPPER(WrappedJointData.JointSettings, NewSettings, LinearDriveStiffness);
+	CVD_COPY_VECTOR_FIELD_TO_WRAPPER(WrappedJointData.JointSettings, NewSettings, LinearDriveDamping);
+	CVD_COPY_VECTOR_FIELD_TO_WRAPPER(WrappedJointData.JointSettings, NewSettings, LinearDriveMaxForce);
+	CVD_COPY_FIELD_TO_WRAPPER(WrappedJointData.JointSettings, NewSettings, AngularDrivePositionTarget);
+	CVD_COPY_VECTOR_FIELD_TO_WRAPPER(WrappedJointData.JointSettings, NewSettings, AngularDriveVelocityTarget);
+	CVD_COPY_FIELD_TO_WRAPPER_WITH_CAST(WrappedJointData.JointSettings, NewSettings, AngularDriveForceMode, EJointForceMode);
+	CVD_COPY_VECTOR_FIELD_TO_WRAPPER(WrappedJointData.JointSettings, NewSettings, AngularDriveStiffness);
+	CVD_COPY_VECTOR_FIELD_TO_WRAPPER(WrappedJointData.JointSettings, NewSettings, AngularDriveDamping);
+	CVD_COPY_VECTOR_FIELD_TO_WRAPPER(WrappedJointData.JointSettings, NewSettings, AngularDriveMaxTorque);
+
+	CVD_COPY_FIELD_TO_WRAPPER(WrappedJointData.JointSettings, NewSettings, LinearBreakForce);
+	CVD_COPY_FIELD_TO_WRAPPER(WrappedJointData.JointSettings, NewSettings, LinearPlasticityLimit);
+
+	CVD_COPY_FIELD_TO_WRAPPER_WITH_CAST(WrappedJointData.JointSettings, NewSettings, LinearPlasticityType, EPlasticityType);
+	
+	CVD_COPY_FIELD_TO_WRAPPER(WrappedJointData.JointSettings, NewSettings, LinearPlasticityInitialDistanceSquared);
+	CVD_COPY_FIELD_TO_WRAPPER(WrappedJointData.JointSettings, NewSettings, AngularBreakTorque);
+	CVD_COPY_FIELD_TO_WRAPPER(WrappedJointData.JointSettings, NewSettings, AngularPlasticityLimit);
+	CVD_COPY_FIELD_TO_WRAPPER(WrappedJointData.JointSettings, NewSettings, ContactTransferScale);
+	
+	NewSettings.ConnectorTransforms[0] = WrappedJointData.JointSettings.ConnectorTransforms[0];
+	NewSettings.ConnectorTransforms[1] = WrappedJointData.JointSettings.ConnectorTransforms[1];
+
+	ConstaintHanlde->SetSettings(MoveTemp(NewSettings));
+	ConstaintHanlde->SetJointState(MoveTemp(NewState));
+}
+
 FChaosVDCharacterGroundConstraint FChaosVDDataWrapperUtils::BuildCharacterGroundConstraintDataWrapper(const Chaos::FCharacterGroundConstraintHandle* ConstraintHandle)
 {
 	FChaosVDCharacterGroundConstraint WrappedConstraintData;
@@ -576,5 +686,3 @@ void FChaosVDDataWrapperUtils::AddTreeLeaves(const TConstArrayView<Chaos::TBound
 		InOutAABBTreeWrapper.BoundingVolumeLeafs.Emplace(MoveTemp(CVDBoundingVolume));
 	}
 }
-
-#endif //WITH_CHAOS_VISUAL_DEBUGGER
diff --git a/Engine/Source/Runtime/Experimental/Chaos/Private/ChaosVisualDebugger/ChaosVisualDebuggerTrace.cpp b/Engine/Source/Runtime/Experimental/Chaos/Private/ChaosVisualDebugger/ChaosVisualDebuggerTrace.cpp
--- a/Engine/Source/Runtime/Experimental/Chaos/Private/ChaosVisualDebugger/ChaosVisualDebuggerTrace.cpp
+++ b/Engine/Source/Runtime/Experimental/Chaos/Private/ChaosVisualDebugger/ChaosVisualDebuggerTrace.cpp
@@ -159,7 +159,7 @@ void FChaosVisualDebuggerTrace::TraceParticle(Chaos::FGeometryParticleHandle* Pa
 	TraceImplicitObject({ GeometryHash, ParticleHandle->GetGeometry() });
 
 	{
-		FChaosVDParticleDataWrapper ParticleDataWrapper = FChaosVDDataWrapperUtils::BuildParticleDataWrapperFromParticle(ParticleHandle, CVDNameTable);
+		FChaosVDParticleDataWrapper ParticleDataWrapper = FChaosVDDataWrapperUtils::BuildParticleDataWrapperFromParticle(ParticleHandle);
 		ParticleDataWrapper.GeometryHash = GeometryHash;
 		ParticleDataWrapper.SolverID = ContextData.Id;
 		
diff --git a/Engine/Source/Runtime/Experimental/Chaos/Private/PBDRigidsSolver.cpp b/Engine/Source/Runtime/Experimental/Chaos/Private/PBDRigidsSolver.cpp
--- a/Engine/Source/Runtime/Experimental/Chaos/Private/PBDRigidsSolver.cpp
+++ b/Engine/Source/Runtime/Experimental/Chaos/Private/PBDRigidsSolver.cpp
@@ -38,6 +38,7 @@
 #include "ProfilingDebugging/CountersTrace.h"
 #include "ProfilingDebugging/CsvProfiler.h"
 #include "ChaosLog.h"
+#include "Physics/Experimental/PhysScene_Chaos.h"
 
 //PRAGMA_DISABLE_OPTIMIZATION
 
@@ -664,6 +665,7 @@ namespace Chaos
 		, MCurrentLock(new FCriticalSection())
 
 		, PerSolverField(nullptr)
+		, Serializer(this)
 	{
 		UE_LOG(LogPBDRigidsSolver, Verbose, TEXT("PBDRigidsSolver::PBDRigidsSolver()"));
 		Reset();
@@ -812,6 +814,11 @@ namespace Chaos
 		}
 	}
 
+	FSolverSerializer& FPBDRigidsSolver::GetSerializer()
+	{
+		return Serializer;
+	}
+
 	void FPBDRigidsSolver::UnregisterObject(FSingleParticlePhysicsProxy* Proxy)
 	{
 		UE_CHAOS_ASYNC_INITBODY_WRITESCOPELOCK(GetExternalDataLock_External());
@@ -1069,6 +1076,9 @@ namespace Chaos
 			});
 	}
 
+
+
+
 	void FPBDRigidsSolver::RegisterObject(FSuspensionConstraint* GTConstraint)
 	{
 		LLM_SCOPE(ELLMTag::ChaosConstraint);
@@ -1682,6 +1692,17 @@ namespace Chaos
 						}
 				
 						JointProxy->PushStateOnPhysicsThread(this, *Manager, DataIdx, Dirty.PropertyData);
+
+#if UE_WITH_REMOTE_OBJECT_HANDLE
+						if (bIsNew)
+						{
+							
+							if (FSerializedDataBufferPtr StateBuffer = Serializer.PopPendingSerializedStateForProxy(JointProxy))
+							{
+								Serializer.ApplySerializedStateToJointConstraint(JointProxy->GetHandle(), *StateBuffer);
+							}
+						}
+#endif
 						break;
 					}
 
diff --git a/Engine/Source/Runtime/Experimental/Chaos/Private/PhysicsProxy/SingleParticlePhysicsProxy.cpp b/Engine/Source/Runtime/Experimental/Chaos/Private/PhysicsProxy/SingleParticlePhysicsProxy.cpp
--- a/Engine/Source/Runtime/Experimental/Chaos/Private/PhysicsProxy/SingleParticlePhysicsProxy.cpp
+++ b/Engine/Source/Runtime/Experimental/Chaos/Private/PhysicsProxy/SingleParticlePhysicsProxy.cpp
@@ -17,6 +17,7 @@
 #include "Chaos/PullPhysicsDataImp.h"
 #include "RewindData.h"
 #include "Chaos/DebugDrawQueue.h"
+#include "Physics/Experimental/PhysScene_Chaos.h"
 
 namespace Chaos
 {
@@ -68,81 +69,99 @@ void PushToPhysicsStateImp(const Chaos::FDirtyPropertiesManager& Manager, Chaos:
 	auto RigidHandle = bHasDynamicData ? static_cast<Chaos::FPBDRigidParticleHandle*>(Handle) : nullptr;
 	const FDirtyChaosProperties& ParticleData = Dirty.PropertyData;
 	FPBDRigidsEvolutionGBF& Evolution = *Solver.GetEvolution();
+	
+	const FParticleNonFrequentData* NewNonFrequentData = ParticleData.FindNonFrequentData(Manager, DataIdx);
 
-	if (bResimInitialized)	//todo: assumes particles are always initialized as enabled. This is not true in future versions of code, so check PushData
-	{
-		Evolution.EnableParticle(Handle);
-	}
-	// move the copied game thread data into the handle
+#if UE_WITH_REMOTE_OBJECT_HANDLE
+	if (FSerializedDataBufferPtr PendingStateData = Dirty.Proxy && !Dirty.Proxy->IsInitialized() ? Solver.GetSerializer().PopPendingSerializedStateForProxy(Dirty.Proxy) : nullptr)
 	{
-		auto NewXR = ParticleData.FindXR(Manager, DataIdx);
-		auto NewNonFrequentData = ParticleData.FindNonFrequentData(Manager, DataIdx);
-		auto NewDynamicMiscData = ParticleData.FindDynamicMisc(Manager, DataIdx);
-
-		if(NewXR)
+		// Currently we don't serialize collision geometry and shape data, so we need to apply it first
+		// This might change in the future
+		if (NewNonFrequentData)
 		{
-			// @todo(chaos): we need to know if this is a teleport or not and pass that on. See UE-165746
-			// For now we just set bIsTeleport to true since that's the no-impact option for SetParticleTransform
-			// (there would be issues if we report a non-teleport move for an initial-position a long way from the origin)
-			const bool bIsTeleport = true;
-			bool bWakeUp = true;
-			if (NewDynamicMiscData && NewDynamicMiscData->ObjectState() == EObjectStateType::Sleeping)
-			{
-				bWakeUp = false;
-			}
-
-			Evolution.SetParticleTransform(Handle, NewXR->X(), NewXR->R(), bIsTeleport, bWakeUp);
-		}
-
-		if(NewNonFrequentData)
-		{
-			// Geometry may have changed, we need to remove the particle and its collisions from the graph
-			Evolution.InvalidateParticle(Handle);
-
 			Handle->SetNonFrequentData(*NewNonFrequentData);
 		}
 
-		auto NewVelocities = bHasKinematicData ? ParticleData.FindVelocities(Manager, DataIdx) : nullptr;
-		if(NewVelocities)
+		Solver.GetSerializer().ApplySerializedStateToParticle(Handle, *PendingStateData);
+	}
+	else
+#endif
+	{
+		if (bResimInitialized)	//todo: assumes particles are always initialized as enabled. This is not true in future versions of code, so check PushData
 		{
-			Evolution.SetParticleVelocities(KinematicHandle, NewVelocities->V(), NewVelocities->W());
+			Evolution.EnableParticle(Handle);
 		}
-
-		auto NewKinematicTargetGT = bHasKinematicData ? ParticleData.FindKinematicTarget(Manager, DataIdx) : nullptr;
-		if (NewKinematicTargetGT)
+		// move the copied game thread data into the handle
 		{
-			Evolution.SetParticleKinematicTarget(KinematicHandle, *NewKinematicTargetGT);
-		}
+			auto NewXR = ParticleData.FindXR(Manager, DataIdx);
+			auto NewDynamicMiscData = ParticleData.FindDynamicMisc(Manager, DataIdx);
 
-		if(NewXR || NewNonFrequentData || NewVelocities || NewKinematicTargetGT)
-		{
-			// Update world-space cached state like the bounds
-			// @todo(chaos): do we need to do this here? It should be done in Integrate and ApplyKinematicTarget so only really Statics need this...
-			const bool bHasKinematicTarget = (NewKinematicTargetGT != nullptr) && (NewKinematicTargetGT->GetMode() == EKinematicTargetMode::Position);
-			const FRigidTransform3 WorldTransform = !bHasKinematicTarget ? FRigidTransform3(Handle->GetX(), Handle->GetR()) : NewKinematicTargetGT->GetTransform();
-			Handle->UpdateWorldSpaceState(WorldTransform, FVec3(0));
+			if(NewXR)
+			{
+				// @todo(chaos): we need to know if this is a teleport or not and pass that on. See UE-165746
+				// For now we just set bIsTeleport to true since that's the no-impact option for SetParticleTransform
+				// (there would be issues if we report a non-teleport move for an initial-position a long way from the origin)
+				const bool bIsTeleport = true;
+				bool bWakeUp = true;
+				if (NewDynamicMiscData && NewDynamicMiscData->ObjectState() == EObjectStateType::Sleeping)
+				{
+					bWakeUp = false;
+				}
 
-			Evolution.DirtyParticle(*Handle);
-		}
+				Evolution.SetParticleTransform(Handle, NewXR->X(), NewXR->R(), bIsTeleport, bWakeUp);
+			}
 
-		if(bHasDynamicData)
-		{
-			if(auto NewData = ParticleData.FindMassProps(Manager,DataIdx))
+			if(NewNonFrequentData)
 			{
-				RigidHandle->SetMassProps(*NewData);
+				// Geometry may have changed, we need to remove the particle and its collisions from the graph
+				Evolution.InvalidateParticle(Handle);
+
+				Handle->SetNonFrequentData(*NewNonFrequentData);
 			}
 
-			if(auto NewData = ParticleData.FindDynamics(Manager, DataIdx))
+			auto NewVelocities = bHasKinematicData ? ParticleData.FindVelocities(Manager, DataIdx) : nullptr;
+			if(NewVelocities)
+			{
+				Evolution.SetParticleVelocities(KinematicHandle, NewVelocities->V(), NewVelocities->W());
+			}
+			
+			auto NewKinematicTargetGT = bHasKinematicData ? ParticleData.FindKinematicTarget(Manager, DataIdx) : nullptr;
+			if (NewKinematicTargetGT)
 			{
-				RigidHandle->SetDynamics(*NewData);
-				Evolution.ResetVSmoothFromForces(*RigidHandle);
+				Evolution.SetParticleKinematicTarget(KinematicHandle, *NewKinematicTargetGT);
 			}
 
-			if(NewDynamicMiscData)
+			if(NewXR || NewNonFrequentData || NewVelocities || NewKinematicTargetGT)
 			{
-				Solver.SetParticleDynamicMisc(RigidHandle, *NewDynamicMiscData);
+				// Update world-space cached state like the bounds
+				// @todo(chaos): do we need to do this here? It should be done in Integrate and ApplyKinematicTarget so only really Statics need this...
+				const bool bHasKinematicTarget = (NewKinematicTargetGT != nullptr) && (NewKinematicTargetGT->GetMode() == EKinematicTargetMode::Position);
+				const FRigidTransform3 WorldTransform = !bHasKinematicTarget ? FRigidTransform3(Handle->GetX(), Handle->GetR()) : NewKinematicTargetGT->GetTransform();
+				Handle->UpdateWorldSpaceState(WorldTransform, FVec3(0));
+
+				Evolution.DirtyParticle(*Handle);
+			}
+
+			if(bHasDynamicData)
+			{
+				if(auto NewData = ParticleData.FindMassProps(Manager,DataIdx))
+				{
+					RigidHandle->SetMassProps(*NewData);
+				}
+
+				if(auto NewData = ParticleData.FindDynamics(Manager, DataIdx))
+				{
+					RigidHandle->SetDynamics(*NewData);
+					Evolution.ResetVSmoothFromForces(*RigidHandle);
+				}
+
+				if(NewDynamicMiscData )
+				{
+					Solver.SetParticleDynamicMisc(RigidHandle, *NewDynamicMiscData);
+				}
 			}
 		}
+	}
 
 		//shape properties
 		bool bUpdateCollisionData = false;
@@ -202,7 +221,6 @@ void PushToPhysicsStateImp(const Chaos::FDirtyPropertiesManager& Manager, Chaos:
 				Evolution.RemoveParticleFromAccelerationStructure(*Handle);
 			}
 		}
-	}
 }
 
 //
@@ -514,4 +532,4 @@ void FSingleParticlePhysicsProxy::ClearEvents()
 		Rigid->ClearEvents();
 	}
 }
-}
\ No newline at end of file
+}
diff --git a/Engine/Source/Runtime/Experimental/Chaos/Public/Chaos/PBDJointConstraints.h b/Engine/Source/Runtime/Experimental/Chaos/Public/Chaos/PBDJointConstraints.h
--- a/Engine/Source/Runtime/Experimental/Chaos/Public/Chaos/PBDJointConstraints.h
+++ b/Engine/Source/Runtime/Experimental/Chaos/Public/Chaos/PBDJointConstraints.h
@@ -14,7 +14,8 @@
 
 namespace Chaos
 {
-	namespace Private
+
+namespace Private
 	{
 		template <typename JointSolverType>
 		class TPBDJointContainerSolver;
@@ -36,6 +37,8 @@ namespace Chaos
 
 		CHAOS_API void CalculateConstraintSpace(FVec3& OutXa, FMatrix33& OutRa, FVec3& OutXb, FMatrix33& OutRb) const;
 
+		
+
 		CHAOS_API bool IsConstraintEnabled() const;
 		CHAOS_API bool IsConstraintBroken() const;
 		CHAOS_API bool IsConstraintBreaking() const;
@@ -51,6 +54,8 @@ namespace Chaos
 
 		CHAOS_API const FPBDJointSettings& GetSettings() const;
 		const FPBDJointSettings& GetJointSettings() const { return GetSettings(); }	//needed for property macros
+	
+		void SetJointState(const class FPBDJointState& NewState);
 
 		// Note that this is the "regular" way to change any settings, but note that if it results
 		// in a change of state/mode there may be some overhead. For simple/numerical changes, it
@@ -213,6 +218,8 @@ namespace Chaos
 		CHAOS_API FConstraintContainerHandle* AddConstraint(const FParticlePair& InConstrainedParticles, const FTransformPair& ConnectorTransforms);
 		CHAOS_API FConstraintContainerHandle* AddConstraint(const FParticlePair& InConstrainedParticles, const FPBDJointSettings& InConstraintSettings);
 
+		void SetConstraintState(int32 ConstraintIndex, const FPBDJointState& InState);
+
 		/**
 		 * Remove the specified constraint.
 		 */
diff --git a/Engine/Source/Runtime/Experimental/Chaos/Public/Chaos/Serialization/SerializationUtils.h b/Engine/Source/Runtime/Experimental/Chaos/Public/Chaos/Serialization/SerializationUtils.h
new file mode 100644
--- /dev/null
+++ b/Engine/Source/Runtime/Experimental/Chaos/Public/Chaos/Serialization/SerializationUtils.h
@@ -0,0 +1,26 @@
+// Copyright Epic Games, Inc. All Rights Reserved.
+
+#pragma once
+#include "Chaos/ParticleHandleFwd.h"
+#include "Containers/Array.h"
+
+namespace Chaos
+{
+	class FConstraintHandle;
+}
+
+namespace Chaos
+{
+	namespace Serialization::Private
+	{
+		/** Lightweight unversioned ustruct serializer */
+		UE_INTERNAL CHAOS_API void FastStructSerialize(UScriptStruct* Struct, void* SourceData, FArchive& Ar);
+
+		/** Lightweight unversioned ustruct serializer */
+		template<typename StructType>
+		UE_INTERNAL void FastStructSerialize(FArchive& Ar, StructType* Data)
+		{
+			FastStructSerialize(StructType::StaticStruct(), reinterpret_cast<void*>(Data), Ar);
+		}
+	}
+}
diff --git a/Engine/Source/Runtime/Experimental/Chaos/Public/Chaos/Serialization/SerializedDataBuffer.h b/Engine/Source/Runtime/Experimental/Chaos/Public/Chaos/Serialization/SerializedDataBuffer.h
new file mode 100644
--- /dev/null
+++ b/Engine/Source/Runtime/Experimental/Chaos/Public/Chaos/Serialization/SerializedDataBuffer.h
@@ -0,0 +1,40 @@
+// Copyright Epic Games, Inc. All Rights Reserved.
+
+#pragma once
+#include "Containers/Array.h"
+#include "Serialization/Archive.h"
+#include "Templates/SharedPointer.h"
+
+namespace Chaos
+{
+	/** Wrapper around a data buffer used for serialization of chaos data */
+	struct FSerializedDataBuffer
+	{
+		FSerializedDataBuffer(const TArray<uint8>& InByteArray) : DataBuffer(InByteArray)
+		{	
+		}
+
+		FSerializedDataBuffer()
+		{
+		}
+
+		TArray<uint8>& GetDataAsByteArrayRef()
+		{
+			return DataBuffer;
+		}
+
+		void Serialize(FArchive& Ar)
+		{
+			Ar << DataBuffer;
+		}
+
+	private:
+		TArray<uint8> DataBuffer;
+	};
+
+	inline FArchive& operator<<(FArchive& Ar, FSerializedDataBuffer& Data)
+	{
+		Data.Serialize(Ar);
+		return Ar;
+	}
+}
diff --git a/Engine/Source/Runtime/Experimental/Chaos/Public/Chaos/Serialization/SolverSerializer.h b/Engine/Source/Runtime/Experimental/Chaos/Public/Chaos/Serialization/SolverSerializer.h
new file mode 100644
--- /dev/null
+++ b/Engine/Source/Runtime/Experimental/Chaos/Public/Chaos/Serialization/SolverSerializer.h
@@ -0,0 +1,85 @@
+// Copyright Epic Games, Inc. All Rights Reserved.
+
+#pragma once
+#include "Chaos/ImplicitFwd.h"
+#include "Chaos/ParticleHandleFwd.h"
+#include "Containers/Array.h"
+#include "Chaos/Framework/PhysicsProxyBase.h"
+#include "Chaos/Serialization/SerializedDataBuffer.h"
+#include "DataWrappers/ChaosVDImplicitObjectDataWrapper.h"
+#include "SolverSerializer.generated.h"
+
+namespace Chaos
+{
+struct FSerializedDataBufferHandle;
+}
+
+namespace Chaos
+{
+class FChaosArchive;
+}
+
+struct FChaosVDParticlePairMidPhase;
+struct FChaosVDCharacterGroundConstraint;
+struct FChaosVDCharacterGroundConstraintDataDataWrapper;
+struct FChaosVDJointConstraint;
+struct FChaosVDParticleDataWrapper;
+
+using FSerializedImplicitObject = FChaosVDImplicitObjectDataWrapper<Chaos::FImplicitObjectPtr, Chaos::FChaosArchive>;
+
+/**
+ * Structure holding the serialized state of a rigid solver
+ * Used to save it to disk or to re-hydrate a new solver instance
+ */
+USTRUCT()
+struct UE_INTERNAL FSerializedSolverScene
+{
+	GENERATED_BODY()
+
+	UPROPERTY()
+	TArray<FChaosVDParticleDataWrapper> ParticleData;
+	UPROPERTY()
+	TArray<FChaosVDJointConstraint> JointConstraintData;
+	UPROPERTY()
+	TArray<FChaosVDCharacterGroundConstraint> CharacterGroundConstraintData;
+	UPROPERTY()
+	TArray<FChaosVDParticlePairMidPhase> CollisionMidPhaseData;
+
+	TArray<FSerializedImplicitObject> ImplicitObjectData;
+};
+
+namespace Chaos
+{
+	class FPBDJointConstraintHandle;
+	class FConstraintHandle;
+	class FPBDRigidsSolver;
+
+	typedef TUniquePtr<FSerializedDataBuffer> FSerializedDataBufferPtr;
+
+	/* Object capable of serializing total or partially a rigid solver instance **/
+	class FSolverSerializer
+	{
+	public:
+		FSolverSerializer(FPBDRigidsSolver* InSolver);
+
+		void SerializeToBuffer(FSerializedDataBuffer& OutSerializedData);
+		void PopulateFromSerializedBuffer(const FSerializedDataBuffer& InSerializedData);
+
+		CHAOS_API void SerializeParticleStateToBuffer(FGeometryParticleHandle* InParticleHandle, FSerializedDataBuffer& OutSerializedData);
+		CHAOS_API void SerializeConstraintStateToBuffer(FConstraintHandle* ConstraintHandlePtr, FSerializedDataBuffer& OutSerializedData);
+
+		void ApplySerializedStateToParticle(FGeometryParticleHandle* InParticleHandle, FSerializedDataBuffer& InSerializedData);
+		void ApplySerializedStateToParticle(FGeometryParticleHandle* InParticleHandle, const FChaosVDParticleDataWrapper& InParticleState);
+
+		void ApplySerializedStateToJointConstraint(FPBDJointConstraintHandle* ConstraintHandlePtr, FSerializedDataBuffer& InSerializedData);
+
+		CHAOS_API void PushPendingSerializedStateForProxy(IPhysicsProxyBase* Proxy, FSerializedDataBufferPtr&& InState);
+		CHAOS_API FSerializedDataBufferPtr PopPendingSerializedStateForProxy(IPhysicsProxyBase* Proxy);
+
+
+	private:
+		FPBDRigidsSolver* SolverInstance;
+		TMap<IPhysicsProxyBase*, FSerializedDataBufferPtr> PendingMigratedPhysicsStateByProxy;
+	};
+}
+
diff --git a/Engine/Source/Runtime/Experimental/Chaos/Public/ChaosVisualDebugger/ChaosVDDataWrapperUtils.h b/Engine/Source/Runtime/Experimental/Chaos/Public/ChaosVisualDebugger/ChaosVDDataWrapperUtils.h
--- a/Engine/Source/Runtime/Experimental/Chaos/Public/ChaosVisualDebugger/ChaosVDDataWrapperUtils.h
+++ b/Engine/Source/Runtime/Experimental/Chaos/Public/ChaosVisualDebugger/ChaosVDDataWrapperUtils.h
@@ -2,13 +2,18 @@
 
 #pragma once
 
-#if WITH_CHAOS_VISUAL_DEBUGGER
 #include "Chaos/AABBTree.h"
 #include "Chaos/Collision/ContactPoint.h"
 #include "Chaos/ParticleHandle.h"
 #include "Chaos/ShapeInstanceFwd.h"
 #include "DataWrappers/ChaosVDCollisionDataWrappers.h"
 #include "DataWrappers/ChaosVDAccelerationStructureDataWrappers.h"
+#include "Chaos/ISpatialAccelerationCollection.h"
+
+namespace Chaos
+{
+class FSolverSerializer;
+}
 
 struct FChaosVDAABBTreeDataWrapper;
 struct FChaosVDCharacterGroundConstraint;
@@ -79,17 +84,22 @@ private:
 	static void CopyCollisionMaterialToDataWrapper(const Chaos::FPBDCollisionConstraintMaterial& InCopyFrom, FChaosVDCollisionMaterial& OutCopyTo);
 
 	/** Creates and populates a FChaosVDParticleDataWrapper with the data of the provided FGeometryParticleHandle */
-	static FChaosVDParticleDataWrapper BuildParticleDataWrapperFromParticle(const Chaos::FGeometryParticleHandle* ParticleHandlePtr, const TSharedRef<Chaos::VisualDebugger::FChaosVDSerializableNameTable>& InNameTableInstance);
+	static FChaosVDParticleDataWrapper BuildParticleDataWrapperFromParticle(const Chaos::FGeometryParticleHandle* ParticleHandlePtr);
+	static FChaosVDParticleDataWrapper BuildParticleDataWrapperFromGTParticle(Chaos::FGeometryParticle& GTParticle);
 
 	/** Creates and populates a FChaosVDConstraint with the data of the provided FPBDCollisionConstraint */
 	static FChaosVDConstraint BuildConstraintDataWrapperFromConstraint(const Chaos::FPBDCollisionConstraint& InConstraint);
 
 	/** Creates and populates a FChaosVDParticlePairMidPhase with the data of the provided FParticlePairMidPhase */
 	static FChaosVDParticlePairMidPhase BuildMidPhaseDataWrapperFromMidPhase(const Chaos::FParticlePairMidPhase& InMidPhase);
-
+public:
 	/** Creates and populates a FChaosVDJointConstraint with the data of the provided FPBDJointConstraintHandle */
 	static FChaosVDJointConstraint BuildJointDataWrapper(const Chaos::FPBDJointConstraintHandle* ConstaintHanlde);
 
+	static void ApplyJointDataWrapperToHandle(Chaos::FPBDJointConstraintHandle* ConstaintHanlde, const FChaosVDJointConstraint& WrappedJointData);
+
+private:
+
 	/** Creates and populates a FChaosVDChartacterGroundConstraint with the data of the provided FCharacterGroundConstraintHandle */
 	static FChaosVDCharacterGroundConstraint BuildCharacterGroundConstraintDataWrapper(const Chaos::FCharacterGroundConstraintHandle* ConstaintHanlde);
 
@@ -114,6 +124,7 @@ private:
 	static void AddTreeLeaves(const TConstArrayView<Chaos::TBoundingVolume<Chaos::FAccelerationStructureHandle>>& LeavesContainer, FChaosVDAABBTreeDataWrapper& InOutAABBTreeWrapper);
 
 	friend FChaosVisualDebuggerTrace;
+	friend Chaos::FSolverSerializer;
 };
 
 template<typename LeafType>
@@ -152,5 +163,3 @@ void FChaosVDDataWrapperUtils::BuildDataWrapperFromAABBStructure(const Chaos::TA
 	
 	OutAABBTreeWrapper.MarkAsValid();
 }
-
-#endif //WITH_CHAOS_VISUAL_DEBUGGER
diff --git a/Engine/Source/Runtime/Experimental/Chaos/Public/PBDRigidsSolver.h b/Engine/Source/Runtime/Experimental/Chaos/Public/PBDRigidsSolver.h
--- a/Engine/Source/Runtime/Experimental/Chaos/Public/PBDRigidsSolver.h
+++ b/Engine/Source/Runtime/Experimental/Chaos/Public/PBDRigidsSolver.h
@@ -27,6 +27,7 @@
 #include "SolverEventFilters.h"
 #include "Chaos/PBDRigidsEvolutionFwd.h"
 #include "ChaosSolversModule.h"
+#include "Chaos/Serialization/SolverSerializer.h"
 
 class FPhysInterface_Chaos;
 struct FChaosSolverConfiguration;
@@ -350,6 +351,7 @@ namespace Chaos
 
 		// Apply callbacks internally 
 		CHAOS_API virtual void ApplyCallbacks_Internal() override;
+		CHAOS_API FSolverSerializer& GetSerializer();
 
 	protected:
 
@@ -403,6 +405,8 @@ namespace Chaos
 
 		TUniquePtr<FPerSolverFieldSystem> PerSolverField;
 
+		FSolverSerializer Serializer;
+
 
 		// Physics material mirrors for the solver. These should generally stay in sync with the global material list from
 		// the game thread. This data is read only in the solver as we should never need to update it here. External threads can
@@ -428,7 +432,6 @@ namespace Chaos
 		TArray<FGeometryCollectionPhysicsProxy*> PendingDestroyGeometryCollectionPhysicsProxy;
 		TArray<FClusterUnionPhysicsProxy*> PendingDestroyClusterUnionProxy;
 
-
 		CHAOS_API void ProcessSinglePushedData_Internal(FPushPhysicsData& PushData);
 		CHAOS_API virtual void ProcessPushedData_Internal(FPushPhysicsData& PushData) override;
 		CHAOS_API void DestroyPendingProxies_Internal();
diff --git a/Engine/Source/Runtime/Experimental/ChaosVisualDebugger/ChaosVDRuntime.Build.cs b/Engine/Source/Runtime/Experimental/ChaosVisualDebugger/ChaosVDRuntime.Build.cs
--- a/Engine/Source/Runtime/Experimental/ChaosVisualDebugger/ChaosVDRuntime.Build.cs
+++ b/Engine/Source/Runtime/Experimental/ChaosVisualDebugger/ChaosVDRuntime.Build.cs
@@ -16,6 +16,16 @@ namespace UnrealBuildTool.Rules
 					"TraceLog"
 				}
 			);
+
+			if (IsChaosVisualDebuggerSupported(Target))
+			{
+				PublicDefinitions.Add("WITH_CHAOS_VISUAL_DEBUGGER=1");
+			}
+			else
+			{
+				PublicDefinitions.Add("WITH_CHAOS_VISUAL_DEBUGGER=0");
+			}
+
 		}
 	}
 }
\ No newline at end of file
diff --git a/Engine/Source/Runtime/Experimental/ChaosVisualDebugger/Private/ChaosVDRuntimeModule.cpp b/Engine/Source/Runtime/Experimental/ChaosVisualDebugger/Private/ChaosVDRuntimeModule.cpp
--- a/Engine/Source/Runtime/Experimental/ChaosVisualDebugger/Private/ChaosVDRuntimeModule.cpp
+++ b/Engine/Source/Runtime/Experimental/ChaosVisualDebugger/Private/ChaosVDRuntimeModule.cpp
@@ -2,6 +2,10 @@
 
 #include "ChaosVDRuntimeModule.h"
 
+#include "Modules/ModuleManager.h"
+
+#if WITH_CHAOS_VISUAL_DEBUGGER
+
 #include "Containers/Array.h"
 #include "HAL/IConsoleManager.h"
 #include "HAL/PlatformProcess.h"
@@ -396,4 +400,9 @@ void FChaosVDRuntimeModule::EnableRequiredTraceChannels()
 #endif
 }
 
-#undef LOCTEXT_NAMESPACE 
+#undef LOCTEXT_NAMESPACE
+#else
+
+IMPLEMENT_MODULE(FDefaultModuleImpl, ChaosVDRuntime);
+
+#endif
diff --git a/Engine/Source/Runtime/Experimental/ChaosVisualDebugger/Public/ChaosVDRuntimeModule.h b/Engine/Source/Runtime/Experimental/ChaosVisualDebugger/Public/ChaosVDRuntimeModule.h
--- a/Engine/Source/Runtime/Experimental/ChaosVisualDebugger/Public/ChaosVDRuntimeModule.h
+++ b/Engine/Source/Runtime/Experimental/ChaosVisualDebugger/Public/ChaosVDRuntimeModule.h
@@ -1,7 +1,7 @@
 // Copyright Epic Games, Inc. All Rights Reserved.
 
 #pragma once
-
+#if WITH_CHAOS_VISUAL_DEBUGGER
 #include "Templates/SharedPointer.h"
 #include "Modules/ModuleInterface.h"
 #include "ProfilingDebugging/TraceAuxiliary.h"
@@ -147,3 +147,4 @@ private:
 
 	static FTransactionallySafeRWLock DelegatesRWLock;
 };
+#endif
diff --git a/Engine/Source/Runtime/Experimental/ChaosVisualDebugger/Public/DataWrappers/ChaosVDParticleDataWrapper.h b/Engine/Source/Runtime/Experimental/ChaosVisualDebugger/Public/DataWrappers/ChaosVDParticleDataWrapper.h
--- a/Engine/Source/Runtime/Experimental/ChaosVisualDebugger/Public/DataWrappers/ChaosVDParticleDataWrapper.h
+++ b/Engine/Source/Runtime/Experimental/ChaosVisualDebugger/Public/DataWrappers/ChaosVDParticleDataWrapper.h
@@ -109,6 +109,18 @@ struct FChaosVDFRigidParticleControlFlags : public FChaosVDWrapperDataBase
 		bHasValidData = true;
 	}
 
+	template <typename TOther>
+	void CopyTo(TOther& Other) const
+	{
+		Other.SetGravityEnabled(bGravityEnabled);
+		Other.SetCCDEnabled(bCCDEnabled);
+		Other.SetOneWayInteractionEnabled(bOneWayInteractionEnabled);
+		Other.SetInertiaConditioningEnabled(bInertiaConditioningEnabled);
+		Other.SetGravityGroupIndex(GravityGroupIndex);
+		Other.SetMACDEnabled(bMACDEnabled);
+		Other.SetGyroscopicTorqueEnabled(bGyroscopicTorqueEnabled);
+	}
+
 	UPROPERTY(VisibleAnywhere, Category= "Particle Control Flags")
 	bool bGravityEnabled;
 	UPROPERTY(VisibleAnywhere, Category= "Particle Control Flags")
@@ -163,7 +175,7 @@ struct FChaosVDParticlePositionRotation : public FChaosVDWrapperDataBase
 	{
 		MX = Other->GetP();
 		MR = Other->GetQ();
-		
+
 		bHasValidData = true;
 	}
 
@@ -213,6 +225,13 @@ struct FChaosVDParticleVelocities : public FChaosVDWrapperDataBase
 		bHasValidData = true;
 	}
 	
+	template <typename TOther>
+	void CopyTo(TOther& Other) const
+	{
+		Other.SetV(MV);
+		Other.SetW(MW);
+	}
+	
 	UPROPERTY(VisibleAnywhere, Category= "Particle Velocities")
 	FVector MV = FVector(ForceInit);
 
@@ -250,7 +269,7 @@ struct FChaosVDParticleDynamics : public FChaosVDWrapperDataBase
 	CHAOSVDRUNTIME_API bool Serialize(FArchive& Ar);
 
 	template <typename TOther>
-	void CopyFrom(const TOther& Other)
+	void CopyFrom(TOther& Other)
 	{
 		MAcceleration = Other.Acceleration();
 		MAngularAcceleration = Other.AngularAcceleration();
@@ -260,6 +279,15 @@ struct FChaosVDParticleDynamics : public FChaosVDWrapperDataBase
 		bHasValidData = true;
 	}
 
+	template <typename TOther>
+	void CopyTo(TOther& Other) const
+	{
+		Other.Acceleration() = MAcceleration;
+        Other.AngularAcceleration() = MAngularAcceleration;
+		Other.LinearImpulseVelocity() =  MLinearImpulseVelocity;
+        Other.AngularImpulseVelocity() = MAngularImpulseVelocity;
+	}
+
 	UPROPERTY(VisibleAnywhere, Category= "Particle Dynamics")
 	FVector MAcceleration = FVector(ForceInit);
 
@@ -316,6 +344,17 @@ struct FChaosVDParticleMassProps : public FChaosVDWrapperDataBase
 		bHasValidData = true;
 	}
 	
+	template <typename TOther>
+	void CopyTo(TOther& Other) const
+	{
+		Other.SetCenterOfMass(MCenterOfMass);
+		Other.SetRotationOfMass(MRotationOfMass);
+		Other.I() = MI;
+		Other.InvI() = MInvI;
+		Other.M() = MM;
+		Other.InvM() = MInvM;
+	}
+	
 	UPROPERTY(VisibleAnywhere, Category= "Particle Mass Props")
 	FVector MCenterOfMass = FVector(ForceInit);
 
@@ -358,7 +397,7 @@ struct FChaosVDParticleDynamicMisc : public FChaosVDWrapperDataBase
 {
 	GENERATED_BODY()
 
-	FChaosVDParticleDynamicMisc(): MAngularEtherDrag(0), MMaxLinearSpeedSq(0), MMaxAngularSpeedSq(0), MInitialOverlapDepenetrationVelocity(0), MSleepThresholdMultiplier(1),
+	FChaosVDParticleDynamicMisc(): MLinearEtherDrag(0), MAngularEtherDrag(0), MMaxLinearSpeedSq(0), MMaxAngularSpeedSq(0), MInitialOverlapDepenetrationVelocity(0), MSleepThresholdMultiplier(1),
 	                               MCollisionGroup(0), MObjectState(), MSleepType(), bDisabled(false)
 	{
 	}
@@ -368,7 +407,7 @@ struct FChaosVDParticleDynamicMisc : public FChaosVDWrapperDataBase
 	template <typename TOther>
 	void CopyFrom(const TOther& Other)
 	{
-		MAngularEtherDrag = Other.LinearEtherDrag();
+		MLinearEtherDrag = Other.LinearEtherDrag();
 		MAngularEtherDrag = Other.AngularEtherDrag();
 		MMaxLinearSpeedSq = Other.MaxLinearSpeedSq();
 		MMaxAngularSpeedSq = Other.MaxAngularSpeedSq();
@@ -385,6 +424,32 @@ struct FChaosVDParticleDynamicMisc : public FChaosVDWrapperDataBase
 
 		bHasValidData = true;
 	}
+	
+	template <typename TOther, typename ControlFlagsType, typename ObjectStateType, typename SleepStateType>
+	void CopyTo(TOther& Other) const
+	{
+		Other.LinearEtherDrag() = MLinearEtherDrag;
+		Other.AngularEtherDrag() = MAngularEtherDrag;
+		Other.MaxLinearSpeedSq() = MMaxLinearSpeedSq;
+		Other.MaxAngularSpeedSq() = MMaxAngularSpeedSq;
+		Other.SetInitialOverlapDepenetrationVelocity(MInitialOverlapDepenetrationVelocity);
+		Other.SetSleepThresholdMultiplier(MSleepThresholdMultiplier);
+		Other.SetCollisionGroup(MCollisionGroup);
+		Other.SetSleepType(static_cast<SleepStateType>(MSleepType));
+		Other.SetCollisionConstraintFlags(MCollisionConstraintFlag);
+
+		ControlFlagsType ControlFlags;
+		MControlFlags.CopyTo(ControlFlags);
+
+		Other.SetControlFlags(ControlFlags);
+
+		Other.SetObjectStateLowLevel(static_cast<ObjectStateType>(MObjectState));
+		
+		Other.SetDisabled(bDisabled);
+	}
+
+	UPROPERTY(VisibleAnywhere, Category= "Particle Dynamic Misc")
+	double MLinearEtherDrag;
 
 	UPROPERTY(VisibleAnywhere, Category= "Particle Dynamic Misc")
 	double MAngularEtherDrag;
@@ -514,6 +579,12 @@ struct FChaosVDParticleCluster : public FChaosVDWrapperDataBase
 		
 		bHasValidData = true;
 	}
+	
+	template <typename TOther>
+	void CopyTo(TOther& Other) const
+	{
+		//TODO:
+	}
 
 	UPROPERTY(VisibleAnywhere, Category= "Particle Cluster")
 	int32 ParentParticleID = INDEX_NONE;
diff --git a/Engine/Source/Runtime/PhysicsCore/Private/ChaosEngineInterface.cpp b/Engine/Source/Runtime/PhysicsCore/Private/ChaosEngineInterface.cpp
--- a/Engine/Source/Runtime/PhysicsCore/Private/ChaosEngineInterface.cpp
+++ b/Engine/Source/Runtime/PhysicsCore/Private/ChaosEngineInterface.cpp
@@ -9,6 +9,10 @@
 #include "Chaos/KinematicTargets.h"
 #include "PhysicsInterfaceDeclaresCore.h"
 #include "PhysicsObjectPhysicsCoreInterface.h"
+#include "Chaos/Serialization/SerializationUtils.h"
+#include "Chaos/Serialization/SerializedDataBuffer.h"
+#include "Chaos/Serialization/SolverSerializer.h"
+#include "Physics/Experimental/PhysInterface_Chaos.h"
 
 FPhysicsDelegatesCore::FOnUpdatePhysXMaterial FPhysicsDelegatesCore::OnUpdatePhysXMaterial;
 
@@ -2090,6 +2094,101 @@ void FChaosEngineInterface::SetLocalTransform(const FPhysicsShapeHandle& InShape
 	}
 }
 
+void FChaosEngineInterface::SerializeParticleStateToBuffer(const FPhysicsActorHandle& InActorHandle, Chaos::FSerializedDataBuffer& OutSerializedData)
+{
+	if (!IsValid(InActorHandle))
+	{
+		return;
+	}
+
+	using namespace Chaos;
+
+	if (Chaos::FPBDRigidsSolver* Solver = InActorHandle->GetSolver<Chaos::FPBDRigidsSolver>())
+	{
+		if (InActorHandle->GetHandle_LowLevel())
+		{
+			Solver->GetSerializer().SerializeParticleStateToBuffer(InActorHandle->GetHandle_LowLevel(), OutSerializedData);
+		}
+		else 
+		{
+			if (FSerializedDataBufferPtr PendingDataHandle = Solver->GetSerializer().PopPendingSerializedStateForProxy(InActorHandle))
+			{
+				OutSerializedData.GetDataAsByteArrayRef().Append(PendingDataHandle->GetDataAsByteArrayRef());
+			}	
+		}
+	}
+}
+
+void FChaosEngineInterface::ApplySerializedStateToParticle(const FPhysicsActorHandle& InActorHandle, Chaos::FSerializedDataBufferPtr&& InSerializedData)
+{
+	if (!IsValid(InActorHandle))
+	{
+		return;
+	}
+
+	if (Chaos::FPBDRigidsSolver* Solver = InActorHandle->GetSolver<Chaos::FPBDRigidsSolver>())
+	{
+		Solver->GetSerializer().PushPendingSerializedStateForProxy(InActorHandle, MoveTemp(InSerializedData));
+	}
+}
+
+void FChaosEngineInterface::SerializeConstraintStateToBuffer(const FPhysicsConstraintHandle& ConstraintHandle, Chaos::FSerializedDataBuffer& OutSerializedData)
+{
+	if (!ConstraintHandle.IsValid())
+	{
+		return;
+	}
+
+	using namespace Chaos;
+
+	switch (EConstraintType Type = ConstraintHandle.Constraint->GetType())
+	{
+	case EConstraintType::NoneType:
+		break;
+	case EConstraintType::JointConstraintType:
+		{
+			if (FJointConstraintPhysicsProxy* JointProxy = ConstraintHandle.Constraint->GetProxy<FJointConstraintPhysicsProxy>())
+			{
+				if (FPBDRigidsSolver* Solver = JointProxy->GetSolver<FPBDRigidsSolver>())
+				{
+					if (JointProxy->GetHandle())
+					{
+						Solver->GetSerializer().SerializeConstraintStateToBuffer(JointProxy->GetHandle(), OutSerializedData);
+					}
+					else
+					{
+						if (FSerializedDataBufferPtr PendingData = Solver->GetSerializer().PopPendingSerializedStateForProxy(JointProxy))
+						{
+							OutSerializedData.GetDataAsByteArrayRef().Append(PendingData->GetDataAsByteArrayRef());
+						}
+					}
+				}
+			}
+		}
+		break;
+	case EConstraintType::SpringConstraintType:
+		break;
+	case EConstraintType::SuspensionConstraintType:
+		break;
+	case EConstraintType::CharacterGroundConstraintType:
+		break;
+	}
+}
+
+void FChaosEngineInterface::ApplySerializedStateToConstraint(const FPhysicsConstraintHandle& ConstraintHandle, Chaos::FSerializedDataBufferPtr&& InSerializedData)
+{
+	if (!ConstraintHandle.IsValid())
+	{
+		return;
+	}
+
+	IPhysicsProxyBase* Proxy = ConstraintHandle.Constraint->GetProxy();
+	if (Chaos::FPBDRigidsSolver* Solver = Proxy ? Proxy->GetSolver<Chaos::FPBDRigidsSolver>() : nullptr)
+	{
+		Solver->GetSerializer().PushPendingSerializedStateForProxy(Proxy, MoveTemp(InSerializedData));
+	}
+}
+
 template<typename AllocatorType>
 int32 GetAllShapesInternalImp_AssumedLocked(const FPhysicsActorHandle& InActorHandle,TArray<FPhysicsShapeReference_Chaos,AllocatorType>& OutShapes)
 {
diff --git a/Engine/Source/Runtime/PhysicsCore/Public/Chaos/ChaosEngineInterface.h b/Engine/Source/Runtime/PhysicsCore/Public/Chaos/ChaosEngineInterface.h
--- a/Engine/Source/Runtime/PhysicsCore/Public/Chaos/ChaosEngineInterface.h
+++ b/Engine/Source/Runtime/PhysicsCore/Public/Chaos/ChaosEngineInterface.h
@@ -201,7 +201,8 @@ typedef uint8 FMaskFilter;
 
 namespace Chaos
 {
-	class FBVHParticles;
+struct FSerializedDataBufferHandle;
+class FBVHParticles;
 
 	template <typename T, int>
 	class TPBDRigidParticles;
@@ -227,6 +228,9 @@ namespace Chaos
 	class FSuspensionConstraint;
 
 	class FTriangleMeshImplicitObject;
+
+	struct FSerializedDataBuffer;
+	using FSerializedDataBufferPtr = TUniquePtr<FSerializedDataBuffer>;
 }
 
 struct FCollisionShape;
@@ -569,4 +573,10 @@ public:
 	static PHYSICSCORE_API void SetIsProbeShape(const FPhysicsShapeHandle& InShape,bool bIsProbeShape);
 	static PHYSICSCORE_API void SetIsQueryShape(const FPhysicsShapeHandle& InShape,bool bIsQueryShape);
 	static PHYSICSCORE_API void SetLocalTransform(const FPhysicsShapeHandle& InShape,const FTransform& NewLocalTransform);
+
+	static PHYSICSCORE_API void SerializeParticleStateToBuffer(const FPhysicsActorHandle& InActorHandle, Chaos::FSerializedDataBuffer& OutSerializedData);
+	static PHYSICSCORE_API void ApplySerializedStateToParticle(const FPhysicsActorHandle& InActorHandle, TUniquePtr<Chaos::FSerializedDataBuffer>&& InSerializedData);
+
+	static PHYSICSCORE_API void SerializeConstraintStateToBuffer(const FPhysicsConstraintHandle& ConstraintHandle, Chaos::FSerializedDataBuffer& OutSerializedData);
+	static PHYSICSCORE_API void ApplySerializedStateToConstraint(const FPhysicsConstraintHandle& ConstraintHandle, Chaos::FSerializedDataBufferPtr&& InSerializedData);
 };
