Dataflow:
- Bug fixes

diff --git a/Engine/Source/Runtime/Experimental/Dataflow/Core/Private/Dataflow/DataflowNodeParameters.cpp b/Engine/Source/Runtime/Experimental/Dataflow/Core/Private/Dataflow/DataflowNodeParameters.cpp
--- a/Engine/Source/Runtime/Experimental/Dataflow/Core/Private/Dataflow/DataflowNodeParameters.cpp
+++ b/Engine/Source/Runtime/Experimental/Dataflow/Core/Private/Dataflow/DataflowNodeParameters.cpp
@@ -118,7 +118,7 @@ namespace UE::Dataflow
 				const FDataflowNode* OwningNode = Callstack[Idx]->GetOwningNode();
 				const FString WarningStr = FString(TEXT("Evaluation failed"));
 
-				if (OnContextHasError.IsBound())
+				if (OnContextHasWarning.IsBound())
 				{
 					OnContextHasWarning.Broadcast(OwningNode, InOutput, WarningStr);
 				}
diff --git a/Engine/Source/Runtime/Experimental/Dataflow/Core/Public/Dataflow/DataflowInputOutput.h b/Engine/Source/Runtime/Experimental/Dataflow/Core/Public/Dataflow/DataflowInputOutput.h
--- a/Engine/Source/Runtime/Experimental/Dataflow/Core/Public/Dataflow/DataflowInputOutput.h
+++ b/Engine/Source/Runtime/Experimental/Dataflow/Core/Public/Dataflow/DataflowInputOutput.h
@@ -340,7 +340,7 @@ const T& FDataflowInput::GetValue(UE::Dataflow::FContext& Context, const T& Defa
 		}
 		if (!ConnectionOut->Evaluate(Context))
 		{
-			Context.SetData(ConnectionOut->CacheKey(), Property, Default, GetOwningNodeGuid(), GetOwningNodeValueHash(), ConnectionOut->GetOwningNodeTimestamp());
+			Context.SetData(ConnectionOut->CacheKey(), Property, Default, GetOwningNodeGuid(), GetOwningNodeValueHash(), GetOwningNodeTimestamp());
 		}
 		if (Context.HasData(ConnectionOut->CacheKey()))
 		{
