Add a new IAS specific host group type.

#rnx

- FOnDemandHostGroup is currently a wrapper around a list of host urls but IAS needs additional info associated with each host group, so we need our own type FIASHostGroup that contains this additional info.
- FIASHostGroup owns a FOnDemandHostGroup but also keeps track of the host groups name (a unique identifier), it's short term error history and it's overall connection status.
- Most of this info previously existed as part of FOnDemandIoBackend but we will need to track it per host group rather than per backend.
- The error status has been removed from FBackendStatus and this class now only tracks features that have been enabled/disabled on the backend.
- Reconnecting to a hostgroup has been moved from FOnDemandIoBackend to FIASHostGroup.
- Tracking when a hostgroup has failed enough to be considered disconnected has also been moved to FIASHostGroup.
- Responsibility for tracking statistics/analytics remains with FOnDemandIoBackend for now.
- Several areas of logging and 'ias.DisplayOnScreenStatistics' debug text has been updated to include the FIASHostGroup name where appropriate.
- This is still a bit work in progress and several places still assume that DefaultHostGroup is the only host group that FOnDemandIoBackend can own. This will be addressed in the next work item.

diff --git a/Engine/Source/Runtime/Experimental/IoStore/OnDemand/Private/IasHostGroup.cpp b/Engine/Source/Runtime/Experimental/IoStore/OnDemand/Private/IasHostGroup.cpp
new file mode 100644
--- /dev/null
+++ b/Engine/Source/Runtime/Experimental/IoStore/OnDemand/Private/IasHostGroup.cpp
@@ -0,0 +1,230 @@
+// Copyright Epic Games, Inc. All Rights Reserved.
+
+#include "IasHostGroup.h"
+
+#include "HAL/IConsoleManager.h"
+#include "IO/IoStoreOnDemand.h"
+#include "IO/OnDemandHostGroup.h"
+#include "LatencyTesting.h"
+#include "Logging/StructuredLog.h"
+
+namespace UE::IoStore
+{
+
+int32 GIasHttpErrorSampleCount = 8;
+static FAutoConsoleVariableRef CVar_IasHttpErrorSampleCount(
+	TEXT("ias.HttpErrorSampleCount"),
+	GIasHttpErrorSampleCount,
+	TEXT("Number of samples for computing the moving average of failed HTTP requests")
+);
+
+float GIasHttpErrorHighWater = 0.5f;
+static FAutoConsoleVariableRef CVar_IasHttpErrorHighWater(
+	TEXT("ias.HttpErrorHighWater"),
+	GIasHttpErrorHighWater,
+	TEXT("High water mark when HTTP streaming will be disabled")
+);
+
+struct FBitWindow
+{
+	void Reset()
+	{
+		Reset(Bits.Num());
+	}
+
+	void Reset(uint32 Count)
+	{
+		Count = FMath::RoundUpToPowerOfTwo(Count);
+		Bits.SetNum(int32(Count), false);
+		Counter = 0;
+		Mask = Count - 1;
+	}
+
+	void Add(bool bValue)
+	{
+		const uint32 Idx = Counter++ & Mask;
+		Bits[Idx] = bValue;
+	}
+
+	float AvgSetBits() const
+	{
+		return float(Bits.CountSetBits()) / float(Bits.Num());
+	}
+
+private:
+	TBitArray<> Bits;
+	uint32 Counter = 0;
+	uint32 Mask = 0;
+};
+
+struct FIASHostGroup::FImpl
+{
+	FImpl() = default;
+	FImpl(FName InName, FOnDemandHostGroup& InHostGroup)
+		: HostGroup(InHostGroup)
+		, Name(InName)
+		{
+		}
+	~FImpl() = default;
+
+	FOnDemandHostGroup HostGroup;
+	FBitWindow HttpErrorHistory;
+	FName Name;
+	bool bHttpEnabled = true;
+};
+
+TIoStatusOr<FIASHostGroup> FIASHostGroup::Create(FName Name, TConstArrayView<FAnsiString> HostUrls)
+{
+	TIoStatusOr<FOnDemandHostGroup> Result = FOnDemandHostGroup::Create(HostUrls);
+
+	if (!Result.IsOk())
+	{
+		return Result.Status();
+	}
+
+	FIASHostGroup HostGroup(Name, Result.ConsumeValueOrDie());
+	HostGroup.Impl->HttpErrorHistory.Reset(GIasHttpErrorSampleCount);
+
+	return HostGroup;
+}
+
+TIoStatusOr<FIASHostGroup> FIASHostGroup::Create(FName Name, TConstArrayView<FString> HostUrls)
+{
+	TIoStatusOr<FOnDemandHostGroup> Result = FOnDemandHostGroup::Create(HostUrls);
+
+	if (!Result.IsOk())
+	{
+		return Result.Status();
+	}
+
+	FIASHostGroup HostGroup(Name, Result.ConsumeValueOrDie());
+	HostGroup.Impl->HttpErrorHistory.Reset(GIasHttpErrorSampleCount);
+
+	return HostGroup;
+}
+
+FIASHostGroup::FIASHostGroup()
+	: Impl(MakeShared<FImpl>())
+{
+}
+
+FIASHostGroup::FIASHostGroup(FName Name, FOnDemandHostGroup&& HostGroup)
+	: Impl(MakeShared<FImpl>(Name, HostGroup))
+{
+}
+
+FIASHostGroup::FIASHostGroup(FIASHostGroup::FSharedImpl&& InImpl)
+	: Impl(MoveTemp(InImpl))
+{
+}
+
+FName FIASHostGroup::GetName() const
+{
+	return Impl->Name;
+}
+
+bool FIASHostGroup::IsConnected() const
+{
+	return Impl->HostGroup.PrimaryHostIndex() != INDEX_NONE;
+}
+
+void FIASHostGroup::Connect(int32 HostIndex)
+{
+	Impl->bHttpEnabled = true;
+	Impl->HttpErrorHistory.Reset();
+
+	SetPrimaryHost(HostIndex);
+}
+
+void FIASHostGroup::Disconnect()
+{
+	Impl->bHttpEnabled = false;
+	Impl->HttpErrorHistory.Reset();
+
+	SetPrimaryHost(INDEX_NONE);
+}
+
+FIASHostGroup::EReconnectionResult FIASHostGroup::AttemptReconnection(FStringView TestPath, uint32 TimeoutMs, std::atomic_bool& CancellationToken)
+{
+	if (!IsConnected())
+	{
+		UE_LOGFMT(LogIas, Log, "[{HostName}] Trying to reconnect to any available endpoint...", GetName());
+
+		if (const int32 Idx = ConnectionTest(GetHostUrls(), TestPath, TimeoutMs, CancellationToken); Idx != INDEX_NONE)
+		{
+			Connect(Idx);
+
+			UE_LOGFMT(LogIas, Log, "[{HostName}] Successfully reconnected to '{Url}'", GetPrimaryHostUrl(), GetName());
+			return EReconnectionResult::Reconnected;
+		}
+		else
+		{
+			return EReconnectionResult::FailedToConnect;
+		}
+	}
+	else if (GetPrimaryHostIndex() != 0)
+	{
+		if (const int32 Idx = ConnectionTest(GetHostUrls().Left(1), TestPath, TimeoutMs, CancellationToken); Idx != INDEX_NONE)
+		{
+			SetPrimaryHost(Idx);
+			UE_LOGFMT(LogIas, Log, "[{HostName}] Reconnected to primary host '{Url}'", GetPrimaryHostUrl(), GetName());
+		}
+	}
+
+	return EReconnectionResult::AlreadyConnected;
+}
+
+void FIASHostGroup::OnSuccessfulResponse()
+{
+	Impl->HttpErrorHistory.Add(false);
+}
+
+bool FIASHostGroup::OnFailedResponse()
+{
+	Impl->HttpErrorHistory.Add(true);
+
+	const float Average = Impl->HttpErrorHistory.AvgSetBits();
+	const bool bAboveHighWaterMark = Average > GIasHttpErrorHighWater;
+
+	UE_LOG(LogIas, Log, TEXT("[%s] %.2f%% the last %d HTTP requests failed"), *Impl->Name.ToString(),  Average * 100.0f, GIasHttpErrorSampleCount);
+
+	if (bAboveHighWaterMark && IsConnected())
+	{
+		Disconnect();
+
+		// TODO: Update log message with host group name?
+		UE_LOG(LogIas, Warning, TEXT("[%s] Host group disabled due to high water mark of %.2f of the last %d requests reached"),
+			*Impl->Name.ToString(), GIasHttpErrorHighWater * 100.0f, GIasHttpErrorSampleCount);
+
+			return true;
+	}
+
+	return false;
+}
+
+const FOnDemandHostGroup& FIASHostGroup::GetUnderlyingHostGroup() const
+{
+	return Impl->HostGroup;
+}
+
+void FIASHostGroup::SetPrimaryHost(int32 Index) const
+{
+	Impl->HostGroup.SetPrimaryHost(Index);
+}
+
+FAnsiStringView FIASHostGroup::GetPrimaryHostUrl() const
+{
+	return Impl->HostGroup.PrimaryHost();
+}
+
+int32 FIASHostGroup::GetPrimaryHostIndex() const
+{
+	return Impl->HostGroup.PrimaryHostIndex();
+}
+
+TConstArrayView<FAnsiString> FIASHostGroup::GetHostUrls() const
+{
+	return Impl->HostGroup.Hosts();
+}
+
+} //namespace UE::IoStore
diff --git a/Engine/Source/Runtime/Experimental/IoStore/OnDemand/Private/IasHostGroup.h b/Engine/Source/Runtime/Experimental/IoStore/OnDemand/Private/IasHostGroup.h
new file mode 100644
--- /dev/null
+++ b/Engine/Source/Runtime/Experimental/IoStore/OnDemand/Private/IasHostGroup.h
@@ -0,0 +1,60 @@
+// Copyright Epic Games, Inc. All Rights Reserved.
+
+#pragma once
+
+#include "IO/IoStatus.h"
+#include "OnDemandBackendStatus.h"
+#include "Templates/SharedPointer.h"
+
+namespace UE::IoStore
+{
+
+class FOnDemandHostGroup;
+
+class FIASHostGroup
+{
+	struct FImpl;
+	using FSharedImpl = TSharedPtr<FImpl, ESPMode::ThreadSafe>;
+
+public:
+	static TIoStatusOr<FIASHostGroup> Create(FName Name, TConstArrayView<FAnsiString> HostUrls);
+	static TIoStatusOr<FIASHostGroup> Create(FName Name, TConstArrayView<FString> HostUrls);
+
+	FIASHostGroup();
+	FIASHostGroup(FName Name, FOnDemandHostGroup&& HostGroup);
+	~FIASHostGroup() = default;
+
+	FName GetName() const;
+
+	bool IsConnected() const;
+	
+	void Connect(int32 HostIndex);
+	void Disconnect();
+
+	enum class EReconnectionResult : uint8
+	{
+		AlreadyConnected = 0,
+		Reconnected,
+		FailedToConnect,
+	};
+
+	EReconnectionResult AttemptReconnection(FStringView TestPath, uint32 TimeoutMs, std::atomic_bool& CancellationToken);
+
+	void OnSuccessfulResponse();
+	bool OnFailedResponse();
+
+	void SetPrimaryHost(int32 Index) const;
+	FAnsiStringView GetPrimaryHostUrl() const;
+	int32 GetPrimaryHostIndex() const;
+
+	const FOnDemandHostGroup& GetUnderlyingHostGroup() const;
+
+	TConstArrayView<FAnsiString> GetHostUrls() const;
+private:
+	FIASHostGroup(FSharedImpl&& Impl);
+
+	FSharedImpl Impl;
+	
+};
+
+} //namespace UE::IoStore
diff --git a/Engine/Source/Runtime/Experimental/IoStore/OnDemand/Private/OnDemandBackendStatus.cpp b/Engine/Source/Runtime/Experimental/IoStore/OnDemand/Private/OnDemandBackendStatus.cpp
--- a/Engine/Source/Runtime/Experimental/IoStore/OnDemand/Private/OnDemandBackendStatus.cpp
+++ b/Engine/Source/Runtime/Experimental/IoStore/OnDemand/Private/OnDemandBackendStatus.cpp
@@ -32,11 +32,7 @@ void FBackendStatus::ToString(FStringBuilderBase& Builder) const
 
 	Builder << TEXT(" | ");
 
-	if (IsHttpError())
-	{
-		Builder << TEXT("Http - Disconnected");
-	}
-	else if (IsHttpEnabled())
+	if (IsHttpEnabled())
 	{
 		Builder << TEXT("Http - Enabled");
 	}
@@ -59,7 +55,7 @@ void FBackendStatus::ToString(FStringBuilderBase& Builder) const
 
 bool FBackendStatus::IsHttpEnabled(uint8 FlagsToTest)
 {
-	constexpr uint8 HttpFlags = uint8(EFlags::HttpEnabled) | uint8(EFlags::HttpError);
+	constexpr uint8 HttpFlags = uint8(EFlags::HttpEnabled);
 	return ((FlagsToTest & HttpFlags) == uint8(EFlags::HttpEnabled)) && GIasHttpEnabled;
 }
 
@@ -106,7 +102,6 @@ FStringBuilderBase& operator<<(FStringBuilderBase& Sb, FBackendStatus::EFlags St
 
 	AppendIf(FBackendStatus::EFlags::CacheEnabled, TEXT("CacheEnabled"));
 	AppendIf(FBackendStatus::EFlags::HttpEnabled, TEXT("HttpEnabled"));
-	AppendIf(FBackendStatus::EFlags::HttpError, TEXT("HttpError"));
 	AppendIf(FBackendStatus::EFlags::HttpBulkOptionalDisabled, TEXT("HttpBulkOptionalDisabled"));
 	AppendIf(FBackendStatus::EFlags::AbandonCache, TEXT("AbandonCache"));
 
diff --git a/Engine/Source/Runtime/Experimental/IoStore/OnDemand/Private/OnDemandBackendStatus.h b/Engine/Source/Runtime/Experimental/IoStore/OnDemand/Private/OnDemandBackendStatus.h
--- a/Engine/Source/Runtime/Experimental/IoStore/OnDemand/Private/OnDemandBackendStatus.h
+++ b/Engine/Source/Runtime/Experimental/IoStore/OnDemand/Private/OnDemandBackendStatus.h
@@ -20,9 +20,8 @@ public:
 		None						= 0,
 		CacheEnabled				= 1 << 0,
 		HttpEnabled					= 1 << 1,
-		HttpError					= 1 << 2,
-		HttpBulkOptionalDisabled	= 1 << 3,
-		AbandonCache				= 1 << 4,
+		HttpBulkOptionalDisabled	= 1 << 2,
+		AbandonCache				= 1 << 3,
 
 		// When adding new values here, remember to update operator<<(FStringBuilderBase& Sb, EFlags StatusFlags) in the cpp!
 	};
@@ -34,11 +33,6 @@ public:
 
 	bool IsHttpEnabled(EIoChunkType ChunkType) const;
 
-	bool IsHttpError() const
-	{
-		return HasAnyFlags(EFlags::HttpError);
-	}
-
 	bool IsCacheEnabled() const
 	{
 		return HasAnyFlags(EFlags::CacheEnabled);
@@ -77,11 +71,6 @@ public:
 		AddOrRemoveFlags(EFlags::CacheEnabled, bEnabled, TEXT("Cache enabled"));
 	}
 
-	void SetHttpError(bool bError)
-	{
-		AddOrRemoveFlags(EFlags::HttpError, bError, TEXT("HTTP streaming error"));
-	}
-
 	void SetAbandonCache(bool bAbandon)
 	{
 		AddOrRemoveFlags(EFlags::AbandonCache, bAbandon, TEXT("Abandon cache"));
diff --git a/Engine/Source/Runtime/Experimental/IoStore/OnDemand/Private/OnDemandIoDispatcherBackend.cpp b/Engine/Source/Runtime/Experimental/IoStore/OnDemand/Private/OnDemandIoDispatcherBackend.cpp
--- a/Engine/Source/Runtime/Experimental/IoStore/OnDemand/Private/OnDemandIoDispatcherBackend.cpp
+++ b/Engine/Source/Runtime/Experimental/IoStore/OnDemand/Private/OnDemandIoDispatcherBackend.cpp
@@ -21,6 +21,7 @@
 #include "IO/IoStore.h"
 #include "IO/IoStoreOnDemand.h"
 #include "IasCache.h"
+#include "IasHostGroup.h"
 #include "LatencyTesting.h"
 #include "Logging/StructuredLog.h"
 #include "Math/NumericLimits.h"
@@ -133,20 +134,6 @@ static FAutoConsoleVariableRef CVar_IasHttpHealthCheckWaitTime(
 	TEXT("Number of milliseconds to wait before reconnecting to avaiable endpoint(s)")
 );
 
-int32 GIasHttpErrorSampleCount = 8;
-static FAutoConsoleVariableRef CVar_IasHttpErrorSampleCount(
-	TEXT("ias.HttpErrorSampleCount"),
-	GIasHttpErrorSampleCount,
-	TEXT("Number of samples for computing the moving average of failed HTTP requests")
-);
-
-float GIasHttpErrorHighWater = 0.5f;
-static FAutoConsoleVariableRef CVar_IasHttpErrorHighWater(
-	TEXT("ias.HttpErrorHighWater"),
-	GIasHttpErrorHighWater,
-	TEXT("High water mark when HTTP streaming will be disabled")
-);
-
 bool GIasHttpEnabled = true;
 static FAutoConsoleVariableRef CVar_IasHttpEnabled(
 	TEXT("ias.HttpEnabled"),
@@ -273,33 +260,6 @@ FAutoConsoleVariableRef CVarOnDemandBackendThreadPriority(
 	return GCompleteMaterializeTaskPriorities[FMath::Clamp(GCompleteMaterializeTaskPriority, 0, UE_ARRAY_COUNT(GCompleteMaterializeTaskPriorities) - 1)];
 }
 
-///////////////////////////////////////////////////////////////////////////////
-struct FBitWindow
-{
-	void Reset(uint32 Count)
-	{
-		Count = FMath::RoundUpToPowerOfTwo(Count);
-		Bits.SetNum(int32(Count), false);
-		Counter = 0;
-		Mask = Count - 1;
-	}
-
-	void Add(bool bValue)
-	{
-		const uint32 Idx = Counter++ & Mask;
-		Bits[Idx] = bValue;
-	}
-
-	float AvgSetBits() const
-	{
-		return float(Bits.CountSetBits()) / float(Bits.Num());
-	}
-
-private:
-	TBitArray<> Bits;
-	uint32 Counter = 0;
-	uint32 Mask = 0;
-};
 ///////////////////////////////////////////////////////////////////////////////
 FIoHash GetChunkKey(const FIoHash& ChunkHash, const FIoOffsetAndLength& Range)
 {
@@ -893,7 +853,7 @@ private:
 	void CompleteCacheRead(FChunkRequest* ChunkRequest);
 	bool CreateDefaultHostGroup();
 	bool ResolveDistributedEndpoint(const FDistributedEndpointUrl& Url, TArray<FString>& OutUrls);
-	void ProcessHttpRequests(FMultiEndpointHttpClient& HttpClient, FBitWindow& HttpErrors, int32 MaxConcurrentRequests);
+	void ProcessHttpRequests(FMultiEndpointHttpClient& HttpClient, FIASHostGroup& HostGroup, int32 MaxConcurrentRequests);
 	int32 WaitForCompleteRequestTasks(float WaitTimeSeconds, float PollTimeSeconds);
 	void DrainHttpRequests();
 	void UpdateThreadPriorityIfNeeded();
@@ -917,9 +877,7 @@ private:
 	FOnDemandIoBackendStats Stats;
 	FDistributedEndpointUrl DistributionUrl;
 	FEventRef DistributedEndpointEvent;
-
-	FOnDemandHostGroup DefaultHostGroup;
-
+	FIASHostGroup DefaultHostGroup;
 	FString EndpointTestPath;
 	EThreadPriority CurrentThreadPriority;
 
@@ -961,7 +919,7 @@ FOnDemandIoBackend::FOnDemandIoBackend(const FOnDemandEndpointConfig& Config, FO
 		FConsoleCommandDelegate::CreateLambda([this]()
 			{
 				UE_LOG(LogIas, Display, TEXT("User invoked http error via 'ias.InvokeHttpFailure'"));
-				BackendStatus.SetHttpError(true);
+				DefaultHostGroup.Disconnect();
 
 				TickBackendEvent->Trigger();
 			}),
@@ -1194,9 +1152,9 @@ void FOnDemandIoBackend::CompleteCacheRead(FChunkRequest* ChunkRequest)
 		break;
 	}
 
-	if (bWasCancelled || BackendStatus.IsHttpEnabled() == false)
+	if (bWasCancelled || (!DefaultHostGroup.IsConnected() || !BackendStatus.IsHttpEnabled()))
 	{
-		UE_CLOG(BackendStatus.IsHttpEnabled() == false, LogIas, Log, TEXT("Chunk was not found in the cache and HTTP is disabled"));
+		UE_CLOG(!DefaultHostGroup.IsConnected() || !BackendStatus.IsHttpEnabled(), LogIas, Log, TEXT("Chunk was not found in the cache and HTTP is disabled"));
 		CompleteRequest(ChunkRequest);
 		return;
 	}
@@ -1218,7 +1176,7 @@ bool FOnDemandIoBackend::Resolve(FIoRequestImpl* Request)
 
 	FChunkRequestParams RequestParams = FChunkRequestParams::Create(Request, ChunkInfo);
 
-	if (BackendStatus.IsHttpEnabled(Request->ChunkId.GetChunkType()) == false)
+	if (BackendStatus.IsHttpEnabled(Request->ChunkId.GetChunkType()) == false || !DefaultHostGroup.IsConnected())
 	{ 
 		// If the cache is not readonly the chunk may get evicted before the request is completed
 		if (BackendStatus.IsCacheReadOnly() == false || Cache->ContainsChunk(RequestParams.ChunkKey) == false)
@@ -1351,7 +1309,7 @@ TIoStatusOr<uint64> FOnDemandIoBackend::GetSizeForChunk(const FIoChunkId& ChunkI
 	FIoOffsetAndLength RequestedRange(ChunkRange.GetOffset(), FMath::Min<uint64>(ChunkInfo.RawSize(), ChunkRange.GetLength()));
 	OutAvailable = ChunkInfo.RawSize();
 
-	if (BackendStatus.IsHttpEnabled(ChunkId.GetChunkType()) == false)
+	if (!BackendStatus.IsHttpEnabled(ChunkId.GetChunkType()) || !DefaultHostGroup.IsConnected())
 	{
 		// If the cache is not readonly the chunk may get evicted before the request is resolved
 		if (BackendStatus.IsCacheReadOnly() == false)
@@ -1432,7 +1390,7 @@ bool FOnDemandIoBackend::CreateDefaultHostGroup()
 		Algo::Rotate(DefaultUrls, GIasHttpPrimaryEndpoint);
 	}
 
-	TIoStatusOr<FOnDemandHostGroup> HostResult = FOnDemandHostGroup::Create(DefaultUrls);
+	TIoStatusOr<FIASHostGroup> HostResult = FIASHostGroup::Create(FName("Default"), DefaultUrls);
 	DefaultUrls.Empty(); // No longer needed
 
 	if (!HostResult.IsOk())
@@ -1524,7 +1482,7 @@ void FOnDemandIoBackend::ReportAnalytics(TArray<FAnalyticsEventAttribute>& OutAn
 
 	Stats.ReportGeneralAnalytics(OutAnalyticsArray);
 
-	if (DefaultHostGroup.PrimaryHostIndex() != INDEX_NONE)
+	if (DefaultHostGroup.IsConnected())
 	{
 		Stats.ReportEndPointAnalytics(OutAnalyticsArray);
 	}
@@ -1551,12 +1509,12 @@ FOnDemandStreamingCacheUsage FOnDemandIoBackend::GetCacheUsage() const
 	return Usage; 
 }
 
-void FOnDemandIoBackend::ProcessHttpRequests(FMultiEndpointHttpClient& HttpClient, FBitWindow& HttpErrors, int32 MaxConcurrentRequests)
+void FOnDemandIoBackend::ProcessHttpRequests(FMultiEndpointHttpClient& HttpClient, FIASHostGroup& HostGroup, int32 MaxConcurrentRequests)
 {
 	int32 NumConcurrentRequests = 0;
 	FChunkRequest* NextChunkRequest = HttpRequests.Dequeue();
 
-	HttpClient.UpdateConnection(DefaultHostGroup);
+	HttpClient.UpdateConnection(HostGroup.GetUnderlyingHostGroup());
 
 	while (NextChunkRequest)
 	{
@@ -1575,7 +1533,7 @@ void FOnDemandIoBackend::ProcessHttpRequests(FMultiEndpointHttpClient& HttpClien
 					CompleteRequest(ChunkRequest);
 					Stats.OnHttpCancel();
 				}
-				else if (BackendStatus.IsHttpEnabled() == false)
+				else if (!HostGroup.IsConnected() || !BackendStatus.IsHttpEnabled())
 				{
 					CompleteRequest(ChunkRequest);
 					// Technically this request is being skipped because of a pre-existing error. It is not
@@ -1592,8 +1550,9 @@ void FOnDemandIoBackend::ProcessHttpRequests(FMultiEndpointHttpClient& HttpClien
 					NumConcurrentRequests++;
 
 					auto AnsiUrl = StringCast<ANSICHAR>(Url.GetData(), Url.Len());
-					HttpClient.Get(DefaultHostGroup, AnsiUrl, ChunkRequest->Params.ChunkRange,
-						[this, &NextChunkRequest, ChunkRequest, &NumConcurrentRequests, &HttpErrors]
+
+					HttpClient.Get(HostGroup.GetUnderlyingHostGroup(), AnsiUrl, ChunkRequest->Params.ChunkRange,
+						[this, &NextChunkRequest, ChunkRequest, &NumConcurrentRequests, &HostGroup]
 						(FMultiEndpointHttpClientResponse&& HttpResponse)
 						{
 							NumConcurrentRequests--;
@@ -1601,39 +1560,32 @@ void FOnDemandIoBackend::ProcessHttpRequests(FMultiEndpointHttpClient& HttpClien
 							if (HttpResponse.RetryCount > 0)
 							{
 								Stats.OnHttpRetry();
-								DefaultHostGroup.SetPrimaryHost(HttpResponse.HostIndex);
 							}
 
 							if (HttpResponse.IsOk())
 							{
-								HttpErrors.Add(false);
+								HostGroup.OnSuccessfulResponse();
+
 								ChunkRequest->Chunk = MoveTemp(HttpResponse.Body);
 								Stats.OnHttpGet(ChunkRequest->Chunk.DataSize(), HttpResponse.DurationMilliseconds);
+
+								UE::Tasks::Launch(UE_SOURCE_LOCATION, [this, ChunkRequest]()
+									{
+										CompleteRequest(ChunkRequest);
+									}, GetRequestCompletionTaskPriority());
 							}
 							else
 							{
 								Stats.OnHttpError();
-								HttpErrors.Add(true);
-
-								DefaultHostGroup.SetPrimaryHost(HttpResponse.HostIndex);
-
-								const float Average = HttpErrors.AvgSetBits();
-								const bool bAboveHighWaterMark = Average > GIasHttpErrorHighWater;
-								UE_LOG(LogIas, Log, TEXT("%.2f%% the last %d HTTP requests failed"), Average * 100.0f, GIasHttpErrorSampleCount);
 
-								if (bAboveHighWaterMark && BackendStatus.IsHttpEnabled())
+								if (HostGroup.OnFailedResponse())
 								{
-									BackendStatus.SetHttpError(true);
-									UE_LOG(LogIas, Warning, TEXT("HTTP streaming disabled due to high water mark of %.2f of the last %d requests reached"),
-										GIasHttpErrorHighWater * 100.0f, GIasHttpErrorSampleCount);
+									// A disconnect was triggered
+									Stats.OnHttpDisconnected();
 								}
-							}
-
-							UE::Tasks::Launch(UE_SOURCE_LOCATION, [this, ChunkRequest]()
-								{
-									CompleteRequest(ChunkRequest);
-								}, GetRequestCompletionTaskPriority());
 
+								CompleteRequest(ChunkRequest);
+							}
 						});
 				}
 			}
@@ -1696,14 +1648,11 @@ uint32 FOnDemandIoBackend::Run()
 		return 0;
 	}
 
-	FBitWindow HttpErrors;
-	HttpErrors.Reset(GIasHttpErrorSampleCount);
-	
 	TUniquePtr<FMultiEndpointHttpClient> HttpClient = FMultiEndpointHttpClient::Create(FMultiEndpointHttpClientConfig
 		{
 			.MaxConnectionCount = GIasHttpConnectionCount,
 			.ReceiveBufferSize = GIasHttpRecvBufKiB >= 0 ? GIasHttpRecvBufKiB << 10 : -1,
-			.MaxRetryCount = (uint32)FMath::Max(DefaultHostGroup.Hosts().Num(), GIasHttpRetryCount),
+			.MaxRetryCount = (uint32)FMath::Max(DefaultHostGroup.GetHostUrls().Num(), GIasHttpRetryCount),
 			.TimeoutMs = GIasHttpFailTimeOutMs,
 			.bEnableThreadSafetyChecks = true,
 			.bAllowChunkedTransfer = GIasHttpAllowChunkedXfer,
@@ -1717,17 +1666,14 @@ uint32 FOnDemandIoBackend::Run()
 
 	const double InitStart = FPlatformTime::Seconds();
 
-	if (ConnectionTest(DefaultHostGroup.PrimaryHost(), GetEndpointTestPath(), GIasHttpTimeOutMs))
+	if (ConnectionTest(DefaultHostGroup.GetPrimaryHostUrl(), GetEndpointTestPath(), GIasHttpTimeOutMs))
 	{
 		BackendStatus.SetHttpEnabled(true);
 		Stats.OnHttpConnected();
 	}
 	else
 	{
-		BackendStatus.SetHttpError(true);
-		HttpErrors.Reset(GIasHttpErrorSampleCount);
-
-		DefaultHostGroup.SetPrimaryHost(INDEX_NONE);
+		DefaultHostGroup.Disconnect();
 	}
 
 	const double InitTime = FPlatformTime::Seconds() - InitStart;
@@ -1737,38 +1683,13 @@ uint32 FOnDemandIoBackend::Run()
 	{
 		UpdateThreadPriorityIfNeeded();
 
-		ProcessHttpRequests(*HttpClient, HttpErrors, FMath::Min(GIasHttpConcurrentRequests, 32));
+		ProcessHttpRequests(*HttpClient, DefaultHostGroup, FMath::Min(GIasHttpConcurrentRequests, 32));
 
 		if (!bStopRequested)
 		{
-			if (BackendStatus.IsHttpError())
+			if (DefaultHostGroup.AttemptReconnection(GetEndpointTestPath(), GIasHttpTimeOutMs, bStopRequested) == FIASHostGroup::EReconnectionResult::Reconnected)
 			{
-				if (DefaultHostGroup.PrimaryHostIndex() != INDEX_NONE)
-				{
-					Stats.OnHttpDisconnected();
-
-					DefaultHostGroup.SetPrimaryHost(INDEX_NONE);
-					HttpErrors.Reset(GIasHttpErrorSampleCount);
-				}
-
-				UE_LOG(LogIas, Log, TEXT("Trying to reconnect to any available endpoint..."));
-
-				if (const int32 Idx = ConnectionTest(DefaultHostGroup.Hosts(), GetEndpointTestPath(), GIasHttpTimeOutMs, bStopRequested); Idx != INDEX_NONE)
-				{
-					Stats.OnHttpConnected();
-
-					DefaultHostGroup.SetPrimaryHost(Idx);
-					BackendStatus.SetHttpError(false);
-					UE_LOGFMT(LogIas, Log, "Successfully reconnected to '{Url}'", DefaultHostGroup.PrimaryHost());
-				}
-			}
-			else if (DefaultHostGroup.PrimaryHostIndex() != 0)
-			{
-				if (const int32 Idx = ConnectionTest(DefaultHostGroup.Hosts().Left(1), GetEndpointTestPath(), GIasHttpTimeOutMs, bStopRequested); Idx != INDEX_NONE)
-				{
-					DefaultHostGroup.SetPrimaryHost(Idx);
-					UE_LOGFMT(LogIas, Log, "Reconnected to primary host '{Url}'", DefaultHostGroup.PrimaryHost());
-				}
+				Stats.OnHttpConnected();
 			}
 
 			if (BackendStatus.ShouldAbandonCache())
@@ -1781,7 +1702,7 @@ uint32 FOnDemandIoBackend::Run()
 					Cache.Release()->Abandon(); // Will delete its self
 				}
 			}
-			TickBackendEvent->Wait(BackendStatus.IsHttpError() ? GIasHttpHealthCheckWaitTime : MAX_uint32);
+			TickBackendEvent->Wait(!DefaultHostGroup.IsConnected() ? GIasHttpHealthCheckWaitTime : MAX_uint32);
 		}
 	}
 
diff --git a/Engine/Source/Runtime/Experimental/IoStore/OnDemand/Private/Statistics.cpp b/Engine/Source/Runtime/Experimental/IoStore/OnDemand/Private/Statistics.cpp
--- a/Engine/Source/Runtime/Experimental/IoStore/OnDemand/Private/Statistics.cpp
+++ b/Engine/Source/Runtime/Experimental/IoStore/OnDemand/Private/Statistics.cpp
@@ -7,7 +7,7 @@
 #include "HAL/IConsoleManager.h"
 #include "HAL/PlatformTime.h"
 #include "IO/IoStoreOnDemand.h"
-#include "IO/OnDemandHostGroup.h"
+#include "IasHostGroup.h"
 #include "Internationalization/Internationalization.h"
 #include "Math/UnrealMathUtility.h"
 #include "Misc/CoreDelegates.h"
@@ -285,22 +285,19 @@ FOnDemandIoBackendStats::FOnDemandIoBackendStats(FBackendStatus& InStatus)
 				BackendStatusText << TEXT("IAS Backend Status: ");
 				BackendStatus.ToString(BackendStatusText);
 
-				FCoreDelegates::EOnScreenMessageSeverity Verbosity = BackendStatus.IsHttpError() ?	FCoreDelegates::EOnScreenMessageSeverity::Error :
-																									FCoreDelegates::EOnScreenMessageSeverity::Info;
-
-				Out.Add(Verbosity, FText::FromStringView(BackendStatusText.ToView()));
+				Out.Add(FCoreDelegates::EOnScreenMessageSeverity::Info, FText::FromStringView(BackendStatusText.ToView()));
 			}
 
-			for (const FOnDemandHostGroup& HostGroup : HostGroups)
+			for (const FIASHostGroup& HostGroup : HostGroups)
 			{
 				FCoreDelegates::EOnScreenMessageSeverity Verbosity = FCoreDelegates::EOnScreenMessageSeverity::Info;
 				TStringBuilder<256> Text;
-				Text << TEXT("IAS HostGroup Status: ");
+				Text << TEXT("IAS HostGroup [") << HostGroup.GetName() << TEXT("] ");
 
-				if(HostGroup.PrimaryHostIndex() != INDEX_NONE)
+				if (HostGroup.IsConnected())
 				{
-					Text << HostGroup.PrimaryHost();
-					Text << TEXT(" (") << HostGroup.PrimaryHostIndex() << TEXT("/") << HostGroup.Hosts().Num() << TEXT(")");
+					Text << HostGroup.GetPrimaryHostUrl();
+					Text << TEXT(" (") << HostGroup.GetPrimaryHostIndex() << TEXT("/") << HostGroup.GetHostUrls().Num() << TEXT(")");
 				}
 				else
 				{
@@ -384,7 +381,7 @@ FOnDemandIoBackendStats* FOnDemandIoBackendStats::Get()
 	return GStatistics;
 }
 
-void FOnDemandIoBackendStats::AddHostGroup(const FOnDemandHostGroup& NewHostGroup)
+void FOnDemandIoBackendStats::AddHostGroup(const FIASHostGroup& NewHostGroup)
 {
 	// TODO: Maybe check for duplicates?
 	HostGroups.Add(NewHostGroup);
@@ -454,7 +451,7 @@ class FAnalyticsRecording : public IAnalyticsRecording
 public:
 	FAnalyticsRecording() = delete;
 
-	FAnalyticsRecording(const TArray<FOnDemandHostGroup>& InHostGroups, const FBackendStatus& InBackendStatus)
+	FAnalyticsRecording(const TArray<FIASHostGroup>& InHostGroups, const FBackendStatus& InBackendStatus)
 		: HostGroups(InHostGroups)
 		, BackendStatus(InBackendStatus)
 	{
@@ -606,7 +603,7 @@ private:
 		FTrackedValue<FCounterAtomicInt, GCacheRejectBytes>			RejectBytes;
 	} Cache;
 
-	const TArray<FOnDemandHostGroup>& HostGroups;
+	const TArray<FIASHostGroup>& HostGroups;
 	const FBackendStatus& BackendStatus;
 	bool bRecording = true;
 };
diff --git a/Engine/Source/Runtime/Experimental/IoStore/OnDemand/Private/Statistics.h b/Engine/Source/Runtime/Experimental/IoStore/OnDemand/Private/Statistics.h
--- a/Engine/Source/Runtime/Experimental/IoStore/OnDemand/Private/Statistics.h
+++ b/Engine/Source/Runtime/Experimental/IoStore/OnDemand/Private/Statistics.h
@@ -20,7 +20,7 @@ LLM_DECLARE_TAG(Ias);
 namespace UE::IoStore
 {
 
-class FOnDemandHostGroup;
+class FIASHostGroup;
 class IAnalyticsRecording;
 struct FAvailableEndPoints;
 
@@ -38,7 +38,7 @@ public:
 	FOnDemandIoBackendStats(FBackendStatus& InBackendStatus) IAS_STATISTICS_IMPL()
 	~FOnDemandIoBackendStats() IAS_STATISTICS_IMPL()
 
-	void AddHostGroup(const FOnDemandHostGroup& HostGroup)  IAS_STATISTICS_IMPL()
+	void AddHostGroup(const FIASHostGroup& HostGroup)  IAS_STATISTICS_IMPL()
 
 	/** Report analytics not directly associated with a specific endpoint */
 	void ReportGeneralAnalytics(TArray<FAnalyticsEventAttribute>& OutAnalyticsArray) const IAS_STATISTICS_IMPL()
@@ -80,7 +80,7 @@ private:
 
 #if IAS_WITH_STATISTICS
 	const FBackendStatus& BackendStatus;
-	TArray<FOnDemandHostGroup> HostGroups;
+	TArray<FIASHostGroup> HostGroups;
 #endif //IAS_WITH_STATISTICS
 };
 
