Dataflow - Fixed crash in UDataflowMesh::GetDynamicMeshRef when an empty object is used.

#rnx

diff --git a/Engine/Source/Runtime/Experimental/Dataflow/Engine/Private/Dataflow/DataflowMesh.cpp b/Engine/Source/Runtime/Experimental/Dataflow/Engine/Private/Dataflow/DataflowMesh.cpp
--- a/Engine/Source/Runtime/Experimental/Dataflow/Engine/Private/Dataflow/DataflowMesh.cpp
+++ b/Engine/Source/Runtime/Experimental/Dataflow/Engine/Private/Dataflow/DataflowMesh.cpp
@@ -5,6 +5,7 @@
 
 UDataflowMesh::UDataflowMesh(const FObjectInitializer& ObjectInitializer)
 	: Super(ObjectInitializer)
+	, DynamicMesh(MakeUnique<UE::Geometry::FDynamicMesh3>())
 {
 }
 
