Remove todo comment (the work was already done)

#rnx

diff --git a/Engine/Source/Runtime/Experimental/IoStore/OnDemand/Private/IasHostGroup.cpp b/Engine/Source/Runtime/Experimental/IoStore/OnDemand/Private/IasHostGroup.cpp
--- a/Engine/Source/Runtime/Experimental/IoStore/OnDemand/Private/IasHostGroup.cpp
+++ b/Engine/Source/Runtime/Experimental/IoStore/OnDemand/Private/IasHostGroup.cpp
@@ -192,11 +192,10 @@ bool FIASHostGroup::OnFailedResponse()
 	{
 		Disconnect();
 
-		// TODO: Update log message with host group name?
 		UE_LOG(LogIas, Warning, TEXT("[%s] Host group disabled due to high water mark of %.2f of the last %d requests reached"),
 			*Impl->Name.ToString(), GIasHttpErrorHighWater * 100.0f, GIasHttpErrorSampleCount);
 
-			return true;
+		return true;
 	}
 
 	return false;
