Iris
- DataStream API changes
- UDataStream now caches InitParameters including name and owning DataStreamManager
- DataStreamImplemenations can now query their state and request stream to be cloesd without going through ReplicationSystem

diff --git a/Engine/Plugins/Runtime/ReplicationSystemTestPlugin/Source/Private/Tests/DataStream/MockDataStream.cpp b/Engine/Plugins/Runtime/ReplicationSystemTestPlugin/Source/Private/Tests/DataStream/MockDataStream.cpp
--- a/Engine/Plugins/Runtime/ReplicationSystemTestPlugin/Source/Private/Tests/DataStream/MockDataStream.cpp
+++ b/Engine/Plugins/Runtime/ReplicationSystemTestPlugin/Source/Private/Tests/DataStream/MockDataStream.cpp
@@ -58,3 +58,8 @@ bool UMockDataStream::HasAcknowledgedAllReliableData() const
 {
 	return true;
 }
+
+void UMockDataStream::Update(const FUpdateParameters& Params)
+{
+	++CallStatus.UpdateCallCount;
+}
diff --git a/Engine/Plugins/Runtime/ReplicationSystemTestPlugin/Source/Private/Tests/DataStream/MockDataStream.h b/Engine/Plugins/Runtime/ReplicationSystemTestPlugin/Source/Private/Tests/DataStream/MockDataStream.h
--- a/Engine/Plugins/Runtime/ReplicationSystemTestPlugin/Source/Private/Tests/DataStream/MockDataStream.h
+++ b/Engine/Plugins/Runtime/ReplicationSystemTestPlugin/Source/Private/Tests/DataStream/MockDataStream.h
@@ -18,6 +18,7 @@ public:
 		uint32 ReadDataCallCount;
 		uint32 ProcessPacketDeliveryStatusCallCount;
 		uint32 ProcessPacketDeliveryStatusMagicValue;
+		uint32 UpdateCallCount;
 	};
 
 	struct FFunctionCallSetup
@@ -40,6 +41,7 @@ protected:
 	virtual void ReadData(UE::Net::FNetSerializationContext& context) override;
 	virtual void ProcessPacketDeliveryStatus(UE::Net::EPacketDeliveryStatus Status, FDataStreamRecord const* Record) override;
 	virtual bool HasAcknowledgedAllReliableData() const override;
+	virtual void Update(const FUpdateParameters& Params) override;
 
 private:
 	struct FRecord : public FDataStreamRecord
diff --git a/Engine/Plugins/Runtime/ReplicationSystemTestPlugin/Source/Private/Tests/DataStream/TestDataStream.cpp b/Engine/Plugins/Runtime/ReplicationSystemTestPlugin/Source/Private/Tests/DataStream/TestDataStream.cpp
--- a/Engine/Plugins/Runtime/ReplicationSystemTestPlugin/Source/Private/Tests/DataStream/TestDataStream.cpp
+++ b/Engine/Plugins/Runtime/ReplicationSystemTestPlugin/Source/Private/Tests/DataStream/TestDataStream.cpp
@@ -198,6 +198,24 @@ UE_NET_TEST_FIXTURE(FTestDataStream, DataStreamGetsReadDataCall)
 	DataStreamManager->ProcessPacketDeliveryStatus(EPacketDeliveryStatus::Discard, Record);
 }
 
+UE_NET_TEST_FIXTURE(FTestDataStream, DataStreamGetsUpdateCall)
+{
+	UMockDataStream::FFunctionCallSetup MockSetup = {};
+
+	UMockDataStream* Mock = CreateMockStream(&MockSetup);
+
+	FNetSerializationContext& Context = CreateDataStreamContext();
+
+	UDataStream::FUpdateParameters DataStreamUpdateParams = { .UpdateType = UDataStream::EUpdateType::PreSendUpdate };
+	DataStreamManager->Update(DataStreamUpdateParams);
+
+	// Make sure Update was called
+	{
+		const UMockDataStream::FFunctionCallStatus& CallStatus = Mock->GetFunctionCallStatus();
+		UE_NET_ASSERT_EQ(CallStatus.UpdateCallCount, 1U);		
+	}
+}
+
 // FTestDataStream implementation
 FTestDataStream::FTestDataStream()
 : FNetworkAutomationTestSuiteFixture()
@@ -597,7 +615,7 @@ UE_NET_TEST_FIXTURE(FTestDynamicCreateDataStreamFixture, TestDynamicDataStream_R
 	// Request close from client
 	ClientReplicationSystem->CloseDataStream(Client->ConnectionIdOnServer, DataStreamName);
 
-	// Verify that stream is open on both sides
+	// Verify that stream is open on server but pending close on client
 	UE_NET_ASSERT_TRUE(GetDataStreamStateOnServer() == UDataStream::EDataStreamState::Open);
 	UE_NET_ASSERT_TRUE(GetDataStreamStateOnClient() == UDataStream::EDataStreamState::PendingClose);
 
@@ -616,6 +634,44 @@ UE_NET_TEST_FIXTURE(FTestDynamicCreateDataStreamFixture, TestDynamicDataStream_R
 	UE_NET_ASSERT_TRUE(GetDataStreamStateOnClient() == UDataStream::EDataStreamState::Invalid);	
 }
 
+// Verify that stream can be closed from client
+UE_NET_TEST_FIXTURE(FTestDynamicCreateDataStreamFixture, TestDynamicDataStream_RequestCloseOnStreamFromClient)
+{
+	UReplicationSystem* ServerReplicationSystem = Server->ReplicationSystem;
+	UReplicationSystem* ClientReplicationSystem = Client->ReplicationSystem;
+
+	UDataStream* ServerNetTokenStream = ServerReplicationSystem->OpenDataStream(Client->ConnectionIdOnServer, DataStreamName);
+
+	RoundTrip();
+	RoundTrip();
+
+	// Verify that stream is open on both sides
+	UE_NET_ASSERT_TRUE(GetDataStreamStateOnServer() == UDataStream::EDataStreamState::Open);
+	UE_NET_ASSERT_TRUE(GetDataStreamStateOnClient() == UDataStream::EDataStreamState::Open);
+
+	// Request close from client
+	UNetTokenDataStream* ClientNetTokenStream = Cast<UNetTokenDataStream>(ClientReplicationSystem->GetDataStream(Client->LocalConnectionId, DataStreamName));
+	ClientNetTokenStream->RequestClose();
+
+	// Verify that stream state
+	UE_NET_ASSERT_TRUE(ServerNetTokenStream->GetState() == UDataStream::EDataStreamState::Open);
+	UE_NET_ASSERT_TRUE(ClientNetTokenStream->GetState() == UDataStream::EDataStreamState::PendingClose);
+
+	// Send and deliver to server
+	Client->UpdateAndSend(Server, true);
+	
+	UE_NET_ASSERT_TRUE(GetDataStreamStateOnServer() == UDataStream::EDataStreamState::PendingClose);
+	UE_NET_ASSERT_TRUE(GetDataStreamStateOnClient() == UDataStream::EDataStreamState::WaitOnCloseConfirmation);
+
+	// Double roundtrip and we should be done
+	RoundTrip();
+	RoundTrip();
+
+	// Verify that stream is invalidated
+	UE_NET_ASSERT_TRUE(GetDataStreamStateOnServer() == UDataStream::EDataStreamState::Invalid);
+	UE_NET_ASSERT_TRUE(GetDataStreamStateOnClient() == UDataStream::EDataStreamState::Invalid);	
+}
+
 // Verify that stream can be closed from client even if we drop request
 UE_NET_TEST_FIXTURE(FTestDynamicCreateDataStreamFixture, TestDynamicDataStream_RequestCloseFromClientIsResentIfDropped)
 {
diff --git a/Engine/Plugins/Runtime/ReplicationSystemTestPlugin/Source/Private/Tests/ReplicationSystem/ReplicationSystemServerClientTestFixture.cpp b/Engine/Plugins/Runtime/ReplicationSystemTestPlugin/Source/Private/Tests/ReplicationSystem/ReplicationSystemServerClientTestFixture.cpp
--- a/Engine/Plugins/Runtime/ReplicationSystemTestPlugin/Source/Private/Tests/ReplicationSystem/ReplicationSystemServerClientTestFixture.cpp
+++ b/Engine/Plugins/Runtime/ReplicationSystemTestPlugin/Source/Private/Tests/ReplicationSystem/ReplicationSystemServerClientTestFixture.cpp
@@ -338,24 +338,31 @@ bool FReplicationSystemTestNode::SendUpdate(uint32 ConnectionId, const TCHAR* De
 		BeginWriteParameters.WriteMode = EDataStreamWriteMode::PostTickDispatch;
 	}
 
-	const bool bResult = (Connection.DataStreamManager->BeginWrite(BeginWriteParameters) != UDataStream::EWriteResult::NoData) && (Connection.DataStreamManager->WriteData(Context, Record)  != UDataStream::EWriteResult::NoData);
-	if (bResult)
+	bool bWroteData = false;
+	if (Connection.DataStreamManager->BeginWrite(BeginWriteParameters) != UDataStream::EWriteResult::NoData)
 	{
-		Writer.CommitWrites();
-		Packet.BitCount = Writer.GetPosBits();
-		Packet.PacketId = PacketId++;
-		if (Desc)
+		Connection.DataStreamManager->WriteData(Context, Record);
+
+		if (Writer.GetPosBits() > 0)
 		{
-			Packet.Desc = FString(Desc);
+			Writer.CommitWrites();
+			Packet.BitCount = Writer.GetPosBits();
+			Packet.PacketId = PacketId++;
+			if (Desc)
+			{
+				Packet.Desc = FString(Desc);
+			}
+
+			Connection.WriteRecords.Enqueue(Record);
+			Connection.WrittenPackets.Enqueue(Packet);
+
+			bWroteData = true;
 		}
 
-		Connection.WriteRecords.Enqueue(Record);
-		Connection.WrittenPackets.Enqueue(Packet);
+		Connection.DataStreamManager->EndWrite();
 	}
 
-	Connection.DataStreamManager->EndWrite();
-
-	if (bResult)
+	if (bWroteData)
 	{
 		UE_NET_TRACE_FLUSH_COLLECTOR(Context.GetTraceCollector(), GetNetTraceId(), Connection.ConnectionId, ENetTracePacketType::Outgoing);
 		UE_NET_TRACE_PACKET_SEND(GetNetTraceId(), Connection.ConnectionId, Packet.PacketId, Packet.BitCount);
@@ -364,7 +371,7 @@ bool FReplicationSystemTestNode::SendUpdate(uint32 ConnectionId, const TCHAR* De
 
 	UE_NET_TRACE_DESTROY_COLLECTOR(Context.GetTraceCollector());
 
-	return bResult;
+	return bWroteData;
 }
 
 void FReplicationSystemTestNode::PostSendUpdate()
diff --git a/Engine/Source/Runtime/Experimental/Iris/Core/Private/Iris/DataStream/DataStream.cpp b/Engine/Source/Runtime/Experimental/Iris/Core/Private/Iris/DataStream/DataStream.cpp
--- a/Engine/Source/Runtime/Experimental/Iris/Core/Private/Iris/DataStream/DataStream.cpp
+++ b/Engine/Source/Runtime/Experimental/Iris/Core/Private/Iris/DataStream/DataStream.cpp
@@ -1,6 +1,7 @@
 // Copyright Epic Games, Inc. All Rights Reserved.
 
 #include "Iris/DataStream/DataStream.h"
+#include "Iris/DataStream/DataStreamManager.h"
 
 UDataStream::~UDataStream()
 {
@@ -17,12 +18,17 @@ void UDataStream::EndWrite()
 
 void UDataStream::Init(const FInitParameters& Params)
 {
+	DataStreamInitParameters = Params;
 }
 
 void UDataStream::Deinit()
 {
 }
 
+void UDataStream::Update(const FUpdateParameters& Params)
+{
+}
+
 const TCHAR* LexToString(const UDataStream::EDataStreamState State)
 {
 	static const TCHAR* Names[] = {
@@ -37,3 +43,23 @@ const TCHAR* LexToString(const UDataStream::EDataStreamState State)
 
 	return State < UDataStream::EDataStreamState::Count ? Names[(uint32)State] : TEXT("");
 }
+
+const UDataStream::EDataStreamState UDataStream::GetState() const
+{
+	if (DataStreamInitParameters.DataStreamManager)
+	{
+		return DataStreamInitParameters.DataStreamManager->GetStreamState(GetDataStreamName());
+	}
+	else
+	{
+		return EDataStreamState::Invalid;
+	}
+}
+
+void UDataStream::RequestClose()
+{
+	if (DataStreamInitParameters.DataStreamManager)
+	{
+		DataStreamInitParameters.DataStreamManager->CloseStream(GetDataStreamName());
+	}
+}
diff --git a/Engine/Source/Runtime/Experimental/Iris/Core/Private/Iris/DataStream/DataStreamManager.cpp b/Engine/Source/Runtime/Experimental/Iris/Core/Private/Iris/DataStream/DataStreamManager.cpp
--- a/Engine/Source/Runtime/Experimental/Iris/Core/Private/Iris/DataStream/DataStreamManager.cpp
+++ b/Engine/Source/Runtime/Experimental/Iris/Core/Private/Iris/DataStream/DataStreamManager.cpp
@@ -35,6 +35,8 @@ public:
 	void Init(const UDataStream::FInitParameters& InitParams);
 	void Deinit();
 
+	void Update(const FUpdateParameters& Params);
+
 	EWriteResult BeginWrite(const UDataStream::FBeginWriteParameters& Params);
 	UDataStream::EWriteResult WriteData(UE::Net::FNetSerializationContext& context, FDataStreamRecord const*& OutRecord);
 	void EndWrite();
@@ -78,7 +80,7 @@ private:
 	};
 
 	void InitRecordStorage();
-	void InitStream(UDataStream* Stream);
+	void InitStream(UDataStream* Stream, FName StreamName);
 	void InitStreams();
 	void DestroyStream(uint32 StreamIndex);
 	void SetStreamState(uint32 StreamIndex, EDataStreamState NewState);
@@ -121,7 +123,8 @@ UDataStreamManager::~UDataStreamManager()
 
 void UDataStreamManager::Init(const UDataStream::FInitParameters& InitParams)
 {
-	Impl->Init(InitParams);
+	UDataStream::FInitParameters Params(this, InitParams);
+	Impl->Init(Params);
 }
 
 void UDataStreamManager::Deinit()
@@ -129,6 +132,11 @@ void UDataStreamManager::Deinit()
 	Impl->Deinit();
 }
 
+void UDataStreamManager::Update(const FUpdateParameters& Params)
+{
+	Impl->Update(Params);
+}
+
 UDataStream::EWriteResult UDataStreamManager::BeginWrite(const UDataStream::FBeginWriteParameters& Params)
 {
 	return Impl->BeginWrite(Params);
@@ -159,6 +167,12 @@ bool UDataStreamManager::HasAcknowledgedAllReliableData() const
 	return Impl->HasAcknowledgedAllReliableData();
 }
 
+bool UDataStreamManager::IsKnownStreamDefinition(const FName StreamName)
+{
+	const UDataStreamDefinitions* StreamDefinitions = GetDefault<UDataStreamDefinitions>();
+	return StreamDefinitions->FindDefinition(StreamName) != nullptr;
+}
+
 ECreateDataStreamResult UDataStreamManager::CreateStream(const FName StreamName)
 {
 	return Impl->CreateStream(StreamName);
@@ -251,6 +265,17 @@ void UDataStreamManager::FImpl::Deinit()
 	StreamState.Reset();
 }
 
+void UDataStreamManager::FImpl::Update(const FUpdateParameters& Params)
+{
+	for (UDataStream* Stream : Streams)
+	{
+		if (IsValid(Stream))
+		{
+			Stream->Update(Params);
+		}
+	}
+}
+
 void UDataStreamManager::FImpl::DestroyStream(uint32 StreamIndex)
 {
 	UDataStream* Stream = Streams[StreamIndex];
@@ -444,17 +469,28 @@ UDataStreamManager::EWriteResult UDataStreamManager::FImpl::WriteData(UE::Net::F
 			}
 		}
 
-		DataStreamMask |= CurrentStreamMask;
+		// Only update DataStreamMask if data was written. 
+		if (SubBitStream.GetPosBits() > 0U)
+		{
+			DataStreamMask |= CurrentStreamMask;		
+			TempStreamRecords[StreamIt] = SubRecord;
+		}
+		else
+		{
+			ensureMsgf(SubRecord == nullptr, TEXT("DataStream '%s' provided a record despite not writing any data."), ToCStr(Stream->GetFName().GetPlainNameString()));
+		}
 
 		ManagerStream.CommitSubstream(SubBitStream);
-		TempStreamRecords[StreamIt] = SubRecord;
+
 		// Set CombinedWriteResult to HasMoreData if any of the result variables is HasMoreData, otherwise take the WriteResult, which will be 'Ok'.
 		CombinedWriteResult = (CombinedWriteResult == EWriteResult::HasMoreData || WriteResult == EWriteResult::HasMoreData) ? EWriteResult::HasMoreData : WriteResult;
 	}
 
-	if ((CombinedWriteResult == EWriteResult::NoData) && !bHasStreamsWithDirtyState)
+	if (!(DataStreamMask || bHasStreamsWithDirtyState))
 	{
 		Context.GetBitStreamWriter()->DiscardSubstream(ManagerStream);
+		// Technically we could also return EWriteResult::HasMoreData
+		CombinedWriteResult = EWriteResult::NoData;
 	}
 	else
 	{
@@ -1028,7 +1064,7 @@ ECreateDataStreamResult UDataStreamManager::FImpl::CreateStreamFromDefinition(co
 	}
 
 	// Init stream
-	InitStream(Stream);
+	InitStream(Stream, Definition.DataStreamName);
 
 	return ECreateDataStreamResult::Success;
 }
@@ -1129,11 +1165,17 @@ void UDataStreamManager::FImpl::InitRecordStorage()
 	Records.Reset();
 }
 
-void UDataStreamManager::FImpl::InitStream(UDataStream* Stream)
+void UDataStreamManager::FImpl::InitStream(UDataStream* Stream, FName DataStreamName)
 {
 	if (IsValid(Stream))
 	{
-		Stream->Init(InitParameters);
+		UDataStream::FInitParameters StreamInitParameters(InitParameters);
+		StreamInitParameters.Name = DataStreamName;
+
+		Stream->Init(StreamInitParameters);
+
+		// Catch if DataStream does not call Super::Init.
+		ensureMsgf(Stream->GetDataStreamName() == DataStreamName, TEXT("DataStream %s did not call Super::Init"), *DataStreamName.ToString());
 	}
 }
 
diff --git a/Engine/Source/Runtime/Experimental/Iris/Core/Private/Iris/ReplicationSystem/NetTokenDataStream.cpp b/Engine/Source/Runtime/Experimental/Iris/Core/Private/Iris/ReplicationSystem/NetTokenDataStream.cpp
--- a/Engine/Source/Runtime/Experimental/Iris/Core/Private/Iris/ReplicationSystem/NetTokenDataStream.cpp
+++ b/Engine/Source/Runtime/Experimental/Iris/Core/Private/Iris/ReplicationSystem/NetTokenDataStream.cpp
@@ -55,6 +55,8 @@ UNetTokenDataStream::~UNetTokenDataStream()
 
 void UNetTokenDataStream::Init(const UDataStream::FInitParameters& Params)
 {
+	Super::Init(Params);
+
 	using namespace UE::Net;
 
 	ReplicationSystemId = Params.ReplicationSystemId;
diff --git a/Engine/Source/Runtime/Experimental/Iris/Core/Private/Iris/ReplicationSystem/ReplicationDataStream.cpp b/Engine/Source/Runtime/Experimental/Iris/Core/Private/Iris/ReplicationSystem/ReplicationDataStream.cpp
--- a/Engine/Source/Runtime/Experimental/Iris/Core/Private/Iris/ReplicationSystem/ReplicationDataStream.cpp
+++ b/Engine/Source/Runtime/Experimental/Iris/Core/Private/Iris/ReplicationSystem/ReplicationDataStream.cpp
@@ -19,6 +19,8 @@ UReplicationDataStream::~UReplicationDataStream()
 
 void UReplicationDataStream::Init(const UDataStream::FInitParameters& Params)
 {
+	Super::Init(Params);
+
 	using namespace UE::Net::Private;
 
 	UReplicationSystem* ReplicationSystem = UE::Net::GetReplicationSystem(Params.ReplicationSystemId);
@@ -45,6 +47,8 @@ void UReplicationDataStream::Init(const UDataStream::FInitParameters& Params)
 
 void UReplicationDataStream::Deinit()
 {
+	Super::Deinit();
+
 	ReplicationWriter = nullptr;
 	ReplicationReader = nullptr;
 }
diff --git a/Engine/Source/Runtime/Experimental/Iris/Core/Private/Iris/ReplicationSystem/ReplicationSystem.cpp b/Engine/Source/Runtime/Experimental/Iris/Core/Private/Iris/ReplicationSystem/ReplicationSystem.cpp
--- a/Engine/Source/Runtime/Experimental/Iris/Core/Private/Iris/ReplicationSystem/ReplicationSystem.cpp
+++ b/Engine/Source/Runtime/Experimental/Iris/Core/Private/Iris/ReplicationSystem/ReplicationSystem.cpp
@@ -692,6 +692,23 @@ public:
 		Delegates.ConnectionRemovedDelegate.Broadcast(ConnectionHandle);
 	}
 
+	void UpdateDataStreams(UDataStream::FUpdateParameters& UpdateParameters)
+	{
+		IRIS_CSV_PROFILER_SCOPE(Iris, ReplicationSystem_UpdateDataStreams);
+
+		FReplicationConnections& Connections = ReplicationSystemInternal.GetConnections();
+		auto UpdateDataStreamsForConnection = [&Connections, &UpdateParameters](uint32 ConnectionId)
+		{
+			FReplicationConnection* Conn = Connections.GetConnection(ConnectionId);
+			if (UDataStreamManager* DataStreamManager = Conn->DataStreamManager.Get())
+			{
+				Conn->DataStreamManager->Update(UpdateParameters);
+			}
+		};
+		const FNetBitArray& ValidConnections = Connections.GetValidConnections();
+		ValidConnections.ForAllSetBits(UpdateDataStreamsForConnection);
+	}
+
 	void UpdateUnresolvableReferenceTracking()
 	{
 		IRIS_CSV_PROFILER_SCOPE(Iris, ReplicationSystem_UpdateUnresolvableReferenceTracking);
@@ -842,6 +859,10 @@ void UReplicationSystem::PreSendUpdate(const FSendUpdateParams& Params)
 			CollectGarbage();
 		}
 
+		// DataStream presend update, this is similar to channel tick()
+		UDataStream::FUpdateParameters DataStreamUpdateParams = { .UpdateType = UDataStream::EUpdateType::PreSendUpdate };
+		Impl->UpdateDataStreams(DataStreamUpdateParams);
+
 		if (bAllowObjectReplication)
 		{
 			UE_NET_TRACE_FRAME_STATSCOUNTER(GetId(), ReplicationSystem.ReplicatedObjectCount, InternalSys.GetNetRefHandleManager().GetActiveObjectCount(), ENetTraceVerbosity::Verbose);
@@ -1134,6 +1155,12 @@ bool UReplicationSystem::RegisterNetBlobHandler(UNetBlobHandler* Handler)
 	return NetBlobManager.RegisterNetBlobHandler(Handler);
 }
 
+/** Returns true if there exists a DataStreamDefinition for the provided Name */
+bool UReplicationSystem::IsKnownDataStreamDefinition(FName Name) const
+{
+	return UDataStreamManager::IsKnownStreamDefinition(Name);
+}
+
 UDataStream* UReplicationSystem::OpenDataStream(uint32 ConnectionId, FName Name)
 {
 	if (!IsServer())
diff --git a/Engine/Source/Runtime/Experimental/Iris/Core/Public/Iris/DataStream/DataStream.h b/Engine/Source/Runtime/Experimental/Iris/Core/Public/Iris/DataStream/DataStream.h
--- a/Engine/Source/Runtime/Experimental/Iris/Core/Public/Iris/DataStream/DataStream.h
+++ b/Engine/Source/Runtime/Experimental/Iris/Core/Public/Iris/DataStream/DataStream.h
@@ -5,6 +5,8 @@
 #include "CoreTypes.h"
 #include "DataStream.generated.h"
 
+class UDataStreamManager;
+
 namespace UE::Net
 {
 	class FNetSerializationContext;
@@ -69,6 +71,12 @@ public:
 		HasMoreData,
 	};
 
+	enum class EUpdateType : uint8
+	{
+		// Update originating from ReplicationSystem::PreSendUpdate
+		PreSendUpdate = 0,
+	};
+
 	enum class EDataStreamState : uint8
 	{
 		// Stream is invalid
@@ -89,10 +97,24 @@ public:
 
 	struct FInitParameters
 	{
+	public:
+		FInitParameters() = default;
+		FInitParameters(UDataStreamManager* InDataStreamManager, const FInitParameters& InParams)
+		{
+			*this = InParams;
+			DataStreamManager = InDataStreamManager;
+		}
+
+		UE::Net::Private::FNetExports* NetExports = nullptr;
+		FName Name;
 		uint32 ReplicationSystemId = 0U;
 		uint32 ConnectionId = 0U;
 		uint32 PacketWindowSize = 0U;
-		UE::Net::Private::FNetExports* NetExports = nullptr;
+
+	private:
+		// We only want this to be accessible from UDataStream base
+		UDataStreamManager* DataStreamManager = nullptr;
+		friend UDataStream;
 	};
 
 	struct FBeginWriteParameters
@@ -103,6 +125,11 @@ public:
 		uint32 MaxPackets = 1U;
 	};
 
+	struct FUpdateParameters
+	{
+		EUpdateType UpdateType;
+	};
+
 public:
 	IRISCORE_API virtual ~UDataStream();
 
@@ -112,6 +139,11 @@ public:
 	/** Called when a created stream is destroyed. */
 	IRISCORE_API virtual void Deinit();
 
+	/** 
+	 * Called to drive required updates during the ReplicationSystem update calls.
+ 	 */
+	IRISCORE_API virtual void Update(const FUpdateParameters& Params);
+
 	/**
 	 * Called before any calls to potential WriteData, if it returns EWriteData::NoData no other calls will be made.
 	 * The purpose of the method is to enable a DataStream to setup data that can persist over multiple calls to WriteData if bandwidth allows.
@@ -155,6 +187,34 @@ public:
 	 * @return true if the stream has no pending reliable data for which it is waiting on an acknowledgement.
 	 */
 	IRISCORE_API virtual bool HasAcknowledgedAllReliableData() const PURE_VIRTUAL(HasAcknowledgedAllReliableData, return true;);
+
+	/**
+	 * Get name of DataStream
+	 */
+	FName GetDataStreamName() const
+	{
+		return DataStreamInitParameters.Name;
+	}
+
+	/**
+ 	 * Initiate close of DataStream. Note: This only applies to DataStreams that are flagged with bDynamicCreate in the DataStreamDefinition
+	 * @param Name The name of the DataStream. Names of valid DataStream are configured in UDataStreamDefinitions.
+	 * @see UDataStreamDefinitions
+	*/
+	IRISCORE_API void RequestClose();
+
+	/** Get the current state of the DataStream. */
+	IRISCORE_API const UDataStream::EDataStreamState GetState() const;
+
+protected:
+	/** Access init parameters. */
+	const UDataStream::FInitParameters& GetInitParameters() const
+	{
+		return DataStreamInitParameters;
+	}
+
+private:
+	FInitParameters DataStreamInitParameters;
 };
 
 const TCHAR* LexToString(const UDataStream::EDataStreamState State);
diff --git a/Engine/Source/Runtime/Experimental/Iris/Core/Public/Iris/DataStream/DataStreamManager.h b/Engine/Source/Runtime/Experimental/Iris/Core/Public/Iris/DataStream/DataStreamManager.h
--- a/Engine/Source/Runtime/Experimental/Iris/Core/Public/Iris/DataStream/DataStreamManager.h
+++ b/Engine/Source/Runtime/Experimental/Iris/Core/Public/Iris/DataStream/DataStreamManager.h
@@ -43,13 +43,15 @@ class UDataStreamManager : public UDataStream
 public:
 	IRISCORE_API virtual ~UDataStreamManager();
 
+	// UDataStream interface
+
 	/** Initializes the manager. No data stream can be created by the manager before this. */
 	virtual IRISCORE_API void Init(const FInitParameters& InitParams) override;
 
 	/** Prepare for destruction. Data streams will get ProcessPacketDeliveryStatus for outstanding packets and then marked as garbage. */
 	IRISCORE_API void Deinit();
 
-	// UDataStream interface
+	IRISCORE_API virtual void Update(const FUpdateParameters& Params) override;
 
 	/** Call BeginWrite on all active data streams. */
 	IRISCORE_API virtual EWriteResult BeginWrite(const FBeginWriteParameters& Params) override;
@@ -71,6 +73,11 @@ public:
 
 	// DataStreamManager specifics
 
+	/**
+	 * Returns true of the stream of with the specified name is present in the UDataStreamDefintions
+	*/
+	IRISCORE_API static bool IsKnownStreamDefinition(const FName StreamName);
+
 	/**
 	 * Creates a DataStream that has been configured via UDataStreamDefinitions.
 	 * @param StreamName The data stream name as configured in a FDataStreamDefinition.
diff --git a/Engine/Source/Runtime/Experimental/Iris/Core/Public/Iris/ReplicationSystem/ReplicationSystem.h b/Engine/Source/Runtime/Experimental/Iris/Core/Public/Iris/ReplicationSystem/ReplicationSystem.h
--- a/Engine/Source/Runtime/Experimental/Iris/Core/Public/Iris/ReplicationSystem/ReplicationSystem.h
+++ b/Engine/Source/Runtime/Experimental/Iris/Core/Public/Iris/ReplicationSystem/ReplicationSystem.h
@@ -282,10 +282,13 @@ public:
 	 * @return A pointer to the DataStream if it exists
 	 * @see UDataStreamDefinitions
 	 */
-	UDataStream* GetDataStream(uint32 ConnectionId, FName Name);
+	IRISCORE_API UDataStream* GetDataStream(uint32 ConnectionId, FName Name);
 
 	/** Const version of GetDataStream */
-	const UDataStream* GetDataStream(uint32 ConnectionId, FName Name) const;
+	IRISCORE_API const UDataStream* GetDataStream(uint32 ConnectionId, FName Name) const;
+
+	/** Returns true if there exists a DataStreamDefinition for the provided Name */
+	IRISCORE_API bool IsKnownDataStreamDefinition(FName Name) const;
 
 	/**
 	 * Request DataStream to be created/opened for the given ConnectionId. NOTE: This only applies to DataStreams that are flagged with bDynamicCreate in the DataStreamDefinition
@@ -293,14 +296,14 @@ public:
 	 * @return A pointer to the DataStream if successfully opened or already existed.
 	 * @see UDataStreamDefinitions
 	 */
-	UDataStream* OpenDataStream(uint32 ConnectionId, FName Name);
+	IRISCORE_API UDataStream* OpenDataStream(uint32 ConnectionId, FName Name);
 
 	/**
 	 * Initiate close of DataStream for the given ConnectionId. NOTE: This only applies to DataStreams that are flagged with bDynamicCreate in the DataStreamDefinition
 	 * @param Name The name of the DataStream. Names of valid DataStream are configured in UDataStreamDefinitions.
 	 * @see UDataStreamDefinitions
 	 */
-	void CloseDataStream(uint32 ConnectionId, FName Name);
+	IRISCORE_API void CloseDataStream(uint32 ConnectionId, FName Name);
 	
 	// NetBlob 
 
