Dataflow:
- Deprecated old math vector nodes
- Versioned up Min/Max math nodes with optional variable number of inputs
- Modified Print node to output result into OutputLog
- Added Separators around Option Pin menus, and fixed text in Rt click menu
- Bug fixes

diff --git a/Engine/Plugins/Experimental/GeometryCollectionPlugin/Source/GeometryCollectionNodes/Private/Dataflow/GeometryCollectionMakeNodes.cpp b/Engine/Plugins/Experimental/GeometryCollectionPlugin/Source/GeometryCollectionNodes/Private/Dataflow/GeometryCollectionMakeNodes.cpp
--- a/Engine/Plugins/Experimental/GeometryCollectionPlugin/Source/GeometryCollectionNodes/Private/Dataflow/GeometryCollectionMakeNodes.cpp
+++ b/Engine/Plugins/Experimental/GeometryCollectionPlugin/Source/GeometryCollectionNodes/Private/Dataflow/GeometryCollectionMakeNodes.cpp
@@ -44,13 +44,14 @@ namespace UE::Dataflow
 {
 	void GeometryCollectionMakeNodes()
 	{
-		DATAFLOW_NODE_REGISTER_CREATION_FACTORY(FMakeLiteralStringDataflowNode);
+		DATAFLOW_NODE_REGISTER_CREATION_FACTORY(FMakeLiteralStringDataflowNode_v2);
 		DATAFLOW_NODE_REGISTER_CREATION_FACTORY(FMakePointsDataflowNode);
 		DATAFLOW_NODE_REGISTER_CREATION_FACTORY(FMakeBoxDataflowNode);
 		DATAFLOW_NODE_REGISTER_CREATION_FACTORY(FMakeSphereDataflowNode);
-		DATAFLOW_NODE_REGISTER_CREATION_FACTORY(FMakeLiteralFloatDataflowNode);
-		DATAFLOW_NODE_REGISTER_CREATION_FACTORY(FMakeLiteralIntDataflowNode);
-		DATAFLOW_NODE_REGISTER_CREATION_FACTORY(FMakeLiteralBoolDataflowNode);
+		DATAFLOW_NODE_REGISTER_CREATION_FACTORY(FMakeLiteralFloatDataflowNode_v2);
+		DATAFLOW_NODE_REGISTER_CREATION_FACTORY(FMakeLiteralDoubleDataflowNode);
+		DATAFLOW_NODE_REGISTER_CREATION_FACTORY(FMakeLiteralIntDataflowNode_v2);
+		DATAFLOW_NODE_REGISTER_CREATION_FACTORY(FMakeLiteralBoolDataflowNode_v2);
 		DATAFLOW_NODE_REGISTER_CREATION_FACTORY(FMakeLiteralVectorDataflowNode);
 		DATAFLOW_NODE_REGISTER_CREATION_FACTORY(FMakeTransformDataflowNode_v2);
 		DATAFLOW_NODE_REGISTER_CREATION_FACTORY(FMakeQuaternionDataflowNode);
@@ -60,6 +61,10 @@ namespace UE::Dataflow
 		DATAFLOW_NODE_REGISTER_CREATION_FACTORY(FBreakTransformDataflowNode);
 		// Deprecated
 		DATAFLOW_NODE_REGISTER_CREATION_FACTORY(FMakeTransformDataflowNode);
+		DATAFLOW_NODE_REGISTER_CREATION_FACTORY(FMakeLiteralStringDataflowNode);
+		DATAFLOW_NODE_REGISTER_CREATION_FACTORY(FMakeLiteralFloatDataflowNode);
+		DATAFLOW_NODE_REGISTER_CREATION_FACTORY(FMakeLiteralIntDataflowNode);
+		DATAFLOW_NODE_REGISTER_CREATION_FACTORY(FMakeLiteralBoolDataflowNode);
 	}
 }
 
@@ -71,6 +76,14 @@ void FMakeLiteralStringDataflowNode::Evaluate(UE::Dataflow::FContext& Context, c
 	}
 }
 
+void FMakeLiteralStringDataflowNode_v2::Evaluate(UE::Dataflow::FContext& Context, const FDataflowOutput* Out) const
+{
+	if (Out->IsA<FString>(&String))
+	{
+		SetValue(Context, String, &String);
+	}
+}
+
 void FMakePointsDataflowNode::Evaluate(UE::Dataflow::FContext& Context, const FDataflowOutput* Out) const
 {
 	if (Out->IsA<TArray<FVector>>(&Points))
@@ -112,7 +125,6 @@ void FMakeSphereDataflowNode::Evaluate(UE::Dataflow::FContext& Context, const FD
 	}
 }
 
-
 void FMakeLiteralFloatDataflowNode::Evaluate(UE::Dataflow::FContext& Context, const FDataflowOutput* Out) const
 {
 	if (Out->IsA<float>(&Float))
@@ -121,6 +133,32 @@ void FMakeLiteralFloatDataflowNode::Evaluate(UE::Dataflow::FContext& Context, co
 	}
 }
 
+void FMakeLiteralFloatDataflowNode_v2::Evaluate(UE::Dataflow::FContext& Context, const FDataflowOutput* Out) const
+{
+	if (Out->IsA<float>(&Float))
+	{
+		SetValue(Context, Float, &Float);
+	}
+}
+
+//-----------------------------------------------------------------------------------------------
+
+FMakeLiteralDoubleDataflowNode::FMakeLiteralDoubleDataflowNode(const UE::Dataflow::FNodeParameters& InParam, FGuid InGuid)
+	: FDataflowNode(InParam, InGuid)
+{
+	RegisterOutputConnection(&Double);
+}
+
+void FMakeLiteralDoubleDataflowNode::Evaluate(UE::Dataflow::FContext& Context, const FDataflowOutput* Out) const
+{
+	if (Out->IsA(&Double))
+	{
+		SetValue(Context, Double, &Double);
+	}
+}
+
+//-----------------------------------------------------------------------------------------------
+
 void FMakeLiteralIntDataflowNode::Evaluate(UE::Dataflow::FContext& Context, const FDataflowOutput* Out) const
 {
 	if (Out->IsA<int32>(&Int))
@@ -129,6 +167,14 @@ void FMakeLiteralIntDataflowNode::Evaluate(UE::Dataflow::FContext& Context, cons
 	}
 }
 
+void FMakeLiteralIntDataflowNode_v2::Evaluate(UE::Dataflow::FContext& Context, const FDataflowOutput* Out) const
+{
+	if (Out->IsA<int32>(&Int))
+	{
+		SetValue(Context, Int, &Int);
+	}
+}
+
 void FMakeLiteralBoolDataflowNode::Evaluate(UE::Dataflow::FContext& Context, const FDataflowOutput* Out) const
 {
 	if (Out->IsA<bool>(&Bool))
@@ -137,6 +183,14 @@ void FMakeLiteralBoolDataflowNode::Evaluate(UE::Dataflow::FContext& Context, con
 	}
 }
 
+void FMakeLiteralBoolDataflowNode_v2::Evaluate(UE::Dataflow::FContext& Context, const FDataflowOutput* Out) const
+{
+	if (Out->IsA<bool>(&Bool))
+	{
+		SetValue(Context, Bool, &Bool);
+	}
+}
+
 void FMakeLiteralVectorDataflowNode::Evaluate(UE::Dataflow::FContext& Context, const FDataflowOutput* Out) const
 {
 	if (Out->IsA<FVector>(&Vector))
diff --git a/Engine/Plugins/Experimental/GeometryCollectionPlugin/Source/GeometryCollectionNodes/Public/Dataflow/GeometryCollectionMakeNodes.h b/Engine/Plugins/Experimental/GeometryCollectionPlugin/Source/GeometryCollectionNodes/Public/Dataflow/GeometryCollectionMakeNodes.h
--- a/Engine/Plugins/Experimental/GeometryCollectionPlugin/Source/GeometryCollectionNodes/Public/Dataflow/GeometryCollectionMakeNodes.h
+++ b/Engine/Plugins/Experimental/GeometryCollectionPlugin/Source/GeometryCollectionNodes/Public/Dataflow/GeometryCollectionMakeNodes.h
@@ -20,11 +20,10 @@ namespace UE_DEPRECATED(5.5, "Use UE::Dataflow instead.") Dataflow {}
 #endif
 
 /**
- *
- * Description for this node
- *
+ * Make a literal string
+ * Deprecated (5.6)
  */
-USTRUCT()
+USTRUCT(meta = (DataflowGeometryCollection, Deprecated = "5.6"))
 struct FMakeLiteralStringDataflowNode : public FDataflowNode
 {
 	GENERATED_USTRUCT_BODY()
@@ -46,9 +45,31 @@ public:
 	virtual void Evaluate(UE::Dataflow::FContext& Context, const FDataflowOutput* Out) const override;
 };
 
+/**
+ * Make a literal string
+ */
+USTRUCT()
+struct FMakeLiteralStringDataflowNode_v2 : public FDataflowNode
+{
+	GENERATED_USTRUCT_BODY()
+	DATAFLOW_NODE_DEFINE_INTERNAL(FMakeLiteralStringDataflowNode_v2, "MakeLiteralString", "Utilities|String", "")
+
+public:
+	UPROPERTY(EditAnywhere, Category = "String", meta = (DataflowOutput));
+	FString String;
+
+	FMakeLiteralStringDataflowNode_v2(const UE::Dataflow::FNodeParameters& InParam, FGuid InGuid = FGuid::NewGuid())
+		: FDataflowNode(InParam, InGuid)
+	{
+		RegisterOutputConnection(&String);
+	}
+
+	virtual void Evaluate(UE::Dataflow::FContext& Context, const FDataflowOutput* Out) const override;
+};
+
 /**
  *
- * Description for this node
+ * Make a points array from specified points
  *
  */
 USTRUCT()
@@ -86,9 +107,7 @@ enum class EMakeBoxDataTypeEnum : uint8
 };
 
 /**
- *
- * Description for this node
- *
+ * Make a box
  */
 USTRUCT()
 struct FMakeBoxDataflowNode : public FDataflowNode
@@ -163,13 +182,11 @@ public:
 	virtual void Evaluate(UE::Dataflow::FContext& Context, const FDataflowOutput* Out) const override;
 };
 
-
 /**
- *
- * Description for this node
- *
+ * Make a float value
+ * Deprecated (5.6)
  */
-USTRUCT()
+USTRUCT(meta = (DataflowGeometryCollection, Deprecated = "5.6"))
 struct FMakeLiteralFloatDataflowNode : public FDataflowNode
 {
 	GENERATED_USTRUCT_BODY()
@@ -191,13 +208,55 @@ public:
 	virtual void Evaluate(UE::Dataflow::FContext& Context, const FDataflowOutput* Out) const override;
 };
 
+/**
+ * Make a float value
+ */
+USTRUCT()
+struct FMakeLiteralFloatDataflowNode_v2 : public FDataflowNode
+{
+	GENERATED_USTRUCT_BODY()
+	DATAFLOW_NODE_DEFINE_INTERNAL(FMakeLiteralFloatDataflowNode_v2, "MakeLiteralFloat", "Math|Float", "")
+
+public:
+	UPROPERTY(EditAnywhere, Category = "Float", meta = (DataflowOutput));
+	float Float = 0.f;
+
+	FMakeLiteralFloatDataflowNode_v2(const UE::Dataflow::FNodeParameters& InParam, FGuid InGuid = FGuid::NewGuid())
+		: FDataflowNode(InParam, InGuid)
+	{
+		RegisterOutputConnection(&Float);
+	}
+
+	virtual void Evaluate(UE::Dataflow::FContext& Context, const FDataflowOutput* Out) const override;
+};
 
 /**
  *
- * Description for this node
+ * Make a double value
  *
  */
 USTRUCT()
+struct FMakeLiteralDoubleDataflowNode : public FDataflowNode
+{
+	GENERATED_USTRUCT_BODY()
+	DATAFLOW_NODE_DEFINE_INTERNAL(FMakeLiteralDoubleDataflowNode, "MakeLiteralDouble", "Math|Double", "")
+
+private:
+	UPROPERTY(EditAnywhere, Category = "Double", meta = (DataflowOutput));
+	double Double = 0.0;
+
+public:
+	FMakeLiteralDoubleDataflowNode(const UE::Dataflow::FNodeParameters& InParam, FGuid InGuid = FGuid::NewGuid());
+
+private:
+	virtual void Evaluate(UE::Dataflow::FContext& Context, const FDataflowOutput* Out) const override;
+};
+
+/**
+ * Make an integer value
+ * Deprecated (5.6)
+ */
+USTRUCT(meta = (DataflowGeometryCollection, Deprecated = "5.6"))
 struct FMakeLiteralIntDataflowNode : public FDataflowNode
 {
 	GENERATED_USTRUCT_BODY()
@@ -220,11 +279,32 @@ public:
 };
 
 /**
- *
- * Description for this node
- *
+ * Make an integer value
  */
 USTRUCT()
+struct FMakeLiteralIntDataflowNode_v2 : public FDataflowNode
+{
+	GENERATED_USTRUCT_BODY()
+	DATAFLOW_NODE_DEFINE_INTERNAL(FMakeLiteralIntDataflowNode_v2, "MakeLiteralInt", "Math|Int", "")
+
+public:
+	UPROPERTY(EditAnywhere, Category = "Int", meta = (DataflowOutput));
+	int32 Int = 0;
+
+	FMakeLiteralIntDataflowNode_v2(const UE::Dataflow::FNodeParameters& InParam, FGuid InGuid = FGuid::NewGuid())
+		: FDataflowNode(InParam, InGuid)
+	{
+		RegisterOutputConnection(&Int);
+	}
+
+	virtual void Evaluate(UE::Dataflow::FContext& Context, const FDataflowOutput* Out) const override;
+};
+
+/**
+ * Make a bool value
+ * Deprecated(5.6)
+ */
+USTRUCT(meta = (DataflowGeometryCollection, Deprecated = "5.6"))
 struct FMakeLiteralBoolDataflowNode : public FDataflowNode
 {
 	GENERATED_USTRUCT_BODY()
@@ -246,6 +326,28 @@ public:
 	virtual void Evaluate(UE::Dataflow::FContext& Context, const FDataflowOutput* Out) const override;
 };
 
+/**
+ * Make a bool value
+ */
+USTRUCT()
+struct FMakeLiteralBoolDataflowNode_v2 : public FDataflowNode
+{
+	GENERATED_USTRUCT_BODY()
+	DATAFLOW_NODE_DEFINE_INTERNAL(FMakeLiteralBoolDataflowNode_v2, "MakeLiteralBool", "Math|Boolean", "")
+
+public:
+	UPROPERTY(EditAnywhere, Category = "Bool", meta = (DataflowOutput));
+	bool Bool = false;
+
+	FMakeLiteralBoolDataflowNode_v2(const UE::Dataflow::FNodeParameters& InParam, FGuid InGuid = FGuid::NewGuid())
+		: FDataflowNode(InParam, InGuid)
+	{
+		RegisterOutputConnection(&Bool);
+	}
+
+	virtual void Evaluate(UE::Dataflow::FContext& Context, const FDataflowOutput* Out) const override;
+};
+
 
 /**
  *
diff --git a/Engine/Plugins/Experimental/GeometryCollectionPlugin/Source/GeometryCollectionNodes/Public/Dataflow/GeometryCollectionMathNodes.h b/Engine/Plugins/Experimental/GeometryCollectionPlugin/Source/GeometryCollectionNodes/Public/Dataflow/GeometryCollectionMathNodes.h
--- a/Engine/Plugins/Experimental/GeometryCollectionPlugin/Source/GeometryCollectionNodes/Public/Dataflow/GeometryCollectionMathNodes.h
+++ b/Engine/Plugins/Experimental/GeometryCollectionPlugin/Source/GeometryCollectionNodes/Public/Dataflow/GeometryCollectionMathNodes.h
@@ -594,7 +594,7 @@ public:
 *
 *
 */
-USTRUCT(meta = (DataflowGeometryCollection))
+USTRUCT(meta = (DataflowGeometryCollection, Deprecated = "5.6"))
 struct FMin3DataflowNode : public FDataflowNode
 {
 	GENERATED_USTRUCT_BODY()
@@ -630,7 +630,7 @@ public:
 *
 *
 */
-USTRUCT(meta = (DataflowGeometryCollection))
+USTRUCT(meta = (DataflowGeometryCollection, Deprecated = "5.6"))
 struct FMax3DataflowNode : public FDataflowNode
 {
 	GENERATED_USTRUCT_BODY()
@@ -1248,7 +1248,7 @@ public:
 *
 *
 */
-USTRUCT(meta = (DataflowGeometryCollection))
+USTRUCT(meta = (DataflowGeometryCollection, Deprecated = "5.6"))
 struct FScaleVectorDataflowNode : public FDataflowNode
 {
 	GENERATED_USTRUCT_BODY()
@@ -1280,7 +1280,7 @@ public:
 *
 *
 */
-USTRUCT(meta = (DataflowGeometryCollection))
+USTRUCT(meta = (DataflowGeometryCollection, Deprecated = "5.6"))
 struct FDotProductDataflowNode : public FDataflowNode
 {
 	GENERATED_USTRUCT_BODY()
@@ -1312,7 +1312,7 @@ public:
 *
 *
 */
-USTRUCT(meta = (DataflowGeometryCollection))
+USTRUCT(meta = (DataflowGeometryCollection, Deprecated = "5.6"))
 struct FCrossProductDataflowNode : public FDataflowNode
 {
 	GENERATED_USTRUCT_BODY()
@@ -1344,7 +1344,7 @@ public:
 *
 *
 */
-USTRUCT(meta = (DataflowGeometryCollection))
+USTRUCT(meta = (DataflowGeometryCollection, Deprecated = "5.6"))
 struct FNormalizeDataflowNode : public FDataflowNode
 {
 	GENERATED_USTRUCT_BODY()
@@ -1376,7 +1376,7 @@ public:
 *
 *
 */
-USTRUCT(meta = (DataflowGeometryCollection))
+USTRUCT(meta = (DataflowGeometryCollection, Deprecated = "5.6"))
 struct FLengthDataflowNode : public FDataflowNode
 {
 	GENERATED_USTRUCT_BODY()
@@ -1404,7 +1404,7 @@ public:
 *
 *
 */
-USTRUCT(meta = (DataflowGeometryCollection))
+USTRUCT(meta = (DataflowGeometryCollection, Deprecated = "5.6"))
 struct FDistanceDataflowNode : public FDataflowNode
 {
 	GENERATED_USTRUCT_BODY()
diff --git a/Engine/Source/Runtime/Experimental/Dataflow/Core/Private/Dataflow/DataflowCoreNodes.cpp b/Engine/Source/Runtime/Experimental/Dataflow/Core/Private/Dataflow/DataflowCoreNodes.cpp
--- a/Engine/Source/Runtime/Experimental/Dataflow/Core/Private/Dataflow/DataflowCoreNodes.cpp
+++ b/Engine/Source/Runtime/Experimental/Dataflow/Core/Private/Dataflow/DataflowCoreNodes.cpp
@@ -196,6 +196,6 @@ FDataflowPrintNode::FDataflowPrintNode(const UE::Dataflow::FNodeParameters& InPa
 
 void FDataflowPrintNode::Evaluate(UE::Dataflow::FContext& Context, const FDataflowOutput* Out) const
 {
-	const FString InValue = GetValue(Context, &Value);
-	UE_LOG(LogTemp, Warning, TEXT("[Dataflow Print] %s"), *InValue);
+	const FString InValue = "Result=" + GetValue(Context, &Value);
+	Context.Info(InValue, this, nullptr);
 }
\ No newline at end of file
diff --git a/Engine/Source/Runtime/Experimental/Dataflow/Core/Private/Dataflow/DataflowMathNodes.cpp b/Engine/Source/Runtime/Experimental/Dataflow/Core/Private/Dataflow/DataflowMathNodes.cpp
--- a/Engine/Source/Runtime/Experimental/Dataflow/Core/Private/Dataflow/DataflowMathNodes.cpp
+++ b/Engine/Source/Runtime/Experimental/Dataflow/Core/Private/Dataflow/DataflowMathNodes.cpp
@@ -24,8 +24,8 @@ namespace UE::Dataflow
 		DATAFLOW_NODE_REGISTER_CREATION_FACTORY(FDataflowMathInverseSquareRootNode);
 		DATAFLOW_NODE_REGISTER_CREATION_FACTORY(FDataflowMathLogNode);
 		DATAFLOW_NODE_REGISTER_CREATION_FACTORY(FDataflowMathLogXNode);
-		DATAFLOW_NODE_REGISTER_CREATION_FACTORY(FDataflowMathMaximumNode);
-		DATAFLOW_NODE_REGISTER_CREATION_FACTORY(FDataflowMathMinimumNode);
+		DATAFLOW_NODE_REGISTER_CREATION_FACTORY(FDataflowMathMinimumNode_v2);
+		DATAFLOW_NODE_REGISTER_CREATION_FACTORY(FDataflowMathMaximumNode_v2);
 		DATAFLOW_NODE_REGISTER_CREATION_FACTORY(FDataflowMathMultiplyNode);
 		DATAFLOW_NODE_REGISTER_CREATION_FACTORY(FDataflowMathNegateNode);
 		DATAFLOW_NODE_REGISTER_CREATION_FACTORY(FDataflowMathOneMinusNode);
@@ -48,6 +48,10 @@ namespace UE::Dataflow
 		DATAFLOW_NODE_REGISTER_CREATION_FACTORY(FDataflowMathArcTan2Node)
 		DATAFLOW_NODE_REGISTER_CREATION_FACTORY(FDataflowMathDegToRadNode);
 		DATAFLOW_NODE_REGISTER_CREATION_FACTORY(FDataflowMathRadToDegNode);
+
+		// Deprecated
+		DATAFLOW_NODE_REGISTER_CREATION_FACTORY(FDataflowMathMinimumNode);
+		DATAFLOW_NODE_REGISTER_CREATION_FACTORY(FDataflowMathMaximumNode);
 	}
 }
 
@@ -179,6 +183,117 @@ double FDataflowMathMinimumNode::ComputeResult(UE::Dataflow::FContext& Context,
 
 //-----------------------------------------------------------------------------------------------
 
+FDataflowMathMinimumNode_v2::FDataflowMathMinimumNode_v2(const UE::Dataflow::FNodeParameters& InParam, FGuid InGuid)
+	: FDataflowNode(InParam, InGuid)
+{
+	RegisterOutputConnection(&Result);
+
+	// Add initial variable inputs
+	for (int32 Index = 0; Index < NumInitialVariableInputs; ++Index)
+	{
+		AddPins();
+	}
+}
+
+void FDataflowMathMinimumNode_v2::Evaluate(UE::Dataflow::FContext& Context, const FDataflowOutput* Out) const
+{
+	if (Out->IsA(&Result))
+	{
+		double MinValue = std::numeric_limits<double>::max();
+		for (int32 Idx = 0; Idx < Inputs.Num(); ++Idx)
+		{
+			const double InputValue = GetValue(Context, GetConnectionReference(Idx));
+			MinValue = FMath::Min(InputValue, MinValue);
+		}
+
+		SetValue(Context, MinValue, &Result);
+	}
+}
+
+bool FDataflowMathMinimumNode_v2::CanAddPin() const
+{
+	return true;
+}
+
+bool FDataflowMathMinimumNode_v2::CanRemovePin() const
+{
+	return Inputs.Num() > 0;
+}
+
+UE::Dataflow::TConnectionReference<FDataflowNumericTypes> FDataflowMathMinimumNode_v2::GetConnectionReference(int32 Index) const
+{
+	return { &Inputs[Index], Index, &Inputs };
+}
+
+TArray<UE::Dataflow::FPin> FDataflowMathMinimumNode_v2::AddPins()
+{
+	const int32 Index = Inputs.AddDefaulted();
+	const FDataflowInput& Input = RegisterInputArrayConnection(GetConnectionReference(Index));
+	return { { UE::Dataflow::FPin::EDirection::INPUT, Input.GetType(), Input.GetName() } };
+}
+
+TArray<UE::Dataflow::FPin> FDataflowMathMinimumNode_v2::GetPinsToRemove() const
+{
+	const int32 Index = (Inputs.Num() - 1);
+	check(Inputs.IsValidIndex(Index));
+	if (const FDataflowInput* const Input = FindInput(GetConnectionReference(Index)))
+	{
+		return { { UE::Dataflow::FPin::EDirection::INPUT, Input->GetType(), Input->GetName() } };
+	}
+	return Super::GetPinsToRemove();
+}
+
+void FDataflowMathMinimumNode_v2::OnPinRemoved(const UE::Dataflow::FPin& Pin)
+{
+	const int32 Index = Inputs.Num() - 1;
+	check(Inputs.IsValidIndex(Index));
+#if DO_CHECK
+	const FDataflowInput* const Input = FindInput(GetConnectionReference(Index));
+	check(Input);
+	check(Input->GetName() == Pin.Name);
+	check(Input->GetType() == Pin.Type);
+#endif
+	Inputs.SetNum(Index);
+
+	return Super::OnPinRemoved(Pin);
+}
+
+void FDataflowMathMinimumNode_v2::PostSerialize(const FArchive& Ar)
+{
+	if (Ar.IsLoading())
+	{
+		check(Inputs.Num() >= 0);
+		// register new elements from the array as inputs
+		for (int32 Index = 0; Index < Inputs.Num(); ++Index)
+		{
+			FindOrRegisterInputArrayConnection(GetConnectionReference(Index));
+		}
+		if (Ar.IsTransacting())
+		{
+			// if we have more inputs than materials then we need to unregister the inputs 
+			const int32 NumVariableInputs = (GetNumInputs() - NumOtherInputs);
+			const int32 NumInputs = Inputs.Num();
+			if (NumVariableInputs > NumInputs)
+			{
+				// Inputs have been removed.
+				// Temporarily expand Collections so we can get connection references.
+				Inputs.SetNum(NumVariableInputs);
+				for (int32 Index = NumInputs; Index < Inputs.Num(); ++Index)
+				{
+					UnregisterInputConnection(GetConnectionReference(Index));
+				}
+				Inputs.SetNum(NumInputs);
+			}
+		}
+		else
+		{
+			ensureAlways(Inputs.Num() + NumOtherInputs == GetNumInputs());
+		}
+	}
+}
+
+//-----------------------------------------------------------------------------------------------
+
 FDataflowMathMaximumNode::FDataflowMathMaximumNode(const UE::Dataflow::FNodeParameters& InParam, FGuid InGuid)
 	: FDataflowMathTwoInputsOperatorNode(InParam, InGuid)
 {
@@ -192,6 +307,117 @@ double FDataflowMathMaximumNode::ComputeResult(UE::Dataflow::FContext& Context,
 
 //-----------------------------------------------------------------------------------------------
 
+FDataflowMathMaximumNode_v2::FDataflowMathMaximumNode_v2(const UE::Dataflow::FNodeParameters& InParam, FGuid InGuid)
+	: FDataflowNode(InParam, InGuid)
+{
+	RegisterOutputConnection(&Result);
+
+	// Add initial variable inputs
+	for (int32 Index = 0; Index < NumInitialVariableInputs; ++Index)
+	{
+		AddPins();
+	}
+}
+
+void FDataflowMathMaximumNode_v2::Evaluate(UE::Dataflow::FContext& Context, const FDataflowOutput* Out) const
+{
+	if (Out->IsA(&Result))
+	{
+		double MaxValue = std::numeric_limits<double>::min();
+		for (int32 Idx = 0; Idx < Inputs.Num(); ++Idx)
+		{
+			const double InputValue = GetValue(Context, GetConnectionReference(Idx));
+			MaxValue = FMath::Max(InputValue, MaxValue);
+		}
+
+		SetValue(Context, MaxValue, &Result);
+	}
+}
+
+bool FDataflowMathMaximumNode_v2::CanAddPin() const
+{
+	return true;
+}
+
+bool FDataflowMathMaximumNode_v2::CanRemovePin() const
+{
+	return Inputs.Num() > 0;
+}
+
+UE::Dataflow::TConnectionReference<FDataflowNumericTypes> FDataflowMathMaximumNode_v2::GetConnectionReference(int32 Index) const
+{
+	return { &Inputs[Index], Index, &Inputs };
+}
+
+TArray<UE::Dataflow::FPin> FDataflowMathMaximumNode_v2::AddPins()
+{
+	const int32 Index = Inputs.AddDefaulted();
+	const FDataflowInput& Input = RegisterInputArrayConnection(GetConnectionReference(Index));
+	return { { UE::Dataflow::FPin::EDirection::INPUT, Input.GetType(), Input.GetName() } };
+}
+
+TArray<UE::Dataflow::FPin> FDataflowMathMaximumNode_v2::GetPinsToRemove() const
+{
+	const int32 Index = (Inputs.Num() - 1);
+	check(Inputs.IsValidIndex(Index));
+	if (const FDataflowInput* const Input = FindInput(GetConnectionReference(Index)))
+	{
+		return { { UE::Dataflow::FPin::EDirection::INPUT, Input->GetType(), Input->GetName() } };
+	}
+	return Super::GetPinsToRemove();
+}
+
+void FDataflowMathMaximumNode_v2::OnPinRemoved(const UE::Dataflow::FPin& Pin)
+{
+	const int32 Index = Inputs.Num() - 1;
+	check(Inputs.IsValidIndex(Index));
+#if DO_CHECK
+	const FDataflowInput* const Input = FindInput(GetConnectionReference(Index));
+	check(Input);
+	check(Input->GetName() == Pin.Name);
+	check(Input->GetType() == Pin.Type);
+#endif
+	Inputs.SetNum(Index);
+
+	return Super::OnPinRemoved(Pin);
+}
+
+void FDataflowMathMaximumNode_v2::PostSerialize(const FArchive& Ar)
+{
+	if (Ar.IsLoading())
+	{
+		check(Inputs.Num() >= 0);
+		// register new elements from the array as inputs
+		for (int32 Index = 0; Index < Inputs.Num(); ++Index)
+		{
+			FindOrRegisterInputArrayConnection(GetConnectionReference(Index));
+		}
+		if (Ar.IsTransacting())
+		{
+			// if we have more inputs than materials then we need to unregister the inputs 
+			const int32 NumVariableInputs = (GetNumInputs() - NumOtherInputs);
+			const int32 NumInputs = Inputs.Num();
+			if (NumVariableInputs > NumInputs)
+			{
+				// Inputs have been removed.
+				// Temporarily expand Collections so we can get connection references.
+				Inputs.SetNum(NumVariableInputs);
+				for (int32 Index = NumInputs; Index < Inputs.Num(); ++Index)
+				{
+					UnregisterInputConnection(GetConnectionReference(Index));
+				}
+				Inputs.SetNum(NumInputs);
+			}
+		}
+		else
+		{
+			ensureAlways(Inputs.Num() + NumOtherInputs == GetNumInputs());
+		}
+	}
+}
+
+//-----------------------------------------------------------------------------------------------
+
 FDataflowMathReciprocalNode::FDataflowMathReciprocalNode(const UE::Dataflow::FNodeParameters& InParam, FGuid InGuid)
 	: FDataflowMathOneInputOperatorNode(InParam, InGuid)
 {
diff --git a/Engine/Source/Runtime/Experimental/Dataflow/Core/Public/Dataflow/DataflowMathNodes.h b/Engine/Source/Runtime/Experimental/Dataflow/Core/Public/Dataflow/DataflowMathNodes.h
--- a/Engine/Source/Runtime/Experimental/Dataflow/Core/Public/Dataflow/DataflowMathNodes.h
+++ b/Engine/Source/Runtime/Experimental/Dataflow/Core/Public/Dataflow/DataflowMathNodes.h
@@ -119,8 +119,12 @@ public:
 	virtual double ComputeResult(UE::Dataflow::FContext& Context, double InA, double InB) const override;
 };
 
-/** Minimum ( Min(A, B) ) */
-USTRUCT()
+/** 
+* Minimum ( Min(A, B) ) 
+* Deprecated (5.6)
+* Use Minimum (V2) with variable number of inputs instead
+*/
+USTRUCT(meta = (Deprecated = "5.6"))
 struct FDataflowMathMinimumNode : public FDataflowMathTwoInputsOperatorNode
 {
 	GENERATED_USTRUCT_BODY()
@@ -131,8 +135,47 @@ public:
 	virtual double ComputeResult(UE::Dataflow::FContext& Context, double InA, double InB) const override;
 };
 
-/** Maximum ( Max(A, B) ) */
+/**
+* Minimum ( Min(A, B, C, ... ) )
+*/
 USTRUCT()
+struct FDataflowMathMinimumNode_v2 : public FDataflowNode
+{
+	GENERATED_USTRUCT_BODY()
+	DATAFLOW_NODE_DEFINE_INTERNAL(FDataflowMathMinimumNode_v2, "Minimum", DATAFLOW_MATH_NODES_CATEGORY, "")
+
+private:
+	UPROPERTY()
+	TArray<FDataflowNumericTypes> Inputs;
+
+	UPROPERTY(meta = (DataflowOutput))
+	FDataflowNumericTypes Result;
+
+	//~ Begin FDataflowNode interface
+	virtual TArray<UE::Dataflow::FPin> AddPins() override;
+	virtual bool CanAddPin() const override;
+	virtual bool CanRemovePin() const override;
+	virtual TArray<UE::Dataflow::FPin> GetPinsToRemove() const override;
+	virtual void OnPinRemoved(const UE::Dataflow::FPin& Pin) override;
+	virtual void PostSerialize(const FArchive& Ar) override;
+	virtual void Evaluate(UE::Dataflow::FContext& Context, const FDataflowOutput* Out) const override;
+	//~ End FDataflowNode interface
+
+	UE::Dataflow::TConnectionReference<FDataflowNumericTypes> GetConnectionReference(int32 Index) const;
+
+	static constexpr int32 NumOtherInputs = 0; // No other inputs
+	static constexpr int32 NumInitialVariableInputs = 2;
+
+public:
+	FDataflowMathMinimumNode_v2(const UE::Dataflow::FNodeParameters& InParam, FGuid InGuid = FGuid::NewGuid());
+};
+
+/**
+* Maximum ( Max(A, B) )
+* Deprecated (5.6)
+* Use Maximum (V2) with variable number of inputs instead
+*/
+USTRUCT(meta = (Deprecated = "5.6"))
 struct FDataflowMathMaximumNode : public FDataflowMathTwoInputsOperatorNode
 {
 	GENERATED_USTRUCT_BODY()
@@ -143,6 +186,39 @@ public:
 	virtual double ComputeResult(UE::Dataflow::FContext& Context, double InA, double InB) const override;
 };
 
+/** Maximum ( Max(A, B, C, ... ) ) */
+USTRUCT()
+struct FDataflowMathMaximumNode_v2 : public FDataflowNode
+{
+	GENERATED_USTRUCT_BODY()
+	DATAFLOW_NODE_DEFINE_INTERNAL(FDataflowMathMaximumNode_v2, "Maximum", DATAFLOW_MATH_NODES_CATEGORY, "")
+
+private:
+	UPROPERTY()
+	TArray<FDataflowNumericTypes> Inputs;
+
+	UPROPERTY(meta = (DataflowOutput))
+	FDataflowNumericTypes Result;
+
+	//~ Begin FDataflowNode interface
+	virtual TArray<UE::Dataflow::FPin> AddPins() override;
+	virtual bool CanAddPin() const override;
+	virtual bool CanRemovePin() const override;
+	virtual TArray<UE::Dataflow::FPin> GetPinsToRemove() const override;
+	virtual void OnPinRemoved(const UE::Dataflow::FPin& Pin) override;
+	virtual void PostSerialize(const FArchive& Ar) override;
+	virtual void Evaluate(UE::Dataflow::FContext& Context, const FDataflowOutput* Out) const override;
+	//~ End FDataflowNode interface
+
+	UE::Dataflow::TConnectionReference<FDataflowNumericTypes> GetConnectionReference(int32 Index) const;
+
+	static constexpr int32 NumOtherInputs = 0; // No other inputs
+	static constexpr int32 NumInitialVariableInputs = 2;
+
+public:
+	FDataflowMathMaximumNode_v2(const UE::Dataflow::FNodeParameters& InParam, FGuid InGuid = FGuid::NewGuid());
+};
+
 /** 
 * Reciprocal( 1 / A )
 * if A is equal to 0, returns Fallback
