Dataflow:
- Added override color to Dataflow node

diff --git a/Engine/Source/Runtime/Experimental/Dataflow/Core/Public/Dataflow/DataflowNode.h b/Engine/Source/Runtime/Experimental/Dataflow/Core/Public/Dataflow/DataflowNode.h
--- a/Engine/Source/Runtime/Experimental/Dataflow/Core/Public/Dataflow/DataflowNode.h
+++ b/Engine/Source/Runtime/Experimental/Dataflow/Core/Public/Dataflow/DataflowNode.h
@@ -59,6 +59,12 @@ PRAGMA_ENABLE_DEPRECATION_WARNINGS
 	UPROPERTY(EditAnywhere, Category = "Dataflow")
 	bool bActive = true;
 
+	UPROPERTY(EditAnywhere, Category = "Dataflow", meta = (InlineEditConditionToggle))
+	bool bOverrideColor = false;
+
+	UPROPERTY(EditAnywhere, Category = "Dataflow", meta = (EditCondition="bOverrideColor"))
+	FLinearColor OverrideColor = FLinearColor::White;
+
 	FDataflowNode()
 		: Guid(FGuid())
 		, Name("Invalid")
@@ -394,6 +400,16 @@ PRAGMA_ENABLE_DEPRECATION_WARNINGS
 		return FrozenProperties;
 	}
 
+	bool IsColorOverriden() const
+	{
+		return bOverrideColor;
+	}
+
+	const FLinearColor& GetOverrideColor() const
+	{
+		return OverrideColor;
+	}
+
 	/**
 	 * Gets the property name for bActive. This exists so the property can be made private later.
 	 */
diff --git a/Engine/Source/Runtime/Experimental/Dataflow/Engine/Private/Dataflow/DataflowEdNode.cpp b/Engine/Source/Runtime/Experimental/Dataflow/Engine/Private/Dataflow/DataflowEdNode.cpp
--- a/Engine/Source/Runtime/Experimental/Dataflow/Engine/Private/Dataflow/DataflowEdNode.cpp
+++ b/Engine/Source/Runtime/Experimental/Dataflow/Engine/Private/Dataflow/DataflowEdNode.cpp
@@ -636,6 +636,10 @@ FLinearColor UDataflowEdNode::GetNodeTitleColor() const
 		{
 			return UE::Dataflow::Private::FrozenTitleColor;
 		}
+		else if (DataflowNode->IsColorOverriden())
+		{
+			return DataflowNode->GetOverrideColor();
+		}
 		return UE::Dataflow::FNodeColorsRegistry::Get().GetNodeTitleColor(DataflowNode->GetCategory());
 	}
 	return FDataflowNode::DefaultNodeTitleColor;
