Iris
- Changed IrisObjectReferencePackageMap to allow for serializing more than 256 unique exports.
- Added LevelOffset to NetTrace::OffsetScope to allow for better injection of out of band Trace events

diff --git a/Engine/Source/Runtime/Experimental/Iris/Core/Private/Iris/ReplicationSystem/NetBlob/NetRPC.cpp b/Engine/Source/Runtime/Experimental/Iris/Core/Private/Iris/ReplicationSystem/NetBlob/NetRPC.cpp
--- a/Engine/Source/Runtime/Experimental/Iris/Core/Private/Iris/ReplicationSystem/NetBlob/NetRPC.cpp
+++ b/Engine/Source/Runtime/Experimental/Iris/Core/Private/Iris/ReplicationSystem/NetBlob/NetRPC.cpp
@@ -520,7 +520,7 @@ void FNetRPC::CallFunction(FNetRPCCallContext& CallContext)
 	{
 		bool bLogRpc = true;
 
-		// Suppress spammy engine RPCs. This could be made a configable list in the future.
+		// Suppress spammy engine RPCs. This could be made a configurable list in the future.
 		const FString& FunctionName = Function->GetName();
 		if (   FunctionName.Contains(TEXT("ServerUpdateCamera"))
 			|| FunctionName.Contains(TEXT("ClientAckGoodMove"))
diff --git a/Engine/Source/Runtime/Experimental/Iris/Core/Private/Iris/Serialization/IrisObjectReferencePackageMap.cpp b/Engine/Source/Runtime/Experimental/Iris/Core/Private/Iris/Serialization/IrisObjectReferencePackageMap.cpp
--- a/Engine/Source/Runtime/Experimental/Iris/Core/Private/Iris/Serialization/IrisObjectReferencePackageMap.cpp
+++ b/Engine/Source/Runtime/Experimental/Iris/Core/Private/Iris/Serialization/IrisObjectReferencePackageMap.cpp
@@ -14,47 +14,36 @@ namespace UE
 
 bool UIrisObjectReferencePackageMap::SerializeObject(FArchive& Ar, UClass* InClass, UObject*& Obj, FNetworkGUID* OutNetGUID)
 {
-	if (!PackageMapExports)
+	UE::Net::FIrisPackageMapExports* PackageMapExports = Ar.IsSaving() ? PackageMapExportsForWriting : PackageMapExportsForReading;
+
+	if (!ensure(PackageMapExports))
 	{
 		return false;
 	}
 
-	constexpr uint8 MaxNumReferences = 255U;
 	UE::Net::FIrisPackageMapExports::FObjectReferenceArray* References = &PackageMapExports->References;
 
 	if (Ar.IsSaving())
 	{
-		int32 Index = MaxNumReferences;
-		if (!References->Find(Obj, Index) && References->Num() < MaxNumReferences)
+		int32 Index;
+		if (!References->Find(Obj, Index))
 		{
 			Index = References->Add(Obj);
 		}
-
-		if (References->IsValidIndex(Index))
-		{
-			uint8 IndexByte = static_cast<uint8>(Index);
-			Ar << IndexByte;
-		}
-		else
-		{
-			ensureMsgf(false, TEXT("UIrisObjectReferencePackageMap::SerializeObject, failed to serialize object reference with Index %u (%s). A Maximum of %u references are currently supported by this PackageMap"),
-				Index, *GetNameSafe(Obj), MaxNumReferences);
-			uint8 IndexByte = MaxNumReferences;
-			Ar << IndexByte;
-			return false;
-		}
+		uint32 IndexToWrite = static_cast<uint32>(Index);
+		Ar.SerializeIntPacked(IndexToWrite);
 	}
 	else
 	{
-		uint8 IndexByte = MaxNumReferences;
-		Ar << IndexByte;
-		if (References->IsValidIndex(IndexByte) && IndexByte < MaxNumReferences)
+		uint32 ReadIndex;
+		Ar.SerializeIntPacked(ReadIndex);
+		if (References->IsValidIndex(ReadIndex))
 		{
-			Obj = (*References)[IndexByte];
+			Obj = (*References)[ReadIndex];
 		}
 		else
 		{
-			ensureMsgf(false, TEXT("UIrisObjectReferencePackageMap::SerializeObject, failed to read object reference index %u is out of bounds. Current ObjectReference num: %u"), IndexByte, References->Num());
+			ensureMsgf(false, TEXT("UIrisObjectReferencePackageMap::SerializeObject, failed to read object reference index %u is out of bounds. Current ObjectReference num: %u"), ReadIndex, References->Num());
 			return false;
 		}
 	}
@@ -66,47 +55,40 @@ bool UIrisObjectReferencePackageMap::SerializeName(FArchive& Ar, FName& InName)
 {
 	using namespace UE::Net;
 
-	if (!(PackageMapExports && bEnableIrisPackageMapNameExports))
+	UE::Net::FIrisPackageMapExports* PackageMapExports = Ar.IsSaving() ? PackageMapExportsForWriting : PackageMapExportsForReading;
+
+	if (!bEnableIrisPackageMapNameExports)
 	{
 		return Super::SerializeName(Ar, InName);
 	}
 
-	FIrisPackageMapExports::FNameArray* Names = &PackageMapExports->Names;
-	constexpr uint8 MaxNumNames = 255U;
+	if (!ensure(PackageMapExports))
+	{
+		return false;
+	}
 
+	FIrisPackageMapExports::FNameArray* Names = &PackageMapExports->Names;
 	if (Ar.IsSaving())
 	{
-		int32 Index = MaxNumNames;
-		if (!Names->Find(InName, Index) && Names->Num() < MaxNumNames)
+		int32 Index;
+		if (!Names->Find(InName, Index))
 		{
 			Index = Names->Add(InName);
 		}
-
-		if (Names->IsValidIndex(Index))
-		{
-			uint8 IndexByte = static_cast<uint8>(Index);
-			Ar << IndexByte;
-		}
-		else
-		{
-			ensureMsgf(false, TEXT("UIrisObjectReferencePackageMap::SerializeName, failed to serialize name with Index %u (%s). A Maximum of %u name are currently supported by this PackageMap"),
-				Index, *InName.ToString(), MaxNumNames);
-			uint8 IndexByte = MaxNumNames;
-			Ar << IndexByte;
-			return false;
-		}
+		uint32 IndexToWrite = static_cast<uint32>(Index);
+		Ar.SerializeIntPacked(IndexToWrite);
 	}
 	else
 	{
-		uint8 IndexByte = MaxNumNames;
-		Ar << IndexByte;
-		if (Names->IsValidIndex(IndexByte) && IndexByte < MaxNumNames)
+		uint32 ReadIndex;
+		Ar.SerializeIntPacked(ReadIndex);
+		if (Names->IsValidIndex(ReadIndex))
 		{
-			InName = (*Names)[IndexByte];
+			InName = (*Names)[ReadIndex];
 		}
 		else
 		{
-			ensureMsgf(false, TEXT("UIrisObjectReferencePackageMap::SerializeName, failed to read name index %u is out of bounds. Current Name num: %u"), IndexByte, Names->Num());
+			ensureMsgf(false, TEXT("UIrisObjectReferencePackageMap::SerializeName, failed to read name index %u is out of bounds. Current Name num: %u"), ReadIndex, Names->Num());
 			return false;
 		}
 	}
@@ -116,16 +98,66 @@ bool UIrisObjectReferencePackageMap::SerializeName(FArchive& Ar, FName& InName)
 
 void UIrisObjectReferencePackageMap::InitForRead(const UE::Net::FIrisPackageMapExports* InPackageMapExports, const UE::Net::FNetTokenResolveContext& InNetTokenResolveContext)
 { 
-	PackageMapExports = const_cast<UE::Net::FIrisPackageMapExports*>(InPackageMapExports);
+	PackageMapExportsForReading = const_cast<UE::Net::FIrisPackageMapExports*>(InPackageMapExports);
 	NetTokenResolveContext = InNetTokenResolveContext;
 }
 
 void UIrisObjectReferencePackageMap::InitForWrite(UE::Net::FIrisPackageMapExports* InPackageMapExports)
 {
-	PackageMapExports = InPackageMapExports;
-	if (ensureMsgf(PackageMapExports, TEXT("UIrisObjectReferencePackageMap requires valid PackageMapExports to capture exports.")))
+	PackageMapExportsForWriting = InPackageMapExports;
+	if (ensureMsgf(PackageMapExportsForWriting, TEXT("UIrisObjectReferencePackageMap requires valid PackageMapExports to capture exports.")))
+	{
+		PackageMapExportsForWriting->Reset();
+	}
+}
+
+namespace UE::Net
+{
+
+FIrisObjectReferencePackageMapReadScope::FIrisObjectReferencePackageMapReadScope(UIrisObjectReferencePackageMap* InPackageMap, const UE::Net::FIrisPackageMapExports* PackageMapExports, const UE::Net::FNetTokenResolveContext* NetTokenResolveContext)
+: PackageMap(InPackageMap)
+{
+	if (PackageMap)
+	{
+		PackageMap->InitForRead(PackageMapExports, *NetTokenResolveContext);
+	}
+}
+
+FIrisObjectReferencePackageMapReadScope::~FIrisObjectReferencePackageMapReadScope()
+{
+	if (PackageMap)
+	{
+		PackageMap->PackageMapExportsForReading = nullptr;
+	}
+}
+
+FIrisObjectReferencePackageMapWriteScope::FIrisObjectReferencePackageMapWriteScope(UIrisObjectReferencePackageMap* InPackageMap, UE::Net::FIrisPackageMapExports* PackageMapExports)
+: PackageMap(InPackageMap)
+{
+	if (!PackageMap)
+	{
+		return;
+	}
+
+	if (PackageMapExports && ensureMsgf(PackageMap->PackageMapExportsForReading != PackageMapExports, TEXT("FIrisObjectReferencePackageMapWriteScope cannot read and write from the same FIrisPackageMapExports as we are reading from.")))
 	{
 		PackageMapExports->Reset();
+		PackageMap->PackageMapExportsForWriting = PackageMapExports;
+	}
+	else
+	{
+		PackageMap->PackageMapExportsForWriting = nullptr;
+	}
+}
+
+FIrisObjectReferencePackageMapWriteScope::~FIrisObjectReferencePackageMapWriteScope()
+{
+	if (PackageMap)
+	{
+		PackageMap->PackageMapExportsForWriting = nullptr;
 	}
 }
 
+} // UE::Net
+
+
diff --git a/Engine/Source/Runtime/Experimental/Iris/Core/Private/Iris/Serialization/IrisPackageMapExportUtil.cpp b/Engine/Source/Runtime/Experimental/Iris/Core/Private/Iris/Serialization/IrisPackageMapExportUtil.cpp
--- a/Engine/Source/Runtime/Experimental/Iris/Core/Private/Iris/Serialization/IrisPackageMapExportUtil.cpp
+++ b/Engine/Source/Runtime/Experimental/Iris/Core/Private/Iris/Serialization/IrisPackageMapExportUtil.cpp
@@ -10,6 +10,7 @@
 #include "Iris/Serialization/NetReferenceCollector.h"
 #include "Iris/Serialization/NetSerializerArrayStorage.h"
 #include "Iris/Serialization/IrisObjectReferencePackageMap.h"
+#include "Iris/Serialization/InternalNetSerializationContext.h"
 #include "Net/Core/NetToken/NetToken.h"
 #include "Net/Core/Trace/NetTrace.h"
 #include "NetExportContext.h"
@@ -204,6 +205,18 @@ void FIrisPackageMapExportsUtil::Quantize(FNetSerializationContext& Context, con
 	}
 }
 
+void FIrisPackageMapExportsUtil::FreeDynamicState(QuantizedType& Value)
+{
+	FNetSerializationContext Context;
+	Private::FInternalNetSerializationContext InternalContext;
+	Context.SetInternalContext(&InternalContext);
+
+	// Clear all info
+	Value.ObjectReferenceStorage.Free(Context);
+	Value.NameStorage.Free(Context);
+	Value.NetTokenStorage.Free(Context);
+}
+
 void FIrisPackageMapExportsUtil::Dequantize(FNetSerializationContext& Context, const QuantizedType& Source, UE::Net::FIrisPackageMapExports& PackageMapExports)
 {
 	// References
diff --git a/Engine/Source/Runtime/Experimental/Iris/Core/Public/Iris/Serialization/IrisObjectReferencePackageMap.h b/Engine/Source/Runtime/Experimental/Iris/Core/Public/Iris/Serialization/IrisObjectReferencePackageMap.h
--- a/Engine/Source/Runtime/Experimental/Iris/Core/Public/Iris/Serialization/IrisObjectReferencePackageMap.h
+++ b/Engine/Source/Runtime/Experimental/Iris/Core/Public/Iris/Serialization/IrisObjectReferencePackageMap.h
@@ -8,9 +8,13 @@
 
 // Forward declarations
 class FNetworkGUID;
+class UIrisObjectReferencePackageMap;
 
 namespace UE::Net
 {
+	class FIrisObjectReferencePackageMapWriteScope;
+	class FIrisObjectReferencePackageMapReadScope;
+
 	// In order to properly capture exported data when calling in to old style NetSerialize methods
 	// we need to capture and inject certain types.
 	struct FIrisPackageMapExports
@@ -18,6 +22,11 @@ namespace UE::Net
 		typedef TArray<TObjectPtr<UObject>, TInlineAllocator<4>> FObjectReferenceArray;
 		typedef TArray<FName, TInlineAllocator<4>> FNameArray;
 
+		bool IsEmpty() const
+		{
+			return References.IsEmpty() && Names.IsEmpty();
+		}
+
 		void Reset()
 		{
 			References.Reset();
@@ -27,6 +36,30 @@ namespace UE::Net
 		FObjectReferenceArray References;
 		FNameArray Names;
 	};
+
+	// Scope that calls InitForRead on target PackageMap and invalidates set PackageMapExports on scope exit.
+	class FIrisObjectReferencePackageMapReadScope
+	{
+	public:
+		IRISCORE_API FIrisObjectReferencePackageMapReadScope(UIrisObjectReferencePackageMap* PackageMap, const UE::Net::FIrisPackageMapExports* PackageMapExports, const UE::Net::FNetTokenResolveContext* NetTokenResolveContext);
+		IRISCORE_API ~FIrisObjectReferencePackageMapReadScope();
+		UIrisObjectReferencePackageMap* GetPackageMap() { return PackageMap; }
+
+	private:
+		UIrisObjectReferencePackageMap* PackageMap = nullptr;
+	};
+
+	// Scope that calls InitForWrite on target PackageMap and invalidates set PackageMapExports on scope exit.
+	class FIrisObjectReferencePackageMapWriteScope
+	{
+	public:
+		IRISCORE_API FIrisObjectReferencePackageMapWriteScope(UIrisObjectReferencePackageMap* PackageMap, UE::Net::FIrisPackageMapExports* PackageMapExports);
+		IRISCORE_API ~FIrisObjectReferencePackageMapWriteScope();
+		UIrisObjectReferencePackageMap* GetPackageMap() { return PackageMap; }
+
+	private:
+		UIrisObjectReferencePackageMap* PackageMap = nullptr;
+	};
 }
 
 /**
@@ -46,15 +79,19 @@ public:
 	// Override SerializeName in order to be able to capture name and serialize them with iris instead.
 	virtual bool SerializeName(FArchive& Ar, FName& InName);
 
+	virtual const UE::Net::FNetTokenResolveContext* GetNetTokenResolveContext() const  override { return &NetTokenResolveContext; }
+
 	// Init for read, we need to set the exports from which we are going to read our data.
 	IRISCORE_API void InitForRead(const UE::Net::FIrisPackageMapExports* PackageMapExports, const UE::Net::FNetTokenResolveContext& InNetTokenResolveContext);
 
 	// Init for write, all captured exports will be serialized as in index and added to the PackageMapExports for later export using iris.
 	IRISCORE_API void InitForWrite(UE::Net::FIrisPackageMapExports* PackageMapExports);
 
-	virtual const UE::Net::FNetTokenResolveContext* GetNetTokenResolveContext() const  override { return &NetTokenResolveContext; }
+private:
+	friend UE::Net::FIrisObjectReferencePackageMapReadScope;
+	friend UE::Net::FIrisObjectReferencePackageMapWriteScope;
 
-protected:
-	UE::Net::FIrisPackageMapExports* PackageMapExports = nullptr;
+	UE::Net::FIrisPackageMapExports* PackageMapExportsForReading = nullptr;
+	UE::Net::FIrisPackageMapExports* PackageMapExportsForWriting = nullptr;
 	UE::Net::FNetTokenResolveContext NetTokenResolveContext;
 };
diff --git a/Engine/Source/Runtime/Experimental/Iris/Core/Public/Iris/Serialization/IrisPackageMapExportUtil.h b/Engine/Source/Runtime/Experimental/Iris/Core/Public/Iris/Serialization/IrisPackageMapExportUtil.h
--- a/Engine/Source/Runtime/Experimental/Iris/Core/Public/Iris/Serialization/IrisPackageMapExportUtil.h
+++ b/Engine/Source/Runtime/Experimental/Iris/Core/Public/Iris/Serialization/IrisPackageMapExportUtil.h
@@ -49,16 +49,17 @@ struct FIrisPackageMapExportsUtil
 	typedef FIrisPackageMapExportsQuantizedType QuantizedType;
 
 	// We need some sort of limit here.
-	static constexpr uint32 MaxExports = 128U;
+	static constexpr uint32 MaxExports = 65536U;
 
 	// Matches NetSerializer functions
 	IRISCORE_API static void Serialize(FNetSerializationContext& Context, const QuantizedType& Value);
 	IRISCORE_API static void Deserialize(FNetSerializationContext& Context, QuantizedType& Value);
-	IRISCORE_API static void Quantize(FNetSerializationContext& Context, const UE::Net::FIrisPackageMapExports& PackageMapExport, TArrayView<const UE::Net::FNetToken> NetTokensPendingExport, QuantizedType& Target);
+	IRISCORE_API static void Quantize(FNetSerializationContext& Context, const UE::Net::FIrisPackageMapExports& PackageMapExports, TArrayView<const UE::Net::FNetToken> NetTokensPendingExport, QuantizedType& Target);
 	IRISCORE_API static void Dequantize(FNetSerializationContext& Context, const QuantizedType& Source, UE::Net::FIrisPackageMapExports& PackageMapExports);
 	IRISCORE_API static bool IsEqual(FNetSerializationContext& Context, const QuantizedType& Value0, const QuantizedType& Value1);
 	IRISCORE_API static void CloneDynamicState(FNetSerializationContext& Context, QuantizedType& Target, const QuantizedType& Source);
 	IRISCORE_API static void FreeDynamicState(FNetSerializationContext& Context, QuantizedType& Value);
+	IRISCORE_API static void FreeDynamicState(QuantizedType& Value);
 	IRISCORE_API static void CollectNetReferences(FNetSerializationContext& Context, const QuantizedType& Value, const FNetSerializerChangeMaskParam& ChangeMaskInfo, FNetReferenceCollector& Collector);
 	IRISCORE_API static bool Validate(FNetSerializationContext&, const QuantizedType& Value);
 
diff --git a/Engine/Source/Runtime/Net/Core/Private/Net/Core/Trace/NetTraceInternal.cpp b/Engine/Source/Runtime/Net/Core/Private/Net/Core/Trace/NetTraceInternal.cpp
--- a/Engine/Source/Runtime/Net/Core/Private/Net/Core/Trace/NetTraceInternal.cpp
+++ b/Engine/Source/Runtime/Net/Core/Private/Net/Core/Trace/NetTraceInternal.cpp
@@ -219,7 +219,7 @@ void FNetTrace::DestroyTraceCollector(FNetTraceCollector* Collector)
 	}
 }
 
-void FNetTrace::FoldTraceCollector(FNetTraceCollector* DstCollector, const FNetTraceCollector* SrcCollector, uint32 Offset)
+void FNetTrace::FoldTraceCollector(FNetTraceCollector* DstCollector, const FNetTraceCollector* SrcCollector, uint32 Offset, uint32 LevelOffset)
 {
 	if (DstCollector && SrcCollector && DstCollector != SrcCollector)
 	{
@@ -227,7 +227,7 @@ void FNetTrace::FoldTraceCollector(FNetTraceCollector* DstCollector, const FNetT
 		check(SrcCollector->BunchEventCount == 0 && SrcCollector->CurrentNestingLevel == 0);
 
 		// When we fold non-bunch events we inject them at the current level
-		const uint32 Level = DstCollector->CurrentNestingLevel;		
+		const uint32 Level = DstCollector->CurrentNestingLevel + LevelOffset;
 		
 		// Make sure that the events fit
 		if (SrcCollector->EventCount + DstCollector->EventCount > (uint32)DstCollector->Events.Num())
diff --git a/Engine/Source/Runtime/Net/Core/Public/Net/Core/Trace/Private/NetTraceInternal.h b/Engine/Source/Runtime/Net/Core/Public/Net/Core/Trace/Private/NetTraceInternal.h
--- a/Engine/Source/Runtime/Net/Core/Public/Net/Core/Trace/Private/NetTraceInternal.h
+++ b/Engine/Source/Runtime/Net/Core/Public/Net/Core/Trace/Private/NetTraceInternal.h
@@ -196,7 +196,7 @@ struct FNetTrace
 	NETCORE_API static void DestroyTraceCollector(FNetTraceCollector* Collector);
 
 	/** Folds events into DstCollector */
-	NETCORE_API static void FoldTraceCollector(FNetTraceCollector* DstCollector, const FNetTraceCollector* SrcCollector, uint32 Offset);
+	NETCORE_API static void FoldTraceCollector(FNetTraceCollector* DstCollector, const FNetTraceCollector* SrcCollector, uint32 Offset, uint32 LevelOffset = 0U);
 	
 	/** Folds bunch into DstCollector, if the src and dst are the same bunch it is assumed that a bunch already is open and that it will be closed */
 	NETCORE_API static void FoldBunchCollector(FNetTraceCollector* DstCollector, const FNetTraceCollector* SrcCollector);
