Disable Async IO for IAS for now due to a potential crash bug.

#rnx

- This change should be temporary while we hunt down the cause of a crash that occurs if/when a user loses connection to the IAS CDNs

diff --git a/Engine/Source/Runtime/Experimental/IoStore/HttpClient/Private/Overlapped.inl b/Engine/Source/Runtime/Experimental/IoStore/HttpClient/Private/Overlapped.inl
--- a/Engine/Source/Runtime/Experimental/IoStore/HttpClient/Private/Overlapped.inl
+++ b/Engine/Source/Runtime/Experimental/IoStore/HttpClient/Private/Overlapped.inl
@@ -5,7 +5,7 @@
 namespace UE::IoStore::HTTP
 {
 
-#if 0
+#if 1
 
 using SocketType = SOCKET;
 
