Dataflow:
- Added Convert nodes using all the any types
- deprecated old set of convert nodes
- Bug fixes

diff --git a/Engine/Plugins/Experimental/GeometryCollectionPlugin/Source/GeometryCollectionNodes/Private/Dataflow/DataflowNodeAndPinTypeColors.cpp b/Engine/Plugins/Experimental/GeometryCollectionPlugin/Source/GeometryCollectionNodes/Private/Dataflow/DataflowNodeAndPinTypeColors.cpp
--- a/Engine/Plugins/Experimental/GeometryCollectionPlugin/Source/GeometryCollectionNodes/Private/Dataflow/DataflowNodeAndPinTypeColors.cpp
+++ b/Engine/Plugins/Experimental/GeometryCollectionPlugin/Source/GeometryCollectionNodes/Private/Dataflow/DataflowNodeAndPinTypeColors.cpp
@@ -35,6 +35,7 @@ namespace UE::Dataflow
 		DATAFLOW_NODE_REGISTER_CREATION_FACTORY_NODE_COLORS_BY_CATEGORY("GeometryCollection|Selection", FLinearColor(1.f, 1.f, 0.05f), CDefaultNodeBodyTintColor);
 		DATAFLOW_NODE_REGISTER_CREATION_FACTORY_NODE_COLORS_BY_CATEGORY("GeometryCollection|Texture", FLinearColor(.1f, 1.f, 0.6f), CDefaultNodeBodyTintColor);
 		DATAFLOW_NODE_REGISTER_CREATION_FACTORY_NODE_COLORS_BY_CATEGORY("GeometryCollection|UV", FLinearColor(.1f, 1.f, 0.6f), CDefaultNodeBodyTintColor);
+		DATAFLOW_NODE_REGISTER_CREATION_FACTORY_NODE_COLORS_BY_CATEGORY("Convert", FLinearColor(1.f, 1.f, 0.1f), CDefaultNodeBodyTintColor);
 
 		// PinType colors
 		DATAFLOW_NODE_REGISTER_CREATION_FACTORY_PIN_SETTINGS_BY_TYPE("FManagedArrayCollection", FLinearColor(0.353393f, 0.454175f, 1.0f, 1.0f), CDefaultWireThickness);
diff --git a/Engine/Plugins/Experimental/GeometryCollectionPlugin/Source/GeometryCollectionNodes/Public/Dataflow/GeometryCollectionConversionNodes.h b/Engine/Plugins/Experimental/GeometryCollectionPlugin/Source/GeometryCollectionNodes/Public/Dataflow/GeometryCollectionConversionNodes.h
--- a/Engine/Plugins/Experimental/GeometryCollectionPlugin/Source/GeometryCollectionNodes/Public/Dataflow/GeometryCollectionConversionNodes.h
+++ b/Engine/Plugins/Experimental/GeometryCollectionPlugin/Source/GeometryCollectionNodes/Public/Dataflow/GeometryCollectionConversionNodes.h
@@ -18,7 +18,7 @@ namespace UE_DEPRECATED(5.5, "Use UE::Dataflow instead.") Dataflow {}
  * Description for this node
  *
  */
-USTRUCT(meta = (DataflowGeometryCollection))
+USTRUCT(meta = (DataflowGeometryCollection, Deprecated = "5.6"))
 struct FVectorToStringDataflowNode : public FDataflowNode
 {
 	GENERATED_USTRUCT_BODY()
@@ -47,7 +47,7 @@ public:
  * Description for this node
  *
  */
-USTRUCT(meta = (DataflowGeometryCollection))
+USTRUCT(meta = (DataflowGeometryCollection, Deprecated = "5.6"))
 struct FFloatToStringDataflowNode : public FDataflowNode
 {
 	GENERATED_USTRUCT_BODY()
@@ -76,7 +76,7 @@ public:
  * Converts an Int to a String
  *
  */
-USTRUCT(meta = (DataflowGeometryCollection))
+USTRUCT(meta = (DataflowGeometryCollection, Deprecated = "5.6"))
 struct FIntToStringDataflowNode : public FDataflowNode
 {
 	GENERATED_USTRUCT_BODY()
@@ -105,7 +105,7 @@ public:
  * Converts a Bool to a String in a form of ("true", "false")
  *
  */
-USTRUCT(meta = (DataflowGeometryCollection))
+USTRUCT(meta = (DataflowGeometryCollection, Deprecated = "5.6"))
 struct FBoolToStringDataflowNode : public FDataflowNode
 {
 	GENERATED_USTRUCT_BODY()
@@ -134,7 +134,7 @@ public:
  * Converts an Int to a Float
  *
  */
-USTRUCT(meta = (DataflowGeometryCollection))
+USTRUCT(meta = (DataflowGeometryCollection, Deprecated = "5.6"))
 struct FIntToFloatDataflowNode : public FDataflowNode
 {
 	GENERATED_USTRUCT_BODY()
@@ -162,7 +162,7 @@ public:
  * Converts an Int to a Double
  *
  */
-USTRUCT(meta = (DataflowGeometryCollection))
+USTRUCT(meta = (DataflowGeometryCollection, Deprecated = "5.6"))
 struct FIntToDoubleDataflowNode : public FDataflowNode
 {
 	GENERATED_USTRUCT_BODY()
@@ -190,7 +190,7 @@ public:
  * Converts an Float to a Double
  *
  */
-USTRUCT(meta = (DataflowGeometryCollection))
+USTRUCT(meta = (DataflowGeometryCollection, Deprecated = "5.6"))
 struct FFloatToDoubleDataflowNode : public FDataflowNode
 {
 	GENERATED_USTRUCT_BODY()
@@ -231,7 +231,7 @@ enum class EFloatToIntFunctionEnum : uint8
  * Converts a Float to Int using the specified method
  *
  */
-USTRUCT(meta = (DataflowGeometryCollection))
+USTRUCT(meta = (DataflowGeometryCollection, Deprecated = "5.6"))
 struct FFloatToIntDataflowNode : public FDataflowNode
 {
 	GENERATED_USTRUCT_BODY()
@@ -266,7 +266,7 @@ public:
  * Converts an Int to a Bool
  *
  */
-USTRUCT(meta = (DataflowGeometryCollection))
+USTRUCT(meta = (DataflowGeometryCollection, Deprecated = "5.6"))
 struct FIntToBoolDataflowNode : public FDataflowNode
 {
 	GENERATED_USTRUCT_BODY()
@@ -295,7 +295,7 @@ public:
  * Converts a Bool to an Int
  *
  */
-USTRUCT(meta = (DataflowGeometryCollection))
+USTRUCT(meta = (DataflowGeometryCollection, Deprecated = "5.6"))
 struct FBoolToIntDataflowNode : public FDataflowNode
 {
 	GENERATED_USTRUCT_BODY()
diff --git a/Engine/Source/Runtime/Experimental/Dataflow/Core/Private/Dataflow/DataflowAnyType.cpp b/Engine/Source/Runtime/Experimental/Dataflow/Core/Private/Dataflow/DataflowAnyType.cpp
--- a/Engine/Source/Runtime/Experimental/Dataflow/Core/Private/Dataflow/DataflowAnyType.cpp
+++ b/Engine/Source/Runtime/Experimental/Dataflow/Core/Private/Dataflow/DataflowAnyType.cpp
@@ -17,6 +17,7 @@ namespace UE::Dataflow
 		UE_DATAFLOW_REGISTER_ANYTYPE(FDataflowStringConvertibleTypes);
 		UE_DATAFLOW_REGISTER_ANYTYPE(FDataflowUObjectConvertibleTypes);
 		UE_DATAFLOW_REGISTER_ANYTYPE(FDataflowSelectionTypes);
+		UE_DATAFLOW_REGISTER_ANYTYPE(FDataflowRotationTypes);
 
 		// array types 
 		UE_DATAFLOW_REGISTER_ANYTYPE(FDataflowVectorArrayTypes);
diff --git a/Engine/Source/Runtime/Experimental/Dataflow/Core/Private/Dataflow/DataflowConvertNodes.cpp b/Engine/Source/Runtime/Experimental/Dataflow/Core/Private/Dataflow/DataflowConvertNodes.cpp
new file mode 100644
--- /dev/null
+++ b/Engine/Source/Runtime/Experimental/Dataflow/Core/Private/Dataflow/DataflowConvertNodes.cpp
@@ -0,0 +1,290 @@
+// Copyright Epic Games, Inc. All Rights Reserved.
+
+#include "Dataflow/DataflowConvertNodes.h"
+#include "Dataflow/DataflowNode.h"
+#include "Dataflow/DataflowNodeFactory.h"
+#include "Dataflow/DataflowNodeColorsRegistry.h"
+
+#include "MathUtil.h"
+
+namespace UE::Dataflow
+{
+	void RegisterDataflowConvertNodes()
+	{
+		// Convert
+		DATAFLOW_NODE_REGISTER_CREATION_FACTORY(FConvertNumericTypesDataflowNode);
+		DATAFLOW_NODE_REGISTER_CREATION_FACTORY(FConvertVectorTypesDataflowNode);
+		DATAFLOW_NODE_REGISTER_CREATION_FACTORY(FConvertStringTypesDataflowNode);
+		DATAFLOW_NODE_REGISTER_CREATION_FACTORY(FConvertBoolTypesDataflowNode);
+		DATAFLOW_NODE_REGISTER_CREATION_FACTORY(FConvertTransformTypesDataflowNode);
+		DATAFLOW_NODE_REGISTER_CREATION_FACTORY(FConvertStringConvertibleTypesDataflowNode);
+		DATAFLOW_NODE_REGISTER_CREATION_FACTORY(FConvertUObjectConvertibleTypesDataflowNode);
+		DATAFLOW_NODE_REGISTER_CREATION_FACTORY(FConvertSelectionTypesDataflowNode);
+		DATAFLOW_NODE_REGISTER_CREATION_FACTORY(FConvertVectorArrayTypesDataflowNode);
+		DATAFLOW_NODE_REGISTER_CREATION_FACTORY(FConvertNumericArrayTypesDataflowNode);
+		DATAFLOW_NODE_REGISTER_CREATION_FACTORY(FConvertStringArrayTypesDataflowNode);
+		DATAFLOW_NODE_REGISTER_CREATION_FACTORY(FConvertBoolArrayTypesDataflowNode);
+		DATAFLOW_NODE_REGISTER_CREATION_FACTORY(FConvertTransformArrayTypesDataflowNode);
+		DATAFLOW_NODE_REGISTER_CREATION_FACTORY(FConvertRotationDataflowNode);
+
+		// Deprecated
+	}
+}
+
+//-----------------------------------------------------------------------------------------------
+
+FConvertNumericTypesDataflowNode::FConvertNumericTypesDataflowNode(const UE::Dataflow::FNodeParameters& InParam, FGuid InGuid)
+	: FDataflowNode(InParam, InGuid)
+{
+	RegisterInputConnection(&In);
+	RegisterOutputConnection(&Out);
+}
+
+void FConvertNumericTypesDataflowNode::Evaluate(UE::Dataflow::FContext& Context, const FDataflowOutput* Output) const
+{
+	if (Output->IsA(&Out))
+	{
+		const double InValue = GetValue(Context, &In);
+		SetValue(Context, InValue, &Out);
+	}
+}
+
+//-----------------------------------------------------------------------------------------------
+
+FConvertVectorTypesDataflowNode::FConvertVectorTypesDataflowNode(const UE::Dataflow::FNodeParameters& InParam, FGuid InGuid)
+	: FDataflowNode(InParam, InGuid)
+{
+	RegisterInputConnection(&In);
+	RegisterOutputConnection(&Out);
+}
+
+void FConvertVectorTypesDataflowNode::Evaluate(UE::Dataflow::FContext& Context, const FDataflowOutput* Output) const
+{
+	if (Output->IsA(&Out))
+	{
+		const FVector4 InValue = GetValue(Context, &In);
+		SetValue(Context, InValue, &Out);
+	}
+}
+
+//-----------------------------------------------------------------------------------------------
+
+FConvertStringTypesDataflowNode::FConvertStringTypesDataflowNode(const UE::Dataflow::FNodeParameters& InParam, FGuid InGuid)
+	: FDataflowNode(InParam, InGuid)
+{
+	RegisterInputConnection(&In);
+	RegisterOutputConnection(&Out);
+}
+
+void FConvertStringTypesDataflowNode::Evaluate(UE::Dataflow::FContext& Context, const FDataflowOutput* Output) const
+{
+	if (Output->IsA(&Out))
+	{
+		const FString InValue = GetValue(Context, &In);
+		SetValue(Context, InValue, &Out);
+	}
+}
+
+//-----------------------------------------------------------------------------------------------
+
+FConvertBoolTypesDataflowNode::FConvertBoolTypesDataflowNode(const UE::Dataflow::FNodeParameters& InParam, FGuid InGuid)
+	: FDataflowNode(InParam, InGuid)
+{
+	RegisterInputConnection(&In);
+	RegisterOutputConnection(&Out);
+}
+
+void FConvertBoolTypesDataflowNode::Evaluate(UE::Dataflow::FContext& Context, const FDataflowOutput* Output) const
+{
+	if (Output->IsA(&Out))
+	{
+		const bool InValue = GetValue(Context, &In);
+		SetValue(Context, InValue, &Out);
+	}
+}
+
+//-----------------------------------------------------------------------------------------------
+
+FConvertTransformTypesDataflowNode::FConvertTransformTypesDataflowNode(const UE::Dataflow::FNodeParameters& InParam, FGuid InGuid)
+	: FDataflowNode(InParam, InGuid)
+{
+	RegisterInputConnection(&In);
+	RegisterOutputConnection(&Out);
+}
+
+void FConvertTransformTypesDataflowNode::Evaluate(UE::Dataflow::FContext& Context, const FDataflowOutput* Output) const
+{
+	if (Output->IsA(&Out))
+	{
+		const FTransform InValue = GetValue(Context, &In);
+		SetValue(Context, InValue, &Out);
+	}
+}
+
+//-----------------------------------------------------------------------------------------------
+
+FConvertStringConvertibleTypesDataflowNode::FConvertStringConvertibleTypesDataflowNode(const UE::Dataflow::FNodeParameters& InParam, FGuid InGuid)
+	: FDataflowNode(InParam, InGuid)
+{
+	RegisterInputConnection(&In);
+	RegisterOutputConnection(&Out);
+}
+
+void FConvertStringConvertibleTypesDataflowNode::Evaluate(UE::Dataflow::FContext& Context, const FDataflowOutput* Output) const
+{
+	if (Output->IsA(&Out))
+	{
+		const FString InValue = GetValue(Context, &In);
+		SetValue(Context, InValue, &Out);
+	}
+}
+
+//-----------------------------------------------------------------------------------------------
+
+FConvertUObjectConvertibleTypesDataflowNode::FConvertUObjectConvertibleTypesDataflowNode(const UE::Dataflow::FNodeParameters& InParam, FGuid InGuid)
+	: FDataflowNode(InParam, InGuid)
+{
+	RegisterInputConnection(&In);
+	RegisterOutputConnection(&Out);
+}
+
+void FConvertUObjectConvertibleTypesDataflowNode::Evaluate(UE::Dataflow::FContext& Context, const FDataflowOutput* Output) const
+{
+	if (Output->IsA(&Out))
+	{
+		const TObjectPtr<UObject> InValue = GetValue(Context, &In);
+		SetValue(Context, InValue, &Out);
+	}
+}
+
+//-----------------------------------------------------------------------------------------------
+
+FConvertSelectionTypesDataflowNode::FConvertSelectionTypesDataflowNode(const UE::Dataflow::FNodeParameters& InParam, FGuid InGuid)
+	: FDataflowNode(InParam, InGuid)
+{
+	RegisterInputConnection(&In);
+	RegisterOutputConnection(&Out);
+}
+
+void FConvertSelectionTypesDataflowNode::Evaluate(UE::Dataflow::FContext& Context, const FDataflowOutput* Output) const
+{
+	if (Output->IsA(&Out))
+	{
+		const FDataflowSelection InValue = GetValue(Context, &In);
+		SetValue(Context, InValue, &Out);
+	}
+}
+
+//-----------------------------------------------------------------------------------------------
+
+FConvertVectorArrayTypesDataflowNode::FConvertVectorArrayTypesDataflowNode(const UE::Dataflow::FNodeParameters& InParam, FGuid InGuid)
+	: FDataflowNode(InParam, InGuid)
+{
+	RegisterInputConnection(&In);
+	RegisterOutputConnection(&Out);
+}
+
+void FConvertVectorArrayTypesDataflowNode::Evaluate(UE::Dataflow::FContext& Context, const FDataflowOutput* Output) const
+{
+	if (Output->IsA(&Out))
+	{
+		const TArray<FVector4> InValue = GetValue(Context, &In);
+		SetValue(Context, InValue, &Out);
+	}
+}
+
+//-----------------------------------------------------------------------------------------------
+
+FConvertNumericArrayTypesDataflowNode::FConvertNumericArrayTypesDataflowNode(const UE::Dataflow::FNodeParameters& InParam, FGuid InGuid)
+	: FDataflowNode(InParam, InGuid)
+{
+	RegisterInputConnection(&In);
+	RegisterOutputConnection(&Out);
+}
+
+void FConvertNumericArrayTypesDataflowNode::Evaluate(UE::Dataflow::FContext& Context, const FDataflowOutput* Output) const
+{
+	if (Output->IsA(&Out))
+	{
+		const TArray<double> InValue = GetValue(Context, &In);
+		SetValue(Context, InValue, &Out);
+	}
+}
+
+//-----------------------------------------------------------------------------------------------
+
+FConvertStringArrayTypesDataflowNode::FConvertStringArrayTypesDataflowNode(const UE::Dataflow::FNodeParameters& InParam, FGuid InGuid)
+	: FDataflowNode(InParam, InGuid)
+{
+	RegisterInputConnection(&In);
+	RegisterOutputConnection(&Out);
+}
+
+void FConvertStringArrayTypesDataflowNode::Evaluate(UE::Dataflow::FContext& Context, const FDataflowOutput* Output) const
+{
+	if (Output->IsA(&Out))
+	{
+		const TArray<FString> InValue = GetValue(Context, &In);
+		SetValue(Context, InValue, &Out);
+	}
+}
+
+//-----------------------------------------------------------------------------------------------
+
+FConvertBoolArrayTypesDataflowNode::FConvertBoolArrayTypesDataflowNode(const UE::Dataflow::FNodeParameters& InParam, FGuid InGuid)
+	: FDataflowNode(InParam, InGuid)
+{
+	RegisterInputConnection(&In);
+	RegisterOutputConnection(&Out);
+}
+
+void FConvertBoolArrayTypesDataflowNode::Evaluate(UE::Dataflow::FContext& Context, const FDataflowOutput* Output) const
+{
+	if (Output->IsA(&Out))
+	{
+		const TArray<bool> InValue = GetValue(Context, &In);
+		SetValue(Context, InValue, &Out);
+	}
+}
+
+//-----------------------------------------------------------------------------------------------
+
+FConvertTransformArrayTypesDataflowNode::FConvertTransformArrayTypesDataflowNode(const UE::Dataflow::FNodeParameters& InParam, FGuid InGuid)
+	: FDataflowNode(InParam, InGuid)
+{
+	RegisterInputConnection(&In);
+	RegisterOutputConnection(&Out);
+}
+
+void FConvertTransformArrayTypesDataflowNode::Evaluate(UE::Dataflow::FContext& Context, const FDataflowOutput* Output) const
+{
+	if (Output->IsA(&Out))
+	{
+		const TArray<FTransform> InValue = GetValue(Context, &In);
+		SetValue(Context, InValue, &Out);
+	}
+}
+
+//-----------------------------------------------------------------------------------------------
+
+FConvertRotationDataflowNode::FConvertRotationDataflowNode(const UE::Dataflow::FNodeParameters& InParam, FGuid InGuid)
+	: FDataflowNode(InParam, InGuid)
+{
+	RegisterInputConnection(&In);
+	RegisterOutputConnection(&Out);
+}
+
+void FConvertRotationDataflowNode::Evaluate(UE::Dataflow::FContext& Context, const FDataflowOutput* Output) const
+{
+	if (Output->IsA(&Out))
+	{
+		const FRotator InValue = GetValue(Context, &In);
+		SetValue(Context, InValue, &Out);
+	}
+}
+
+//-----------------------------------------------------------------------------------------------
+
+
+
+
diff --git a/Engine/Source/Runtime/Experimental/Dataflow/Core/Private/Dataflow/DataflowMathNodes.cpp b/Engine/Source/Runtime/Experimental/Dataflow/Core/Private/Dataflow/DataflowMathNodes.cpp
--- a/Engine/Source/Runtime/Experimental/Dataflow/Core/Private/Dataflow/DataflowMathNodes.cpp
+++ b/Engine/Source/Runtime/Experimental/Dataflow/Core/Private/Dataflow/DataflowMathNodes.cpp
@@ -836,3 +836,6 @@ double FDataflowMathRadToDegNode::ComputeResult(UE::Dataflow::FContext& Context,
 {
 	return FMath::RadiansToDegrees(InA);
 }
+
+//-----------------------------------------------------------------------------------------------
+
diff --git a/Engine/Source/Runtime/Experimental/Dataflow/Core/Private/Dataflow/DataflowNodeFactory.cpp b/Engine/Source/Runtime/Experimental/Dataflow/Core/Private/Dataflow/DataflowNodeFactory.cpp
--- a/Engine/Source/Runtime/Experimental/Dataflow/Core/Private/Dataflow/DataflowNodeFactory.cpp
+++ b/Engine/Source/Runtime/Experimental/Dataflow/Core/Private/Dataflow/DataflowNodeFactory.cpp
@@ -7,6 +7,7 @@
 #include "Dataflow/DataflowImageNodes.h"
 #include "Dataflow/DataflowMathNodes.h"
 #include "Dataflow/DataflowVectorNodes.h"
+#include "Dataflow/DataflowConvertNodes.h"
 #include "Misc/MessageDialog.h"
 #include "Logging/LogMacros.h"
 #include "Dataflow/DataflowCoreNodeAndPinTypeColors.h"
@@ -28,6 +29,7 @@ namespace UE::Dataflow
 		RegisterDataflowMathNodes();
 		RegisterDataflowVectorNodes();
 		RegisterDataflowImageNodes();
+		RegisterDataflowConvertNodes();
 	}
 
 	TSharedPtr<FDataflowNode> FNodeFactory::NewNodeFromRegisteredType(FGraph& Graph, const FNewNodeParameters& Param)
diff --git a/Engine/Source/Runtime/Experimental/Dataflow/Core/Public/Dataflow/DataflowAnyType.h b/Engine/Source/Runtime/Experimental/Dataflow/Core/Public/Dataflow/DataflowAnyType.h
--- a/Engine/Source/Runtime/Experimental/Dataflow/Core/Public/Dataflow/DataflowAnyType.h
+++ b/Engine/Source/Runtime/Experimental/Dataflow/Core/Public/Dataflow/DataflowAnyType.h
@@ -59,7 +59,7 @@ struct FDataflowVectorTypes : public FDataflowAnyType
 	FVector4 Value = FVector4(0, 0, 0, 0);
 };
 
-/** String types (FString or FName) */
+/** String types (FString or FName or FText) */
 USTRUCT()
 struct FDataflowStringTypes : public FDataflowAnyType
 {
@@ -205,6 +205,19 @@ struct FDataflowTransformArrayTypes : public FDataflowAnyType
 	TArray<FTransform> Value;
 };
 
+/** Rotation types */
+USTRUCT()
+struct FDataflowRotationTypes : public FDataflowAnyType
+{
+	using FPolicyType = FDataflowRotationTypePolicy;
+	using FStorageType = FRotator;
+
+	GENERATED_USTRUCT_BODY()
+
+	UPROPERTY()
+	FRotator Value;
+};
+
 namespace UE::Dataflow
 {
 	DATAFLOWCORE_API bool AreTypesCompatible(FName TypeA, FName TypeB);
diff --git a/Engine/Source/Runtime/Experimental/Dataflow/Core/Public/Dataflow/DataflowConvertNodes.h b/Engine/Source/Runtime/Experimental/Dataflow/Core/Public/Dataflow/DataflowConvertNodes.h
new file mode 100644
--- /dev/null
+++ b/Engine/Source/Runtime/Experimental/Dataflow/Core/Public/Dataflow/DataflowConvertNodes.h
@@ -0,0 +1,394 @@
+// Copyright Epic Games, Inc. All Rights Reserved.
+
+#pragma once
+
+//#include "ChaosLog.h"
+//#include "CoreMinimal.h"
+//#include "Dataflow/DataflowConnection.h"
+//#include "Dataflow/DataflowNodeParameters.h"
+#include "Dataflow/DataflowNode.h"
+
+#include "DataflowConvertNodes.generated.h"
+
+//--------------------------------------------------------------------------
+//
+// Convert nodes
+//
+//--------------------------------------------------------------------------
+
+#define DATAFLOW_CONVERT_NODES_CATEGORY "Convert"
+
+/**
+* Convert Numeric types
+* (double, float, int64, uint64, int32, uint32, int16, uint16, int8, uint8)
+*/
+USTRUCT()
+struct FConvertNumericTypesDataflowNode : public FDataflowNode
+{
+	GENERATED_USTRUCT_BODY()
+	DATAFLOW_NODE_DEFINE_INTERNAL(FConvertNumericTypesDataflowNode, "ConvertNumericTypes", DATAFLOW_CONVERT_NODES_CATEGORY, "")
+
+private:
+	/** Input value */
+	UPROPERTY(meta = (DataflowInput, DataflowIntrinsic))
+	FDataflowNumericTypes In;
+
+	/** Output value */
+	UPROPERTY(meta = (DataflowOutput))
+	FDataflowNumericTypes Out;
+
+	//~ Begin FDataflowNode interface
+	virtual void Evaluate(UE::Dataflow::FContext& Context, const FDataflowOutput* Out) const override;
+	//~ End FDataflowNode interface
+
+public:
+	FConvertNumericTypesDataflowNode(const UE::Dataflow::FNodeParameters& InParam, FGuid InGuid = FGuid::NewGuid());
+};
+
+/**
+* Convert Vector types
+* (2D, 3D and 4D vector, single and double precision)
+*/
+USTRUCT()
+struct FConvertVectorTypesDataflowNode : public FDataflowNode
+{
+	GENERATED_USTRUCT_BODY()
+	DATAFLOW_NODE_DEFINE_INTERNAL(FConvertVectorTypesDataflowNode, "ConvertVectorTypes", DATAFLOW_CONVERT_NODES_CATEGORY, "")
+
+private:
+	/** Input value */
+	UPROPERTY(meta = (DataflowInput, DataflowIntrinsic))
+	FDataflowVectorTypes In;
+
+	/** Output value */
+	UPROPERTY(meta = (DataflowOutput))
+	FDataflowVectorTypes Out;
+
+	//~ Begin FDataflowNode interface
+	virtual void Evaluate(UE::Dataflow::FContext& Context, const FDataflowOutput* Out) const override;
+	//~ End FDataflowNode interface
+
+public:
+	FConvertVectorTypesDataflowNode(const UE::Dataflow::FNodeParameters& InParam, FGuid InGuid = FGuid::NewGuid());
+};
+
+/**
+* Convert String types 
+* (FString or FName or FText)
+*/
+USTRUCT()
+struct FConvertStringTypesDataflowNode : public FDataflowNode
+{
+	GENERATED_USTRUCT_BODY()
+	DATAFLOW_NODE_DEFINE_INTERNAL(FConvertStringTypesDataflowNode, "ConvertStringTypes", DATAFLOW_CONVERT_NODES_CATEGORY, "")
+
+private:
+	/** Input value */
+	UPROPERTY(meta = (DataflowInput, DataflowIntrinsic))
+	FDataflowStringTypes In;
+
+	/** Output value */
+	UPROPERTY(meta = (DataflowOutput))
+	FDataflowStringTypes Out;
+
+	//~ Begin FDataflowNode interface
+	virtual void Evaluate(UE::Dataflow::FContext& Context, const FDataflowOutput* Out) const override;
+	//~ End FDataflowNode interface
+
+public:
+	FConvertStringTypesDataflowNode(const UE::Dataflow::FNodeParameters& InParam, FGuid InGuid = FGuid::NewGuid());
+};
+
+/**
+* Convert Bool types
+*/
+USTRUCT()
+struct FConvertBoolTypesDataflowNode : public FDataflowNode
+{
+	GENERATED_USTRUCT_BODY()
+	DATAFLOW_NODE_DEFINE_INTERNAL(FConvertBoolTypesDataflowNode, "ConvertBoolTypes", DATAFLOW_CONVERT_NODES_CATEGORY, "")
+
+private:
+	/** Input value */
+	UPROPERTY(meta = (DataflowInput, DataflowIntrinsic))
+	FDataflowBoolTypes In;
+
+	/** Output value */
+	UPROPERTY(meta = (DataflowOutput))
+	FDataflowBoolTypes Out;
+
+	//~ Begin FDataflowNode interface
+	virtual void Evaluate(UE::Dataflow::FContext& Context, const FDataflowOutput* Out) const override;
+	//~ End FDataflowNode interface
+
+public:
+	FConvertBoolTypesDataflowNode(const UE::Dataflow::FNodeParameters& InParam, FGuid InGuid = FGuid::NewGuid());
+};
+
+/**
+* Convert Transform types
+*/
+USTRUCT()
+struct FConvertTransformTypesDataflowNode : public FDataflowNode
+{
+	GENERATED_USTRUCT_BODY()
+	DATAFLOW_NODE_DEFINE_INTERNAL(FConvertTransformTypesDataflowNode, "ConvertTransformTypes", DATAFLOW_CONVERT_NODES_CATEGORY, "")
+
+private:
+	/** Input value */
+	UPROPERTY(meta = (DataflowInput, DataflowIntrinsic))
+	FDataflowTransformTypes In;
+
+	/** Output value */
+	UPROPERTY(meta = (DataflowOutput))
+	FDataflowTransformTypes Out;
+
+	//~ Begin FDataflowNode interface
+	virtual void Evaluate(UE::Dataflow::FContext& Context, const FDataflowOutput* Out) const override;
+	//~ End FDataflowNode interface
+
+public:
+	FConvertTransformTypesDataflowNode(const UE::Dataflow::FNodeParameters& InParam, FGuid InGuid = FGuid::NewGuid());
+};
+
+/**
+* Convert String convertible types
+* (String types, Numeric types, Vector types and Booleans)
+*/
+USTRUCT()
+struct FConvertStringConvertibleTypesDataflowNode : public FDataflowNode
+{
+	GENERATED_USTRUCT_BODY()
+	DATAFLOW_NODE_DEFINE_INTERNAL(FConvertStringConvertibleTypesDataflowNode, "ConvertStringConvertibleTypes", DATAFLOW_CONVERT_NODES_CATEGORY, "")
+
+private:
+	/** Input value */
+	UPROPERTY(meta = (DataflowInput, DataflowIntrinsic))
+	FDataflowStringConvertibleTypes In;
+
+	/** Output value */
+	UPROPERTY(meta = (DataflowOutput))
+	FDataflowStringConvertibleTypes Out;
+
+	//~ Begin FDataflowNode interface
+	virtual void Evaluate(UE::Dataflow::FContext& Context, const FDataflowOutput* Out) const override;
+	//~ End FDataflowNode interface
+
+public:
+	FConvertStringConvertibleTypesDataflowNode(const UE::Dataflow::FNodeParameters& InParam, FGuid InGuid = FGuid::NewGuid());
+};
+
+/**
+* Convert UObject types
+*/
+USTRUCT()
+struct FConvertUObjectConvertibleTypesDataflowNode : public FDataflowNode
+{
+	GENERATED_USTRUCT_BODY()
+	DATAFLOW_NODE_DEFINE_INTERNAL(FConvertUObjectConvertibleTypesDataflowNode, "ConvertUObjectConvertibleTypes", DATAFLOW_CONVERT_NODES_CATEGORY, "")
+
+private:
+	/** Input value */
+	UPROPERTY(meta = (DataflowInput, DataflowIntrinsic))
+	FDataflowUObjectConvertibleTypes In;
+
+	/** Output value */
+	UPROPERTY(meta = (DataflowOutput))
+	FDataflowUObjectConvertibleTypes Out;
+
+	//~ Begin FDataflowNode interface
+	virtual void Evaluate(UE::Dataflow::FContext& Context, const FDataflowOutput* Out) const override;
+	//~ End FDataflowNode interface
+
+public:
+	FConvertUObjectConvertibleTypesDataflowNode(const UE::Dataflow::FNodeParameters& InParam, FGuid InGuid = FGuid::NewGuid());
+};
+
+/**
+* Convert Selection types
+*/
+USTRUCT()
+struct FConvertSelectionTypesDataflowNode : public FDataflowNode
+{
+	GENERATED_USTRUCT_BODY()
+	DATAFLOW_NODE_DEFINE_INTERNAL(FConvertSelectionTypesDataflowNode, "ConvertSelectionTypes", DATAFLOW_CONVERT_NODES_CATEGORY, "")
+
+private:
+	/** Input value */
+	UPROPERTY(meta = (DataflowInput, DataflowIntrinsic))
+	FDataflowSelectionTypes In;
+
+	/** Output value */
+	UPROPERTY(meta = (DataflowOutput))
+	FDataflowSelectionTypes Out;
+
+	//~ Begin FDataflowNode interface
+	virtual void Evaluate(UE::Dataflow::FContext& Context, const FDataflowOutput* Out) const override;
+	//~ End FDataflowNode interface
+
+public:
+	FConvertSelectionTypesDataflowNode(const UE::Dataflow::FNodeParameters& InParam, FGuid InGuid = FGuid::NewGuid());
+};
+
+/**
+* Convert Vector array types
+*/
+USTRUCT()
+struct FConvertVectorArrayTypesDataflowNode : public FDataflowNode
+{
+	GENERATED_USTRUCT_BODY()
+	DATAFLOW_NODE_DEFINE_INTERNAL(FConvertVectorArrayTypesDataflowNode, "ConvertVectorArrayTypes", DATAFLOW_CONVERT_NODES_CATEGORY, "")
+
+private:
+	/** Input value */
+	UPROPERTY(meta = (DataflowInput, DataflowIntrinsic))
+	FDataflowVectorArrayTypes In;
+
+	/** Output value */
+	UPROPERTY(meta = (DataflowOutput))
+	FDataflowVectorArrayTypes Out;
+
+	//~ Begin FDataflowNode interface
+	virtual void Evaluate(UE::Dataflow::FContext& Context, const FDataflowOutput* Out) const override;
+	//~ End FDataflowNode interface
+
+public:
+	FConvertVectorArrayTypesDataflowNode(const UE::Dataflow::FNodeParameters& InParam, FGuid InGuid = FGuid::NewGuid());
+};
+
+/**
+* Convert Numeric array types
+*/
+USTRUCT()
+struct FConvertNumericArrayTypesDataflowNode : public FDataflowNode
+{
+	GENERATED_USTRUCT_BODY()
+	DATAFLOW_NODE_DEFINE_INTERNAL(FConvertNumericArrayTypesDataflowNode, "ConvertNumericArrayTypes", DATAFLOW_CONVERT_NODES_CATEGORY, "")
+
+private:
+	/** Input value */
+	UPROPERTY(meta = (DataflowInput, DataflowIntrinsic))
+	FDataflowNumericArrayTypes In;
+
+	/** Output value */
+	UPROPERTY(meta = (DataflowOutput))
+	FDataflowNumericArrayTypes Out;
+
+	//~ Begin FDataflowNode interface
+	virtual void Evaluate(UE::Dataflow::FContext& Context, const FDataflowOutput* Out) const override;
+	//~ End FDataflowNode interface
+
+public:
+	FConvertNumericArrayTypesDataflowNode(const UE::Dataflow::FNodeParameters& InParam, FGuid InGuid = FGuid::NewGuid());
+};
+
+/**
+* Convert String array types
+*/
+USTRUCT()
+struct FConvertStringArrayTypesDataflowNode : public FDataflowNode
+{
+	GENERATED_USTRUCT_BODY()
+	DATAFLOW_NODE_DEFINE_INTERNAL(FConvertStringArrayTypesDataflowNode, "ConvertStringArrayTypes", DATAFLOW_CONVERT_NODES_CATEGORY, "")
+
+private:
+	/** Input value */
+	UPROPERTY(meta = (DataflowInput, DataflowIntrinsic))
+	FDataflowStringArrayTypes In;
+
+	/** Output value */
+	UPROPERTY(meta = (DataflowOutput))
+	FDataflowStringArrayTypes Out;
+
+	//~ Begin FDataflowNode interface
+	virtual void Evaluate(UE::Dataflow::FContext& Context, const FDataflowOutput* Out) const override;
+	//~ End FDataflowNode interface
+
+public:
+	FConvertStringArrayTypesDataflowNode(const UE::Dataflow::FNodeParameters& InParam, FGuid InGuid = FGuid::NewGuid());
+};
+
+/**
+* Convert Bool array types
+*/
+USTRUCT()
+struct FConvertBoolArrayTypesDataflowNode : public FDataflowNode
+{
+	GENERATED_USTRUCT_BODY()
+	DATAFLOW_NODE_DEFINE_INTERNAL(FConvertBoolArrayTypesDataflowNode, "ConvertBoolArrayTypes", DATAFLOW_CONVERT_NODES_CATEGORY, "")
+
+private:
+	/** Input value */
+	UPROPERTY(meta = (DataflowInput, DataflowIntrinsic))
+	FDataflowBoolArrayTypes In;
+
+	/** Output value */
+	UPROPERTY(meta = (DataflowOutput))
+	FDataflowBoolArrayTypes Out;
+
+	//~ Begin FDataflowNode interface
+	virtual void Evaluate(UE::Dataflow::FContext& Context, const FDataflowOutput* Out) const override;
+	//~ End FDataflowNode interface
+
+public:
+	FConvertBoolArrayTypesDataflowNode(const UE::Dataflow::FNodeParameters& InParam, FGuid InGuid = FGuid::NewGuid());
+};
+
+/**
+* Convert Transform array types
+*/
+USTRUCT()
+struct FConvertTransformArrayTypesDataflowNode : public FDataflowNode
+{
+	GENERATED_USTRUCT_BODY()
+	DATAFLOW_NODE_DEFINE_INTERNAL(FConvertTransformArrayTypesDataflowNode, "ConvertTransformArrayTypes", DATAFLOW_CONVERT_NODES_CATEGORY, "")
+
+private:
+	/** Input value */
+	UPROPERTY(meta = (DataflowInput, DataflowIntrinsic))
+	FDataflowTransformArrayTypes In;
+
+	/** Output value */
+	UPROPERTY(meta = (DataflowOutput))
+	FDataflowTransformArrayTypes Out;
+
+	//~ Begin FDataflowNode interface
+	virtual void Evaluate(UE::Dataflow::FContext& Context, const FDataflowOutput* Out) const override;
+	//~ End FDataflowNode interface
+
+public:
+	FConvertTransformArrayTypesDataflowNode(const UE::Dataflow::FNodeParameters& InParam, FGuid InGuid = FGuid::NewGuid());
+};
+
+/**
+* Convert Rotation
+* (FQuat, FRotator, FVector)
+*/
+USTRUCT()
+struct FConvertRotationDataflowNode : public FDataflowNode
+{
+	GENERATED_USTRUCT_BODY()
+	DATAFLOW_NODE_DEFINE_INTERNAL(FConvertRotationDataflowNode, "ConvertRotation", DATAFLOW_CONVERT_NODES_CATEGORY, "")
+
+private:
+	/** Input value */
+	UPROPERTY(meta = (DataflowInput, DataflowIntrinsic))
+	FDataflowRotationTypes In;
+
+	/** Output value */
+	UPROPERTY(meta = (DataflowOutput))
+	FDataflowRotationTypes Out;
+
+	//~ Begin FDataflowNode interface
+	virtual void Evaluate(UE::Dataflow::FContext& Context, const FDataflowOutput* Out) const override;
+	//~ End FDataflowNode interface
+
+public:
+	FConvertRotationDataflowNode(const UE::Dataflow::FNodeParameters& InParam, FGuid InGuid = FGuid::NewGuid());
+};
+
+namespace UE::Dataflow
+{
+	void RegisterDataflowConvertNodes();
+}
+
diff --git a/Engine/Source/Runtime/Experimental/Dataflow/Core/Public/Dataflow/DataflowTypePolicy.h b/Engine/Source/Runtime/Experimental/Dataflow/Core/Public/Dataflow/DataflowTypePolicy.h
--- a/Engine/Source/Runtime/Experimental/Dataflow/Core/Public/Dataflow/DataflowTypePolicy.h
+++ b/Engine/Source/Runtime/Experimental/Dataflow/Core/Public/Dataflow/DataflowTypePolicy.h
@@ -57,6 +57,7 @@ UE_DATAFLOW_POLICY_DECLARE_TYPENAME(int64)
 UE_DATAFLOW_POLICY_DECLARE_TYPENAME(float)
 UE_DATAFLOW_POLICY_DECLARE_TYPENAME(double)
 UE_DATAFLOW_POLICY_DECLARE_TYPENAME(FName)
+UE_DATAFLOW_POLICY_DECLARE_TYPENAME(FText)
 UE_DATAFLOW_POLICY_DECLARE_TYPENAME(FString)
 UE_DATAFLOW_POLICY_DECLARE_TYPENAME(TObjectPtr<UObject>)
 
@@ -223,7 +224,7 @@ struct FDataflowVectorArrayPolicy:
 };
 
 struct FDataflowStringTypePolicy :
-	public TDataflowMultiTypePolicy<FString, FName>
+	public TDataflowMultiTypePolicy<FString, FName, FText>
 {
 };
 
@@ -237,6 +238,11 @@ struct FDataflowStringArrayPolicy :
 {
 };
 
+struct FDataflowRotationTypePolicy :
+	public TDataflowMultiTypePolicy<FVector, FQuat, FRotator>
+{
+};
+
 /**
 * string comvertible types
 * - FString / Fname
@@ -370,6 +376,10 @@ struct FDataflowConverter<FString>
 		{
 			To = From.ToString();
 		}
+		else if constexpr (std::is_same_v<TFromType, FText>)
+		{
+			To = From.ToString();
+		}
 		else if constexpr (std::is_same_v<TFromType, bool>)
 		{
 			To = FString((From == true) ? "True" : "False");
@@ -395,6 +405,10 @@ struct FDataflowConverter<FString>
 		{
 			To = FName(From);
 		}
+		else if constexpr (std::is_same_v<TToType, FText>)
+		{
+			To = FText::FromString(From);
+		}
 		else if constexpr (std::is_same_v<TToType, bool>)
 		{
 			To = From.ToBool();
@@ -561,3 +575,41 @@ struct FDataflowConverter<TArray<ArrayType>>
 		}
 	}
 };
+
+template <>
+struct FDataflowConverter<FRotator>
+{
+	template <typename TFromType>
+	static void From(const TFromType& From, FRotator& To)
+	{
+		if constexpr (std::is_same_v<TFromType, FQuat>)
+		{
+			To = FRotator(From);
+		}
+		else if constexpr (std::is_same_v<TFromType, FVector>)
+		{
+			To = FRotator::MakeFromEuler(From);
+		}
+		else
+		{
+			To = From;
+		}
+	}
+
+	template <typename TToType>
+	static void To(const FRotator& From, TToType& To)
+	{
+		if constexpr (std::is_same_v<TToType, FQuat>)
+		{
+			To = FQuat::MakeFromRotator(From);
+		}
+		else if constexpr (std::is_same_v<TToType, FVector>)
+		{
+			To = FQuat::MakeFromRotator(From).Euler();
+		}
+		else
+		{
+			To = From;
+		}
+	}
+};
