Geometry Collections: Ensure object existence in post-physics callback

diff --git a/Engine/Source/Runtime/Experimental/GeometryCollectionEngine/Private/GeometryCollection/GeometryCollectionComponent.cpp b/Engine/Source/Runtime/Experimental/GeometryCollectionEngine/Private/GeometryCollection/GeometryCollectionComponent.cpp
--- a/Engine/Source/Runtime/Experimental/GeometryCollectionEngine/Private/GeometryCollection/GeometryCollectionComponent.cpp
+++ b/Engine/Source/Runtime/Experimental/GeometryCollectionEngine/Private/GeometryCollection/GeometryCollectionComponent.cpp
@@ -3934,10 +3934,27 @@ PRAGMA_ENABLE_DEPRECATION_WARNINGS
 	}
 #endif
 	PhysicsProxy = new FGeometryCollectionPhysicsProxy(this, *DynamicCollection, SimulationParameters, InitialSimFilter, InitialQueryFilter, CollectorGuid);
-	PhysicsProxy->SetPostPhysicsSyncCallback([this]() { OnPostPhysicsSync(); });
-	PhysicsProxy->SetPostParticlesCreatedCallback([this]() { OnPostCreateParticles(); });
 
-	
+	PhysicsProxy->SetPostPhysicsSyncCallback(
+		[WeakThis = MakeWeakObjectPtr(this)]()
+		{
+			if (UGeometryCollectionComponent* This = WeakThis.Get())
+			{
+				This->OnPostPhysicsSync();
+			}
+		}
+	);
+
+	PhysicsProxy->SetPostParticlesCreatedCallback(
+		[WeakThis = MakeWeakObjectPtr(this)]()
+		{
+			if (UGeometryCollectionComponent* This = WeakThis.Get())
+			{
+				This->OnPostCreateParticles();
+			}
+		}
+	);
+
 	if (GetIsReplicated())
 	{
 		// using net mode and not local role because at this time in the initialization client and server both have an authority local role
