Dataflow:
- Fixed missing initialization

diff --git a/Engine/Source/Runtime/Experimental/Dataflow/Core/Public/Dataflow/DataflowAnyType.h b/Engine/Source/Runtime/Experimental/Dataflow/Core/Public/Dataflow/DataflowAnyType.h
--- a/Engine/Source/Runtime/Experimental/Dataflow/Core/Public/Dataflow/DataflowAnyType.h
+++ b/Engine/Source/Runtime/Experimental/Dataflow/Core/Public/Dataflow/DataflowAnyType.h
@@ -215,7 +215,7 @@ struct FDataflowRotationTypes : public FDataflowAnyType
 	GENERATED_USTRUCT_BODY()
 
 	UPROPERTY()
-	FRotator Value;
+	FRotator Value = FRotator(EForceInit::ForceInitToZero);
 };
 
 namespace UE::Dataflow
