No longer assume that all ias requests use the same hostgroup (although for now they do)

#rnx

- FChunkRequest now contains a FIASHostGroup to use with the request.
-- For now we assign the DefaultHostGroup to all requests.
- FOnDemandIoBackend::ProcessHttpRequests no longer takes a hostgroup, we use the hostgroup assigned to each request instead.
- Replaced FMultiEndpointHttpClient::UpdateConnection with ::UpdateConnections, this version will check all current host groups used by the client rather than checking one specific hostgroup being provided.

diff --git a/Engine/Source/Runtime/Experimental/IoStore/OnDemand/Private/OnDemandHttpClient.cpp b/Engine/Source/Runtime/Experimental/IoStore/OnDemand/Private/OnDemandHttpClient.cpp
--- a/Engine/Source/Runtime/Experimental/IoStore/OnDemand/Private/OnDemandHttpClient.cpp
+++ b/Engine/Source/Runtime/Experimental/IoStore/OnDemand/Private/OnDemandHttpClient.cpp
@@ -473,17 +473,21 @@ bool FMultiEndpointHttpClient::Tick(int32 WaitTimeMs, uint32 MaxKiBPerSecond)
 	return bIsIdle == false;
 }
 
-void FMultiEndpointHttpClient::UpdateConnection(const FOnDemandHostGroup& HostGroup)
+void FMultiEndpointHttpClient::UpdateConnections()
 {
-	if (FConnection* Connection = FindConnection(HostGroup))
+	UE_CHECK_THREADSAFETY(FMultiEndpointHttpClient::UpdateConnections);
+
+	for (TUniquePtr<FConnection>& Connection : Connections)
 	{
-		Connection->CurrentHost = HostGroup.PrimaryHostIndex();
+		check(Connection.IsValid());
+
+		Connection->CurrentHost = Connection->HostGroup.PrimaryHostIndex();
 
 		if (Connection->CurrentHost != INDEX_NONE)
 		{
 			if (Connection->Pools[Connection->CurrentHost].IsValid() == false)
 			{
-				Connection->Pools[Connection->CurrentHost] = CreateConnection(HostGroup.PrimaryHost());
+				Connection->Pools[Connection->CurrentHost] = CreateConnection(Connection->HostGroup.PrimaryHost());
 			}
 		}
 	}
diff --git a/Engine/Source/Runtime/Experimental/IoStore/OnDemand/Private/OnDemandHttpClient.h b/Engine/Source/Runtime/Experimental/IoStore/OnDemand/Private/OnDemandHttpClient.h
--- a/Engine/Source/Runtime/Experimental/IoStore/OnDemand/Private/OnDemandHttpClient.h
+++ b/Engine/Source/Runtime/Experimental/IoStore/OnDemand/Private/OnDemandHttpClient.h
@@ -137,7 +137,7 @@ public:
 	bool Tick(int32 WaitTimeMs, uint32 MaxKiBPerSecond);
 	bool Tick() { return Tick(-1, 0); }
 
-	void UpdateConnection(const FOnDemandHostGroup& HostGroup);
+	void UpdateConnections();
 
 private:
 	using FHttpConnectionPools = TArray<TUniquePtr<HTTP::FConnectionPool>, TInlineAllocator<4>>;
diff --git a/Engine/Source/Runtime/Experimental/IoStore/OnDemand/Private/OnDemandIoDispatcherBackend.cpp b/Engine/Source/Runtime/Experimental/IoStore/OnDemand/Private/OnDemandIoDispatcherBackend.cpp
--- a/Engine/Source/Runtime/Experimental/IoStore/OnDemand/Private/OnDemandIoDispatcherBackend.cpp
+++ b/Engine/Source/Runtime/Experimental/IoStore/OnDemand/Private/OnDemandIoDispatcherBackend.cpp
@@ -573,6 +573,7 @@ struct FChunkRequest
 	FChunkRequestParams Params;
 	FIoRequestImpl* RequestHead;
 	FIoRequestImpl* RequestTail;
+	FIASHostGroup HostGroup;
 	FIoBuffer Chunk;
 	uint64 StartTime;
 	int32 Priority;
@@ -716,6 +717,9 @@ class FOnDemandIoBackend final
 			ChunkRequestCount++;
 			Inflight.Add(Params.ChunkKey, ChunkRequest);
 
+			// Paranoid check to make sure that no host group has currently been assigned
+			check(ChunkRequest->HostGroup.GetHostUrls().IsEmpty());
+
 			return ChunkRequest;
 		}
 
@@ -853,7 +857,7 @@ private:
 	void CompleteCacheRead(FChunkRequest* ChunkRequest);
 	bool CreateDefaultHostGroup();
 	bool ResolveDistributedEndpoint(const FDistributedEndpointUrl& Url, TArray<FString>& OutUrls);
-	void ProcessHttpRequests(FMultiEndpointHttpClient& HttpClient, FIASHostGroup& HostGroup, int32 MaxConcurrentRequests);
+	void ProcessHttpRequests(FMultiEndpointHttpClient& HttpClient, int32 MaxConcurrentRequests);
 	int32 WaitForCompleteRequestTasks(float WaitTimeSeconds, float PollTimeSeconds);
 	void DrainHttpRequests();
 	void UpdateThreadPriorityIfNeeded();
@@ -1241,6 +1245,10 @@ bool FOnDemandIoBackend::Resolve(FIoRequestImpl* Request)
 		check(GetStatus == EIoErrorCode::NotFound);
 	}
 
+	// Currently all requests are assigned the DefaultHostGroup
+	// TODO: When next fixing this, make sure we deal with the !DefaultHostGroup.IsConnected() above
+	ChunkRequest->HostGroup = DefaultHostGroup;
+
 	Stats.OnHttpEnqueue();
 	HttpRequests.EnqueueByPriority(ChunkRequest);
 	TickBackendEvent->Trigger();
@@ -1509,12 +1517,14 @@ FOnDemandStreamingCacheUsage FOnDemandIoBackend::GetCacheUsage() const
 	return Usage; 
 }
 
-void FOnDemandIoBackend::ProcessHttpRequests(FMultiEndpointHttpClient& HttpClient, FIASHostGroup& HostGroup, int32 MaxConcurrentRequests)
+void FOnDemandIoBackend::ProcessHttpRequests(FMultiEndpointHttpClient& HttpClient, int32 MaxConcurrentRequests)
 {
 	int32 NumConcurrentRequests = 0;
 	FChunkRequest* NextChunkRequest = HttpRequests.Dequeue();
 
-	HttpClient.UpdateConnection(HostGroup.GetUnderlyingHostGroup());
+	// TODO: It would be better to only update connections as they need it, consider doing this on
+	// hostgroup connect/disconnect events.
+	HttpClient.UpdateConnections();
 
 	while (NextChunkRequest)
 	{
@@ -1533,7 +1543,7 @@ void FOnDemandIoBackend::ProcessHttpRequests(FMultiEndpointHttpClient& HttpClien
 					CompleteRequest(ChunkRequest);
 					Stats.OnHttpCancel();
 				}
-				else if (!HostGroup.IsConnected() || !BackendStatus.IsHttpEnabled())
+				else if (!ChunkRequest->HostGroup.IsConnected() || !BackendStatus.IsHttpEnabled())
 				{
 					CompleteRequest(ChunkRequest);
 					// Technically this request is being skipped because of a pre-existing error. It is not
@@ -1551,8 +1561,8 @@ void FOnDemandIoBackend::ProcessHttpRequests(FMultiEndpointHttpClient& HttpClien
 
 					auto AnsiUrl = StringCast<ANSICHAR>(Url.GetData(), Url.Len());
 
-					HttpClient.Get(HostGroup.GetUnderlyingHostGroup(), AnsiUrl, ChunkRequest->Params.ChunkRange,
-						[this, &NextChunkRequest, ChunkRequest, &NumConcurrentRequests, &HostGroup]
+					HttpClient.Get(ChunkRequest->HostGroup.GetUnderlyingHostGroup(), AnsiUrl, ChunkRequest->Params.ChunkRange,
+						[this, &NextChunkRequest, ChunkRequest, &NumConcurrentRequests]
 						(FMultiEndpointHttpClientResponse&& HttpResponse)
 						{
 							NumConcurrentRequests--;
@@ -1564,7 +1574,7 @@ void FOnDemandIoBackend::ProcessHttpRequests(FMultiEndpointHttpClient& HttpClien
 
 							if (HttpResponse.IsOk())
 							{
-								HostGroup.OnSuccessfulResponse();
+								ChunkRequest->HostGroup.OnSuccessfulResponse();
 
 								ChunkRequest->Chunk = MoveTemp(HttpResponse.Body);
 								Stats.OnHttpGet(ChunkRequest->Chunk.DataSize(), HttpResponse.DurationMilliseconds);
@@ -1578,7 +1588,7 @@ void FOnDemandIoBackend::ProcessHttpRequests(FMultiEndpointHttpClient& HttpClien
 							{
 								Stats.OnHttpError();
 
-								if (HostGroup.OnFailedResponse())
+								if (ChunkRequest->HostGroup.OnFailedResponse())
 								{
 									// A disconnect was triggered
 									Stats.OnHttpDisconnected();
@@ -1683,7 +1693,7 @@ uint32 FOnDemandIoBackend::Run()
 	{
 		UpdateThreadPriorityIfNeeded();
 
-		ProcessHttpRequests(*HttpClient, DefaultHostGroup, FMath::Min(GIasHttpConcurrentRequests, 32));
+		ProcessHttpRequests(*HttpClient, FMath::Min(GIasHttpConcurrentRequests, 32));
 
 		if (!bStopRequested)
 		{
