GetChunkUrl now accepts the host url as an ansi string to avoid conversions.

#rnx

diff --git a/Engine/Source/Runtime/Experimental/IoStore/OnDemand/Private/OnDemandIoStore.cpp b/Engine/Source/Runtime/Experimental/IoStore/OnDemand/Private/OnDemandIoStore.cpp
--- a/Engine/Source/Runtime/Experimental/IoStore/OnDemand/Private/OnDemandIoStore.cpp
+++ b/Engine/Source/Runtime/Experimental/IoStore/OnDemand/Private/OnDemandIoStore.cpp
@@ -64,7 +64,7 @@ namespace Private
 
 ///////////////////////////////////////////////////////////////////////////////
 static FAnsiStringBuilderBase& GetChunkUrl(
-	const FStringView& Host,
+	const FAnsiStringView& Host,
 	const FOnDemandContainer& Container,
 	const FOnDemandChunkEntry& Entry,
 	FAnsiStringBuilderBase& OutUrl)
@@ -1246,8 +1246,7 @@ FIoStatus FOnDemandIoStore::TickMountRequest(FMountRequest& MountRequest)
 
 				const uint32 RetryCount = 2;
 				TIoStatusOr<FIoBuffer> ContainerHeaderChunk = FHttpClient::Get(
-					Private::GetChunkUrl(FString(Args.HostGroup.PrimaryHost()), *Container, ChunkInfo.ChunkEntry(),
-					ChunkUrlBuilder).ToView(),
+					Private::GetChunkUrl(Args.HostGroup.PrimaryHost(), *Container, ChunkInfo.ChunkEntry(), ChunkUrlBuilder).ToView(),
 					RetryCount,
 					EHttpRedirects::Follow);
 
