Cloth repulsion constraint improvements.

Use a tet volume-based approach rather than point repulsion. Still not optimized (e.g., no ISPC), but good starting point for testing.

diff --git a/Engine/Plugins/ChaosCloth/Source/ChaosCloth/Private/ChaosCloth/ChaosClothConstraints.cpp b/Engine/Plugins/ChaosCloth/Source/ChaosCloth/Private/ChaosCloth/ChaosClothConstraints.cpp
--- a/Engine/Plugins/ChaosCloth/Source/ChaosCloth/Private/ChaosCloth/ChaosClothConstraints.cpp
+++ b/Engine/Plugins/ChaosCloth/Source/ChaosCloth/Private/ChaosCloth/ChaosClothConstraints.cpp
@@ -21,9 +21,9 @@
 #include "Chaos/PBDCollisionSpringConstraints.h"
 #include "Chaos/PBDSelfCollisionSphereConstraints.h"
 #include "Chaos/PBDKinematicTriangleMeshCollisions.h"
-#include "Chaos/PBDRepulsionConstraints.h"
 #include "Chaos/PBDTriangleMeshCollisions.h"
 #include "Chaos/PBDTriangleMeshIntersections.h"
+#include "Chaos/PBDUnilateralTetConstraints.h"
 #include "Chaos/PBDEvolution.h"
 #include "Chaos/PBDSoftBodyCollisionConstraint.h"
 #include "Chaos/SoftsExternalForces.h"
@@ -1586,11 +1586,11 @@ void FClothConstraints::CreateClothClothConstraints(
 					}
 					else if (ConstraintName == TEXT("VertexFaceRepulsionConstraint"))
 					{
-						if (Softs::FPBDRepulsionConstraints::IsEnabled(ConfigProperties))
+						if (Softs::FPBDVertexFaceRepulsionConstraints::IsEnabled(ConfigProperties))
 						{
 							if (ensure(!RepulsionConstraints))
 							{
-								RepulsionConstraints = MakeShared<Softs::FPBDRepulsionConstraints>(
+								RepulsionConstraints = MakeShared<Softs::FPBDVertexFaceRepulsionConstraints>(
 									Evolution->GetSoftBodyParticles(ParticleRangeId),
 									ConfigProperties,
 									SpringConstraintFacade
diff --git a/Engine/Plugins/ChaosCloth/Source/ChaosCloth/Private/ChaosCloth/ChaosClothVisualization.cpp b/Engine/Plugins/ChaosCloth/Source/ChaosCloth/Private/ChaosCloth/ChaosClothVisualization.cpp
--- a/Engine/Plugins/ChaosCloth/Source/ChaosCloth/Private/ChaosCloth/ChaosClothVisualization.cpp
+++ b/Engine/Plugins/ChaosCloth/Source/ChaosCloth/Private/ChaosCloth/ChaosClothVisualization.cpp
@@ -21,11 +21,11 @@
 #include "Chaos/PBDCollisionSpringConstraints.h"
 #include "Chaos/PBDFlatWeightMap.h"
 #include "Chaos/PBDLongRangeConstraints.h"
-#include "Chaos/PBDRepulsionConstraints.h"
 #include "Chaos/PBDSelfCollisionSphereConstraints.h"
 #include "Chaos/PBDSphericalConstraint.h"
 #include "Chaos/PBDSpringConstraints.h"
 #include "Chaos/PBDTriangleMeshCollisions.h"
+#include "Chaos/PBDUnilateralTetConstraints.h"
 #include "Chaos/SkinnedTriangleMesh.h"
 #include "Chaos/SoftsExternalForces.h"
 #include "Chaos/SoftsMultiResConstraints.h"
@@ -3035,11 +3035,10 @@ namespace Chaos
 					DrawPoint(PDI, Pos1, ConstraintColor, nullptr);
 				}
 			}
-			if (const Softs::FPBDRepulsionConstraints* const RepulsionConstraints = ClothConstraints.GetRepulsionConstraints().Get())
+			if (const Softs::FPBDVertexFaceRepulsionConstraints* const RepulsionConstraints = ClothConstraints.GetRepulsionConstraints().Get())
 			{
 				const TConstArrayView<Softs::FSolverVec3> Positions = Cloth->GetParticlePositions(Solver);
 				const TArray<TVec4<int32>>& Constraints = RepulsionConstraints->GetConstraints();
-				const TArray<Softs::FSolverReal>& RestLengths = RepulsionConstraints->GetRestLengths();
 				const TArray<bool>& IsActive = RepulsionConstraints->GetConstraintIsActive();
 				for (int32 ConstraintIndex = 0; ConstraintIndex < Constraints.Num(); ++ConstraintIndex)
 				{
diff --git a/Engine/Plugins/ChaosCloth/Source/ChaosCloth/Public/ChaosCloth/ChaosClothConstraints.h b/Engine/Plugins/ChaosCloth/Source/ChaosCloth/Public/ChaosCloth/ChaosClothConstraints.h
--- a/Engine/Plugins/ChaosCloth/Source/ChaosCloth/Public/ChaosCloth/ChaosClothConstraints.h
+++ b/Engine/Plugins/ChaosCloth/Source/ChaosCloth/Public/ChaosCloth/ChaosClothConstraints.h
@@ -124,7 +124,7 @@ namespace Chaos
 		const TSharedPtr<Softs::FXPBDVertexConstraints>& GetClothVertexSpringConstraints() const { return ClothVertexSpringConstraints; };
 		const TSharedPtr<Softs::FXPBDVertexFaceConstraints>& GetClothVertexFaceSpringConstraints() const { return ClothVertexFaceSpringConstraints; }
 		const TSharedPtr<Softs::FXPBDFaceConstraints>& GetClothFaceSpringConstraints() const { return ClothFaceSpringConstraints; }
-		const TSharedPtr<Softs::FPBDRepulsionConstraints>& GetRepulsionConstraints() const { return RepulsionConstraints; }
+		const TSharedPtr<Softs::FPBDVertexFaceRepulsionConstraints>& GetRepulsionConstraints() const { return RepulsionConstraints; }
 		// ---- End of debug functions ----
 
 	private:
@@ -225,7 +225,7 @@ namespace Chaos
 		TSharedPtr<Softs::FXPBDVertexConstraints> ClothVertexSpringConstraints;
 		TSharedPtr<Softs::FXPBDVertexFaceConstraints> ClothVertexFaceSpringConstraints;
 		TSharedPtr<Softs::FXPBDFaceConstraints> ClothFaceSpringConstraints;
-		TSharedPtr<Softs::FPBDRepulsionConstraints> RepulsionConstraints;
+		TSharedPtr<Softs::FPBDVertexFaceRepulsionConstraints> RepulsionConstraints;
 		//~ End Force-based solver only constraints
 
 		// Exactly one of these should be non-null
diff --git a/Engine/Plugins/ChaosClothAssetEditor/Source/ChaosClothAssetDataflowNodes/Private/ChaosClothAsset/SimulationClothVertexFaceSpringConfigNode.cpp b/Engine/Plugins/ChaosClothAssetEditor/Source/ChaosClothAssetDataflowNodes/Private/ChaosClothAsset/SimulationClothVertexFaceSpringConfigNode.cpp
--- a/Engine/Plugins/ChaosClothAssetEditor/Source/ChaosClothAssetDataflowNodes/Private/ChaosClothAsset/SimulationClothVertexFaceSpringConfigNode.cpp
+++ b/Engine/Plugins/ChaosClothAssetEditor/Source/ChaosClothAssetDataflowNodes/Private/ChaosClothAsset/SimulationClothVertexFaceSpringConfigNode.cpp
@@ -616,10 +616,10 @@ void FChaosClothAssetSimulationClothVertexFaceSpringConfigNode::AddProperties(FP
 {
 	if (!bAppendToExisting)
 	{
-		if (bUsePlanarRepulsionConstraints)
+		if (bUseTetRepulsionConstraints)
 		{
 			PropertyHelper.SetProperty(this, &VertexFaceRepulsionStiffness, {}, Chaos::Softs::ECollectionPropertyFlags::Animatable);
-			PropertyHelper.SetProperty(this, &VertexFaceRepulsionFrictionCoefficient, {}, Chaos::Softs::ECollectionPropertyFlags::Animatable);
+			PropertyHelper.SetProperty(this, &VertexFaceMaxRepulsionIters, {}, Chaos::Softs::ECollectionPropertyFlags::Animatable);
 		}
 		else
 		{
@@ -654,7 +654,7 @@ void FChaosClothAssetSimulationClothVertexFaceSpringConfigNode::EvaluateClothCol
 			TargetWeightsArray[Index] = { TargetWeights[Index][0], TargetWeights[Index][1], TargetWeights[Index][2] };
 		}
 
-		const FString ConstraintName = bUsePlanarRepulsionConstraints ? TEXT("VertexFaceRepulsionConstraint") : TEXT("VertexFaceSpringConstraint");
+		const FString ConstraintName = bUseTetRepulsionConstraints ? TEXT("VertexFaceRepulsionConstraint") : TEXT("VertexFaceSpringConstraint");
 
 		// Try to find existing constraint of this type.
 		bool bFound = false;
diff --git a/Engine/Plugins/ChaosClothAssetEditor/Source/ChaosClothAssetDataflowNodes/Public/ChaosClothAsset/SimulationClothVertexFaceSpringConfigNode.h b/Engine/Plugins/ChaosClothAssetEditor/Source/ChaosClothAssetDataflowNodes/Public/ChaosClothAsset/SimulationClothVertexFaceSpringConfigNode.h
--- a/Engine/Plugins/ChaosClothAssetEditor/Source/ChaosClothAssetDataflowNodes/Public/ChaosClothAsset/SimulationClothVertexFaceSpringConfigNode.h
+++ b/Engine/Plugins/ChaosClothAssetEditor/Source/ChaosClothAssetDataflowNodes/Public/ChaosClothAsset/SimulationClothVertexFaceSpringConfigNode.h
@@ -113,29 +113,29 @@ private:
 	UPROPERTY(EditAnywhere, Category = "Cloth Vertex Face Spring Properties")
 	bool bAppendToExisting = false;
 
-	/** Treat as planar repulsion constraints (e.g., for self collisions) rather than spring constraints */
+	/** Treat as tetrahedral repulsion constraints (e.g., for self collisions) rather than spring constraints */
 	UPROPERTY(EditAnywhere, Category = "Cloth Vertex Face Spring Properties")
-	bool bUsePlanarRepulsionConstraints = false;
+	bool bUseTetRepulsionConstraints = false;
 
 	/** Extension Stiffness is the spring stiffness applied when the spring is currently longer than its rest length. This is a low-high range, but there are currently no ways to author per-spring stiffnesses, so only Low is used in practice.*/
-	UPROPERTY(EditAnywhere, Category = "Cloth Vertex Face Spring Properties", Meta = (UIMin = "0", UIMax = "10000", ClampMin = "0", ClampMax = "10000000", InteractorName = "VertexFaceSpringExtensionStiffness", EditCondition = "!bUsePlanarRepulsionConstraints && !bAppendToExisting", EditConditionHides))
+	UPROPERTY(EditAnywhere, Category = "Cloth Vertex Face Spring Properties", Meta = (UIMin = "0", UIMax = "10000", ClampMin = "0", ClampMax = "10000000", InteractorName = "VertexFaceSpringExtensionStiffness", EditCondition = "!bUseTetRepulsionConstraints && !bAppendToExisting", EditConditionHides))
 	FVector2f VertexFaceSpringExtensionStiffness = { 100.f, 100.f };
 
 	/** Compression Stiffness is the spring stiffness applied when the spring is currently shorter than its rest length. This is a low-high range, but there are currently no ways to author per-spring stiffnesses, so only Low is used in practice.*/
-	UPROPERTY(EditAnywhere, Category = "Cloth Vertex Face Spring Properties", Meta = (UIMin = "0", UIMax = "10000", ClampMin = "0", ClampMax = "10000000", InteractorName = "VertexFaceSpringCompressionStiffness", EditCondition = "!bUsePlanarRepulsionConstraints && !bAppendToExisting", EditConditionHides))
+	UPROPERTY(EditAnywhere, Category = "Cloth Vertex Face Spring Properties", Meta = (UIMin = "0", UIMax = "10000", ClampMin = "0", ClampMax = "10000000", InteractorName = "VertexFaceSpringCompressionStiffness", EditCondition = "!bUseTetRepulsionConstraints && !bAppendToExisting", EditConditionHides))
 	FVector2f VertexFaceSpringCompressionStiffness = { 100.f, 100.f };
 
 	/** This damping is the relative to critical damping. This is a low-high range, but there are currently no ways to author per-spring stiffnesses, so only Low is used in practice.*/
-	UPROPERTY(EditAnywhere, Category = "Cloth Vertex Face Spring Properties", Meta = (UIMin = "0", UIMax = "10", ClampMin = "0", ClampMax = "1000", InteractorName = "VertexFaceSpringDamping", EditCondition = "!bUsePlanarRepulsionConstraints && !bAppendToExisting", EditConditionHides))
+	UPROPERTY(EditAnywhere, Category = "Cloth Vertex Face Spring Properties", Meta = (UIMin = "0", UIMax = "10", ClampMin = "0", ClampMax = "1000", InteractorName = "VertexFaceSpringDamping", EditCondition = "!bUseTetRepulsionConstraints && !bAppendToExisting", EditConditionHides))
 	FVector2f VertexFaceSpringDamping = { 0.f, 0.f };
 
 	/** Stiffness for repulsion constraints */
-	UPROPERTY(EditAnywhere, Category = "Cloth Vertex Face Spring Properties", Meta = (ClampMin = "0", ClampMax = "1", InteractorName = "VertexFaceRepulsionStiffness", EditCondition = "bUsePlanarRepulsionConstraints && !bAppendToExisting", EditConditionHides))
+	UPROPERTY(EditAnywhere, Category = "Cloth Vertex Face Spring Properties", Meta = (ClampMin = "0", ClampMax = "1", InteractorName = "VertexFaceRepulsionStiffness", EditCondition = "bUseTetRepulsionConstraints && !bAppendToExisting", EditConditionHides))
 	float VertexFaceRepulsionStiffness = 0.5f;
 
-	/** Friction for repulsion constraints */
-	UPROPERTY(EditAnywhere, Category = "Cloth Vertex Face Spring Properties", Meta = (ClampMin = "0", ClampMax = "1", InteractorName = "VertexFaceRepulsionStiffness", EditCondition = "bUsePlanarRepulsionConstraints && !bAppendToExisting", EditConditionHides))
-	float VertexFaceRepulsionFrictionCoefficient = 0.f;
+	/** Max Number of iterations to apply (per solver iteration). Helps resolve more collisions, but at additional compute cost. */
+	UPROPERTY(EditAnywhere, Category = "Cloth Vertex Face Spring Properties", Meta = (ClampMin = "1", UIMax = "5", InteractorName = "VertexFaceRepulsionStiffness", EditCondition = "bUseTetRepulsionConstraints && !bAppendToExisting", EditConditionHides))
+	int32 VertexFaceMaxRepulsionIters = 1;
 
 	/** Construction data for procedurally generating constraints.*/
 	UPROPERTY(EditAnywhere, EditFixedSize, Category = "Cloth Vertex Face Spring Construction", Meta = (SkipInDisplayNameChain))
diff --git a/Engine/Source/Runtime/Experimental/Chaos/Private/Chaos/PBDRepulsionConstraints.cpp b/Engine/Source/Runtime/Experimental/Chaos/Private/Chaos/PBDRepulsionConstraints.cpp
deleted file mode 100644
diff --git a/Engine/Source/Runtime/Experimental/Chaos/Private/Chaos/PBDUnilateralTetConstraints.cpp b/Engine/Source/Runtime/Experimental/Chaos/Private/Chaos/PBDUnilateralTetConstraints.cpp
new file mode 100644
--- /dev/null
+++ b/Engine/Source/Runtime/Experimental/Chaos/Private/Chaos/PBDUnilateralTetConstraints.cpp
@@ -0,0 +1,251 @@
+// Copyright Epic Games, Inc. All Rights Reserved.
+#include "Chaos/PBDUnilateralTetConstraints.h"
+#include "Chaos/GraphColoring.h"
+#include "Chaos/Triangle.h"
+#include "ChaosStats.h"
+#include "ChaosLog.h"
+#include "HAL/IConsoleManager.h"
+#if 0 && INTEL_ISPC
+#include "PBDUnilateralTetConstraints.ispc.generated.h"
+#endif
+
+#if 0 && INTEL_ISPC
+
+#if !defined(CHAOS_UNILATERAL_TET_ISPC_ENABLED_DEFAULT)
+#define CHAOS_UNILATERAL_TET_ISPC_ENABLED_DEFAULT 1
+#endif
+
+#if !defined(USE_ISPC_KERNEL_CONSOLE_VARIABLES_IN_SHIPPING)
+#define USE_ISPC_KERNEL_CONSOLE_VARIABLES_IN_SHIPPING 0
+#endif
+
+// Support run-time toggling on supported platforms in non-shipping configurations
+#if (UE_BUILD_SHIPPING && !USE_ISPC_KERNEL_CONSOLE_VARIABLES_IN_SHIPPING)
+static constexpr bool bChaos_UnilateralTet_ISPC_Enabled = INTEL_ISPC && CHAOS_UNILATERAL_TET_ISPC_ENABLED_DEFAULT;
+#else
+
+static bool bChaos_UnilateralTet_ISPC_Enabled;
+
+#endif
+
+static_assert(sizeof(ispc::FVector4f) == sizeof(Chaos::Softs::FPAndInvM), "sizeof(ispc::FVector4f) != sizeof(Chaos::Softs::FPAndInvM)");
+static_assert(sizeof(ispc::FIntVector4) == sizeof(Chaos::TVec4<int32>), "sizeof(ispc::FIntVector4) != sizeof(Chaos::TVec4<int32>");
+
+#endif
+namespace Chaos::Softs
+{
+	FPBDUnilateralTetConstraints::FPBDUnilateralTetConstraints(
+		const FSolverParticlesRange& Particles,
+		TArray<TVector<int32, 4>>&& InConstraints,
+		TArray<FSolverReal>&& InVolumes,
+		FSolverReal InStiffness,
+		int32 InMaxNumIters)
+		: Constraints(MoveTemp(InConstraints))
+		, Volumes(MoveTemp(InVolumes))
+		, Stiffness(InStiffness)
+		, MaxNumIters(InMaxNumIters)
+	{
+		TrimKinematicConstraints(Particles);
+		InitColor(Particles);
+
+#if CHAOS_DEBUG_DRAW
+		ConstraintIsActive.SetNum(Constraints.Num());
+#endif
+	}
+
+	void FPBDUnilateralTetConstraints::TrimKinematicConstraints(const FSolverParticlesRange& Particles)
+	{
+		TArray<TVector<int32, 4>> TrimmedConstraints;
+		TArray<FSolverReal> TrimmedVolumes;
+		TrimmedConstraints.Reserve(Constraints.Num());
+		TrimmedVolumes.Reserve(Constraints.Num());
+
+		for (int32 ConstraintIndex = 0; ConstraintIndex < Constraints.Num(); ++ConstraintIndex)
+		{
+			if (Particles.InvM(Constraints[ConstraintIndex][0]) != (FSolverReal)0.f ||
+				Particles.InvM(Constraints[ConstraintIndex][1]) != (FSolverReal)0.f ||
+				Particles.InvM(Constraints[ConstraintIndex][2]) != (FSolverReal)0.f ||
+				Particles.InvM(Constraints[ConstraintIndex][3]) != (FSolverReal)0.f)
+			{
+				TrimmedConstraints.Add(Constraints[ConstraintIndex]);
+				TrimmedVolumes.Add(Volumes[ConstraintIndex]);
+			}
+		}
+
+		TrimmedConstraints.Shrink();
+		TrimmedVolumes.Shrink();
+
+		Constraints = MoveTemp(TrimmedConstraints);
+		Volumes = MoveTemp(TrimmedVolumes);
+	}
+
+	void FPBDUnilateralTetConstraints::InitColor(const FSolverParticlesRange& Particles)
+	{
+		const int32 NumConstraints = Constraints.Num();
+		const TArray<TArray<int32>> ConstraintsPerColor = FGraphColoring::ComputeGraphColoringParticlesOrRange(Constraints, Particles, 0, Particles.Size());
+		// Reorder constraints based on color so each array in ConstraintsPerColor contains contiguous elements.
+		TArray<TVector<int32, 4>> ReorderedConstraints;
+		TArray<FSolverReal> ReorderedRestLengths;
+		TArray<FSolverReal> ReorderedVolumes;
+		ReorderedConstraints.SetNumUninitialized(NumConstraints);
+		ReorderedRestLengths.SetNumUninitialized(NumConstraints);
+		ReorderedVolumes.SetNumUninitialized(NumConstraints);
+
+		ConstraintsPerColorStartIndex.Reset(ConstraintsPerColor.Num() + 1);
+		int32 ReorderedIndex = 0;
+		for (const TArray<int32>& ConstraintsBatch : ConstraintsPerColor)
+		{
+			ConstraintsPerColorStartIndex.Add(ReorderedIndex);
+			for (const int32 OrigIndex : ConstraintsBatch)
+			{
+				ReorderedConstraints[ReorderedIndex] = Constraints[OrigIndex];
+				ReorderedVolumes[ReorderedIndex] = Volumes[OrigIndex];
+
+				++ReorderedIndex;
+			}
+		}
+		ConstraintsPerColorStartIndex.Add(ReorderedIndex);
+
+		Constraints = MoveTemp(ReorderedConstraints);
+		Volumes = MoveTemp(ReorderedVolumes);
+	}
+	
+	void FPBDUnilateralTetConstraints::Apply(FSolverParticlesRange& Particles, const FSolverReal Dt) const
+	{
+		ApplyVolumeConstraint(Particles, Dt);
+	}
+
+	void FPBDUnilateralTetConstraints::ApplyVolumeConstraint(FSolverParticlesRange& Particles, const FSolverReal Dt) const
+	{
+		// TODO ISPC
+		const TVec4<int32>* const ConstraintsData = Constraints.GetData();
+
+		FPAndInvM* const PAndInvM = Particles.GetPAndInvM().GetData();
+
+		for (int32 Iter = 0; Iter < MaxNumIters; ++Iter)
+		{
+			bool bAnyActive = false;
+			for (int32 ConstraintIndex = 0; ConstraintIndex < Constraints.Num(); ++ConstraintIndex)
+			{
+#if CHAOS_DEBUG_DRAW
+				ConstraintIsActive[ConstraintIndex] = false;
+#endif
+				const TVec4<int32>& Constraint = ConstraintsData[ConstraintIndex];
+				const int32 Index1 = Constraint[1];
+				const int32 Index2 = Constraint[2];
+				const int32 Index3 = Constraint[3];
+				const int32 Index4 = Constraint[0];
+				FSolverVec3& P1 = PAndInvM[Index1].P;
+				FSolverVec3& P2 = PAndInvM[Index2].P;
+				FSolverVec3& P3 = PAndInvM[Index3].P;
+				FSolverVec3& P4 = PAndInvM[Index4].P;
+
+				TVec4<FSolverVec3> Grads;
+				const FSolverVec3 P2P1 = P2 - P1;
+				const FSolverVec3 P4P1 = P4 - P1;
+				const FSolverVec3 P3P1 = P3 - P1;
+				Grads[1] = FSolverVec3::CrossProduct(P3P1, P4P1) / (FSolverReal)6.;
+				Grads[2] = FSolverVec3::CrossProduct(P4P1, P2P1) / (FSolverReal)6.;
+				Grads[3] = FSolverVec3::CrossProduct(P2P1, P3P1) / (FSolverReal)6.;
+				Grads[0] = -(Grads[1] + Grads[2] + Grads[3]);
+
+				const FSolverReal Volume = FSolverVec3::DotProduct(FSolverVec3::CrossProduct(P2P1, P3P1), P4P1) / (FSolverReal)6.;
+				if ((Volume - Volumes[ConstraintIndex]) >= 0)
+				{
+					continue;
+				}
+
+				const FSolverReal S = Stiffness * (Volume - Volumes[ConstraintIndex]) / (
+					PAndInvM[Index1].InvM * Grads[0].SizeSquared() +
+					PAndInvM[Index2].InvM * Grads[1].SizeSquared() +
+					PAndInvM[Index3].InvM * Grads[2].SizeSquared() +
+					PAndInvM[Index4].InvM * Grads[3].SizeSquared());
+
+				P1 -= S * PAndInvM[Index1].InvM * Grads[0];
+				P2 -= S * PAndInvM[Index2].InvM * Grads[1];
+				P3 -= S * PAndInvM[Index3].InvM * Grads[2];
+				P4 -= S * PAndInvM[Index4].InvM * Grads[3];
+
+				bAnyActive = true;
+#if CHAOS_DEBUG_DRAW
+				ConstraintIsActive[ConstraintIndex] = true;
+#endif
+			}
+			if (!bAnyActive)
+			{
+				break;
+			}
+		}
+	}
+
+
+	FPBDVertexFaceRepulsionConstraints::FPBDVertexFaceRepulsionConstraints(
+		const FSolverParticlesRange& Particles,
+		const FCollectionPropertyConstFacade& PropertyCollection,
+		const FEmbeddedSpringConstraintFacade& SpringConstraintFacade)
+		: FPBDUnilateralTetConstraints(
+			Particles
+			, ExtractConstraintIndices(SpringConstraintFacade)
+			, ExtractVolumes(Particles, SpringConstraintFacade)
+			, FMath::Clamp(GetVertexFaceRepulsionStiffness(PropertyCollection, DefaultStiffness), MinStiffness, MaxStiffness)
+			, FMath::Max(GetVertexFaceMaxRepulsionIters(PropertyCollection, 1), 1)
+		)
+		, VertexFaceRepulsionStiffnessIndex(PropertyCollection)
+		, VertexFaceMaxRepulsionItersIndex(PropertyCollection)
+	{
+	}
+
+	void FPBDVertexFaceRepulsionConstraints::SetProperties(const FCollectionPropertyConstFacade& PropertyCollection)
+	{
+		if (IsVertexFaceRepulsionStiffnessMutable(PropertyCollection))
+		{
+			Stiffness = FMath::Clamp(GetVertexFaceRepulsionStiffness(PropertyCollection), MinStiffness, MaxStiffness);
+		}
+		if (IsVertexFaceMaxRepulsionItersMutable(PropertyCollection))
+		{
+			MaxNumIters = FMath::Max(GetVertexFaceMaxRepulsionIters(PropertyCollection), 1);
+		}
+	}
+
+	TArray<TVector<int32, 4>> FPBDVertexFaceRepulsionConstraints::ExtractConstraintIndices(const FEmbeddedSpringConstraintFacade& SpringConstraintFacade)
+	{
+		const TConstArrayView<TArray<int32>>& InSourceIndices = SpringConstraintFacade.GetSourceIndexConst();
+		const TConstArrayView<TArray<int32>>& InTargetIndices = SpringConstraintFacade.GetTargetIndexConst();
+
+		const int32 NumConstraints = InSourceIndices.Num();
+		check(InTargetIndices.Num() == NumConstraints);
+
+		TArray<TVector<int32,4>> ResultConstraints;
+		ResultConstraints.Reserve(NumConstraints);
+
+		for (int32 ConstraintIdx = 0; ConstraintIdx < NumConstraints; ++ConstraintIdx)
+		{
+			ResultConstraints.Emplace(InSourceIndices[ConstraintIdx][0], InTargetIndices[ConstraintIdx][0], InTargetIndices[ConstraintIdx][1], InTargetIndices[ConstraintIdx][2]);
+		}
+		return ResultConstraints;
+	}
+
+	TArray<FSolverReal> FPBDVertexFaceRepulsionConstraints::ExtractVolumes(const FSolverParticlesRange& Particles, const FEmbeddedSpringConstraintFacade& SpringConstraintFacade)
+	{
+		const TConstArrayView<TArray<int32>>& InTargetIndices = SpringConstraintFacade.GetTargetIndexConst();
+		const TConstArrayView<float>& InSpringLengths = SpringConstraintFacade.GetSpringLengthConst();
+
+		const int32 NumConstraints = InTargetIndices.Num();
+		check(InSpringLengths.Num() == NumConstraints);
+
+		TArray<FSolverReal> ResultVolumes;
+		ResultVolumes.Reserve(NumConstraints);
+
+		for (int32 ConstraintIdx = 0; ConstraintIdx < NumConstraints; ++ConstraintIdx)
+		{
+			const TArray<int32>& Constraint = InTargetIndices[ConstraintIdx];
+			// Calculate Volume = area of triangle * restlength.
+			const FSolverReal Volume =
+				(FSolverVec3::CrossProduct(Particles.X(Constraint[1]) - Particles.X(Constraint[0]), Particles.X(Constraint[2]) - Particles.X(Constraint[0]))).Length()
+				* InSpringLengths[ConstraintIdx] / (FSolverReal)6.;
+
+			ResultVolumes.Emplace(Volume);
+		}
+		return ResultVolumes;
+	}
+} // End namespace Chaos::Softs
diff --git a/Engine/Source/Runtime/Experimental/Chaos/Public/Chaos/PBDRepulsionConstraints.h b/Engine/Source/Runtime/Experimental/Chaos/Public/Chaos/PBDRepulsionConstraints.h
deleted file mode 100644
diff --git a/Engine/Source/Runtime/Experimental/Chaos/Public/Chaos/PBDSoftsEvolutionFwd.h b/Engine/Source/Runtime/Experimental/Chaos/Public/Chaos/PBDSoftsEvolutionFwd.h
--- a/Engine/Source/Runtime/Experimental/Chaos/Public/Chaos/PBDSoftsEvolutionFwd.h
+++ b/Engine/Source/Runtime/Experimental/Chaos/Public/Chaos/PBDSoftsEvolutionFwd.h
@@ -82,7 +82,7 @@ namespace Chaos
 		class FXPBDVertexConstraints;
 		class FXPBDVertexFaceConstraints;
 		class FXPBDFaceConstraints;
-		class FPBDRepulsionConstraints;
+		class FPBDVertexFaceRepulsionConstraints;
 
 		// Softs solver forces
 		class FVelocityAndPressureField;
diff --git a/Engine/Source/Runtime/Experimental/Chaos/Public/Chaos/PBDUnilateralTetConstraints.h b/Engine/Source/Runtime/Experimental/Chaos/Public/Chaos/PBDUnilateralTetConstraints.h
new file mode 100644
--- /dev/null
+++ b/Engine/Source/Runtime/Experimental/Chaos/Public/Chaos/PBDUnilateralTetConstraints.h
@@ -0,0 +1,94 @@
+// Copyright Epic Games, Inc. All Rights Reserved.
+#pragma once
+
+#include "Chaos/CollectionEmbeddedSpringConstraintFacade.h"
+#include "Chaos/CollectionPropertyFacade.h"
+#include "Chaos/Core.h"
+#include "Chaos/PBDFlatWeightMap.h"
+#include "Chaos/PBDSoftsEvolutionFwd.h"
+#include "Chaos/PBDSoftsSolverParticles.h"
+#include "Chaos/SoftsSolverParticlesRange.h"
+#include "Containers/Array.h"
+
+namespace Chaos::Softs
+{
+
+	class FPBDUnilateralTetConstraints
+	{
+	public:
+		static constexpr FSolverReal MinStiffness = (FSolverReal)0;
+		static constexpr FSolverReal MaxStiffness = (FSolverReal)1.f;
+		static constexpr FSolverReal DefaultStiffness = (FSolverReal)0.5f;
+
+		CHAOS_API FPBDUnilateralTetConstraints(
+			const FSolverParticlesRange& Particles,
+			TArray<TVector<int32, 4>>&& InConstraints,
+			TArray<FSolverReal>&& InVolumes,
+			FSolverReal InStiffness,
+			int32 InMaxNumIters);
+
+		~FPBDUnilateralTetConstraints() = default;
+
+		void ApplyProperties(const FSolverReal /*Dt*/, const int32 /*NumIterations*/)
+		{
+		}
+
+		CHAOS_API void Apply(FSolverParticlesRange& Particles, const FSolverReal Dt) const;
+		
+		const TArray<TVector<int32, 4>>& GetConstraints() const { return Constraints; }
+
+#if CHAOS_DEBUG_DRAW
+		const TArray<bool>& GetConstraintIsActive() const { return ConstraintIsActive; }
+#endif
+
+	private:
+
+		void TrimKinematicConstraints(const FSolverParticlesRange& Particles);
+		void InitColor(const FSolverParticlesRange& Particles);
+
+		void ApplyVolumeConstraint(FSolverParticlesRange& Particles, const FSolverReal Dt) const;
+#if INTEL_ISPC
+		void ApplyVolumeConstraint_ISPC(FSolverParticlesRange& Particles, const FSolverReal Dt) const {}
+#endif
+
+		TArray<TVector<int32, 4>> Constraints;
+		TArray<FSolverReal> Volumes;
+
+		TArray<int32> ConstraintsPerColorStartIndex; // Constraints are ordered so each batch is contiguous. This is ColorNum + 1 length so it can be used as start and end.
+
+#if CHAOS_DEBUG_DRAW
+		// Used for debug draw to show which constraints were active last apply
+		mutable TArray<bool> ConstraintIsActive;
+#endif
+
+	protected:
+		FSolverReal Stiffness;
+		int32 MaxNumIters;
+	};
+
+	class FPBDVertexFaceRepulsionConstraints final : public FPBDUnilateralTetConstraints
+	{
+	public:
+
+		static bool IsEnabled(const FCollectionPropertyConstFacade& PropertyCollection)
+		{
+			return IsVertexFaceRepulsionStiffnessEnabled(PropertyCollection, false);
+		}
+
+		CHAOS_API FPBDVertexFaceRepulsionConstraints(
+			const FSolverParticlesRange& Particles,
+			const FCollectionPropertyConstFacade& PropertyCollection,
+			const FEmbeddedSpringConstraintFacade& SpringConstraintFacade);
+
+		CHAOS_API void SetProperties(const FCollectionPropertyConstFacade& PropertyCollection);
+
+	private:
+		static TArray<TVector<int32, 4>> ExtractConstraintIndices(const FEmbeddedSpringConstraintFacade& SpringConstraintFacade);
+		static TArray<FSolverReal> ExtractVolumes(const FSolverParticlesRange& Particles, const FEmbeddedSpringConstraintFacade& SpringConstraintFacade);
+
+		UE_CHAOS_DECLARE_PROPERTYCOLLECTION_NAME(VertexFaceRepulsionStiffness, float);
+		UE_CHAOS_DECLARE_PROPERTYCOLLECTION_NAME(VertexFaceMaxRepulsionIters, int32);
+	};
+
+}  // End namespace Chaos::Softs
+
