Iris
- ChunkedDataStream, Experimental DataStream to send bytebuffers from external serialization along with object references.

diff --git a/Engine/Plugins/Runtime/ReplicationSystemTestPlugin/Source/Private/Tests/DataStream/TestChunkedDataStream.cpp b/Engine/Plugins/Runtime/ReplicationSystemTestPlugin/Source/Private/Tests/DataStream/TestChunkedDataStream.cpp
new file mode 100644
--- /dev/null
+++ b/Engine/Plugins/Runtime/ReplicationSystemTestPlugin/Source/Private/Tests/DataStream/TestChunkedDataStream.cpp
@@ -0,0 +1,410 @@
+// Copyright Epic Games, Inc. All Rights Reserved.
+
+#include "Logging/LogScopedVerbosityOverride.h"
+#include "NetworkAutomationTest.h"
+#include "NetworkAutomationTestMacros.h"
+
+#include "Iris/Core/IrisLog.h"
+#include "Iris/DataStream/DataStreamManager.h"
+#include "Iris/DataStream/DataStreamDefinitions.h"
+#include "Iris/DataStream/DataStreamManager.h"
+#include "Iris/ReplicationSystem/ChunkedDataStream/ChunkedDataStream.h"
+
+#include "Iris/PacketControl/PacketNotification.h"
+#include "Iris/Serialization/NetBitStreamReader.h"
+#include "Iris/Serialization/NetBitStreamWriter.h"
+#include "Iris/Serialization/NetSerializationContext.h"
+#include "Tests/ReplicationSystem/ReplicationSystemServerClientTestFixture.h"
+#include "Iris/ReplicationSystem/ReplicationSystem.h"
+#include "Iris/ReplicationSystem/NetTokenStore.h"
+#include "Iris/ReplicationSystem/StringTokenStore.h"
+#include "Iris/ReplicationSystem/NetTokenDataStream.h"
+#include "Templates/Casts.h"
+
+#include "Serialization/MemoryWriter.h"
+#include "Serialization/MemoryReader.h"
+
+namespace UE::Net::Private
+{
+
+class FTestChunkedDataStreamFixture : public FReplicationSystemServerClientTestFixture
+{
+protected:
+	virtual void SetUp() override
+	{
+		FReplicationSystemServerClientTestFixture::SetUp();
+
+		// We can override what we want as long as we do it before any connections are created.
+		DataStreamUtil.SetUp();
+
+		// Add a dynamic DataStream
+		{
+			FDataStreamTestUtil::FAddDataStreamDefinitionParams Params;
+			Params.bDynamicCreate = true;
+			DataStreamUtil.AddDataStreamDefinition(TEXT("ChunkedData"), TEXT("/Script/IrisCore.ChunkedDataStream"), Params);	
+		}
+
+		{
+			DataStreamUtil.AddDataStreamDefinition(TEXT("Replication"), TEXT("/Script/IrisCore.ReplicationDataStream"));
+		}
+
+		DataStreamUtil.FixupDefinitions();
+
+		// Add a client
+		Client = CreateClient();
+	}
+
+	UDataStream::EDataStreamState GetDataStreamStateOnServer() const
+	{
+		return Server->GetConnectionInfo(Client->ConnectionIdOnServer).DataStreamManager->GetStreamState(DataStreamName);
+	}
+
+	UDataStream::EDataStreamState GetDataStreamStateOnClient() const
+	{
+		return Client->GetConnectionInfo(Client->LocalConnectionId).DataStreamManager->GetStreamState(DataStreamName);
+	}
+
+	void RoundTrip(bool bDeliver = true)
+	{
+		Server->UpdateAndSend({Client}, bDeliver);
+		Client->UpdateAndSend(Server, true);
+	}
+
+	void RoundTripWithLatencyAndDeliverPercentage(uint32 InFlightCount = 1, int32 DropPercentage = 10, int32 PacketMinSize = 128, int32 PacketMaxSize = 2048)
+	{
+		uint32 SentCount = 0U;
+		for (uint32 It = 0; It <= InFlightCount; ++It)
+		{
+			Server->SetMaxSendPacketSize(FMath::RandRange(PacketMinSize, PacketMaxSize) & ~0x3);
+			Server->PreSendUpdate();
+
+			SentCount += Server->SendTo(Client) ? 1U : 0U;
+
+			Server->PostSendUpdate();
+		}
+
+		while (SentCount--)
+		{
+			const bool bDeliver =  FMath::RandRange(0, 100) > DropPercentage;
+			Server->DeliverTo(Client, bDeliver);
+		}
+				
+		Client->UpdateAndSend(Server, true);
+	}
+
+	const FName DataStreamName = "ChunkedData";
+	FReplicationSystemTestClient* Client = nullptr;
+	FDataStreamTestUtil DataStreamUtil;
+};
+
+// Basic functionality test
+UE_NET_TEST_FIXTURE(FTestChunkedDataStreamFixture, TestChunkedDataStream)
+{
+	UReplicationSystem* ServerReplicationSystem = Server->ReplicationSystem;
+	UReplicationSystem* ClientReplicationSystem = Client->ReplicationSystem;
+
+	// Open stream
+	UChunkedDataStream* ServerStream = Cast<UChunkedDataStream>(ServerReplicationSystem->OpenDataStream(Client->ConnectionIdOnServer, DataStreamName));
+
+	// Double roundtrip to ensure that stream is fully open/acked
+	RoundTrip();
+	RoundTrip();
+
+	UChunkedDataStream* ClientStream = Cast<UChunkedDataStream>(ClientReplicationSystem->GetDataStream(Client->LocalConnectionId, DataStreamName));
+
+	// Verify that stream is open
+	UE_NET_ASSERT_TRUE(ClientStream->GetState() == UDataStream::EDataStreamState::Open);
+
+	// Send a payload
+	TSharedRef<TArray64<uint8>> Payload = MakeShared<TArray64<uint8>>();
+	Payload->SetNumZeroed(1567);
+	ServerStream->EnqueuePayload(Payload);
+
+	// Send some data
+	RoundTrip();
+	RoundTrip();
+
+	// Verify that we have got the payload
+	UE_NET_ASSERT_EQ(ClientStream->GetNumReceivedPayloadsPendingDispatch(), 1U);
+
+	// Loop until we have received the payload
+	bool bHasReceivedPayload = false;	
+
+	// Dispatch
+	ClientStream->DispatchReceivedPayload([this, &Payload, &bHasReceivedPayload](TConstArrayView64<uint8> ReceivedPayload)
+	{
+		bHasReceivedPayload = true;
+		UE_NET_ASSERT_EQ(ReceivedPayload.Num(), Payload->Num());
+	});
+
+	UE_NET_ASSERT_TRUE(bHasReceivedPayload);
+
+	// Verify that we have dispatched all received payloads
+	UE_NET_ASSERT_EQ(ClientStream->GetNumReceivedPayloadsPendingDispatch(), 0U);
+}
+
+// Verify that stream can be closed from client
+UE_NET_TEST_FIXTURE(FTestChunkedDataStreamFixture, TestChunkedDataStream_RequestCloseFromClient)
+{
+	UReplicationSystem* ServerReplicationSystem = Server->ReplicationSystem;
+	UReplicationSystem* ClientReplicationSystem = Client->ReplicationSystem;
+
+	// Open stream
+	UChunkedDataStream* ServerStream = Cast<UChunkedDataStream>(ServerReplicationSystem->OpenDataStream(Client->ConnectionIdOnServer, DataStreamName));
+
+	RoundTrip();
+	RoundTrip();
+
+	UChunkedDataStream* ClientStream = Cast<UChunkedDataStream>(ClientReplicationSystem->GetDataStream(Client->LocalConnectionId, DataStreamName));
+
+	// Verify open
+	UE_NET_ASSERT_TRUE(ClientStream != nullptr);
+
+	// Request close from client
+	ClientStream->RequestClose();
+
+	// Verify that stream is open on both sides
+	UE_NET_ASSERT_TRUE(GetDataStreamStateOnServer() == UDataStream::EDataStreamState::Open);
+	UE_NET_ASSERT_TRUE(GetDataStreamStateOnClient() == UDataStream::EDataStreamState::PendingClose);
+
+	// Send and deliver to server
+	Client->UpdateAndSend(Server, true);
+	
+	UE_NET_ASSERT_TRUE(GetDataStreamStateOnServer() == UDataStream::EDataStreamState::PendingClose);
+	UE_NET_ASSERT_TRUE(GetDataStreamStateOnClient() == UDataStream::EDataStreamState::WaitOnCloseConfirmation);
+
+	// Double roundtrip and we should be done
+	RoundTrip();
+	RoundTrip();
+
+	// Verify that stream is invalidated
+	UE_NET_ASSERT_TRUE(GetDataStreamStateOnServer() == UDataStream::EDataStreamState::Invalid);
+	UE_NET_ASSERT_TRUE(GetDataStreamStateOnClient() == UDataStream::EDataStreamState::Invalid);	
+}
+
+// Send some payloads with varying size
+UE_NET_TEST_FIXTURE(FTestChunkedDataStreamFixture, TestChunkedDataStream_SendManyWithVaryingPacketSizes)
+{
+	UReplicationSystem* ServerReplicationSystem = Server->ReplicationSystem;
+	UReplicationSystem* ClientReplicationSystem = Client->ReplicationSystem;
+
+	// Open stream
+	UChunkedDataStream* ServerStream = Cast<UChunkedDataStream>(ServerReplicationSystem->OpenDataStream(Client->ConnectionIdOnServer, DataStreamName));
+	RoundTrip();
+	RoundTrip();
+	UChunkedDataStream* ClientStream = Cast<UChunkedDataStream>(ClientReplicationSystem->GetDataStream(Client->LocalConnectionId, DataStreamName));
+
+	// Send payloads
+	const uint32 NumPayloads = 313;
+	TSharedPtr<TArray64<uint8>> Payloads[NumPayloads];
+	for (uint32 It=0; It<NumPayloads; ++It)
+	{
+		Payloads[It] = MakeShared<TArray64<uint8>>();
+		Payloads[It]->SetNumZeroed(FMath::RandRange(10, 5000));
+		ServerStream->EnqueuePayload(Payloads[It]);
+	}
+
+	// Loop send and dispatch until we have received and dispatched all data.
+	uint32 NumRecvdPayloads = 0;	
+	while (NumRecvdPayloads < NumPayloads)
+	{
+		// Send data with varying packet loss and unacked packets in-flight
+		RoundTripWithLatencyAndDeliverPercentage(128);
+
+		// Verify received payloads
+		ClientStream->DispatchReceivedPayloads([this, &Payloads, &NumRecvdPayloads](TConstArrayView64<uint8> Data)
+		{
+			UE_NET_ASSERT_EQ(Data.Num(), Payloads[NumRecvdPayloads]->Num());
+			++NumRecvdPayloads;
+		});
+	}
+	UE_NET_ASSERT_EQ(NumRecvdPayloads, NumPayloads);
+}
+
+// Send some payloads with varying size and exports
+UE_NET_TEST_FIXTURE(FTestChunkedDataStreamFixture, TestChunkedDataStream_SendManyWithVaryingPacketSizesAndExports)
+{
+	UReplicationSystem* ServerReplicationSystem = Server->ReplicationSystem;
+	UReplicationSystem* ClientReplicationSystem = Client->ReplicationSystem;
+
+	// Open stream
+	UChunkedDataStream* ServerStream = Cast<UChunkedDataStream>(ServerReplicationSystem->OpenDataStream(Client->ConnectionIdOnServer, DataStreamName));
+	RoundTrip();
+	RoundTrip();
+	UChunkedDataStream* ClientStream = Cast<UChunkedDataStream>(ClientReplicationSystem->GetDataStream(Client->LocalConnectionId, DataStreamName));
+
+	// Some refs to test exports
+	UObject* ObjectRefs[] =
+	{
+		nullptr,
+		UTestReplicatedIrisObject::StaticClass(),
+		UReplicationSystem::StaticClass(),
+		UReplicatedTestObject::StaticClass()
+	};
+	const int32 MaxObjectRefsIndex = UE_ARRAY_COUNT(ObjectRefs);
+
+	// Send some data
+	const uint32 NumPayloads = 64;
+	TSharedPtr<TArray64<uint8>> Payloads[NumPayloads];
+	TArray<uint32> WrittenExportIndices[NumPayloads];
+
+	for (uint32 It=0; It<NumPayloads; ++It)
+	{
+		// We use this to capture exports written
+		FChunkedDataStreamExportWriteScope ScopedExports(ServerStream);
+
+		Payloads[It] = MakeShared<TArray64<uint8>>();
+		Payloads[It]->SetNumZeroed(FMath::RandRange(10, 5000));
+
+		// Randomize a few object references and write them to the paylaod using the PackageMap associated with the DataStream
+		{
+			FMemoryWriter64 Ar(*Payloads[It]);
+
+			int32 ReferenceCount = FMath::RandRange(0, MaxObjectRefsIndex);
+			Ar << ReferenceCount;
+			for (int32 RefIt = 0; RefIt < ReferenceCount; ++RefIt)
+			{
+				const int32 ReferenceIndex = FMath::RandRange(0, MaxObjectRefsIndex - 1);
+				WrittenExportIndices[It].Add(ReferenceIndex);
+				ScopedExports.GetPackageMap()->SerializeObject(Ar, UObject::StaticClass(), ObjectRefs[ReferenceIndex], nullptr);
+			}
+		}
+
+		ServerStream->EnqueuePayload(Payloads[It]);
+	}
+
+	uint32 NumRecvdPayloads = 0;	
+	while (NumRecvdPayloads < NumPayloads)
+	{
+		RoundTripWithLatencyAndDeliverPercentage(0, 0, 1500, 1500);
+
+		// Init Packagemap for reading references
+		FChunkedDataStreamExportReadScope ScopedExports(ClientStream);
+
+		ClientStream->DispatchReceivedPayloads([this, &ScopedExports, &Payloads, &ObjectRefs, &WrittenExportIndices, &NumRecvdPayloads](TConstArrayView64<uint8> Data)
+		{
+			UE_NET_ASSERT_EQ(Data.Num(), Payloads[NumRecvdPayloads]->Num());
+
+			// Read payload including references and validate them.
+			{
+				FMemoryReaderView Ar(MakeMemoryView(Data.GetData(), Data.Num()));
+
+				// Read reference count.
+				int32 ReferenceCount = 0U;
+				Ar << ReferenceCount;
+
+				// Validate that we recevied the expected number of references
+				UE_NET_ASSERT_EQ(ReferenceCount, WrittenExportIndices[NumRecvdPayloads].Num());
+
+				// Validate references
+				for (int32 RefIt = 0; RefIt < ReferenceCount; ++RefIt)
+				{
+					UObject* SomePtr = nullptr;
+					ScopedExports.GetPackageMap()->SerializeObject(Ar, UObject::StaticClass(), SomePtr, nullptr);
+
+					const int32 ReferenceIndex = WrittenExportIndices[NumRecvdPayloads][RefIt];
+					UE_NET_ASSERT_EQ(SomePtr, ObjectRefs[ReferenceIndex]);
+				}
+			}
+
+			++NumRecvdPayloads;
+		});
+	}
+	UE_NET_ASSERT_EQ(NumRecvdPayloads, NumPayloads);
+}
+
+UE_NET_TEST_FIXTURE(FTestChunkedDataStreamFixture, TestChunkedDataStream_CannotEnqueueMoreThanMaxEnqueuedPayloadBytes)
+{
+	UReplicationSystem* ServerReplicationSystem = Server->ReplicationSystem;
+	UReplicationSystem* ClientReplicationSystem = Client->ReplicationSystem;
+
+	// Open stream
+	UChunkedDataStream* ServerStream = Cast<UChunkedDataStream>(ServerReplicationSystem->OpenDataStream(Client->ConnectionIdOnServer, DataStreamName));
+	RoundTrip();
+	RoundTrip();
+
+	UChunkedDataStream* ClientStream = Cast<UChunkedDataStream>(ClientReplicationSystem->GetDataStream(Client->LocalConnectionId, DataStreamName));
+
+	// Set limit
+	ServerStream->SetMaxEnqueuedPayloadBytes(1500);
+
+	// Try to enqueue some data
+	{
+		TSharedRef<TArray64<uint8>> Payload = MakeShared<TArray64<uint8>>();
+		Payload->SetNumZeroed(1024);
+		const bool bResult = ServerStream->EnqueuePayload(Payload);
+		UE_NET_ASSERT_TRUE(bResult);
+	}
+
+	// Try to enqueue some data that should not be allowed
+	{
+		TSharedRef<TArray64<uint8>> Payload = MakeShared<TArray64<uint8>>();
+		Payload->SetNumZeroed(1024);
+		const bool bResult = ServerStream->EnqueuePayload(Payload);
+		UE_NET_ASSERT_FALSE(bResult);
+	}
+}
+
+UE_NET_TEST_FIXTURE(FTestChunkedDataStreamFixture, TestChunkedDataStream_ClientWillSetErrorIfTooManyUndispatchedPayloadBytes)
+{
+	UReplicationSystem* ServerReplicationSystem = Server->ReplicationSystem;
+	UReplicationSystem* ClientReplicationSystem = Client->ReplicationSystem;
+
+	// Open stream
+	UChunkedDataStream* ServerStream = Cast<UChunkedDataStream>(ServerReplicationSystem->OpenDataStream(Client->ConnectionIdOnServer, DataStreamName));
+	RoundTrip();
+	RoundTrip();
+
+	UChunkedDataStream* ClientStream = Cast<UChunkedDataStream>(ClientReplicationSystem->GetDataStream(Client->LocalConnectionId, DataStreamName));
+
+	// Set limit on receiving end on how much undispatched data we should allow to be queued up.
+	ClientStream->SetMaxUndispatchedPayloadBytes(2500);
+
+	// Enqueue some data
+	{
+		TSharedRef<TArray64<uint8>> Payload = MakeShared<TArray64<uint8>>();
+		Payload->SetNumZeroed(1024);
+		ServerStream->EnqueuePayload(Payload);
+	}
+	{
+		TSharedRef<TArray64<uint8>> Payload = MakeShared<TArray64<uint8>>();
+		Payload->SetNumZeroed(1024);
+		ServerStream->EnqueuePayload(Payload);
+	}
+
+	// Tick send/receive
+	RoundTrip();
+	RoundTrip();
+
+	// Verify that we have got the first two expected payloads
+	UE_NET_ASSERT_EQ(ClientStream->GetNumReceivedPayloadsPendingDispatch(), 2U);
+
+	// No error detected
+	UE_NET_ASSERT_FALSE(ClientStream->HasError());
+
+	// Send more data that will exceed the set limit.
+	{
+		TSharedRef<TArray64<uint8>> Payload = MakeShared<TArray64<uint8>>();
+		Payload->SetNumZeroed(10000);
+		ServerStream->EnqueuePayload(Payload);
+	}
+
+	// Tick send/receive
+	RoundTrip();
+
+	// Stream should now be in error state, and will ignore incoming chunks (still read from bitstream but they will be discarded.)
+	UE_NET_ASSERT_TRUE(ClientStream->HasError());
+
+	// Tick send/receive
+	RoundTrip();
+	RoundTrip();
+	RoundTrip();
+	RoundTrip();
+	RoundTrip();
+	RoundTrip();
+
+	// Verify that we have got the first two expected payloads
+	UE_NET_ASSERT_EQ(ClientStream->GetNumReceivedPayloadsPendingDispatch(), 2U);
+}
+
+}	
diff --git a/Engine/Source/Runtime/Experimental/Iris/Core/Private/Iris/ReplicationSystem/ChunkedDataStream/ChunkedDataReader.cpp b/Engine/Source/Runtime/Experimental/Iris/Core/Private/Iris/ReplicationSystem/ChunkedDataStream/ChunkedDataReader.cpp
new file mode 100644
--- /dev/null
+++ b/Engine/Source/Runtime/Experimental/Iris/Core/Private/Iris/ReplicationSystem/ChunkedDataStream/ChunkedDataReader.cpp
@@ -0,0 +1,494 @@
+// Copyright Epic Games, Inc. All Rights Reserved.
+
+#include "ChunkedDataReader.h"
+
+#include "Iris/PacketControl/PacketNotification.h"
+#include "Iris/ReplicationSystem/ReplicationSystem.h"
+#include "Iris/ReplicationSystem/ReplicationSystemInternal.h"
+#include "Iris/Serialization/NetBitStreamReader.h"
+#include "Iris/Serialization/NetSerializationContext.h"
+#include "Iris/Serialization/NetBitStreamUtil.h"
+#include "Net/Core/Trace/NetTrace.h"
+#include "Misc/ScopeExit.h"
+
+namespace UE::Net::Private
+{
+
+FChunkedDataReader::FDataChunk::FDataChunk()
+: PartCount(0U)
+, SequenceNumber(uint16(-1))
+, PartByteCount(0U)
+, bIsFirstChunk(0U)
+, bIsExportChunk(0U)
+{
+}
+
+const uint32 FChunkedDataReader::FDataChunk::GetPartPayloadByteCount() const
+{
+	// Size is encoded by a combination of PartCount and PayloadByteCount
+	// The first part contains the size of the entire payload encoded as PartCount * ChunkSize + PayloadByteCount
+	return (bIsFirstChunk && (PartCount > 1U)) ? FChunkedDataStreamParameters::ChunkSize : PartByteCount;
+}
+
+void FChunkedDataReader::FDataChunk::Deserialize(UE::Net::FNetSerializationContext& Context)
+{
+	FNetBitStreamReader* Reader = Context.GetBitStreamReader();
+
+	bIsFirstChunk = Reader->ReadBool() ? 1U : 0U;
+	if (bIsFirstChunk)
+	{
+		bIsExportChunk = Reader->ReadBool();
+		PartCount = ReadPackedUint32(Reader);
+	}
+	else
+	{
+		bIsExportChunk = 0U;
+		PartCount = 0U;
+	}
+			
+	const bool bIsFullChunk = Reader->ReadBool();
+	const uint32 ReadPartByteCount = bIsFullChunk ? FChunkedDataStreamParameters::ChunkSize : ReadPackedUint16(Reader);
+
+	if (Context.HasErrorOrOverflow())
+	{
+		return;
+	}
+
+	PartByteCount = (uint16)FMath::Min(FChunkedDataStreamParameters::ChunkSize, ReadPartByteCount);
+
+	// Read actual payload
+	UE_NET_TRACE_SCOPE(Payload, *Reader, Context.GetTraceCollector(), ENetTraceVerbosity::Verbose);
+
+	// Size is encoded as a combination of PartCount and PayloadByteCount
+	// The first part contains the size of the entire payload encoded as PartCount * ChunkSize + PayloadByteCount
+	// So the actual payload size of the first part depends on if there are multiple parts or not.
+	const uint32 PartPayloadByteCount = (bIsFirstChunk && (PartCount > 1U)) ? FChunkedDataStreamParameters::ChunkSize : PartByteCount;
+	PartPayload.SetNum(PartPayloadByteCount);
+	ReadBytes(Reader, PartPayload.GetData(), PartPayloadByteCount);			
+}
+
+FChunkedDataReader::FChunkedDataReader(const UDataStream::FInitParameters& InParams)
+: InitParams(InParams)
+, ReplicationSystem(UE::Net::GetReplicationSystem(InParams.ReplicationSystemId))
+, ObjectReferenceCache(&ReplicationSystem->GetReplicationSystemInternal()->GetObjectReferenceCache())
+{
+	FNetTokenStore* NetTokenStore = ReplicationSystem->GetNetTokenStore();
+
+	// Setup internal context
+	ResolveContext.ConnectionId = InitParams.ConnectionId;
+	ResolveContext.RemoteNetTokenStoreState = NetTokenStore->GetRemoteNetTokenStoreState(InitParams.ConnectionId);
+
+	NetTokenResolveContext.NetTokenStore = NetTokenStore;
+	NetTokenResolveContext.RemoteNetTokenStoreState = ResolveContext.RemoteNetTokenStoreState;
+}
+
+void FChunkedDataReader::ResetResolvedReferences()
+{
+	// Make sure to release all references that we are holding on to
+	if (ObjectReferenceCache)
+	{
+		for (const FNetRefHandle& RefHandle : ResolvedReferences)
+		{
+			ObjectReferenceCache->RemoveTrackedQueuedBatchObjectReference(RefHandle);
+		}
+	}
+}
+
+FChunkedDataReader::~FChunkedDataReader()
+{
+	ResetResolvedReferences();
+}
+
+bool FChunkedDataReader::ProcessExportPayload(FNetSerializationContext& Context, FRecvQueueEntry& Entry)
+{
+	FNetBitStreamReader ExportsReader;
+	ExportsReader.InitBits(Entry.Payload.GetData(), Entry.Payload.Num() * 8U);
+	FNetSerializationContext SubContext = Context.MakeSubContext(&ExportsReader);
+
+	FNetTraceCollector* ExportsTraceCollector = nullptr;
+#if UE_NET_TRACE_ENABLED
+	FNetTraceCollector ExportsTraceCollectorOnStack;
+	ExportsTraceCollector = &ExportsTraceCollectorOnStack;
+#endif
+
+	SubContext.SetTraceCollector(ExportsTraceCollector);
+
+	UE_NET_TRACE_NAMED_SCOPE(ExportsTraceScope, ExportsPayload, ExportsReader, ExportsTraceCollector, ENetTraceVerbosity::Trace);
+
+ON_SCOPE_EXIT
+	{
+#if UE_NET_TRACE_ENABLED
+		UE_NET_TRACE_EXIT_NAMED_SCOPE(ExportsTraceScope);
+
+		// Append huge object state at end of stream.
+		if (FNetTraceCollector* TraceCollector = Context.GetTraceCollector())
+		{
+			FNetBitStreamReader& Reader = *Context.GetBitStreamReader();
+			// Inject after all other trace events
+			const uint32 LevelOffset = 3U;
+			FNetTrace::FoldTraceCollector(TraceCollector, ExportsTraceCollector, GetBitStreamPositionForNetTrace(Reader) + MultiExportsPayLoadOffset, LevelOffset);
+
+			MultiExportsPayLoadOffset += ExportsReader.GetPosBits();;
+		}
+#endif
+	};
+
+	// Read and process exports
+	const uint32 ExportsOffset = ExportsReader.ReadBits(FChunkedDataStreamParameters::NumBitsForExportOffset);
+
+	if (SubContext.HasErrorOrOverflow())
+	{
+		return false;
+	}
+
+	uint32 ExportsEndPosition = 0U;
+	if (ExportsOffset != 0U)
+	{
+		const uint32 ReturnPos = ExportsReader.GetPosBits();
+		ExportsReader.Seek(ReturnPos + ExportsOffset);
+
+		if (!ObjectReferenceCache->ReadExports(SubContext, &Entry.References->MustBeMappedReferences))
+		{
+			return false;
+		}
+
+		ExportsEndPosition = ExportsReader.GetPosBits();
+		ExportsReader.Seek(ReturnPos);
+	}
+	FIrisPackageMapExportsUtil::Deserialize(SubContext, Entry.References->QuantizedExports);
+
+	// Just to get tracing to report nicely.
+#if UE_NET_TRACE_ENABLED
+	if (!SubContext.HasErrorOrOverflow() && ExportsOffset != 0U)
+	{
+		ExportsReader.Seek(ExportsEndPosition);
+	}
+#endif
+
+	return !SubContext.HasErrorOrOverflow();
+}
+
+void FChunkedDataReader::AssemblePayloadsPendingAssembly(UE::Net::FNetSerializationContext& Context)
+{
+	// Reset ExportsPayLoadOffset
+	MultiExportsPayLoadOffset = 0U;
+	while (!DataChunksPendingAssembly.IsEmpty())
+	{
+		const FDataChunk& CurrentChunk = DataChunksPendingAssembly.First();
+
+		if (CurrentChunk.SequenceNumber != ExpectedSeq)
+		{
+			break;
+		}
+
+		// If we have encountered an error, we will no longer try to assemble received chunks
+		if (!HasError())
+		{
+			bool bLocalHasError = false;
+			if (CurrentChunk.bIsFirstChunk)
+			{
+				FRecvQueueEntry* CurrentEntry = ReceiveQueue.IsEmpty() ? nullptr : &ReceiveQueue.Last();
+
+				// Validate that previous entry is complete
+				if (ensure(!CurrentEntry || (CurrentEntry->RemainingByteCount == 0U)))
+				{
+					// If the last received payload was an export we append to the same entry.
+					if (!CurrentEntry || CurrentChunk.bIsExportChunk || !(CurrentEntry->bHasProcessedExports && CurrentEntry->Payload.Num() == 0U))
+					{
+						// Create new chunk to dispatch and start assembling payload
+						CurrentEntry = &ReceiveQueue.Emplace_GetRef();							
+
+						if (CurrentChunk.bIsExportChunk)
+						{
+							CurrentEntry->References = MakeUnique<FReferencesForImport>();					
+						}
+					}
+
+					CurrentEntry->RemainingByteCount = ((CurrentChunk.PartCount - 1U) * FChunkedDataStreamParameters::ChunkSize) + CurrentChunk.PartByteCount;
+						
+					UE_LOG_CHUNKEDDATASTREAM_CONN(Verbose, TEXT("AssemblePayloadsPendingAssembly Size: %u PartCount: %u IsExportPayload: %u"), CurrentEntry->RemainingByteCount, CurrentChunk.PartCount, CurrentChunk.bIsExportChunk);
+
+					if (CurrentUndispatchedPayloadBytes + CurrentEntry->RemainingByteCount > MaxUndispatchedPayloadBytes)
+					{
+						SetError(TEXT("Error: MaxUndispatchedPayloadBytes exceeded."));
+					}
+					else
+					{
+						CurrentEntry->Payload.Reserve(CurrentEntry->RemainingByteCount);
+						CurrentUndispatchedPayloadBytes += CurrentEntry->RemainingByteCount;
+					}
+				}
+				else
+				{
+					SetError(TEXT("Error: Encountered new payload when previous one still is not fully received, DataStream will be closed."));
+				}
+			}
+			if (!HasError())
+			{
+				FRecvQueueEntry& CurrentEntry = ReceiveQueue.Last();
+				if (ensure((uint32)CurrentChunk.PartPayload.Num() <= CurrentEntry.RemainingByteCount))
+				{
+					CurrentEntry.Payload.Append(CurrentChunk.PartPayload);
+					CurrentEntry.RemainingByteCount -= CurrentChunk.PartPayload.Num();
+						
+					if (CurrentEntry.RemainingByteCount == 0 && !CurrentEntry.bHasProcessedExports && CurrentEntry.References)
+					{
+						// Read and process exports as soon as export payload is assembled
+						if (ProcessExportPayload(Context, CurrentEntry))
+						{
+							// If the export payload has been processed, we reset the payload so that we can reuse the same entry for the actual data payload
+							CurrentUndispatchedPayloadBytes -= CurrentEntry.Payload.Num();
+							CurrentEntry.Payload.Reset();
+							CurrentEntry.bHasProcessedExports = true;
+						}
+						else
+						{
+							SetError(TEXT("Error: Failed to ProcessExportPayload, DataStream will be closed."));
+							bLocalHasError = true;
+						}
+					}
+				}
+				else
+				{
+					SetError(TEXT("Error: Received more data than expected when assembling payload, DataStream will be closed."));
+				}
+			}
+		}
+
+		// We are done with this chunk
+		++ExpectedSeq;
+		DataChunksPendingAssembly.PopFront();
+	}
+
+	if (DataChunksPendingAssembly.IsEmpty())
+	{
+		DataChunksPendingAssembly.Trim();
+	}
+}
+
+bool FChunkedDataReader::TryResolveUnresolvedMustBeMappedReferences(TArray<FNetRefHandle>& MustBeMappedReferences)
+{
+	// Resolve
+	TArray<FNetRefHandle> Unresolved;
+	TArray<TPair<FNetRefHandle, UObject*>, TInlineAllocator<4>> QueuedObjectsToTrack;
+
+	Unresolved.Reserve(MustBeMappedReferences.Num());
+	QueuedObjectsToTrack.SetNum(MustBeMappedReferences.Num());
+
+	// Try to resolve references
+	for (FNetRefHandle Handle : MustBeMappedReferences)
+	{
+		UObject* ResolvedObject = nullptr;
+		ENetObjectReferenceResolveResult ResolveResult = ObjectReferenceCache->ResolveObjectReference(FObjectReferenceCache::MakeNetObjectReference(Handle), ResolveContext, ResolvedObject);
+		if (EnumHasAnyFlags(ResolveResult, ENetObjectReferenceResolveResult::HasUnresolvedMustBeMappedReferences) && !ObjectReferenceCache->IsNetRefHandleBroken(Handle, true))
+		{
+			Unresolved.Add(Handle);
+		}
+		else if (ResolveResult == ENetObjectReferenceResolveResult::None)
+		{
+			QueuedObjectsToTrack.Emplace(Handle, ResolvedObject);
+		}
+	}
+
+	// If more references are resolved, add them to tracking list
+	for (TPair<FNetRefHandle, UObject*>& NetRefHandleObjectPair : QueuedObjectsToTrack)
+	{
+		if (!ResolvedReferences.Contains(NetRefHandleObjectPair.Key))
+		{
+			ResolvedReferences.Add(NetRefHandleObjectPair.Key);
+			ObjectReferenceCache->AddTrackedQueuedBatchObjectReference(NetRefHandleObjectPair.Key, NetRefHandleObjectPair.Value);
+		}
+	}
+
+	if (Unresolved.Num())
+	{
+		// Upate status
+		MustBeMappedReferences = MoveTemp(Unresolved);
+	
+		return false;
+	}
+
+	// Nothing more to do.
+	MustBeMappedReferences.Reset();
+
+	return true;	
+}
+
+UChunkedDataStream::EDispatchResult FChunkedDataReader::DispatchReceivedPayload(TFunctionRef<void(TConstArrayView64<uint8>)> DispatchPayloadFunction)
+{
+	if (ReceiveQueue.IsEmpty() || HasError())
+	{
+		return UChunkedDataStream::EDispatchResult::NothingToDispatch;
+	}
+
+	FRecvQueueEntry& Entry = ReceiveQueue.First();
+	if (Entry.RemainingByteCount != 0U)
+	{
+		return UChunkedDataStream::EDispatchResult::NothingToDispatch;
+	}
+
+	const bool bProcessReferences = Entry.References.IsValid();
+	if (bProcessReferences)
+	{					
+		if (ObjectReferenceCache->ShouldAsyncLoad() && !TryResolveUnresolvedMustBeMappedReferences(Entry.References->MustBeMappedReferences))
+		{
+			// Wait for async loading to complete if we have any mustbemapped entries
+			UE_LOG_CHUNKEDDATASTREAM_CONN(Verbose, TEXT("Waiting for %d MustBeMapped references to be resolvable"), Entry.References->MustBeMappedReferences.Num());
+			return UChunkedDataStream::EDispatchResult::WaitingForMustBeMappedReferences;
+		}
+
+		// Setup context for dispatch
+		FInternalNetSerializationContext InternalContext;
+		FInternalNetSerializationContext::FInitParameters InternalContextInitParams;
+		InternalContextInitParams.ReplicationSystem = ReplicationSystem;
+		InternalContextInitParams.ObjectResolveContext = ResolveContext;
+		InternalContext.Init(InternalContextInitParams);
+
+		FNetSerializationContext Context;
+		Context.SetLocalConnectionId(InitParams.ConnectionId);
+		Context.SetInternalContext(&InternalContext);
+
+		// Dequantize exports
+		FIrisPackageMapExportsUtil::Dequantize(Context, Entry.References->QuantizedExports, PackageMapExports);
+	}
+
+	if (Entry.Payload.Num())
+	{
+		UE_LOG_CHUNKEDDATASTREAM_CONN(Verbose, TEXT("Dispatching payload with %u bytes and %u potential exports"), Entry.Payload.Num(), PackageMapExports.References.Num());
+		DispatchPayloadFunction(MakeConstArrayView(Entry.Payload));
+		CurrentUndispatchedPayloadBytes -= Entry.Payload.Num();
+	}
+
+	if (!bProcessReferences)
+	{
+		ResetResolvedReferences();
+	}
+	ReceiveQueue.PopFront();
+
+	if (ReceiveQueue.IsEmpty())
+	{
+		ReceiveQueue.Trim();
+	}
+
+	return UChunkedDataStream::EDispatchResult::Ok;
+}
+
+UChunkedDataStream::EDispatchResult FChunkedDataReader::DispatchReceivedPayloads(TFunctionRef<void(TConstArrayView64<uint8>)> DispatchPayloadFunction)
+{		
+	UChunkedDataStream::EDispatchResult Result = UChunkedDataStream::EDispatchResult::Ok;
+	for (; Result == UChunkedDataStream::EDispatchResult::Ok; Result = DispatchReceivedPayload(DispatchPayloadFunction))
+	{
+	}
+
+	return Result;
+}
+
+uint32 FChunkedDataReader::GetNumReceivedPayloadsPendingDispatch() const
+{
+	uint32 ReadyToDispatchCount = 0U;
+	for (const FRecvQueueEntry& Entry : ReceiveQueue)
+	{
+		ReadyToDispatchCount += Entry.RemainingByteCount == 0U ? 1U : 0U;
+	}
+	return ReadyToDispatchCount;
+}
+
+void FChunkedDataReader::ReadData(UE::Net::FNetSerializationContext& Context)
+{
+	FNetBitStreamReader* Reader = Context.GetBitStreamReader();
+
+	// $TODO: set this up in DataStreamManager instead. UE-243627
+	FInternalNetSerializationContext InternalContext;
+	FInternalNetSerializationContext::FInitParameters InternalContextInitParams;
+	InternalContextInitParams.ReplicationSystem = ReplicationSystem;
+	InternalContextInitParams.ObjectResolveContext = ResolveContext;
+	InternalContext.Init(InternalContextInitParams);
+
+	Context.SetLocalConnectionId(InitParams.ConnectionId);
+	Context.SetInternalContext(&InternalContext);
+
+	uint16 LastReadSeq = uint16(-1);
+	while (Reader->ReadBool())
+	{
+		if (Reader->IsOverflown())
+		{
+			break;
+		}
+
+		UE_NET_TRACE_SCOPE(DataChunk, *Reader, Context.GetTraceCollector(), ENetTraceVerbosity::Verbose);
+
+		UE_NET_TRACE_NAMED_DYNAMIC_NAME_SCOPE(SequenceScope, static_cast<const TCHAR*>(nullptr), *Reader, Context.GetTraceCollector(), ENetTraceVerbosity::VeryVerbose);
+
+		// Read sequence number
+		const bool bIsInSequence = Reader->ReadBool();
+		uint16 ReadSeq = 0U;
+		if (bIsInSequence)
+		{
+			ReadSeq = (LastReadSeq + 1U) & FChunkedDataStreamParameters::SequenceBitMask;
+		}
+		else
+		{
+			ReadSeq = (uint16)Reader->ReadBits(FChunkedDataStreamParameters::SequenceBitCount);
+		}
+
+		if (Reader->IsOverflown())
+		{
+			break;
+		}
+
+		LastReadSeq = ReadSeq;
+
+		const uint16 Seq = LastReadSeq;
+		const uint16 SeqDelta = (Seq - ExpectedSeq) & FChunkedDataStreamParameters::SequenceBitMask;
+
+		// Make room to store missing sequence numbers.
+		{
+			DataChunksPendingAssembly.Reserve(FMath::Max(DataChunksPendingAssembly.Num(), (int32)(SeqDelta + 1)));
+			while ((uint32)DataChunksPendingAssembly.Num() <= SeqDelta)
+			{
+				DataChunksPendingAssembly.Add(FDataChunk());
+			}
+		}
+
+		FDataChunk& Chunk = DataChunksPendingAssembly[SeqDelta];
+		Chunk.SequenceNumber = ExpectedSeq + SeqDelta;
+		Chunk.Deserialize(Context);
+
+#if UE_NET_TRACE_ENABLED
+		if (FNetTrace::GetNetTraceVerbosityEnabled(ENetTraceVerbosity::VeryVerbose))
+		{
+			TStringBuilder<64> Builder;
+			if (Chunk.bIsFirstChunk)
+			{
+				Builder.Appendf(TEXT("Seq %u First part of %u"), Chunk.SequenceNumber, Chunk.PartCount);
+			}
+			else
+			{
+				Builder.Appendf(TEXT("Seq %u"), Chunk.SequenceNumber);
+			}
+			UE_NET_TRACE_SET_SCOPE_NAME(SequenceScope, Builder.ToString());
+		}
+#endif // UE_NET_TRACE_ENABLED
+
+		UE_LOG_CHUNKEDDATASTREAM_CONN(Verbose, TEXT("Deserialize Seq:%u (local:%u), Expected %u"), Seq, Chunk.SequenceNumber, ExpectedSeq);
+	}
+
+	// Assemble data chunks that we have received
+	AssemblePayloadsPendingAssembly(Context);
+}
+
+void FChunkedDataReader::SetError(const FString& InErrorMessage)
+{
+	UE_LOG_CHUNKEDDATASTREAM_CONN(Error, TEXT("FChunkedDataReader::ErrorEncountered() %s"), *InErrorMessage);
+	if (!bHasError)
+	{
+		bHasError = true;
+	}
+}
+
+bool FChunkedDataReader::HasError() const
+{
+	return bHasError;
+}
+
+} // End of namespace(s)
diff --git a/Engine/Source/Runtime/Experimental/Iris/Core/Private/Iris/ReplicationSystem/ChunkedDataStream/ChunkedDataReader.h b/Engine/Source/Runtime/Experimental/Iris/Core/Private/Iris/ReplicationSystem/ChunkedDataStream/ChunkedDataReader.h
new file mode 100644
--- /dev/null
+++ b/Engine/Source/Runtime/Experimental/Iris/Core/Private/Iris/ReplicationSystem/ChunkedDataStream/ChunkedDataReader.h
@@ -0,0 +1,125 @@
+// Copyright Epic Games, Inc. All Rights Reserved.
+
+#pragma once
+
+#include "ChunkedDataStreamCommon.h"
+#include "Containers/Array.h"
+#include "Containers/RingBuffer.h"
+#include "Iris/DataStream/DataStream.h"
+#include "Iris/ReplicationSystem/ChunkedDataStream/ChunkedDataStream.h"
+#include "Iris/ReplicationSystem/ObjectReferenceCacheFwd.h"
+#include "Iris/Serialization/IrisObjectReferencePackageMap.h"
+#include "Iris/Serialization/IrisPackageMapExportUtil.h"
+#include "Iris/Serialization/NetExportContext.h"
+#include "Net/Core/NetToken/NetTokenExportContext.h"
+
+namespace UE::Net::Private
+{
+
+// Used by ChunkedDataStream to reading and dispatching incoming data
+class FChunkedDataReader
+{
+public:
+
+	using EWriteResult = UDataStream::EWriteResult;
+	using FBeginWriteParameters = UDataStream::FBeginWriteParameters;
+	using FInitParameters = UDataStream::FInitParameters;
+
+	struct FReferencesForImport
+	{
+		FIrisPackageMapExportsQuantizedType QuantizedExports;
+		TArray<FNetRefHandle> MustBeMappedReferences;
+
+		~FReferencesForImport()
+		{
+			FIrisPackageMapExportsUtil::FreeDynamicState(QuantizedExports);
+		}
+	};
+
+	// RecvQueue entry
+	struct FRecvQueueEntry
+	{
+		// Returns true if this is a processed export payload
+		bool GetIsProcessedExportPayload() const
+		{
+			return bHasProcessedExports;
+		}
+
+		TArray<uint8, TAlignedHeapAllocator<4>> Payload;
+		TUniquePtr<FReferencesForImport> References;
+		uint32 RemainingByteCount = 0U;
+		bool bHasProcessedExports = false;
+	};
+
+	struct FDataChunk
+	{
+		FDataChunk();
+
+		const uint32 GetPartPayloadByteCount() const;
+		void Deserialize(UE::Net::FNetSerializationContext& Context);
+
+		TArray<uint8> PartPayload;
+		uint32 PartCount;
+		uint16 SequenceNumber;
+		uint16 PartByteCount : 14;
+		uint16 bIsFirstChunk : 1;
+		uint16 bIsExportChunk : 1;
+	};
+
+public:
+
+	FChunkedDataReader(const UDataStream::FInitParameters& InParams);
+	~FChunkedDataReader();
+
+	bool ProcessExportPayload(FNetSerializationContext& Context, FRecvQueueEntry& Entry);
+	void AssemblePayloadsPendingAssembly(UE::Net::FNetSerializationContext& Context);
+	bool TryResolveUnresolvedMustBeMappedReferences(TArray<FNetRefHandle>& MustBeMappedReferences);
+	UChunkedDataStream::EDispatchResult DispatchReceivedPayload(TFunctionRef<void(TConstArrayView64<uint8>)> DispatchPayloadFunction);
+	UChunkedDataStream::EDispatchResult DispatchReceivedPayloads(TFunctionRef<void(TConstArrayView64<uint8>)> DispatchPayloadFunction);
+
+	uint32 GetNumReceivedPayloadsPendingDispatch() const;
+	void ReadData(UE::Net::FNetSerializationContext& Context);
+	void SetError(const FString& InErrorMessage);
+	bool HasError() const;
+	void ResetResolvedReferences();
+
+public:
+
+	friend class FChunkedDataStreamExportReadScope;
+
+	// Incoming data
+	TRingBuffer<FDataChunk> DataChunksPendingAssembly;
+
+	// Received data, ready to dispatch
+	TRingBuffer<FRecvQueueEntry> ReceiveQueue;
+
+	// Next expected sequence number
+	uint16 ExpectedSeq = 0;
+
+	// We have encountered and error, and should close the DataStream
+	bool bHasError = false;
+
+	// Cached on init
+	FInitParameters InitParams;
+	UReplicationSystem* ReplicationSystem = nullptr;
+	FObjectReferenceCache* ObjectReferenceCache = nullptr;
+	FNetObjectResolveContext ResolveContext;
+	FNetTokenResolveContext NetTokenResolveContext;
+
+	// Resolved references for which we have are holding on to references to avoid GC, must be released on exit
+	TArray<FNetRefHandle, TInlineAllocator<4>> ResolvedReferences;
+
+	// Exports
+	UE::Net::FIrisPackageMapExports PackageMapExports;
+
+	// Maximum undispatched payload bytes, if this is overflown datastream will be put in error state and closed
+	uint64 MaxUndispatchedPayloadBytes = 10485760U;
+
+	// Current number of received payload bytes ready to dispatch
+	uint64 CurrentUndispatchedPayloadBytes = 0U;
+
+	// Offset used when folding multiple exports payload processed after reading the same pacekt
+	uint32 MultiExportsPayLoadOffset = 0U;
+};
+
+} // End of namespace(s)
diff --git a/Engine/Source/Runtime/Experimental/Iris/Core/Private/Iris/ReplicationSystem/ChunkedDataStream/ChunkedDataStream.cpp b/Engine/Source/Runtime/Experimental/Iris/Core/Private/Iris/ReplicationSystem/ChunkedDataStream/ChunkedDataStream.cpp
new file mode 100644
--- /dev/null
+++ b/Engine/Source/Runtime/Experimental/Iris/Core/Private/Iris/ReplicationSystem/ChunkedDataStream/ChunkedDataStream.cpp
@@ -0,0 +1,185 @@
+// Copyright Epic Games, Inc. All Rights Reserved.
+
+#include "Iris/ReplicationSystem/ChunkedDataStream/ChunkedDataStream.h"
+#include "ChunkedDataReader.h"
+#include "ChunkedDataWriter.h"
+
+DEFINE_LOG_CATEGORY(LogIrisChunkedDataStream);
+
+bool UChunkedDataStream::EnqueuePayload(const TSharedPtr<TArray64<uint8>>& Payload)
+{
+	return ChunkedWriter->EnqueuePayload(Payload);
+}
+
+void UChunkedDataStream::Init(const UDataStream::FInitParameters& Params)
+{
+	Super::Init(Params);
+
+	using namespace UE::Net::Private;
+
+	PackageMap = TObjectPtr<UIrisObjectReferencePackageMap>(NewObject<UIrisObjectReferencePackageMap>());
+
+	ChunkedWriter = new FChunkedDataWriter(Params);
+	ChunkedReader = new FChunkedDataReader(Params);
+}
+
+void UChunkedDataStream::Deinit()
+{
+	Super::Deinit();
+
+	if (ChunkedWriter)
+	{
+		delete ChunkedWriter;
+		ChunkedWriter = nullptr;
+	}
+	if (ChunkedReader)
+	{
+		delete ChunkedReader;
+		ChunkedReader = nullptr;
+	}
+
+	// Release Packagemap
+	PackageMap = nullptr;
+}
+
+UDataStream::EWriteResult UChunkedDataStream::BeginWrite(const FBeginWriteParameters& Params)
+{
+	if (ChunkedWriter != nullptr)
+	{
+		return ChunkedWriter->BeginWrite(Params);
+	}
+
+	return EWriteResult::NoData;
+}
+
+UDataStream::EWriteResult UChunkedDataStream::WriteData(UE::Net::FNetSerializationContext& Context, FDataStreamRecord const*& OutRecord)
+{
+	if (ChunkedWriter != nullptr)
+	{
+		return ChunkedWriter->WriteData(Context, OutRecord);
+	}
+
+	return EWriteResult::NoData;
+}
+
+void UChunkedDataStream::ReadData(UE::Net::FNetSerializationContext& Context)
+{
+	if (ChunkedReader != nullptr)
+	{
+		return ChunkedReader->ReadData(Context);
+	}
+}
+
+void UChunkedDataStream::ProcessPacketDeliveryStatus(UE::Net::EPacketDeliveryStatus Status, FDataStreamRecord const* InRecord)
+{
+	if (ChunkedWriter)
+	{
+		ChunkedWriter->ProcessPacketDeliveryStatus(Status, InRecord);
+	}
+}
+
+bool UChunkedDataStream::HasAcknowledgedAllReliableData() const
+{
+	if (ChunkedWriter)
+	{
+		return ChunkedWriter->HasAcknowledgedAllReliableData();
+	}
+
+	return true;
+}
+
+UChunkedDataStream::EDispatchResult UChunkedDataStream::DispatchReceivedPayloads(TFunctionRef<void(TConstArrayView64<uint8>)> DispatchPayloadFunction)
+{
+	if (!ChunkedReader)
+	{
+		return EDispatchResult::NothingToDispatch;
+	}
+
+	return ChunkedReader->DispatchReceivedPayloads(DispatchPayloadFunction);
+}
+
+UChunkedDataStream::EDispatchResult UChunkedDataStream::DispatchReceivedPayload(TFunctionRef<void(TConstArrayView64<uint8>)> DispatchPayloadFunction)
+{
+	if (!ChunkedReader)
+	{
+		return EDispatchResult::NothingToDispatch;
+	}
+
+	return ChunkedReader->DispatchReceivedPayload(DispatchPayloadFunction);
+}
+
+uint32 UChunkedDataStream::GetNumReceivedPayloadsPendingDispatch() const
+{
+	if (ChunkedReader)
+	{
+		return ChunkedReader->GetNumReceivedPayloadsPendingDispatch();
+	}
+
+	return 0U;
+}
+
+UIrisObjectReferencePackageMap* UChunkedDataStream::GetPackageMap()
+{
+	return PackageMap.Get();
+}
+
+const UIrisObjectReferencePackageMap* UChunkedDataStream::GetPackageMap() const
+{
+	return PackageMap.Get();
+}
+
+FChunkedDataStreamExportWriteScope::FChunkedDataStreamExportWriteScope(UChunkedDataStream* DataStream)
+: WriteScope(DataStream ? DataStream->GetPackageMap() : nullptr, DataStream && DataStream->ChunkedWriter ? &DataStream->ChunkedWriter->PackageMapExports : nullptr)
+{
+}
+
+UIrisObjectReferencePackageMap* FChunkedDataStreamExportWriteScope::GetPackageMap()
+{
+	return WriteScope.GetPackageMap();
+}
+
+FChunkedDataStreamExportReadScope::FChunkedDataStreamExportReadScope(UChunkedDataStream* DataStream)
+: ReadScope(DataStream ? DataStream->GetPackageMap() : nullptr, DataStream && DataStream->ChunkedReader ? &DataStream->ChunkedReader->PackageMapExports : nullptr, DataStream && DataStream->ChunkedReader ? &DataStream->ChunkedReader->NetTokenResolveContext : nullptr)
+{
+}
+
+UIrisObjectReferencePackageMap* FChunkedDataStreamExportReadScope::GetPackageMap()
+{
+	return ReadScope.GetPackageMap();
+}
+
+uint32 UChunkedDataStream::GetQueuedByteCount() const
+{
+	if (ChunkedWriter)
+	{
+		return ChunkedWriter->GetQueuedBytes();
+	}
+
+	return 0U;
+}
+
+void UChunkedDataStream::SetMaxUndispatchedPayloadBytes(uint32  MaxUndispatchedPayloadBytes)
+{
+	if (ChunkedReader)
+	{
+		ChunkedReader->MaxUndispatchedPayloadBytes = MaxUndispatchedPayloadBytes;
+	}
+}
+
+void UChunkedDataStream::SetMaxEnqueuedPayloadBytes(uint32 MaxEnqueuedPayloadBytes)
+{
+	if (ChunkedWriter)
+	{
+		ChunkedWriter->SendBufferMaxSize = MaxEnqueuedPayloadBytes;
+	}
+}
+
+bool UChunkedDataStream::HasError() const
+{
+	if (ChunkedReader)
+	{
+		return ChunkedReader->HasError();
+	}
+
+	return false;
+}
diff --git a/Engine/Source/Runtime/Experimental/Iris/Core/Private/Iris/ReplicationSystem/ChunkedDataStream/ChunkedDataStreamCommon.h b/Engine/Source/Runtime/Experimental/Iris/Core/Private/Iris/ReplicationSystem/ChunkedDataStream/ChunkedDataStreamCommon.h
new file mode 100644
--- /dev/null
+++ b/Engine/Source/Runtime/Experimental/Iris/Core/Private/Iris/ReplicationSystem/ChunkedDataStream/ChunkedDataStreamCommon.h
@@ -0,0 +1,33 @@
+// Copyright Epic Games, Inc. All Rights Reserved.
+
+#pragma once
+
+#include "Iris/Core/IrisLog.h"
+#include "Net/Core/Trace/NetTrace.h"
+
+#ifndef UE_NET_CHUNKEDDATASTREAM_LOG_COMPILE_VERBOSITY
+// Don't compile verbose logs in Shipping builds	
+#if UE_BUILD_SHIPPING
+#	define UE_NET_CHUNKEDDATASTREAM_LOG_COMPILE_VERBOSITY Log
+#else
+#	define UE_NET_CHUNKEDDATASTREAM_LOG_COMPILE_VERBOSITY All
+#endif
+#endif
+
+DECLARE_LOG_CATEGORY_EXTERN(LogIrisChunkedDataStream, Log, UE_NET_CHUNKEDDATASTREAM_LOG_COMPILE_VERBOSITY);
+#define UE_LOG_CHUNKEDDATASTREAM_CONN(Verbosity, Format, ...)  UE_LOG(LogIrisChunkedDataStream, Verbosity, TEXT("ChunkedDataStream: R:%u :C%u ") Format, InitParams.ReplicationSystemId, InitParams.ConnectionId, ##__VA_ARGS__)
+
+namespace UE::Net::Private
+{
+
+struct FChunkedDataStreamParameters
+{	
+	static constexpr uint32 SequenceBitCount = 11U;
+	static constexpr uint32 MaxUnackedDataChunkCount = 1 << SequenceBitCount;
+	static constexpr uint32 SequenceBitMask = (1U << SequenceBitCount) - 1U;
+	static constexpr uint32 ChunkSize = 192U;
+
+	static constexpr uint32 NumBitsForExportOffset = 32U;
+};
+
+} // End of namespace(s)
diff --git a/Engine/Source/Runtime/Experimental/Iris/Core/Private/Iris/ReplicationSystem/ChunkedDataStream/ChunkedDataWriter.cpp b/Engine/Source/Runtime/Experimental/Iris/Core/Private/Iris/ReplicationSystem/ChunkedDataStream/ChunkedDataWriter.cpp
new file mode 100644
--- /dev/null
+++ b/Engine/Source/Runtime/Experimental/Iris/Core/Private/Iris/ReplicationSystem/ChunkedDataStream/ChunkedDataWriter.cpp
@@ -0,0 +1,473 @@
+// Copyright Epic Games, Inc. All Rights Reserved.
+
+#include "ChunkedDataWriter.h"
+#include "Iris/Serialization/NetSerializationContext.h"
+
+#include "Iris/PacketControl/PacketNotification.h"
+#include "Iris/ReplicationSystem/NetExports.h"
+#include "Iris/ReplicationSystem/ReplicationSystem.h"
+#include "Iris/ReplicationSystem/ReplicationSystemInternal.h"
+#include "Iris/Serialization/NetBitStreamUtil.h"
+#include "Containers/ContainersFwd.h"
+#include "Net/Core/Trace/NetTrace.h"
+#include "Misc/ScopeExit.h"
+
+namespace UE::Net::Private
+{
+
+FChunkedDataWriter::FDataChunk::FDataChunk()
+: PayloadByteOffset(0U)
+, PartCount(0U)
+, SequenceNumber(uint16(-1))
+, PartByteCount(0U)
+, bIsFirstChunk(0U)
+, bIsExportChunk(0U)
+{
+}
+
+void FChunkedDataWriter::FDataChunk::Serialize(UE::Net::FNetSerializationContext& Context) const
+{
+	FNetBitStreamWriter* Writer = Context.GetBitStreamWriter();
+
+	// Size is encoded as a combination of PartCount and PayloadByteCount
+	// The first part contains the size of the entire payload encoded as PartCount * ChunkSize + PayloadByteCount
+	uint32 PartPayLoadBytesToWrite = PartByteCount;
+	if (Writer->WriteBool(bIsFirstChunk))
+	{
+		Writer->WriteBool(bIsExportChunk);
+		WritePackedUint32(Writer, PartCount);
+		PartPayLoadBytesToWrite = PartCount > 1U ? FChunkedDataStreamParameters::ChunkSize : PartByteCount;
+	}
+	const bool bIsFullChunk = PartByteCount == FChunkedDataStreamParameters::ChunkSize;
+	if (!Writer->WriteBool(bIsFullChunk))
+	{
+		WritePackedUint16(Writer, PartByteCount);
+	}
+
+	// Write actual payload
+	UE_NET_TRACE_SCOPE(Payload, *Writer, Context.GetTraceCollector(), ENetTraceVerbosity::Verbose);
+#if UE_NET_TRACE_ENABLED
+		if (bIsFirstChunk && bIsExportChunk)
+		{
+			if (FNetTraceCollector* Collector = SrcEntry->References->TraceCollector)
+			{
+				FNetTrace::FoldTraceCollector(Context.GetTraceCollector(), Collector, GetBitStreamPositionForNetTrace(*Writer));
+			}
+		}
+#endif // UE_NET_TRACE_ENABLED
+
+	const uint8* PayloadData = bIsExportChunk ? SrcEntry->References->ExportsPayload.GetData() : SrcEntry->Payload->GetData();
+	WriteBytes(Writer, PayloadData + PayloadByteOffset, PartPayLoadBytesToWrite);
+}
+
+FChunkedDataWriter::FChunkedDataWriter(const FInitParameters& InParams)
+: InitParams(InParams)
+, ReplicationSystem(UE::Net::GetReplicationSystem(InParams.ReplicationSystemId))
+, ObjectReferenceCache(&ReplicationSystem->GetReplicationSystemInternal()->GetObjectReferenceCache())
+{
+}
+
+bool FChunkedDataWriter::SplitPayload(FSendQueueEntry& SrcEntry, TConstArrayView<uint8> Payload, bool bIsExportPayload)
+{
+	const uint8* SrcPayload = Payload.GetData();
+	const uint32 SrcPayloadBytes = Payload.Num();
+
+	const uint32 ChunkCount = (SrcPayloadBytes + FChunkedDataStreamParameters::ChunkSize - 1U) / FChunkedDataStreamParameters::ChunkSize;
+
+	DataChunksPendingSend.Reserve(DataChunksPendingSend.Num() + int32(ChunkCount));
+
+	UE_LOG_CHUNKEDDATASTREAM_CONN(Verbose, TEXT("Split Payload %u Bytes into %u chunks"), SrcPayloadBytes, ChunkCount);
+	
+	uint32 CurrentOffset = 0U;
+	for (uint32 ChunkIt = 0, ChunkEndIt = ChunkCount; ChunkIt != ChunkEndIt; ++ChunkIt)
+	{
+		const bool bIsFirstChunk = ChunkIt == 0U;
+
+		FDataChunk Chunk;
+			
+		Chunk.SrcEntry = &SrcEntry;
+		Chunk.PartCount = ChunkCount;
+		Chunk.bIsFirstChunk = bIsFirstChunk ? 1U : 0U;
+		Chunk.bIsExportChunk = bIsExportPayload ? 1U : 0U;
+		Chunk.SequenceNumber = NextSequenceNumber++;
+
+		// Copy relevant data from our temporary buffer.
+		Chunk.PayloadByteOffset = CurrentOffset;
+		Chunk.PartByteCount = (uint16)FPlatformMath::Min(SrcPayloadBytes - CurrentOffset, FChunkedDataStreamParameters::ChunkSize);
+		CurrentOffset += Chunk.PartByteCount;
+
+		// We encode full payload size in first part as (SrcPayloadBytes - ((PartCount - 1) * FChunkedDataStreamParameters::ChunkSize))
+		if (bIsFirstChunk && Chunk.PartCount > 1U)
+		{
+			const uint16 LastPartByteCount = SrcPayloadBytes % FChunkedDataStreamParameters::ChunkSize;
+			Chunk.PartByteCount = LastPartByteCount != 0 ? LastPartByteCount : FChunkedDataStreamParameters::ChunkSize;
+		}
+			
+		DataChunksPendingSend.Add(MoveTemp(Chunk));
+	}
+
+	return ChunkCount > 0;
+}
+
+FChunkedDataWriter::FReferencesForExport* FChunkedDataWriter::CreateExportPayload()
+{
+	if (!PackageMapExports.IsEmpty() && NetTokensPendingExport.IsEmpty())
+	{
+		TUniquePtr<FReferencesForExport> Result = MakeUnique<FReferencesForExport>();
+		Result->ExportsPayload.AddUninitialized(ExportsBufferMaxSize);
+
+		FNetBitStreamWriter ExportsWriter;
+		const uint32 MaxExportsPayloadBytes = Result->ExportsPayload.Num();
+		ExportsWriter.InitBytes(Result->ExportsPayload.GetData(), MaxExportsPayloadBytes);
+
+		// Create context
+		FNetSerializationContext Context(&ExportsWriter);
+		FInternalNetSerializationContext InternalContext(ReplicationSystem);
+		Context.SetInternalContext(&InternalContext);
+		Context.SetLocalConnectionId(InitParams.ConnectionId);
+
+		// Temporary quantized state
+		FIrisPackageMapExportsQuantizedType QuantizedExports = {};
+
+		// We need to release dynamic data on exit
+		ON_SCOPE_EXIT
+		{
+			FIrisPackageMapExportsUtil::FreeDynamicState(QuantizedExports);
+		};
+
+		FIrisPackageMapExportsUtil::Quantize(Context, PackageMapExports, NetTokensPendingExport, QuantizedExports);
+
+		// Setup export scope
+		FNetExports::FExportScope ExportScope = InitParams.NetExports->MakeExportScope(Context, Result->BatchExports);
+
+#if UE_NET_TRACE_ENABLED
+		FNetTraceCollector* LocalTraceCollector = UE_NET_TRACE_CREATE_COLLECTOR(ENetTraceVerbosity::Trace);
+		Context.SetTraceCollector(LocalTraceCollector);
+#endif
+
+		const uint32 ExportHeaderPos = ExportsWriter.GetPosBits();
+
+		UE_NET_TRACE_SCOPE(ExportPayload, ExportsWriter, Context.GetTraceCollector(), ENetTraceVerbosity::Verbose);
+
+		// Header
+		ExportsWriter.WriteBits(0, FChunkedDataStreamParameters::NumBitsForExportOffset);
+
+		// Append potential exports
+		ObjectReferenceCache->AddPendingExports(Context, MakeArrayView(QuantizedExports.ObjectReferenceStorage.GetData(), QuantizedExports.ObjectReferenceStorage.Num()));
+		
+		// Serialize the reference data
+		FIrisPackageMapExportsUtil::Serialize(Context, QuantizedExports);
+
+		const uint32 WrittenBitsInBatch = (ExportsWriter.GetPosBits() - ExportHeaderPos) - FChunkedDataStreamParameters::NumBitsForExportOffset;
+
+		// Serialize exports if there are any
+		const FObjectReferenceCache::EWriteExportsResult WriteExportResult = ObjectReferenceCache->WritePendingExports(Context, 0);
+		if (WriteExportResult == FObjectReferenceCache::EWriteExportsResult::BitStreamOverflow)
+		{
+			return nullptr;
+		}
+		else if (WriteExportResult == FObjectReferenceCache::EWriteExportsResult::WroteExports)
+		{
+			// Go back and update header that we have exports to read.
+			FNetBitStreamWriteScope SizeScope(ExportsWriter, ExportHeaderPos);
+			ExportsWriter.WriteBits(WrittenBitsInBatch, FChunkedDataStreamParameters::NumBitsForExportOffset);
+		}
+		ExportsWriter.CommitWrites();
+
+		// Trim down the payload.
+		Result->ExportsPayload.SetNum(ExportsWriter.GetPosBytes(), EAllowShrinking::Yes);
+#if UE_NET_TRACE_ENABLED
+		Result->TraceCollector = LocalTraceCollector;
+		LocalTraceCollector = nullptr;
+#endif
+
+		return Result.Release();
+	}
+
+	return nullptr;
+}
+
+void FChunkedDataWriter::ResetExports()
+{
+	PackageMapExports.Reset();
+	NetTokensPendingExport.Reset();
+}
+
+bool FChunkedDataWriter::EnqueuePayload(const TSharedPtr<TArray64<uint8>>& Payload)
+{
+	uint32 TotalPayloadByteCount = (uint32)Payload->Num();
+
+	// Nothing to send
+	if (TotalPayloadByteCount == 0U)
+	{
+		return true;
+	}
+
+	bool bCanEnqueueData = (TotalPayloadByteCount + CurrentBytesInSendQueue) <= SendBufferMaxSize;
+
+	// Do we have exports
+	TUniquePtr<FReferencesForExport> Exports(bCanEnqueueData ? CreateExportPayload() : nullptr);
+	if (Exports)
+	{
+		TotalPayloadByteCount += Exports->ExportsPayload.Num();
+		bCanEnqueueData = (TotalPayloadByteCount + CurrentBytesInSendQueue) <= SendBufferMaxSize;
+	}
+
+	if (!bCanEnqueueData)
+	{
+		UE_LOG_CHUNKEDDATASTREAM_CONN(Warning, TEXT("EnqueuePayload SendBufferFull: Cannot enqueue payload with %u Bytes, CurrentBytesInSendQueue %u"), TotalPayloadByteCount, CurrentBytesInSendQueue);
+		return false;
+	}
+
+	TUniquePtr<FSendQueueEntry>& NewEntry = SendQueue.Emplace_GetRef(MakeUnique<FSendQueueEntry>(Payload));
+	NewEntry->References = MoveTemp(Exports);
+	CurrentBytesInSendQueue += TotalPayloadByteCount;
+
+	UE_LOG_CHUNKEDDATASTREAM_CONN(Verbose, TEXT("EnqueuePayload NewEntry %u Bytes, %u exports CurrentBytesInSendQueue %u"), Payload->Num(), PackageMapExports.References.Num(), CurrentBytesInSendQueue);
+
+	ResetExports();
+
+	return true;
+}
+
+bool FChunkedDataWriter::CanSend() const
+{
+	return (uint32)DataChunksPendingAck.Num() < FChunkedDataStreamParameters::MaxUnackedDataChunkCount;
+}
+
+UDataStream::EWriteResult FChunkedDataWriter::BeginWrite(const FBeginWriteParameters& Params)
+{
+	if (SendQueue.Num() && CanSend())
+	{
+		return EWriteResult::HasMoreData;
+	}
+
+	return EWriteResult::NoData;
+}
+
+bool FChunkedDataWriter::HasAcknowledgedAllReliableData() const
+{
+	return SendQueue.IsEmpty();
+}
+
+bool FChunkedDataWriter::UpdateSendQueue()
+{
+	for (TUniquePtr<FSendQueueEntry>& Entry : SendQueue)
+	{
+		if (Entry->RefCount == 0U)
+		{
+			// Should we enqueue exports as separate entry?
+			if (Entry->References)
+			{
+				const bool bIsExportPayload = true;
+				SplitPayload(*Entry, Entry->References->ExportsPayload, bIsExportPayload);
+			}
+			SplitPayload(*Entry, *Entry->Payload.Get());
+			return DataChunksPendingSend.Num() > 0;
+		}
+	}
+
+	if (!DataChunksPendingSend.IsEmpty())
+	{
+		return true;
+	}	
+		
+	return false;
+}
+
+UDataStream::EWriteResult FChunkedDataWriter::WriteData(UE::Net::FNetSerializationContext& Context, FDataStreamRecord const*& OutRecord)
+{
+	FNetBitStreamWriter* Writer = Context.GetBitStreamWriter();
+	
+	// Write chunks, we need at least 1 bit free...
+	if (SendQueue.Num() == 0 || Writer->GetBitsLeft() < 1U)
+	{
+		// If we have no pending data in-flight we can trim down our storage
+		if (SendQueue.Num() == 0)
+		{
+			DataChunksPendingSend.Trim();
+			DataChunksPendingAck.Trim();
+		}
+		
+		return EWriteResult::NoData;
+	}
+
+	// Write data until we have no more data to write or it does not fit
+	uint32 WrittenCount = 0;
+
+	// Setup a substream and context for writing data
+	FNetBitStreamWriter SubStream = Writer->CreateSubstream(Writer->GetBitsLeft() - 1U);
+	FNetSerializationContext SubContext = Context.MakeSubContext(&SubStream);
+
+	uint16 PrevWrittenSeq = uint16(-1);
+	int32 CurrentChunkIndex = 0;
+	bool bHasMoreDataToSend = false;
+	for (;;)
+	{
+		bHasMoreDataToSend = CanSend() && UpdateSendQueue();
+		if (!bHasMoreDataToSend || CurrentChunkIndex >= DataChunksPendingSend.Num())
+		{
+			break;
+		}
+
+		const FDataChunk& Chunk = DataChunksPendingSend.GetAtIndexNoCheck(CurrentChunkIndex);
+		const uint16 CurrentSeq = Chunk.SequenceNumber;
+				
+		if (IsSequenceSent(CurrentSeq))
+		{
+			++CurrentChunkIndex;
+			continue;
+		}
+
+		{
+			UE_NET_TRACE_SCOPE(DataChunk, SubStream, SubContext.GetTraceCollector(), ENetTraceVerbosity::Verbose);
+
+			FNetBitStreamRollbackScope SequenceRollback(SubStream);
+
+			// Continuation marker
+			SubStream.WriteBool(true);
+
+			// Write sequence number, only if it differs from previous written one.
+			const uint16 Seq = Chunk.SequenceNumber & FChunkedDataStreamParameters::SequenceBitMask;
+			const bool bIsInSequence = Seq == ((PrevWrittenSeq + 1U) & FChunkedDataStreamParameters::SequenceBitMask);
+			PrevWrittenSeq = Seq;
+
+			UE_NET_TRACE_NAMED_DYNAMIC_NAME_SCOPE(SequenceScope, static_cast<const TCHAR*>(nullptr), SubStream, SubContext.GetTraceCollector(), ENetTraceVerbosity::VeryVerbose);
+
+			if (!SubStream.WriteBool(bIsInSequence))
+			{
+				SubStream.WriteBits(Seq, FChunkedDataStreamParameters::SequenceBitCount);
+			}
+
+			// Write chunk
+			Chunk.Serialize(SubContext);
+
+			if (SubStream.IsOverflown())
+			{
+				break;
+			}
+
+#if UE_NET_TRACE_ENABLED
+			if (FNetTrace::GetNetTraceVerbosityEnabled(ENetTraceVerbosity::VeryVerbose))
+			{
+				TStringBuilder<64> Builder;
+				if (Chunk.bIsFirstChunk)
+				{
+					Builder.Appendf(TEXT("Seq %u First part of %u"), Chunk.SequenceNumber, Chunk.PartCount);
+				}
+				else
+				{
+					Builder.Appendf(TEXT("Seq %u"), Chunk.SequenceNumber);
+				}
+				UE_NET_TRACE_SET_SCOPE_NAME(SequenceScope, Builder.ToString());
+			}
+#endif // UE_NET_TRACE_ENABLED
+
+			UE_LOG_CHUNKEDDATASTREAM_CONN(Verbose, TEXT("Serialized Seq:%u (local:%u)"), Seq, CurrentSeq);
+
+		}
+
+		// Enqueue in our record as well for resending if we drop data
+		DataChunksPendingAck.Add((uint16)CurrentSeq);
+		SetSequenceIsSent(CurrentSeq);
+		++WrittenCount;
+		++CurrentChunkIndex;
+	}
+
+	// Commit substream
+	if (WrittenCount)
+	{
+		Writer->CommitSubstream(SubStream);
+		Writer->WriteBool(false);
+
+		// Store number of written batches in the external record pointer
+		UPTRINT& OutRecordCount = *reinterpret_cast<UPTRINT*>(&OutRecord);
+		OutRecordCount = WrittenCount;
+
+		return bHasMoreDataToSend ? EWriteResult::HasMoreData : EWriteResult::Ok;
+	}
+	else
+	{
+		Writer->DiscardSubstream(SubStream);
+
+		return bHasMoreDataToSend ? EWriteResult::HasMoreData : EWriteResult::NoData;
+	}
+}
+
+void FChunkedDataWriter::RemoveSendQueueEntry(FSendQueueEntry* SendQueueEntry)
+{
+	const int32 Index = SendQueue.IndexOfByPredicate([SendQueueEntry](const TUniquePtr<FSendQueueEntry>& Entry) { return Entry.Get() == SendQueueEntry; });
+	if (Index != INDEX_NONE)
+	{
+		SendQueue.RemoveAt(Index);
+	}
+}
+
+void FChunkedDataWriter::PopDeliveredChunks()
+{
+	while (!DataChunksPendingSend.IsEmpty())
+	{
+		const FDataChunk& CurrentChunk = DataChunksPendingSend.First();
+
+		FSendQueueEntry* SendQueueEntry = CurrentChunk.SrcEntry;
+
+		const uint32 Index = SequenceToIndex(CurrentChunk.SequenceNumber);
+		if (!IsIndexAcked(Index))
+		{
+			break;
+		}
+
+		ClearIndexIsAcked(Index);
+		ClearIndexIsSent(Index);
+
+		DataChunksPendingSend.PopFrontNoCheck();
+
+		// NOTE: it is intentional that we wait with removing the SendQueue entry until we know that all data before it has
+		// been delivered to the client to ensure that potential exports has been processed before we acknowledge them
+		// If refcount is zero, we have delivered our payload and can treat exports as exported
+		if (SendQueueEntry->RefCount == 0U)
+		{
+			UE_LOG_CHUNKEDDATASTREAM_CONN(Verbose, TEXT("Completed SendQueueEntry 0x%p"), SendQueueEntry);
+			// We need to explicitly acknowledge exports made through the huge object batch
+			if (SendQueueEntry->References)
+			{
+				InitParams.NetExports->AcknowledgeBatchExports(SendQueueEntry->References->BatchExports);
+				CurrentBytesInSendQueue -= SendQueueEntry->References->ExportsPayload.Num();
+			}
+			CurrentBytesInSendQueue -= (uint32)SendQueueEntry->Payload->Num();
+			RemoveSendQueueEntry(SendQueueEntry);
+		}
+	}
+}
+
+void FChunkedDataWriter::ProcessPacketDeliveryStatus(UE::Net::EPacketDeliveryStatus Status, FDataStreamRecord const* Record)
+{
+	// The Record pointer is used as storage for the number of batches to process
+	uint32 RecordCount = (uint32)reinterpret_cast<UPTRINT>(Record);
+
+	if (Status == UE::Net::EPacketDeliveryStatus::Lost)
+	{
+		while (RecordCount)
+		{
+			// Mark entries as not sent
+			const uint32 LostSeq = DataChunksPendingAck.PopFrontValue();
+			ClearSequenceIsSent(LostSeq);
+			--RecordCount;
+			UE_LOG_CHUNKEDDATASTREAM_CONN(Verbose, TEXT("Notified Dropped Seq %u"), LostSeq);
+		}
+	}
+	else
+	{
+		while (RecordCount)
+		{
+			const uint32 DeliveredSeq = DataChunksPendingAck.PopFrontValue();
+			// Mark entries as not sent
+			SetSequenceIsAcked(DeliveredSeq);
+			--RecordCount;
+			UE_LOG_CHUNKEDDATASTREAM_CONN(Verbose, TEXT("Notified Delivered Seq %u"), DeliveredSeq);
+		}
+		PopDeliveredChunks();
+	}
+}
+
+} // End of namespace(s)
diff --git a/Engine/Source/Runtime/Experimental/Iris/Core/Private/Iris/ReplicationSystem/ChunkedDataStream/ChunkedDataWriter.h b/Engine/Source/Runtime/Experimental/Iris/Core/Private/Iris/ReplicationSystem/ChunkedDataStream/ChunkedDataWriter.h
new file mode 100644
--- /dev/null
+++ b/Engine/Source/Runtime/Experimental/Iris/Core/Private/Iris/ReplicationSystem/ChunkedDataStream/ChunkedDataWriter.h
@@ -0,0 +1,200 @@
+// Copyright Epic Games, Inc. All Rights Reserved.
+
+#pragma once
+
+#include "ChunkedDataStreamCommon.h"
+#include "Containers/Array.h"
+#include "Containers/RingBuffer.h"
+#include "Iris/DataStream/DataStream.h"
+#include "Iris/Serialization/IrisObjectReferencePackageMap.h"
+#include "Iris/Serialization/IrisPackageMapExportUtil.h"
+#include "Iris/Serialization/NetExportContext.h"
+#include "Net/Core/NetToken/NetTokenExportContext.h"
+
+namespace UE::Net::Private
+{
+
+// Used by ChunkedDataStream to capture and send payloads and exports
+class FChunkedDataWriter
+{
+public:
+	using EWriteResult = UDataStream::EWriteResult;
+	using FBeginWriteParameters = UDataStream::FBeginWriteParameters;
+	using FInitParameters = UDataStream::FInitParameters;
+	FChunkedDataWriter(const FInitParameters& InParams);
+
+	// Tracks references associated with an enqueued payload
+	struct FReferencesForExport
+	{
+		TArray<uint8, TAlignedHeapAllocator<>> ExportsPayload;
+		FNetExportContext::FBatchExports BatchExports;
+
+		~FReferencesForExport()
+		{
+			if (TraceCollector)
+			{
+				UE_NET_TRACE_DESTROY_COLLECTOR(TraceCollector);
+			}
+		}
+		FNetTraceCollector* TraceCollector = nullptr;
+	};
+	
+	// Enqueued payload to send
+	struct FSendQueueEntry
+	{
+		FSendQueueEntry() {}
+		FSendQueueEntry(FSendQueueEntry&& Other)
+		: Payload(MoveTemp(Other.Payload))
+		, References(MoveTemp(Other.References))
+		, RefCount(Other.RefCount)
+		{
+			Other.RefCount = 0;
+		}
+
+		FSendQueueEntry(const TSharedPtr<TArray64<uint8>>& InPayload)
+		: Payload(InPayload)
+		{
+		}
+
+		void AddRef() const
+		{
+			++RefCount;
+		}
+
+		void Release() const
+		{
+			--RefCount;
+		}
+	
+		TSharedPtr<TArray64<uint8>> Payload;
+		TUniquePtr<FReferencesForExport> References;
+		mutable int32 RefCount = 0;
+	};
+
+	// Split chunk of data, referencing source
+	struct FDataChunk
+	{
+		FDataChunk();
+		void Serialize(UE::Net::FNetSerializationContext& Context) const;
+		
+		// Hold a reference to the datachunk as source data is shared
+		TRefCountPtr<FSendQueueEntry> SrcEntry;
+		uint32 PayloadByteOffset;
+		uint32 PartCount;
+		uint16 SequenceNumber;
+		uint16 PartByteCount : 14U;
+		uint16 bIsFirstChunk : 1U;
+		uint16 bIsExportChunk : 1U;
+	};
+
+public:
+	uint32 SequenceToIndex(uint32 Seq) const
+	{
+		return Seq % FChunkedDataStreamParameters::MaxUnackedDataChunkCount;
+	}
+
+	bool IsIndexAcked(uint32 Index) const
+	{
+		return (Acked[Index >> 5U] & (1U << (Index & 31U))) != 0U;
+	}
+
+	void SetIndexIsAcked(uint32 Index)
+	{
+		Acked[Index >> 5U] |= (1U << (Index & 31U));
+	}
+
+	void SetSequenceIsAcked(uint32 Seq)
+	{
+		return SetIndexIsAcked(SequenceToIndex(Seq));
+	}
+
+	void ClearIndexIsAcked(uint32 Index)
+	{
+		Acked[Index >> 5U] &= ~(1U << (Index & 31U));
+	}
+
+	bool IsIndexSent(uint32 Index) const
+	{
+		return (Sent[Index >> 5U] & (1U << (Index & 31U))) != 0U;
+	}
+
+	bool IsSequenceSent(uint32 Seq) const
+	{
+		return IsIndexSent(SequenceToIndex(Seq));
+	}
+
+	void SetIndexIsSent(uint32 Index)
+	{
+		Sent[Index >> 5U] |= (1U << (Index & 31U));
+	}
+
+	void SetSequenceIsSent(uint32 Seq)
+	{
+		return SetIndexIsSent(SequenceToIndex(Seq));
+	}
+
+	void ClearIndexIsSent(uint32 Index)
+	{
+		Sent[Index >> 5U] &= ~(1U << (Index & 31U));
+	}
+
+	void ClearSequenceIsSent(uint32 Seq)
+	{
+		return ClearIndexIsSent(SequenceToIndex(Seq));
+	}
+
+	bool SplitPayload(FSendQueueEntry& SrcEntry, TConstArrayView<uint8> Payload, bool bIsExportPayload = false);
+	FReferencesForExport* CreateExportPayload();
+	void ResetExports();
+	bool EnqueuePayload(const TSharedPtr<TArray64<uint8>>& Payload);
+	bool CanSend() const;
+	bool UpdateSendQueue();
+	EWriteResult BeginWrite(const FBeginWriteParameters& Params);
+	EWriteResult WriteData(UE::Net::FNetSerializationContext& Context, FDataStreamRecord const*& OutRecord);
+	bool HasAcknowledgedAllReliableData() const;
+	void PopDeliveredChunks();
+	void ProcessPacketDeliveryStatus(UE::Net::EPacketDeliveryStatus Status, FDataStreamRecord const* Record);
+	void RemoveSendQueueEntry(FSendQueueEntry* SendQueueEntry);
+
+	// This could be more precise by by updating CurrentBytesInSendQueue based on individual datachunks.
+	uint32 GetQueuedBytes() const
+	{
+		return CurrentBytesInSendQueue;
+	}
+
+public:
+	// Payload data
+	TArray<TUniquePtr<FSendQueueEntry>> SendQueue;
+
+	// Split data chunks
+	TRingBuffer<FDataChunk> DataChunksPendingSend;
+
+	// In-flight data chunks pending ack
+	TRingBuffer<uint16> DataChunksPendingAck;
+
+	// Track status of entries in the DataChunksPendingSend
+	uint32 Sent[(FChunkedDataStreamParameters::MaxUnackedDataChunkCount + 31)/32] = {};
+	uint32 Acked[(FChunkedDataStreamParameters::MaxUnackedDataChunkCount + 31)/32] = {};
+	uint16 NextSequenceNumber = 0U;
+	friend class FChunkedDataStreamExportWriteScope;
+
+	// Cached copy of DataStream init params
+	FInitParameters InitParams;
+	UReplicationSystem* ReplicationSystem = nullptr;
+	FObjectReferenceCache* ObjectReferenceCache = nullptr;
+
+	// Total number of bytes in send queue.
+	uint32 CurrentBytesInSendQueue = 0U;
+
+	// Just for sanity
+	uint32 ExportsBufferMaxSize = 524288U;
+
+	// We do not allow more paylaod bytes to be enqueued than this.
+	uint32 SendBufferMaxSize = 10485760U;
+
+	// Exports
+	UE::Net::FIrisPackageMapExports PackageMapExports;
+	UE::Net::FNetTokenExportContext::FNetTokenExports NetTokensPendingExport;
+};
+
+} // End of namespace(s)
diff --git a/Engine/Source/Runtime/Experimental/Iris/Core/Public/Iris/ReplicationSystem/ChunkedDataStream/ChunkedDataStream.h b/Engine/Source/Runtime/Experimental/Iris/Core/Public/Iris/ReplicationSystem/ChunkedDataStream/ChunkedDataStream.h
new file mode 100644
--- /dev/null
+++ b/Engine/Source/Runtime/Experimental/Iris/Core/Public/Iris/ReplicationSystem/ChunkedDataStream/ChunkedDataStream.h
@@ -0,0 +1,129 @@
+// Copyright Epic Games, Inc. All Rights Reserved.
+
+#pragma once
+
+#include "Iris/DataStream/DataStream.h"
+#include "Iris/Serialization/IrisObjectReferencePackageMap.h"
+#include "ChunkedDataStream.generated.h"
+
+namespace UE::Net
+{
+	struct FIrisPackageMapExports;
+	namespace Private
+	{
+		class FChunkedDataWriter;
+		class FChunkedDataReader;
+	}
+}
+
+/** Scope used to setup PackageMap owned by ChunkedDataStream to write and capture exports */
+class FChunkedDataStreamExportWriteScope
+{
+public:
+	IRISCORE_API FChunkedDataStreamExportWriteScope(UChunkedDataStream* DataStream);
+	IRISCORE_API UIrisObjectReferencePackageMap* GetPackageMap();
+
+private:
+	UE::Net::FIrisObjectReferencePackageMapWriteScope WriteScope;
+};
+
+/** Scope used to setup PackageMap owned by ChunkedDataStream for reading captured exports */
+class FChunkedDataStreamExportReadScope
+{
+public:
+	IRISCORE_API FChunkedDataStreamExportReadScope(UChunkedDataStream* DataStream);
+	IRISCORE_API UIrisObjectReferencePackageMap* GetPackageMap();
+
+private:
+	UE::Net::FIrisObjectReferencePackageMapReadScope ReadScope;
+};
+
+/** 
+ * ChunkedDataStream
+ * Experimental DataStream used to split and carry large payloads with potential exports
+ */
+UCLASS(MinimalAPI)
+class UChunkedDataStream : public UDataStream
+{
+	GENERATED_BODY()
+
+public:
+	enum class EDispatchResult : uint8
+	{
+		Ok,
+		WaitingForMustBeMappedReferences,
+		NothingToDispatch
+	};
+
+	/**
+	 * Enqueue Payload for sending,
+	 * Object References written to the payload by using the PackageMap associated with the DataStream will be appended to the payload.
+	 * @param Payload SharedPtr to Payload to send, The DataStream will hold a shared reference until the transfer is complete.
+	 * @return false if SendBuffer is full
+	 */
+	IRISCORE_API bool EnqueuePayload(const TSharedPtr<TArray64<uint8>>& Payload);
+
+	/**
+	 * Dispatch received Payload
+	 * Object References can be read from the payload through the PackageMap associated with the DataStream.
+	 * @param DispatchDataChunkFunction Reference to TFunction to process Payload.
+	 */
+	IRISCORE_API EDispatchResult DispatchReceivedPayload(TFunctionRef<void(TConstArrayView64<uint8>)> DispatchPayloadFunction);
+
+	/**
+	* Dispatch all received Payloads
+	* Object References can be read from the payload through the PackageMap associated with the DataStream.
+	* @param DispatchDataChunkFunction Reference to TFunction to process Payload.
+	*/
+	IRISCORE_API EDispatchResult DispatchReceivedPayloads(TFunctionRef<void(TConstArrayView64<uint8>)> DispatchPayloadFunction);
+
+	/** Get the number of received payloads that are ready for dispatch */
+	IRISCORE_API uint32 GetNumReceivedPayloadsPendingDispatch() const;
+
+	/** Get UIrisObjectReferencePackageMap associated with the DataStream */
+	IRISCORE_API UIrisObjectReferencePackageMap* GetPackageMap();
+	IRISCORE_API const UIrisObjectReferencePackageMap* GetPackageMap() const;
+
+	/** Get number of payload bytes that is yet to be acknowledged */
+	IRISCORE_API uint32 GetQueuedByteCount() const;
+
+	/** 
+	 * Set the maximum number of undispatched payload bytes we can have on the receiing side. 
+	 * @param MaxUndispatchedPayloadBytes Stream will be set in error state and close if we have received too much data without dispatching it
+	 */
+	IRISCORE_API void SetMaxUndispatchedPayloadBytes(uint32 MaxUndispatchedPayloadBytes);
+
+	/** 
+	* Set the maximum number of enqueued payload bytes we can have on the sending side. 
+	* @param MaxEnqueuedPayloadBytes Stream will be set in error state and close if we have received too much data without dispatching it
+	*/
+	IRISCORE_API void SetMaxEnqueuedPayloadBytes(uint32 MaxEnqueuedPayloadBytes);
+
+	/** 
+	* Returns true if the stream is in an error state and should be closed.
+	*/
+	IRISCORE_API bool HasError() const;
+
+protected:
+
+	// UDataStream interface
+	IRISCORE_API virtual void Init(const UDataStream::FInitParameters& Params) override;
+	IRISCORE_API virtual void Deinit() override;
+
+	IRISCORE_API virtual EWriteResult BeginWrite(const FBeginWriteParameters& Params) override;
+	IRISCORE_API virtual EWriteResult WriteData(UE::Net::FNetSerializationContext& Context, FDataStreamRecord const*& OutRecord) override;
+
+	IRISCORE_API virtual void ReadData(UE::Net::FNetSerializationContext& Context) override;
+	IRISCORE_API virtual void ProcessPacketDeliveryStatus(UE::Net::EPacketDeliveryStatus Status, FDataStreamRecord const* Record) override;
+	IRISCORE_API virtual bool HasAcknowledgedAllReliableData() const override;
+
+private:
+	friend FChunkedDataStreamExportWriteScope;
+	friend FChunkedDataStreamExportReadScope;
+
+	UE::Net::Private::FChunkedDataWriter* ChunkedWriter = nullptr;
+	UE::Net::Private::FChunkedDataReader* ChunkedReader = nullptr;
+
+	UPROPERTY(Transient)
+	TObjectPtr<UIrisObjectReferencePackageMap> PackageMap = nullptr;
+};
