Add an overload of FMultiEndpointHttpClient::Get that takes a single url as it's input.

#rnx

- This mimics the older FHttpClient::Get code path which provided an easy way for code to make a one off GET request.
- Since the method is blocking we only provide an overload that returns the FMultiEndpointHttpClientResponse, we cannot have a delegate version like the other ::Get methods.
- Replaced use of FHttpClient::Get with this new overload.

diff --git a/Engine/Source/Runtime/Experimental/IoStore/OnDemand/Private/OnDemandHttpClient.cpp b/Engine/Source/Runtime/Experimental/IoStore/OnDemand/Private/OnDemandHttpClient.cpp
--- a/Engine/Source/Runtime/Experimental/IoStore/OnDemand/Private/OnDemandHttpClient.cpp
+++ b/Engine/Source/Runtime/Experimental/IoStore/OnDemand/Private/OnDemandHttpClient.cpp
@@ -100,6 +100,7 @@ static void LogHttpResult(const TCHAR* Host, const TCHAR* Url, uint32 StatusCode
 	UE_LOG(LogIas, VeryVerbose, TEXT("http-%3u: %5" UINT64_FMT "ms %5" UINT64_FMT "KiB[%7" UINT64_FMT "] '%S' %s%s"), StatusCode, DurationMs, Size, Offset, Memo, Host, Url);
 };
 
+#if 0
 void FHttpClient::Get(FAnsiStringView Url, const FIoOffsetAndLength& Range, FGetCallback&& Callback)
 {
 	check(CurrentEndpoint != INDEX_NONE);
@@ -116,6 +117,7 @@ void FHttpClient::Get(FAnsiStringView Url, FGetCallback&& Callback)
 {
 	Get(Url, FIoOffsetAndLength(), MoveTemp(Callback));
 }
+#endif
 
 bool FHttpClient::Tick(int32 WaitTimeMs, uint32 MaxKiBPerSecond)
 {
@@ -414,6 +416,67 @@ TUniquePtr<FMultiEndpointHttpClient> FMultiEndpointHttpClient::Create(const FMul
 	return TUniquePtr<FMultiEndpointHttpClient>(new FMultiEndpointHttpClient(Config));
 }
 
+TIoStatusOr<FMultiEndpointHttpClientResponse> FMultiEndpointHttpClient::Get(FAnsiStringView Url, const FMultiEndpointHttpClientConfig& Config)
+{
+	using namespace UE::IoStore::HTTP;
+
+	FEventLoop::FRequestParams Params = FEventLoop::FRequestParams
+	{
+		.bAutoRedirect = Config.Redirects == EHttpRedirects::Follow
+	};
+
+	FEventLoop Loop;
+	FIoBuffer Body;
+	TStringBuilder<128> Reason;
+	uint32 StatusCode = 0;
+	const uint32 MaxAttempts = FMath::Min(Config.MaxRetryCount, 3u);
+
+	const uint64 StartTime = FPlatformTime::Cycles64();
+
+	for (uint32 Attempt = 0; Attempt <= MaxAttempts; ++Attempt)
+	{
+		Loop.Send(Loop.Request("GET", Url, &Params), [&Body, &Reason, &StatusCode](const FTicketStatus& Status)
+			{
+				if (Status.GetId() == FTicketStatus::EId::Response)
+				{
+					Status.GetResponse().SetDestination(&Body);
+					StatusCode = Status.GetResponse().GetStatusCode();
+					return;
+				}
+
+				if (Status.GetId() == FTicketStatus::EId::Error)
+				{
+					Reason << Status.GetError().Reason;
+				}
+			});
+
+		while (Loop.Tick(-1))
+		{
+			// Busy loop
+		}
+
+		if (StatusCode > 199 && StatusCode < 300)
+		{
+			FMultiEndpointHttpClientResponse Response
+			{
+				.Body = MoveTemp(Body),
+				.DurationMilliseconds = uint64(FPlatformTime::ToMilliseconds64(FPlatformTime::Cycles64() - StartTime)),
+				.StatusCode = StatusCode,
+				.RetryCount = Attempt,
+			};
+
+			return Response;
+		}
+	}
+
+	if (Reason.Len() == 0)
+	{
+		Reason << TEXT("StatusCode: ") << StatusCode;
+	}
+
+	return FIoStatus(EIoErrorCode::ReadError, Reason.ToView());
+}
+
 void FMultiEndpointHttpClient::Get(const FOnDemandHostGroup& Endpoint, FAnsiStringView RelativeUrl, FOnHttpResponse&& OnResponse)
 {
 	FIoOffsetAndLength DefaultChunkRange(0,0);
diff --git a/Engine/Source/Runtime/Experimental/IoStore/OnDemand/Private/OnDemandHttpClient.h b/Engine/Source/Runtime/Experimental/IoStore/OnDemand/Private/OnDemandHttpClient.h
--- a/Engine/Source/Runtime/Experimental/IoStore/OnDemand/Private/OnDemandHttpClient.h
+++ b/Engine/Source/Runtime/Experimental/IoStore/OnDemand/Private/OnDemandHttpClient.h
@@ -129,7 +129,10 @@ public:
 
 	using FOnHttpResponse = TFunction<void(FMultiEndpointHttpClientResponse&&)>;
 
-	static TUniquePtr<FMultiEndpointHttpClient> Create(const FMultiEndpointHttpClientConfig& Config);
+	[[nodiscard]] static TUniquePtr<FMultiEndpointHttpClient> Create(const FMultiEndpointHttpClientConfig& Config);
+
+	/** Blocking method */
+	[[nodiscard]] static TIoStatusOr<FMultiEndpointHttpClientResponse> Get(FAnsiStringView Url, const FMultiEndpointHttpClientConfig& Config);
 
 	void Get(const FOnDemandHostGroup& HostGroup, FAnsiStringView RelativeUrl, FOnHttpResponse&& OnResponse);
 	void Get(const FOnDemandHostGroup& HostGroup, FAnsiStringView RelativeUrl, const FIoOffsetAndLength& ChunkRange, FOnHttpResponse&& OnResponse);
diff --git a/Engine/Source/Runtime/Experimental/IoStore/OnDemand/Private/OnDemandIoStore.cpp b/Engine/Source/Runtime/Experimental/IoStore/OnDemand/Private/OnDemandIoStore.cpp
--- a/Engine/Source/Runtime/Experimental/IoStore/OnDemand/Private/OnDemandIoStore.cpp
+++ b/Engine/Source/Runtime/Experimental/IoStore/OnDemand/Private/OnDemandIoStore.cpp
@@ -1245,25 +1245,25 @@ FIoStatus FOnDemandIoStore::TickMountRequest(FMountRequest& MountRequest)
 					*Container->Name, *LexToString(ChunkId));
 
 				const uint32 RetryCount = 2;
-				TIoStatusOr<FIoBuffer> ContainerHeaderChunk = FHttpClient::Get(
+				TIoStatusOr<FMultiEndpointHttpClientResponse> Response = FMultiEndpointHttpClient::Get(
 					Private::GetChunkUrl(Args.HostGroup.PrimaryHost(), *Container, ChunkInfo.ChunkEntry(), ChunkUrlBuilder).ToView(),
-					RetryCount,
-					EHttpRedirects::Follow);
+					FMultiEndpointHttpClientConfig{.MaxRetryCount = RetryCount}
+					);
 
-				if (ContainerHeaderChunk.IsOk() == false)
+				if (Response.IsOk() == false)
 				{
-					return ContainerHeaderChunk.Status();
+					return Response.Status();
 				}
 
 				if (EnumHasAnyFlags(Container->Flags, EOnDemandContainerFlags::PendingEncryptionKey))
 				{
-					RequestedContainerHeaders.Add(Container->UniqueName(), ContainerHeaderChunk.ConsumeValueOrDie());
+					RequestedContainerHeaders.Add(Container->UniqueName(), Response.ConsumeValueOrDie().Body);
 					continue;
 				}
 
 				TIoStatusOr<FSharedContainerHeader> Header = Private::DeserializeContainerHeader(
 					ChunkInfo,
-					ContainerHeaderChunk.ValueOrDie().GetView());
+					Response.ConsumeValueOrDie().Body.GetView());
 
 				if (Header.IsOk() == false)
 				{
diff --git a/Engine/Source/Runtime/Experimental/IoStore/OnDemand/Private/OnDemandIoStore.h b/Engine/Source/Runtime/Experimental/IoStore/OnDemand/Private/OnDemandIoStore.h
--- a/Engine/Source/Runtime/Experimental/IoStore/OnDemand/Private/OnDemandIoStore.h
+++ b/Engine/Source/Runtime/Experimental/IoStore/OnDemand/Private/OnDemandIoStore.h
@@ -249,6 +249,7 @@ struct FOnDemandChunkInfo
 	inline TConstArrayView<FIoBlockHash>	BlockHashes() const;
 	FAnsiStringView							ChunksDirectory() const { return SharedContainer->ChunksDirectory; }
 	const FOnDemandChunkEntry&				ChunkEntry() const { return Entry; }
+	inline FOnDemandHostGroup				HostGroup() const;
 
 	bool									IsValid() const { return SharedContainer.IsValid(); }
 	operator								bool() const { return IsValid(); }
@@ -279,6 +280,11 @@ TConstArrayView<FIoBlockHash> FOnDemandChunkInfo::BlockHashes() const
 		: TConstArrayView<FIoBlockHash>(SharedContainer->BlockHashes.GetData() + Entry.BlockOffset, Entry.BlockCount);
 }
 
+FOnDemandHostGroup FOnDemandChunkInfo::HostGroup() const
+{
+	return SharedContainer->HostGroup;
+}
+
 FOnDemandChunkInfo FOnDemandChunkInfo::Find(FSharedOnDemandContainer Container, const FIoChunkId& ChunkId)
 {
 	check(Container.IsValid());
