Remove one additional use of FHttpClient::Get missing from CL 39941710

#rnx

diff --git a/Engine/Source/Runtime/Experimental/IoStore/OnDemand/Private/OnDemandToc.cpp b/Engine/Source/Runtime/Experimental/IoStore/OnDemand/Private/OnDemandToc.cpp
--- a/Engine/Source/Runtime/Experimental/IoStore/OnDemand/Private/OnDemandToc.cpp
+++ b/Engine/Source/Runtime/Experimental/IoStore/OnDemand/Private/OnDemandToc.cpp
@@ -660,7 +660,11 @@ TIoStatusOr<FOnDemandToc> FOnDemandToc::LoadFromFile(const FString& FilePath, bo
 TIoStatusOr<FOnDemandToc> FOnDemandToc::LoadFromUrl(FAnsiStringView Url, uint32 RetryCount, bool bFollowRedirects)
 {
 	const EHttpRedirects Redirects = bFollowRedirects ? EHttpRedirects::Follow : EHttpRedirects::Disabled;
-	TIoStatusOr<FIoBuffer> Response = FHttpClient::Get(Url, RetryCount, Redirects); 
+	TIoStatusOr<FMultiEndpointHttpClientResponse> Response =  FMultiEndpointHttpClient::Get(Url, FMultiEndpointHttpClientConfig
+	{ 
+		.MaxRetryCount = RetryCount ,
+		.Redirects  = Redirects
+	});
 
 	if (!Response.IsOk())
 	{
@@ -668,7 +672,7 @@ TIoStatusOr<FOnDemandToc> FOnDemandToc::LoadFromUrl(FAnsiStringView Url, uint32
 		return Status;
 	}
 
-	FMemoryReaderView Ar(Response.ValueOrDie().GetView());
+	FMemoryReaderView Ar(Response.ValueOrDie().Body.GetView());
 	FOnDemandToc Toc;
 	Ar << Toc;
 
