Remove UE::IoStore::FHttpClient now that it is no longer used.

#rnx

diff --git a/Engine/Source/Runtime/Experimental/IoStore/OnDemand/Private/OnDemandHttpClient.cpp b/Engine/Source/Runtime/Experimental/IoStore/OnDemand/Private/OnDemandHttpClient.cpp
--- a/Engine/Source/Runtime/Experimental/IoStore/OnDemand/Private/OnDemandHttpClient.cpp
+++ b/Engine/Source/Runtime/Experimental/IoStore/OnDemand/Private/OnDemandHttpClient.cpp
@@ -94,310 +94,6 @@ FIoStatus LoadDefaultHttpCertificates(bool& bWasLoaded)
 	return DefaultCerts.Status;
 }
 
-static void LogHttpResult(const TCHAR* Host, const TCHAR* Url, uint32 StatusCode, uint64 DurationMs, uint64 Size, uint64 Offset, const char* Memo = "ok")
-{
-	Size >>= 10;
-	UE_LOG(LogIas, VeryVerbose, TEXT("http-%3u: %5" UINT64_FMT "ms %5" UINT64_FMT "KiB[%7" UINT64_FMT "] '%S' %s%s"), StatusCode, DurationMs, Size, Offset, Memo, Host, Url);
-};
-
-#if 0
-void FHttpClient::Get(FAnsiStringView Url, const FIoOffsetAndLength& Range, FGetCallback&& Callback)
-{
-	check(CurrentEndpoint != INDEX_NONE);
-	IssueRequest(FRequestParams
-	{
-		.Url = FString(Url),
-		.Range = Range,
-		.Callback = MoveTemp(Callback),
-		.Endpoint = CurrentEndpoint
-	});
-}
-
-void FHttpClient::Get(FAnsiStringView Url, FGetCallback&& Callback)
-{
-	Get(Url, FIoOffsetAndLength(), MoveTemp(Callback));
-}
-#endif
-
-bool FHttpClient::Tick(int32 WaitTimeMs, uint32 MaxKiBPerSecond)
-{
-	EventLoop.Throttle(MaxKiBPerSecond);
-	const uint32 TicketCount = EventLoop.Tick(WaitTimeMs);
-
-	if (Retries.IsEmpty() == false)
-	{
-		const int32 RequestCount = FMath::Min(Retries.Num(), int32(HTTP::FEventLoop::MaxActiveTickets - TicketCount));
-		for (int32 Idx = 0; Idx < RequestCount; Idx++)
-		{
-			IssueRequest(MoveTemp(Retries[Idx]));
-		}
-		Retries.RemoveAtSwap(0, RequestCount);
-	}
-
-	const bool bIsIdle = EventLoop.IsIdle();
-	if (bIsIdle)
-	{
-		for (int32 Idx = 0, Count = Connections.Num(); Idx < Count; ++Idx)
-		{
-			if (Idx != CurrentEndpoint)
-			{
-				Connections[Idx].Reset();
-			}
-		}
-	}
-
-	return bIsIdle == false;
-}
-
-void FHttpClient::SetEndpoint(int32 Endpoint)
-{
-	if (CurrentEndpoint == Endpoint)
-	{
-		return;
-	}
-
-	EnsureConnection(Endpoint);
-	CurrentEndpoint = Endpoint;
-}
-
-FHttpClient::FHttpClient(FHttpClientConfig&& ClientConfig)
-	: Config(MoveTemp(ClientConfig))
-{
-	check(Config.Endpoints.IsEmpty() == false);
-	Connections.SetNum(Config.Endpoints.Num());
-	SetEndpoint(Config.PrimaryEndpoint);
-
-	EventLoop.SetFailTimeout(Config.TimeoutMs);
-}
-
-TUniquePtr<FHttpClient> FHttpClient::Create(FHttpClientConfig&& ClientConfig)
-{
-	if (ClientConfig.Endpoints.IsEmpty())
-	{
-		return TUniquePtr<FHttpClient>();
-	}
-
-	for (const FString& Ep : ClientConfig.Endpoints)
-	{
-		if (Ep.IsEmpty())
-		{
-			return TUniquePtr<FHttpClient>();
-		}
-
-		//TODO: Get rid of all string conversions
-		if (UE::IoStore::HTTP::FConnectionPool::IsValidHostUrl(StringCast<ANSICHAR>(*Ep, Ep.Len())) == false)
-		{
-			return TUniquePtr<FHttpClient>();
-		}
-	}
-
-	return TUniquePtr<FHttpClient>(new FHttpClient(MoveTemp(ClientConfig)));
-}
-
-TUniquePtr<FHttpClient> FHttpClient::Create(const FString& Endpoint)
-{
-	FHttpClientConfig Config;
-	Config.Endpoints.Add(Endpoint);
-	return Create(MoveTemp(Config));
-}
-
-void FHttpClient::IssueRequest(FRequestParams&& Params)
-{
-	using namespace UE::IoStore::HTTP;
-
-	//TODO: Remove string conversion
-	const auto Url = StringCast<ANSICHAR>(*Params.Url, Params.Url.Len());
-
-	auto Sink = [
-		this,
-		Params = MoveTemp(Params),
-		Buffer = FIoBuffer(),
-		StartTime = FPlatformTime::Cycles64(),
-		StatusCode = uint32(0)]
-		(const FTicketStatus& Status) mutable
-		{
-			const uint64 DurationMs = (uint64)FPlatformTime::ToMilliseconds64(FPlatformTime::Cycles64() - StartTime);
-			switch (Status.GetId())
-			{
-			case FTicketStatus::EId::Response:
-			{
-				FResponse& Response = Status.GetResponse();
-				StatusCode = Response.GetStatusCode();
-				Response.SetDestination(&Buffer);
-				TrackCdnCacheStats(Response);
-				break;
-			}
-			case FTicketStatus::EId::Content:
-			{
-				const FIoBuffer& Content = Status.GetContent();
-				LogHttpResult(*GetEndpointUrl(Params.Endpoint), *Params.Url, StatusCode, DurationMs, Content.GetSize(), Params.Range.GetOffset());
-				
-				const bool bSuccessful = StatusCode > 199 && StatusCode < 300;
-				if (bSuccessful && Content.GetSize() > 0)
-				{
-					check(Params.Range.GetLength() == 0 || Params.Range.GetLength() == Content.GetSize());
-					Params.Callback(Content, DurationMs);
-					if (Params.Attempt > 0 && Config.bChangeEndpointAfterSuccessfulRetry)
-					{
-						SetEndpoint(Params.Endpoint);
-					}
-				}
-				else
-				{
-					const bool bServerError = StatusCode > 499 && StatusCode < 600;
-					if (Params.Attempt < Config.MaxRetryCount)
-					{
-						const bool bNextEndpoint = bServerError == false && Params.Attempt > 0;
-						RetryRequest(MoveTemp(Params), bNextEndpoint);
-					}
-					else
-					{
-						const EIoErrorCode ErrorCode = bServerError ? EIoErrorCode::ReadError : EIoErrorCode::NotFound;
-						Params.Callback(FIoStatus(ErrorCode), DurationMs);
-					}
-				}
-				break;
-			}
-			case FTicketStatus::EId::Error:
-			{
-				if (Params.Attempt < Config.MaxRetryCount)
-				{
-					const bool bNextEndpoint = Params.Attempt > 0;
-					RetryRequest(MoveTemp(Params), bNextEndpoint);
-				}
-				else
-				{
-					LogHttpResult(*GetEndpointUrl(Params.Endpoint), *Params.Url, StatusCode, DurationMs, 0, Params.Range.GetOffset(), Status.GetError().Reason);
-					Params.Callback(FIoStatus(EIoErrorCode::ReadError, FString(Status.GetError().Reason)), DurationMs);
-				}
-				break;
-			}
-			default:
-			{
-				LogHttpResult(*GetEndpointUrl(Params.Endpoint), *Params.Url, StatusCode, DurationMs, 0, Params.Range.GetOffset(), "cancelled");
-				Params.Callback(FIoStatus(EIoErrorCode::Cancelled, FString(Status.GetError().Reason)), DurationMs);
-				break;
-			}
-			}
-		};
-
-	TUniquePtr<FConnectionPool>& Connection = Connections[Params.Endpoint];
-	check(Connection.IsValid());
-
-	FRequest Request = EventLoop.Get(Url, *Connection);
-
-	if (Params.Range.GetOffset() > 0 || Params.Range.GetLength() > 0)
-	{
-		Request.Header(ANSITEXTVIEW("Range"),
-			WriteToAnsiString<64>(ANSITEXTVIEW("bytes="), Params.Range.GetOffset(), ANSITEXTVIEW("-"), Params.Range.GetOffset() + Params.Range.GetLength() - 1));
-	}
-
-	EventLoop.Send(MoveTemp(Request), MoveTemp(Sink));
-}
-
-void FHttpClient::RetryRequest(FRequestParams&& Params, bool bNextEndpoint)
-{
-	check(Params.Attempt < Config.MaxRetryCount);
-	if (Params.Attempt == 0)
-	{
-		FOnDemandIoBackendStats::Get()->OnHttpRetry();
-	}
-	if (bNextEndpoint && Config.Endpoints.Num() > 1)
-	{
-		Params.Endpoint = (Params.Endpoint + 1) % Config.Endpoints.Num();
-		EnsureConnection(Params.Endpoint);
-	}
-	Params.Attempt++;
-	Retries.Emplace(MoveTemp(Params));
-}
-
-TUniquePtr<HTTP::FConnectionPool> FHttpClient::CreateConnection(const FStringView& HostAddr)
-{
-	auto ServiceUrlAnsi = StringCast<ANSICHAR>(HostAddr.GetData(), HostAddr.Len());
-	HTTP::FConnectionPool::FParams Params;
-	ensure(Params.SetHostFromUrl(ServiceUrlAnsi) >= 0);
-	if (Config.ReceiveBufferSize >= 0)
-	{
-		UE_LOG(LogIas, Log, TEXT("HTTP client receive buffer size set to %d"), Config.ReceiveBufferSize);
-		Params.RecvBufSize = Config.ReceiveBufferSize;
-	}
-	Params.ConnectionCount = uint16(Config.MaxConnectionCount);
-
-	return MakeUnique<HTTP::FConnectionPool>(Params);
-}
-
-void FHttpClient::EnsureConnection(int32 Connection)
-{
-	check(Connection >= INDEX_NONE && Connection < Connections.Num());
-	if (Connection != INDEX_NONE && Connections[Connection].IsValid() == false)
-	{
-		Connections[Connection] = CreateConnection(Config.Endpoints[Connection]);
-	}
-}
-
-const FString& FHttpClient::GetEndpointUrl(int32 Endpoint)
-{
-	check(Endpoint >= INDEX_NONE && Endpoint < Config.Endpoints.Num());
-	static const FString None(TEXT("None"));
-	return Endpoint == INDEX_NONE ? None : Config.Endpoints[Endpoint];
-}
-
-TIoStatusOr<FIoBuffer> FHttpClient::Get(FAnsiStringView Url, uint32 RetryCount, EHttpRedirects Redirects)
-{
-	using namespace HTTP;
-
-	FEventLoop::FRequestParams Params = FEventLoop::FRequestParams
-	{
-		.bAutoRedirect = Redirects == EHttpRedirects::Follow
-	};
-
-	FEventLoop Loop;
-	FIoBuffer Body;
-	TStringBuilder<128> Reason;
-	uint32 StatusCode = 0;
-	const uint32 MaxAttempts = FMath::Min(RetryCount, 3u);
-
-	for (uint32 Attempt = 0; Attempt <= MaxAttempts; ++Attempt)
-	{
-		Loop.Send(Loop.Request("GET", Url, &Params), [&Body, &Reason, &StatusCode](const FTicketStatus& Status)
-		{
-			if (Status.GetId() == FTicketStatus::EId::Response)
-			{
-				Status.GetResponse().SetDestination(&Body);
-				StatusCode = Status.GetResponse().GetStatusCode();
-				return;
-			}
-
-			if (Status.GetId() == FTicketStatus::EId::Error)
-			{
-				Reason << Status.GetError().Reason;
-			}
-		});
-
-		while (Loop.Tick(-1))
-			;
-
-		if (StatusCode > 199 && StatusCode < 300)
-		{
-			return Body;
-		}
-	}
-
-	if (Reason.Len() == 0)
-	{
-		Reason << TEXT("StatusCode: ") << StatusCode;
-	}
-
-	return FIoStatus(EIoErrorCode::ReadError, Reason.ToView());
-}
-
-TIoStatusOr<FIoBuffer> FHttpClient::Get(FStringView Url, uint32 RetryCount, EHttpRedirects Redirects)
-{
-	auto AnsiUrl = StringCast<ANSICHAR>(Url.GetData(), Url.Len());
-	return Get(AnsiUrl, RetryCount, Redirects);
-}
-
-
 FMultiEndpointHttpClient::FMultiEndpointHttpClient(const FMultiEndpointHttpClientConfig& InConfig)
 	: Config(InConfig)
 {
diff --git a/Engine/Source/Runtime/Experimental/IoStore/OnDemand/Private/OnDemandHttpClient.h b/Engine/Source/Runtime/Experimental/IoStore/OnDemand/Private/OnDemandHttpClient.h
--- a/Engine/Source/Runtime/Experimental/IoStore/OnDemand/Private/OnDemandHttpClient.h
+++ b/Engine/Source/Runtime/Experimental/IoStore/OnDemand/Private/OnDemandHttpClient.h
@@ -26,63 +26,6 @@ enum class EHttpRedirects
 	Follow
 };
 
-struct FHttpClientConfig
-{
-	TArray<FString> Endpoints;
-	int32 PrimaryEndpoint = 0;
-	int32 MaxConnectionCount = 8;
-	int32 MaxRetryCount = 1;
-	int32 ReceiveBufferSize = -1;
-	int32 TimeoutMs = 0;
-	bool bChangeEndpointAfterSuccessfulRetry = true;
-};
-
-class FHttpClient
-{
-public:
-	using FGetCallback = TFunction<void(TIoStatusOr<FIoBuffer>, uint64 DurationMs)>;
-
-	static TUniquePtr<FHttpClient> Create(FHttpClientConfig&& ClientConfig);
-	static TUniquePtr<FHttpClient> Create(const FString& Endpoint);
-
-	void Get(FAnsiStringView Url, const FIoOffsetAndLength& Range, FGetCallback&& Callback);
-	void Get(FAnsiStringView Url, FGetCallback&& Callback);
-
-	bool Tick(int32 WaitTimeMs, uint32 MaxKiBPerSecond);
-	bool Tick() { return Tick(-1, 0); }
-
-	int32 GetEndpoint() const { return CurrentEndpoint; }
-	int32 GetPrimaryEndpoint() const { return Config.PrimaryEndpoint; }
-	void SetEndpoint(int32 Endpoint);
-	bool IsUsingPrimaryEndpoint() const { return CurrentEndpoint == Config.PrimaryEndpoint; }
-
-	static TIoStatusOr<FIoBuffer> Get(FAnsiStringView Url, uint32 RetryCount, EHttpRedirects Redirects);
-	static TIoStatusOr<FIoBuffer> Get(FStringView Url, uint32 RetryCount, EHttpRedirects Redirects);
-
-private:
-	struct FRequestParams
-	{
-		FString Url;
-		FIoOffsetAndLength Range;
-		FGetCallback Callback;
-		int32 Attempt = 0;
-		int32 Endpoint = INDEX_NONE;
-	};
-
-	FHttpClient(FHttpClientConfig&& ClientConfig);
-	void RetryRequest(FRequestParams&& Params, bool bNextEndpoint);
-	void IssueRequest(FRequestParams&& Params);
-	TUniquePtr<HTTP::FConnectionPool> CreateConnection(const FStringView& HostAddr);
-	void EnsureConnection(int32 Connection);
-	const FString& GetEndpointUrl(int32 Endpoint);
-
-	FHttpClientConfig Config;
-	HTTP::FEventLoop EventLoop;
-	TArray<TUniquePtr<HTTP::FConnectionPool>> Connections;
-	TArray<FRequestParams> Retries;
-	int32 CurrentEndpoint = INDEX_NONE;
-};
-
 struct FMultiEndpointHttpClientConfig
 {
 	int32				MaxConnectionCount = 4;
