Make the hierarchy facade not hang trying to find a collection root when the parent array contains a reference loop

diff --git a/Engine/Source/Runtime/Experimental/Chaos/Private/GeometryCollection/Facades/CollectionHierarchyFacade.cpp b/Engine/Source/Runtime/Experimental/Chaos/Private/GeometryCollection/Facades/CollectionHierarchyFacade.cpp
--- a/Engine/Source/Runtime/Experimental/Chaos/Private/GeometryCollection/Facades/CollectionHierarchyFacade.cpp
+++ b/Engine/Source/Runtime/Experimental/Chaos/Private/GeometryCollection/Facades/CollectionHierarchyFacade.cpp
@@ -37,15 +37,14 @@ namespace Chaos::Facades
 		if (IsValid())
 		{
 			const TManagedArray<int32>& Parents = ParentAttribute.Get();
-			if (Parents.Num() > 0)
+			for (int32 TransformIndex = 0; TransformIndex < Parents.Num(); ++TransformIndex)
 			{
-				int32 RootTransformIndex = 0;
-				while (Parents[RootTransformIndex] != INDEX_NONE)
+				if (Parents[TransformIndex] == INDEX_NONE)
 				{
-					RootTransformIndex = Parents[RootTransformIndex];
+					return TransformIndex;
 				}
-				return RootTransformIndex;
 			}
+			ensureMsgf(Parents.Num() == 0, TEXT("Non-empty collection has invalid hierarchy (no root found)"));
 		}
 		return INDEX_NONE;
 	}
