Dataflow : Allow anytype pin to be dragged onto input and output nodes in a subgraph
- Use the storage type of the anytype as the concrete type for the input / output

diff --git a/Engine/Source/Runtime/Experimental/Dataflow/Engine/Private/Dataflow/DataflowSubGraphNodes.cpp b/Engine/Source/Runtime/Experimental/Dataflow/Engine/Private/Dataflow/DataflowSubGraphNodes.cpp
--- a/Engine/Source/Runtime/Experimental/Dataflow/Engine/Private/Dataflow/DataflowSubGraphNodes.cpp
+++ b/Engine/Source/Runtime/Experimental/Dataflow/Engine/Private/Dataflow/DataflowSubGraphNodes.cpp
@@ -126,6 +126,36 @@ namespace UE::Dataflow::SubGraphNodes::Private
 		const UE::Dataflow::ISubGraphContextCallback& SubGraphCallback;
 	};
 
+	FPropertyBagPropertyDesc GetPropertyBagPropertyDescFromDataflowConnection(const FDataflowConnection& Connection)
+	{
+		const FProperty* ConnectionProperty = Connection.GetProperty();
+		check(ConnectionProperty);
+
+		// AnyType require special treatment to make sure we give a concrete type
+		if (Connection.IsAnyType())
+		{
+			if (ConnectionProperty->GetClass()->IsChildOf(FStructProperty::StaticClass()))
+			{
+				if (const FStructProperty* StructProperty = CastField<FStructProperty>(ConnectionProperty))
+				{
+					if (const bool bInheritsFromAnyType = StructProperty->Struct->IsChildOf<FDataflowAnyType>())
+					{
+						for (TFieldIterator<FProperty> FieldIt(StructProperty->Struct); FieldIt; ++FieldIt)
+						{
+							FProperty* Prop = *FieldIt;
+							if (Prop->GetFName() == "Value")
+							{
+								return FPropertyBagPropertyDesc(Connection.GetName(), Prop);
+							}
+						}
+					}
+				}
+			}
+		}
+		// default simply set property desc from the property
+		return FPropertyBagPropertyDesc(Connection.GetName(), Connection.GetProperty());
+	}
+
 }
 
 //////////////////////////////////////////////////////////////////////////////////////////////////////////
@@ -433,9 +463,10 @@ bool FDataflowSubGraphInputNode::SupportsDropConnectionOnNode(FName TypeName, UE
 
 const FDataflowConnection* FDataflowSubGraphInputNode::OnDropConnectionOnNode(const FDataflowConnection& DroppedConnection)
 {
+	using namespace  UE::Dataflow::SubGraphNodes;
 	if (SupportsDropConnectionOnNode(DroppedConnection.GetType(), DroppedConnection.GetDirection()))
 	{
-		const FPropertyBagPropertyDesc Desc(DroppedConnection.GetName(), DroppedConnection.GetProperty());
+		const FPropertyBagPropertyDesc Desc = Private::GetPropertyBagPropertyDescFromDataflowConnection(DroppedConnection);
 		PropertyBag.AddProperties({ Desc });
 		DynamicConnections.Refresh();
 		return FindOutput(Desc.Name);
@@ -479,9 +510,10 @@ bool FDataflowSubGraphOutputNode::SupportsDropConnectionOnNode(FName TypeName, U
 
 const FDataflowConnection* FDataflowSubGraphOutputNode::OnDropConnectionOnNode(const FDataflowConnection& DroppedConnection)
 {
+	using namespace  UE::Dataflow::SubGraphNodes;
 	if (SupportsDropConnectionOnNode(DroppedConnection.GetType(), DroppedConnection.GetDirection()))
 	{
-		const FPropertyBagPropertyDesc Desc(DroppedConnection.GetName(), DroppedConnection.GetProperty());
+		const FPropertyBagPropertyDesc Desc = Private::GetPropertyBagPropertyDescFromDataflowConnection(DroppedConnection);
 		PropertyBag.AddProperties({ Desc });
 		DynamicConnections.Refresh();
 		return FindInput(Desc.Name);
