Added aabb tree algorithm and tests.
This also adds a shared node as well as a visualizer that can be enabled with an ifdef.

diff --git a/Engine/Source/Runtime/Experimental/Chaos/Public/Chaos/AABB.h b/Engine/Source/Runtime/Experimental/Chaos/Public/Chaos/AABB.h
--- a/Engine/Source/Runtime/Experimental/Chaos/Public/Chaos/AABB.h
+++ b/Engine/Source/Runtime/Experimental/Chaos/Public/Chaos/AABB.h
@@ -147,6 +147,11 @@ namespace Chaos
 			return true;
 		}
 
+		FORCEINLINE bool Contains(const TAABB<T, d>& Other) const
+		{
+			return Contains(Other.Min()) && Contains(Other.Max());
+		}
+
 		FORCEINLINE const TAABB<T, d>& BoundingBox() const { return *this; }
 
 		FORCEINLINE uint16 GetMaterialIndex(uint32 HintIndex) const { return 0; }
