convex properties interface can be constructed from fmanagedarraycollection instead of fgeometrycollection (since it's just calling the base constructor, which takes fmanagedarraycollection)

diff --git a/Engine/Source/Runtime/Experimental/Chaos/Private/GeometryCollection/GeometryCollectionConvexPropertiesInterface.cpp b/Engine/Source/Runtime/Experimental/Chaos/Private/GeometryCollection/GeometryCollectionConvexPropertiesInterface.cpp
--- a/Engine/Source/Runtime/Experimental/Chaos/Private/GeometryCollection/GeometryCollectionConvexPropertiesInterface.cpp
+++ b/Engine/Source/Runtime/Experimental/Chaos/Private/GeometryCollection/GeometryCollectionConvexPropertiesInterface.cpp
@@ -12,7 +12,7 @@ const FName FGeometryCollectionConvexPropertiesInterface::ConvexCanExceedFractio
 const FName FGeometryCollectionConvexPropertiesInterface::ConvexRemoveOverlapsMethodAttribute = "RemoveOverlapsMethod";
 const FName FGeometryCollectionConvexPropertiesInterface::ConvexRemoveOverlapsShrinkAttribute = "RemoveOverlapsShrinkPercent";
 
-FGeometryCollectionConvexPropertiesInterface::FGeometryCollectionConvexPropertiesInterface(FGeometryCollection* InGeometryCollection)
+FGeometryCollectionConvexPropertiesInterface::FGeometryCollectionConvexPropertiesInterface(FManagedArrayCollection* InGeometryCollection)
 	: FManagedArrayInterface(InGeometryCollection)
 {}
 
diff --git a/Engine/Source/Runtime/Experimental/Chaos/Public/GeometryCollection/GeometryCollectionConvexPropertiesInterface.h b/Engine/Source/Runtime/Experimental/Chaos/Public/GeometryCollection/GeometryCollectionConvexPropertiesInterface.h
--- a/Engine/Source/Runtime/Experimental/Chaos/Public/GeometryCollection/GeometryCollectionConvexPropertiesInterface.h
+++ b/Engine/Source/Runtime/Experimental/Chaos/Public/GeometryCollection/GeometryCollectionConvexPropertiesInterface.h
@@ -41,7 +41,7 @@ public :
 		float OverlapRemovalShrinkPercent = 0.0f;
 	};
 
-	CHAOS_API FGeometryCollectionConvexPropertiesInterface(FGeometryCollection* InGeometryCollection);
+	CHAOS_API FGeometryCollectionConvexPropertiesInterface(FManagedArrayCollection* InGeometryCollection);
 
 	CHAOS_API void InitializeInterface() override;
 
