Collect pending completions prior to deletion.

#rnx

diff --git a/Engine/Source/Runtime/Experimental/IoStore/HttpClient/Private/Overlapped.inl b/Engine/Source/Runtime/Experimental/IoStore/HttpClient/Private/Overlapped.inl
--- a/Engine/Source/Runtime/Experimental/IoStore/HttpClient/Private/Overlapped.inl
+++ b/Engine/Source/Runtime/Experimental/IoStore/HttpClient/Private/Overlapped.inl
@@ -5,7 +5,7 @@
 namespace UE::IoStore::HTTP
 {
 
-#if 1
+#if 0
 
 using SocketType = SOCKET;
 
@@ -15,6 +15,8 @@ using SocketType = SOCKET;
 using SocketType	= struct FWsaSocket*;
 using PollType		= HANDLE;
 
+static int32 Poll_CollectCompletions(PollType, int32);
+
 ////////////////////////////////////////////////////////////////////////////////
 struct FWsaSocket
 {
@@ -28,6 +30,7 @@ public:
 					FWsaSocket(SOCKET InSocket);
 					~FWsaSocket();
 	bool			IsAsync() const;
+	void			Cancel(bool bAndComplete=false);
 	int32			Connect(const sockaddr* Addr, socklen_t AddrLen);
 	int32			Send(const char* Data, uint32 DataSize, uint32 SendFlags);
 	int32			Recv(char* Data, uint32 DataSize, uint32 RecvFlags);
@@ -38,6 +41,7 @@ public:
 	WSAOVERLAPPED	SendOverlap = {};
 	WSAOVERLAPPED	RecvOverlap = {};
 	SOCKET			Socket;
+	PollType		Iocp;
 	uint16			Flags = 0;
 	int16			PollEvents = 0;
 };
@@ -51,6 +55,7 @@ FWsaSocket::FWsaSocket(SOCKET InSocket)
 ////////////////////////////////////////////////////////////////////////////////
 FWsaSocket::~FWsaSocket()
 {
+	Cancel(true);
 	closesocket(Socket);
 }
 
@@ -60,9 +65,33 @@ bool FWsaSocket::IsAsync() const
 	return (Flags & Flag_Overlapped) != 0;
 }
 
+////////////////////////////////////////////////////////////////////////////////
+void FWsaSocket::Cancel(bool bAndComplete)
+{
+	if ((Flags & Flag_Busy) == 0)		return;
+	if ((Flags & Flag_Overlapped) == 0)	return;
+
+	CancelIo(HANDLE(Socket));
+
+	if (!bAndComplete)
+	{
+		return;
+	}
+
+	for (int32 i = 1; i; ::SwitchToThread(), i = (Flags & Flag_Busy))
+	{
+		if (Poll_CollectCompletions(Iocp, 0) < 0)
+		{
+			return;
+		}
+	}
+}
+
 ////////////////////////////////////////////////////////////////////////////////
 void FWsaSocket::OnCompletion(uint32 Transfered, WSAOVERLAPPED& Overlap, uint32 Result)
 {
+	check((Flags & ~(Flags & (Flag_Busy|Flag_Overlapped))) == 0);
+
 	if (Result != ERROR_SUCCESS)
 	{
 		if (Result >= 0xffff)
@@ -389,12 +418,35 @@ static bool Poll_Register(PollType Iocp, SocketType Socket)
 		return !GIasEnableOverlappedIo;
 	}
 
+	Socket->Iocp = Iocp;
+
 	auto Key = ULONG_PTR(Socket);
 	auto Fd = HANDLE(Socket->Socket);
 	HANDLE Result = CreateIoCompletionPort(Fd, Iocp, Key, 0);
 	return (Result != nullptr);
 }
 
+////////////////////////////////////////////////////////////////////////////////
+static int32 Poll_CollectCompletions(PollType Iocp, int32 TimeOutMs)
+{
+	DWORD EntryCount = 0;
+	OVERLAPPED_ENTRY Entries[16];
+	if (!GetQueuedCompletionStatusEx(Iocp, Entries, UE_ARRAY_COUNT(Entries), &EntryCount, TimeOutMs, false))
+	{
+		return -1;
+	}
+
+	for (const OVERLAPPED_ENTRY& Entry : MakeArrayView(Entries, EntryCount))
+	{
+		uint32 Result = uint32(Entry.Internal);
+		auto& Socket = *(FWsaSocket*)(Entry.lpCompletionKey);
+		auto& Overlap = *(WSAOVERLAPPED*)(Entry.lpOverlapped);
+		Socket.OnCompletion(Entry.dwNumberOfBytesTransferred, Overlap, Result);
+	}
+
+	return EntryCount;
+}
+
 ////////////////////////////////////////////////////////////////////////////////
 static int32 Poll_Wait(PollType Iocp, WSAPOLLFD* const Polls, uint32 Num, int32 TimeOutMs)
 {
@@ -410,21 +462,11 @@ static int32 Poll_Wait(PollType Iocp, WSAPOLLFD* const Polls, uint32 Num, int32
 		return poll(Polls, Num, TimeOutMs);
 	}
 
-	DWORD EntryCount = 0;
-	OVERLAPPED_ENTRY Entries[16];
-	if (!GetQueuedCompletionStatusEx(Iocp, Entries, UE_ARRAY_COUNT(Entries), &EntryCount, TimeOutMs, false))
+	if (Poll_CollectCompletions(Iocp, TimeOutMs) < 0)
 	{
 		return (GetLastError() == WAIT_TIMEOUT) - 1;
 	}
 
-	for (const OVERLAPPED_ENTRY& Entry : MakeArrayView(Entries, EntryCount))
-	{
-		uint32 Result = uint32(Entry.Internal);
-		auto& Socket = *(FWsaSocket*)(Entry.lpCompletionKey);
-		auto& Overlap = *(WSAOVERLAPPED*)(Entry.lpOverlapped);
-		Socket.OnCompletion(Entry.dwNumberOfBytesTransferred, Overlap, Result);
-	}
-
 	int32 Ret = 0;
 	for (WSAPOLLFD& Poll : MakeArrayView(Polls, Num))
 	{
@@ -441,7 +483,6 @@ static int32 Poll_Wait(PollType Iocp, WSAPOLLFD* const Polls, uint32 Num, int32
 	}
 	*/
 
-	check(!EntryCount == !Ret);
 	return Ret;
 }
 
