Fix bug that was causing FMultiEndpointHttpClient to always print verbose logging.

#rnx

- When I added the option to customize the log category for FMultiEndpointHttpClient::Log I accidently broke the verbosity filtering. For some reason I was thinking that it would be handled inside of FMsg::Logf but that has never been true.
-- This caused the logging to always be written (if a valid LogCategoryName was provided) with VeryVerbose verbosity, even if that category was supposed to be filtered out for that verbosity level.
- Replaced the LogCategoryName and LogVerbosity options from FMultiEndpointHttpClientConfig, we now pass in a pointer to the log category that we want to use.
- If the log category is a nullptr (default) then we log nothing, otherwise we use the categories verbosity filter to decide to log or not.
- In order to respect NO_LOGGING we need to define our own log category type so that the calling code can always be FMultiEndpointHttpClientConfig{.LogCategory = &LogMyCategory} and doesn't need to take the NO_LOGGING macro into account.

diff --git a/Engine/Source/Runtime/Experimental/IoStore/OnDemand/Private/OnDemandHttpClient.cpp b/Engine/Source/Runtime/Experimental/IoStore/OnDemand/Private/OnDemandHttpClient.cpp
--- a/Engine/Source/Runtime/Experimental/IoStore/OnDemand/Private/OnDemandHttpClient.cpp
+++ b/Engine/Source/Runtime/Experimental/IoStore/OnDemand/Private/OnDemandHttpClient.cpp
@@ -414,8 +414,7 @@ FMultiEndpointHttpClient::FConnection* FMultiEndpointHttpClient::FindConnection(
 void FMultiEndpointHttpClient::Log(const FRequest& Request) const
 {
 #if !NO_LOGGING
-	auto ShouldLogCategory = [](FName Name)->bool {return !Name.IsNone();}; // For LogCategoryName to resolve to a FName if UE_LOG_LAZY_CATEGORY_NAMES if defined and enabled
-	if (!ShouldLogCategory(Config.LogCategoryName))
+	if (Config.LogCategory == nullptr || Config.LogCategory->IsSuppressed(ELogVerbosity::VeryVerbose))
 	{
 		return;
 	}
@@ -423,7 +422,7 @@ void FMultiEndpointHttpClient::Log(const FRequest& Request) const
 	const uint64 DurationMs = (uint64)FPlatformTime::ToMilliseconds64(FPlatformTime::Cycles64() - Request.StartTime);
 	const uint64 Size = Request.Body.GetSize() >> 10;
 
-	FMsg::Logf(__FILE__, __LINE__, Config.LogCategoryName, Config.LogVerbosity,
+	FMsg::Logf(__FILE__, __LINE__, Config.LogCategory->GetCategoryName(), ELogVerbosity::VeryVerbose,
 		TEXT("http-%3u: %5" UINT64_FMT "ms %5" UINT64_FMT "KiB %s (Attempt %d/%d)"),
 		Request.StatusCode,
 		DurationMs,
diff --git a/Engine/Source/Runtime/Experimental/IoStore/OnDemand/Private/OnDemandHttpClient.h b/Engine/Source/Runtime/Experimental/IoStore/OnDemand/Private/OnDemandHttpClient.h
--- a/Engine/Source/Runtime/Experimental/IoStore/OnDemand/Private/OnDemandHttpClient.h
+++ b/Engine/Source/Runtime/Experimental/IoStore/OnDemand/Private/OnDemandHttpClient.h
@@ -13,6 +13,15 @@
 #include "Templates/Function.h"
 #include "Templates/UniquePtr.h"
 
+// When NO_LOGGING is enabled, then log categories will change type to 'FNoLoggingCategory' which does
+// not inherit from 'FLogCategoryBase' so in order to keep code compiling we need to be able to switch
+// types as well.
+#if !NO_LOGGING
+	#define UE_LOG_CATEGORY_TYPE FLogCategoryBase
+#else
+	#define UE_LOG_CATEGORY_TYPE FNoLoggingCategory
+#endif //!NO_LOGGING
+
 namespace UE::IoStore
 {
 
@@ -28,19 +37,15 @@ enum class EHttpRedirects
 
 struct FMultiEndpointHttpClientConfig
 {
-	int32				MaxConnectionCount = 4;
-	int32				ReceiveBufferSize = -1;
-	uint32				MaxRetryCount = 0;
-	int32				TimeoutMs = 0;
-	EHttpRedirects		Redirects = EHttpRedirects::Follow;
-	bool				bEnableThreadSafetyChecks = false;
-	bool				bAllowChunkedTransfer = true;
-
-#if !NO_LOGGING
-	/** When not set (the default) this will disable logging, it is up to the system creating the client to assign a LogCategory */
-	FLogCategoryName	LogCategoryName;
-	ELogVerbosity::Type	LogVerbosity = ELogVerbosity::Log;
-#endif //!NO_LOGGING
+	int32						MaxConnectionCount = 4;
+	int32						ReceiveBufferSize = -1;
+	uint32						MaxRetryCount = 0;
+	int32						TimeoutMs = 0;
+	EHttpRedirects				Redirects = EHttpRedirects::Follow;
+	bool						bEnableThreadSafetyChecks = false;
+	bool						bAllowChunkedTransfer = true;
+	/** Logging will be disabled if this is set to a nullptr, it is up to the calling system to assign a LogCategory */
+	const UE_LOG_CATEGORY_TYPE*	LogCategory = nullptr;
 };
 
 struct FMultiEndpointHttpClientResponse
@@ -129,3 +134,5 @@ private:
 };
 
 } // namespace UE::IoStore
+
+#undef UE_LOG_CATEGORY_TYPE
diff --git a/Engine/Source/Runtime/Experimental/IoStore/OnDemand/Private/OnDemandIoDispatcherBackend.cpp b/Engine/Source/Runtime/Experimental/IoStore/OnDemand/Private/OnDemandIoDispatcherBackend.cpp
--- a/Engine/Source/Runtime/Experimental/IoStore/OnDemand/Private/OnDemandIoDispatcherBackend.cpp
+++ b/Engine/Source/Runtime/Experimental/IoStore/OnDemand/Private/OnDemandIoDispatcherBackend.cpp
@@ -1666,11 +1666,7 @@ uint32 FOnDemandIoBackend::Run()
 			.TimeoutMs = GIasHttpFailTimeOutMs,
 			.bEnableThreadSafetyChecks = true,
 			.bAllowChunkedTransfer = GIasHttpAllowChunkedXfer,
-
-#if !NO_LOGGING
-			.LogCategoryName = LogIas.GetCategoryName(),
-			.LogVerbosity = ELogVerbosity::VeryVerbose
-#endif // !NO_LOGGING
+			.LogCategory = &LogIas
 		});
 	check(HttpClient.IsValid());
 
diff --git a/Engine/Source/Runtime/Experimental/IoStore/OnDemand/Private/OnDemandIoStore.cpp b/Engine/Source/Runtime/Experimental/IoStore/OnDemand/Private/OnDemandIoStore.cpp
--- a/Engine/Source/Runtime/Experimental/IoStore/OnDemand/Private/OnDemandIoStore.cpp
+++ b/Engine/Source/Runtime/Experimental/IoStore/OnDemand/Private/OnDemandIoStore.cpp
@@ -1449,9 +1449,7 @@ FOnDemandInstallResult FOnDemandIoStore::TickInstallRequest(FOnDemandInternalIns
 	FMultiEndpointHttpClientConfig HttpConfig;
 	HttpConfig.MaxConnectionCount = 8;
 	HttpConfig.MaxRetryCount = 2;
-#if !NO_LOGGING
-	HttpConfig.LogCategoryName = LogIoStoreOnDemand.GetCategoryName();
-#endif // !NO_LOGGING
+	HttpConfig.LogCategory = &LogIoStoreOnDemand;
 
 	TUniquePtr<FMultiEndpointHttpClient> HttpClient = FMultiEndpointHttpClient::Create(MoveTemp(HttpConfig));
 	if (HttpClient.IsValid() == false)
