Iris
- Fixes for compiling without iris

diff --git a/Engine/Source/Runtime/Experimental/Iris/Stub/Public/Iris/ReplicationSystem/ChunkedDataStream/ChunkedDataStream.h b/Engine/Source/Runtime/Experimental/Iris/Stub/Public/Iris/ReplicationSystem/ChunkedDataStream/ChunkedDataStream.h
new file mode 100644
--- /dev/null
+++ b/Engine/Source/Runtime/Experimental/Iris/Stub/Public/Iris/ReplicationSystem/ChunkedDataStream/ChunkedDataStream.h
@@ -0,0 +1,12 @@
+// Copyright Epic Games, Inc. All Rights Reserved.
+
+#pragma once
+
+#include "UObject/ObjectMacros.h"
+#include "ChunkedDataStream.generated.h"
+
+UCLASS(MinimalAPI, transient)
+class UChunkedDataStream : public UObject
+{
+	GENERATED_BODY()
+};
diff --git a/Engine/Source/Runtime/Experimental/Iris/Stub/Public/Iris/ReplicationSystem/Filtering/NetObjectConnectionFilter.h b/Engine/Source/Runtime/Experimental/Iris/Stub/Public/Iris/ReplicationSystem/Filtering/NetObjectConnectionFilter.h
new file mode 100644
--- /dev/null
+++ b/Engine/Source/Runtime/Experimental/Iris/Stub/Public/Iris/ReplicationSystem/Filtering/NetObjectConnectionFilter.h
@@ -0,0 +1,21 @@
+// Copyright Epic Games, Inc. All Rights Reserved.
+
+#pragma once
+
+#include "UObject/ObjectMacros.h"
+#include "NetObjectConnectionFilter.generated.h"
+
+// Stub used when we are not compiling with iris to workaround UHT dependencies
+static_assert(UE_WITH_IRIS == 0, "IrisStub module should not be used when iris is enabled");
+
+UCLASS(Transient, MinimalAPI)
+class UNetObjectConnectionFilterConfig : public UObject
+{
+	GENERATED_BODY()
+};
+
+UCLASS(Transient, MinimalAPI)
+class UNetObjectConnectionFilter : public UObject
+{
+	GENERATED_BODY()
+};
diff --git a/Engine/Source/Runtime/Experimental/Iris/Stub/Public/Iris/ReplicationSystem/Filtering/NetObjectGridFilter.h b/Engine/Source/Runtime/Experimental/Iris/Stub/Public/Iris/ReplicationSystem/Filtering/NetObjectGridFilter.h
new file mode 100644
--- /dev/null
+++ b/Engine/Source/Runtime/Experimental/Iris/Stub/Public/Iris/ReplicationSystem/Filtering/NetObjectGridFilter.h
@@ -0,0 +1,21 @@
+// Copyright Epic Games, Inc. All Rights Reserved.
+
+#pragma once
+
+#include "UObject/ObjectMacros.h"
+#include "NetObjectGridFilter.generated.h"
+
+// Stub used when we are not compiling with iris to workaround UHT dependencies
+static_assert(UE_WITH_IRIS == 0, "IrisStub module should not be used when iris is enabled");
+
+UCLASS(Transient, MinimalAPI)
+class UNetObjectGridFilterConfig : public UObject
+{
+	GENERATED_BODY()
+};
+
+UCLASS(Transient, MinimalAPI)
+class UNetObjectGridFilter : public UObject
+{
+	GENERATED_BODY()
+};
diff --git a/Engine/Source/Runtime/Experimental/Iris/Stub/Public/Iris/ReplicationSystem/Prioritization/SphereNetObjectPrioritizer.h b/Engine/Source/Runtime/Experimental/Iris/Stub/Public/Iris/ReplicationSystem/Prioritization/SphereNetObjectPrioritizer.h
new file mode 100644
--- /dev/null
+++ b/Engine/Source/Runtime/Experimental/Iris/Stub/Public/Iris/ReplicationSystem/Prioritization/SphereNetObjectPrioritizer.h
@@ -0,0 +1,15 @@
+// Copyright Epic Games, Inc. All Rights Reserved.
+
+#pragma once
+
+#include "Iris/ReplicationSystem/Prioritization/NetObjectPrioritizer.h"
+#include "SphereNetObjectPrioritizer.generated.h"
+
+// Stub used when we are not compiling with iris to workaround UHT dependencies
+static_assert(UE_WITH_IRIS == 0, "IrisStub module should not be used when iris is enabled");
+	
+UCLASS(transient, MinimalAPI)
+class USphereNetObjectPrioritizerConfig : public UNetObjectPrioritizerConfig
+{
+	GENERATED_BODY()
+};
