simplify collection selection facade's leaf and cluster selection logic, and remove its dependency on the level attribute

diff --git a/Engine/Source/Runtime/Experimental/Chaos/Private/GeometryCollection/Facades/CollectionTransformSelectionFacade.cpp b/Engine/Source/Runtime/Experimental/Chaos/Private/GeometryCollection/Facades/CollectionTransformSelectionFacade.cpp
--- a/Engine/Source/Runtime/Experimental/Chaos/Private/GeometryCollection/Facades/CollectionTransformSelectionFacade.cpp
+++ b/Engine/Source/Runtime/Experimental/Chaos/Private/GeometryCollection/Facades/CollectionTransformSelectionFacade.cpp
@@ -420,20 +420,17 @@ namespace GeometryCollection::Facades
 
 		if (CanSelectLeaf())
 		{
-			const TManagedArray<int32>& Levels = LevelAttribute.Get();
 			const TManagedArray<int32>& SimType = SimulationTypeAttribute.Get();
 
-			const int32 NumTransforms = ParentAttribute.Num();
-
-			const int32 ViewLevel = INDEX_NONE;
-
-			OutSelection = GetBonesByLevel(ViewLevel, true, true);
-
-			OutSelection.SetNum(Algo::RemoveIf(OutSelection, [&](int32 BoneIdx)
+			const int32 NumTransforms = SimType.Num();
+			for (int32 BoneIdx = 0; BoneIdx < NumTransforms; BoneIdx++)
+			{
+				bool bIsRigid = SimType[BoneIdx] == FGeometryCollection::ESimulationTypes::FST_Rigid;
+				if (bIsRigid)
 				{
-					return SimType[BoneIdx] != FGeometryCollection::ESimulationTypes::FST_Rigid
-						|| (ViewLevel != INDEX_NONE && Levels[BoneIdx] != ViewLevel);
-				}));
+					OutSelection.Add(BoneIdx);
+				}
+			}
 		}
 		
 		return OutSelection;
@@ -445,20 +442,17 @@ namespace GeometryCollection::Facades
 
 		if (CanSelectCluster())
 		{
-			const TManagedArray<int32>& Levels = LevelAttribute.Get();
 			const TManagedArray<int32>& SimType = SimulationTypeAttribute.Get();
 
-			const int32 NumTransforms = ParentAttribute.Num();
-
-			const int32 ViewLevel = INDEX_NONE;
-
-			OutSelection = GetBonesByLevel(ViewLevel, true, true);
-
-			OutSelection.SetNum(Algo::RemoveIf(OutSelection, [&](int32 BoneIdx)
+			const int32 NumTransforms = SimType.Num();
+			for (int32 BoneIdx = 0; BoneIdx < NumTransforms; BoneIdx++)
+			{
+				bool bIsCluster = SimType[BoneIdx] == FGeometryCollection::ESimulationTypes::FST_Clustered;
+				if (bIsCluster)
 				{
-					return SimType[BoneIdx] != FGeometryCollection::ESimulationTypes::FST_Clustered
-						|| (ViewLevel != INDEX_NONE && Levels[BoneIdx] != ViewLevel);
-				}));
+					OutSelection.Add(BoneIdx);
+				}
+			}
 		}
 		
 		return OutSelection;
diff --git a/Engine/Source/Runtime/Experimental/Chaos/Public/GeometryCollection/Facades/CollectionTransformSelectionFacade.h b/Engine/Source/Runtime/Experimental/Chaos/Public/GeometryCollection/Facades/CollectionTransformSelectionFacade.h
--- a/Engine/Source/Runtime/Experimental/Chaos/Public/GeometryCollection/Facades/CollectionTransformSelectionFacade.h
+++ b/Engine/Source/Runtime/Experimental/Chaos/Public/GeometryCollection/Facades/CollectionTransformSelectionFacade.h
@@ -93,11 +93,11 @@ namespace GeometryCollection::Facades
 		CHAOS_API TArray<int32> SelectRandom(bool bDeterministic, float RandomSeed, float RandomThresholdVal) const;
 
 		/**  */
-		bool CanSelectLeaf() const { return ParentAttribute.IsValid() && LevelAttribute.IsValid() && SimulationTypeAttribute.IsValid(); }
+		bool CanSelectLeaf() const { return SimulationTypeAttribute.IsValid(); }
 		CHAOS_API TArray<int32> SelectLeaf() const;
 
 		/**  */
-		bool CanSelectCluster() const { return ParentAttribute.IsValid() && LevelAttribute.IsValid() && SimulationTypeAttribute.IsValid(); }
+		bool CanSelectCluster() const { return SimulationTypeAttribute.IsValid(); }
 		CHAOS_API TArray<int32> SelectCluster() const;
 
 		/**  */
