Fix OnDemandIOStore cancel
Fix OnDemandIOStore progress

diff --git a/Engine/Source/Runtime/Experimental/IoStore/OnDemand/Private/OnDemandIoStore.cpp b/Engine/Source/Runtime/Experimental/IoStore/OnDemand/Private/OnDemandIoStore.cpp
--- a/Engine/Source/Runtime/Experimental/IoStore/OnDemand/Private/OnDemandIoStore.cpp
+++ b/Engine/Source/Runtime/Experimental/IoStore/OnDemand/Private/OnDemandIoStore.cpp
@@ -1321,7 +1321,7 @@ FOnDemandInstallResult FOnDemandIoStore::TickInstallRequest(FOnDemandInternalIns
 
 	FOnDemandInstallResult OutResult;
 	FIoStatus& Status					= OutResult.Status;
-	FOnDemandInstallProgress& Progress	= InstallRequest.Progress;
+	FOnDemandInstallProgress& Progress	= OutResult.Progress;
 
 	const double StartTime = FPlatformTime::Seconds();
 	ON_SCOPE_EXIT
@@ -1647,10 +1647,9 @@ void FOnDemandIoStore::CompleteInstallRequest(FSharedInternalInstallRequest Requ
 
 	if (!Request->OnCompleted)
 	{
-		return Request->Status.store(Status);
+		Request->Status.store(Status);
 	}
-
-	if (EnumHasAnyFlags(Request->Args.Options, EOnDemandInstallOptions::CallbackOnGameThread))
+	else if (EnumHasAnyFlags(Request->Args.Options, EOnDemandInstallOptions::CallbackOnGameThread))
 	{
 		ExecuteOnGameThread(
 			UE_SOURCE_LOCATION,
@@ -2054,14 +2053,6 @@ void FOnDemandIoStore::CancelInstallRequest(FSharedInternalInstallRequest Instal
 	}
 
 	UE_LOG(LogIoStoreOnDemand, Log, TEXT("Cancelling install request, %s"), *LexToString(*InstallRequest));
-
-	FOnDemandInstallResult InstallResult
-	{
-		.Status				= FIoStatus(EIoErrorCode::Cancelled),
-		.DurationInSeconds 	= FPlatformTime::Seconds() - InstallRequest->StartTime,
-		.Progress			= MoveTemp(InstallRequest->Progress)
-	};
-	CompleteInstallRequest(InstallRequest, MoveTemp(InstallResult));
 }
 
 void FOnDemandIoStore::UpdateInstallRequestPriority(FSharedInternalInstallRequest InstallRequest, EOnDemandInstallPriority NewPriority)
@@ -2192,13 +2183,15 @@ FOnDemandInstallRequest::~FOnDemandInstallRequest()
 }
 
 FOnDemandInstallRequest::FOnDemandInstallRequest(FOnDemandInstallRequest&& Other)
-	: Request(MoveTemp(Other.Request))
+	: IoStore(MoveTemp(Other.IoStore))
+	, Request(MoveTemp(Other.Request))
 {
 }
 
 FOnDemandInstallRequest& FOnDemandInstallRequest::operator=(FOnDemandInstallRequest&& Other)
 {
 	Request = MoveTemp(Other.Request);
+	IoStore = MoveTemp(Other.IoStore);
 	return *this;
 }
 
diff --git a/Engine/Source/Runtime/Experimental/IoStore/OnDemand/Private/OnDemandIoStore.h b/Engine/Source/Runtime/Experimental/IoStore/OnDemand/Private/OnDemandIoStore.h
--- a/Engine/Source/Runtime/Experimental/IoStore/OnDemand/Private/OnDemandIoStore.h
+++ b/Engine/Source/Runtime/Experimental/IoStore/OnDemand/Private/OnDemandIoStore.h
@@ -83,7 +83,6 @@ public:
 	FOnDemandInstallArgs		Args;
 	FOnDemandInstallCompleted	OnCompleted;
 	FOnDemandInstallProgressed	OnProgress;
-	FOnDemandInstallProgress	Progress;
 	double						StartTime = FPlatformTime::Seconds();
 	uint32						SeqNo = NextSeqNo++;
 	std::atomic<EStatus>		Status{EStatus::Pending};
