Geometry collection : Add transforms to root proxy data
- transforms are local to the component they are used on (WorldTransform = MeshTransform * ComponentToWorldTransform)

diff --git a/Engine/Source/Runtime/Experimental/GeometryCollectionEngine/Private/GeometryCollection/GeometryCollectionComponent.cpp b/Engine/Source/Runtime/Experimental/GeometryCollectionEngine/Private/GeometryCollection/GeometryCollectionComponent.cpp
--- a/Engine/Source/Runtime/Experimental/GeometryCollectionEngine/Private/GeometryCollection/GeometryCollectionComponent.cpp
+++ b/Engine/Source/Runtime/Experimental/GeometryCollectionEngine/Private/GeometryCollection/GeometryCollectionComponent.cpp
@@ -1452,16 +1452,20 @@ bool UGeometryCollectionComponent::DoCustomNavigableGeometryExport(FNavigableGeo
 
 	if (bUseRootProxyForNavigation)
 	{
+		const FGeometryCollectionProxyMeshData& ProxyMeshData = RestCollection->RootProxyData;
+
 		bool bHasData = false;
-		for (int32 MeshIndex = 0; MeshIndex < RestCollection->RootProxyData.ProxyMeshes.Num(); MeshIndex++)
+		for (int32 MeshIndex = 0; MeshIndex < ProxyMeshData.ProxyMeshes.Num(); MeshIndex++)
 		{
-			const TObjectPtr<UStaticMesh>& ProxyMesh = RestCollection->RootProxyData.ProxyMeshes[MeshIndex];
+			const TObjectPtr<UStaticMesh>& ProxyMesh = ProxyMeshData.ProxyMeshes[MeshIndex];
 			if (ProxyMesh != nullptr)
 			{
 				const int32 RootIndex = GetRootIndex();
 				const FTransform& CompToWorld = GetComponentToWorld(); 
 				const FTransform FinalTransform = (DynamicCollection ? FTransform(DynamicCollection->GetTransform(RootIndex)) : FTransform::Identity) * CompToWorld;
 				const FVector Scale3D = FinalTransform.GetScale3D();
+				const FTransform MeshLocalTransform{ ProxyMeshData.GetMeshTransform(MeshIndex) };
+
 				if (!Scale3D.IsZero())
 				{
 					if (const UNavCollisionBase* NavCollision = ProxyMesh->GetNavCollision())
@@ -1471,7 +1475,7 @@ bool UGeometryCollectionComponent::DoCustomNavigableGeometryExport(FNavigableGeo
 							continue;
 						}
 						
-						bHasData = NavCollision->ExportGeometry(FinalTransform, GeomExport) || bHasData;
+						bHasData = NavCollision->ExportGeometry(MeshLocalTransform * FinalTransform, GeomExport) || bHasData;
 					}
 				}
 			}
@@ -3360,13 +3364,16 @@ FTransform UGeometryCollectionComponent::GetSocketTransform(FName InSocketName,
 		}
 		if (!bFoundSocket)
 		{
-			for (const TObjectPtr<UStaticMesh>& ProxyMesh : RestCollection->RootProxyData.ProxyMeshes)
+			const FGeometryCollectionProxyMeshData& ProxyMeshData = RestCollection->RootProxyData;
+			for (int32 MeshIdx = 0; MeshIdx < ProxyMeshData.ProxyMeshes.Num(); ++MeshIdx)
 			{
+				const TObjectPtr<UStaticMesh>& ProxyMesh = ProxyMeshData.ProxyMeshes[MeshIdx];
 				if (ProxyMesh)
 				{
 					if (UStaticMeshSocket* MeshSocket = ProxyMesh->FindSocket(InSocketName))
 					{
 						SocketComponentSpaceTransform = FTransform(MeshSocket->RelativeRotation, MeshSocket->RelativeLocation, MeshSocket->RelativeScale);
+						SocketComponentSpaceTransform = SocketComponentSpaceTransform * FTransform(ProxyMeshData.GetMeshTransform(MeshIdx));
 						bFoundSocket = true;
 						break;
 					}
@@ -3970,8 +3977,10 @@ PRAGMA_ENABLE_DEPRECATION_WARNINGS
 	}
 
 	auto CreateTraceCollisionFromStaticMeshGeometry = [this](const FTransform& InToLocal, TArray<Chaos::FImplicitObjectPtr>& OutGeoms, Chaos::FShapesArray& OutShapes) {
-		for (const TObjectPtr<UStaticMesh>& ProxyMesh : RestCollection->RootProxyData.ProxyMeshes)
+		const FGeometryCollectionProxyMeshData& ProxyMeshData = RestCollection->RootProxyData;
+		for (int32 MeshIdx = 0; MeshIdx < ProxyMeshData.ProxyMeshes.Num(); ++MeshIdx)
 		{
+			const TObjectPtr<UStaticMesh>& ProxyMesh = RestCollection->RootProxyData.ProxyMeshes[MeshIdx];
 			if (const UStaticMesh* StaticMesh = ProxyMesh.Get())
 			{
 				if (UBodySetup* BodySetup = ProxyMesh->GetBodySetup())
@@ -3980,11 +3989,13 @@ PRAGMA_ENABLE_DEPRECATION_WARNINGS
 					BodyInstance.BuildBodyFilterData(BodyCollisionData.CollisionFilterData);
 					FBodyInstance::BuildBodyCollisionFlags(BodyCollisionData.CollisionFlags, BodyInstance.GetCollisionEnabled(), BodySetup->GetCollisionTraceFlag() == CTF_UseComplexAsSimple);
 
+					const FTransform MeshLocalTransform(ProxyMeshData.GetMeshTransform(MeshIdx));
+
 					FGeometryAddParams AddParams;
 					AddParams.bDoubleSided = BodySetup->bDoubleSidedGeometry;
 					AddParams.CollisionData = BodyCollisionData;
 					AddParams.CollisionTraceType = BodySetup->GetCollisionTraceFlag();
-					AddParams.WorldTransform = PhysicsProxy->GetWorldTransform_External();
+					AddParams.WorldTransform = MeshLocalTransform * PhysicsProxy->GetWorldTransform_External();
 					AddParams.Scale = AddParams.WorldTransform.GetScale3D();
 					AddParams.SimpleMaterial = GetPhysicalMaterial();
 					AddParams.LocalTransform = InToLocal;
@@ -6300,7 +6311,9 @@ void UGeometryCollectionComponent::SetRootProxyLocalTransform(int32 Index, const
 	if (RootProxyLocalTransforms.IsEmpty() && RestCollection)
 	{
 		const int32 NumProxyMeshes = RestCollection->RootProxyData.ProxyMeshes.Num();
-		RootProxyLocalTransforms.Init(FTransform3f::Identity, NumProxyMeshes);
+		RootProxyLocalTransforms.Reserve(NumProxyMeshes);
+		RootProxyLocalTransforms.Append(RestCollection->RootProxyData.MeshTransforms);
+		RootProxyLocalTransforms.Append(&FTransform3f::Identity, (NumProxyMeshes - RootProxyLocalTransforms.Num()));
 	}
 	if (RootProxyLocalTransforms.IsValidIndex(Index))
 	{
@@ -6375,7 +6388,14 @@ void UGeometryCollectionComponent::RefreshCustomRenderer()
 				const FTransform CompSpaceRootTransform(ComponentSpaceTransforms.RequestRootTransform());
 				if (RootProxyLocalTransforms.IsEmpty())
 				{
-					RendererInterface->UpdateRootTransform(*RestCollection, CompSpaceRootTransform);
+					if (RestCollection->RootProxyData.MeshTransforms.IsEmpty())
+					{
+						RendererInterface->UpdateRootTransform(*RestCollection, CompSpaceRootTransform);
+					}
+					else
+					{
+						RendererInterface->UpdateRootTransforms(*RestCollection, CompSpaceRootTransform, MakeArrayView(RestCollection->RootProxyData.MeshTransforms));
+					}
 				}
 				else
 				{
diff --git a/Engine/Source/Runtime/Experimental/GeometryCollectionEngine/Private/GeometryCollection/GeometryCollectionObject.cpp b/Engine/Source/Runtime/Experimental/GeometryCollectionEngine/Private/GeometryCollection/GeometryCollectionObject.cpp
--- a/Engine/Source/Runtime/Experimental/GeometryCollectionEngine/Private/GeometryCollection/GeometryCollectionObject.cpp
+++ b/Engine/Source/Runtime/Experimental/GeometryCollectionEngine/Private/GeometryCollection/GeometryCollectionObject.cpp
@@ -240,6 +240,14 @@ void FGeometryCollectionSizeSpecificData::PostSerialize(const FArchive& Ar)
 }
 #endif
 
+const FTransform3f& FGeometryCollectionProxyMeshData::GetMeshTransform(int32 MeshIndex) const
+{
+	if (MeshTransforms.IsValidIndex(MeshIndex))
+	{
+		return MeshTransforms[MeshIndex];
+	}
+	return FTransform3f::Identity;
+}
 
 
 void FillSharedSimulationSizeSpecificData(FSharedSimulationSizeSpecificData& ToData, const FGeometryCollectionSizeSpecificData& FromData)
diff --git a/Engine/Source/Runtime/Experimental/GeometryCollectionEngine/Public/GeometryCollection/GeometryCollectionObject.h b/Engine/Source/Runtime/Experimental/GeometryCollectionEngine/Public/GeometryCollection/GeometryCollectionObject.h
--- a/Engine/Source/Runtime/Experimental/GeometryCollectionEngine/Public/GeometryCollection/GeometryCollectionObject.h
+++ b/Engine/Source/Runtime/Experimental/GeometryCollectionEngine/Public/GeometryCollection/GeometryCollectionObject.h
@@ -326,15 +326,24 @@ struct TStructOpsTypeTraits<FGeometryCollectionSizeSpecificData> : public TStruc
 	static constexpr EPropertyObjectReferenceType WithSerializerObjectReferences = EPropertyObjectReferenceType::None;
 };
 
-
-
 USTRUCT(BlueprintType)
 struct FGeometryCollectionProxyMeshData
 {
 	GENERATED_BODY()
 
+	/** References to meshes to be used as proxies for rendering when the geometry collection is not broken yet */
 	UPROPERTY(EditAnywhere, BlueprintReadWrite, Category = "Rendering", meta = (AllowedClasses = "/Script/Engine.StaticMesh"))
 	TArray<TObjectPtr<UStaticMesh>> ProxyMeshes;
+
+	/** 
+	* Transform of the mesh relative to the geometry collection it belongs to 
+	* if an entry is missing or if the arary is empty , those transforms are assume to be identity
+	* making any proxy mesh done before 5.6 be fully backward compatible
+	*/
+	UPROPERTY(EditAnywhere, BlueprintReadWrite, Category = "Rendering")
+	TArray<FTransform3f> MeshTransforms;
+
+	GEOMETRYCOLLECTIONENGINE_API const FTransform3f& GetMeshTransform(int32 MeshIndex) const;
 };
 
 USTRUCT()
diff --git a/Engine/Source/Runtime/NavigationSystem/Private/NavModifierComponent.cpp b/Engine/Source/Runtime/NavigationSystem/Private/NavModifierComponent.cpp
--- a/Engine/Source/Runtime/NavigationSystem/Private/NavModifierComponent.cpp
+++ b/Engine/Source/Runtime/NavigationSystem/Private/NavModifierComponent.cpp
@@ -165,15 +165,18 @@ void UNavModifierComponent::CalculateBounds() const
     			{
     				Bounds += GeometryCollection->Bounds.GetBox();
 
-    				for (const TObjectPtr<UStaticMesh>& ProxyMesh : RestCollection->RootProxyData.ProxyMeshes)
+					const FGeometryCollectionProxyMeshData& ProxyMeshData = RestCollection->RootProxyData;
+					for (int32 MeshIdx = 0; MeshIdx < ProxyMeshData.ProxyMeshes.Num(); ++MeshIdx)
     				{
+						const TObjectPtr<UStaticMesh>& ProxyMesh = ProxyMeshData.ProxyMeshes[MeshIdx];
     					if (ProxyMesh != nullptr)
     					{
     						const UBodySetup* Body = ProxyMesh->GetBodySetup();
     						if (Body)
     						{
     							const FTransform& ParentTM = PrimComp->GetComponentTransform();
-    							PopulateComponentBounds(ParentTM, *Body);	
+								const FTransform LocalMeshTransform(ProxyMeshData.GetMeshTransform(MeshIdx));
+    							PopulateComponentBounds(LocalMeshTransform* ParentTM, *Body);
     						}
     					}
     				}
