[DSTM] Switching to delta-serialization by default for physics state migration.

Also fixed an issue where during skeletal mesh migration, the data for the last processed joint constraint is not applied.

Lastly, I updated CVD's data wrappers so we can use either the custom serialization code or UE's property serialization code independently.

diff --git a/Engine/Source/Runtime/Engine/Classes/Components/SkeletalMeshComponent.h b/Engine/Source/Runtime/Engine/Classes/Components/SkeletalMeshComponent.h
--- a/Engine/Source/Runtime/Engine/Classes/Components/SkeletalMeshComponent.h
+++ b/Engine/Source/Runtime/Engine/Classes/Components/SkeletalMeshComponent.h
@@ -1992,12 +1992,17 @@ private:
 		template<typename ElementType>
 		void AppendDataFromElement(ElementType& InElement)
 		{
-			DataOffsetsPerElementIndex.Emplace(MigratedStateAsBytes.GetDataAsByteArrayRef().Num());
+			DataOffsetsPerElementIndex.Emplace(GetSize());
 			InElement.SerializeState(MigratedStateAsBytes);
 		}
 
 		void Serialize(FArchive& Ar);
 
+		int32 GetSize() const
+		{
+			return MigratedStateAsBytes.GetSize();
+		}
+
 	private:
 		int32 GetDataSizeForElementIndex(int32 ElementIndex) const;
 		int32 GetDataOffset(int32 ElementIndex) const;
diff --git a/Engine/Source/Runtime/Engine/Private/Components/SkeletalMeshComponent.cpp b/Engine/Source/Runtime/Engine/Private/Components/SkeletalMeshComponent.cpp
--- a/Engine/Source/Runtime/Engine/Private/Components/SkeletalMeshComponent.cpp
+++ b/Engine/Source/Runtime/Engine/Private/Components/SkeletalMeshComponent.cpp
@@ -3801,23 +3801,42 @@ void USkeletalMeshComponent::FMigratedPhysicsState::ReadElementDataIntoBuffer(Ch
 
 	TArray<uint8>& AsByteArray = MigratedStateAsBytes.GetDataAsByteArrayRef();
 
-	if (StartPos + SizeToCopy < AsByteArray.Num())
+	if (ensure(StartPos + SizeToCopy < AsByteArray.Num()))
 	{
 		TargetBuffer.GetDataAsByteArrayRef().Append(&AsByteArray[StartPos], SizeToCopy);
 	}
+	else
+	{
+		UE_LOG(LogSkeletalMesh, Error, TEXT("[Physics State Migration] Attempted to read data out of bounds! | Buffer size [%d] | Start Pos [%d]| SizeToCopy [%d]"), AsByteArray.Num(), StartPos, SizeToCopy);
+	}
 	
 	CurrentReadElementIndex++;
 }
 
 int32 USkeletalMeshComponent::FMigratedPhysicsState::GetDataSizeForElementIndex(int32 ElementIndex) const
 {
-	if (!DataOffsetsPerElementIndex.IsValidIndex(ElementIndex))
+	if (!ensure(DataOffsetsPerElementIndex.IsValidIndex(ElementIndex)))
 	{
+		UE_LOG(LogSkeletalMesh, Error, TEXT("[Physics State Migration] Attempted to get the data size for an invalid element index | Requested Index [%d] | Max Element Index [%d] "), ElementIndex, DataOffsetsPerElementIndex.Num());
 		return 0;
 	}
 
-	int32 NextIndex = ElementIndex + 1;
-	return DataOffsetsPerElementIndex.IsValidIndex(NextIndex) ? DataOffsetsPerElementIndex[NextIndex] - DataOffsetsPerElementIndex[ElementIndex] : 0;
+	const int32 NextIndex = ElementIndex + 1;
+
+	if (DataOffsetsPerElementIndex.IsValidIndex(NextIndex))
+	{
+		return DataOffsetsPerElementIndex[NextIndex] - DataOffsetsPerElementIndex[ElementIndex];
+	}
+
+	if (ensure(DataOffsetsPerElementIndex.Num() == NextIndex))
+	{
+		// If this is the last element, the size of data to read is the remaining of the buffer from where we currently are at
+		const int32 TotalInnerBufferSizeMinusOne = GetSize() - 1;
+		return TotalInnerBufferSizeMinusOne - DataOffsetsPerElementIndex[ElementIndex];
+	}
+
+	UE_LOG(LogSkeletalMesh, Error, TEXT("[Physics State Migration] Failed to get the data size for Index [%d] "), ElementIndex);
+	return 0;
 }
 
 int32 USkeletalMeshComponent::FMigratedPhysicsState::GetDataOffset(int32 ElementIndex) const
diff --git a/Engine/Source/Runtime/Experimental/Chaos/Private/Chaos/Serialization/SerializationUtils.cpp b/Engine/Source/Runtime/Experimental/Chaos/Private/Chaos/Serialization/SerializationUtils.cpp
--- a/Engine/Source/Runtime/Experimental/Chaos/Private/Chaos/Serialization/SerializationUtils.cpp
+++ b/Engine/Source/Runtime/Experimental/Chaos/Private/Chaos/Serialization/SerializationUtils.cpp
@@ -16,6 +16,9 @@ namespace Chaos
 {
 	namespace Serialization::Private
 	{
+		static bool bUseFastStructSerializer = false;
+		static FAutoConsoleVariableRef CVarSolverSerializerUseFastStructSerializer(TEXT("p.chaos.serialization.UseBuiltInSerializer"), bUseFastStructSerializer, TEXT("If set to true, solver serialization will use the new fast struct serializer. Otherwise it will use the built in serializer"));
+	
 		static bool CanBulkSerialize(FProperty* Property)
 		{
 	#if PLATFORM_LITTLE_ENDIAN
@@ -83,8 +86,15 @@ namespace Chaos
 			}
 		}
 
-		void FastStructSerialize(UScriptStruct* Struct, void* SourceData, FArchive& Ar)
+		void FastStructSerialize(UScriptStruct* Struct, void* SourceData, FArchive& Ar, void* Defaults)
 		{
+			if (!bUseFastStructSerializer)
+			{
+				Ar.SetUseUnversionedPropertySerialization(true);
+				Struct->SerializeTaggedProperties(Ar, static_cast<uint8*>(SourceData), Struct, static_cast<uint8*>(Defaults));
+				return;
+			}
+	
 			for (FProperty* Property = Struct->PropertyLink; Property != nullptr; Property = Property->PropertyLinkNext)
 			{
 				for (int32 Idx = 0; Idx < Property->ArrayDim; Idx++)
diff --git a/Engine/Source/Runtime/Experimental/Chaos/Public/Chaos/Serialization/SerializationUtils.h b/Engine/Source/Runtime/Experimental/Chaos/Public/Chaos/Serialization/SerializationUtils.h
--- a/Engine/Source/Runtime/Experimental/Chaos/Public/Chaos/Serialization/SerializationUtils.h
+++ b/Engine/Source/Runtime/Experimental/Chaos/Public/Chaos/Serialization/SerializationUtils.h
@@ -14,13 +14,14 @@ namespace Chaos
 	namespace Serialization::Private
 	{
 		/** Lightweight unversioned ustruct serializer */
-		UE_INTERNAL CHAOS_API void FastStructSerialize(UScriptStruct* Struct, void* SourceData, FArchive& Ar);
+		UE_INTERNAL CHAOS_API void FastStructSerialize(UScriptStruct* Struct, void* SourceData, FArchive& Ar, void* Defaults = nullptr);
 
 		/** Lightweight unversioned ustruct serializer */
 		template<typename StructType>
 		UE_INTERNAL void FastStructSerialize(FArchive& Ar, StructType* Data)
 		{
-			FastStructSerialize(StructType::StaticStruct(), reinterpret_cast<void*>(Data), Ar);
+			StructType DefaultStruct = StructType();
+			FastStructSerialize(StructType::StaticStruct(), reinterpret_cast<void*>(Data), Ar, &DefaultStruct);
 		}
 	}
 }
diff --git a/Engine/Source/Runtime/Experimental/Chaos/Public/Chaos/Serialization/SerializedDataBuffer.h b/Engine/Source/Runtime/Experimental/Chaos/Public/Chaos/Serialization/SerializedDataBuffer.h
--- a/Engine/Source/Runtime/Experimental/Chaos/Public/Chaos/Serialization/SerializedDataBuffer.h
+++ b/Engine/Source/Runtime/Experimental/Chaos/Public/Chaos/Serialization/SerializedDataBuffer.h
@@ -23,6 +23,11 @@ namespace Chaos
 			return DataBuffer;
 		}
 
+		int32 GetSize() const
+		{
+			return DataBuffer.Num();
+		}
+
 		void Serialize(FArchive& Ar)
 		{
 			Ar << DataBuffer;
diff --git a/Engine/Source/Runtime/Experimental/ChaosVisualDebugger/Public/DataWrappers/ChaosVDDataSerializationMacros.h b/Engine/Source/Runtime/Experimental/ChaosVisualDebugger/Public/DataWrappers/ChaosVDDataSerializationMacros.h
--- a/Engine/Source/Runtime/Experimental/ChaosVisualDebugger/Public/DataWrappers/ChaosVDDataSerializationMacros.h
+++ b/Engine/Source/Runtime/Experimental/ChaosVisualDebugger/Public/DataWrappers/ChaosVDDataSerializationMacros.h
@@ -32,14 +32,5 @@
 		{\
 			Data.Serialize(Ar); \
 			return Ar; \
-		} \
-		template<>\
-		struct TStructOpsTypeTraits<Type> : public TStructOpsTypeTraitsBase2<Type> \
-		{\
-			enum\
-			{\
-				WithSerializer = true,\
-			};\
-		};\
-
+		} 
 #endif
\ No newline at end of file
diff --git a/Engine/Source/Runtime/Experimental/ChaosVisualDebugger/Public/DataWrappers/ChaosVDParticleDataWrapper.h b/Engine/Source/Runtime/Experimental/ChaosVisualDebugger/Public/DataWrappers/ChaosVDParticleDataWrapper.h
--- a/Engine/Source/Runtime/Experimental/ChaosVisualDebugger/Public/DataWrappers/ChaosVDParticleDataWrapper.h
+++ b/Engine/Source/Runtime/Experimental/ChaosVisualDebugger/Public/DataWrappers/ChaosVDParticleDataWrapper.h
@@ -143,11 +143,7 @@ struct FChaosVDFRigidParticleControlFlags : public FChaosVDWrapperDataBase
 	uint8 ProjectionSolverIterationCount;
 };
 
-inline FArchive& operator<<(FArchive& Ar, FChaosVDFRigidParticleControlFlags& Data)
-{
-	Data.Serialize(Ar);
-	return Ar;
-}
+CVD_IMPLEMENT_SERIALIZER(FChaosVDFRigidParticleControlFlags)
 
 template<>
 struct TStructOpsTypeTraits<FChaosVDFRigidParticleControlFlags> : public TStructOpsTypeTraitsBase2<FChaosVDFRigidParticleControlFlags>
@@ -188,20 +184,7 @@ struct FChaosVDParticlePositionRotation : public FChaosVDWrapperDataBase
 	FQuat MR = FQuat(ForceInit);
 };
 
-inline FArchive& operator<<(FArchive& Ar, FChaosVDParticlePositionRotation& Data)
-{
-	Data.Serialize(Ar);
-	return Ar;
-}
-
-template<>
-struct TStructOpsTypeTraits<FChaosVDParticlePositionRotation> : public TStructOpsTypeTraitsBase2<FChaosVDParticlePositionRotation>
-{
-	enum
-	{
-		WithSerializer = true,
-	};
-};
+CVD_IMPLEMENT_SERIALIZER(FChaosVDParticlePositionRotation)
 
 /** Simplified UStruct version of FParticleVelocities.
  * Used to be able to show the values in the editor and allow changes via the Property Editor.
@@ -239,20 +222,7 @@ struct FChaosVDParticleVelocities : public FChaosVDWrapperDataBase
 	FVector MW = FVector(ForceInit);
 };
 
-inline FArchive& operator<<(FArchive& Ar, FChaosVDParticleVelocities& Data)
-{
-	Data.Serialize(Ar);
-	return Ar;
-}
-
-template<>
-struct TStructOpsTypeTraits<FChaosVDParticleVelocities> : public TStructOpsTypeTraitsBase2<FChaosVDParticleVelocities>
-{
-	enum
-	{
-		WithSerializer = true,
-	};
-};
+CVD_IMPLEMENT_SERIALIZER(FChaosVDParticleVelocities)
 
 /** Simplified UStruct version of FParticleDynamics.
  * Used to be able to show the values in the editor and allow changes via the Property Editor.
@@ -301,21 +271,7 @@ struct FChaosVDParticleDynamics : public FChaosVDWrapperDataBase
 	FVector MAngularImpulseVelocity = FVector(ForceInit);
 };
 
-inline FArchive& operator<<(FArchive& Ar, FChaosVDParticleDynamics& Data)
-{
-	Data.Serialize(Ar);
-	return Ar;
-}
-
-template<>
-struct TStructOpsTypeTraits<FChaosVDParticleDynamics> : public TStructOpsTypeTraitsBase2<FChaosVDParticleDynamics>
-{
-	enum
-	{
-		WithSerializer = true,
-	};
-};
-
+CVD_IMPLEMENT_SERIALIZER(FChaosVDParticleDynamics)
 
 /** Simplified UStruct version of FParticleMassProps.
  * Used to be able to show the values in the editor and allow changes via the Property Editor.
@@ -374,20 +330,7 @@ struct FChaosVDParticleMassProps : public FChaosVDWrapperDataBase
 	double MInvM = 0.0;
 };
 
-inline FArchive& operator<<(FArchive& Ar, FChaosVDParticleMassProps& Data)
-{
-	Data.Serialize(Ar);
-	return Ar;
-}
-
-template<>
-struct TStructOpsTypeTraits<FChaosVDParticleMassProps> : public TStructOpsTypeTraitsBase2<FChaosVDParticleMassProps>
-{
-	enum
-	{
-		WithSerializer = true,
-	};
-};
+CVD_IMPLEMENT_SERIALIZER(FChaosVDParticleMassProps)
 
 /** Simplified UStruct version of FParticleDynamicMisc.
  * Used to be able to show the values in the editor and allow changes via the Property Editor.
@@ -485,20 +428,7 @@ struct FChaosVDParticleDynamicMisc : public FChaosVDWrapperDataBase
 	bool bDisabled;
 };
 
-inline FArchive& operator<<(FArchive& Ar, FChaosVDParticleDynamicMisc& Data)
-{
-	Data.Serialize(Ar);
-	return Ar;
-}
-
-template<>
-struct TStructOpsTypeTraits<FChaosVDParticleDynamicMisc> : public TStructOpsTypeTraitsBase2<FChaosVDParticleDynamicMisc>
-{
-	enum
-	{
-		WithSerializer = true,
-	};
-};
+CVD_IMPLEMENT_SERIALIZER(FChaosVDParticleDynamicMisc)
 
 /** Represents the data of a connectivity Edge that CVD can use to reconstruct it during playback */
 USTRUCT()
@@ -521,20 +451,7 @@ struct FChaosVDConnectivityEdge
 	}
 };
 
-inline FArchive& operator<<(FArchive& Ar, FChaosVDConnectivityEdge& Data)
-{
-	Data.Serialize(Ar);
-	return Ar;
-}
-
-template<>
-struct TStructOpsTypeTraits<FChaosVDConnectivityEdge> : public TStructOpsTypeTraitsBase2<FChaosVDConnectivityEdge>
-{
-	enum
-	{
-		WithSerializer = true,
-	};
-};
+CVD_IMPLEMENT_SERIALIZER(FChaosVDConnectivityEdge)
 
 /** Structure contained data from a Clustered Particle.
  * Used to be able to show the values in the editor and allow changes via the Property Editor.
@@ -626,21 +543,7 @@ struct FChaosVDParticleCluster : public FChaosVDWrapperDataBase
 	bool bIsChildToParentLocked = false;
 };
 
-inline FArchive& operator<<(FArchive& Ar, FChaosVDParticleCluster& Data)
-{
-	Data.Serialize(Ar);
-	return Ar;
-}
-
-template<>
-struct TStructOpsTypeTraits<FChaosVDParticleCluster> : public TStructOpsTypeTraitsBase2<FChaosVDParticleCluster>
-{
-	enum
-	{
-		WithSerializer = true,
-	};
-};
-
+CVD_IMPLEMENT_SERIALIZER(FChaosVDParticleCluster)
 
 /** Simplified UStruct version of FChaosVDParticleDataWrapper.
  * Used to be able to show the values in the editor and allow changes via the Property Editor.
@@ -706,17 +609,4 @@ struct FChaosVDParticleDataWrapper : public FChaosVDWrapperDataBase
 	CHAOSVDRUNTIME_API bool Serialize(FArchive& Ar);
 };
 
-inline FArchive& operator<<(FArchive& Ar,FChaosVDParticleDataWrapper& Data)
-{
-	Data.Serialize(Ar);
-	return Ar;
-}
-
-template<>
-struct TStructOpsTypeTraits<FChaosVDParticleDataWrapper> : public TStructOpsTypeTraitsBase2<FChaosVDParticleDataWrapper>
-{
-	enum
-	{
-		WithSerializer = true,
-	};
-};
+CVD_IMPLEMENT_SERIALIZER(FChaosVDParticleDataWrapper)
