Make FInstancedStructDescriptorCache use transactionally safe critical sections for AutoRTFM compatability

diff --git a/Engine/Source/Runtime/Experimental/Iris/Core/Private/Iris/Serialization/InstancedStructNetSerializer.cpp b/Engine/Source/Runtime/Experimental/Iris/Core/Private/Iris/Serialization/InstancedStructNetSerializer.cpp
--- a/Engine/Source/Runtime/Experimental/Iris/Core/Private/Iris/Serialization/InstancedStructNetSerializer.cpp
+++ b/Engine/Source/Runtime/Experimental/Iris/Core/Private/Iris/Serialization/InstancedStructNetSerializer.cpp
@@ -654,7 +654,7 @@ bool FInstancedStructDescriptorCache::IsSupportedType(const UScriptStruct* Struc
 
 TRefCountPtr<const FReplicationStateDescriptor> FInstancedStructDescriptorCache::FindDescriptor(FName StructPath)
 {
-	FScopeLock Lock(&Mutex);
+	UE::TScopeLock Lock(Mutex);
 	if (MaxCachedDescriptorCount > 0)
 	{
 		return DescriptorLruCache.FindAndTouchRef(StructPath);
@@ -735,7 +735,7 @@ TRefCountPtr<const FReplicationStateDescriptor> FInstancedStructDescriptorCache:
 	Descriptor = FReplicationStateDescriptorBuilder::CreateDescriptorForStruct(Struct, Params);
 
 	{
-		FScopeLock Lock(&Mutex);
+		UE::TScopeLock Lock(Mutex);
 		if (MaxCachedDescriptorCount > 0)
 		{
 			DescriptorLruCache.Add(StructPath, Descriptor);
diff --git a/Engine/Source/Runtime/Experimental/Iris/Core/Public/Iris/Serialization/InstancedStructNetSerializer.h b/Engine/Source/Runtime/Experimental/Iris/Core/Public/Iris/Serialization/InstancedStructNetSerializer.h
--- a/Engine/Source/Runtime/Experimental/Iris/Core/Public/Iris/Serialization/InstancedStructNetSerializer.h
+++ b/Engine/Source/Runtime/Experimental/Iris/Core/Public/Iris/Serialization/InstancedStructNetSerializer.h
@@ -6,6 +6,7 @@
 #include "Iris/Serialization/NetSerializer.h"
 #include "Containers/LruCache.h"
 #include "UObject/ObjectMacros.h"
+#include "Misc/TransactionallySafeCriticalSection.h"
 
 #include "InstancedStructNetSerializer.generated.h"
 
@@ -46,7 +47,7 @@ public:
 private:
 	TRefCountPtr<const FReplicationStateDescriptor> CreateAndCacheDescriptor(const UScriptStruct* Struct, FName StructPath);
 
-	FCriticalSection Mutex;
+	FTransactionallySafeCriticalSection Mutex;
 	// LRU cache for descriptors for limited descriptor counts.
 	TLruCache<FName, TRefCountPtr<const FReplicationStateDescriptor>> DescriptorLruCache;
 	// Map struct name -> FReplicationStateDescriptor for unlimited descriptor counts. 
