Dataflow : fix poitential crash in subgraph input node

diff --git a/Engine/Source/Runtime/Experimental/Dataflow/Engine/Private/Dataflow/DataflowSubGraphNodes.cpp b/Engine/Source/Runtime/Experimental/Dataflow/Engine/Private/Dataflow/DataflowSubGraphNodes.cpp
--- a/Engine/Source/Runtime/Experimental/Dataflow/Engine/Private/Dataflow/DataflowSubGraphNodes.cpp
+++ b/Engine/Source/Runtime/Experimental/Dataflow/Engine/Private/Dataflow/DataflowSubGraphNodes.cpp
@@ -431,12 +431,11 @@ void FDataflowSubGraphInputNode::Evaluate(UE::Dataflow::FContext& Context, const
 {
 	using namespace UE::Dataflow::SubGraphNodes::Private;
 
-	check(Out);
 	if (FContextSubGraph* SubGraphContext = Context.AsType<FContextSubGraph>())
 	{
 		SubGraphContext->EvaluateParentContext(*Out);
 	}
-	else
+	else if (Out)
 	{
 		Out->SetNullValue(Context);
 	}
