Dataflow : Fix issue with EdNode not having the right Dataflow node Guid right after deserialization
- Added OnNodeInvalided handler

diff --git a/Engine/Source/Runtime/Experimental/Dataflow/Engine/Private/Dataflow/DataflowObject.cpp b/Engine/Source/Runtime/Experimental/Dataflow/Engine/Private/Dataflow/DataflowObject.cpp
--- a/Engine/Source/Runtime/Experimental/Dataflow/Engine/Private/Dataflow/DataflowObject.cpp
+++ b/Engine/Source/Runtime/Experimental/Dataflow/Engine/Private/Dataflow/DataflowObject.cpp
@@ -16,6 +16,7 @@
 
 FDataflowAssetDelegates::FOnVariablesChanged FDataflowAssetDelegates::OnVariablesChanged;
 FDataflowAssetDelegates::FOnSubGraphsChanged FDataflowAssetDelegates::OnSubGraphsChanged;
+FDataflowAssetDelegates::FOnNodeInvalidated FDataflowAssetDelegates::OnNodeInvalidated;
 
 namespace UE::Dataflow::CVars
 {
@@ -183,6 +184,7 @@ void UDataflow::PostLoad()
 			if (UDataflowEdNode* DataflowEdNode = Cast<UDataflowEdNode>(&EdNode))
 			{
 				DataflowEdNode->SetDataflowGraph(Dataflow);
+				DataflowEdNode->SetDataflowNodeGuid(DataflowEdNode->DataflowNodeGuid);
 				DataflowEdNode->UpdatePinsFromDataflowNode();
 			}
 
diff --git a/Engine/Source/Runtime/Experimental/Dataflow/Engine/Public/Dataflow/DataflowObject.h b/Engine/Source/Runtime/Experimental/Dataflow/Engine/Public/Dataflow/DataflowObject.h
--- a/Engine/Source/Runtime/Experimental/Dataflow/Engine/Public/Dataflow/DataflowObject.h
+++ b/Engine/Source/Runtime/Experimental/Dataflow/Engine/Public/Dataflow/DataflowObject.h
@@ -21,6 +21,7 @@
 #include "DataflowObject.generated.h"
 
 class FArchive;
+struct FDataflowNode;
 class UDataflow;
 class UObject;
 class UDataflowEdNode;
@@ -48,6 +49,9 @@ struct FDataflowAssetDelegates
 	/** Called when sub-graphs are edited ( add, remove ) */
 	DECLARE_MULTICAST_DELEGATE_ThreeParams(FOnSubGraphsChanged, const UDataflow*, const FGuid&, UE::Dataflow::ESubGraphChangedReason);
 	static DATAFLOWENGINE_API FOnSubGraphsChanged OnSubGraphsChanged;
+
+	DECLARE_MULTICAST_DELEGATE_TwoParams(FOnNodeInvalidated, UDataflow&, FDataflowNode&);
+	static DATAFLOWENGINE_API FOnNodeInvalidated OnNodeInvalidated;
 };
 
 /**
