It is more logical and neater to have the poller interface work with waitables as a way to get to the underlying handle.

#rnx

diff --git a/Engine/Source/Runtime/Experimental/IoStore/HttpClient/Private/Loop.inl b/Engine/Source/Runtime/Experimental/IoStore/HttpClient/Private/Loop.inl
--- a/Engine/Source/Runtime/Experimental/IoStore/HttpClient/Private/Loop.inl
+++ b/Engine/Source/Runtime/Experimental/IoStore/HttpClient/Private/Loop.inl
@@ -1007,8 +1007,11 @@ void FPeerGroup::TickSend(FTickState& State, FHost& Host, FPoller& Poller)
 		FOutcome Outcome = FOutcome::None();
 
 		FSocket Socket;
-		if (Socket.Create(Poller))
+		if (Socket.Create())
 		{
+			FWaitable Waitable = Socket.GetWaitable();
+			Poller.Register(Waitable);
+
 			Outcome = Host.Connect(Socket);
 		}
 		else
diff --git a/Engine/Source/Runtime/Experimental/IoStore/HttpClient/Private/Socket.inl b/Engine/Source/Runtime/Experimental/IoStore/HttpClient/Private/Socket.inl
--- a/Engine/Source/Runtime/Experimental/IoStore/HttpClient/Private/Socket.inl
+++ b/Engine/Source/Runtime/Experimental/IoStore/HttpClient/Private/Socket.inl
@@ -123,25 +123,59 @@ uint16 Socket_HtoNs(uint16 v) { return ByteSwap(v); }
 namespace UE::IoStore::HTTP
 {
 
+// {{{1 wait ...................................................................
+
+////////////////////////////////////////////////////////////////////////////////
+class FWaitable
+{
+public:
+				FWaitable() = default;
+				FWaitable(FWaitable&&) = default;
+				FWaitable& operator = (FWaitable&&) = default;
+
+private:
+	friend		class FSocket;
+	friend		class FWaiter;
+	friend		class FPoller;
+				FWaitable(SocketType InSocket);
+	SocketType	Socket = InvalidSocket;
+
+private:
+				FWaitable(const FWaitable&) = delete;
+				FWaitable& operator = (const FWaitable&) = delete;
+};
+
+////////////////////////////////////////////////////////////////////////////////
+FWaitable::FWaitable(SocketType InSocket)
+: Socket(InSocket)
+{
+}
+
+
+
 // {{{1 poll ...................................................................
 
 ////////////////////////////////////////////////////////////////////////////////
-struct FPollFd
 #if defined(IAS_HTTP_USE_POLL)
-	: public pollfd
-{
+	using FPollFdBase = pollfd;
 #else
-{
-	SocketType			fd;
-	int16				events;
-	int16				revents;
-	static const int16	POLLIN  = 1 << 0;
-	static const int16	POLLOUT = 1 << 1;
-	static const int16	POLLERR = 1 << 2;
-	static const int16	POLLHUP = POLLERR;
-	static const int16	POLLNVAL= POLLERR;
+	struct FPollFdBase
+	{
+		SocketType			fd;
+		int16				events;
+		int16				revents;
+		static const int16	POLLIN  = 1 << 0;
+		static const int16	POLLOUT = 1 << 1;
+		static const int16	POLLERR = 1 << 2;
+		static const int16	POLLHUP = POLLERR;
+		static const int16	POLLNVAL= POLLERR;
+	};
 #endif
 
+////////////////////////////////////////////////////////////////////////////////
+struct FPollFd
+	: public FPollFdBase
+{
 	static const int16	PollIn	= int16(POLLIN);
 	static const int16	PollOut	= int16(POLLOUT);
 	static const int16	PollErr	= int16(POLLERR);
@@ -172,11 +206,11 @@ private:
 class FPoller
 {
 public:
-				FPoller()					{ Handle = Poll_CreateContext(); }
-				FPoller(FPoller&& Rhs)		{ Swap(Handle, Rhs.Handle); }
-				~FPoller()					{ Poll_DestroyContext(Handle); }
-	FPoller&	operator = (FPoller&& Rhs)	{ Swap(Handle, Rhs.Handle); return *this; }
-	void		Register(SocketType Socket)	{ Poll_Register(Handle, Socket); }
+				FPoller()							{ Handle = Poll_CreateContext(); }
+				FPoller(FPoller&& Rhs)				{ Swap(Handle, Rhs.Handle); }
+				~FPoller()							{ Poll_DestroyContext(Handle); }
+	FPoller&	operator = (FPoller&& Rhs)			{ Swap(Handle, Rhs.Handle); return *this; }
+	bool		Register(const FWaitable& Waitable)	{ return Poll_Register(Handle, Waitable.Socket); }
 	int32		Wait(FPollFd* Polls, int32 Num, int32 TimeOutMs);
 
 private:
@@ -240,34 +274,7 @@ int32 FPoller::Wait(FPollFd* Polls, int32 Num, int32 TimeOutMs)
 
 
 
-// {{{1 wait ...................................................................
-
-////////////////////////////////////////////////////////////////////////////////
-class FWaitable
-{
-public:
-				FWaitable() = default;
-				FWaitable(FWaitable&&) = default;
-				FWaitable& operator = (FWaitable&&) = default;
-
-private:
-	friend		class FSocket;
-	friend		class FWaiter;
-				FWaitable(SocketType InSocket);
-	SocketType	Socket = InvalidSocket;
-
-private:
-				FWaitable(const FWaitable&) = delete;
-				FWaitable& operator = (const FWaitable&) = delete;
-};
-
-////////////////////////////////////////////////////////////////////////////////
-FWaitable::FWaitable(SocketType InSocket)
-: Socket(InSocket)
-{
-}
-
-
+// {{{1 waiter .................................................................
 
 ////////////////////////////////////////////////////////////////////////////////
 class FWaiter
@@ -366,7 +373,6 @@ public:
 	bool		IsValid() const				{ return Socket != InvalidSocket; }
 	FWaitable	GetWaitable() const			{ return { Socket }; }
 	bool		Create();
-	bool		Create(FPoller& Poller);
 	void		Destroy();
 	FOutcome	Connect(uint32 Ip, uint32 Port);
 	void		Disconnect();
@@ -409,18 +415,6 @@ bool FSocket::Create()
 	return true;
 }
 
-////////////////////////////////////////////////////////////////////////////////
-bool FSocket::Create(FPoller& Poller)
-{
-	if (!Create())
-	{
-		return false;
-	}
-
-	Poller.Register(Socket);
-	return true;
-}
-
 ////////////////////////////////////////////////////////////////////////////////
 void FSocket::Destroy()
 {
