Ditch the redundant peer typedef

#rnx

diff --git a/Engine/Source/Runtime/Experimental/IoStore/HttpClient/Private/Loop.inl b/Engine/Source/Runtime/Experimental/IoStore/HttpClient/Private/Loop.inl
--- a/Engine/Source/Runtime/Experimental/IoStore/HttpClient/Private/Loop.inl
+++ b/Engine/Source/Runtime/Experimental/IoStore/HttpClient/Private/Loop.inl
@@ -7,10 +7,8 @@ namespace UE::IoStore::HTTP
 
 // {{{1 event-loop-int .........................................................
 
-using FPeerType = FHttpPeer;
-
 ////////////////////////////////////////////////////////////////////////////////
-static FOutcome DoSend(FActivity* Activity, FPeerType& Peer)
+static FOutcome DoSend(FActivity* Activity, FHttpPeer& Peer)
 {
 	Trace(Activity, ETrace::StateChange, Activity->State);
 
@@ -60,7 +58,7 @@ static FOutcome DoSend(FActivity* Activity, FPeerType& Peer)
 }
 
 ////////////////////////////////////////////////////////////////////////////////
-static FOutcome DoRecvPeer(FActivity* Activity, FPeerType& Peer, int32& MaxRecvSize, int32 Size)
+static FOutcome DoRecvPeer(FActivity* Activity, FHttpPeer& Peer, int32& MaxRecvSize, int32 Size)
 {
 	Size = FMath::Min(Size, MaxRecvSize);
 	check(Size >= 0);
@@ -99,7 +97,7 @@ static FOutcome DoRecvPeer(FActivity* Activity, FPeerType& Peer, int32& MaxRecvS
 }
 
 ////////////////////////////////////////////////////////////////////////////////
-static FOutcome DoRecvMessage(FActivity* Activity, FPeerType& Peer)
+static FOutcome DoRecvMessage(FActivity* Activity, FHttpPeer& Peer)
 {
 	TRACE_CPUPROFILER_EVENT_SCOPE(IasHttp::DoRecvMessage);
 
@@ -315,7 +313,7 @@ static FOutcome DoRecvMessage(FActivity* Activity, FPeerType& Peer)
 }
 
 ////////////////////////////////////////////////////////////////////////////////
-static FOutcome DoRecvContent(FActivity* Activity, FPeerType& Peer, int32& MaxRecvSize)
+static FOutcome DoRecvContent(FActivity* Activity, FHttpPeer& Peer, int32& MaxRecvSize)
 {
 	TRACE_CPUPROFILER_EVENT_SCOPE(IasHttp::DoRecvContent);
 
@@ -344,7 +342,7 @@ static FOutcome DoRecvContent(FActivity* Activity, FPeerType& Peer, int32& MaxRe
 }
 
 ////////////////////////////////////////////////////////////////////////////////
-static FOutcome DoRecvStream(FActivity* Activity, FPeerType& Peer, int32& MaxRecvSize)
+static FOutcome DoRecvStream(FActivity* Activity, FHttpPeer& Peer, int32& MaxRecvSize)
 {
 	auto RaiseCrLfError = [Activity] ()
 	{
@@ -545,7 +543,7 @@ static FOutcome DoRecvStream(FActivity* Activity, FPeerType& Peer, int32& MaxRec
 }
 
 ////////////////////////////////////////////////////////////////////////////////
-static FOutcome DoRecv(FActivity* Activity, FPeerType& Peer, int32& MaxRecvSize)
+static FOutcome DoRecv(FActivity* Activity, FHttpPeer& Peer, int32& MaxRecvSize)
 {
 	using EState = FActivity::EState;
 
@@ -749,7 +747,7 @@ private:
 	void				SendInternal(FTickState& State);
 	FActivity*			Send = nullptr;
 	FActivity*			Recv = nullptr;
-	FPeerType			Peer;
+	FHttpPeer			Peer;
 	uint32				LastUseMs = 0;
 	uint8				IsKeepAlive = 0;
 	bool				bNegotiating = false;
@@ -806,7 +804,7 @@ void FPeerGroup::Fail(FTickState& State, const char* Reason)
 		Activity = Next;
 	}
 
-	Peer = FPeerType();
+	Peer = FHttpPeer();
 	Send = Recv = nullptr;
 	bWaiting = false;
 	IsKeepAlive = 0;
@@ -996,7 +994,7 @@ void FPeerGroup::TickSend(FTickState& State, FHost& Host, FPoller& Poller)
 	if (Peer.IsValid() && LastUseMs + GIdleMs < State.NowMs)
 	{
 		LastUseMs = State.NowMs;
-		Peer = FPeerType();
+		Peer = FHttpPeer();
 	}
 
 	// We don't have a connected socket on first use, or if a keep-alive:close
@@ -1034,7 +1032,7 @@ void FPeerGroup::TickSend(FTickState& State, FHost& Host, FPoller& Poller)
 
 		FCertRootsRef VerifyCert = Host.GetVerifyCert();
 		const char* HostName = Host.GetHostName().GetData();
-		Peer = FPeerType(MoveTemp(Socket), VerifyCert, HostName);
+		Peer = FHttpPeer(MoveTemp(Socket), VerifyCert, HostName);
 	}
 
 	Send = Pending;
