[Newton] Disable Physics Replication LOD and enable Resimulation for all replicated physics objects.

#changelist validated
#virtualized

diff --git a/Engine/Source/Runtime/Experimental/Chaos/Private/RewindData.cpp b/Engine/Source/Runtime/Experimental/Chaos/Private/RewindData.cpp
--- a/Engine/Source/Runtime/Experimental/Chaos/Private/RewindData.cpp
+++ b/Engine/Source/Runtime/Experimental/Chaos/Private/RewindData.cpp
@@ -381,7 +381,7 @@ bool FRewindData::CheckQuaternionThreshold(FQuat A, FQuat B, float ThresholdDegr
 }
 
 
-CHAOS_API bool bResimAllowRewindToResimulatedFrames = false;
+CHAOS_API bool bResimAllowRewindToResimulatedFrames = true;
 FAutoConsoleVariableRef CVarResimAllowRewindToResimulatedFrames(TEXT("p.Resim.AllowRewindToResimulatedFrames"), bResimAllowRewindToResimulatedFrames, TEXT("Allow rewinding back to a frame that was previously part of a resimulation. If a resimulation is performed between frame 100-110, allow a new resim from 105-115 if needed, else next resim will be able to start from frame 111."));
 
 bool FRewindData::RewindToFrame(int32 Frame)
