Chaos Cloth Asset - Added non const methods to access the DataflowInstance and set the overrides.

#rnx

diff --git a/Engine/Plugins/ChaosClothAsset/Source/ChaosClothAssetEngine/Private/ChaosClothAsset/ClothAssetBase.cpp b/Engine/Plugins/ChaosClothAsset/Source/ChaosClothAssetEngine/Private/ChaosClothAsset/ClothAssetBase.cpp
--- a/Engine/Plugins/ChaosClothAsset/Source/ChaosClothAssetEngine/Private/ChaosClothAsset/ClothAssetBase.cpp
+++ b/Engine/Plugins/ChaosClothAsset/Source/ChaosClothAssetEngine/Private/ChaosClothAsset/ClothAssetBase.cpp
@@ -145,6 +145,13 @@ PRAGMA_DISABLE_DEPRECATION_WARNINGS
 PRAGMA_ENABLE_DEPRECATION_WARNINGS
 }
 
+FDataflowInstance& UChaosClothAssetBase::GetDataflowInstance()
+{
+PRAGMA_DISABLE_DEPRECATION_WARNINGS
+	return DataflowInstance;
+PRAGMA_ENABLE_DEPRECATION_WARNINGS
+}
+
 void UChaosClothAssetBase::BeginDestroy()
 {
 	check(IsInGameThread());
diff --git a/Engine/Plugins/ChaosClothAsset/Source/ChaosClothAssetEngine/Public/ChaosClothAsset/ClothAssetBase.h b/Engine/Plugins/ChaosClothAsset/Source/ChaosClothAssetEngine/Public/ChaosClothAsset/ClothAssetBase.h
--- a/Engine/Plugins/ChaosClothAsset/Source/ChaosClothAssetEngine/Public/ChaosClothAsset/ClothAssetBase.h
+++ b/Engine/Plugins/ChaosClothAsset/Source/ChaosClothAssetEngine/Public/ChaosClothAsset/ClothAssetBase.h
@@ -68,9 +68,11 @@ public:
 	//~ End IDataflowContentOwner interface 
 
 	//~ Begin IDataflowInstanceInterface interface
-	virtual const FDataflowInstance& GetDataflowInstance() const;
+	virtual const FDataflowInstance& GetDataflowInstance() const override;
+	virtual FDataflowInstance& GetDataflowInstance() override;
 	//~ End IDataflowInstanceInterface interface
 
+
 	//~ Begin UObject interface
 	virtual void BeginDestroy() override;
 	virtual bool IsReadyForFinishDestroy() override;
diff --git a/Engine/Source/Runtime/Experimental/Dataflow/Engine/Private/Dataflow/DataflowInstance.cpp b/Engine/Source/Runtime/Experimental/Dataflow/Engine/Private/Dataflow/DataflowInstance.cpp
--- a/Engine/Source/Runtime/Experimental/Dataflow/Engine/Private/Dataflow/DataflowInstance.cpp
+++ b/Engine/Source/Runtime/Experimental/Dataflow/Engine/Private/Dataflow/DataflowInstance.cpp
@@ -189,7 +189,7 @@ const FInstancedPropertyBag& FDataflowVariableOverrides::GetVariables() const
 	return Variables;
 }
 
-const FInstancedPropertyBag& FDataflowVariableOverrides::GetVariables()
+FInstancedPropertyBag& FDataflowVariableOverrides::GetVariables()
 {
 	return Variables;
 }
@@ -343,17 +343,17 @@ const FInstancedPropertyBag& FDataflowInstance::GetVariables() const
 	return VariableOverrides.GetVariables();
 }
 
-const FInstancedPropertyBag& FDataflowInstance::GetVariables()
+FInstancedPropertyBag& FDataflowInstance::GetVariables()
 {
 	return VariableOverrides.GetVariables();
 }
 
-FDataflowVariableOverrides& FDataflowInstance::GetVariableOverrides()
+const FDataflowVariableOverrides& FDataflowInstance::GetVariableOverrides() const
 {
 	return VariableOverrides;
 }
 
-const FDataflowVariableOverrides& FDataflowInstance::GetVariableOverrides() const
+FDataflowVariableOverrides& FDataflowInstance::GetVariableOverrides()
 {
 	return VariableOverrides;
 }
diff --git a/Engine/Source/Runtime/Experimental/Dataflow/Engine/Public/Dataflow/DataflowInstance.h b/Engine/Source/Runtime/Experimental/Dataflow/Engine/Public/Dataflow/DataflowInstance.h
--- a/Engine/Source/Runtime/Experimental/Dataflow/Engine/Public/Dataflow/DataflowInstance.h
+++ b/Engine/Source/Runtime/Experimental/Dataflow/Engine/Public/Dataflow/DataflowInstance.h
@@ -41,7 +41,7 @@ struct FDataflowVariableOverrides
 
 	DATAFLOWENGINE_API const FInstancedPropertyBag& GetVariables() const;
 
-	DATAFLOWENGINE_API const FInstancedPropertyBag& GetVariables();
+	DATAFLOWENGINE_API FInstancedPropertyBag& GetVariables();
 
 	DATAFLOWENGINE_API static FName GetVariablePropertyName();
 
@@ -100,7 +100,7 @@ public:
 
 	DATAFLOWENGINE_API const FInstancedPropertyBag& GetVariables() const;
 
-	DATAFLOWENGINE_API const FInstancedPropertyBag& GetVariables();
+	DATAFLOWENGINE_API FInstancedPropertyBag& GetVariables();
 
 	/** Get the dataflow terminal member property name */
 	DATAFLOWENGINE_API static FName GetDataflowTerminalPropertyName();
@@ -110,6 +110,8 @@ public:
 
 	DATAFLOWENGINE_API const FDataflowVariableOverrides& GetVariableOverrides() const;
 
+	DATAFLOWENGINE_API FDataflowVariableOverrides& GetVariableOverrides();
+
 	DATAFLOWENGINE_API void SyncVariables();
 
 	/** 
@@ -137,7 +139,6 @@ private:
 
 private:
 	friend class UE::Dataflow::Private::FVariablesOverridesDetails;
-	DATAFLOWENGINE_API FDataflowVariableOverrides& GetVariableOverrides();
 
 #if WITH_EDITOR
 	void OnOwnerPostEditChangeProperty(struct FPropertyChangedEvent& PropertyChangedEvent);
@@ -162,6 +163,7 @@ class IDataflowInstanceInterface
 
 public:
 	virtual const FDataflowInstance& GetDataflowInstance() const = 0;
+	virtual FDataflowInstance& GetDataflowInstance() = 0;
 };
 
 namespace UE::Dataflow::InstanceUtils
diff --git a/Engine/Source/Runtime/Experimental/GeometryCollectionEngine/Public/GeometryCollection/GeometryCollectionObject.h b/Engine/Source/Runtime/Experimental/GeometryCollectionEngine/Public/GeometryCollection/GeometryCollectionObject.h
--- a/Engine/Source/Runtime/Experimental/GeometryCollectionEngine/Public/GeometryCollection/GeometryCollectionObject.h
+++ b/Engine/Source/Runtime/Experimental/GeometryCollectionEngine/Public/GeometryCollection/GeometryCollectionObject.h
@@ -863,7 +863,8 @@ public:
 	UDataflow* GetDataflowAsset() const;
 
 	//~ Begin IDataflowGeneratedAssetInterface
-	virtual const FDataflowInstance& GetDataflowInstance() const { return DataflowInstance; };
+	virtual const FDataflowInstance& GetDataflowInstance() const override { return DataflowInstance; };
+	virtual FDataflowInstance& GetDataflowInstance() override { return DataflowInstance; };
 	//~ End IDataflowGeneratedAssetInterface
 
 	//~ Begin IDataflowContentOwner interface
