Chaosflesh: Supported multiple transform sources for kinematic updates using multiple SKM sources.
-Fixed bone matching bug when RefSkeleton and Skeleton of SKM are not in the same order.

diff --git a/Engine/Plugins/Experimental/ChaosFlesh/Source/ChaosFleshEngine/Private/ChaosFlesh/ChaosDeformableTetrahedralComponent.cpp b/Engine/Plugins/Experimental/ChaosFlesh/Source/ChaosFleshEngine/Private/ChaosFlesh/ChaosDeformableTetrahedralComponent.cpp
--- a/Engine/Plugins/Experimental/ChaosFlesh/Source/ChaosFleshEngine/Private/ChaosFlesh/ChaosDeformableTetrahedralComponent.cpp
+++ b/Engine/Plugins/Experimental/ChaosFlesh/Source/ChaosFleshEngine/Private/ChaosFlesh/ChaosDeformableTetrahedralComponent.cpp
@@ -178,7 +178,7 @@ UDeformablePhysicsComponent::FDataMapValue UDeformableTetrahedralComponent::NewD
 							{
 								if (const USkeleton* Skeleton = SkeletalMesh->GetSkeleton())
 								{
-									TSet<int32> Roots = TransformSource.GetTransformSource(Skeleton->GetName(), Skeleton->GetGuid().ToString());
+									TSet<int32> Roots = TransformSource.GetTransformSource(Skeleton->GetName(), Skeleton->GetGuid().ToString(), SkeletalMesh->GetName());
 									if (!Roots.IsEmpty() && ensureMsgf(Roots.Num() == 1, TEXT("Error: Only supports a single root per skeleton.(%s)"), *Skeleton->GetName()))
 									{
 										TArray<FTransform> ComponentLocalPose;
@@ -187,7 +187,7 @@ UDeformablePhysicsComponent::FDataMapValue UDeformableTetrahedralComponent::NewD
 										const TArray<FTransform>& ComponentTransforms = SkeletalMeshComponent->GetComponentSpaceTransforms();
 										if (ComponentLocalPose.Num() == ComponentTransforms.Num())
 										{
-											for (int Adx = Roots.Array()[0], Cdx = 0; Adx <= AnimationTransforms.Num() && Cdx < ComponentTransforms.Num(); Adx++, Cdx++)
+											for (int Adx = Roots.Array()[0], Cdx = 0; Adx < AnimationTransforms.Num() && Cdx < ComponentTransforms.Num(); Adx++, Cdx++)
 											{
 												// @todo(flesh) : Can we just use one array?
 												AnimationTransforms[Adx] = ComponentTransforms[Cdx];
@@ -775,7 +775,7 @@ bool UDeformableTetrahedralComponent::IsSkeletalMeshAnimationCompatible(const US
 					{
 						if (const USkeleton* Skeleton = SkeletalMesh->GetSkeleton())
 						{
-							TSet<int32> Roots = TransformSource.GetTransformSource(Skeleton->GetName(), Skeleton->GetGuid().ToString());
+							TSet<int32> Roots = TransformSource.GetTransformSource(Skeleton->GetName(), Skeleton->GetGuid().ToString(), SkeletalMesh->GetName());
 							if (!Roots.IsEmpty() && ensureMsgf(Roots.Num() == 1, TEXT("Error: Only supports a single root per skeleton.(%s)"), *Skeleton->GetName()))
 							{
 								return true;
diff --git a/Engine/Plugins/Experimental/Dataflow/Source/DataflowEditor/Private/Dataflow/DataflowSimulationControls.cpp b/Engine/Plugins/Experimental/Dataflow/Source/DataflowEditor/Private/Dataflow/DataflowSimulationControls.cpp
--- a/Engine/Plugins/Experimental/Dataflow/Source/DataflowEditor/Private/Dataflow/DataflowSimulationControls.cpp
+++ b/Engine/Plugins/Experimental/Dataflow/Source/DataflowEditor/Private/Dataflow/DataflowSimulationControls.cpp
@@ -158,7 +158,8 @@ namespace UE::Dataflow
 			BoneIndices.SetNumUninitialized(NumBones);
 			for (int32 Index = 0; Index < NumBones; ++Index)
 			{
-				BoneIndices[Index] = StaticCast<FBoneIndexType>(Index);
+				int32 SkeletonBoneIndex = InSkeletalMesh->GetSkeleton()->GetReferenceSkeleton().FindBoneIndex(ReferenceSkeleton->GetBoneName(Index));
+				BoneIndices[Index] = StaticCast<FBoneIndexType>(SkeletonBoneIndex);
 			}
 
 			FBoneContainer BoneContainer;
diff --git a/Engine/Source/Runtime/Experimental/Chaos/Private/GeometryCollection/Facades/CollectionTransformSourceFacade.cpp b/Engine/Source/Runtime/Experimental/Chaos/Private/GeometryCollection/Facades/CollectionTransformSourceFacade.cpp
--- a/Engine/Source/Runtime/Experimental/Chaos/Private/GeometryCollection/Facades/CollectionTransformSourceFacade.cpp
+++ b/Engine/Source/Runtime/Experimental/Chaos/Private/GeometryCollection/Facades/CollectionTransformSourceFacade.cpp
@@ -17,17 +17,20 @@ namespace GeometryCollection::Facades
 	const FName FTransformSource::SourceNameAttributeName = "Name";
 	const FName FTransformSource::SourceGuidAttributeName = "GuidID";
 	const FName FTransformSource::SourceRootsAttributeName = "Roots";
+	const FName FTransformSource::SourceSkeletalMeshNameAttributeName = "SkeletalMeshName";
 
 	FTransformSource::FTransformSource(FManagedArrayCollection& InCollection)
 		: SourceNameAttribute(InCollection, SourceNameAttributeName, TransformSourceGroupName)
 		, SourceGuidAttribute(InCollection, SourceGuidAttributeName, TransformSourceGroupName)
 		, SourceRootsAttribute(InCollection, SourceRootsAttributeName, TransformSourceGroupName, FTransformCollection::TransformGroup)
+		, SourceSkeletalMeshNameAttribute(InCollection, SourceSkeletalMeshNameAttributeName, TransformSourceGroupName)
 	{}
 
 	FTransformSource::FTransformSource(const FManagedArrayCollection& InCollection)
 		: SourceNameAttribute(InCollection, SourceNameAttributeName, TransformSourceGroupName)
 		, SourceGuidAttribute(InCollection, SourceGuidAttributeName, TransformSourceGroupName)
 		, SourceRootsAttribute(InCollection, SourceRootsAttributeName, TransformSourceGroupName, FTransformCollection::TransformGroup)
+		, SourceSkeletalMeshNameAttribute(InCollection, SourceSkeletalMeshNameAttributeName, TransformSourceGroupName)
 	{}
 	//
 	//  Initialization
@@ -39,6 +42,7 @@ namespace GeometryCollection::Facades
 		SourceNameAttribute.Add();
 		SourceGuidAttribute.Add();
 		SourceRootsAttribute.Add();
+		SourceSkeletalMeshNameAttribute.Add();
 	}
 
 	bool FTransformSource::IsValid() const
@@ -49,7 +53,7 @@ namespace GeometryCollection::Facades
 	//
 	//  Add Data
 	//
-	void FTransformSource::AddTransformSource(const FString& InName, const FString& InGuid, const TSet<int32>& InRoots)
+	void FTransformSource::AddTransformSource(const FString& InName, const FString& InGuid, const TSet<int32>& InRoots, const FString& SKMName)
 	{
 		check(!IsConst());
 		DefineSchema();
@@ -58,12 +62,13 @@ namespace GeometryCollection::Facades
 		SourceNameAttribute.Modify()[Idx] = InName;
 		SourceGuidAttribute.Modify()[Idx] = InGuid;
 		SourceRootsAttribute.Modify()[Idx] = InRoots;
+		SourceSkeletalMeshNameAttribute.Modify()[Idx] = SKMName;
 	}
 
 	//
 	//  Get Data
 	//
-	TSet<int32> FTransformSource::GetTransformSource(const FString& InName, const FString& InGuid) const
+	TSet<int32> FTransformSource::GetTransformSource(const FString& InName, const FString& InGuid, const FString& SKMName) const
 	{
 		if (IsValid())
 		{
@@ -76,7 +81,10 @@ namespace GeometryCollection::Facades
 			{
 				if (Guids[i] == InGuid && Names[i].Equals(InName))
 				{
-					return Roots[i];
+					if (!SourceSkeletalMeshNameAttribute.IsValid() || SourceSkeletalMeshNameAttribute[i].IsEmpty() || SourceSkeletalMeshNameAttribute[i].Equals(SKMName))
+					{
+						return Roots[i];
+					}
 				}
 			}
 		}
diff --git a/Engine/Source/Runtime/Experimental/Chaos/Public/GeometryCollection/Facades/CollectionTransformSourceFacade.h b/Engine/Source/Runtime/Experimental/Chaos/Public/GeometryCollection/Facades/CollectionTransformSourceFacade.h
--- a/Engine/Source/Runtime/Experimental/Chaos/Public/GeometryCollection/Facades/CollectionTransformSourceFacade.h
+++ b/Engine/Source/Runtime/Experimental/Chaos/Public/GeometryCollection/Facades/CollectionTransformSourceFacade.h
@@ -40,7 +40,7 @@ namespace GeometryCollection::Facades
 		static CHAOS_API const FName SourceNameAttributeName;
 		static CHAOS_API const FName SourceGuidAttributeName;
 		static CHAOS_API const FName SourceRootsAttributeName;
-
+		static CHAOS_API const FName SourceSkeletalMeshNameAttributeName;
 		/**
 		* FSelectionFacade Constuctor
 		* @param VertixDependencyGroup : GroupName the index attribute is dependent on. 
@@ -62,21 +62,22 @@ namespace GeometryCollection::Facades
 		* @param Name : Name of the owner of the transform set.
 		* @param Guid : Guid of the owner of the transform set.
 		* @param Roots : Root indices of the transform set.
+		* @param SKMName : Name of the skeletal mesh owner of the transform set.
 		*/
-		CHAOS_API void AddTransformSource(const FString& Name, const FString& Guid, const TSet<int32>& Roots);
+		CHAOS_API void AddTransformSource(const FString& Name, const FString& Guid, const TSet<int32>& Roots, const FString& SKMName = FString());
 
 		/**
 		* Query for root indices.  
 		* @param Name : Name of the owner of the transform set.  
 		* @param Guid : Guid of the owner of the transform set.
 		*/
-		CHAOS_API TSet<int32> GetTransformSource(const FString& Name, const FString& Guid) const;
+		CHAOS_API TSet<int32> GetTransformSource(const FString& Name, const FString& Guid, const FString& SKMName = FString()) const;
 
 	private:
 		TManagedArrayAccessor<FString> SourceNameAttribute;
 		TManagedArrayAccessor<FString> SourceGuidAttribute;
 		TManagedArrayAccessor< TSet<int32 > > SourceRootsAttribute;
-		
+		TManagedArrayAccessor<FString> SourceSkeletalMeshNameAttribute;
 	};
 
 }
diff --git a/Engine/Source/Runtime/Experimental/GeometryCollectionEngine/Private/GeometryCollection/GeometryCollectionEngineConversion.cpp b/Engine/Source/Runtime/Experimental/GeometryCollectionEngine/Private/GeometryCollection/GeometryCollectionEngineConversion.cpp
--- a/Engine/Source/Runtime/Experimental/GeometryCollectionEngine/Private/GeometryCollection/GeometryCollectionEngineConversion.cpp
+++ b/Engine/Source/Runtime/Experimental/GeometryCollectionEngine/Private/GeometryCollection/GeometryCollectionEngineConversion.cpp
@@ -1202,7 +1202,7 @@ bool FGeometryCollectionEngineConversion::AppendSkeletalMesh(const USkeletalMesh
 		GeometryCollection::Facades::FTransformSource TransformSourceFacade(*InManagedArrayCollection);
 		TSet<int32> Roots;
 		Roots.Add(RootIndex);
-		TransformSourceFacade.AddTransformSource(Skeleton->GetName(), Skeleton->GetGuid().ToString(), Roots);
+		TransformSourceFacade.AddTransformSource(Skeleton->GetName(), Skeleton->GetGuid().ToString(), Roots, InSkeletalMesh->GetName());
 	}
 
 	if (bImportTransformOnly)
@@ -1276,7 +1276,7 @@ bool FGeometryCollectionEngineConversion::AppendSkeletalMesh(const USkeletalMesh
 		if (RootIndex != INDEX_NONE)
 		{
 			Children[RootIndex].Add(ComponentTransformIndex);
-			BaseName = FString::Printf(TEXT("%s_SplitMesh"), *BoneName[RootIndex]);
+			BaseName = FString::Printf(TEXT("%s_SplitMesh"), *InSkeletalMesh->GetName());
 		}
 
 		BoneName[ComponentTransformIndex] = MakeUnique(BaseName, SplitMeshIndex);
@@ -1472,7 +1472,7 @@ void FGeometryCollectionEngineConversion::AppendSkeleton(const USkeleton* InSkel
 		}
 
 		ensure(Roots.Num());
-		TransformSourceFacade.AddTransformSource(InSkeleton->GetName(), InSkeleton->GetGuid().ToString(), Roots);
+		TransformSourceFacade.AddTransformSource(InSkeleton->GetName(), InSkeleton->GetGuid().ToString(), Roots, FString() /*SKMName*/);
 	}
 }
 
@@ -1811,7 +1811,10 @@ void FGeometryCollectionEngineConversion::ConvertCollectionToSkeleton(const FMan
 				FMeshBoneInfo Info(BoneName, BoneName.ToString(), SkeletionParentIndex);
 				{
 					FReferenceSkeletonModifier Edit(OutSkeleton);
-					Edit.Add(Info, Transform, true /*bAllowMultipleRoots*/);
+					if (Edit.FindBoneIndex(BoneName) == INDEX_NONE) //Bone does not exist
+					{
+						Edit.Add(Info, Transform, true /*bAllowMultipleRoots*/);
+					}
 				}
 				AddMapping(CurrentIndex, OutSkeleton->GetReferenceSkeleton().GetRawBoneNum());
 
