Dataflow : fix double call to input PullValue when forwarding an input to an output
- the call to PullValue is already done in FowardInput method ofthe output

diff --git a/Engine/Source/Runtime/Experimental/Dataflow/Core/Private/Dataflow/DataflowNode.cpp b/Engine/Source/Runtime/Experimental/Dataflow/Core/Private/Dataflow/DataflowNode.cpp
--- a/Engine/Source/Runtime/Experimental/Dataflow/Core/Private/Dataflow/DataflowNode.cpp
+++ b/Engine/Source/Runtime/Experimental/Dataflow/Core/Private/Dataflow/DataflowNode.cpp
@@ -1401,7 +1401,6 @@ void FDataflowNode::ForwardInput(UE::Dataflow::FContext& Context, const UE::Data
 		if (const FDataflowInput* Input = FindInput(InputReference))
 		{
 			// we need to pull the value first so the upstream of the graph evaluate 
-			Input->PullValue(Context);
 			Output->ForwardInput(Input, Context);
 		}
 		else
