Dataflow : fix timestamp issue for forwarded cached entries
- Timestamp was using the refernced cache timestamp instead of using the node timestamp
- This would cause situation where the forwarding node could be re-evaluated multiple time if its timestamp was in thefuture compare to the node upstream it was refrencing thedata from

diff --git a/Engine/Source/Runtime/Experimental/Dataflow/Core/Private/Dataflow/DataflowInputOutput.cpp b/Engine/Source/Runtime/Experimental/Dataflow/Core/Private/Dataflow/DataflowInputOutput.cpp
--- a/Engine/Source/Runtime/Experimental/Dataflow/Core/Private/Dataflow/DataflowInputOutput.cpp
+++ b/Engine/Source/Runtime/Experimental/Dataflow/Core/Private/Dataflow/DataflowInputOutput.cpp
@@ -375,7 +375,7 @@ void FDataflowOutput::ForwardInput(const FDataflowInput* Input, UE::Dataflow::FC
 			if (const FDataflowOutput* ConnectionOut = Input->GetConnection())
 			{
 				Input->PullValue(Context);
-				Context.SetDataReference(CacheKey(), Property, ConnectionOut->CacheKey());
+				Context.SetDataReference(CacheKey(), Property, ConnectionOut->CacheKey(), GetOwningNodeTimestamp());
 			}
 		}
 		else
diff --git a/Engine/Source/Runtime/Experimental/Dataflow/Core/Public/Dataflow/DataflowNodeParameters.h b/Engine/Source/Runtime/Experimental/Dataflow/Core/Public/Dataflow/DataflowNodeParameters.h
--- a/Engine/Source/Runtime/Experimental/Dataflow/Core/Public/Dataflow/DataflowNodeParameters.h
+++ b/Engine/Source/Runtime/Experimental/Dataflow/Core/Public/Dataflow/DataflowNodeParameters.h
@@ -417,12 +417,13 @@ namespace UE::Dataflow
 			SetDataImpl(InKey, MoveTemp(DataStoreEntry));
 		}
 
-		void SetDataReference(FContextCacheKey Key, const FProperty* Property, FContextCacheKey ReferenceKey)
+		void SetDataReference(FContextCacheKey Key, const FProperty* Property, FContextCacheKey ReferenceKey, const FTimestamp& InTimestamp)
 		{
 			// find the reference key to get 
 			if (const TUniquePtr<FContextCacheElementBase>* CacheElement = GetDataImpl(ReferenceKey))
 			{
 				TUniquePtr<FContextCacheElementBase> CacheReferenceElement = (*CacheElement)->CreateReference(ReferenceKey);
+				CacheReferenceElement->SetTimestamp(InTimestamp);
 
 				SetDataImpl(Key, MoveTemp(CacheReferenceElement));
 			}
