Chaosflesh: Various fixes
-stop loading chaos cache every tick
-added safety check in DeleteFleshVertices
-added out-of-range check in Transform Selection to Face/Vertex Selection

diff --git a/Engine/Plugins/Experimental/ChaosFlesh/Source/ChaosFleshNodes/Private/Dataflow/ChaosFleshAppendTetrahedralCollectionNode.cpp b/Engine/Plugins/Experimental/ChaosFlesh/Source/ChaosFleshNodes/Private/Dataflow/ChaosFleshAppendTetrahedralCollectionNode.cpp
--- a/Engine/Plugins/Experimental/ChaosFlesh/Source/ChaosFleshNodes/Private/Dataflow/ChaosFleshAppendTetrahedralCollectionNode.cpp
+++ b/Engine/Plugins/Experimental/ChaosFlesh/Source/ChaosFleshNodes/Private/Dataflow/ChaosFleshAppendTetrahedralCollectionNode.cpp
@@ -73,11 +73,22 @@ void FDeleteFleshVerticesDataflowNode::Evaluate(UE::Dataflow::FContext& Context,
 	if (Out->IsA<FManagedArrayCollection>(&Collection))
 	{
 		TUniquePtr<FFleshCollection> InCollection(GetValue<FManagedArrayCollection>(Context, &Collection).NewCopy<FFleshCollection>());
-		if (IsConnected(&VertexSelection))
+		if (IsConnected(&Collection) && IsConnected(&VertexSelection))
 		{
 			const FDataflowVertexSelection& InDataflowVertexSelection = GetValue<FDataflowVertexSelection>(Context, &VertexSelection);
-			TArray<int32> VertexSelected = InDataflowVertexSelection.AsArray();
-			InCollection->RemoveVertices(VertexSelected);
+			if (InDataflowVertexSelection.Num() == InCollection->NumElements(FGeometryCollection::VerticesGroup))
+			{
+				TArray<int32> VertexSelected = InDataflowVertexSelection.AsArray();
+				InCollection->RemoveVertices(VertexSelected);
+			}
+			else
+			{
+				Context.Warning(FString::Printf(
+					TEXT("DeleteFleshVertices Node: VertexSelection has different size (%d) than the number of vertices (%d) in the Collection."),
+					InDataflowVertexSelection.Num(),
+					InCollection->NumElements(FGeometryCollection::VerticesGroup)),
+					this, Out);
+			}
 		}
 		SetValue<const FManagedArrayCollection&>(Context, MoveTemp(*InCollection), &Collection);
 	}
diff --git a/Engine/Plugins/Experimental/ChaosFlesh/Source/ChaosFleshNodes/Public/Dataflow/ChaosFleshAppendTetrahedralCollectionNode.h b/Engine/Plugins/Experimental/ChaosFlesh/Source/ChaosFleshNodes/Public/Dataflow/ChaosFleshAppendTetrahedralCollectionNode.h
--- a/Engine/Plugins/Experimental/ChaosFlesh/Source/ChaosFleshNodes/Public/Dataflow/ChaosFleshAppendTetrahedralCollectionNode.h
+++ b/Engine/Plugins/Experimental/ChaosFlesh/Source/ChaosFleshNodes/Public/Dataflow/ChaosFleshAppendTetrahedralCollectionNode.h
@@ -70,10 +70,10 @@ struct FDeleteFleshVerticesDataflowNode : public FDataflowNode
 	DATAFLOW_NODE_DEFINE_INTERNAL(FDeleteFleshVerticesDataflowNode, "DeleteFleshVertices", "Flesh", "")
 	DATAFLOW_NODE_RENDER_TYPE("SurfaceRender", FGeometryCollection::StaticType(), "Collection")
 public:
-	UPROPERTY(meta = (DataflowInput, DataflowOutput, DataflowPassthrough = "Collection"));
+	UPROPERTY(meta = (DataflowInput, DataflowIntrinsic, DataflowOutput, DataflowPassthrough = "Collection"));
 	FManagedArrayCollection Collection;
 
-	UPROPERTY(meta = (DataflowInput))
+	UPROPERTY(meta = (DataflowInput, DataflowIntrinsic))
 	FDataflowVertexSelection VertexSelection;
 
 	FDeleteFleshVerticesDataflowNode(const UE::Dataflow::FNodeParameters& InParam, FGuid InGuid = FGuid::NewGuid())
diff --git a/Engine/Plugins/Experimental/Dataflow/Source/DataflowEditor/Private/Dataflow/DataflowSimulationScene.cpp b/Engine/Plugins/Experimental/Dataflow/Source/DataflowEditor/Private/Dataflow/DataflowSimulationScene.cpp
--- a/Engine/Plugins/Experimental/Dataflow/Source/DataflowEditor/Private/Dataflow/DataflowSimulationScene.cpp
+++ b/Engine/Plugins/Experimental/Dataflow/Source/DataflowEditor/Private/Dataflow/DataflowSimulationScene.cpp
@@ -247,7 +247,7 @@ void FDataflowSimulationScene::UpdateSimulationCache(float DeltaSeconds)
 		{
 			// Update the cached simulation at some point in time
 			const float ScrollBarTime = SimulationTime - TimeRange[0];
-			if (!SimulationGenerator->IsSimulating() || ScrollBarTime != CacheManager->StartTime)
+			if (ScrollBarTime != CacheManager->StartTime)
 			{
 				CacheManager->SetStartTime(ScrollBarTime);
 			}
diff --git a/Engine/Source/Runtime/Experimental/Chaos/Private/GeometryCollection/Facades/CollectionTransformSelectionFacade.cpp b/Engine/Source/Runtime/Experimental/Chaos/Private/GeometryCollection/Facades/CollectionTransformSelectionFacade.cpp
--- a/Engine/Source/Runtime/Experimental/Chaos/Private/GeometryCollection/Facades/CollectionTransformSelectionFacade.cpp
+++ b/Engine/Source/Runtime/Experimental/Chaos/Private/GeometryCollection/Facades/CollectionTransformSelectionFacade.cpp
@@ -1236,12 +1236,19 @@ namespace GeometryCollection::Facades
 
 			for (int32 TransformIdx : InTransformSelection)
 			{
-				const int32 FaceIdxStart = FaceStarts[TransformToGeometryIndices[TransformIdx]];
-
-				for (int32 Offset = 0; Offset < FaceCounts[TransformToGeometryIndices[TransformIdx]]; ++Offset)
+				if (TransformToGeometryIndices.IsValidIndex(TransformIdx))
 				{
-					const int32 FaceIdx = FaceIdxStart + Offset;
-					OutSelection.Add(FaceIdx);
+					const int32 GeometryIdx = TransformToGeometryIndices[TransformIdx];
+					if (FaceStarts.IsValidIndex(GeometryIdx))
+					{
+						const int32 FaceIdxStart = FaceStarts[GeometryIdx];
+
+						for (int32 Offset = 0; Offset < FaceCounts[GeometryIdx]; ++Offset)
+						{
+							const int32 FaceIdx = FaceIdxStart + Offset;
+							OutSelection.Add(FaceIdx);
+						}
+					}
 				}
 			}
 		}
@@ -1280,12 +1287,19 @@ namespace GeometryCollection::Facades
 
 			for (int32 TransformIdx : InTransformSelection)
 			{
-				const int32 VertexIndexStart = VertexStarts[TransformToGeometryIndices[TransformIdx]];
-
-				for (int32 Offset = 0; Offset < VertexCounts[TransformToGeometryIndices[TransformIdx]]; ++Offset)
+				if (TransformToGeometryIndices.IsValidIndex(TransformIdx))
 				{
-					const int32 VertexIdx = VertexIndexStart + Offset;
-					OutSelection.Add(VertexIdx);
+					const int32 GeometryIdx = TransformToGeometryIndices[TransformIdx];
+					if (VertexStarts.IsValidIndex(GeometryIdx))
+					{
+						const int32 VertexIndexStart = VertexStarts[TransformToGeometryIndices[TransformIdx]];
+
+						for (int32 Offset = 0; Offset < VertexCounts[TransformToGeometryIndices[TransformIdx]]; ++Offset)
+						{
+							const int32 VertexIdx = VertexIndexStart + Offset;
+							OutSelection.Add(VertexIdx);
+						}
+					}
 				}
 			}
 		}
