Chaosflesh: Added SetCollisionArea dataflow node to specify collision-enabled vertices. Updated deformable solver to parse collision vertex attribute.

diff --git a/Engine/Plugins/Experimental/ChaosFlesh/Source/ChaosFleshNodes/Private/Dataflow/ChaosFleshNodesPlugin.cpp b/Engine/Plugins/Experimental/ChaosFlesh/Source/ChaosFleshNodes/Private/Dataflow/ChaosFleshNodesPlugin.cpp
--- a/Engine/Plugins/Experimental/ChaosFlesh/Source/ChaosFleshNodes/Private/Dataflow/ChaosFleshNodesPlugin.cpp
+++ b/Engine/Plugins/Experimental/ChaosFlesh/Source/ChaosFleshNodes/Private/Dataflow/ChaosFleshNodesPlugin.cpp
@@ -64,6 +64,7 @@ void IChaosFleshNodesPlugin::StartupModule()
 	DATAFLOW_NODE_REGISTER_CREATION_FACTORY(FKinematicSkeletonConstraintDataflowNode);
 	DATAFLOW_NODE_REGISTER_CREATION_FACTORY(FKinematicSkeletalMeshInitializationDataflowNode);
 	DATAFLOW_NODE_REGISTER_CREATION_FACTORY(FRadialTetrahedronDataflowNodes);
+	DATAFLOW_NODE_REGISTER_CREATION_FACTORY(FSetCollidableVerticesDataflowNode);
 	DATAFLOW_NODE_REGISTER_CREATION_FACTORY(FSetFleshBonePositionTargetBindingDataflowNode);
 	DATAFLOW_NODE_REGISTER_CREATION_FACTORY(FSetFleshDefaultPropertiesNode);
 	DATAFLOW_NODE_REGISTER_CREATION_FACTORY(FSetMuscleActivationParameterNode);
diff --git a/Engine/Plugins/Experimental/ChaosFlesh/Source/ChaosFleshNodes/Private/Dataflow/ChaosFleshSetVertexTrianglePositionTargetBindingNode.cpp b/Engine/Plugins/Experimental/ChaosFlesh/Source/ChaosFleshNodes/Private/Dataflow/ChaosFleshSetVertexTrianglePositionTargetBindingNode.cpp
--- a/Engine/Plugins/Experimental/ChaosFlesh/Source/ChaosFleshNodes/Private/Dataflow/ChaosFleshSetVertexTrianglePositionTargetBindingNode.cpp
+++ b/Engine/Plugins/Experimental/ChaosFlesh/Source/ChaosFleshNodes/Private/Dataflow/ChaosFleshSetVertexTrianglePositionTargetBindingNode.cpp
@@ -7,6 +7,7 @@
 #include "Chaos/TriangleCollisionPoint.h"
 #include "GeometryCollection/Facades/CollectionMeshFacade.h"
 #include "GeometryCollection/Facades/CollectionPositionTargetFacade.h"
+#include "GeometryCollection/Facades/CollectionCollisionFacade.h"
 #include "ChaosFlesh/TetrahedralCollection.h"
 #include "ChaosFlesh/ChaosFleshCollectionFacade.h"
 
@@ -161,3 +162,18 @@ void FSetVertexTrianglePositionTargetBindingDataflowNode::Evaluate(UE::Dataflow:
 		SetValue(Context, MoveTemp(InCollection), &Collection);
 	}
 }
+
+void FSetCollidableVerticesDataflowNode::Evaluate(UE::Dataflow::FContext& Context, const FDataflowOutput* Out) const
+{
+	if (Out->IsA<DataType>(&Collection))
+	{
+		DataType InCollection = GetValue<DataType>(Context, &Collection);
+		if (IsConnected(&VertexSelection))
+		{
+			GeometryCollection::Facades::FCollisionFacade CollisionFacade(InCollection);
+			const FDataflowVertexSelection& InDataflowVertexSelection = GetValue(Context, &VertexSelection);
+			CollisionFacade.SetCollisionEnabled(InDataflowVertexSelection.AsArray());
+		}
+		SetValue(Context, MoveTemp(InCollection), &Collection);
+	}
+}
\ No newline at end of file
diff --git a/Engine/Plugins/Experimental/ChaosFlesh/Source/ChaosFleshNodes/Public/Dataflow/ChaosFleshSetVertexTrianglePositionTargetBindingNode.h b/Engine/Plugins/Experimental/ChaosFlesh/Source/ChaosFleshNodes/Public/Dataflow/ChaosFleshSetVertexTrianglePositionTargetBindingNode.h
--- a/Engine/Plugins/Experimental/ChaosFlesh/Source/ChaosFleshNodes/Public/Dataflow/ChaosFleshSetVertexTrianglePositionTargetBindingNode.h
+++ b/Engine/Plugins/Experimental/ChaosFlesh/Source/ChaosFleshNodes/Public/Dataflow/ChaosFleshSetVertexTrianglePositionTargetBindingNode.h
@@ -16,26 +16,25 @@ USTRUCT(meta = (DataflowFlesh))
 struct FSetVertexTrianglePositionTargetBindingDataflowNode : public FDataflowNode
 {
 	GENERATED_USTRUCT_BODY()
-		DATAFLOW_NODE_DEFINE_INTERNAL(FSetVertexTrianglePositionTargetBindingDataflowNode, "SetVertexTrianglePositionTargetBinding", "Flesh", "")
-		DATAFLOW_NODE_RENDER_TYPE("SurfaceRender",FGeometryCollection::StaticType(),  "Collection")
+	DATAFLOW_NODE_DEFINE_INTERNAL(FSetVertexTrianglePositionTargetBindingDataflowNode, "SetVertexTrianglePositionTargetBinding", "Flesh", "")
+	DATAFLOW_NODE_RENDER_TYPE("SurfaceRender",FGeometryCollection::StaticType(),  "Collection")
 
 public:
 	typedef FManagedArrayCollection DataType;
 
-
 	UPROPERTY(meta = (DataflowInput, DataflowOutput, DisplayName = "Collection", DataflowPassthrough = "Collection"))
-		FManagedArrayCollection Collection;
+	FManagedArrayCollection Collection;
 
 	UPROPERTY(EditAnywhere, Category = "Dataflow")
-		float PositionTargetStiffness = 10000.f;
+	float PositionTargetStiffness = 10000.f;
 
 	/* (optional) only create weak constraints from vertices in VertexSelection to triangles in other geometries */
 	UPROPERTY(meta = (DataflowInput, DisplayName = "(Optional) VertexSelection"))
-		FDataflowVertexSelection VertexSelection;
+	FDataflowVertexSelection VertexSelection;
 	
 	/* weak constraint search radius = Largest side of the collection's bounding box * VertexRadiusRatio */
 	UPROPERTY(EditAnywhere, Category = "Dataflow", meta = (ClampMin = "0.0", ClampMax = "2.0", UIMin = "0.0", UIMax = "2.0"))
-		float VertexRadiusRatio = .001f;
+	float VertexRadiusRatio = .001f;
 
 	FSetVertexTrianglePositionTargetBindingDataflowNode(const UE::Dataflow::FNodeParameters& InParam, FGuid InGuid = FGuid::NewGuid())
 		: FDataflowNode(InParam, InGuid)
@@ -47,3 +46,36 @@ public:
 
 	virtual void Evaluate(UE::Dataflow::FContext& Context, const FDataflowOutput* Out) const override;
 };
+
+/**
+ * Set custom vertices so that only these vertices can collide with other surfaces. 
+ * Unselected vertices will not collide with unselected vertices.
+ */
+USTRUCT(meta = (DataflowFlesh))
+struct FSetCollidableVerticesDataflowNode : public FDataflowNode
+{
+	GENERATED_USTRUCT_BODY()
+	DATAFLOW_NODE_DEFINE_INTERNAL(FSetCollidableVerticesDataflowNode, "SetCollidableVertices", "Flesh", "")
+	DATAFLOW_NODE_RENDER_TYPE("SurfaceRender", FGeometryCollection::StaticType(), "Collection")
+
+public:
+	FSetCollidableVerticesDataflowNode(const UE::Dataflow::FNodeParameters& InParam, FGuid InGuid = FGuid::NewGuid())
+		: FDataflowNode(InParam, InGuid)
+	{
+		RegisterInputConnection(&Collection);
+		RegisterInputConnection(&VertexSelection);
+		RegisterOutputConnection(&Collection, &Collection);
+	}
+
+private:
+	typedef FManagedArrayCollection DataType;
+
+	UPROPERTY(meta = (DataflowInput, DataflowIntrinsic, DataflowOutput, DisplayName = "Collection", DataflowPassthrough = "Collection"))
+	FManagedArrayCollection Collection;
+
+	/* Vertices selected to be able to collide with others. Unselected vertices will not collide with unselected vertices*/
+	UPROPERTY(meta = (DataflowInput, DataflowIntrinsic, DisplayName = "Collision Vertex Selection"))
+	FDataflowVertexSelection VertexSelection;
+
+	virtual void Evaluate(UE::Dataflow::FContext& Context, const FDataflowOutput* Out) const override;
+};
\ No newline at end of file
diff --git a/Engine/Source/Runtime/Experimental/Chaos/Private/Chaos/Deformable/ChaosDeformableSolver.cpp b/Engine/Source/Runtime/Experimental/Chaos/Private/Chaos/Deformable/ChaosDeformableSolver.cpp
--- a/Engine/Source/Runtime/Experimental/Chaos/Private/Chaos/Deformable/ChaosDeformableSolver.cpp
+++ b/Engine/Source/Runtime/Experimental/Chaos/Private/Chaos/Deformable/ChaosDeformableSolver.cpp
@@ -35,6 +35,7 @@
 #include "GeometryCollection/Facades/CollectionKinematicBindingFacade.h"
 #include "GeometryCollection/Facades/CollectionVertexBoneWeightsFacade.h"
 #include "GeometryCollection/Facades/CollectionPositionTargetFacade.h"
+#include "GeometryCollection/Facades/CollectionCollisionFacade.h"
 #include "GeometryCollection/Facades/CollectionConstraintOverrideFacade.h"
 #include "GeometryCollection/Facades/CollectionMeshFacade.h"
 #include "GeometryCollection/Facades/CollectionMuscleActivationFacade.h"
@@ -164,7 +165,7 @@ namespace Chaos::Softs
 		{
 			ParticleComponentIndex.Reset(new TArray<int32>());
 			SurfaceTriangleMesh.Reset(new Chaos::FTriangleMesh());
-			SurfaceVertices.Reset(new TArray<int32>);
+			SurfaceCollisionVertices.Reset(new TArray<int32>);
 		}
 		if (Property.bUseGridBasedConstraints)
 		{
@@ -1451,16 +1452,34 @@ namespace Chaos::Softs
 						{
 							int32 SurfaceOffset = SurfaceElements->Num();
 							SurfaceElements->SetNum(SurfaceOffset + Indices->Num());
+							TArray<int32> SurfaceVertices;
 							for (int32 i = 0; i < Indices->Num(); i++)
 							{
 								for (int32 j = 0; j < 3; j++)
 								{
 									(*SurfaceElements)[i + SurfaceOffset][j] = VertexOffset + (*Indices)[i][j];
+									SurfaceVertices.Add(VertexOffset + (*Indices)[i][j]);
 								}
 							}
+							TSet<int32> UniqueSurfaceVertices(SurfaceVertices);
+							GeometryCollection::Facades::FCollisionFacade CollisionFacade(*Rest);
+							if (CollisionFacade.IsValid()) //Collision vertices are specified
+							{
+								for (int32 SurfaceVertexIdx: UniqueSurfaceVertices)
+								{
+									if (CollisionFacade.IsCollisionEnabled(SurfaceVertexIdx - VertexOffset))
+									{
+										SurfaceCollisionVertices->Add(SurfaceVertexIdx);
+									}
+								}
+							}
+							else
+							{
+								SurfaceCollisionVertices->Append(UniqueSurfaceVertices.Array());
+							}
 						}
-						VertexOffset += Vertex->Num();
 
+						VertexOffset += Vertex->Num();
 						
 						if (!Property.bDoInComponentSpringCollision || Property.bDoSphereRepulsion) //Component-Component collisions or sphere repulsion
 						{
@@ -1491,7 +1510,6 @@ namespace Chaos::Softs
 		}
 
 		SurfaceTriangleMesh->Init(*SurfaceElements);
-		SurfaceTriangleMesh->GetVertexSetAsArray(*SurfaceVertices);
 
 		VertexOffset = 0;
 
@@ -1567,11 +1585,11 @@ namespace Chaos::Softs
 			int32 NRadius = Property.NRingExcluded;
 			SurfaceTriangleMesh->GetPointToNeighborsMap();
 			SurfaceTriangleMesh->GetPointToTriangleMap(); //Initialize Maps before parallel
-			ParticleTriangleExclusionMap.Reserve(SurfaceVertices->Num());
-			PhysicsParallelFor(SurfaceVertices->Num(),
+			ParticleTriangleExclusionMap.Reserve(SurfaceCollisionVertices->Num());
+			PhysicsParallelFor(SurfaceCollisionVertices->Num(),
 				[this, &NRadius](int32 i)
 				{
-					int32 VertexId = (*SurfaceVertices)[i];
+					int32 VertexId = (*SurfaceCollisionVertices)[i];
 					TSet<int32>& TriangleSet = ParticleTriangleExclusionMap.FindOrAdd(VertexId);
 					TSet<int32> NRing;
 					if (NRadius > 1)
@@ -1835,11 +1853,11 @@ namespace Chaos::Softs
 					this->TriangleMeshCollisions->InitFlesh(InParticles, Property.SpringCollisionSearchRadius, this->Property.bCollideWithFullMesh);
 					if (Property.bDoInComponentSpringCollision)
 					{
-						this->GSWeakConstraints->CollisionDetectionSpatialHashInComponent(this->Evolution->Particles(), *SurfaceVertices, *SurfaceTriangleMesh, ParticleTriangleExclusionMap, TriangleMeshCollisions->GetDynamicSpatialHash(), Property.SpringCollisionSearchRadius, Property.SpringCollisionStiffness, Property.bAllowSliding);
+						this->GSWeakConstraints->CollisionDetectionSpatialHashInComponent(this->Evolution->Particles(), *SurfaceCollisionVertices, *SurfaceTriangleMesh, ParticleTriangleExclusionMap, TriangleMeshCollisions->GetDynamicSpatialHash(), Property.SpringCollisionSearchRadius, Property.SpringCollisionStiffness, Property.bAllowSliding);
 					}
 					else
 					{
-						this->GSWeakConstraints->CollisionDetectionSpatialHash(this->Evolution->Particles(), *SurfaceVertices, *SurfaceTriangleMesh, *ParticleComponentIndex, TriangleMeshCollisions->GetDynamicSpatialHash(), Property.SpringCollisionSearchRadius, Property.SpringCollisionStiffness, Property.bAllowSliding);
+						this->GSWeakConstraints->CollisionDetectionSpatialHash(this->Evolution->Particles(), *SurfaceCollisionVertices, *SurfaceTriangleMesh, *ParticleComponentIndex, TriangleMeshCollisions->GetDynamicSpatialHash(), Property.SpringCollisionSearchRadius, Property.SpringCollisionStiffness, Property.bAllowSliding);
 					}
 					TArray<TArray<int32>> WCCollisionConstraints, WCCollisionIncidentElements, WCCollisionIncidentElementsLocal;
 					this->GSWeakConstraints->ComputeCollisionWCDataSimplified(WCCollisionConstraints, WCCollisionIncidentElements, WCCollisionIncidentElementsLocal);
@@ -1855,7 +1873,7 @@ namespace Chaos::Softs
 			Evolution->ConstraintInits()[InitIndex] 
 				= [this](FSolverParticles& InParticles, const FSolverReal Dt)
 				{
-					this->GSSphereRepulsionConstraints->UpdateSphereRepulsionConstraints(this->Evolution->Particles(), *SurfaceVertices, *ParticleComponentIndex);
+					this->GSSphereRepulsionConstraints->UpdateSphereRepulsionConstraints(this->Evolution->Particles(), *SurfaceCollisionVertices, *ParticleComponentIndex);
 					TArray<TArray<int32>> SphereRepulsionConstraints, SphereRepulsionIncidentElements, SphereRepulsionIncidentElementsLocal;
 					this->GSSphereRepulsionConstraints->ReturnSphereRepulsionConstraints(SphereRepulsionConstraints, SphereRepulsionIncidentElements, SphereRepulsionIncidentElementsLocal);
 					this->GSMainConstraint->AddTransientConstraints(SphereRepulsionConstraints, SphereRepulsionIncidentElements, SphereRepulsionIncidentElementsLocal);
diff --git a/Engine/Source/Runtime/Experimental/Chaos/Private/GeometryCollection/Facades/CollectionCollisionFacade.cpp b/Engine/Source/Runtime/Experimental/Chaos/Private/GeometryCollection/Facades/CollectionCollisionFacade.cpp
new file mode 100644
--- /dev/null
+++ b/Engine/Source/Runtime/Experimental/Chaos/Private/GeometryCollection/Facades/CollectionCollisionFacade.cpp
@@ -0,0 +1,66 @@
+// Copyright Epic Games, Inc. All Rights Reserved.
+
+/*=============================================================================
+	GeometryCollection.cpp: FGeometryCollection methods.
+=============================================================================*/
+
+#include "GeometryCollection/Facades/CollectionCollisionFacade.h"
+#include "GeometryCollection/GeometryCollection.h"
+
+namespace GeometryCollection::Facades
+{
+	// Attributes
+	const FName FCollisionFacade::IsCollisionEnabledAttributeName = "IsCollisionEnabled";
+
+	FCollisionFacade::FCollisionFacade(FManagedArrayCollection& InCollection)
+		: ConstCollection(InCollection)
+		, Collection(&InCollection)
+		, IsCollisionEnabledAttribute(InCollection, IsCollisionEnabledAttributeName, FGeometryCollection::VerticesGroup)
+	{
+		DefineSchema();
+	}
+
+	FCollisionFacade::FCollisionFacade(const FManagedArrayCollection& InCollection)
+		: ConstCollection(InCollection)
+		, Collection(nullptr)
+		, IsCollisionEnabledAttribute(InCollection, IsCollisionEnabledAttributeName, FGeometryCollection::VerticesGroup)
+	{
+	}
+
+	/** Define the facade */
+
+	void FCollisionFacade::DefineSchema()
+	{
+		check(!IsConst());
+		IsCollisionEnabledAttribute.AddAndFill(false);
+	}
+
+	bool FCollisionFacade::IsValid() const
+	{
+		return IsCollisionEnabledAttribute.IsValid();
+	}
+
+	void FCollisionFacade::SetCollisionEnabled(const TArray<int32>& VertexIndices)
+	{
+		check(!IsConst());
+		TManagedArray<bool>& IsCollisionEnabledArray = IsCollisionEnabledAttribute.Modify();
+		for (int32 VertexIdx : VertexIndices)
+		{
+			if (IsCollisionEnabledAttribute.IsValidIndex(VertexIdx))
+			{
+				IsCollisionEnabledArray[VertexIdx] = true;
+			}
+		}
+	}
+
+	bool FCollisionFacade::IsCollisionEnabled(int32 VertexIndex) const
+	{
+		if (IsCollisionEnabledAttribute.IsValid() && IsCollisionEnabledAttribute.IsValidIndex(VertexIndex))
+		{
+			return IsCollisionEnabledAttribute.Get()[VertexIndex];
+		}
+		return false;
+	}
+};
+
+
diff --git a/Engine/Source/Runtime/Experimental/Chaos/Public/Chaos/Deformable/ChaosDeformableSolver.h b/Engine/Source/Runtime/Experimental/Chaos/Public/Chaos/Deformable/ChaosDeformableSolver.h
--- a/Engine/Source/Runtime/Experimental/Chaos/Public/Chaos/Deformable/ChaosDeformableSolver.h
+++ b/Engine/Source/Runtime/Experimental/Chaos/Public/Chaos/Deformable/ChaosDeformableSolver.h
@@ -220,7 +220,7 @@ namespace Chaos::Softs
 		TUniquePtr <TArray<TVec3<int32>>> TetmeshSurfaceElements;
 		TUniquePtr <TArray<Chaos::TVec4<int32>>> AllElements;
 		TUniquePtr <FTriangleMesh> SurfaceTriangleMesh;
-		TUniquePtr <TArray<int32>> SurfaceVertices;
+		TUniquePtr <TArray<int32>> SurfaceCollisionVertices;
 		TUniquePtr <TArray<TArray<int32>>> AllIncidentElements;
 		TUniquePtr <TArray<TArray<int32>>> AllIncidentElementsLocal;
 		TUniquePtr <TArray<FSolverReal>> AllTetEMeshArray;
diff --git a/Engine/Source/Runtime/Experimental/Chaos/Public/GeometryCollection/Facades/CollectionCollisionFacade.h b/Engine/Source/Runtime/Experimental/Chaos/Public/GeometryCollection/Facades/CollectionCollisionFacade.h
new file mode 100644
--- /dev/null
+++ b/Engine/Source/Runtime/Experimental/Chaos/Public/GeometryCollection/Facades/CollectionCollisionFacade.h
@@ -0,0 +1,51 @@
+// Copyright Epic Games, Inc. All Rights Reserved.
+
+#pragma once
+#include "Containers/UnrealString.h"
+#include "Containers/Array.h"
+#include "GeometryCollection/ManagedArrayCollection.h"
+#include "GeometryCollection/ManagedArrayAccessor.h"
+
+namespace GeometryCollection::Facades
+{
+
+	/**
+	* FCollisionFacade
+	* 
+	* Defines common API for storing collision information.
+	* 
+	*/
+	class FCollisionFacade
+	{
+	public:
+
+		// Attributes
+		CHAOS_API static const FName IsCollisionEnabledAttributeName;
+
+		/**
+		* FCollisionFacade Constuctor
+		*/
+
+		CHAOS_API FCollisionFacade(FManagedArrayCollection& InCollection);
+
+		CHAOS_API FCollisionFacade(const FManagedArrayCollection& InCollection);
+		
+		/** Define the facade */
+		CHAOS_API void DefineSchema();
+
+		/** Is the Facade const */
+		bool IsConst() const { return Collection == nullptr; }
+
+		CHAOS_API bool IsValid() const;
+
+		CHAOS_API void SetCollisionEnabled(const TArray<int32>& VertexIndices);
+
+		CHAOS_API bool IsCollisionEnabled(int32 VertexIndex) const;
+	private:
+		const FManagedArrayCollection& ConstCollection;
+		FManagedArrayCollection* Collection = nullptr;
+
+		TManagedArrayAccessor<bool> IsCollisionEnabledAttribute;
+	};
+
+}
