Add the concept of resolving to FIASHostGroup

#rnx

- Eventually we will move the code for resolving a distributed endpoint to FIASHostGroup and this change starts to lay the ground work for that.
- We can now create a FIASHostGroup with no urls (unresolved) and later apply the urls (resolving)
- The statistics debug display can also show when a hostgroup is currently resolving rather than being connected or an error state.
-- Note that we don't pass the host group to statistics until it is resolved at the moment so this won't do anything right now.

diff --git a/Engine/Source/Runtime/Experimental/IoStore/OnDemand/Private/IasHostGroup.cpp b/Engine/Source/Runtime/Experimental/IoStore/OnDemand/Private/IasHostGroup.cpp
--- a/Engine/Source/Runtime/Experimental/IoStore/OnDemand/Private/IasHostGroup.cpp
+++ b/Engine/Source/Runtime/Experimental/IoStore/OnDemand/Private/IasHostGroup.cpp
@@ -60,6 +60,10 @@ private:
 struct FIASHostGroup::FImpl
 {
 	FImpl() = default;
+	FImpl(FName InName)
+		: Name(InName)
+	{
+	}
 	FImpl(FName InName, FOnDemandHostGroup& InHostGroup)
 		: HostGroup(InHostGroup)
 		, Name(InName)
@@ -67,6 +71,13 @@ struct FIASHostGroup::FImpl
 		}
 	~FImpl() = default;
 
+	void Reset(FOnDemandHostGroup&& InHostGroup)
+	{
+		HostGroup = MoveTemp(InHostGroup);
+		HttpErrorHistory.Reset(GIasHttpErrorSampleCount);
+		bHttpEnabled = true;
+	}
+
 	FOnDemandHostGroup HostGroup;
 	FBitWindow HttpErrorHistory;
 	FName Name;
@@ -108,6 +119,11 @@ FIASHostGroup::FIASHostGroup()
 {
 }
 
+FIASHostGroup::FIASHostGroup(FName Name)
+	: Impl(MakeShared<FImpl>(Name))
+{
+}
+
 FIASHostGroup::FIASHostGroup(FName Name, FOnDemandHostGroup&& HostGroup)
 	: Impl(MakeShared<FImpl>(Name, HostGroup))
 {
@@ -123,11 +139,55 @@ FName FIASHostGroup::GetName() const
 	return Impl->Name;
 }
 
+bool FIASHostGroup::IsResolved() const
+{
+	return !Impl->HostGroup.IsEmpty();
+}
+
 bool FIASHostGroup::IsConnected() const
 {
 	return Impl->HostGroup.PrimaryHostIndex() != INDEX_NONE;
 }
 
+FIoStatus FIASHostGroup::Resolve(TConstArrayView<FAnsiString> HostUrls)
+{
+	if (IsResolved())
+	{
+		return FIoStatus(EIoErrorCode::InvalidCode, TEXT("Host group is already resolved"));
+	}
+
+	TIoStatusOr<FOnDemandHostGroup> Result = FOnDemandHostGroup::Create(HostUrls);
+
+	if (!Result.IsOk())
+	{
+		return Result.Status();
+	}
+
+	Impl->Reset(Result.ConsumeValueOrDie());
+
+	return FIoStatus(EIoErrorCode::Ok);
+}
+
+FIoStatus FIASHostGroup::Resolve(TConstArrayView<FString> HostUrls)
+{
+	if (IsResolved())
+	{
+		return FIoStatus(EIoErrorCode::InvalidCode, TEXT("Host group is already resolved"));
+	}
+
+	TIoStatusOr<FOnDemandHostGroup> Result = FOnDemandHostGroup::Create(HostUrls);
+
+	if (!Result.IsOk())
+	{
+		return Result.Status();
+	}
+
+	Impl->Reset(Result.ConsumeValueOrDie());
+
+	return FIoStatus(EIoErrorCode::Ok);
+}
+
+
 void FIASHostGroup::Connect(int32 HostIndex)
 {
 	Impl->bHttpEnabled = true;
diff --git a/Engine/Source/Runtime/Experimental/IoStore/OnDemand/Private/IasHostGroup.h b/Engine/Source/Runtime/Experimental/IoStore/OnDemand/Private/IasHostGroup.h
--- a/Engine/Source/Runtime/Experimental/IoStore/OnDemand/Private/IasHostGroup.h
+++ b/Engine/Source/Runtime/Experimental/IoStore/OnDemand/Private/IasHostGroup.h
@@ -21,12 +21,17 @@ public:
 	static TIoStatusOr<FIASHostGroup> Create(FName Name, TConstArrayView<FString> HostUrls);
 
 	FIASHostGroup();
+	FIASHostGroup(FName Name);
 	FIASHostGroup(FName Name, FOnDemandHostGroup&& HostGroup);
 	~FIASHostGroup() = default;
 
 	FName GetName() const;
 
+	bool IsResolved() const;
 	bool IsConnected() const;
+
+	FIoStatus Resolve(TConstArrayView<FAnsiString> HostUrls);
+	FIoStatus Resolve(TConstArrayView<FString> HostUrls);
 	
 	void Connect(int32 HostIndex);
 	void Disconnect();
diff --git a/Engine/Source/Runtime/Experimental/IoStore/OnDemand/Private/OnDemandIoDispatcherBackend.cpp b/Engine/Source/Runtime/Experimental/IoStore/OnDemand/Private/OnDemandIoDispatcherBackend.cpp
--- a/Engine/Source/Runtime/Experimental/IoStore/OnDemand/Private/OnDemandIoDispatcherBackend.cpp
+++ b/Engine/Source/Runtime/Experimental/IoStore/OnDemand/Private/OnDemandIoDispatcherBackend.cpp
@@ -899,6 +899,7 @@ FOnDemandIoBackend::FOnDemandIoBackend(const FOnDemandEndpointConfig& Config, FO
 	: IoStore(InIoStore)
 	, Cache(MoveTemp(InCache))
 	, Stats(BackendStatus)
+	, DefaultHostGroup(FName("Default"))
 	, CurrentThreadPriority(EThreadPriority::TPri_Num)
 {
 	EndpointTestPath = Config.TocPath;
@@ -1398,17 +1399,15 @@ bool FOnDemandIoBackend::CreateDefaultHostGroup()
 		Algo::Rotate(DefaultUrls, GIasHttpPrimaryEndpoint);
 	}
 
-	TIoStatusOr<FIASHostGroup> HostResult = FIASHostGroup::Create(FName("Default"), DefaultUrls);
+	FIoStatus HostResult = DefaultHostGroup.Resolve(DefaultUrls);
 	DefaultUrls.Empty(); // No longer needed
-
+	
 	if (!HostResult.IsOk())
 	{
-		UE_LOG(LogIas, Error, TEXT("HTTP streaming disabled, could not create the default host group (%s)"), *HostResult.Status().ToString());
+		UE_LOG(LogIas, Error, TEXT("HTTP streaming disabled, could not create the default host group (%s)"), *HostResult.ToString());
 		return false;
 	}
 
-	DefaultHostGroup = HostResult.ConsumeValueOrDie();
-
 	Stats.AddHostGroup(DefaultHostGroup);
 
 	return true;
diff --git a/Engine/Source/Runtime/Experimental/IoStore/OnDemand/Private/Statistics.cpp b/Engine/Source/Runtime/Experimental/IoStore/OnDemand/Private/Statistics.cpp
--- a/Engine/Source/Runtime/Experimental/IoStore/OnDemand/Private/Statistics.cpp
+++ b/Engine/Source/Runtime/Experimental/IoStore/OnDemand/Private/Statistics.cpp
@@ -299,6 +299,10 @@ FOnDemandIoBackendStats::FOnDemandIoBackendStats(FBackendStatus& InStatus)
 					Text << HostGroup.GetPrimaryHostUrl();
 					Text << TEXT(" (") << HostGroup.GetPrimaryHostIndex() << TEXT("/") << HostGroup.GetHostUrls().Num() << TEXT(")");
 				}
+				else if (HostGroup.IsResolved())
+				{
+					Text << TEXT("Resolving...");
+				}
 				else
 				{
 					Text << TEXT("Disconnected");
