Chaosflesh: Fixed a bug where values in TSet are modified by reference. This caused TSet.Find(Value) to miss because hash key was created from the old value.

diff --git a/Engine/Source/Runtime/Experimental/Chaos/Private/GeometryCollection/TransformCollection.cpp b/Engine/Source/Runtime/Experimental/Chaos/Private/GeometryCollection/TransformCollection.cpp
--- a/Engine/Source/Runtime/Experimental/Chaos/Private/GeometryCollection/TransformCollection.cpp
+++ b/Engine/Source/Runtime/Experimental/Chaos/Private/GeometryCollection/TransformCollection.cpp
@@ -92,9 +92,18 @@ void FTransformCollection::Append(const FManagedArrayCollection& InCollection)
 		for (int Idx = Offset; Idx < Size; Idx++)
 		{
 			if (Parent[Idx] != INDEX_NONE)
-				Parent[Idx] += Offset;			
-			for (int32& Val : Children[Idx])
+			{
+				Parent[Idx] += Offset;
+			}
+			//Modify the temporary copied array to avoid directly changing values in the TSet
+			//Hash value depends on the value of the elements in the TSet
+			//If only the value is changed, Find() will not be able to find the value
+			TArray<int32> ChildrenArray = Children[Idx].Array();
+			for (int32& Val : ChildrenArray)
+			{
 				Val += Offset;
+			}
+			Children[Idx] = TSet<int32>(ChildrenArray);
 		}
 	}
 	else
