Fix crash when a component of the modular vehicle is destroyed

diff --git a/Engine/Source/Runtime/Experimental/ChaosVehicles/ChaosVehiclesCore/Private/SimModule/SimulationModuleBase.cpp b/Engine/Source/Runtime/Experimental/ChaosVehicles/ChaosVehiclesCore/Private/SimModule/SimulationModuleBase.cpp
--- a/Engine/Source/Runtime/Experimental/ChaosVehicles/ChaosVehiclesCore/Private/SimModule/SimulationModuleBase.cpp
+++ b/Engine/Source/Runtime/Experimental/ChaosVehicles/ChaosVehiclesCore/Private/SimModule/SimulationModuleBase.cpp
@@ -70,10 +70,11 @@ FVehicleBlackboard* ISimulationModuleBase::GetSimBlackboard()
 
 FPBDRigidClusteredParticleHandle* ISimulationModuleBase::GetClusterParticle(FClusterUnionPhysicsProxy* Proxy)
 { 
-	if (ParticleIdx.IsValid() && CachedParticle && (CachedParticle->UniqueIdx() == ParticleIdx))
-	{
-		return CachedParticle;
-	}
+	// TODO: This caching method fails if the component that owns the particle is destroyed
+	//if (ParticleIdx.IsValid() && CachedParticle && (CachedParticle->UniqueIdx() == ParticleIdx))
+	//{
+	//	return CachedParticle;
+	//}
 
 	CachedParticle = nullptr;
 
