Clone cached io buffer data, speculative fix for a read after free happening in FIoChunkEncoding::Decode. Disabled by default

diff --git a/Engine/Source/Runtime/Experimental/IoStore/OnDemand/Private/JournaledCache.cpp b/Engine/Source/Runtime/Experimental/IoStore/OnDemand/Private/JournaledCache.cpp
--- a/Engine/Source/Runtime/Experimental/IoStore/OnDemand/Private/JournaledCache.cpp
+++ b/Engine/Source/Runtime/Experimental/IoStore/OnDemand/Private/JournaledCache.cpp
@@ -54,6 +54,15 @@ TRACE_DECLARE_INT_COUNTER(IasAllowance, TEXT("Ias/CacheAllowance"));
 TRACE_DECLARE_INT_COUNTER(IasOpCount,   TEXT("Ias/CacheOpCount"));
 TRACE_DECLARE_INT_COUNTER(IasReadCursor,TEXT("Ias/CacheReadCursor"));
 
+////////////////////////////////////////////////////////////////////////////////
+int32 CloneCachedIOBufferData = 0;
+static FAutoConsoleVariableRef CloneCachedIOBufferDataCVar(
+	TEXT("ias.CloneCachedIOBufferData"),
+	CloneCachedIOBufferData,
+	TEXT("Clone cached io buffer data, speculative fix for a read after free happening in FIoChunkEncoding::Decode. Disabled by default"),
+	ECVF_Default
+);
+
 ////////////////////////////////////////////////////////////////////////////////
 static int32 LoadCache(class FDiskCache&);
 
@@ -1378,14 +1387,31 @@ FCache::FGetToken FCache::Get(uint64 Key, FIoBuffer& OutData) const
 
 	// Nothing's on disk, so lets try the memory cache
 	FReadScopeLock _(MemLock);
-	if (const FIoBuffer* Data = MemCache.Get(Key); Data != nullptr)
+
+
+	if (CloneCachedIOBufferData > 0)
 	{
-		OutData = *Data;
+		// Speculative fix for a read after free happening in FIoChunkEncoding::Decode
+		if (Partial.Key == Key)
+		{
+			OutData = FIoBuffer(FIoBuffer::Clone, Partial.Data.Data(), Partial.Data.GetSize());
+		}
+		else if (const FIoBuffer* Data = MemCache.Get(Key); Data != nullptr)
+		{
+			OutData = FIoBuffer(FIoBuffer::Clone, Data->Data(), Data->GetSize());
+		}
 	}
-
-	if (Partial.Key == Key)
+	else
 	{
-		OutData = Partial.Data;
+		if (const FIoBuffer* Data = MemCache.Get(Key); Data != nullptr)
+		{
+			OutData = *Data;
+		}
+
+		if (Partial.Key == Key)
+		{
+			OutData = Partial.Data;
+		}
 	}
 
 	return 0;
