Dataflow : add on screen node performances
- can be disabled using p.Dataflow.ShowPerfData 0

diff --git a/Engine/Plugins/Experimental/Dataflow/Source/DataflowEditor/Private/Dataflow/DataflowSNode.cpp b/Engine/Plugins/Experimental/Dataflow/Source/DataflowEditor/Private/Dataflow/DataflowSNode.cpp
--- a/Engine/Plugins/Experimental/Dataflow/Source/DataflowEditor/Private/Dataflow/DataflowSNode.cpp
+++ b/Engine/Plugins/Experimental/Dataflow/Source/DataflowEditor/Private/Dataflow/DataflowSNode.cpp
@@ -16,6 +16,7 @@
 #include "Styling/AppStyle.h"
 #include "Widgets/Layout/SBorder.h"
 #include "Widgets/Input/SCheckBox.h"
+#include "Widgets/Text/STextBlock.h"
 #include "Widgets/SBoxPanel.h"
 #include "Widgets/Images/SImage.h"
 #include "Editor/Transactor.h"
@@ -25,6 +26,12 @@
 
 #define LOCTEXT_NAMESPACE "SDataflowEdNode"
 
+namespace Dataflow::Private
+{
+	static bool bShowPerfData = true;
+	static FAutoConsoleVariableRef CVarChaosBoxCalcBoundsISPCEnabled(TEXT("p.Dataflow.ShowPerfData"), bShowPerfData, TEXT("Whether to display perf data on top of dataflow nodes"));
+}
+
 const FSlateBrush* SDataflowEdNode::GetPinButtonImage() const
 {
 	if (DataflowGraphNode && DataflowGraphNode->ShouldWireframeRenderNode())
@@ -63,6 +70,9 @@ void SDataflowEdNode::Construct(const FArguments& InArgs, UDataflowEdNode* InNod
 				return EVisibility::Collapsed;
 			});
 
+	PerfWidget = SNew(STextBlock);
+
+
 	//
 	//  Cached
 	//
@@ -223,6 +233,42 @@ TArray<FOverlayWidgetInfo> SDataflowEdNode::GetOverlayWidgets(bool bSelected, co
 			FreezeImageInfo.Widget = FreezeImageWidget;
 			Widgets.Add(FreezeImageInfo);
 		}
+
+		if (Dataflow::Private::bShowPerfData)
+		{
+			if (TSharedPtr<UE::Dataflow::FContext> Context = DataflowInterface->GetDataflowContext())
+			{
+				if (TSharedPtr<const FDataflowNode> DataflowNode = DataflowGraphNode->GetDataflowNode())
+				{
+					UE::Dataflow::FContextPerfData::FData PerfData = Context->GetPerfDataForNode(*DataflowNode);
+					if (PerfData.ExclusiveTimeMs > 0 && PerfData.InclusiveTimeMs > 0)
+					{
+						constexpr double FadeTimeMs = 5000;
+						const double MsSinceUpdate = FPlatformTime::ToMilliseconds64(FPlatformTime::Cycles64() - PerfData.LastTimestamp.Value);
+						const double LerpFactor = FMath::Clamp(MsSinceUpdate / FadeTimeMs, 0.0, 1.0);
+						const FLinearColor Color = FMath::Lerp(FLinearColor::Red, FLinearColor::White, LerpFactor);
+						PerfWidget->SetColorAndOpacity(Color);
+						PerfWidget->SetVisibility(EVisibility::Visible);
+						FNumberFormattingOptions FmtOptions;
+						FmtOptions.MaximumFractionalDigits = 2;
+						FmtOptions.MinimumFractionalDigits = 2;
+
+						PerfWidget->SetText(FText::Format(LOCTEXT("DataflowNodePerfDataFormat", "Time: {0} ms"), FText::AsNumber(PerfData.ExclusiveTimeMs, &FmtOptions)));
+					}
+					else
+					{
+						PerfWidget->SetVisibility(EVisibility::Collapsed);
+						PerfWidget->SetText(FText());
+					}
+
+					static const FVector2D FreezeImageOverhang(10.f);  // The icon must slightly overhang to make space for the lower output pins
+					FOverlayWidgetInfo PerfWidgetInfo;
+					PerfWidgetInfo.OverlayOffset = FVector2D{ 0, -20 };
+					PerfWidgetInfo.Widget = PerfWidget;
+					Widgets.Add(PerfWidgetInfo);
+				}
+			}
+		}
 	}
 
 	/*
diff --git a/Engine/Plugins/Experimental/Dataflow/Source/DataflowEditor/Public/Dataflow/DataflowSNode.h b/Engine/Plugins/Experimental/Dataflow/Source/DataflowEditor/Public/Dataflow/DataflowSNode.h
--- a/Engine/Plugins/Experimental/Dataflow/Source/DataflowEditor/Public/Dataflow/DataflowSNode.h
+++ b/Engine/Plugins/Experimental/Dataflow/Source/DataflowEditor/Public/Dataflow/DataflowSNode.h
@@ -11,6 +11,7 @@
 class UDataflowEdNode;
 class SCheckBox;
 class SImage;
+class STextBlock;
 
 //
 // SDataflowEdNode
@@ -68,6 +69,7 @@ private:
 	FCheckBoxStyle CheckBoxStyle;
 	TSharedPtr<SCheckBox> RenderCheckBoxWidget;
 	TSharedPtr<SImage> FreezeImageWidget;
+	TSharedPtr<STextBlock> PerfWidget;
 
 	//FCheckBoxStyle CacheStatusStyle;
 	//TSharedPtr<SCheckBox> CacheStatus;
diff --git a/Engine/Source/Runtime/Experimental/Dataflow/Core/Private/Dataflow/DataflowInputOutput.cpp b/Engine/Source/Runtime/Experimental/Dataflow/Core/Private/Dataflow/DataflowInputOutput.cpp
--- a/Engine/Source/Runtime/Experimental/Dataflow/Core/Private/Dataflow/DataflowInputOutput.cpp
+++ b/Engine/Source/Runtime/Experimental/Dataflow/Core/Private/Dataflow/DataflowInputOutput.cpp
@@ -274,8 +274,7 @@ bool FDataflowOutput::EvaluateImpl(UE::Dataflow::FContext& Context) const
 	{
 		UE_LOG(LogChaosDataflow, VeryVerbose, TEXT("FDataflowOutput::EvaluateImpl(): Node [%s], Output [%s]"), *GetOwningNode()->GetName().ToString(), *GetName().ToString());
 
-		UE::Dataflow::FContextScopedCallstack Callstack(Context, this);
-		if (Callstack.IsLoopDetected())
+		if (Context.IsInCallstack(this))
 		{
 			ensureMsgf(false, TEXT("Connection %s is already in the callstack, this is certainly because of a loop in the graph"), *GetName().ToString());
 			return false;
@@ -287,8 +286,11 @@ bool FDataflowOutput::EvaluateImpl(UE::Dataflow::FContext& Context) const
 			UE_LOG(LogChaosDataflow, VeryVerbose, TEXT("FDataflowOutput::EvaluateImpl(): Context has data, NodeTimestamp [%ul], CacheTimestamp [%lu]"), OwningNode->GetTimestamp().Value, Context.GetTimestamp(CacheKey()).Value);
 			return true;
 		}
-		// if not, evaluate
+
+		// if not, add to the callstack and evaluate
 		UE_LOG(LogChaosDataflow, Verbose, TEXT("FDataflowNode::Evaluate(): Node [%s], Output [%s], NodeTimestamp [%lu]"), *GetOwningNode()->GetName().ToString(), *GetName().ToString(), OwningNode->GetTimestamp().Value);
+
+		UE::Dataflow::FContextScopedCallstack Callstack(Context, this);
 		OwningNode->Evaluate(Context, this);
 		
 		// Validation
diff --git a/Engine/Source/Runtime/Experimental/Dataflow/Core/Private/Dataflow/DataflowNodeParameters.cpp b/Engine/Source/Runtime/Experimental/Dataflow/Core/Private/Dataflow/DataflowNodeParameters.cpp
--- a/Engine/Source/Runtime/Experimental/Dataflow/Core/Private/Dataflow/DataflowNodeParameters.cpp
+++ b/Engine/Source/Runtime/Experimental/Dataflow/Core/Private/Dataflow/DataflowNodeParameters.cpp
@@ -12,9 +12,85 @@
 
 namespace UE::Dataflow
 {
+
+	////////////////////////////////////////////////////////////////////////////////////////////////////////////////////////////////////////////
+
+	void FContextPerfData::Reset()
+	{
+		DataPerOutput.Reset();
+	}
+
+	void FContextPerfData::Accumulate(const FDataflowConnection* Connection, uint64 TotalTime, uint64 ExternalTime)
+	{
+		if (Connection)
+		{
+			FData& Data = DataPerOutput.FindOrAdd(Connection->GetGuid());
+			Data.ExclusiveTimeMs = FMath::Max(Data.ExclusiveTimeMs, (float)FPlatformTime::ToMilliseconds64(TotalTime - ExternalTime));
+			Data.InclusiveTimeMs = FMath::Max(Data.InclusiveTimeMs, (float)FPlatformTime::ToMilliseconds64(TotalTime));
+			Data.LastTimestamp = Connection->GetOwningNodeTimestamp();
+		}
+	}
+	////////////////////////////////////////////////////////////////////////////////////////////////////////////////////////////////////////////
+
+	void FContextCallstack::Push(const FDataflowConnection* Connection)
+	{
+		FEntry Entry
+		{
+			.Connection = Connection,
+			.StartTime = FPlatformTime::Cycles64(),
+			.ExternalTime = 0,
+		};
+		Callstack.Push(Entry);
+	}
+
+	void FContextCallstack::Pop(const FDataflowConnection* Connection, uint64& OutTotalTime, uint64& OutExternalTime)
+	{
+		OutTotalTime = 0;
+		if (ensure(!Callstack.IsEmpty()))
+		{
+			FEntry& Entry = Callstack.Top();
+			OutTotalTime = FPlatformTime::Cycles64() - Entry.StartTime;
+			OutExternalTime = Entry.ExternalTime;
+
+			Callstack.Pop();
+
+			if (Callstack.Num() > 0)
+			{
+				FEntry& PreviousEntry = Callstack.Top();
+				PreviousEntry.ExternalTime += OutTotalTime;
+			}
+		}
+	}
+
+	const FDataflowConnection* FContextCallstack::Top() const
+	{
+		if (Callstack.IsEmpty())
+		{
+			return nullptr;
+		}
+		return Callstack.Top().Connection;
+	}
+
+	int32 FContextCallstack::Num() const
+	{
+		return Callstack.Num();
+	}
+
+	const FDataflowConnection* FContextCallstack::operator[](int32 Index)
+	{
+		return Callstack[Index].Connection;
+	}
+
+	bool FContextCallstack::Contains(const FDataflowConnection* Connection) const
+	{
+		return (Callstack.FindByKey(Connection) != nullptr);
+	}
+
 	uint64 FTimestamp::Invalid = 0;
 	uint64 FTimestamp::Current() { return FPlatformTime::Cycles64(); }
 
+	//////////////////////////////////////////////////////////////////////
+
 	FTimestamp FContext::GetTimestamp(FContextCacheKey Key) const
 	{
 		if (const TUniquePtr<FContextCacheElementBase>* const Cache = const_cast<FContext*>(this)->GetDataImpl(Key))
@@ -24,6 +100,45 @@ namespace UE::Dataflow
 		return FTimestamp::Invalid;
 	}
 
+	FContextPerfData::FData FContext::GetPerfDataForNode(const FDataflowNode& Node) const
+	{
+		FContextPerfData::FData NodeData;
+#if DATAFLOW_EDITOR_EVALUATION
+		for (FDataflowOutput* Output : Node.GetOutputs())
+		{
+			if (Output)
+			{
+				if (const FContextPerfData::FData* OutputData = PerfData.DataPerOutput.Find(Output->GetGuid()))
+				{
+					NodeData.ExclusiveTimeMs += OutputData->ExclusiveTimeMs;
+					NodeData.InclusiveTimeMs += OutputData->InclusiveTimeMs;
+				}
+			}
+		}
+#endif
+		NodeData.LastTimestamp = Node.GetTimestamp();
+		return NodeData;
+	}
+
+	void FContext::ResetPerfDataForNode(const FDataflowNode& Node)
+	{
+#if DATAFLOW_EDITOR_EVALUATION
+		for (FDataflowOutput* Output : Node.GetOutputs())
+		{
+			if (Output)
+			{
+				if (const FContextPerfData::FData* OutputData = PerfData.DataPerOutput.Find(Output->GetGuid()))
+				{
+					if (Node.GetTimestamp().Value > OutputData->LastTimestamp.Value)
+					{
+						PerfData.DataPerOutput.Remove(Output->GetGuid());
+					}
+				}
+			}
+		}
+#endif
+	}
+
 	void FContext::PushToCallstack(const FDataflowConnection* Connection)
 	{
 #if DATAFLOW_EDITOR_EVALUATION
@@ -35,7 +150,10 @@ namespace UE::Dataflow
 	{
 #if DATAFLOW_EDITOR_EVALUATION
 		ensure(Connection == Callstack.Top());
-		Callstack.Pop();
+		uint64 TotalTime = 0;
+		uint64 ExternalTime = 0;
+		Callstack.Pop(Connection, TotalTime, ExternalTime);
+		PerfData.Accumulate(Connection, TotalTime, ExternalTime);
 #endif
 	}
 
@@ -147,6 +265,7 @@ namespace UE::Dataflow
 		, Connection(InConnection)
 	{
 		bLoopDetected = Context.IsInCallstack(Connection);
+		Context.ResetPerfDataForNode(*Connection->GetOwningNode());
 		Context.PushToCallstack(Connection);
 	}
 
diff --git a/Engine/Source/Runtime/Experimental/Dataflow/Core/Public/Dataflow/DataflowNodeParameters.h b/Engine/Source/Runtime/Experimental/Dataflow/Core/Public/Dataflow/DataflowNodeParameters.h
--- a/Engine/Source/Runtime/Experimental/Dataflow/Core/Public/Dataflow/DataflowNodeParameters.h
+++ b/Engine/Source/Runtime/Experimental/Dataflow/Core/Public/Dataflow/DataflowNodeParameters.h
@@ -339,9 +339,45 @@ inline FArchive& operator<<(FArchive& Ar, UE::Dataflow::FContextCache& ValueIn)
 	return Ar;
 }
 
-
 namespace UE::Dataflow
 {
+	struct FContextPerfData
+	{
+		void Reset();
+		void Accumulate(const FDataflowConnection* Connection, uint64 TotalTime, uint64 ExternalTime);
+
+		struct FData
+		{
+			float InclusiveTimeMs = 0;
+			float ExclusiveTimeMs = 0;
+			UE::Dataflow::FTimestamp LastTimestamp = UE::Dataflow::FTimestamp(0); // timestamp of the owner node when data was last collected
+		};
+		TMap<FGuid, FData> DataPerOutput;
+	};
+
+	struct FContextCallstack
+	{
+		void Push(const FDataflowConnection* Connection);
+		void Pop(const FDataflowConnection* Connection, uint64& OutTotalTime, uint64& OutExternalTime);
+		const FDataflowConnection* Top() const;
+
+		int32 Num() const;
+		const FDataflowConnection* operator[](int32 Index);
+		bool Contains(const FDataflowConnection* Connection) const;
+
+		struct FEntry
+		{
+			const FDataflowConnection* Connection;
+			uint64 StartTime = 0;
+			uint64 ExternalTime = 0;
+
+			bool operator==(const FDataflowConnection* OtherConnection) const
+			{
+				return (OtherConnection == Connection);
+			}
+		};
+		TArray<FEntry> Callstack;
+	};
 
 	class FContext
 	{
@@ -494,6 +530,9 @@ namespace UE::Dataflow
 		DATAFLOWCORE_API void PopFromCallstack(const FDataflowConnection* Connection);
 		DATAFLOWCORE_API bool IsInCallstack(const FDataflowConnection* Connection) const;
 
+		DATAFLOWCORE_API FContextPerfData::FData GetPerfDataForNode(const FDataflowNode& Node) const;
+		DATAFLOWCORE_API void ResetPerfDataForNode(const FDataflowNode& Node);
+
 		DATAFLOWCORE_API bool IsCacheEntryAfterTimestamp(FContextCacheKey InKey, const FTimestamp InTimestamp) const;
 
 		DATAFLOWCORE_API void Info(const FString& InInfo, const FDataflowNode* InNode = nullptr, const FDataflowOutput* InOutput = nullptr);
@@ -529,7 +568,8 @@ namespace UE::Dataflow
 
 	private:
 #if DATAFLOW_EDITOR_EVALUATION
-		TArray<const FDataflowConnection*> Callstack;
+		FContextCallstack Callstack;
+		FContextPerfData PerfData;
 #endif
 		TSet<const FDataflowNode*> NodesWithInfo;
 		TSet<const FDataflowNode*> NodesWithWarning;
