Fix event not triggerered with correct probe flag
Probe flag set in the modifier should be initilaized before being used

diff --git a/Engine/Source/Runtime/Experimental/Chaos/Private/Chaos/Collision/ParticlePairMidPhase.cpp b/Engine/Source/Runtime/Experimental/Chaos/Private/Chaos/Collision/ParticlePairMidPhase.cpp
--- a/Engine/Source/Runtime/Experimental/Chaos/Private/Chaos/Collision/ParticlePairMidPhase.cpp
+++ b/Engine/Source/Runtime/Experimental/Chaos/Private/Chaos/Collision/ParticlePairMidPhase.cpp
@@ -1825,7 +1825,9 @@ namespace Chaos
 			// NOTE: Probe constraints are always active and we run collision detection for them at the end opf the frame
 			bool bIsActive = true;
 
-			if (!Constraint->IsProbe())
+			// We should test the original probe flag, because the modifier could have changed it at this point
+			// Reseting flag will occur only when constraint will be activated.
+			if (!Constraint->BoundsTestFlags.bIsProbe)
 			{
 				if (!bUseCCDSweep)
 				{
