Chaosflesh:
-added DeleteVertexTrianglePositionTargetBinding node
-added delete function in PositionTarget facade

diff --git a/Engine/Plugins/Experimental/ChaosFlesh/Source/ChaosFleshNodes/Private/Dataflow/ChaosFleshNodesPlugin.cpp b/Engine/Plugins/Experimental/ChaosFlesh/Source/ChaosFleshNodes/Private/Dataflow/ChaosFleshNodesPlugin.cpp
--- a/Engine/Plugins/Experimental/ChaosFlesh/Source/ChaosFleshNodes/Private/Dataflow/ChaosFleshNodesPlugin.cpp
+++ b/Engine/Plugins/Experimental/ChaosFlesh/Source/ChaosFleshNodes/Private/Dataflow/ChaosFleshNodesPlugin.cpp
@@ -51,6 +51,7 @@ void IChaosFleshNodesPlugin::StartupModule()
 	DATAFLOW_NODE_REGISTER_CREATION_FACTORY(FComputeMuscleActivationDataNode_v2);
 	DATAFLOW_NODE_REGISTER_CREATION_FACTORY(FCreateTetrahedronDataflowNode);
 	DATAFLOW_NODE_REGISTER_CREATION_FACTORY(FDeleteFleshVerticesDataflowNode);
+	DATAFLOW_NODE_REGISTER_CREATION_FACTORY(FDeleteVertexTrianglePositionTargetBindingDataflowNode);
 	DATAFLOW_NODE_REGISTER_CREATION_FACTORY(FFleshAssetTerminalDataflowNode);
 	DATAFLOW_NODE_REGISTER_CREATION_FACTORY(FGenerateOriginInsertionNode);
 	DATAFLOW_NODE_REGISTER_CREATION_FACTORY(FGenerateSkeletalBindings);
diff --git a/Engine/Plugins/Experimental/ChaosFlesh/Source/ChaosFleshNodes/Private/Dataflow/ChaosFleshSetVertexTrianglePositionTargetBindingNode.cpp b/Engine/Plugins/Experimental/ChaosFlesh/Source/ChaosFleshNodes/Private/Dataflow/ChaosFleshSetVertexTrianglePositionTargetBindingNode.cpp
--- a/Engine/Plugins/Experimental/ChaosFlesh/Source/ChaosFleshNodes/Private/Dataflow/ChaosFleshSetVertexTrianglePositionTargetBindingNode.cpp
+++ b/Engine/Plugins/Experimental/ChaosFlesh/Source/ChaosFleshNodes/Private/Dataflow/ChaosFleshSetVertexTrianglePositionTargetBindingNode.cpp
@@ -163,6 +163,42 @@ void FSetVertexTrianglePositionTargetBindingDataflowNode::Evaluate(UE::Dataflow:
 	}
 }
 
+void FDeleteVertexTrianglePositionTargetBindingDataflowNode::Evaluate(UE::Dataflow::FContext& Context, const FDataflowOutput* Out) const
+{
+	if (Out->IsA<DataType>(&Collection))
+	{
+		DataType InCollection = GetValue<DataType>(Context, &Collection);
+		if (IsConnected(&VertexSelection1) && IsConnected(&VertexSelection2))
+		{
+			const FDataflowVertexSelection& InVertexSelection1 = GetValue(Context, &VertexSelection1);
+			const FDataflowVertexSelection& InVertexSelection2 = GetValue(Context, &VertexSelection2);
+			const int32 VertexCount = InCollection.NumElements(FGeometryCollection::VerticesGroup);
+			if (InVertexSelection1.Num() != VertexCount)
+			{
+				Context.Error(FString::Printf(
+					TEXT("VertexSelection1 size (%d) is not equal to the collection's vertex count (%d)"),
+					InVertexSelection1.Num(), VertexCount),
+					this, Out);
+			}
+			if (InVertexSelection2.Num() != VertexCount)
+			{
+				Context.Error(FString::Printf(
+					TEXT("VertexSelection2 size (%d) is not equal to the collection's vertex count (%d)"),
+					InVertexSelection2.Num(), VertexCount),
+					this, Out);
+			}
+			GeometryCollection::Facades::FPositionTargetFacade PositionTargets(InCollection);
+			const int32 NumRemoved = PositionTargets.RemovePositionTarget(
+				[&InVertexSelection1](const int32 VertexIdx) { return InVertexSelection1.IsSelected(VertexIdx); },
+				[&InVertexSelection2](const int32 VertexIdx) { return InVertexSelection2.IsSelected(VertexIdx); });
+			Context.Info(FString::Printf(
+				TEXT("DeleteVertexTrianglePositionTargetBinding: removed %d springs between two VertexSelections"), 
+				NumRemoved), this, Out);
+		}
+		SetValue(Context, MoveTemp(InCollection), &Collection);
+	}
+}
+
 void FSetCollidableVerticesDataflowNode::Evaluate(UE::Dataflow::FContext& Context, const FDataflowOutput* Out) const
 {
 	if (Out->IsA<DataType>(&Collection))
diff --git a/Engine/Plugins/Experimental/ChaosFlesh/Source/ChaosFleshNodes/Public/Dataflow/ChaosFleshSetVertexTrianglePositionTargetBindingNode.h b/Engine/Plugins/Experimental/ChaosFlesh/Source/ChaosFleshNodes/Public/Dataflow/ChaosFleshSetVertexTrianglePositionTargetBindingNode.h
--- a/Engine/Plugins/Experimental/ChaosFlesh/Source/ChaosFleshNodes/Public/Dataflow/ChaosFleshSetVertexTrianglePositionTargetBindingNode.h
+++ b/Engine/Plugins/Experimental/ChaosFlesh/Source/ChaosFleshNodes/Public/Dataflow/ChaosFleshSetVertexTrianglePositionTargetBindingNode.h
@@ -47,6 +47,43 @@ public:
 	virtual void Evaluate(UE::Dataflow::FContext& Context, const FDataflowOutput* Out) const override;
 };
 
+/**
+ * Delete vertex-triangle weak constraints (zero rest length springs) between VertexSelection1 and VertexSelection2.
+ */
+USTRUCT(meta = (DataflowFlesh))
+struct FDeleteVertexTrianglePositionTargetBindingDataflowNode : public FDataflowNode
+{
+	GENERATED_USTRUCT_BODY()
+	DATAFLOW_NODE_DEFINE_INTERNAL(FDeleteVertexTrianglePositionTargetBindingDataflowNode, "DeleteVertexTrianglePositionTargetBinding", "Flesh", "")
+	DATAFLOW_NODE_RENDER_TYPE("SurfaceRender", FGeometryCollection::StaticType(), "Collection")
+
+public:
+	FDeleteVertexTrianglePositionTargetBindingDataflowNode(const UE::Dataflow::FNodeParameters& InParam, FGuid InGuid = FGuid::NewGuid())
+		: FDataflowNode(InParam, InGuid)
+	{
+		RegisterInputConnection(&Collection);
+		RegisterInputConnection(&VertexSelection1);
+		RegisterInputConnection(&VertexSelection2);
+		RegisterOutputConnection(&Collection, &Collection);
+	}
+
+private:
+	typedef FManagedArrayCollection DataType;
+
+	UPROPERTY(meta = (DataflowInput, DataflowIntrinsic, DataflowOutput, DisplayName = "Collection", DataflowPassthrough = "Collection"))
+	FManagedArrayCollection Collection;
+
+	/* This node deletes springs between VertexSelection1 and VertexSelection2. */
+	UPROPERTY(meta = (DataflowInput, DataflowIntrinsic, DisplayName = "VertexSelection1"))
+	FDataflowVertexSelection VertexSelection1;
+
+	/* This node deletes springs between VertexSelection1 and VertexSelection2. */
+	UPROPERTY(meta = (DataflowInput, DataflowIntrinsic, DisplayName = "VertexSelection2"))
+	FDataflowVertexSelection VertexSelection2;
+
+	virtual void Evaluate(UE::Dataflow::FContext& Context, const FDataflowOutput* Out) const override;
+};
+
 /**
  * Set custom vertices so that only these vertices can collide with other surfaces. 
  * Unselected vertices will not collide with unselected vertices.
diff --git a/Engine/Source/Runtime/Experimental/Chaos/Private/GeometryCollection/Facades/CollectionPositionTargetFacade.cpp b/Engine/Source/Runtime/Experimental/Chaos/Private/GeometryCollection/Facades/CollectionPositionTargetFacade.cpp
--- a/Engine/Source/Runtime/Experimental/Chaos/Private/GeometryCollection/Facades/CollectionPositionTargetFacade.cpp
+++ b/Engine/Source/Runtime/Experimental/Chaos/Private/GeometryCollection/Facades/CollectionPositionTargetFacade.cpp
@@ -101,12 +101,12 @@ namespace GeometryCollection::Facades
 
 			if (SourceIndexAttribute.Num() > DataIndex && DataIndex > -1)
 			{
-				ReturnData.SourceIndex = SourceIndexAttribute.Get()[DataIndex];
+				ReturnData.SourceIndex = SourceIndexAttribute[DataIndex];
 			}
 
 			if (TargetIndexAttribute.Num() > DataIndex && DataIndex > -1)
 			{
-				ReturnData.TargetIndex = TargetIndexAttribute.Get()[DataIndex];
+				ReturnData.TargetIndex = TargetIndexAttribute[DataIndex];
 			}
 
 			if (SourceWeightsAttribute.Num() > DataIndex && DataIndex > -1)
@@ -133,7 +133,7 @@ namespace GeometryCollection::Facades
 			for (int32 Index = 0; Index < Collection->NumElements(GroupName); ++Index)
 			{
 				bool bDelete = false;
-				for (int32 TargetIdx : TargetIndexAttribute.Get()[Index])
+				for (int32 TargetIdx : TargetIndexAttribute[Index])
 				{
 					if (TargetIdx < 0)
 					{
@@ -144,7 +144,7 @@ namespace GeometryCollection::Facades
 				}
 				if (!bDelete)
 				{
-					for (int32 SourceIdx : SourceIndexAttribute.Get()[Index])
+					for (int32 SourceIdx : SourceIndexAttribute[Index])
 					{
 						if (SourceIdx < 0)
 						{
@@ -158,4 +158,36 @@ namespace GeometryCollection::Facades
 		}
 		return InvalidConstraintIndices.Num();
 	}
+
+	int32 FPositionTargetFacade::RemovePositionTarget(TFunctionRef<bool(const int32 VertexIdx)> IsVertexGroup1, TFunctionRef<bool(const int32 VertexIdx)> IsVertexGroup2)
+	{
+		check(!IsConst());
+		TArray<int32> InvalidConstraintIndices;
+		if (IsValid())
+		{
+			for (int32 Index = 0; Index < Collection->NumElements(GroupName); ++Index)
+			{
+				bool IsTargetGroup1 = true;
+				bool IsTargetGroup2 = true;
+				for (int32 TargetIdx : TargetIndexAttribute[Index])
+				{
+					IsTargetGroup1 = IsTargetGroup1 && IsVertexGroup1(TargetIdx);
+					IsTargetGroup2 = IsTargetGroup2 && IsVertexGroup2(TargetIdx);
+				}
+				bool IsSourceGroup1 = true;
+				bool IsSourceGroup2 = true;
+				for (int32 SourceIdx : SourceIndexAttribute[Index])
+				{
+					IsSourceGroup1 = IsSourceGroup1 && IsVertexGroup1(SourceIdx);
+					IsSourceGroup2 = IsSourceGroup2 && IsVertexGroup2(SourceIdx);
+				}
+				if ((IsTargetGroup1 && IsSourceGroup2) || (IsTargetGroup2 && IsSourceGroup1))
+				{
+					InvalidConstraintIndices.Add(Index);
+				}
+			}
+			Collection->RemoveElements(GroupName, InvalidConstraintIndices);
+		}
+		return InvalidConstraintIndices.Num();
+	}
 }
diff --git a/Engine/Source/Runtime/Experimental/Chaos/Public/GeometryCollection/Facades/CollectionPositionTargetFacade.h b/Engine/Source/Runtime/Experimental/Chaos/Public/GeometryCollection/Facades/CollectionPositionTargetFacade.h
--- a/Engine/Source/Runtime/Experimental/Chaos/Public/GeometryCollection/Facades/CollectionPositionTargetFacade.h
+++ b/Engine/Source/Runtime/Experimental/Chaos/Public/GeometryCollection/Facades/CollectionPositionTargetFacade.h
@@ -79,6 +79,8 @@ namespace GeometryCollection::Facades
 		
 		/** Remove position targets with invalid indices*/
 		CHAOS_API int32 RemoveInvalidPositionTarget();
+		/** Remove position targets between two groups of vertices*/
+		CHAOS_API int32 RemovePositionTarget(TFunctionRef<bool(const int32 VertexIdx)> IsVertexGroup1, TFunctionRef<bool(const int32 VertexIdx)> IsVertexGroup2);
 
 	protected:
 		const FManagedArrayCollection& ConstCollection;
