Chaos - fix active collision list after island wake

Collision constraints reactivated on island wake were not being added to the active list. This meant that collision events for awakened bodies were not being reported on the waking frame. For an event listener, this makes it appear that the bodies are not colliding for one frame, making it impossible to track continuous contacts across sleep/wake cycles.
This has been broken for a while (probably since adding per-worker collision allocation).

diff --git a/Engine/Source/Runtime/Experimental/Chaos/Private/Chaos/Collision/CollisionConstraintAllocator.cpp b/Engine/Source/Runtime/Experimental/Chaos/Private/Chaos/Collision/CollisionConstraintAllocator.cpp
--- a/Engine/Source/Runtime/Experimental/Chaos/Private/Chaos/Collision/CollisionConstraintAllocator.cpp
+++ b/Engine/Source/Runtime/Experimental/Chaos/Private/Chaos/Collision/CollisionConstraintAllocator.cpp
@@ -98,23 +98,10 @@ namespace Chaos
 			if (!ContextAllocators.IsEmpty())
 			{
 				const TArray<FPBDCollisionConstraint*>& NewConstraints = ContextAllocators[0]->NewActiveConstraints;
-				// Add the constraints to the active list (and the CCD list if appropriate)
 				for (FPBDCollisionConstraint* Constraint : NewConstraints)
 				{
-					FPBDCollisionConstraintContainerCookie& Cookie = Constraint->GetContainerCookie();
-
-					// Add the constraint to the active list and update its epoch
-					checkSlow(ActiveConstraints.Find(Constraint) == INDEX_NONE);
-					Cookie.ConstraintIndex = ActiveConstraints.Add(Constraint);
-
-					// If the constraint uses CCD, keep it in another list so we don't have to search the full list for CCD contacts
-					if (Constraint->GetCCDEnabled())
-					{
-						checkSlow(ActiveCCDConstraints.Find(Constraint) == INDEX_NONE);
-						Cookie.CCDConstraintIndex = ActiveCCDConstraints.Add(Constraint);
-					}
-
-					Cookie.LastUsedEpoch = CurrentEpoch;
+					// Add the constraint to the active list (and the CCD list if appropriate)
+					AddActiveConstraintImpl(Constraint);
 				}
 				ContextAllocators[0]->NewActiveConstraints.Reset();
 			}
@@ -176,6 +163,36 @@ namespace Chaos
 			}
 		}
 
+		void FCollisionConstraintAllocator::AddActiveConstraint(FPBDCollisionConstraint* Constraint)
+		{
+			// Add the constraint to the active list, but only if it hasn't been activated this epoch already
+			FPBDCollisionConstraintContainerCookie& Cookie = Constraint->GetContainerCookie();
+			if (Cookie.LastUsedEpoch != CurrentEpoch)
+			{
+				AddActiveConstraintImpl(Constraint);
+			}
+		}
+
+		void FCollisionConstraintAllocator::AddActiveConstraintImpl(FPBDCollisionConstraint* Constraint)
+		{
+			FPBDCollisionConstraintContainerCookie& Cookie = Constraint->GetContainerCookie();
+
+			Cookie.LastUsedEpoch = CurrentEpoch;
+			Cookie.ConstraintIndex = INDEX_NONE;
+			Cookie.CCDConstraintIndex = INDEX_NONE;
+
+			// Add to the active list and remember the index for O(1) removal
+			checkSlow(ActiveConstraints.Find(Constraint) == INDEX_NONE);
+			Cookie.ConstraintIndex = ActiveConstraints.Add(Constraint);
+
+			// If the constraint uses CCD, keep it in another list so we don't have to search the full list for CCD contacts
+			if (Constraint->GetCCDEnabled())
+			{
+				checkSlow(ActiveCCDConstraints.Find(Constraint) == INDEX_NONE);
+				Cookie.CCDConstraintIndex = ActiveCCDConstraints.Add(Constraint);
+			}
+		}
+
 		void FCollisionConstraintAllocator::RemoveActiveConstraint(FPBDCollisionConstraint& Constraint)
 		{
 			FPBDCollisionConstraintContainerCookie& Cookie = Constraint.GetContainerCookie();
diff --git a/Engine/Source/Runtime/Experimental/Chaos/Private/Chaos/Collision/ParticlePairMidPhase.cpp b/Engine/Source/Runtime/Experimental/Chaos/Private/Chaos/Collision/ParticlePairMidPhase.cpp
--- a/Engine/Source/Runtime/Experimental/Chaos/Private/Chaos/Collision/ParticlePairMidPhase.cpp
+++ b/Engine/Source/Runtime/Experimental/Chaos/Private/Chaos/Collision/ParticlePairMidPhase.cpp
@@ -681,16 +681,12 @@ namespace Chaos
 	{
 		if (Constraint.IsValid() && IsUsedSince(SleepEpoch))
 		{
-			// We need to refresh the epoch so that the constraint state will be used as the previous
-			// state if the pair is still colliding in the next tick. However, we are not in the active 
-			// collisions array so reset the indices.
-			Constraint->GetContainerCookie().LastUsedEpoch = CurrentEpoch;
-			Constraint->GetContainerCookie().ConstraintIndex = INDEX_NONE;
-			Constraint->GetContainerCookie().CCDConstraintIndex = INDEX_NONE;
-			LastUsedEpoch = CurrentEpoch;
+			// Update the epoch and add to the active list
+			Constraint->ConcreteContainer()->GetConstraintAllocator().AddActiveConstraint(Constraint.Get());
+			LastUsedEpoch = Constraint->GetContainerCookie().LastUsedEpoch;
 
 			// We have skipped collision detection for this particle because it was asleep, so we need to update the transforms...
-			// NOTE: this relies on the shape world transforms being up-to-date. They are usually updated in Integarte which
+			// NOTE: this relies on the shape world transforms being up-to-date. They are usually updated in Integrate which
 			// is also skipped for sleeping particles, so they must be updated manually when waking partciles (see IslandManager)
 			Constraint->SetShapeWorldTransforms(Shape0->GetLeafWorldTransform(GetParticle0()), Shape1->GetLeafWorldTransform(GetParticle1()));
 		}
@@ -2032,9 +2028,8 @@ namespace Chaos
 			FPBDCollisionConstraintPtr& Constraint = KVP.Value;
 			if (Constraint->GetContainerCookie().LastUsedEpoch >= LastUsedEpoch)
 			{
-				Constraint->GetContainerCookie().LastUsedEpoch = CurrentEpoch;
-				Constraint->GetContainerCookie().ConstraintIndex = INDEX_NONE;
-				Constraint->GetContainerCookie().CCDConstraintIndex = INDEX_NONE;
+				// Update the epoch and add to the active list
+				Constraint->ConcreteContainer()->GetConstraintAllocator().AddActiveConstraint(Constraint.Get());
 			}
 		}
 	}
@@ -2230,9 +2225,8 @@ namespace Chaos
 	{
 		if (Constraint.IsValid() && (Constraint->GetContainerCookie().LastUsedEpoch >= LastUsedEpoch))
 		{
-			Constraint->GetContainerCookie().LastUsedEpoch = CurrentEpoch;
-			Constraint->GetContainerCookie().ConstraintIndex = INDEX_NONE;
-			Constraint->GetContainerCookie().CCDConstraintIndex = INDEX_NONE;
+			// Update the epoch and add to the active list
+			Constraint->ConcreteContainer()->GetConstraintAllocator().AddActiveConstraint(Constraint.Get());
 		}
 	}
 
diff --git a/Engine/Source/Runtime/Experimental/Chaos/Public/Chaos/Collision/CollisionConstraintAllocator.h b/Engine/Source/Runtime/Experimental/Chaos/Public/Chaos/Collision/CollisionConstraintAllocator.h
--- a/Engine/Source/Runtime/Experimental/Chaos/Public/Chaos/Collision/CollisionConstraintAllocator.h
+++ b/Engine/Source/Runtime/Experimental/Chaos/Public/Chaos/Collision/CollisionConstraintAllocator.h
@@ -80,14 +80,15 @@ namespace Chaos
 
 			/**
 			 * Set the constraint as Active - it will be added to the graph and solved this tick.
-			 * NOTE: This actually just enqueues the constraint - actual activation happens after the parallel phase is done.
+			 * NOTE: The definitive active constraints list is compiled after the parallel phase 
+			 * is complete (FCollisionConstraintAllocator::ProcessNewConstraints).
 			 */
 			bool ActivateConstraint(FPBDCollisionConstraint* Constraint)
 			{
-				// When we wake an Island, we reactivate all constraints for all dynamic particles in the island. This
-				// results in duplicate calls to active for constraints involving two dynamic particles, hence the check for CurrentEpoch.
-				// @todo(chaos): fix duplicate calls from island wake. See UpdateSleepState in IslandManager.cpp
 				FPBDCollisionConstraintContainerCookie& Cookie = Constraint->GetContainerCookie();
+				// Use the LastUsedEpoch to catch duplicate add attempts. This is still possible
+				// via rewind/resim and InjectCollision. Ideally we wouldn't need do update
+				// the cookie data until FCollisionConstraintAllocator::ProcessNewCollisions...
 				if (Cookie.LastUsedEpoch != CurrentEpoch)
 				{
 					Cookie.LastUsedEpoch = CurrentEpoch;
@@ -413,6 +414,14 @@ namespace Chaos
 				ProcessNewItems();
 			}
 
+			/**
+			 * This should be called when a sleeping constraint is activated after the narrow phase, typically
+			 * in the island manager. This is mainly so that any systems asking for a list of active constraints
+			 * are provided all constraints that may have contributed to the simulation, including those that
+			 * were added when an island is awoken.
+			 */
+			void AddActiveConstraint(FPBDCollisionConstraint* Constraint);
+
 			/**
 			 * @brief Add a set of pre-built constraints and build required internal mapping data
 			 * This is used by the resim cache when restoring constraints after a desync
@@ -554,6 +563,7 @@ namespace Chaos
 			// Collect all the constraints activated in the collision tasks and register them (calls ActivateConstraintImp on each)
 			CHAOS_API void ProcessNewConstraints();
 
+			CHAOS_API void AddActiveConstraintImpl(FPBDCollisionConstraint* Constraint);
 			CHAOS_API void RemoveActiveConstraint(FPBDCollisionConstraint& Constraint);
 
 			CHAOS_API void SortActiveConstraints();
