Chaos - Event dispatch will now dispatch events even when there are no collision so that, for example, we can track begin/end collisions correctly

diff --git a/Engine/Source/Runtime/Experimental/Chaos/Private/EventManager.cpp b/Engine/Source/Runtime/Experimental/Chaos/Private/EventManager.cpp
--- a/Engine/Source/Runtime/Experimental/Chaos/Private/EventManager.cpp
+++ b/Engine/Source/Runtime/Experimental/Chaos/Private/EventManager.cpp
@@ -4,6 +4,13 @@
 
 namespace Chaos
 {
+	namespace Private
+	{
+		// This is required to be able to track continuous collisions (so that we can detect the "no collision" state)
+		bool bEventManagerDispatchEmptyEvents = true;
+		FAutoConsoleVariableRef CVarEventManagerDispatchEmptyEvent(TEXT("p.Chaos.EventManager.DispatchEmptyEvents"), bEventManagerDispatchEmptyEvents, TEXT("Whether to dispatch events even if the event buffers are empty"));
+	}
+
 	void FEventManager::Reset()
 	{
 		ContainerLock.WriteLock();
diff --git a/Engine/Source/Runtime/Experimental/Chaos/Private/EventsData.cpp b/Engine/Source/Runtime/Experimental/Chaos/Private/EventsData.cpp
new file mode 100644
--- /dev/null
+++ b/Engine/Source/Runtime/Experimental/Chaos/Private/EventsData.cpp
@@ -0,0 +1,91 @@
+// Copyright Epic Games, Inc. All Rights Reserved.
+#include "EventsData.h"
+#include "EventManager.h"
+
+namespace Chaos
+{
+	FCollisionEventDataIterator::FCollisionEventDataIterator(const TArrayView<const IPhysicsProxyBase*>& InProxies, const Chaos::FCollisionEventData& InCollisionEventData)
+		: ProxyIndex(INDEX_NONE)
+		, ProxyCollisionIndex(INDEX_NONE)
+		, ProxyCollisionIndices(nullptr)
+		, Proxies(InProxies)
+		, CollisionEventData(InCollisionEventData)
+	{
+		Reset();
+	}
+
+	void FCollisionEventDataIterator::Reset()
+	{
+		ProxyIndex = 0;
+		ProxyCollisionIndex = 0;
+		ProxyCollisionIndices = nullptr;
+
+		if (ProxyIndex < Proxies.Num())
+		{
+			const IPhysicsProxyBase* NextProxy = Proxies[ProxyIndex];
+			ProxyCollisionIndices = CollisionEventData.PhysicsProxyToCollisionIndices.PhysicsProxyToIndicesMap.Find(NextProxy);
+		}
+
+		while (!IsFinished() && !GetCurrentCollidingDataIndex().IsValid())
+		{
+			Next();
+		}
+	}
+
+	void FCollisionEventDataIterator::Next()
+	{
+		if (!IsFinished())
+		{
+			// Move to the next event for the current proxy
+			++ProxyCollisionIndex;
+
+			// If we run out of events for that proxy, move to the next proxy
+			// Skip over any proxies with no collision events
+			// NOTE: Next() may be from Reset() so we have to handle null ProxyCollisionIndices
+			if ((ProxyCollisionIndices == nullptr) || (ProxyCollisionIndex >= ProxyCollisionIndices->Num()))
+			{
+				ProxyCollisionIndex = 0;
+				ProxyCollisionIndices = nullptr;
+				while (!IsFinished() && !GetCurrentCollidingDataIndex().IsValid())
+				{
+					if (++ProxyIndex < Proxies.Num())
+					{
+						const IPhysicsProxyBase* CurrentProxy = Proxies[ProxyIndex];
+						ProxyCollisionIndices = CollisionEventData.PhysicsProxyToCollisionIndices.PhysicsProxyToIndicesMap.Find(CurrentProxy);
+
+						// Skip empty proxies
+						if ((ProxyCollisionIndices != nullptr) && (ProxyCollisionIndices->Num() == 0))
+						{
+							ProxyCollisionIndices = nullptr;
+						}
+					}
+				}
+			}
+		}
+	}
+
+	FCollidingDataIndex FCollisionEventDataIterator::GetCurrentCollidingDataIndex() const
+	{
+		if (!IsFinished())
+		{
+			if ((ProxyCollisionIndices != nullptr) && (ProxyCollisionIndex < ProxyCollisionIndices->Num()))
+			{
+				bool bSwapProxyOrder;
+				int32 EncodedCollisionIndex = (*ProxyCollisionIndices)[ProxyCollisionIndex];
+				int32 CollisionIndex = FEventManager::DecodeCollisionIndex(EncodedCollisionIndex, bSwapProxyOrder);
+				int32 CollisionProxyIndex = bSwapProxyOrder ? 1 : 0;
+				return FCollidingDataIndex(CollisionIndex, CollisionProxyIndex);
+			}
+		}
+		return {};
+	}
+
+	const IPhysicsProxyBase* FCollisionEventDataIterator::GetCurrentProxy() const
+	{
+		if (!IsFinished())
+		{
+			return Proxies[ProxyIndex];
+		}
+		return nullptr;
+	}
+}
\ No newline at end of file
diff --git a/Engine/Source/Runtime/Experimental/Chaos/Public/EventManager.h b/Engine/Source/Runtime/Experimental/Chaos/Public/EventManager.h
--- a/Engine/Source/Runtime/Experimental/Chaos/Public/EventManager.h
+++ b/Engine/Source/Runtime/Experimental/Chaos/Public/EventManager.h
@@ -12,6 +12,10 @@
 
 namespace Chaos
 {
+	namespace Private
+	{
+		extern bool bEventManagerDispatchEmptyEvents;
+	}
 
 	//#todo : need to enable disable events, or not register them if disabled.
 	//#todo : add timers
@@ -298,10 +302,16 @@ namespace Chaos
 		virtual void DispatchConsumerData()
 		{
 			const PayloadType* Buffer = EventBuffer.Get()->GetConsumerBuffer();
-			if (IsEventDataEmpty(Buffer))
+
+			// We dispatch events even when the event buffers are empty because the absence
+			// of collisions/breaks/etc is required information for some use-cases like
+			// tracking continuous collisions.
+			const bool bDispatchEmptyEvents = Private::bEventManagerDispatchEmptyEvents;
+			if (IsEventDataEmpty(Buffer) && !bDispatchEmptyEvents)
 			{
 				return;
 			}
+
 			UnregisterDeferedHandler();
 			HandlerLock.ReadLock();
 			const TMap<IPhysicsProxyBase*, TArray<int32>>* Map = GetProxyToIndexMap(Buffer); // Use t his map to get all proxies used in the event buffer
diff --git a/Engine/Source/Runtime/Experimental/Chaos/Public/EventsData.h b/Engine/Source/Runtime/Experimental/Chaos/Public/EventsData.h
--- a/Engine/Source/Runtime/Experimental/Chaos/Public/EventsData.h
+++ b/Engine/Source/Runtime/Experimental/Chaos/Public/EventsData.h
@@ -306,4 +306,140 @@ namespace Chaos
 			return nullptr;
 		}
 	}
+
+	// An index into a FCollisionEventData array of FCollidingDataobjects obtained from
+	// the PhysicsProxyToIndicesMap for a specific proxy.
+	// Indices held in FCollisionEventData::PhysicsProxyToCollisionIndices are encoded
+	// to include the proxy index in the collision data structure (0 or 1).
+	// This is useful if you collect a set of collision indices for a set of proxies, and
+	// the store a flat list of FCollidingDataIndexes. In this case it may be expensive
+	// to determine which of the proxies in the collision belongs to one of our proxies.
+	class FCollidingDataIndex
+	{
+	public:
+		FCollidingDataIndex() 
+			: ProxyIndex(0)
+			, CollisionIndex(INDEX_NONE)
+		{
+		}
+
+		// InProxyIndex: The index of the proxy in the FCollidingData (0 or 1)
+		// InCollisionIndex: The index of the FCollidingData in the FCollisionEventData array
+		FCollidingDataIndex(int32 InCollisionIndex, int32 InProxyIndex)
+			: ProxyIndex(InProxyIndex)
+			, CollisionIndex(InCollisionIndex)
+		{
+			check(InCollisionIndex >= 0);
+			check(InProxyIndex >= 0);
+			check(InProxyIndex <= 1);
+		}
+
+		// Reset to invalid state
+		void Reset()
+		{
+			ProxyIndex = 0;
+			CollisionIndex = INDEX_NONE;
+		}
+
+		// Is this index a valid index into the FCollidingData array?
+		bool IsValid() const
+		{
+			return CollisionIndex >= 0;
+		}
+
+		// The index into the FCollidingData array
+		int32 GetIndex() const
+		{
+			return CollisionIndex;
+		}
+
+		// The index of our body in the collision data. This will be zero or one.
+		int32 GetProxyIndex() const
+		{
+			return ProxyIndex;
+		}
+
+	private:
+		uint32 ProxyIndex : 1;
+		int32 CollisionIndex : 31;
+	};
+
+	// It wouldn't matter if it were larger, but just in case more flags are added and we do something unexpected...
+	static_assert(sizeof(FCollidingDataIndex) == sizeof(int32), "FCollidingDataIndex should be same size as int32");
+
+	// Iterates over the FCollidingData objects for a set of proxies
+	//
+	// Usage:
+	//	for (FCollisionEventDataIterator It(MyProxyList, CollisionData); It; ++It)
+	//	{
+	//		Chaos::FCollidingDataIndex CollidingDataIndex = It.GetCurrentCollidingDataIndex();
+	//		if (CollidingDataIndex.IsValid())
+	//		{
+	//			// CollidingData is between one ot the proxies in MyProxyList and another (which in general may also be in MyProxyList!)
+	//			const Chaos::FCollidingData& CollidingData = CollisionEventData.CollisionData.AllCollisionsArray[CollidingDataIndex.GetIndex()];
+	//			
+	//			Chaos::FVec3 Normal = (CollidingDataIndex.GetProxyIndex() == 0) ? CollidingData.Normal : -CollidingData.Normal
+	//			float Mass = (CollidingDataIndex.GetProxyIndex() == 0) ? CollidingData.Mass1 : CollidingData.Mass2;
+	//			// ...
+	//		}}
+	//	}
+	//
+	class FCollisionEventDataIterator
+	{
+	public:
+		CHAOS_API FCollisionEventDataIterator(const TArrayView<const IPhysicsProxyBase*>& InProxies, const Chaos::FCollisionEventData& InCollisionEventData);
+
+		// Reset the iterator to point at the first FCollidingData for the first proxy that has any collisions.
+		CHAOS_API void Reset();
+
+		// Move to the next collision. If we have reached the end of collisions for the current proxy, 
+		// move to the next proxy with collisions and select its first collision.
+		CHAOS_API void Next();
+
+		// Get the index of the current collision in the CollisionEventData's CollidingData array
+		// or INDEX_NONE if we have reached the end of the collisions for all proxies.
+		CHAOS_API FCollidingDataIndex GetCurrentCollidingDataIndex() const;
+
+		// Get the Proxy that the iterator is currently pointing to
+		// or null if we have reached the end of the collisions for all proxies.
+		CHAOS_API const IPhysicsProxyBase* GetCurrentProxy() const;
+
+		// Have we reached the end of the collisions for all proxies?
+		bool IsFinished() const
+		{
+			return (ProxyIndex >= Proxies.Num());
+		}
+
+		// Returns true if we are still iterating (have not reached the end)
+		operator bool() const
+		{
+			return !IsFinished();
+		}
+
+		// Move to the next non-empty collision/proxy
+		FCollisionEventDataIterator& operator++()
+		{
+			Next();
+			return *this;
+		}
+
+		// NOTE: Deliberately not providing a post-increment operator
+		// because it is more expensive to copy than users might expect 
+		// (although not that bad really, so feel free to add one if
+		// you really need it!)
+		FCollisionEventDataIterator operator++(int) = delete;
+
+	private:
+		// Current Iteration state
+		int32 ProxyIndex;
+		int32 ProxyCollisionIndex;
+		const TArray<int32>* ProxyCollisionIndices;
+
+		// Set of proxies we care about
+		TArrayView<const IPhysicsProxyBase*> Proxies;
+
+		// Collision data from the whole scene
+		const Chaos::FCollisionEventData& CollisionEventData;
+	};
+
 }
