Chaos clustering : fix bug where partial cluster do not inherit gravity settings
- Took the opportunity to transfer a lot more attributes from the parent particle to the newly created one

diff --git a/Engine/Source/Runtime/Experimental/Chaos/Private/Chaos/PBDRigidClustering.cpp b/Engine/Source/Runtime/Experimental/Chaos/Private/Chaos/PBDRigidClustering.cpp
--- a/Engine/Source/Runtime/Experimental/Chaos/Private/Chaos/PBDRigidClustering.cpp
+++ b/Engine/Source/Runtime/Experimental/Chaos/Private/Chaos/PBDRigidClustering.cpp
@@ -576,7 +576,25 @@ namespace Chaos
 		NewParticle->SetInternalCluster(true);
 		NewParticle->SetClusterId(ClusterId(nullptr, Children.Num()));
 		NewParticle->SetIsAnchored(false);
-		for (auto& Constituent : Children) MEvolution.DoInternalParticleInitilization(Constituent, NewParticle);
+		if (Parent)
+		{
+			NewParticle->SetGravityEnabled(Parent->GravityEnabled());
+			NewParticle->SetGravityGroupIndex(Parent->GravityGroupIndex());
+			NewParticle->SetInertiaConditioningEnabled(Parent->InertiaConditioningEnabled());
+			NewParticle->SetCCDEnabled(Parent->CCDEnabled());
+			NewParticle->SetMACDEnabled(Parent->MACDEnabled());
+			NewParticle->SetOneWayInteraction(Parent->OneWayInteraction());
+			NewParticle->SetLinearEtherDrag(Parent->LinearEtherDrag());
+			NewParticle->SetAngularEtherDrag(Parent->AngularEtherDrag());
+			NewParticle->SetInitialOverlapDepenetrationVelocity(Parent->InitialOverlapDepenetrationVelocity());
+			NewParticle->SetSleepThresholdMultiplier(Parent->SleepThresholdMultiplier());
+			NewParticle->SetIterationSettings(Parent->IterationSettings());
+		}
+
+		for (auto& Constituent : Children)
+		{
+			MEvolution.DoInternalParticleInitilization(Constituent, NewParticle);
+		}
 
 		//
 		// Update clustering data structures.
