Allow ondemand containers to resolve an FIasHostGroup when mounted.

#rnx

- Add a very basic manager (FHostGroupManager) that should hold all active FIASHostGroups
- The ondemand container can now resolve the hostgroup via a name (currently we use the same default name for all IAS containers)
- This allows us to assign the container hostgroup to requests from that hostgroup.
- Start to introduce the concept of resolved vs unresolved hostgroups. This will let us move the distributed endpoint logic from the IAS backend to the host group.

diff --git a/Engine/Source/Runtime/Experimental/IoStore/OnDemand/Private/IasHostGroup.cpp b/Engine/Source/Runtime/Experimental/IoStore/OnDemand/Private/IasHostGroup.cpp
--- a/Engine/Source/Runtime/Experimental/IoStore/OnDemand/Private/IasHostGroup.cpp
+++ b/Engine/Source/Runtime/Experimental/IoStore/OnDemand/Private/IasHostGroup.cpp
@@ -2,6 +2,7 @@
 
 #include "IasHostGroup.h"
 
+#include "Algo/Find.h"
 #include "HAL/IConsoleManager.h"
 #include "IO/IoStoreOnDemand.h"
 #include "IO/OnDemandHostGroup.h"
@@ -286,4 +287,48 @@ TConstArrayView<FAnsiString> FIASHostGroup::GetHostUrls() const
 	return Impl->HostGroup.Hosts();
 }
 
+FHostGroupManager& FHostGroupManager::Get()
+{
+	static FHostGroupManager Instance;
+	return Instance;
+}
+
+TIoStatusOr<FIASHostGroup> FHostGroupManager::Register(FName Name)
+{
+	FIASHostGroup HostGroup(Name);
+
+	HostGroups.Add(HostGroup);
+
+	return MoveTemp(HostGroup);
+}
+
+TIoStatusOr<FIASHostGroup> FHostGroupManager::Register(FName Name, TConstArrayView<FAnsiString> HostUrls)
+{
+	TIoStatusOr<FIASHostGroup> Result = FIASHostGroup::Create(Name, HostUrls);
+
+	if (Result.IsOk())
+	{
+		HostGroups.Add(Result.ValueOrDie());
+	}
+
+	return Result;
+}
+
+FIASHostGroup FHostGroupManager::Find(FName Name)
+{
+	FIASHostGroup* FoundHostGroup = Algo::FindByPredicate(HostGroups, [Name](const FIASHostGroup& HostGroup)->bool
+	{
+		return HostGroup.GetName() == Name;
+	});
+
+	if (FoundHostGroup != nullptr)
+	{
+		return *FoundHostGroup;
+	}
+	else
+	{
+		return FIASHostGroup();
+	}
+}
+
 } //namespace UE::IoStore
diff --git a/Engine/Source/Runtime/Experimental/IoStore/OnDemand/Private/IasHostGroup.h b/Engine/Source/Runtime/Experimental/IoStore/OnDemand/Private/IasHostGroup.h
--- a/Engine/Source/Runtime/Experimental/IoStore/OnDemand/Private/IasHostGroup.h
+++ b/Engine/Source/Runtime/Experimental/IoStore/OnDemand/Private/IasHostGroup.h
@@ -63,4 +63,24 @@ private:
 	
 };
 
+class FHostGroupManager final
+{
+	FHostGroupManager() =default;
+	~FHostGroupManager() = default;
+
+public:
+
+	static FHostGroupManager& Get(); 
+
+	// TODO: Note that currently we have no way to remove FIASHostGroups
+
+	TIoStatusOr<FIASHostGroup> Register(FName Name); // <- Unresolved
+	TIoStatusOr<FIASHostGroup> Register(FName Name, TConstArrayView<FAnsiString> HostUrls); // <- Resolved
+
+	FIASHostGroup				Find(FName Name);
+private:
+
+	TArray<FIASHostGroup> HostGroups;
+};
+
 } //namespace UE::IoStore
diff --git a/Engine/Source/Runtime/Experimental/IoStore/OnDemand/Private/OnDemandIoDispatcherBackend.cpp b/Engine/Source/Runtime/Experimental/IoStore/OnDemand/Private/OnDemandIoDispatcherBackend.cpp
--- a/Engine/Source/Runtime/Experimental/IoStore/OnDemand/Private/OnDemandIoDispatcherBackend.cpp
+++ b/Engine/Source/Runtime/Experimental/IoStore/OnDemand/Private/OnDemandIoDispatcherBackend.cpp
@@ -899,9 +899,10 @@ FOnDemandIoBackend::FOnDemandIoBackend(const FOnDemandEndpointConfig& Config, FO
 	: IoStore(InIoStore)
 	, Cache(MoveTemp(InCache))
 	, Stats(BackendStatus)
-	, DefaultHostGroup(FName("Default"))
 	, CurrentThreadPriority(EThreadPriority::TPri_Num)
 {
+	DefaultHostGroup = FHostGroupManager::Get().Register(FName("DefaultOnDemand")).ConsumeValueOrDie();
+
 	EndpointTestPath = Config.TocPath;
 	if (Config.DistributionUrl.IsEmpty() == false)
 	{
@@ -1248,7 +1249,7 @@ bool FOnDemandIoBackend::Resolve(FIoRequestImpl* Request)
 
 	// Currently all requests are assigned the DefaultHostGroup
 	// TODO: When next fixing this, make sure we deal with the !DefaultHostGroup.IsConnected() above
-	ChunkRequest->HostGroup = DefaultHostGroup;
+	ChunkRequest->HostGroup = ChunkInfo.HostGroup();
 
 	Stats.OnHttpEnqueue();
 	HttpRequests.EnqueueByPriority(ChunkRequest);
diff --git a/Engine/Source/Runtime/Experimental/IoStore/OnDemand/Private/OnDemandIoStore.cpp b/Engine/Source/Runtime/Experimental/IoStore/OnDemand/Private/OnDemandIoStore.cpp
--- a/Engine/Source/Runtime/Experimental/IoStore/OnDemand/Private/OnDemandIoStore.cpp
+++ b/Engine/Source/Runtime/Experimental/IoStore/OnDemand/Private/OnDemandIoStore.cpp
@@ -1209,6 +1209,8 @@ FIoStatus FOnDemandIoStore::TickMountRequest(FMountRequest& MountRequest)
 		if (EnumHasAnyFlags(Args.Options, EOnDemandMountOptions::StreamOnDemand))
 		{
 			EnumAddFlags(Container->Flags, EOnDemandContainerFlags::StreamOnDemand);
+
+			Container->IasHostGroup = FHostGroupManager::Get().Find(FName("DefaultOnDemand"));
 		}
 		else if (EnumHasAnyFlags(Args.Options, EOnDemandMountOptions::InstallOnDemand))
 		{
@@ -1831,6 +1833,8 @@ void FOnDemandIoStore::CreateContainersFromToc(
 		Container->BlockSizes			= MoveTemp(ContainerEntry.BlockSizes);
 		Container->BlockHashes			= MoveTemp(ContainerEntry.BlockHashes);
 		Container->ContainerId			= ContainerEntry.ContainerId;
+
+		Container->CompressionFormats.Reserve(1);
 		Container->CompressionFormats.Add(CompressionFormat);
 
 		const EIoContainerFlags ContainerFlags = static_cast<EIoContainerFlags>(ContainerEntry.ContainerFlags);
diff --git a/Engine/Source/Runtime/Experimental/IoStore/OnDemand/Private/OnDemandIoStore.h b/Engine/Source/Runtime/Experimental/IoStore/OnDemand/Private/OnDemandIoStore.h
--- a/Engine/Source/Runtime/Experimental/IoStore/OnDemand/Private/OnDemandIoStore.h
+++ b/Engine/Source/Runtime/Experimental/IoStore/OnDemand/Private/OnDemandIoStore.h
@@ -8,10 +8,11 @@
 #include "Containers/AnsiString.h"
 #include "Containers/BitArray.h"
 #include "IO/IoBuffer.h"
-#include "IO/IoStoreOnDemand.h"
-#include "IO/IoHash.h"
 #include "IO/IoChunkEncoding.h"
 #include "IO/IoChunkId.h"
+#include "IO/IoHash.h"
+#include "IO/IoStoreOnDemand.h"
+#include "IasHostGroup.h"
 #include "Misc/AES.h"
 #include "Misc/EnumClassFlags.h"
 
@@ -150,6 +151,7 @@ struct FOnDemandContainer
 	FAES::FAESKey							EncryptionKey;
 	FSharedContainerHeader					Header;
 	FOnDemandHostGroup						HostGroup;
+	FIASHostGroup							IasHostGroup;
 	FString									EncryptionKeyGuid;
 	FString									Name;
 	FString									MountId;
@@ -248,7 +250,7 @@ struct FOnDemandChunkInfo
 	inline TConstArrayView<FIoBlockHash>	BlockHashes() const;
 	FAnsiStringView							ChunksDirectory() const { return SharedContainer->ChunksDirectory; }
 	const FOnDemandChunkEntry&				ChunkEntry() const { return Entry; }
-	inline FOnDemandHostGroup				HostGroup() const;
+	inline FIASHostGroup					HostGroup() const;
 
 	bool									IsValid() const { return SharedContainer.IsValid(); }
 	operator								bool() const { return IsValid(); }
@@ -279,9 +281,9 @@ TConstArrayView<FIoBlockHash> FOnDemandChunkInfo::BlockHashes() const
 		: TConstArrayView<FIoBlockHash>(SharedContainer->BlockHashes.GetData() + Entry.BlockOffset, Entry.BlockCount);
 }
 
-FOnDemandHostGroup FOnDemandChunkInfo::HostGroup() const
+FIASHostGroup FOnDemandChunkInfo::HostGroup() const
 {
-	return SharedContainer->HostGroup;
+	return SharedContainer->IasHostGroup;
 }
 
 FOnDemandChunkInfo FOnDemandChunkInfo::Find(FSharedOnDemandContainer Container, const FIoChunkId& ChunkId)
