IAS latency and connection tests now take the test path as an ansi string.

#rnx

- Since the low level http client takes all strings in ansi format we have steady moved the higher level code to that too.
- We now only convert the string once, when it is first applied to the backend rather than every time we try  a test.

diff --git a/Engine/Source/Runtime/Experimental/IoStore/OnDemand/Private/IasHostGroup.cpp b/Engine/Source/Runtime/Experimental/IoStore/OnDemand/Private/IasHostGroup.cpp
--- a/Engine/Source/Runtime/Experimental/IoStore/OnDemand/Private/IasHostGroup.cpp
+++ b/Engine/Source/Runtime/Experimental/IoStore/OnDemand/Private/IasHostGroup.cpp
@@ -205,7 +205,7 @@ void FIASHostGroup::Disconnect()
 	SetPrimaryHost(INDEX_NONE);
 }
 
-FIASHostGroup::EReconnectionResult FIASHostGroup::AttemptReconnection(FStringView TestPath, uint32 TimeoutMs, std::atomic_bool& CancellationToken)
+FIASHostGroup::EReconnectionResult FIASHostGroup::AttemptReconnection(FAnsiStringView TestPath, uint32 TimeoutMs, std::atomic_bool& CancellationToken)
 {
 	if (!IsConnected())
 	{
diff --git a/Engine/Source/Runtime/Experimental/IoStore/OnDemand/Private/IasHostGroup.h b/Engine/Source/Runtime/Experimental/IoStore/OnDemand/Private/IasHostGroup.h
--- a/Engine/Source/Runtime/Experimental/IoStore/OnDemand/Private/IasHostGroup.h
+++ b/Engine/Source/Runtime/Experimental/IoStore/OnDemand/Private/IasHostGroup.h
@@ -43,7 +43,7 @@ public:
 		FailedToConnect,
 	};
 
-	EReconnectionResult AttemptReconnection(FStringView TestPath, uint32 TimeoutMs, std::atomic_bool& CancellationToken);
+	EReconnectionResult AttemptReconnection(FAnsiStringView TestPath, uint32 TimeoutMs, std::atomic_bool& CancellationToken);
 
 	void OnSuccessfulResponse();
 	bool OnFailedResponse();
diff --git a/Engine/Source/Runtime/Experimental/IoStore/OnDemand/Private/LatencyTesting.cpp b/Engine/Source/Runtime/Experimental/IoStore/OnDemand/Private/LatencyTesting.cpp
--- a/Engine/Source/Runtime/Experimental/IoStore/OnDemand/Private/LatencyTesting.cpp
+++ b/Engine/Source/Runtime/Experimental/IoStore/OnDemand/Private/LatencyTesting.cpp
@@ -14,7 +14,7 @@
 namespace UE::IoStore
 {
 
-void LatencyTest(FAnsiStringView InUrl, FStringView InPath, uint32 TimeOutMs, TArrayView<int32> OutResults)
+void LatencyTest(FAnsiStringView InUrl, FAnsiStringView InPath, uint32 TimeOutMs, TArrayView<int32> OutResults)
 {
 	using namespace HTTP;
 
@@ -93,7 +93,7 @@ void LatencyTest(FAnsiStringView InUrl, FStringView InPath, uint32 TimeOutMs, TA
 	}
 }
 
-bool ConnectionTest(FAnsiStringView Url, FStringView Path, uint32 TimeoutMs)
+bool ConnectionTest(FAnsiStringView Url, FAnsiStringView Path, uint32 TimeoutMs)
 {
 	TRACE_CPUPROFILER_EVENT_SCOPE(IasBackend::ConnectionTest);
 
@@ -115,7 +115,7 @@ bool ConnectionTest(FAnsiStringView Url, FStringView Path, uint32 TimeoutMs)
 	}
 }
 
-int32 ConnectionTest(TConstArrayView<FAnsiString> Urls, FStringView Path, uint32 TimeoutMs, std::atomic_bool& bCancel)
+int32 ConnectionTest(TConstArrayView<FAnsiString> Urls, FAnsiStringView Path, uint32 TimeoutMs, std::atomic_bool& bCancel)
 {
 	TRACE_CPUPROFILER_EVENT_SCOPE(IasBackend::ConnectionTest);
 
diff --git a/Engine/Source/Runtime/Experimental/IoStore/OnDemand/Private/LatencyTesting.h b/Engine/Source/Runtime/Experimental/IoStore/OnDemand/Private/LatencyTesting.h
--- a/Engine/Source/Runtime/Experimental/IoStore/OnDemand/Private/LatencyTesting.h
+++ b/Engine/Source/Runtime/Experimental/IoStore/OnDemand/Private/LatencyTesting.h
@@ -15,7 +15,7 @@ namespace UE::IoStore
  * @param TimeoutMs		How long to wait for a response (in ms) before timing out
  * @param OutResults	A pre-sized array to hold the results. One latency test will be made per element in this container.
  */
-void LatencyTest(FAnsiStringView Url, FStringView Path, uint32 TimeoutMs, TArrayView<int32> OutResults);
+void LatencyTest(FAnsiStringView Url, FAnsiStringView Path, uint32 TimeoutMs, TArrayView<int32> OutResults);
 
 /**
  * An easy way to test if a valid connection can be made to a CDN or not. Instead of returning the latency to the CDN this
@@ -27,7 +27,7 @@ void LatencyTest(FAnsiStringView Url, FStringView Path, uint32 TimeoutMs, TArray
  * 
  * @return True if the CDN can be reached and false if it cannot.
  */
-bool ConnectionTest(FAnsiStringView Url, FStringView Path, uint32 TimeoutMs);
+bool ConnectionTest(FAnsiStringView Url, FAnsiStringView Path, uint32 TimeoutMs);
 
 /**
  * Runs latency tests on a list of urls to find the first CDN that can be reached.
@@ -39,6 +39,6 @@ bool ConnectionTest(FAnsiStringView Url, FStringView Path, uint32 TimeoutMs);
  *						attempt to exit as fast as possible, usually when the next latency test has completed.
  * @return				The index of the CDN that was successfully reached. INDEX_NONE will be returned if no CDN could be reached.
  */
-int32 ConnectionTest(TConstArrayView<FAnsiString> Urls, FStringView Path, uint32 TimeoutMs, std::atomic_bool& bCancel);
+int32 ConnectionTest(TConstArrayView<FAnsiString> Urls, FAnsiStringView Path, uint32 TimeoutMs, std::atomic_bool& bCancel);
 
 } // namespace UE::IoStore
diff --git a/Engine/Source/Runtime/Experimental/IoStore/OnDemand/Private/OnDemandIoDispatcherBackend.cpp b/Engine/Source/Runtime/Experimental/IoStore/OnDemand/Private/OnDemandIoDispatcherBackend.cpp
--- a/Engine/Source/Runtime/Experimental/IoStore/OnDemand/Private/OnDemandIoDispatcherBackend.cpp
+++ b/Engine/Source/Runtime/Experimental/IoStore/OnDemand/Private/OnDemandIoDispatcherBackend.cpp
@@ -47,7 +47,7 @@
 /** When enabled we will run some limited testing on start up for issues that are hard to reproduce with normal gameplay */
 #define UE_ENABLE_IAS_TESTING 0 && !UE_BUILD_SHIPPING
 
-/** When enabled it is possible to disable request cancelation via the cvar '' */
+/** When enabled it is possible to disable request cancellation via the cvar '' */
 #define UE_ALLOW_DISABLE_CANCELLING 1 && !UE_BUILD_SHIPPING
 
 namespace UE::IoStore
@@ -851,7 +851,7 @@ public:
 private:
 	bool Resolve(FIoRequestImpl* Request);
 
-	const FString& GetEndpointTestPath() const;
+	const FAnsiString& GetEndpointTestPath() const;
 	void ConditionallyStartBackendThread();
 	void CompleteRequest(FChunkRequest* ChunkRequest);
 	void CompleteCacheRead(FChunkRequest* ChunkRequest);
@@ -882,7 +882,7 @@ private:
 	FDistributedEndpointUrl DistributionUrl;
 	FEventRef DistributedEndpointEvent;
 	FIASHostGroup DefaultHostGroup;
-	FString EndpointTestPath;
+	FAnsiString EndpointTestPath;
 	EThreadPriority CurrentThreadPriority;
 
 	mutable FRWLock Lock;
@@ -903,7 +903,8 @@ FOnDemandIoBackend::FOnDemandIoBackend(const FOnDemandEndpointConfig& Config, FO
 {
 	DefaultHostGroup = FHostGroupManager::Get().Register(FName("DefaultOnDemand")).ConsumeValueOrDie();
 
-	EndpointTestPath = Config.TocPath;
+	EndpointTestPath = StringCast<ANSICHAR>(*Config.TocPath);
+
 	if (Config.DistributionUrl.IsEmpty() == false)
 	{
 		DistributionUrl = { Config.DistributionUrl, Config.FallbackUrl };
@@ -992,7 +993,7 @@ void FOnDemandIoBackend::Shutdown()
 	BackendContext.Reset();
 }
 
-const FString& FOnDemandIoBackend::GetEndpointTestPath() const
+const FAnsiString& FOnDemandIoBackend::GetEndpointTestPath() const
 {
 	return EndpointTestPath;
 }
