Move the CDN test path to be stored in the host group not the backend.

IAS latency and connection tests now take the test path as an ansi string.

#rnx

diff --git a/Engine/Source/Runtime/Experimental/IoStore/OnDemand/Private/IasHostGroup.cpp b/Engine/Source/Runtime/Experimental/IoStore/OnDemand/Private/IasHostGroup.cpp
--- a/Engine/Source/Runtime/Experimental/IoStore/OnDemand/Private/IasHostGroup.cpp
+++ b/Engine/Source/Runtime/Experimental/IoStore/OnDemand/Private/IasHostGroup.cpp
@@ -61,8 +61,9 @@ private:
 struct FIASHostGroup::FImpl
 {
 	FImpl() = default;
-	FImpl(FName InName)
-		: Name(InName)
+	FImpl(FName InName, FAnsiStringView InTestPath)
+		: TestPath(InTestPath)
+		, Name(InName)
 	{
 	}
 	FImpl(FName InName, FOnDemandHostGroup& InHostGroup)
@@ -81,6 +82,7 @@ struct FIASHostGroup::FImpl
 
 	FOnDemandHostGroup HostGroup;
 	FBitWindow HttpErrorHistory;
+	FAnsiString TestPath;
 	FName Name;
 	bool bHttpEnabled = true;
 };
@@ -120,8 +122,8 @@ FIASHostGroup::FIASHostGroup()
 {
 }
 
-FIASHostGroup::FIASHostGroup(FName Name)
-	: Impl(MakeShared<FImpl>(Name))
+FIASHostGroup::FIASHostGroup(FName Name, FAnsiStringView TestPath)
+	: Impl(MakeShared<FImpl>(Name, TestPath))
 {
 }
 
@@ -140,6 +142,11 @@ FName FIASHostGroup::GetName() const
 	return Impl->Name;
 }
 
+const FAnsiString& FIASHostGroup::GetTestPath() const
+{
+	return Impl->TestPath;
+}
+
 bool FIASHostGroup::IsResolved() const
 {
 	return !Impl->HostGroup.IsEmpty();
@@ -205,13 +212,13 @@ void FIASHostGroup::Disconnect()
 	SetPrimaryHost(INDEX_NONE);
 }
 
-FIASHostGroup::EReconnectionResult FIASHostGroup::AttemptReconnection(FAnsiStringView TestPath, uint32 TimeoutMs, std::atomic_bool& CancellationToken)
+FIASHostGroup::EReconnectionResult FIASHostGroup::AttemptReconnection(uint32 TimeoutMs, std::atomic_bool& CancellationToken)
 {
 	if (!IsConnected())
 	{
 		UE_LOGFMT(LogIas, Log, "[{HostName}] Trying to reconnect to any available endpoint...", GetName());
 
-		if (const int32 Idx = ConnectionTest(GetHostUrls(), TestPath, TimeoutMs, CancellationToken); Idx != INDEX_NONE)
+		if (const int32 Idx = ConnectionTest(GetHostUrls(), Impl->TestPath, TimeoutMs, CancellationToken); Idx != INDEX_NONE)
 		{
 			Connect(Idx);
 
@@ -225,7 +232,7 @@ FIASHostGroup::EReconnectionResult FIASHostGroup::AttemptReconnection(FAnsiStrin
 	}
 	else if (GetPrimaryHostIndex() != 0)
 	{
-		if (const int32 Idx = ConnectionTest(GetHostUrls().Left(1), TestPath, TimeoutMs, CancellationToken); Idx != INDEX_NONE)
+		if (const int32 Idx = ConnectionTest(GetHostUrls().Left(1), Impl->TestPath, TimeoutMs, CancellationToken); Idx != INDEX_NONE)
 		{
 			SetPrimaryHost(Idx);
 			UE_LOGFMT(LogIas, Log, "[{HostName}] Reconnected to primary host '{Url}'", GetPrimaryHostUrl(), GetName());
@@ -293,9 +300,9 @@ FHostGroupManager& FHostGroupManager::Get()
 	return Instance;
 }
 
-TIoStatusOr<FIASHostGroup> FHostGroupManager::Register(FName Name)
+TIoStatusOr<FIASHostGroup> FHostGroupManager::Register(FName Name, FAnsiStringView TestPath)
 {
-	FIASHostGroup HostGroup(Name);
+	FIASHostGroup HostGroup(Name, TestPath);
 
 	HostGroups.Add(HostGroup);
 
diff --git a/Engine/Source/Runtime/Experimental/IoStore/OnDemand/Private/IasHostGroup.h b/Engine/Source/Runtime/Experimental/IoStore/OnDemand/Private/IasHostGroup.h
--- a/Engine/Source/Runtime/Experimental/IoStore/OnDemand/Private/IasHostGroup.h
+++ b/Engine/Source/Runtime/Experimental/IoStore/OnDemand/Private/IasHostGroup.h
@@ -21,11 +21,12 @@ public:
 	static TIoStatusOr<FIASHostGroup> Create(FName Name, TConstArrayView<FString> HostUrls);
 
 	FIASHostGroup();
-	FIASHostGroup(FName Name);
+	FIASHostGroup(FName Name, FAnsiStringView TestPath);
 	FIASHostGroup(FName Name, FOnDemandHostGroup&& HostGroup);
 	~FIASHostGroup() = default;
 
 	FName GetName() const;
+	const FAnsiString& GetTestPath() const;
 
 	bool IsResolved() const;
 	bool IsConnected() const;
@@ -43,7 +44,7 @@ public:
 		FailedToConnect,
 	};
 
-	EReconnectionResult AttemptReconnection(FAnsiStringView TestPath, uint32 TimeoutMs, std::atomic_bool& CancellationToken);
+	EReconnectionResult AttemptReconnection(uint32 TimeoutMs, std::atomic_bool& CancellationToken);
 
 	void OnSuccessfulResponse();
 	bool OnFailedResponse();
@@ -74,7 +75,7 @@ public:
 
 	// TODO: Note that currently we have no way to remove FIASHostGroups
 
-	TIoStatusOr<FIASHostGroup> Register(FName Name); // <- Unresolved
+	TIoStatusOr<FIASHostGroup> Register(FName Name, FAnsiStringView TestPath); // <- Unresolved
 	TIoStatusOr<FIASHostGroup> Register(FName Name, TConstArrayView<FAnsiString> HostUrls); // <- Resolved
 
 	FIASHostGroup				Find(FName Name);
diff --git a/Engine/Source/Runtime/Experimental/IoStore/OnDemand/Private/OnDemandIoDispatcherBackend.cpp b/Engine/Source/Runtime/Experimental/IoStore/OnDemand/Private/OnDemandIoDispatcherBackend.cpp
--- a/Engine/Source/Runtime/Experimental/IoStore/OnDemand/Private/OnDemandIoDispatcherBackend.cpp
+++ b/Engine/Source/Runtime/Experimental/IoStore/OnDemand/Private/OnDemandIoDispatcherBackend.cpp
@@ -851,7 +851,6 @@ public:
 private:
 	bool Resolve(FIoRequestImpl* Request);
 
-	const FAnsiString& GetEndpointTestPath() const;
 	void ConditionallyStartBackendThread();
 	void CompleteRequest(FChunkRequest* ChunkRequest);
 	void CompleteCacheRead(FChunkRequest* ChunkRequest);
@@ -882,7 +881,7 @@ private:
 	FDistributedEndpointUrl DistributionUrl;
 	FEventRef DistributedEndpointEvent;
 	FIASHostGroup DefaultHostGroup;
-	FAnsiString EndpointTestPath;
+
 	EThreadPriority CurrentThreadPriority;
 
 	mutable FRWLock Lock;
@@ -901,9 +900,9 @@ FOnDemandIoBackend::FOnDemandIoBackend(const FOnDemandEndpointConfig& Config, FO
 	, Stats(BackendStatus)
 	, CurrentThreadPriority(EThreadPriority::TPri_Num)
 {
-	DefaultHostGroup = FHostGroupManager::Get().Register(FName("DefaultOnDemand")).ConsumeValueOrDie();
+	FAnsiString EndpointTestPath = StringCast<ANSICHAR>(*Config.TocPath).Get();
 
-	EndpointTestPath = StringCast<ANSICHAR>(*Config.TocPath);
+	DefaultHostGroup = FHostGroupManager::Get().Register(FName("DefaultOnDemand"), EndpointTestPath).ConsumeValueOrDie();
 
 	if (Config.DistributionUrl.IsEmpty() == false)
 	{
@@ -993,11 +992,6 @@ void FOnDemandIoBackend::Shutdown()
 	BackendContext.Reset();
 }
 
-const FAnsiString& FOnDemandIoBackend::GetEndpointTestPath() const
-{
-	return EndpointTestPath;
-}
-
 void FOnDemandIoBackend::UpdateThreadPriorityIfNeeded()
 {
 	// Read the thread priority from the cvar
@@ -1673,7 +1667,7 @@ uint32 FOnDemandIoBackend::Run()
 
 	const double InitStart = FPlatformTime::Seconds();
 
-	if (ConnectionTest(DefaultHostGroup.GetPrimaryHostUrl(), GetEndpointTestPath(), GIasHttpTimeOutMs))
+	if (ConnectionTest(DefaultHostGroup.GetPrimaryHostUrl(), DefaultHostGroup.GetTestPath(), GIasHttpTimeOutMs))
 	{
 		BackendStatus.SetHttpEnabled(true);
 		Stats.OnHttpConnected();
@@ -1694,7 +1688,7 @@ uint32 FOnDemandIoBackend::Run()
 
 		if (!bStopRequested)
 		{
-			if (DefaultHostGroup.AttemptReconnection(GetEndpointTestPath(), GIasHttpTimeOutMs, bStopRequested) == FIASHostGroup::EReconnectionResult::Reconnected)
+			if (DefaultHostGroup.AttemptReconnection(GIasHttpTimeOutMs, bStopRequested) == FIASHostGroup::EReconnectionResult::Reconnected)
 			{
 				Stats.OnHttpConnected();
 			}
