Tolerate cases where more data could be received than asked for.

#rnx
#megacompiled

diff --git a/Engine/Source/Runtime/Experimental/IoStore/HttpClient/Private/Loop.inl b/Engine/Source/Runtime/Experimental/IoStore/HttpClient/Private/Loop.inl
--- a/Engine/Source/Runtime/Experimental/IoStore/HttpClient/Private/Loop.inl
+++ b/Engine/Source/Runtime/Experimental/IoStore/HttpClient/Private/Loop.inl
@@ -89,9 +89,8 @@ static FOutcome DoRecvPeer(FActivity* Activity, FHttpPeer& Peer, int32& MaxRecvS
 	check(Outcome.IsOk());
 
 	int32 Result = Outcome.GetResult();
-	check(Result <= MaxRecvSize);
 	Activity->StateParam += Result;
-	MaxRecvSize -= Result;
+	MaxRecvSize = FMath::Max(MaxRecvSize - Result, 0);
 
 	return Outcome;
 }
