Replace IOnDemandIoStore::GetStreamingOptions with ::IsOnDemandStreamingEnabled.

#rnx

- ::GetStreamingOptions was only ever being called to check if the http system was enabled or not. Since we want to change how this flag works internally it is easier to replace the method with one that checks for the specific condition that we want.

diff --git a/Engine/Plugins/Experimental/RuntimeTelemetry/Source/Private/RuntimeTelemetry.cpp b/Engine/Plugins/Experimental/RuntimeTelemetry/Source/Private/RuntimeTelemetry.cpp
--- a/Engine/Plugins/Experimental/RuntimeTelemetry/Source/Private/RuntimeTelemetry.cpp
+++ b/Engine/Plugins/Experimental/RuntimeTelemetry/Source/Private/RuntimeTelemetry.cpp
@@ -28,7 +28,7 @@ void FRuntimeTelemetry::RecordEvent_IoStoreOnDemand(const FString& Context, TArr
 		if (UE::IoStore::IOnDemandIoStore* IoStore = UE::IoStore::TryGetOnDemandIoStore())
 		{
 			// Only record event data if IoStoreOnDemand is enabled
-			if (EnumHasAnyFlags(IoStore->GetStreamingOptions(), EOnDemandStreamingOptions::Enabled))
+			if (IoStore->IsOnDemandStreamingEnabled())
 			{
 				const int SchemaVersion = 1;
 
diff --git a/Engine/Source/Runtime/Experimental/IoStore/OnDemand/Private/OnDemandIoStore.cpp b/Engine/Source/Runtime/Experimental/IoStore/OnDemand/Private/OnDemandIoStore.cpp
--- a/Engine/Source/Runtime/Experimental/IoStore/OnDemand/Private/OnDemandIoStore.cpp
+++ b/Engine/Source/Runtime/Experimental/IoStore/OnDemand/Private/OnDemandIoStore.cpp
@@ -2083,6 +2083,11 @@ FOnDemandCacheUsage FOnDemandIoStore::GetCacheUsage() const
 	return CacheUsage;
 }
 
+bool FOnDemandIoStore::IsOnDemandStreamingEnabled() const
+{
+	return StreamingBackend != nullptr;
+}
+
 void FOnDemandIoStore::SetStreamingOptions(EOnDemandStreamingOptions Options) 
 {
 	StreamingOptions = Options;
diff --git a/Engine/Source/Runtime/Experimental/IoStore/OnDemand/Private/OnDemandIoStore.h b/Engine/Source/Runtime/Experimental/IoStore/OnDemand/Private/OnDemandIoStore.h
--- a/Engine/Source/Runtime/Experimental/IoStore/OnDemand/Private/OnDemandIoStore.h
+++ b/Engine/Source/Runtime/Experimental/IoStore/OnDemand/Private/OnDemandIoStore.h
@@ -352,7 +352,7 @@ public:
 	virtual TIoStatusOr<uint64> 		GetInstallSize(const FOnDemandGetInstallSizeArgs& Args) const override;
 	virtual FIoStatus					GetInstallSizesByMountId(const FOnDemandGetInstallSizeArgs& Args, TMap<FString, uint64>& OutSizesByMountId) const override;
 	virtual FOnDemandCacheUsage			GetCacheUsage() const override;
-	virtual EOnDemandStreamingOptions	GetStreamingOptions() const override { return StreamingOptions; }
+	virtual bool						IsOnDemandStreamingEnabled() const override;
 	virtual void						SetStreamingOptions(EOnDemandStreamingOptions Options) override;
 	virtual void						ReportAnalytics(TArray<FAnalyticsEventAttribute>& OutAnalyticsArray) const override;
 	virtual TUniquePtr<IAnalyticsRecording> StartAnalyticsRecording() const override;
diff --git a/Engine/Source/Runtime/Experimental/IoStore/OnDemand/Public/IO/IoStoreOnDemand.h b/Engine/Source/Runtime/Experimental/IoStore/OnDemand/Public/IO/IoStoreOnDemand.h
--- a/Engine/Source/Runtime/Experimental/IoStore/OnDemand/Public/IO/IoStoreOnDemand.h
+++ b/Engine/Source/Runtime/Experimental/IoStore/OnDemand/Public/IO/IoStoreOnDemand.h
@@ -449,8 +449,8 @@ public:
 	virtual FIoStatus GetInstallSizesByMountId(const FOnDemandGetInstallSizeArgs& Args, TMap<FString, uint64>& OutSizesByMountId) const = 0;
 	/** Returns the total cache size in bytes. */
 	virtual FOnDemandCacheUsage GetCacheUsage() const = 0;
-	/** Return the curren streaming optinos. */
-	virtual EOnDemandStreamingOptions GetStreamingOptions() const = 0;
+	/** Return if the OnDemand streaming system is enabled. */
+	virtual bool IsOnDemandStreamingEnabled() const = 0;
 	/** Set streaming options.*/
 	virtual void SetStreamingOptions(EOnDemandStreamingOptions Options) = 0;
 	/**
