[Chaos][PhysicsProxyBase]

Fixing a NaN that will happen when ErrorX is zero because the ProjectOnto() call will divide by zero.

-- Code --
PhysicsProxyBase:
FProxyInterpolationError::DirectionalDecay() - Adding an early exit if ErrorX is nearly zero.

#tests Ran a project that was seeing this NaN consistently to make sure it stopped.

diff --git a/Engine/Source/Runtime/Experimental/Chaos/Public/Chaos/Framework/PhysicsProxyBase.h b/Engine/Source/Runtime/Experimental/Chaos/Public/Chaos/Framework/PhysicsProxyBase.h
--- a/Engine/Source/Runtime/Experimental/Chaos/Public/Chaos/Framework/PhysicsProxyBase.h
+++ b/Engine/Source/Runtime/Experimental/Chaos/Public/Chaos/Framework/PhysicsProxyBase.h
@@ -340,6 +340,11 @@ struct FProxyInterpolationError : FProxyInterpolationBase
 
 	virtual const bool DirectionalDecay(Chaos::FVec3 Direction, float ErrorDirectionalDecayMultiplier) override
 	{
+		if (ErrorX.IsNearlyZero())
+		{
+			return false;
+		}
+
 		ErrorDirectionalDecayMultiplier = ErrorInterpolationSettings.Get() ? ErrorInterpolationSettings.Get()->ErrorDirectionalDecayMultiplier : ErrorDirectionalDecayMultiplier;
 		if (ErrorDirectionalDecayMultiplier > 0.0f && IsErrorSmoothing() && SimTicks > 0)
 		{
