Dataflow selection : add group name to selection objects

diff --git a/Engine/Source/Runtime/Experimental/Dataflow/Core/Public/Dataflow/DataflowSelection.h b/Engine/Source/Runtime/Experimental/Dataflow/Core/Public/Dataflow/DataflowSelection.h
--- a/Engine/Source/Runtime/Experimental/Dataflow/Core/Public/Dataflow/DataflowSelection.h
+++ b/Engine/Source/Runtime/Experimental/Dataflow/Core/Public/Dataflow/DataflowSelection.h
@@ -2,6 +2,8 @@
 
 #pragma once
 
+
+#include "UObject/NameTypes.h"
 #include "DataflowSelection.generated.h"
 
 USTRUCT()
@@ -9,6 +11,13 @@ struct FDataflowSelection
 {
 	GENERATED_USTRUCT_BODY()
 
+	FDataflowSelection()
+	{}
+
+	FDataflowSelection(FName InGroupName)
+		: GroupName(InGroupName)
+	{}
+
 	DATAFLOWCORE_API void Initialize(int32 NumBits, bool Value);
 	DATAFLOWCORE_API void Initialize(const FDataflowSelection& Other);
 	DATAFLOWCORE_API void Clear();
@@ -42,7 +51,7 @@ struct FDataflowSelection
 	DATAFLOWCORE_API const TBitArray<>& GetBitArray() const { return SelectionArray; };
 
 private:
-
+	FName GroupName;
 	TBitArray<> SelectionArray;
 };
 
@@ -50,6 +59,12 @@ USTRUCT()
 struct FDataflowTransformSelection : public FDataflowSelection
 {
 	GENERATED_USTRUCT_BODY()
+
+	inline static const FName TransformGroupName = "Transform";
+
+	FDataflowTransformSelection()
+		: FDataflowSelection(TransformGroupName)
+	{}
 };
 
 
@@ -57,6 +72,12 @@ USTRUCT()
 struct FDataflowVertexSelection : public FDataflowSelection
 {
 	GENERATED_USTRUCT_BODY()
+
+	inline static const FName VerticesGroupName = "Vertices";
+
+	FDataflowVertexSelection()
+		: FDataflowSelection(VerticesGroupName)
+	{}
 };
 
 
@@ -64,16 +85,34 @@ USTRUCT()
 struct FDataflowFaceSelection : public FDataflowSelection
 {
 	GENERATED_USTRUCT_BODY()
+
+	inline static const FName FacesGroupName = "Faces";
+
+	FDataflowFaceSelection()
+		: FDataflowSelection(FacesGroupName)
+	{}
 };
 
 USTRUCT()
 struct FDataflowGeometrySelection : public FDataflowSelection
 {
 	GENERATED_USTRUCT_BODY()
+
+	inline static const FName GeometryGroupName = "Geometry";
+
+	FDataflowGeometrySelection()
+		: FDataflowSelection(GeometryGroupName)
+	{}
 };
 
 USTRUCT()
 struct FDataflowMaterialSelection : public FDataflowSelection
 {
 	GENERATED_USTRUCT_BODY()
+
+	inline static const FName MaterialGroupName = "Material";
+
+	FDataflowMaterialSelection()
+		: FDataflowSelection(MaterialGroupName)
+	{}
 };
