[NetTokens] Bring NetTokens to modules outside of Iris/NetCore.

- Config managed type_id
- Callbacks enabling other modules to register providers
- Changes to support struct serialization
- Test changes to support these

diff --git a/Engine/Config/BaseEngine.ini b/Engine/Config/BaseEngine.ini
--- a/Engine/Config/BaseEngine.ini
+++ b/Engine/Config/BaseEngine.ini
@@ -1482,6 +1482,11 @@ bEnsureAllClassesAreFullyPushModel=false
 +SupportsStructNetSerializerList=(StructName=RootMotionSource_MoveToForce)
 +SupportsStructNetSerializerList=(StructName=RootMotionSource_RadialForce)
 
+[/Script/IrisCore.NetTokenTypeIdConfig]
++ReservedTypeIds=(StoreTypeName="StringTokenStore", TypeID=0)
++ReservedTypeIds=(StoreTypeName="NameTokenStore", TypeID=1)
++ReservedTypeIds=(StoreTypeName="GameplayTagTokenStore", TypeID=2)
+
 ;Iris - end
 
 [/Script/Engine.GameEngine]
diff --git a/Engine/Plugins/Runtime/ReplicationSystemTestPlugin/Source/Private/TestMessage.cpp b/Engine/Plugins/Runtime/ReplicationSystemTestPlugin/Source/Private/TestMessage.cpp
--- a/Engine/Plugins/Runtime/ReplicationSystemTestPlugin/Source/Private/TestMessage.cpp
+++ b/Engine/Plugins/Runtime/ReplicationSystemTestPlugin/Source/Private/TestMessage.cpp
@@ -4,6 +4,7 @@
 #include "Iris/ReplicationSystem/NetRefHandle.h"
 #include "Misc/StringBuilder.h"
 #include "Net/Core/NetHandle/NetHandle.h"
+#include "Net/Core/NetToken/NetToken.h"
 
 namespace UE::Net
 {
@@ -22,4 +23,9 @@ FTestMessage& operator<<(FTestMessage& Message, const FNetRefHandle& NetRefHandl
 	return Message << StringBuilder;
 }
 
+FTestMessage& operator<<(FTestMessage& Message, const FNetToken& NetToken)
+{
+	return Message << NetToken.ToString();
+}
+
 }
diff --git a/Engine/Plugins/Runtime/ReplicationSystemTestPlugin/Source/Private/TestMessage.h b/Engine/Plugins/Runtime/ReplicationSystemTestPlugin/Source/Private/TestMessage.h
--- a/Engine/Plugins/Runtime/ReplicationSystemTestPlugin/Source/Private/TestMessage.h
+++ b/Engine/Plugins/Runtime/ReplicationSystemTestPlugin/Source/Private/TestMessage.h
@@ -11,6 +11,7 @@ namespace UE::Net
 {
 	class FNetHandle;
 	class FNetRefHandle;
+	class FNetToken;
 }
 
 namespace UE::Net
@@ -141,6 +142,8 @@ FTestMessage& operator<<(FTestMessage& Message, const FNetHandle& NetHandle);
 
 FTestMessage& operator<<(FTestMessage& Message, const FNetRefHandle& NetRefHandle);
 
+FTestMessage& operator<<(FTestMessage& Message, const FNetToken& NetToken);
+
 #if defined(_WIN32) && !defined(_WIN64)
 inline FTestMessage& operator<<(FTestMessage& Message, SIZE_T Value)
 {
diff --git a/Engine/Plugins/Runtime/ReplicationSystemTestPlugin/Source/Private/Tests/ReplicationSystem/ReplicationSystemServerClientTestFixture.cpp b/Engine/Plugins/Runtime/ReplicationSystemTestPlugin/Source/Private/Tests/ReplicationSystem/ReplicationSystemServerClientTestFixture.cpp
--- a/Engine/Plugins/Runtime/ReplicationSystemTestPlugin/Source/Private/Tests/ReplicationSystem/ReplicationSystemServerClientTestFixture.cpp
+++ b/Engine/Plugins/Runtime/ReplicationSystemTestPlugin/Source/Private/Tests/ReplicationSystem/ReplicationSystemServerClientTestFixture.cpp
@@ -10,11 +10,55 @@
 #include "Iris/ReplicationSystem/ReplicationWriter.h"
 #include "Iris/ReplicationSystem/ReplicationReader.h"
 #include "Iris/PacketControl/PacketNotification.h"
+#include "Iris/ReplicationSystem/NetTokenStore.h"
 #include "Net/Core/Trace/NetTrace.h"
 
 namespace UE::Net
 {
 
+void FNetTokenDataStoreTestUtil::SetUp()
+{
+	StoreNetTokenStoreConfig();
+	AddNetTokenStoreTypeIdPair(TEXT("StringTokenStore"), 0);
+	AddNetTokenStoreTypeIdPair(TEXT("NameTokenStore"), 1);
+}
+
+void FNetTokenDataStoreTestUtil::TearDown()
+{
+	RestoreNetTokenStoreConfig();
+}
+
+void FNetTokenDataStoreTestUtil::AddNetTokenStoreTypeIdPair(FString StoreTypeName, uint32 TypeID)
+{
+	bool bFound = false;
+	for (const auto& Elem :  NetTokenTypeIdConfig->ReservedTypeIds)
+	{
+		if (Elem.StoreTypeName == StoreTypeName)
+		{
+			bFound = true;
+			break;
+		}
+	}
+	if (!bFound)
+	{
+		FNetTokenStoreTypeIdPair TypeIdPair = {.StoreTypeName = StoreTypeName, .TypeID = TypeID};
+		NetTokenTypeIdConfig->ReservedTypeIds.Add(TypeIdPair);
+	}
+}
+
+void FNetTokenDataStoreTestUtil::StoreNetTokenStoreConfig()
+{
+	GetDefault<UNetTokenTypeIdConfig>();
+	NetTokenTypeIdConfig = static_cast<UNetTokenTypeIdConfig*>(GetMutableDefault<UNetTokenTypeIdConfig>());
+	check(NetTokenTypeIdConfig != nullptr);
+	OriginalReservedTypeIds = NetTokenTypeIdConfig->ReservedTypeIds;
+}
+
+void FNetTokenDataStoreTestUtil::RestoreNetTokenStoreConfig()
+{
+	Swap(NetTokenTypeIdConfig->ReservedTypeIds, OriginalReservedTypeIds);
+}
+
 // FDataStreamTestUtil implementation
 FDataStreamTestUtil::FDataStreamTestUtil()
 : DataStreamDefinitions(nullptr)
@@ -81,6 +125,8 @@ void FReplicationSystemTestNode::Setup(bool bIsServer, const TCHAR* Name, FRepli
 	{
 		using namespace UE::Net;
 
+		NetTokenDataStoreUtil.SetUp();
+		
 		NetTokenStore = MakeUnique<FNetTokenStore>();
 
 		FNetTokenStore::FInitParams NetTokenStoreInitParams;
@@ -127,6 +173,7 @@ FReplicationSystemTestNode::~FReplicationSystemTestNode()
 	LOG_SCOPE_VERBOSITY_OVERRIDE(LogIris, ELogVerbosity::Error);
 	FReplicationSystemFactory::DestroyReplicationSystem(ReplicationSystem);
 	CreatedObjects.Empty();
+	NetTokenDataStoreUtil.TearDown();
 
 	// End NetTrace session for this instance
 	UE_NET_TRACE_END_SESSION(NetTraceId);
@@ -532,6 +579,8 @@ void FReplicationSystemServerClientTestFixture::SetUp()
 	DataStreamUtil.AddDataStreamDefinition(TEXT("Replication"), TEXT("/Script/IrisCore.ReplicationDataStream"));
 	DataStreamUtil.FixupDefinitions();
 
+	NetTokenDataStoreUtil.SetUp();
+
 	Server = new FReplicationSystemTestServer(GetName());
 }
 
@@ -543,6 +592,7 @@ void FReplicationSystemServerClientTestFixture::TearDown()
 	}
 	delete Server;
 	DataStreamUtil.TearDown();
+	NetTokenDataStoreUtil.TearDown();
 
 	FNetworkAutomationTestSuiteFixture::TearDown();
 }
diff --git a/Engine/Plugins/Runtime/ReplicationSystemTestPlugin/Source/Private/Tests/ReplicationSystem/ReplicationSystemServerClientTestFixture.h b/Engine/Plugins/Runtime/ReplicationSystemTestPlugin/Source/Private/Tests/ReplicationSystem/ReplicationSystemServerClientTestFixture.h
--- a/Engine/Plugins/Runtime/ReplicationSystemTestPlugin/Source/Private/Tests/ReplicationSystem/ReplicationSystemServerClientTestFixture.h
+++ b/Engine/Plugins/Runtime/ReplicationSystemTestPlugin/Source/Private/Tests/ReplicationSystem/ReplicationSystemServerClientTestFixture.h
@@ -11,6 +11,7 @@
 
 #include "Iris/DataStream/DataStreamDefinitions.h"
 #include "Iris/DataStream/DataStreamManager.h"
+#include "Iris/ReplicationSystem/NetTokenStore.h"
 #include "Iris/ReplicationSystem/NetTokenDataStream.h"
 #include "Iris/ReplicationSystem/ReplicationDataStream.h"
 #include "Iris/ReplicationSystem/ReplicationSystem.h"
@@ -55,6 +56,20 @@ protected:
 	bool* PointerToFixupComplete;
 };
 
+class FNetTokenDataStoreTestUtil
+{
+public:
+	void SetUp();
+	void TearDown();
+	void StoreNetTokenStoreConfig();
+	void RestoreNetTokenStoreConfig();
+	void AddNetTokenStoreTypeIdPair(FString StoreTypeName, uint32 TypeID);
+
+private:
+	UNetTokenTypeIdConfig* NetTokenTypeIdConfig = nullptr;
+	TArray<FNetTokenStoreTypeIdPair> OriginalReservedTypeIds;
+};
+
 // Implements everything we need to drive the replication system to test the system
 class FReplicationSystemTestNode
 {
@@ -185,6 +200,7 @@ public:
 	EReplicationSystemSendPass CurrentSendPass = EReplicationSystemSendPass::Invalid;
 
 private:
+	FNetTokenDataStoreTestUtil NetTokenDataStoreUtil;
 	TArray<FConnectionInfo> Connections;
 	uint32 PacketId = 0U;
 };
@@ -240,6 +256,7 @@ protected:
 	void DestroyClient(FReplicationSystemTestClient* Client);
 
 	FDataStreamTestUtil DataStreamUtil;
+	FNetTokenDataStoreTestUtil NetTokenDataStoreUtil;
 	FReplicationSystemTestServer* Server;
 	TArray<FReplicationSystemTestClient*> Clients;
 };
diff --git a/Engine/Plugins/Runtime/ReplicationSystemTestPlugin/Source/Private/Tests/ReplicationSystem/ReplicationSystemTestFixture.h b/Engine/Plugins/Runtime/ReplicationSystemTestPlugin/Source/Private/Tests/ReplicationSystem/ReplicationSystemTestFixture.h
--- a/Engine/Plugins/Runtime/ReplicationSystemTestPlugin/Source/Private/Tests/ReplicationSystem/ReplicationSystemTestFixture.h
+++ b/Engine/Plugins/Runtime/ReplicationSystemTestPlugin/Source/Private/Tests/ReplicationSystem/ReplicationSystemTestFixture.h
@@ -5,6 +5,7 @@
 #include "NetworkAutomationTest.h"
 #include "NetworkAutomationTestMacros.h"
 #include "ReplicatedTestObject.h"
+#include "ReplicationSystemServerClientTestFixture.h"
 #include "Containers/Array.h"
 #include "Templates/UniquePtr.h"
 #include "Iris/Core/IrisLog.h"
@@ -35,6 +36,8 @@ protected:
 		{
 			using namespace UE::Net;
 
+			NetTokenDataStoreUtil.SetUp();
+			
 			NetTokenStore = MakeUnique<FNetTokenStore>();
 
 			FNetTokenStore::FInitParams NetTokenStoreInitParams;
@@ -71,6 +74,7 @@ protected:
 		FReplicationSystemFactory::DestroyReplicationSystem(ReplicationSystem);
 		LogIris.SetVerbosity(IrisLogVerbosity);
 		CreatedObjects.Empty();
+		NetTokenDataStoreUtil.TearDown();
 	}
 
 	// Creates a test object without components
@@ -124,6 +128,7 @@ protected:
 		return CreatedObject;
 	}
 
+	FNetTokenDataStoreTestUtil NetTokenDataStoreUtil;
 	TUniquePtr<UE::Net::FNetTokenStore> NetTokenStore;
 	UReplicationSystem* ReplicationSystem;
 	UReplicatedTestObjectBridge* ReplicationBridge;
diff --git a/Engine/Source/Runtime/Engine/Classes/Engine/NetDriver.h b/Engine/Source/Runtime/Engine/Classes/Engine/NetDriver.h
--- a/Engine/Source/Runtime/Engine/Classes/Engine/NetDriver.h
+++ b/Engine/Source/Runtime/Engine/Classes/Engine/NetDriver.h
@@ -1983,6 +1983,17 @@ public:
 	/** Returns NetTokenStore that is required to create and serialize NetTokens */
 	const UE::Net::FNetTokenStore* GetNetTokenStore() const { return NetTokenStore.Get(); }
 	UE::Net::FNetTokenStore* GetNetTokenStore() { return NetTokenStore.Get(); }
+	
+	DECLARE_MULTICAST_DELEGATE_OneParam(FOnNetTokenStoreReady, UNetDriver*);
+private:
+	FOnNetTokenStoreReady OnNetTokenStoreReadyDelegate;
+public:
+	/** Fired when the NetTokenStore is ready for use and has been initialized, but before it will be used for any NetSerializations */
+	FOnNetTokenStoreReady::RegistrationType& OnNetTokenStoreReady()
+	{
+		return OnNetTokenStoreReadyDelegate;
+	}
+
 #else
 	/** Returns NetTokenStore that is required to create and serialize NetTokens */
 	const UE::Net::FNetTokenStore* GetNetTokenStore() const { return nullptr; }
diff --git a/Engine/Source/Runtime/Engine/Private/NetDriver.cpp b/Engine/Source/Runtime/Engine/Private/NetDriver.cpp
--- a/Engine/Source/Runtime/Engine/Private/NetDriver.cpp
+++ b/Engine/Source/Runtime/Engine/Private/NetDriver.cpp
@@ -1637,6 +1637,7 @@ bool UNetDriver::InitBase(bool bInitAsClient, FNetworkNotify* InNotify, const FU
 		NetTokenStore->CreateAndRegisterDataStore<FStringTokenStore>();
 		NetTokenStore->CreateAndRegisterDataStore<FNameTokenStore>();
 		NetTokenStore->CreateAndRegisterDataStore<FGameplayTagTokenStore>();
+		OnNetTokenStoreReadyDelegate.Broadcast(this);
 	}
 #endif
 	
diff --git a/Engine/Source/Runtime/Experimental/Iris/Core/Private/Iris/ReplicationSystem/NameTokenStore.cpp b/Engine/Source/Runtime/Experimental/Iris/Core/Private/Iris/ReplicationSystem/NameTokenStore.cpp
--- a/Engine/Source/Runtime/Experimental/Iris/Core/Private/Iris/ReplicationSystem/NameTokenStore.cpp
+++ b/Engine/Source/Runtime/Experimental/Iris/Core/Private/Iris/ReplicationSystem/NameTokenStore.cpp
@@ -100,7 +100,7 @@ void FNameTokenStore::WriteTokenData(FNetSerializationContext& Context, FNetToke
 	WriteString(Context.GetBitStreamWriter(), StoredFNames[TokenStoreKey.GetKeyIndex()].ToString());
 }
 
-void FNameTokenStore::WriteTokenData(FArchive& Ar, FNetTokenStoreKey TokenStoreKey) const
+void FNameTokenStore::WriteTokenData(FArchive& Ar, FNetTokenStoreKey TokenStoreKey, UPackageMap* Map /*= nullptr*/) const
 {
 	UE_NET_TRACE_DYNAMIC_NAME_SCOPE(StoredFNames[TokenStoreKey.GetKeyIndex()], static_cast<FNetBitWriter&>(Ar), GetTraceCollector(static_cast<FNetBitWriter&>(Ar)), ENetTraceVerbosity::VeryVerbose);
 	UE_LOG_FNAMETOKEN(TEXT("FNameTokenStore::WriteTokenData %s %s"), *(StoredTokens[TokenStoreKey.GetKeyIndex()].ToString()), *(StoredFNames[TokenStoreKey.GetKeyIndex()].ToString()));
@@ -131,7 +131,7 @@ FNetTokenDataStore::FNetTokenStoreKey FNameTokenStore::ReadTokenData(FNetSeriali
 	}
 }
 
-FNetTokenDataStore::FNetTokenStoreKey FNameTokenStore::ReadTokenData(FArchive& Ar, const FNetToken& NetToken)
+FNetTokenDataStore::FNetTokenStoreKey FNameTokenStore::ReadTokenData(FArchive& Ar, const FNetToken& NetToken, UPackageMap* Map /*= nullptr*/)
 {
 	FName Name;
 	UPackageMap::StaticSerializeName(Ar, Name);
diff --git a/Engine/Source/Runtime/Experimental/Iris/Core/Private/Iris/ReplicationSystem/NetTokenStore.cpp b/Engine/Source/Runtime/Experimental/Iris/Core/Private/Iris/ReplicationSystem/NetTokenStore.cpp
--- a/Engine/Source/Runtime/Experimental/Iris/Core/Private/Iris/ReplicationSystem/NetTokenStore.cpp
+++ b/Engine/Source/Runtime/Experimental/Iris/Core/Private/Iris/ReplicationSystem/NetTokenStore.cpp
@@ -13,6 +13,8 @@
 
 namespace UE::Net
 {
+class FNetSerializationContext;
+	
 
 class FNetTokenStoreState
 {
@@ -303,18 +305,60 @@ bool FNetTokenStore::RegisterDataStore(TUniquePtr<FNetTokenDataStore> DataStore,
 		return false;
 	}
 
-	if (!ensure(DataStore->GetTypeId() == FNetToken::InvalidTokenTypeId))
+	const UNetTokenTypeIdConfig* TypeIDConfig = GetDefault<UNetTokenTypeIdConfig>();
+	const uint32 TypeId = TypeIDConfig->GetTypeID(TokenStoreName);
+	DataStore->TypeId = TypeId;
+	if (!ensureAlwaysMsgf(DataStore->TypeId != FNetToken::InvalidTokenTypeId, TEXT("No TypeID information found in [/Script/IrisCore.NetTokenTypeIdConfig] Engine.ini for %s"), *TokenStoreName.ToString()))
 	{
-		// Already registered
+		UE_LOG(LogNetToken, Error, TEXT("No TypeID information found in [/Script/IrisCore.NetTokenTypeIdConfig] Engine.ini for %s"), *TokenStoreName.ToString());
 		return false;
 	}
 
-	DataStore->TypeId = TokenDataStores.Num();
-	TokenDataStores.Emplace(TokenStoreName, MoveTemp(DataStore));
+	if (!ensure(TypeId != FNetToken::InvalidTokenTypeId && TypeId < FNetToken::MaxTypeIdCount))
+	{
+		UE_LOG(LogNetToken, Error, TEXT("Invalid TypeID information found for %s"), *TokenStoreName.ToString());
+		return false;
+	}
+
+	bool bLargeEnough = static_cast<uint32>(TokenDataStores.Num()) > DataStore->TypeId;
+	bool bSomethingAlreadyExists = bLargeEnough ? TokenDataStores[DataStore->TypeId].Value.IsValid() : false;
+	FName ExistingName = bLargeEnough ? TokenDataStores[DataStore->TypeId].Key : NAME_None;
+	if (!ensureAlwaysMsgf(!bSomethingAlreadyExists, TEXT("An existing TokenDataStore (Name %s) has been created with that TypeId when Trying to add New TokenDataStore with name: %s"),*ExistingName.ToString(),*TokenStoreName.ToString()))
+	{
+		UE_LOG(LogNetToken, Error, TEXT("An existing TokenDataStore (Name %s) has been created with that TypeId when Trying to add New TokenDataStore with name: %s"),*ExistingName.ToString(),*TokenStoreName.ToString());
+		return false;
+	}
+
+	// Need to resize to
+	if (!bLargeEnough)
+	{
+		TokenDataStores.SetNum(DataStore->TypeId + 1);
+	}
+
+	TokenDataStores.EmplaceAt(DataStore->TypeId, TokenStoreName, MoveTemp(DataStore));
 
 	return true;
 }
 
+bool FNetTokenStore::UnRegisterDataStore(FName TokenStoreName)
+{
+	if (TokenDataStores.Num() >= FNetToken::MaxTypeIdCount)
+	{
+		return false;
+	}
+	
+	if (FNetTokenDataStore* DataStore = GetDataStore(TokenStoreName))
+	{
+		// Delete but don't resize the array
+		TokenDataStores[DataStore->TypeId].Value.Reset();
+		TokenDataStores[DataStore->TypeId].Key = NAME_None;
+		return true;
+	}
+
+	// Failed
+	return false;
+}
+
 void FNetTokenStore::WriteTokenData(FNetSerializationContext& Context, const FNetToken NetToken) const
 {
 	if (NetToken.IsValid())
@@ -328,7 +372,7 @@ void FNetTokenStore::WriteTokenData(FNetSerializationContext& Context, const FNe
 	}
 }
 
-void FNetTokenStore::WriteTokenData(FArchive& Ar, const FNetToken NetToken) const
+void FNetTokenStore::WriteTokenData(FArchive& Ar, const FNetToken NetToken, UPackageMap* Map /*= nullptr*/) const
 {
 	if (NetToken.IsValid())
 	{
@@ -344,7 +388,7 @@ void FNetTokenStore::WriteTokenData(FArchive& Ar, const FNetToken NetToken) cons
 		const FNetTokenStoreKey TokenKey = LocalNetTokenStoreState->TokenInfoArray[NetToken.GetTypeId()][NetToken.GetIndex()];
 
 		// Write token data
-		TokenDataStores[NetToken.GetTypeId()].Get<1>()->WriteTokenData(Ar, TokenKey);
+		TokenDataStores[NetToken.GetTypeId()].Get<1>()->WriteTokenData(Ar, TokenKey, Map);
 	}
 }
 
@@ -408,7 +452,7 @@ void FNetTokenStore::ReadTokenData(FNetSerializationContext& Context, const FNet
 	}
 }
 
-void FNetTokenStore::ReadTokenData(FArchive& Ar, const FNetToken NetToken, FNetTokenStoreState& RemoteNetTokenStoreState)
+void FNetTokenStore::ReadTokenData(FArchive& Ar, const FNetToken NetToken, FNetTokenStoreState& RemoteNetTokenStoreState, UPackageMap* Map /*= nullptr*/)
 {
 	if (NetToken.IsValid())
 	{
@@ -424,7 +468,7 @@ void FNetTokenStore::ReadTokenData(FArchive& Ar, const FNetToken NetToken, FNetT
 		}
 
 		FNetTokenDataStore* DataStore = TokenDataStores[TokenTypeId].Get<1>().Get();
-		const FNetTokenStoreKey StoreKey = DataStore->ReadTokenData(Ar, NetToken);
+		const FNetTokenStoreKey StoreKey = DataStore->ReadTokenData(Ar, NetToken, Map);
 
 		if (!ValidateAndStoreNetTokenData(*DataStore, RemoteNetTokenStoreState, NetToken, StoreKey))
 		{
@@ -510,3 +554,47 @@ TArray<FNetToken> FNetTokenStore::GetAllNetTokens() const
 }
 
 }
+
+uint32 UNetTokenTypeIdConfig::GetTypeID(const FString& TypeName) const
+{
+	// Need to reevaluate every retrieval in the event of hotfixes, or ini changes due to GFP loading/unloading.
+	bool bTypeIdsAppearValid = ReservedTypeIdsAppearValid();
+	if (!ensureAlwaysMsgf(bTypeIdsAppearValid, TEXT("Duplicate or Invalid TypeIds detected!")))
+	{
+		return UE::Net::FNetToken::InvalidTokenTypeId;
+	}
+	for (const FNetTokenStoreTypeIdPair& TypePair : ReservedTypeIds)
+	{
+		if (TypePair.StoreTypeName == TypeName)
+		{
+			return TypePair.TypeID;
+		}
+	}
+	checkf(false, TEXT("Unknown Token Store Type %s. StoreType MUST be declared in Engine.ini."), *TypeName);
+	return UE::Net::FNetToken::InvalidTokenTypeId;
+}
+
+bool UNetTokenTypeIdConfig::ReservedTypeIdsAppearValid() const
+{
+	if (ReservedTypeIds.Num() == 0)
+	{
+		UE_LOG(LogNetToken, Error, TEXT("No TypeID information found in [/Script/IrisCore.NetTokenTypeIdConfig] Engine.ini"));
+		return false;
+	}
+	TMap<uint32,FString> Found;
+	for (const FNetTokenStoreTypeIdPair& TypeIdPair : ReservedTypeIds)
+	{
+		if (!ensureAlwaysMsgf(TypeIdPair.TypeID < UE::Net::FNetToken::MaxTypeIdCount, TEXT("TypeID %d found for %s that is larger than FNetToken::MaxIdCount %d"), TypeIdPair.TypeID, *TypeIdPair.StoreTypeName, UE::Net::FNetToken::MaxTypeIdCount))
+		{
+			UE_LOG(LogNetToken, Error, TEXT("TypeID %d found for %s that is larger than FNetToken::MaxIdCount %d"), TypeIdPair.TypeID, *TypeIdPair.StoreTypeName, UE::Net::FNetToken::MaxTypeIdCount);
+			return false;
+		}
+		if (!ensureAlwaysMsgf(!Found.Contains(TypeIdPair.TypeID), TEXT("Duplicate TypeIDs found: %s and %s with TypeId: %u"), *Found[TypeIdPair.TypeID], *TypeIdPair.StoreTypeName, TypeIdPair.TypeID))
+		{
+			UE_LOG(LogNetToken, Error, TEXT("Duplicate TypeID information found in [/Script/IrisCore.NetTokenTypeIdConfig] Engine.ini for%s and %s with TypeId: %u"), *Found[TypeIdPair.TypeID], *TypeIdPair.StoreTypeName, TypeIdPair.TypeID);
+			return false;
+		}
+		Found.Add(TypeIdPair.TypeID, TypeIdPair.StoreTypeName);
+	}
+	return true;
+}
diff --git a/Engine/Source/Runtime/Experimental/Iris/Core/Private/Iris/ReplicationSystem/StringTokenStore.cpp b/Engine/Source/Runtime/Experimental/Iris/Core/Private/Iris/ReplicationSystem/StringTokenStore.cpp
--- a/Engine/Source/Runtime/Experimental/Iris/Core/Private/Iris/ReplicationSystem/StringTokenStore.cpp
+++ b/Engine/Source/Runtime/Experimental/Iris/Core/Private/Iris/ReplicationSystem/StringTokenStore.cpp
@@ -108,7 +108,7 @@ void FStringTokenStore::WriteTokenData(FNetSerializationContext& Context, FNetTo
 	WriteString(Context.GetBitStreamWriter(), FStringView(StoredStrings[TokenStoreKey.GetKeyIndex()]));
 }
 
-void FStringTokenStore::WriteTokenData(FArchive& Ar, FNetTokenStoreKey TokenStoreKey) const
+void FStringTokenStore::WriteTokenData(FArchive& Ar, FNetTokenStoreKey TokenStoreKey, UPackageMap* Map /*= nullptr*/) const
 {
 	FString Temp(StoredStrings[TokenStoreKey.GetKeyIndex()]);
 	Ar << Temp;
@@ -132,7 +132,7 @@ FNetTokenDataStore::FNetTokenStoreKey FStringTokenStore::ReadTokenData(FNetSeria
 	}
 }
 
-FNetTokenDataStore::FNetTokenStoreKey FStringTokenStore::ReadTokenData(FArchive& Ar, const FNetToken& NetToken)
+FNetTokenDataStore::FNetTokenStoreKey FStringTokenStore::ReadTokenData(FArchive& Ar, const FNetToken& NetToken, UPackageMap* Map /*= nullptr*/)
 {
 	// Read the token data and add it to the string store without assigning LocalToken
 	FString Temp;
diff --git a/Engine/Source/Runtime/Experimental/Iris/Core/Public/Iris/ReplicationSystem/NameTokenStore.h b/Engine/Source/Runtime/Experimental/Iris/Core/Public/Iris/ReplicationSystem/NameTokenStore.h
--- a/Engine/Source/Runtime/Experimental/Iris/Core/Public/Iris/ReplicationSystem/NameTokenStore.h
+++ b/Engine/Source/Runtime/Experimental/Iris/Core/Public/Iris/ReplicationSystem/NameTokenStore.h
@@ -29,11 +29,11 @@ public:
 protected:
 	// Serialize data for a token, note there is not validation in this function
 	IRISCORE_API virtual void WriteTokenData(FNetSerializationContext& Context, FNetTokenStoreKey TokenStoreKey) const override;
-	IRISCORE_API virtual void WriteTokenData(FArchive& Archive, FNetTokenStoreKey TokenStoreKey) const override;
+	IRISCORE_API virtual void WriteTokenData(FArchive& Archive, FNetTokenStoreKey TokenStoreKey, UPackageMap* Map = nullptr) const override;
 
 	// Read data for a token, returns a valid StoreKey if successful read
 	IRISCORE_API virtual FNetTokenStoreKey ReadTokenData(FNetSerializationContext& Context, const FNetToken& NetToken) override;
-	IRISCORE_API virtual FNetTokenStoreKey ReadTokenData(FArchive& Archive, const FNetToken& NetToken) override;
+	IRISCORE_API virtual FNetTokenStoreKey ReadTokenData(FArchive& Archive, const FNetToken& NetToken, UPackageMap* Map = nullptr) override;
 
 	// Create and store data for Token
 	FNetTokenStoreKey GetOrCreateTokenStoreKey(FName Name);
diff --git a/Engine/Source/Runtime/Experimental/Iris/Core/Public/Iris/ReplicationSystem/NetTokenStore.h b/Engine/Source/Runtime/Experimental/Iris/Core/Public/Iris/ReplicationSystem/NetTokenStore.h
--- a/Engine/Source/Runtime/Experimental/Iris/Core/Public/Iris/ReplicationSystem/NetTokenStore.h
+++ b/Engine/Source/Runtime/Experimental/Iris/Core/Public/Iris/ReplicationSystem/NetTokenStore.h
@@ -6,8 +6,11 @@
 #include "Net/Core/NetToken/NetToken.h"
 #include "UObject/NameTypes.h"
 #include "Containers/Array.h"
+#include "Containers/Map.h"
 #include "Templates/Tuple.h"
 
+#include "NetTokenStore.generated.h"
+
 /** 
 * The idea with NetTokens is to allow export of "stable"-pieces of data such as string and names by replacing them with a NetToken during quantization/serialization.
 * The data associated with a NetToken is then exported separately from the data. As soon as the exported data has been acknowledged the data will only be serialized using the NetToken.
@@ -54,9 +57,41 @@ namespace UE::Net
 	}
 }
 
-namespace UE::Net
+USTRUCT()
+struct FNetTokenStoreTypeIdPair
+{
+	GENERATED_BODY()
+	UPROPERTY()
+	FString StoreTypeName;
+	UPROPERTY()
+	uint32 TypeID;
+	bool operator<(const FNetTokenStoreTypeIdPair& Other) const
+	{
+		return TypeID < Other.TypeID;
+	}
+};
+	
+UCLASS(Transient, MinimalAPI, config=Engine)
+class UNetTokenTypeIdConfig : public UObject
 {
+	GENERATED_BODY()
+
+public:
+	UPROPERTY(Config)
+	TArray<FNetTokenStoreTypeIdPair> ReservedTypeIds;
+
+	uint32 GetTypeID(const FString& TypeName) const;
+	uint32 GetTypeID(const FName& TypeName) const
+	{
+		return GetTypeID(TypeName.ToString());
+	}
+	
+protected:
+	bool ReservedTypeIdsAppearValid() const;
+};
 
+namespace UE::Net
+{
 /**
  * NetTokenDataStore
  * Implemented per type to store and serialize data associated with a NetToken.
@@ -118,10 +153,11 @@ protected:
 	virtual void WriteTokenData(FNetSerializationContext& Context, FNetTokenStoreKey Key) const = 0;
 	virtual FNetTokenDataStore::FNetTokenStoreKey ReadTokenData(FNetSerializationContext& Context, const FNetToken& NetToken) = 0;
 
-	virtual void WriteTokenData(FArchive& Ar, FNetTokenStoreKey Key) const = 0;
-	virtual FNetTokenDataStore::FNetTokenStoreKey ReadTokenData(FArchive& Ar, const FNetToken& NetToken) = 0;
-
-	FNetTokenDataStore::FNetTokenStoreKey GetTokenKey(FNetToken Token, const FNetTokenStoreState& TokenStoreState) const;
+	//Do not export things through the UPackageMap. This will break things.
+	virtual void WriteTokenData(FArchive& Ar, FNetTokenStoreKey Key, UPackageMap* Map = nullptr) const = 0;
+	virtual FNetTokenDataStore::FNetTokenStoreKey ReadTokenData(FArchive& Ar, const FNetToken& NetToken, UPackageMap* Map = nullptr) = 0;
+	
+	IRISCORE_API FNetTokenDataStore::FNetTokenStoreKey GetTokenKey(FNetToken Token, const FNetTokenStoreState& TokenStoreState) const;
 
 	inline FNetToken::FTypeId GetTypeId() const
 	{
@@ -181,6 +217,9 @@ public:
 	/** Register DataStore and return true if it was registered. */
 	IRISCORE_API bool RegisterDataStore(TUniquePtr<FNetTokenDataStore> DataStore, FName TokenStoreName);
 
+	/** UnRegister  a DataStore and return true if it was registered. */
+	IRISCORE_API bool UnRegisterDataStore(FName TokenStoreName);
+
 	/** Get const access to data store by name. */
 	IRISCORE_API const FNetTokenDataStore* GetDataStore(FName Name) const;
 
@@ -195,6 +234,13 @@ public:
 		
 		return RegisterDataStore(MakeUnique<T>(*this), T::GetTokenStoreName());
 	}
+	
+	/** Unregister a data store of specified type. */
+	template<typename T>
+	bool DeleteAndUnRegisterDataStore()
+	{
+		return UnRegisterDataStore(T::GetTokenStoreName());
+	}
 
 	/** Return data store of specified type. */
 	template<typename T>
@@ -240,11 +286,11 @@ public:
 	/** Read data associated with the NetToken. */
 	IRISCORE_API void ReadTokenData(FNetSerializationContext& Context, const FNetToken NetToken, FNetTokenStoreState& RemoteNetTokenStoreState);
 
-	/** Write data associated with the NetToken. */
-	IRISCORE_API void WriteTokenData(FArchive& Ar, const FNetToken NetToken) const;
-
-	/** Read data associated with the NetToken. */
-	IRISCORE_API void ReadTokenData(FArchive& Ar, const FNetToken NetToken, FNetTokenStoreState& RemoteNetTokenStoreState);
+	/** Write data associated with the NetToken. Do not export anything via the UPackageMap. */
+	IRISCORE_API void WriteTokenData(FArchive& Ar, const FNetToken NetToken, UPackageMap* Map = nullptr) const;
+	
+	/** Read data associated with the NetToken. Do not export anything via the UPackageMap. */
+	IRISCORE_API void ReadTokenData(FArchive& Ar, const FNetToken NetToken, FNetTokenStoreState& RemoteNetTokenStoreState, UPackageMap* Map = nullptr);
 
 	/** Conditionally write NetTokenData unless already exported. */
 	IRISCORE_API void ConditionalWriteNetTokenData(FNetSerializationContext& Context, Private::FNetExportContext* ExportContext, const FNetToken NetToken) const;
@@ -374,5 +420,4 @@ inline FNetTokenDataStore::FNetTokenStoreKey FNetTokenDataStore::GetNextNetToken
 
 	return FNetTokenStoreKey();
 }
-
 }
diff --git a/Engine/Source/Runtime/Experimental/Iris/Core/Public/Iris/ReplicationSystem/StringTokenStore.h b/Engine/Source/Runtime/Experimental/Iris/Core/Public/Iris/ReplicationSystem/StringTokenStore.h
--- a/Engine/Source/Runtime/Experimental/Iris/Core/Public/Iris/ReplicationSystem/StringTokenStore.h
+++ b/Engine/Source/Runtime/Experimental/Iris/Core/Public/Iris/ReplicationSystem/StringTokenStore.h
@@ -39,11 +39,11 @@ public:
 protected:
 	// Serialize data for a token, note there is not validation in this function
 	virtual void WriteTokenData(FNetSerializationContext& Context, FNetTokenStoreKey TokenStoreKey) const override;
-	virtual void WriteTokenData(FArchive& Archive, FNetTokenStoreKey TokenStoreKey) const override;
+	virtual void WriteTokenData(FArchive& Archive, FNetTokenStoreKey TokenStoreKey, UPackageMap* Map = nullptr) const override;
 
 	// Read data for a token, returns a valid StoreKey if successful read
 	virtual FNetTokenStoreKey ReadTokenData(FNetSerializationContext& Context, const FNetToken& NetToken) override;
-	virtual FNetTokenStoreKey ReadTokenData(FArchive& Archive, const FNetToken& NetToken) override;
+	virtual FNetTokenStoreKey ReadTokenData(FArchive& Archive, const FNetToken& NetToken, UPackageMap* Map = nullptr) override;
 
 	// Create a persistent string
 	FNetTokenStoreKey GetOrCreatePersistentString(const TCHAR* Name, uint32 Length);
