[ChaosVehiclesCore] Resolve nullptr access in FSimModuleTree::SetSimState during some initialization scenearios.

diff --git a/Engine/Source/Runtime/Experimental/ChaosVehicles/ChaosVehiclesCore/Private/SimModule/SimModuleTree.cpp b/Engine/Source/Runtime/Experimental/ChaosVehicles/ChaosVehiclesCore/Private/SimModule/SimModuleTree.cpp
--- a/Engine/Source/Runtime/Experimental/ChaosVehicles/ChaosVehiclesCore/Private/SimModule/SimModuleTree.cpp
+++ b/Engine/Source/Runtime/Experimental/ChaosVehicles/ChaosVehiclesCore/Private/SimModule/SimModuleTree.cpp
@@ -564,7 +564,7 @@ void FSimModuleTree::SetNetState(Chaos::FModuleNetDataArray& ModuleDatas)
 
 	for (TSharedPtr<FModuleNetData>& DataElement : ModuleDatas)
 	{
-		if (!SimulationModuleTree.IsEmpty() && DataElement->SimArrayIndex < SimulationModuleTree.Num())
+		if (!SimulationModuleTree.IsEmpty() && DataElement && DataElement->SimArrayIndex < SimulationModuleTree.Num())
 		{
 			if (ISimulationModuleBase* SimModule = SimulationModuleTree[DataElement->SimArrayIndex].SimModule)
 			{
@@ -585,7 +585,7 @@ void FSimModuleTree::SetSimState(const Chaos::FModuleNetDataArray& ModuleDatas)
 
 	for (const TSharedPtr<FModuleNetData>& DataElement : ModuleDatas)
 	{
-		if (!SimulationModuleTree.IsEmpty() && DataElement->SimArrayIndex < SimulationModuleTree.Num())
+		if (!SimulationModuleTree.IsEmpty() && DataElement && DataElement->SimArrayIndex < SimulationModuleTree.Num())
 		{
 			if (ISimulationModuleBase* SimModule = SimulationModuleTree[DataElement->SimArrayIndex].SimModule)
 			{
