Geometry collection component : Fix out of bounds access in SetRootProxyLocalTransform that could result in corrupted transforms
- the code was misusing the append function by passing a pointer to the a single transform identity instead of an array of it
- fixed it by using an for loop  instead

diff --git a/Engine/Source/Runtime/Experimental/GeometryCollectionEngine/Private/GeometryCollection/GeometryCollectionComponent.cpp b/Engine/Source/Runtime/Experimental/GeometryCollectionEngine/Private/GeometryCollection/GeometryCollectionComponent.cpp
--- a/Engine/Source/Runtime/Experimental/GeometryCollectionEngine/Private/GeometryCollection/GeometryCollectionComponent.cpp
+++ b/Engine/Source/Runtime/Experimental/GeometryCollectionEngine/Private/GeometryCollection/GeometryCollectionComponent.cpp
@@ -6313,7 +6313,10 @@ void UGeometryCollectionComponent::SetRootProxyLocalTransform(int32 Index, const
 		const int32 NumProxyMeshes = RestCollection->RootProxyData.ProxyMeshes.Num();
 		RootProxyLocalTransforms.Reserve(NumProxyMeshes);
 		RootProxyLocalTransforms.Append(RestCollection->RootProxyData.MeshTransforms);
-		RootProxyLocalTransforms.Append(&FTransform3f::Identity, (NumProxyMeshes - RootProxyLocalTransforms.Num()));
+		for (int32 TransformIndex = RootProxyLocalTransforms.Num(); TransformIndex < NumProxyMeshes; ++TransformIndex)
+		{
+			RootProxyLocalTransforms.Add(FTransform3f::Identity);
+		}
 	}
 	if (RootProxyLocalTransforms.IsValidIndex(Index))
 	{
