Dataflow : Make sure the node warning / error / failed / log data is cleared just before evaluating a node
- that fixes stale information when evaluating the graph

diff --git a/Engine/Source/Runtime/Experimental/Dataflow/Core/Private/Dataflow/DataflowNodeParameters.cpp b/Engine/Source/Runtime/Experimental/Dataflow/Core/Private/Dataflow/DataflowNodeParameters.cpp
--- a/Engine/Source/Runtime/Experimental/Dataflow/Core/Private/Dataflow/DataflowNodeParameters.cpp
+++ b/Engine/Source/Runtime/Experimental/Dataflow/Core/Private/Dataflow/DataflowNodeParameters.cpp
@@ -260,12 +260,21 @@ namespace UE::Dataflow
 		NodesFailed.Empty();
 	}
 
+	void FContext::ClearNodeData(const FDataflowNode* InNode)
+	{
+		NodesWithInfo.Remove(InNode);
+		NodesWithWarning.Remove(InNode);
+		NodesWithError.Remove(InNode);
+		NodesFailed.Remove(InNode);
+	}
+
 	FContextScopedCallstack::FContextScopedCallstack(FContext& InContext, const FDataflowConnection* InConnection)
 		: Context(InContext)
 		, Connection(InConnection)
 	{
 		bLoopDetected = Context.IsInCallstack(Connection);
 		Context.ResetPerfDataForNode(*Connection->GetOwningNode());
+		Context.ClearNodeData(Connection->GetOwningNode());
 		Context.PushToCallstack(Connection);
 	}
 
diff --git a/Engine/Source/Runtime/Experimental/Dataflow/Core/Public/Dataflow/DataflowNodeParameters.h b/Engine/Source/Runtime/Experimental/Dataflow/Core/Public/Dataflow/DataflowNodeParameters.h
--- a/Engine/Source/Runtime/Experimental/Dataflow/Core/Public/Dataflow/DataflowNodeParameters.h
+++ b/Engine/Source/Runtime/Experimental/Dataflow/Core/Public/Dataflow/DataflowNodeParameters.h
@@ -547,6 +547,7 @@ namespace UE::Dataflow
 		DATAFLOWCORE_API bool NodeFailed(const FDataflowNode* InNode) { return NodesFailed.Contains(InNode); };
 
 		DATAFLOWCORE_API void ClearNodesData();
+		DATAFLOWCORE_API void ClearNodeData(const FDataflowNode* InNode);
 
 		DECLARE_MULTICAST_DELEGATE_TwoParams(FOnNodeBeginEvaluateMulticast, const FDataflowNode* Node, const FDataflowOutput* Output);
 		FOnNodeBeginEvaluateMulticast OnNodeBeginEvaluateMulticast;
