Removing IAS options EOnDemandStreamingOptions::Disabled/Enabled as it was conflicting with internal status setting.

#rnx

- Removed the APIs for disabling the ondemand http streaming systems as at this point it is like trying to disable the hard disk.
- The cvar 'ias.HttpEnabled' still exists if we really need to disable http streaming for what ever reason.
- The http enable part of backend status is now pretty much only used to track the period where the thread is starting up, we will probably look to remove the backend status soon.
- Fixed a bug where we would not enable the http portion of the backend status once the thread is running if the initial connection ran into any issues. This bug was somewhat hidden by client code forcing the http status to be enabled anyway.

diff --git a/Engine/Source/Runtime/Experimental/IoStore/OnDemand/Private/OnDemandIoDispatcherBackend.cpp b/Engine/Source/Runtime/Experimental/IoStore/OnDemand/Private/OnDemandIoDispatcherBackend.cpp
--- a/Engine/Source/Runtime/Experimental/IoStore/OnDemand/Private/OnDemandIoDispatcherBackend.cpp
+++ b/Engine/Source/Runtime/Experimental/IoStore/OnDemand/Private/OnDemandIoDispatcherBackend.cpp
@@ -834,8 +834,6 @@ public:
 
 	// I/O Http backend
 	virtual void SetBulkOptionalEnabled(bool bEnabled) override;
-	virtual void SetEnabled(bool bEnabled) override;
-	virtual bool IsEnabled() const override;
 	virtual void AbandonCache() override;
 	
 	virtual void ReportAnalytics(TArray<FAnalyticsEventAttribute>& OutAnalyticsArray) const override;
@@ -1457,16 +1455,6 @@ void FOnDemandIoBackend::SetBulkOptionalEnabled(bool bEnabled)
 	BackendStatus.SetHttpOptionalBulkEnabled(bEnabled);
 }
 
-void FOnDemandIoBackend::SetEnabled(bool bEnabled)
-{
-	BackendStatus.SetHttpEnabled(bEnabled);
-}
-
-bool FOnDemandIoBackend::IsEnabled() const
-{
-	return BackendStatus.IsHttpEnabled();
-}
-
 void FOnDemandIoBackend::AbandonCache()
 {
 	BackendStatus.SetCacheEnabled(false);
@@ -1646,8 +1634,6 @@ uint32 FOnDemandIoBackend::Run()
 {
 	LLM_SCOPE_BYTAG(Ias);
 
-	BackendStatus.SetHttpEnabled(false); // Http is disabled until the thread starts up successfully
-
 	if (!CreateDefaultHostGroup())
 	{
 		return 0;
@@ -1669,14 +1655,16 @@ uint32 FOnDemandIoBackend::Run()
 
 	if (ConnectionTest(DefaultHostGroup.GetPrimaryHostUrl(), DefaultHostGroup.GetTestPath(), GIasHttpTimeOutMs))
 	{
-		BackendStatus.SetHttpEnabled(true);
 		Stats.OnHttpConnected();
+		
 	}
 	else
 	{
 		DefaultHostGroup.Disconnect();
 	}
 
+	BackendStatus.SetHttpEnabled(true);
+
 	const double InitTime = FPlatformTime::Seconds() - InitStart;
 	UE_LOG(LogIas, Display, TEXT("HostGroup init took %.3f seconds"), InitTime);
 
diff --git a/Engine/Source/Runtime/Experimental/IoStore/OnDemand/Private/OnDemandIoDispatcherBackend.h b/Engine/Source/Runtime/Experimental/IoStore/OnDemand/Private/OnDemandIoDispatcherBackend.h
--- a/Engine/Source/Runtime/Experimental/IoStore/OnDemand/Private/OnDemandIoDispatcherBackend.h
+++ b/Engine/Source/Runtime/Experimental/IoStore/OnDemand/Private/OnDemandIoDispatcherBackend.h
@@ -79,8 +79,6 @@ public:
 	virtual ~IOnDemandIoDispatcherBackend() = default;
 
 	virtual void SetBulkOptionalEnabled(bool bInEnabled) = 0;
-	virtual void SetEnabled(bool bInEnabled) = 0;
-	virtual bool IsEnabled() const = 0;
 	virtual void AbandonCache() = 0;
 	virtual void ReportAnalytics(TArray<FAnalyticsEventAttribute>& OutAnalyticsArray) const = 0;
 	virtual TUniquePtr<IAnalyticsRecording> StartAnalyticsRecording() const = 0;
diff --git a/Engine/Source/Runtime/Experimental/IoStore/OnDemand/Private/OnDemandIoStore.cpp b/Engine/Source/Runtime/Experimental/IoStore/OnDemand/Private/OnDemandIoStore.cpp
--- a/Engine/Source/Runtime/Experimental/IoStore/OnDemand/Private/OnDemandIoStore.cpp
+++ b/Engine/Source/Runtime/Experimental/IoStore/OnDemand/Private/OnDemandIoStore.cpp
@@ -639,7 +639,7 @@ FIoStatus FOnDemandIoStore::InitializeStreamingBackend(const FOnDemandEndpointCo
 	{
 		return FIoStatus(EIoErrorCode::InvalidParameter, TEXT("Failed to create streaming backend"));
 	}
-	StreamingBackend->SetEnabled(EnumHasAnyFlags(StreamingOptions, EOnDemandStreamingOptions::Enabled));
+
 	StreamingBackend->SetBulkOptionalEnabled(EnumHasAnyFlags(StreamingOptions, EOnDemandStreamingOptions::OptionalBulkDataDisabled) == false);
 
 	int32 BackendPriority = -10;
@@ -2093,7 +2093,6 @@ void FOnDemandIoStore::SetStreamingOptions(EOnDemandStreamingOptions Options)
 	StreamingOptions = Options;
 	if (StreamingBackend.IsValid())
 	{
-		StreamingBackend->SetEnabled(EnumHasAnyFlags(Options, EOnDemandStreamingOptions::Enabled));
 		StreamingBackend->SetBulkOptionalEnabled(EnumHasAnyFlags(Options, EOnDemandStreamingOptions::OptionalBulkDataDisabled) == false);
 	}
 }
diff --git a/Engine/Source/Runtime/Experimental/IoStore/OnDemand/Private/OnDemandIoStore.h b/Engine/Source/Runtime/Experimental/IoStore/OnDemand/Private/OnDemandIoStore.h
--- a/Engine/Source/Runtime/Experimental/IoStore/OnDemand/Private/OnDemandIoStore.h
+++ b/Engine/Source/Runtime/Experimental/IoStore/OnDemand/Private/OnDemandIoStore.h
@@ -398,7 +398,7 @@ private:
 	TArray<FSharedPurgeRequest>				PurgeRequests;
 	UE::FMutex								RequestMutex;
 
-	EOnDemandStreamingOptions				StreamingOptions = EOnDemandStreamingOptions::Enabled;
+	EOnDemandStreamingOptions				StreamingOptions = EOnDemandStreamingOptions::Default;
 	bool									bTicking = false;
 	bool									bTickRequested = false;
 	TFuture<void>							TickFuture;
diff --git a/Engine/Source/Runtime/Experimental/IoStore/OnDemand/Public/IO/IoStoreOnDemand.h b/Engine/Source/Runtime/Experimental/IoStore/OnDemand/Public/IO/IoStoreOnDemand.h
--- a/Engine/Source/Runtime/Experimental/IoStore/OnDemand/Public/IO/IoStoreOnDemand.h
+++ b/Engine/Source/Runtime/Experimental/IoStore/OnDemand/Public/IO/IoStoreOnDemand.h
@@ -393,12 +393,10 @@ public:
 /** Streaming options (IAS). */
 enum class EOnDemandStreamingOptions
 {
-	/** On-demand streaming disabled. */
-	Disabled					= 0,
-	/** On-demand streaming enabled. */
-	Enabled						= (1 << 0),
+	/** No special options enabled. */
+	Default						= 0,
 	/** Optional bulk data disabled. */
-	OptionalBulkDataDisabled	= (1 << 1)
+	OptionalBulkDataDisabled	= (1 << 0)
 };
 ENUM_CLASS_FLAGS(EOnDemandStreamingOptions);
 
