Iris - Moved typedef into proper namespace

#rnx

diff --git a/Engine/Source/Runtime/Experimental/Iris/Core/Private/Iris/ReplicationSystem/ReplicationSystem.cpp b/Engine/Source/Runtime/Experimental/Iris/Core/Private/Iris/ReplicationSystem/ReplicationSystem.cpp
--- a/Engine/Source/Runtime/Experimental/Iris/Core/Private/Iris/ReplicationSystem/ReplicationSystem.cpp
+++ b/Engine/Source/Runtime/Experimental/Iris/Core/Private/Iris/ReplicationSystem/ReplicationSystem.cpp
@@ -1098,7 +1098,7 @@ void UReplicationSystem::SetStaticPriority(FNetRefHandle Handle, float Priority)
 	return Impl->ReplicationSystemInternal.GetPrioritization().SetStaticPriority(ObjectInternalIndex, Priority);
 }
 
-bool UReplicationSystem::SetPrioritizer(FNetRefHandle Handle, FNetObjectPrioritizerHandle Prioritizer)
+bool UReplicationSystem::SetPrioritizer(FNetRefHandle Handle, UE::Net::FNetObjectPrioritizerHandle Prioritizer)
 {
 	using namespace UE::Net::Private;
 
@@ -1117,7 +1117,7 @@ bool UReplicationSystem::SetPrioritizer(FNetRefHandle Handle, FNetObjectPrioriti
 	return ReplicationSystemInternal.GetPrioritization().SetPrioritizer(ObjectInternalIndex, Prioritizer);
 }
 
-FNetObjectPrioritizerHandle UReplicationSystem::GetPrioritizerHandle(const FName PrioritizerName) const
+UE::Net::FNetObjectPrioritizerHandle UReplicationSystem::GetPrioritizerHandle(const FName PrioritizerName) const
 {
 	return Impl->ReplicationSystemInternal.GetPrioritization().GetPrioritizerHandle(PrioritizerName);
 }
diff --git a/Engine/Source/Runtime/Experimental/Iris/Core/Public/Iris/ReplicationSystem/Prioritization/NetObjectPrioritizer.h b/Engine/Source/Runtime/Experimental/Iris/Core/Public/Iris/ReplicationSystem/Prioritization/NetObjectPrioritizer.h
--- a/Engine/Source/Runtime/Experimental/Iris/Core/Public/Iris/ReplicationSystem/Prioritization/NetObjectPrioritizer.h
+++ b/Engine/Source/Runtime/Experimental/Iris/Core/Public/Iris/ReplicationSystem/Prioritization/NetObjectPrioritizer.h
@@ -11,12 +11,27 @@ struct FNetObjectPrioritizationInfo;
 class UReplicationSystem;
 namespace UE::Net
 {
+	typedef uint32 FNetObjectPrioritizerHandle;
 	struct FReplicationInstanceProtocol;
 	struct FReplicationProtocol;
 	struct FReplicationView;
 }
 
-typedef uint32 FNetObjectPrioritizerHandle;
+namespace UE::Net
+{
+
+/** Used to represent an invalid handle to a prioritizer. */
+constexpr FNetObjectPrioritizerHandle InvalidNetObjectPrioritizerHandle = ~FNetObjectPrioritizerHandle(0);
+
+/**
+ * The handle of the default prioritizer.
+ * The first valid prioritizer definition will assume the role as default spatial prioritizer. All objects with a RepTag_WorldLocation tag 
+ * will be added to the default prioritizer. To override the behavior a prioritizer must be set via calls to the ReplicationSystem.
+ * @see UNetObjectPrioritizerDefinitions
+ */
+constexpr FNetObjectPrioritizerHandle DefaultSpatialNetObjectPrioritizerHandle = FNetObjectPrioritizerHandle(0);
+
+}
 
 /**
   * Parameters passed to UNetObjectPrioritizer::Prioritize.
@@ -187,14 +202,3 @@ public:
 protected:
 	IRISCORE_API UNetObjectPrioritizer();
 };
-
-/** Used to represent an invalid handle to a prioritizer. */
-constexpr FNetObjectPrioritizerHandle InvalidNetObjectPrioritizerHandle = ~FNetObjectPrioritizerHandle(0);
-
-/**
- * The handle of the default prioritizer.
- * The first valid prioritizer definition will assume the role as default spatial prioritizer. All objects with a RepTag_WorldLocation tag 
- * will be added to the default prioritizer. To override the behavior a prioritizer must be set via calls to the ReplicationSystem.
- * @see UNetObjectPrioritizerDefinitions
- */
-constexpr FNetObjectPrioritizerHandle DefaultSpatialNetObjectPrioritizerHandle = FNetObjectPrioritizerHandle(0);
