Iris- Fixed memory stomp in FChunkedDataWriter

#rnx

diff --git a/Engine/Source/Runtime/Experimental/Iris/Core/Private/Iris/ReplicationSystem/ChunkedDataStream/ChunkedDataWriter.cpp b/Engine/Source/Runtime/Experimental/Iris/Core/Private/Iris/ReplicationSystem/ChunkedDataStream/ChunkedDataWriter.cpp
--- a/Engine/Source/Runtime/Experimental/Iris/Core/Private/Iris/ReplicationSystem/ChunkedDataStream/ChunkedDataWriter.cpp
+++ b/Engine/Source/Runtime/Experimental/Iris/Core/Private/Iris/ReplicationSystem/ChunkedDataStream/ChunkedDataWriter.cpp
@@ -116,70 +116,76 @@ FChunkedDataWriter::FReferencesForExport* FChunkedDataWriter::CreateExportPayloa
 		TUniquePtr<FReferencesForExport> Result = MakeUnique<FReferencesForExport>();
 		Result->ExportsPayload.AddUninitialized(ExportsBufferMaxSize);
 
-		FNetBitStreamWriter ExportsWriter;
-		const uint32 MaxExportsPayloadBytes = Result->ExportsPayload.Num();
-		ExportsWriter.InitBytes(Result->ExportsPayload.GetData(), MaxExportsPayloadBytes);
+		uint32 ExportsPayloadBytes = 0;
+		{
+			FNetBitStreamWriter ExportsWriter;
+			const uint32 MaxExportsPayloadBytes = Result->ExportsPayload.Num();
+			ExportsWriter.InitBytes(Result->ExportsPayload.GetData(), MaxExportsPayloadBytes);
 
-		// Create context
-		FNetSerializationContext Context(&ExportsWriter);
-		FInternalNetSerializationContext InternalContext(ReplicationSystem);
-		Context.SetInternalContext(&InternalContext);
-		Context.SetLocalConnectionId(InitParams.ConnectionId);
+			// Create context
+			FNetSerializationContext Context(&ExportsWriter);
+			FInternalNetSerializationContext InternalContext(ReplicationSystem);
+			Context.SetInternalContext(&InternalContext);
+			Context.SetLocalConnectionId(InitParams.ConnectionId);
 
-		// Temporary quantized state
-		FIrisPackageMapExportsQuantizedType QuantizedExports = {};
+			// Temporary quantized state
+			FIrisPackageMapExportsQuantizedType QuantizedExports = {};
 
-		// We need to release dynamic data on exit
-		ON_SCOPE_EXIT
-		{
-			FIrisPackageMapExportsUtil::FreeDynamicState(QuantizedExports);
-		};
+			// We need to release dynamic data on exit
+			ON_SCOPE_EXIT
+			{
+				FIrisPackageMapExportsUtil::FreeDynamicState(QuantizedExports);
+			};
 
-		FIrisPackageMapExportsUtil::Quantize(Context, PackageMapExports, NetTokensPendingExport, QuantizedExports);
+			FIrisPackageMapExportsUtil::Quantize(Context, PackageMapExports, NetTokensPendingExport, QuantizedExports);
 
-		// Setup export scope
-		FNetExports::FExportScope ExportScope = InitParams.NetExports->MakeExportScope(Context, Result->BatchExports);
+			// Setup export scope
+			FNetExports::FExportScope ExportScope = InitParams.NetExports->MakeExportScope(Context, Result->BatchExports);
 
 #if UE_NET_TRACE_ENABLED
-		FNetTraceCollector* LocalTraceCollector = UE_NET_TRACE_CREATE_COLLECTOR(ENetTraceVerbosity::Trace);
-		Context.SetTraceCollector(LocalTraceCollector);
+			FNetTraceCollector* LocalTraceCollector = UE_NET_TRACE_CREATE_COLLECTOR(ENetTraceVerbosity::Trace);
+			Context.SetTraceCollector(LocalTraceCollector);
 #endif
 
-		const uint32 ExportHeaderPos = ExportsWriter.GetPosBits();
+			const uint32 ExportHeaderPos = ExportsWriter.GetPosBits();
 
-		UE_NET_TRACE_SCOPE(ExportPayload, ExportsWriter, Context.GetTraceCollector(), ENetTraceVerbosity::Verbose);
+			UE_NET_TRACE_SCOPE(ExportPayload, ExportsWriter, Context.GetTraceCollector(), ENetTraceVerbosity::Verbose);
 
-		// Header
-		ExportsWriter.WriteBits(0, FChunkedDataStreamParameters::NumBitsForExportOffset);
+			// Header
+			ExportsWriter.WriteBits(0, FChunkedDataStreamParameters::NumBitsForExportOffset);
 
-		// Append potential exports
-		ObjectReferenceCache->AddPendingExports(Context, MakeArrayView(QuantizedExports.ObjectReferenceStorage.GetData(), QuantizedExports.ObjectReferenceStorage.Num()));
-		
-		// Serialize the reference data
-		FIrisPackageMapExportsUtil::Serialize(Context, QuantizedExports);
+			// Append potential exports
+			ObjectReferenceCache->AddPendingExports(Context, MakeArrayView(QuantizedExports.ObjectReferenceStorage.GetData(), QuantizedExports.ObjectReferenceStorage.Num()));
 
-		const uint32 WrittenBitsInBatch = (ExportsWriter.GetPosBits() - ExportHeaderPos) - FChunkedDataStreamParameters::NumBitsForExportOffset;
+			// Serialize the reference data
+			FIrisPackageMapExportsUtil::Serialize(Context, QuantizedExports);
 
-		// Serialize exports if there are any
-		const FObjectReferenceCache::EWriteExportsResult WriteExportResult = ObjectReferenceCache->WritePendingExports(Context, 0);
-		if (WriteExportResult == FObjectReferenceCache::EWriteExportsResult::BitStreamOverflow)
-		{
-			return nullptr;
-		}
-		else if (WriteExportResult == FObjectReferenceCache::EWriteExportsResult::WroteExports)
-		{
-			// Go back and update header that we have exports to read.
-			FNetBitStreamWriteScope SizeScope(ExportsWriter, ExportHeaderPos);
-			ExportsWriter.WriteBits(WrittenBitsInBatch, FChunkedDataStreamParameters::NumBitsForExportOffset);
-		}
-		ExportsWriter.CommitWrites();
+			const uint32 WrittenBitsInBatch = (ExportsWriter.GetPosBits() - ExportHeaderPos) - FChunkedDataStreamParameters::NumBitsForExportOffset;
+
+			// Serialize exports if there are any
+			const FObjectReferenceCache::EWriteExportsResult WriteExportResult = ObjectReferenceCache->WritePendingExports(Context, 0);
+			if (WriteExportResult == FObjectReferenceCache::EWriteExportsResult::BitStreamOverflow)
+			{
+				return nullptr;
+			}
+			else if (WriteExportResult == FObjectReferenceCache::EWriteExportsResult::WroteExports)
+			{
+				// Go back and update header that we have exports to read.
+				FNetBitStreamWriteScope SizeScope(ExportsWriter, ExportHeaderPos);
+				ExportsWriter.WriteBits(WrittenBitsInBatch, FChunkedDataStreamParameters::NumBitsForExportOffset);
+			}
+			ExportsWriter.CommitWrites();
+
+			ExportsPayloadBytes = ExportsWriter.GetPosBytes();
 
-		// Trim down the payload.
-		Result->ExportsPayload.SetNum(ExportsWriter.GetPosBytes(), EAllowShrinking::Yes);
 #if UE_NET_TRACE_ENABLED
-		Result->TraceCollector = LocalTraceCollector;
-		LocalTraceCollector = nullptr;
+			Result->TraceCollector = LocalTraceCollector;
+			LocalTraceCollector = nullptr;
 #endif
+		}
+
+		// Trim down the payload.
+		Result->ExportsPayload.SetNum(ExportsPayloadBytes, EAllowShrinking::Yes);
 
 		return Result.Release();
 	}
