Chaos struct size optimizations through member reordering.

FKinematicTarget: 64 ->48 bytes
FParticleNonFrequentData: 32 -> 24 bytes
FParticleMassProps: 112 -> 96 bytes

FPBDRigidParticle: 464 -> 432 bytes. Moved from 512 byte mem bucket to 448 byte bucket.

diff --git a/Engine/Source/Runtime/Experimental/Chaos/Public/Chaos/KinematicTargets.h b/Engine/Source/Runtime/Experimental/Chaos/Public/Chaos/KinematicTargets.h
--- a/Engine/Source/Runtime/Experimental/Chaos/Public/Chaos/KinematicTargets.h
+++ b/Engine/Source/Runtime/Experimental/Chaos/Public/Chaos/KinematicTargets.h
@@ -45,8 +45,8 @@ namespace Chaos
 		}
 
 		FKinematicTarget()
-			: Position(0)
-			, Rotation(FRotation3f::FromIdentity())
+			: Rotation(FRotation3f::FromIdentity())
+			, Position(0)
 			, Mode(EKinematicTargetMode::None)
 		{
 		}
@@ -185,21 +185,21 @@ namespace Chaos
 
 	private:
 		explicit FKinematicTarget(const FRigidTransform3& InTransform)
-			: Position(InTransform.GetTranslation())
-			, Rotation(InTransform.GetRotation())
+			: Rotation(InTransform.GetRotation())
+			, Position(InTransform.GetTranslation())
 			, Mode(EKinematicTargetMode::Position)
 		{
 		}
 
 		FKinematicTarget(const FVec3& InPosition, const FRotation3f& InRotation)
-			: Position(InPosition)
-			, Rotation(InRotation)
+			: Rotation(InRotation)
+			, Position(InPosition)
 			, Mode(EKinematicTargetMode::Position)
 		{
 		}
 
-		FVec3 Position;
 		FRotation3f Rotation;
+		FVec3 Position;
 		EKinematicTargetMode Mode;
 	};
 
