Chaos - API supporting changes

diff --git a/Engine/Source/Runtime/Experimental/Chaos/Private/PBDRigidsSolver.cpp b/Engine/Source/Runtime/Experimental/Chaos/Private/PBDRigidsSolver.cpp
--- a/Engine/Source/Runtime/Experimental/Chaos/Private/PBDRigidsSolver.cpp
+++ b/Engine/Source/Runtime/Experimental/Chaos/Private/PBDRigidsSolver.cpp
@@ -683,29 +683,13 @@ namespace Chaos
 		MEvolution->SetPreIntegrateCallback(
 			[this](FReal Dt)
 			{
-				for (ISimCallbackObject* Callback : SimCallbackObjects)
-				{
-					if (Callback->HasOption(ESimCallbackOptions::PreIntegrate))
-					{
-						FScopedTraceSolverCallback TraceCallback(Callback);
-						Callback->PreIntegrate_Internal();
-					}
-				}
-
-				PreIntegrateDebugDraw(Dt);
+				OnEvolutionPreIntegrate(Dt);
 			});
 
 		MEvolution->SetPostIntegrateCallback(
 			[this](FReal Dt)
 			{
-				for (ISimCallbackObject* Callback : SimCallbackObjects)
-				{
-					if (Callback->HasOption(ESimCallbackOptions::PostIntegrate))
-					{
-						FScopedTraceSolverCallback TraceCallback(Callback);
-						Callback->PostIntegrate_Internal();
-					}
-				}
+				OnEvolutionPostIntegrate(Dt);
 			});
 
 		MEvolution->SetPreSolveCallback(
@@ -737,6 +721,33 @@ namespace Chaos
 			});
 	}
 
+	void FPBDRigidsSolver::OnEvolutionPreIntegrate(FReal Dt)
+	{
+		for (ISimCallbackObject* Callback : SimCallbackObjects)
+		{
+			if (Callback->HasOption(ESimCallbackOptions::PreIntegrate))
+			{
+				FScopedTraceSolverCallback TraceCallback(Callback);
+				Callback->PreIntegrate_Internal();
+			}
+		}
+
+		PreIntegrateDebugDraw(Dt);
+	}
+
+	void FPBDRigidsSolver::OnEvolutionPostIntegrate(FReal Dt)
+	{
+		for (ISimCallbackObject* Callback : SimCallbackObjects)
+		{
+			if (Callback->HasOption(ESimCallbackOptions::PostIntegrate))
+			{
+				FScopedTraceSolverCallback TraceCallback(Callback);
+				Callback->PostIntegrate_Internal();
+			}
+		}
+	}
+
+
 #if CHAOS_DEBUG_NAME
 	void FPBDRigidsSolver::OnDebugNameChanged()
 	{
diff --git a/Engine/Source/Runtime/Experimental/Chaos/Public/PBDRigidsSolver.h b/Engine/Source/Runtime/Experimental/Chaos/Public/PBDRigidsSolver.h
--- a/Engine/Source/Runtime/Experimental/Chaos/Public/PBDRigidsSolver.h
+++ b/Engine/Source/Runtime/Experimental/Chaos/Public/PBDRigidsSolver.h
@@ -353,6 +353,9 @@ namespace Chaos
 		CHAOS_API virtual void ApplyCallbacks_Internal() override;
 		CHAOS_API FSolverSerializer& GetSerializer();
 
+		CHAOS_API void OnEvolutionPreIntegrate(FReal Dt);
+		CHAOS_API void OnEvolutionPostIntegrate(FReal Dt);
+
 	protected:
 
 #if CHAOS_DEBUG_NAME
