Remove more use of DefaultHostGroup from the IAS backend.

#rnx

- Remove direct use of DefaultHostGroup and replace it with calls to FHostGroupManager.
- Now we only reference DefaultHostGroup when setting it up.
- In theory at this point any container could mount it's own host group and the backend would work with it. We just need to add a way for the container to reference it.
- Added thread safety to FHostGroupManager as we now have cvars invoking debugging behaviour via FHostGroupManager::DisconnectAll.

diff --git a/Engine/Source/Runtime/Experimental/IoStore/OnDemand/Private/IasHostGroup.cpp b/Engine/Source/Runtime/Experimental/IoStore/OnDemand/Private/IasHostGroup.cpp
--- a/Engine/Source/Runtime/Experimental/IoStore/OnDemand/Private/IasHostGroup.cpp
+++ b/Engine/Source/Runtime/Experimental/IoStore/OnDemand/Private/IasHostGroup.cpp
@@ -3,11 +3,13 @@
 #include "IasHostGroup.h"
 
 #include "Algo/Find.h"
+#include "Async/UniqueLock.h"
 #include "HAL/IConsoleManager.h"
 #include "IO/IoStoreOnDemand.h"
 #include "IO/OnDemandHostGroup.h"
 #include "LatencyTesting.h"
 #include "Logging/StructuredLog.h"
+#include "Statistics.h"
 
 namespace UE::IoStore
 {
@@ -304,7 +306,10 @@ TIoStatusOr<FIASHostGroup> FHostGroupManager::Register(FName Name, FAnsiStringVi
 {
 	FIASHostGroup HostGroup(Name, TestPath);
 
-	HostGroups.Add(HostGroup);
+	{
+		UE::TUniqueLock _(Mutex);
+		HostGroups.Add(HostGroup);
+	}
 
 	return MoveTemp(HostGroup);
 }
@@ -315,6 +320,7 @@ TIoStatusOr<FIASHostGroup> FHostGroupManager::Register(FName Name, TConstArrayVi
 
 	if (Result.IsOk())
 	{
+		UE::TUniqueLock _(Mutex);
 		HostGroups.Add(Result.ValueOrDie());
 	}
 
@@ -323,6 +329,8 @@ TIoStatusOr<FIASHostGroup> FHostGroupManager::Register(FName Name, TConstArrayVi
 
 FIASHostGroup FHostGroupManager::Find(FName Name)
 {
+	UE::TUniqueLock _(Mutex);
+
 	FIASHostGroup* FoundHostGroup = Algo::FindByPredicate(HostGroups, [Name](const FIASHostGroup& HostGroup)->bool
 	{
 		return HostGroup.GetName() == Name;
@@ -338,4 +346,38 @@ FIASHostGroup FHostGroupManager::Find(FName Name)
 	}
 }
 
+void FHostGroupManager::Tick(uint32 TimeoutMs, std::atomic_bool& CancellationToken)
+{
+	UE::TUniqueLock _(Mutex);
+	for (FIASHostGroup& HostGroup : HostGroups)
+	{
+		if (HostGroup.AttemptReconnection(TimeoutMs, CancellationToken) == FIASHostGroup::EReconnectionResult::Reconnected)
+		{
+			FOnDemandIoBackendStats::Get()->OnHttpConnected(); // TODO: Try to avoid singleton access somehow
+		}
+	}
+}
+
+void FHostGroupManager::DisconnectAll()
+{
+	UE::TUniqueLock _(Mutex);
+	for (FIASHostGroup& HostGroup : HostGroups)
+	{
+		HostGroup.Disconnect();
+	}
+}
+
+uint32 FHostGroupManager::GetNumDisconnctedHosts() const
+{
+	uint32 NumDisconnectedHosts = 0;
+
+	UE::TUniqueLock _(Mutex);
+	for (const FIASHostGroup& HostGroup : HostGroups)
+	{
+		NumDisconnectedHosts += HostGroup.IsConnected() == false ? 1 : 0;
+	}
+
+	return NumDisconnectedHosts;
+}
+
 } //namespace UE::IoStore
diff --git a/Engine/Source/Runtime/Experimental/IoStore/OnDemand/Private/IasHostGroup.h b/Engine/Source/Runtime/Experimental/IoStore/OnDemand/Private/IasHostGroup.h
--- a/Engine/Source/Runtime/Experimental/IoStore/OnDemand/Private/IasHostGroup.h
+++ b/Engine/Source/Runtime/Experimental/IoStore/OnDemand/Private/IasHostGroup.h
@@ -2,6 +2,7 @@
 
 #pragma once
 
+#include "Async/Mutex.h"
 #include "IO/IoStatus.h"
 #include "OnDemandBackendStatus.h"
 #include "Templates/SharedPointer.h"
@@ -11,7 +12,7 @@ namespace UE::IoStore
 
 class FOnDemandHostGroup;
 
-class FIASHostGroup
+class FIASHostGroup final
 {
 	struct FImpl;
 	using FSharedImpl = TSharedPtr<FImpl, ESPMode::ThreadSafe>;
@@ -79,9 +80,17 @@ public:
 	TIoStatusOr<FIASHostGroup> Register(FName Name, TConstArrayView<FAnsiString> HostUrls); // <- Resolved
 
 	FIASHostGroup				Find(FName Name);
+
+	void Tick(uint32 TimeoutMs, std::atomic_bool& CancellationToken);
+
+	void DisconnectAll(); 
+	uint32 GetNumDisconnctedHosts() const;
+
 private:
 
 	TArray<FIASHostGroup> HostGroups;
+
+	mutable UE::FMutex Mutex;
 };
 
 } //namespace UE::IoStore
diff --git a/Engine/Source/Runtime/Experimental/IoStore/OnDemand/Private/OnDemandIoDispatcherBackend.cpp b/Engine/Source/Runtime/Experimental/IoStore/OnDemand/Private/OnDemandIoDispatcherBackend.cpp
--- a/Engine/Source/Runtime/Experimental/IoStore/OnDemand/Private/OnDemandIoDispatcherBackend.cpp
+++ b/Engine/Source/Runtime/Experimental/IoStore/OnDemand/Private/OnDemandIoDispatcherBackend.cpp
@@ -919,13 +919,13 @@ FOnDemandIoBackend::FOnDemandIoBackend(const FOnDemandEndpointConfig& Config, FO
 	DynamicConsoleCommands.Emplace(
 	IConsoleManager::Get().RegisterConsoleCommand(
 		TEXT("ias.InvokeHttpFailure"),
-		TEXT("Marks the current ias http connection as failed forcing the system to try to reconnect"),
+		TEXT("Marks the current ias http connections as failed and forcing them  to try to reconnect"),
 		FConsoleCommandDelegate::CreateLambda([this]()
 			{
 				UE_LOG(LogIas, Display, TEXT("User invoked http error via 'ias.InvokeHttpFailure'"));
-				DefaultHostGroup.Disconnect();
+				FHostGroupManager::Get().DisconnectAll();
 
-				TickBackendEvent->Trigger();
+				TickBackendEvent->Trigger(); // Wake up the backend in case it is ideal waiting for requests
 			}),
 		ECVF_Default)
 	);
@@ -1240,8 +1240,6 @@ bool FOnDemandIoBackend::Resolve(FIoRequestImpl* Request)
 		check(GetStatus == EIoErrorCode::NotFound);
 	}
 
-	// Currently all requests are assigned the DefaultHostGroup
-	// TODO: When next fixing this, make sure we deal with the !DefaultHostGroup.IsConnected() above
 	ChunkRequest->HostGroup = ChunkInfo.HostGroup();
 
 	Stats.OnHttpEnqueue();
@@ -1312,7 +1310,7 @@ TIoStatusOr<uint64> FOnDemandIoBackend::GetSizeForChunk(const FIoChunkId& ChunkI
 	FIoOffsetAndLength RequestedRange(ChunkRange.GetOffset(), FMath::Min<uint64>(ChunkInfo.RawSize(), ChunkRange.GetLength()));
 	OutAvailable = ChunkInfo.RawSize();
 
-	if (!BackendStatus.IsHttpEnabled(ChunkId.GetChunkType()) || !DefaultHostGroup.IsConnected())
+	if (!BackendStatus.IsHttpEnabled(ChunkId.GetChunkType()))
 	{
 		// If the cache is not readonly the chunk may get evicted before the request is resolved
 		if (BackendStatus.IsCacheReadOnly() == false)
@@ -1676,10 +1674,7 @@ uint32 FOnDemandIoBackend::Run()
 
 		if (!bStopRequested)
 		{
-			if (DefaultHostGroup.AttemptReconnection(GIasHttpTimeOutMs, bStopRequested) == FIASHostGroup::EReconnectionResult::Reconnected)
-			{
-				Stats.OnHttpConnected();
-			}
+			FHostGroupManager::Get().Tick(GIasHttpTimeOutMs, bStopRequested);
 
 			if (BackendStatus.ShouldAbandonCache())
 			{
@@ -1691,7 +1686,8 @@ uint32 FOnDemandIoBackend::Run()
 					Cache.Release()->Abandon(); // Will delete its self
 				}
 			}
-			TickBackendEvent->Wait(!DefaultHostGroup.IsConnected() ? GIasHttpHealthCheckWaitTime : MAX_uint32);
+
+			TickBackendEvent->Wait(FHostGroupManager::Get().GetNumDisconnctedHosts() > 0 ? GIasHttpHealthCheckWaitTime : MAX_uint32);
 		}
 	}
 
