Geometry collection : Make sur dynamic removal facades are accessible from outside of their plugin

diff --git a/Engine/Source/Runtime/Experimental/GeometryCollectionEngine/Private/GeometryCollection/GeometryCollectionEngineRemoval.cpp b/Engine/Source/Runtime/Experimental/GeometryCollectionEngine/Private/GeometryCollection/GeometryCollectionEngineRemoval.cpp
--- a/Engine/Source/Runtime/Experimental/GeometryCollectionEngine/Private/GeometryCollection/GeometryCollectionEngineRemoval.cpp
+++ b/Engine/Source/Runtime/Experimental/GeometryCollectionEngine/Private/GeometryCollection/GeometryCollectionEngineRemoval.cpp
@@ -149,6 +149,26 @@ void FGeometryCollectionRemoveOnSleepDynamicFacade::DefineSchema()
 	LastPositionAttribute.AddAndFill(FVector::ZeroVector);
 }
 
+float FGeometryCollectionRemoveOnSleepDynamicFacade::GetSleepRemovalDuration(int32 TransformIndex) const
+{
+	return SleepRemovalDurationAttribute[TransformIndex];
+}
+
+float FGeometryCollectionRemoveOnSleepDynamicFacade::GetMaxSleepTime(int32 TransformIndex) const
+{
+	return MaxSleepTimeAttribute[TransformIndex];
+}
+
+void FGeometryCollectionRemoveOnSleepDynamicFacade::SetSleepRemovalDuration(int32 TransformIndex, float SleepRemovalDuration)
+{
+	SleepRemovalDurationAttribute.ModifyAt(TransformIndex, SleepRemovalDuration);
+}
+
+void FGeometryCollectionRemoveOnSleepDynamicFacade::SetMaxSleepTime(int32 TransformIndex, float MaxSleepTime)
+{
+	MaxSleepTimeAttribute.ModifyAt(TransformIndex, MaxSleepTime);
+}
+
 void FGeometryCollectionRemoveOnSleepDynamicFacade::SetAttributeValues(const FVector2D& MaximumSleepTime, const FVector2D& RemovalDuration)
 {
 	check(!IsConst());
@@ -177,7 +197,7 @@ void FGeometryCollectionRemoveOnSleepDynamicFacade::SetAttributeValues(const FVe
 
 bool FGeometryCollectionRemoveOnSleepDynamicFacade::IsRemovalActive(int32 TransformIndex) const
 {
-	return (MaxSleepTimeAttribute.Get()[TransformIndex] >= 0);
+	return (MaxSleepTimeAttribute[TransformIndex] >= 0);
 }
 
 bool FGeometryCollectionRemoveOnSleepDynamicFacade::ComputeSlowMovingState(int32 TransformIndex, const FVector& Position, float DeltaTime, FVector::FReal VelocityThreshold)
diff --git a/Engine/Source/Runtime/Experimental/GeometryCollectionEngine/Public/GeometryCollection/GeometryCollectionEngineRemoval.h b/Engine/Source/Runtime/Experimental/GeometryCollectionEngine/Public/GeometryCollection/GeometryCollectionEngineRemoval.h
--- a/Engine/Source/Runtime/Experimental/GeometryCollectionEngine/Public/GeometryCollection/GeometryCollectionEngineRemoval.h
+++ b/Engine/Source/Runtime/Experimental/GeometryCollectionEngine/Public/GeometryCollection/GeometryCollectionEngineRemoval.h
@@ -25,31 +25,31 @@ public:
 	static constexpr float DisabledPostBreakDuration = -1;
 	static constexpr float CrumblingRemovalTimer = -1;
 	
-	FGeometryCollectionRemoveOnBreakDynamicFacade(FGeometryDynamicCollection& InCollection);
+	GEOMETRYCOLLECTIONENGINE_API FGeometryCollectionRemoveOnBreakDynamicFacade(FGeometryDynamicCollection& InCollection);
 
 	/** 
 	 * returns true if all the necessary attributes are present
 	 * if not then the API can be used to create  
 	 */
-	bool IsValid() const;	
+	GEOMETRYCOLLECTIONENGINE_API bool IsValid() const;
 
 	/** Is this facade const access */
-	bool IsConst() const;
+	GEOMETRYCOLLECTIONENGINE_API bool IsConst() const;
 
 	/** Add the relevant attributes */
-	void DefineSchema();
+	GEOMETRYCOLLECTIONENGINE_API void DefineSchema();
 
 	/**
 	 * Add the necessary attributes if they are missing and initialize them if necessary
 	 * @param RemoveOnBreakFacade remove on break facade from the rest collection that contain the original user set attributes
 	*/
-	void SetAttributeValues(const GeometryCollection::Facades::FCollectionRemoveOnBreakFacade& RemoveOnBreakFacade);
+	GEOMETRYCOLLECTIONENGINE_API void SetAttributeValues(const GeometryCollection::Facades::FCollectionRemoveOnBreakFacade& RemoveOnBreakFacade);
 
 	/** true if the removal is active for a specific piece */
-	bool IsRemovalActive(int32 TransformIndex) const;
+	GEOMETRYCOLLECTIONENGINE_API bool IsRemovalActive(int32 TransformIndex) const;
 
 	/** true if a specific transform uses cluster crumbling */
-	bool UseClusterCrumbling(int32 TransformIndex) const;
+	GEOMETRYCOLLECTIONENGINE_API bool UseClusterCrumbling(int32 TransformIndex) const;
 	
 	/**
 	 * Update break timer and return the matching decay  
@@ -57,7 +57,7 @@ public:
 	 * @param DeltaTime elapsed time since the last update in second
 	 * @return decay value computed from the timer and duration ( [0,1] range )
 	 */
-	float UpdateBreakTimerAndComputeDecay(int32 TransformIndex, float DeltaTime);
+	GEOMETRYCOLLECTIONENGINE_API float UpdateBreakTimerAndComputeDecay(int32 TransformIndex, float DeltaTime);
 	
 private:
 	/** Time elapsed since the break in seconds */
@@ -79,29 +79,35 @@ private:
 class FGeometryCollectionRemoveOnSleepDynamicFacade
 {
 public:
-	FGeometryCollectionRemoveOnSleepDynamicFacade(FManagedArrayCollection& InCollection);
+	GEOMETRYCOLLECTIONENGINE_API FGeometryCollectionRemoveOnSleepDynamicFacade(FManagedArrayCollection& InCollection);
 
 	/** 
 	 * returns true if all the necessary attributes are present
 	 * if not then the API can be used to create  
 	 */
-	bool IsValid() const;	
+	GEOMETRYCOLLECTIONENGINE_API bool IsValid() const;
 
 	/** Is this facade const access */
-	bool IsConst() const;
+	GEOMETRYCOLLECTIONENGINE_API bool IsConst() const;
 
 	/** Add the relevant attributes */
-	void DefineSchema();
+	GEOMETRYCOLLECTIONENGINE_API void DefineSchema();
+
+	GEOMETRYCOLLECTIONENGINE_API float GetSleepRemovalDuration(int32 TransformIndex) const;
+	GEOMETRYCOLLECTIONENGINE_API float GetMaxSleepTime(int32 TransformIndex) const;
+
+	GEOMETRYCOLLECTIONENGINE_API void SetSleepRemovalDuration(int32 TransformIndex, float SleepRemovalDuration);
+	GEOMETRYCOLLECTIONENGINE_API void SetMaxSleepTime(int32 TransformIndex, float MaxSleepTime);
 
 	/**
 	 * Add the necessary attributes if they are missing and initialize them if necessary
 	 * @param MaximumSleepTime range of time to initialize the sleep duration
 	 * @param RemovalDuration range of time to initialize the removal duration
 	*/
-	void SetAttributeValues(const FVector2D& MaximumSleepTime, const FVector2D& RemovalDuration);
+	GEOMETRYCOLLECTIONENGINE_API void SetAttributeValues(const FVector2D& MaximumSleepTime, const FVector2D& RemovalDuration);
 
 	/** true if the removal is active for a specific piece */
-	bool IsRemovalActive(int32 TransformIndex) const;
+	GEOMETRYCOLLECTIONENGINE_API bool IsRemovalActive(int32 TransformIndex) const;
 
 	/**
 	 * Compute the slow moving state and update from the last position
@@ -110,17 +116,17 @@ public:
 	 * @param DeltaTime elapsed time since last update
 	 * @return true if the piece if the piece is considered slow moving 
 	 **/
-	bool ComputeSlowMovingState(int32 TransformIndex, const FVector& Position, float DeltaTime, FVector::FReal VelocityThreshold);
+	GEOMETRYCOLLECTIONENGINE_API bool ComputeSlowMovingState(int32 TransformIndex, const FVector& Position, float DeltaTime, FVector::FReal VelocityThreshold);
 	
 	/**
 	 * Update the sleep timer
 	 * @param TransformIndex index of the transform to update
 	 * @param DeltaTime elapsed time since last update
 	 */
-	void UpdateSleepTimer(int32 TransformIndex, float DeltaTime);
+	GEOMETRYCOLLECTIONENGINE_API void UpdateSleepTimer(int32 TransformIndex, float DeltaTime);
 
 	/** Compute decay from elapsed timer and duration attributes */
-	float ComputeDecay(int32 TransformIndex) const;
+	GEOMETRYCOLLECTIONENGINE_API float ComputeDecay(int32 TransformIndex) const;
 
 	
 private:
@@ -143,25 +149,25 @@ private:
 class FGeometryCollectionDecayDynamicFacade
 {
 public:
-	FGeometryCollectionDecayDynamicFacade(FManagedArrayCollection& InCollection);
+	GEOMETRYCOLLECTIONENGINE_API FGeometryCollectionDecayDynamicFacade(FManagedArrayCollection& InCollection);
 
 	/** 
 	 * returns true if all the necessary attributes are present
 	 * if not then the API can be used to create  
 	 */
-	bool IsValid() const;	
+	GEOMETRYCOLLECTIONENGINE_API bool IsValid() const;
 
 	/** Add the necessary attributes if they are missing and initialize them if necessary */
-	void AddAttributes();
+	GEOMETRYCOLLECTIONENGINE_API void AddAttributes();
 
 	/** Get decay value for a specific transform index */
-	float GetDecay(int32 TransformIndex) const;
+	GEOMETRYCOLLECTIONENGINE_API float GetDecay(int32 TransformIndex) const;
 
 	/** Set decay value for a specific transform index */
-	void SetDecay(int32 TransformIndex, float DecayValue);
+	GEOMETRYCOLLECTIONENGINE_API void SetDecay(int32 TransformIndex, float DecayValue);
 
 	/** Get the size of the decay attribute - this should match the number of transforms of the collection */
-	int32 GetDecayAttributeSize() const;
+	GEOMETRYCOLLECTIONENGINE_API int32 GetDecayAttributeSize() const;
 
 private:
 	/** state of decay ([0-1] range) */
