Dataflow : refine anytype pin drag and drop on subgraph output/input nodes and use the right concrete type if set on the dragged-from pin

diff --git a/Engine/Source/Runtime/Experimental/Dataflow/Core/Public/Dataflow/DataflowTypePolicy.h b/Engine/Source/Runtime/Experimental/Dataflow/Core/Public/Dataflow/DataflowTypePolicy.h
--- a/Engine/Source/Runtime/Experimental/Dataflow/Core/Public/Dataflow/DataflowTypePolicy.h
+++ b/Engine/Source/Runtime/Experimental/Dataflow/Core/Public/Dataflow/DataflowTypePolicy.h
@@ -16,6 +16,13 @@ struct TDataflowPolicyTypeName
 {
 };
 
+namespace UE::Dataflow
+{
+	template<typename T>
+	FName GetTypeName() { return TDataflowPolicyTypeName<T>::GetName(); }
+}
+
+
 // special version for void 
 template<>
 struct TDataflowPolicyTypeName<void>
diff --git a/Engine/Source/Runtime/Experimental/Dataflow/Engine/Private/Dataflow/DataflowSubGraphNodes.cpp b/Engine/Source/Runtime/Experimental/Dataflow/Engine/Private/Dataflow/DataflowSubGraphNodes.cpp
--- a/Engine/Source/Runtime/Experimental/Dataflow/Engine/Private/Dataflow/DataflowSubGraphNodes.cpp
+++ b/Engine/Source/Runtime/Experimental/Dataflow/Engine/Private/Dataflow/DataflowSubGraphNodes.cpp
@@ -134,18 +134,89 @@ namespace UE::Dataflow::SubGraphNodes::Private
 		// AnyType require special treatment to make sure we give a concrete type
 		if (Connection.IsAnyType())
 		{
-			if (ConnectionProperty->GetClass()->IsChildOf(FStructProperty::StaticClass()))
+			EPropertyBagPropertyType PropertyType = EPropertyBagPropertyType::Count;
+
+			using namespace UE::Dataflow;
+			const FName Type = Connection.GetType();
+			if (Type == GetTypeName<bool>())
+			{
+				PropertyType = EPropertyBagPropertyType::Bool;
+			}
+			else if (Type == GetTypeName<int32>())
+			{
+				PropertyType = EPropertyBagPropertyType::Int32;
+			}
+			else if (Type == GetTypeName<int64>())
+			{
+				PropertyType = EPropertyBagPropertyType::Int64;
+			}
+			else if (Type == GetTypeName<uint32>())
+			{
+				PropertyType = EPropertyBagPropertyType::UInt32;
+			}
+			else if (Type == GetTypeName<uint64>())
+			{
+				PropertyType = EPropertyBagPropertyType::UInt64;
+			}
+			else if (Type == GetTypeName<float>())
+			{
+				PropertyType = EPropertyBagPropertyType::Float;
+			}
+			else if (Type == GetTypeName<double>())
+			{
+				PropertyType = EPropertyBagPropertyType::Double;
+			}
+			else if (Type == GetTypeName<FName>())
+			{
+				PropertyType = EPropertyBagPropertyType::Name;
+			}
+			else if (Type == GetTypeName<FString>())
+			{
+				PropertyType = EPropertyBagPropertyType::String;
+			}
+			else if (Type == GetTypeName<FText>())
+			{
+				PropertyType = EPropertyBagPropertyType::Text;
+			}
+			else
+			{
+				const FString TypeStr(Type.ToString());
+				if (TypeStr.StartsWith("U"))
+				{
+					// Nothing for now and use thefallback - may need fixing if we find there's an issue with that code path 
+				}
+				else if (TypeStr.StartsWith("F"))
+				{
+					static UPackage* CoreUObjectPkg = FindObjectChecked<UPackage>(nullptr, TEXT("/Script/CoreUObject"));
+					const FName ShortTypeName(TypeStr.RightChop(1));
+					UScriptStruct* Result = (UScriptStruct*)StaticFindObjectFastInternal(UScriptStruct::StaticClass(), CoreUObjectPkg, ShortTypeName, false, RF_NoFlags, EInternalObjectFlags::None);
+					if (Result)
+					{
+						return FPropertyBagPropertyDesc(Connection.GetName(), EPropertyBagPropertyType::Struct, Result);
+					}
+				}
+			}
+
+			if (PropertyType != EPropertyBagPropertyType::Count)
+			{
+				return FPropertyBagPropertyDesc(Connection.GetName(), PropertyType);
+			}
+			else
 			{
-				if (const FStructProperty* StructProperty = CastField<FStructProperty>(ConnectionProperty))
+				// Fallback let's use the AnyType default (Value) type
+				if (ConnectionProperty->GetClass()->IsChildOf(FStructProperty::StaticClass()))
 				{
-					if (const bool bInheritsFromAnyType = StructProperty->Struct->IsChildOf<FDataflowAnyType>())
+					if (const FStructProperty* StructProperty = CastField<FStructProperty>(ConnectionProperty))
 					{
-						for (TFieldIterator<FProperty> FieldIt(StructProperty->Struct); FieldIt; ++FieldIt)
+						if (const bool bInheritsFromAnyType = StructProperty->Struct->IsChildOf<FDataflowAnyType>())
 						{
-							FProperty* Prop = *FieldIt;
-							if (Prop->GetFName() == "Value")
+							for (TFieldIterator<FProperty> FieldIt(StructProperty->Struct); FieldIt; ++FieldIt)
 							{
-								return FPropertyBagPropertyDesc(Connection.GetName(), Prop);
+								FProperty* Prop = *FieldIt;
+								if (Prop->GetFName() == "Value")
+								{
+									return FPropertyBagPropertyDesc(Connection.GetName(), Prop);
+								}
 							}
 						}
 					}
