Include KiB/s in perf sample and packed samples a little tighter.

#rnx

diff --git a/Engine/Source/Runtime/Experimental/IoStore/HttpClient/Internal/IO/Http/Client.h b/Engine/Source/Runtime/Experimental/IoStore/HttpClient/Internal/IO/Http/Client.h
--- a/Engine/Source/Runtime/Experimental/IoStore/HttpClient/Internal/IO/Http/Client.h
+++ b/Engine/Source/Runtime/Experimental/IoStore/HttpClient/Internal/IO/Http/Client.h
@@ -165,9 +165,12 @@ class UE_API FTicketPerf
 public:
 	struct FSample
 	{
-		uint32			SendMs;
-		uint32			WaitMs;
-		uint32			RecvMs;
+		uint16			SendMs;
+		uint16			WaitMs;
+		uint16			RecvMs;
+		uint16			RecvKiBps;
+		uint32			GetTotalMs() const		{ return SendMs + WaitMs + RecvMs; }
+		uint32			GetSizeEstKiB() const	{ return (RecvKiBps * RecvMs) / 1000; }
 	};
 
 	FSample				GetSample() const;
@@ -185,13 +188,8 @@ public:
 
 	struct FError
 	{
-	//	FError() = default;
-	//	FError
-	//	~FError() = default;
-
-
-		const char*	Reason;
-		uint32		Code;
+		const char*		Reason;
+		uint32			Code;
 	};
 
 	EId					GetId() const;
diff --git a/Engine/Source/Runtime/Experimental/IoStore/HttpClient/Private/Api.inl b/Engine/Source/Runtime/Experimental/IoStore/HttpClient/Private/Api.inl
--- a/Engine/Source/Runtime/Experimental/IoStore/HttpClient/Private/Api.inl
+++ b/Engine/Source/Runtime/Experimental/IoStore/HttpClient/Private/Api.inl
@@ -267,7 +267,7 @@ FTicketStatus::FError FTicketStatus::GetError() const
 ////////////////////////////////////////////////////////////////////////////////
 FTicketPerf::FSample FTicketPerf::GetSample() const
 {
-	const auto* Activity = (FActivity*)this;
+	const auto& Activity = *(FActivity*)this;
 
 	static uint64 Freq;
 	if (Freq == 0)
@@ -275,14 +275,24 @@ FTicketPerf::FSample FTicketPerf::GetSample() const
 		Freq = uint64(1.0 / FPlatformTime::GetSecondsPerCycle64());
 	}
 
-	auto ToMs = [] (uint64 Value) { return uint32((Value * 1000ull) / Freq); };
+	auto Clamp = [] (auto Value) { return uint16(FMath::Min(uint32(Value), 0xffffu)); };
+	auto ToMs = [&] (uint64 Value) { return Clamp((Value * 1000ull) / Freq); };
 
-	const FStopwatch& Stopwatch = Activity->Stopwatch;
-	return {
+	const FStopwatch& Stopwatch = Activity.Stopwatch;
+	FSample Sample = {
 		ToMs(Stopwatch.GetInterval(0)),
 		ToMs(Stopwatch.GetInterval(1)),
 		ToMs(Stopwatch.GetInterval(2)),
 	};
+
+	uint32 Bps = ~0u;
+	if (Sample.RecvMs)
+	{
+		Bps = uint32((Activity.Response.ContentLength * 1000ull) / Sample.RecvMs);
+	}
+	Sample.RecvKiBps = Clamp(Bps >> 10);
+
+	return Sample;
 }
 
 #endif // IAS_HTTP_WITH_PERF
