-Changed ESimulationInitializationState to an enum class so that's it's undecorated global member names can't cause ambiguities / name shadowing

diff --git a/Engine/Source/Runtime/Experimental/Chaos/Public/GeometryCollection/GeometryCollectionSimulationCoreTypes.h b/Engine/Source/Runtime/Experimental/Chaos/Public/GeometryCollection/GeometryCollectionSimulationCoreTypes.h
--- a/Engine/Source/Runtime/Experimental/Chaos/Public/GeometryCollection/GeometryCollectionSimulationCoreTypes.h
+++ b/Engine/Source/Runtime/Experimental/Chaos/Public/GeometryCollection/GeometryCollectionSimulationCoreTypes.h
@@ -82,7 +82,7 @@ struct FSharedSimulationSizeSpecificData
 //
 //
 //
-enum ESimulationInitializationState : uint8 { Unintialized = 0, Activated, Created, Initialized };
+enum class ESimulationInitializationState : uint8 { Unintialized = 0, Activated, Created, Initialized };
 
 
 /**
diff --git a/Engine/Source/Runtime/Experimental/GeometryCollectionEngine/Public/GeometryCollection/GeometryCollectionComponent.h b/Engine/Source/Runtime/Experimental/GeometryCollectionEngine/Public/GeometryCollection/GeometryCollectionComponent.h
--- a/Engine/Source/Runtime/Experimental/GeometryCollectionEngine/Public/GeometryCollection/GeometryCollectionComponent.h
+++ b/Engine/Source/Runtime/Experimental/GeometryCollectionEngine/Public/GeometryCollection/GeometryCollectionComponent.h
@@ -43,7 +43,7 @@ struct FDamageCollector;
 class FPhysScene_Chaos;
 class AISMPoolActor;
 class IGeometryCollectionExternalRenderInterface;
-enum ESimulationInitializationState : uint8;
+enum class ESimulationInitializationState : uint8;
 enum class EClusterConnectionTypeEnum : uint8;
 enum class EInitialVelocityTypeEnum : uint8;
 enum class EObjectStateTypeEnum : uint8;
