FOnDemandIoStore
- Don't fail if no ondemand containers are mounted. We want to trivially succeed when asking for packages that may exist locally.

diff --git a/Engine/Source/Runtime/Experimental/IoStore/OnDemand/Private/OnDemandIoStore.cpp b/Engine/Source/Runtime/Experimental/IoStore/OnDemand/Private/OnDemandIoStore.cpp
--- a/Engine/Source/Runtime/Experimental/IoStore/OnDemand/Private/OnDemandIoStore.cpp
+++ b/Engine/Source/Runtime/Experimental/IoStore/OnDemand/Private/OnDemandIoStore.cpp
@@ -1934,11 +1934,6 @@ FIoStatus FOnDemandIoStore::GetContainersAndPackagesForInstall(
 		}
 	}
 
-	if (OutContainersForInstallation.IsEmpty())
-	{
-		return EIoErrorCode::NotFound;
-	}
-
 	// It's not allowed to install all content
 	if (MountId.IsEmpty() && TagSets.IsEmpty() && PackageIds.IsEmpty())
 	{
