Chaos - fix event dispatcher if handlers are attached before the physics scene is created.

diff --git a/Engine/Source/Runtime/Experimental/ChaosSolverEngine/Private/Chaos/ChaosGameplayEventDispatcher.cpp b/Engine/Source/Runtime/Experimental/ChaosSolverEngine/Private/Chaos/ChaosGameplayEventDispatcher.cpp
--- a/Engine/Source/Runtime/Experimental/ChaosSolverEngine/Private/Chaos/ChaosGameplayEventDispatcher.cpp
+++ b/Engine/Source/Runtime/Experimental/ChaosSolverEngine/Private/Chaos/ChaosGameplayEventDispatcher.cpp
@@ -321,6 +321,11 @@ void UChaosGameplayEventDispatcher::FillPhysicsProxy(FPhysScene_Chaos& Scene, TA
 
 TArray<UObject*> UChaosGameplayEventDispatcher::GetInterestedProxyOwnersForCollisionEvents()
 {
+	if ((GetWorld() == nullptr) || (GetWorld()->GetPhysicsScene() == nullptr))
+	{
+		return {};
+	}
+
 	TArray<UObject*> Result;
 	FPhysScene_Chaos& Scene = *(GetWorld()->GetPhysicsScene());
 		
@@ -335,7 +340,13 @@ TArray<UObject*> UChaosGameplayEventDispatcher::GetInterestedProxyOwnersForColli
 
 TArray<UObject*> UChaosGameplayEventDispatcher::GetInterestedProxyOwnersForBreakingEvents()
 {
+	if ((GetWorld() == nullptr) || (GetWorld()->GetPhysicsScene() == nullptr))
+	{
+		return {};
+	}
+
 	TArray<UObject*> Result;
+
 	FPhysScene_Chaos& Scene = *(GetWorld()->GetPhysicsScene());
 
 	// look through all the components that someone is interested in and get all the proxies
@@ -349,6 +360,11 @@ TArray<UObject*> UChaosGameplayEventDispatcher::GetInterestedProxyOwnersForBreak
 
 TArray<UObject*> UChaosGameplayEventDispatcher::GetInterestedProxyOwnersForRemovalEvents()
 {
+	if ((GetWorld() == nullptr) || (GetWorld()->GetPhysicsScene() == nullptr))
+	{
+		return {};
+	}
+
 	TArray<UObject*> Result;
 	FPhysScene_Chaos& Scene = *(GetWorld()->GetPhysicsScene());
 
@@ -363,6 +379,11 @@ TArray<UObject*> UChaosGameplayEventDispatcher::GetInterestedProxyOwnersForRemov
 
 TArray<UObject*> UChaosGameplayEventDispatcher::GetInterestedProxyOwnersForCrumblingEvents()
 {
+	if ((GetWorld() == nullptr) || (GetWorld()->GetPhysicsScene() == nullptr))
+	{
+		return {};
+	}
+
 	TArray<UObject*> Result;
 	FPhysScene_Chaos& Scene = *(GetWorld()->GetPhysicsScene());
 
@@ -380,6 +401,11 @@ void UChaosGameplayEventDispatcher::HandleCollisionEvents(const Chaos::FCollisio
 	// todo(chaos) : this code is very similar to FPhysScene_Chaos::HandleCollisionEvents, we should propably consolidate if possible or share as much code as possible 
 	SCOPE_CYCLE_COUNTER(STAT_DispatchCollisionEvents);
 
+	if ((GetWorld() == nullptr) || (GetWorld()->GetPhysicsScene() == nullptr))
+	{
+		return;
+	}
+
 	FPhysScene_Chaos& Scene = *(GetWorld()->GetPhysicsScene());
 
 	PendingChaosCollisionNotifies.Reset();
