Dataflow add templated IsType() to connection

diff --git a/Engine/Source/Runtime/Experimental/Dataflow/Core/Public/Dataflow/DataflowConnection.h b/Engine/Source/Runtime/Experimental/Dataflow/Core/Public/Dataflow/DataflowConnection.h
--- a/Engine/Source/Runtime/Experimental/Dataflow/Core/Public/Dataflow/DataflowConnection.h
+++ b/Engine/Source/Runtime/Experimental/Dataflow/Core/Public/Dataflow/DataflowConnection.h
@@ -235,6 +235,12 @@ public:
 	FName GetType() const { return Type; }
 	FName GetOriginalType() const { return OriginalType; }
 
+	template <typename T>
+	bool IsType() const
+	{
+		return Type == UE::Dataflow::GetTypeName<T>();
+	}
+
 	FGuid GetGuid() const { return Guid; }
 	void SetGuid(FGuid InGuid) { Guid = InGuid; }
 	
