Dataflow : update subgraph node to use IsType method from connection

diff --git a/Engine/Source/Runtime/Experimental/Dataflow/Engine/Private/Dataflow/DataflowSubGraphNodes.cpp b/Engine/Source/Runtime/Experimental/Dataflow/Engine/Private/Dataflow/DataflowSubGraphNodes.cpp
--- a/Engine/Source/Runtime/Experimental/Dataflow/Engine/Private/Dataflow/DataflowSubGraphNodes.cpp
+++ b/Engine/Source/Runtime/Experimental/Dataflow/Engine/Private/Dataflow/DataflowSubGraphNodes.cpp
@@ -136,51 +136,49 @@ namespace UE::Dataflow::SubGraphNodes::Private
 		{
 			EPropertyBagPropertyType PropertyType = EPropertyBagPropertyType::Count;
 
-			using namespace UE::Dataflow;
-			const FName Type = Connection.GetType();
-			if (Type == GetTypeName<bool>())
+			if (Connection.IsType<bool>())
 			{
 				PropertyType = EPropertyBagPropertyType::Bool;
 			}
-			else if (Type == GetTypeName<int32>())
+			else if (Connection.IsType<int32>())
 			{
 				PropertyType = EPropertyBagPropertyType::Int32;
 			}
-			else if (Type == GetTypeName<int64>())
+			else if (Connection.IsType<int64>())
 			{
 				PropertyType = EPropertyBagPropertyType::Int64;
 			}
-			else if (Type == GetTypeName<uint32>())
+			else if (Connection.IsType<uint32>())
 			{
 				PropertyType = EPropertyBagPropertyType::UInt32;
 			}
-			else if (Type == GetTypeName<uint64>())
+			else if (Connection.IsType<uint64>())
 			{
 				PropertyType = EPropertyBagPropertyType::UInt64;
 			}
-			else if (Type == GetTypeName<float>())
+			else if (Connection.IsType<float>())
 			{
 				PropertyType = EPropertyBagPropertyType::Float;
 			}
-			else if (Type == GetTypeName<double>())
+			else if (Connection.IsType<double>())
 			{
 				PropertyType = EPropertyBagPropertyType::Double;
 			}
-			else if (Type == GetTypeName<FName>())
+			else if (Connection.IsType<FName>())
 			{
 				PropertyType = EPropertyBagPropertyType::Name;
 			}
-			else if (Type == GetTypeName<FString>())
+			else if (Connection.IsType<FString>())
 			{
 				PropertyType = EPropertyBagPropertyType::String;
 			}
-			else if (Type == GetTypeName<FText>())
+			else if (Connection.IsType<FText>())
 			{
 				PropertyType = EPropertyBagPropertyType::Text;
 			}
 			else
 			{
-				const FString TypeStr(Type.ToString());
+				const FString TypeStr(Connection.GetType().ToString());
 				if (TypeStr.StartsWith("U"))
 				{
 					// Nothing for now and use thefallback - may need fixing if we find there's an issue with that code path 
