Dataflow:
- Temporarily disabled FConvertUObjectConvertibleTypesDataflowNode node until it will get fixed
- Fixed FConvertSelectionTypesDataflowNode node

diff --git a/Engine/Source/Runtime/Experimental/Chaos/Private/GeometryCollection/Facades/CollectionTransformSelectionFacade.cpp b/Engine/Source/Runtime/Experimental/Chaos/Private/GeometryCollection/Facades/CollectionTransformSelectionFacade.cpp
--- a/Engine/Source/Runtime/Experimental/Chaos/Private/GeometryCollection/Facades/CollectionTransformSelectionFacade.cpp
+++ b/Engine/Source/Runtime/Experimental/Chaos/Private/GeometryCollection/Facades/CollectionTransformSelectionFacade.cpp
@@ -1136,17 +1136,44 @@ namespace GeometryCollection::Facades
 
 			const int32 NumGeos = ConstCollection.NumElements(FGeometryCollection::GeometryGroup);
 
-			for (int32 VertexIdx : InVertexSelection)
+			if (!bAllElementsMustBeSelected)
+			{
+				for (int32 VertexIdx : InVertexSelection)
+				{
+					for (int32 GeometryIdx = 0; GeometryIdx < NumGeos; ++GeometryIdx)
+					{
+						const int32 VertexStart = VertexStarts[GeometryIdx];
+						const int32 VertexEnd = VertexStart + VertexCounts[GeometryIdx];
+
+						if (VertexIdx >= VertexStart && VertexIdx < VertexEnd)
+						{
+							OutSelection.AddUnique(TransformIndices[GeometryIdx]);
+							break;
+						}
+					}
+				}
+			}
+			else
 			{
 				for (int32 GeometryIdx = 0; GeometryIdx < NumGeos; ++GeometryIdx)
 				{
 					const int32 VertexStart = VertexStarts[GeometryIdx];
-					const int32 VertexEnd = VertexStart + VertexCounts[GeometryIdx];
 
-					if (VertexIdx >= VertexStart && VertexIdx < VertexEnd)
+					bool bAllIncluded = true;
+					for (int32 Offset = 0; Offset < VertexCounts[GeometryIdx]; ++Offset)
+					{
+						const int32 VertexIdx = VertexStart + Offset;
+
+						if (!InVertexSelection.Contains(VertexIdx))
+						{
+							bAllIncluded = false;
+							break;
+						}
+					}
+
+					if (bAllIncluded)
 					{
 						OutSelection.AddUnique(TransformIndices[GeometryIdx]);
-						break;
 					}
 				}
 			}
@@ -1167,17 +1194,44 @@ namespace GeometryCollection::Facades
 
 			const int32 NumGeos = ConstCollection.NumElements(FGeometryCollection::GeometryGroup);
 
-			for (int32 FaceIdx : InFaceSelection)
+			if (!bAllElementsMustBeSelected)
+			{
+				for (int32 FaceIdx : InFaceSelection)
+				{
+					for (int32 GeometryIdx = 0; GeometryIdx < NumGeos; ++GeometryIdx)
+					{
+						const int32 FaceStart = FaceStarts[GeometryIdx];
+						const int32 FaceEnd = FaceStart + FaceCounts[GeometryIdx];
+
+						if (FaceIdx >= FaceStart && FaceIdx < FaceEnd)
+						{
+							OutSelection.AddUnique(TransformIndices[GeometryIdx]);
+							break;
+						}
+					}
+				}
+			}
+			else
 			{
 				for (int32 GeometryIdx = 0; GeometryIdx < NumGeos; ++GeometryIdx)
 				{
 					const int32 FaceStart = FaceStarts[GeometryIdx];
-					const int32 FaceEnd = FaceStart + FaceCounts[GeometryIdx];
 
-					if (FaceIdx >= FaceStart && FaceIdx < FaceEnd)
+					bool bAllIncluded = true;
+					for (int32 Offset = 0; Offset < FaceCounts[GeometryIdx]; ++Offset)
+					{
+						const int32 FaceIdx = FaceStart + Offset;
+
+						if (!InFaceSelection.Contains(FaceIdx))
+						{
+							bAllIncluded = false;
+							break;
+						}
+					}
+
+					if (bAllIncluded)
 					{
 						OutSelection.AddUnique(TransformIndices[GeometryIdx]);
-						break;
 					}
 				}
 			}
@@ -1199,24 +1253,41 @@ namespace GeometryCollection::Facades
 			const int32 NumFaces = ConstCollection.NumElements(FGeometryCollection::FacesGroup);
 			const int32 NumVertices = ConstCollection.NumElements(FGeometryCollection::VerticesGroup);
 
-			TArray<TArray<int32>> VertexToFaceIndices;
-			VertexToFaceIndices.SetNum(NumVertices);
+			if (!bAllElementsMustBeSelected)
+			{
+				TArray<TArray<int32>> VertexToFaceIndices;
+				VertexToFaceIndices.SetNum(NumVertices);
 
-			const TManagedArray<FIntVector>& Indices = IndicesAttribute.Get();
+				const TManagedArray<FIntVector>& Indices = IndicesAttribute.Get();
 
-			for (int32 FaceIdx = 0; FaceIdx < NumFaces; ++FaceIdx)
-			{
-				VertexToFaceIndices[Indices[FaceIdx].X].AddUnique(FaceIdx);
-				VertexToFaceIndices[Indices[FaceIdx].Y].AddUnique(FaceIdx);
-				VertexToFaceIndices[Indices[FaceIdx].Z].AddUnique(FaceIdx);
-			}
+				for (int32 FaceIdx = 0; FaceIdx < NumFaces; ++FaceIdx)
+				{
+					VertexToFaceIndices[Indices[FaceIdx].X].AddUnique(FaceIdx);
+					VertexToFaceIndices[Indices[FaceIdx].Y].AddUnique(FaceIdx);
+					VertexToFaceIndices[Indices[FaceIdx].Z].AddUnique(FaceIdx);
+				}
 
-			// Convert VertexSelection to FaceSelection
-			for (int32 VertexIdx : InVertexSelection)
+				// Convert VertexSelection to FaceSelection
+				for (int32 VertexIdx : InVertexSelection)
+				{
+					for (int32 FaceIdx = 0; FaceIdx < VertexToFaceIndices[VertexIdx].Num(); ++FaceIdx)
+					{
+						OutSelection.AddUnique(VertexToFaceIndices[VertexIdx][FaceIdx]);
+					}
+				}
+			}
+			else
 			{
-				for (int32 FaceIdx = 0; FaceIdx < VertexToFaceIndices[VertexIdx].Num(); ++FaceIdx)
+				const TManagedArray<FIntVector>& Indices = IndicesAttribute.Get();
+
+				for (int32 FaceIdx = 0; FaceIdx < NumFaces; ++FaceIdx)
 				{
-					OutSelection.AddUnique(VertexToFaceIndices[VertexIdx][FaceIdx]);
+					if (InVertexSelection.Contains(Indices[FaceIdx].X) &&
+						InVertexSelection.Contains(Indices[FaceIdx].Y) &&
+						InVertexSelection.Contains(Indices[FaceIdx].Z))
+					{
+						OutSelection.AddUnique(FaceIdx);
+					}
 				}
 			}
 		}
@@ -1307,6 +1378,215 @@ namespace GeometryCollection::Facades
 		return OutSelection;
 	}
 
-}; // GeometryCollection::Facades
+	TArray<int32> FCollectionTransformSelectionFacade::ConvertTransformSelectionToGeometrySelection(const TArray<int32>& InTransformSelection) const
+	{
+		TArray<int32> OutSelection;
+
+		if (TransformToGeometryIndexAttribute.IsValid())
+		{
+			const TManagedArray<int32>& TransformToGeometryIndices = TransformToGeometryIndexAttribute.Get();
+
+			for (int32 TransformIdx : InTransformSelection)
+			{
+				if (TransformToGeometryIndices.IsValidIndex(TransformIdx))
+				{
+					const int32 GeometryIdx = TransformToGeometryIndices[TransformIdx];
+					OutSelection.Add(GeometryIdx);
+				}
+			}
+		}
+
+		return OutSelection;
+	}
+
+	TArray<int32> FCollectionTransformSelectionFacade::ConvertGeometrySelectionToTransformSelection(const TArray<int32>& InGeometrySelection) const
+	{
+		TArray<int32> OutSelection;
+
+		if (TransformIndexAttribute.IsValid())
+		{
+			const TManagedArray<int32>& TransformIndices = TransformIndexAttribute.Get();
+
+			for (int32 GeometryIdx : InGeometrySelection)
+			{
+				if (TransformIndices.IsValidIndex(GeometryIdx))
+				{
+					const int32 TransformIdx = TransformIndexAttribute[GeometryIdx];
+					OutSelection.Add(TransformIdx);
+				}
+			}
+		}
+
+		return OutSelection;
+	}
+
+	TArray<int32> FCollectionTransformSelectionFacade::ConvertGeometrySelectionToFaceSelection(const TArray<int32>& InGeometrySelection) const
+	{
+		TArray<int32> OutSelection;
+
+		if (FaceStartAttribute.IsValid() && FaceCountAttribute.IsValid())
+		{
+			const TManagedArray<int32>& FaceStarts = FaceStartAttribute.Get();
+			const TManagedArray<int32>& FaceCounts = FaceCountAttribute.Get();
+
+			for (int32 GeometryIdx : InGeometrySelection)
+			{
+				if (FaceStarts.IsValidIndex(GeometryIdx))
+				{
+					const int32 FaceIdxStart = FaceStarts[GeometryIdx];
+
+					for (int32 Offset = 0; Offset < FaceCounts[GeometryIdx]; ++Offset)
+					{
+						const int32 FaceIdx = FaceIdxStart + Offset;
+						OutSelection.Add(FaceIdx);
+					}
+				}
+			}
+		}
 
+		return OutSelection;
+	}
+
+	TArray<int32> FCollectionTransformSelectionFacade::ConvertGeometrySelectionToVertexSelection(const TArray<int32>& InGeometrySelection) const
+	{
+		TArray<int32> OutSelection;
+
+		if (VertexStartAttribute.IsValid() && VertexCountAttribute.IsValid())
+		{
+			const TManagedArray<int32>& VertexStarts = VertexStartAttribute.Get();
+			const TManagedArray<int32>& VertexCounts = VertexCountAttribute.Get();
+
+			for (int32 GeometryIdx : InGeometrySelection)
+			{
+				if (VertexStarts.IsValidIndex(GeometryIdx))
+				{
+					const int32 VertexIdxStart = VertexStarts[GeometryIdx];
+
+					for (int32 Offset = 0; Offset < VertexCounts[GeometryIdx]; ++Offset)
+					{
+						const int32 VertexIdx = VertexIdxStart + Offset;
+						OutSelection.Add(VertexIdx);
+					}
+				}
+			}
+		}
+
+		return OutSelection;
+	}
+	
+	TArray<int32> FCollectionTransformSelectionFacade::ConvertFaceSelectionToGeometrySelection(const TArray<int32>& InFaceSelection, bool bAllElementsMustBeSelected) const
+	{
+		TArray<int32> OutSelection;
+
+		if (FaceStartAttribute.IsValid() && FaceCountAttribute.IsValid())
+		{
+			const TManagedArray<int32>& FaceStarts = FaceStartAttribute.Get();
+			const TManagedArray<int32>& FaceCounts = FaceCountAttribute.Get();
+
+			const int32 NumGeos = ConstCollection.NumElements(FGeometryCollection::GeometryGroup);
+
+			if (!bAllElementsMustBeSelected)
+			{
+				for (int32 FaceIdx : InFaceSelection)
+				{
+					for (int32 GeometryIdx = 0; GeometryIdx < NumGeos; ++GeometryIdx)
+					{
+						const int32 FaceStart = FaceStarts[GeometryIdx];
+						const int32 FaceEnd = FaceStart + FaceCounts[GeometryIdx];
+
+						if (FaceIdx >= FaceStart && FaceIdx < FaceEnd)
+						{
+							OutSelection.AddUnique(GeometryIdx);
+							break;
+						}
+					}
+				}
+			}
+			else
+			{
+				for (int32 GeometryIdx = 0; GeometryIdx < NumGeos; ++GeometryIdx)
+				{
+					const int32 FaceStart = FaceStarts[GeometryIdx];
+
+					bool bAllIncluded = true;
+					for (int32 Offset = 0; Offset < FaceCounts[GeometryIdx]; ++Offset)
+					{
+						const int32 FaceIdx = FaceStart + Offset;
+
+						if (!InFaceSelection.Contains(FaceIdx))
+						{
+							bAllIncluded = false;
+							break;
+						}
+					}
 
+					if (bAllIncluded)
+					{
+						OutSelection.AddUnique(GeometryIdx);
+					}
+				}
+			}
+		}
+
+		return OutSelection;
+	}
+
+	TArray<int32> FCollectionTransformSelectionFacade::ConvertVertexSelectionToGeometrySelection(const TArray<int32>& InVertexSelection, bool bAllElementsMustBeSelected) const
+	{
+		TArray<int32> OutSelection;
+
+		if (VertexStartAttribute.IsValid() && VertexCountAttribute.IsValid())
+		{
+			const TManagedArray<int32>& VertexStarts = VertexStartAttribute.Get();
+			const TManagedArray<int32>& VertexCounts = VertexCountAttribute.Get();
+
+			const int32 NumGeos = ConstCollection.NumElements(FGeometryCollection::GeometryGroup);
+
+			if (!bAllElementsMustBeSelected)
+			{
+
+				for (int32 VertexIdx : InVertexSelection)
+				{
+					for (int32 GeometryIdx = 0; GeometryIdx < NumGeos; ++GeometryIdx)
+					{
+						const int32 VertexStart = VertexStarts[GeometryIdx];
+						const int32 VertexEnd = VertexStart + VertexCounts[GeometryIdx];
+
+						if (VertexIdx >= VertexStart && VertexIdx < VertexEnd)
+						{
+							OutSelection.AddUnique(GeometryIdx);
+							break;
+						}
+					}
+				}
+			}
+			else
+			{
+				for (int32 GeometryIdx = 0; GeometryIdx < NumGeos; ++GeometryIdx)
+				{
+					const int32 VertexStart = VertexStarts[GeometryIdx];
+
+					bool bAllIncluded = true;
+					for (int32 Offset = 0; Offset < VertexCounts[GeometryIdx]; ++Offset)
+					{
+						const int32 VertexIdx = VertexStart + Offset;
+
+						if (!InVertexSelection.Contains(VertexIdx))
+						{
+							bAllIncluded = false;
+							break;
+						}
+					}
+
+					if (bAllIncluded)
+					{
+						OutSelection.AddUnique(GeometryIdx);
+					}
+				}
+			}
+		}
+
+		return OutSelection;
+	}
+
+}; // GeometryCollection::Facades
diff --git a/Engine/Source/Runtime/Experimental/Chaos/Public/GeometryCollection/Facades/CollectionTransformSelectionFacade.h b/Engine/Source/Runtime/Experimental/Chaos/Public/GeometryCollection/Facades/CollectionTransformSelectionFacade.h
--- a/Engine/Source/Runtime/Experimental/Chaos/Public/GeometryCollection/Facades/CollectionTransformSelectionFacade.h
+++ b/Engine/Source/Runtime/Experimental/Chaos/Public/GeometryCollection/Facades/CollectionTransformSelectionFacade.h
@@ -213,6 +213,36 @@ namespace GeometryCollection::Facades
 		*/
 		CHAOS_API TArray<int32> ConvertTransformSelectionToVertexSelection(const TArray<int32>& InTransformSelection) const;
 
+		/**
+		* Convert a TransformSelection to a GeometrySelection
+		*/
+		CHAOS_API TArray<int32> ConvertTransformSelectionToGeometrySelection(const TArray<int32>& InTransformSelection) const;
+
+		/**
+		* Convert a GeometrySelection to a TransformSelection
+		*/
+		CHAOS_API TArray<int32> ConvertGeometrySelectionToTransformSelection(const TArray<int32>& InGeometrySelection) const;
+
+		/**
+		* Convert a GeometrySelection to a FaceSelection
+		*/
+		CHAOS_API TArray<int32> ConvertGeometrySelectionToFaceSelection(const TArray<int32>& InGeometrySelection) const;
+
+		/**
+		* Convert a GeometrySelection to a VartexSelection
+		*/
+		CHAOS_API TArray<int32> ConvertGeometrySelectionToVertexSelection(const TArray<int32>& InGeometrySelection) const;
+
+		/**
+		* Convert a FaceSelection to a GeometrySelection
+		*/
+		CHAOS_API TArray<int32> ConvertFaceSelectionToGeometrySelection(const TArray<int32>& InFaceSelection, bool bAllElementsMustBeSelected) const;
+
+		/**
+		* Convert a VertexSelection to a GeometrySelection
+		*/
+		CHAOS_API TArray<int32> ConvertVertexSelectionToGeometrySelection(const TArray<int32>& InVertexSelection, bool bAllElementsMustBeSelected) const;
+
 
 	private:
 		// This a const facade
diff --git a/Engine/Source/Runtime/Experimental/Dataflow/Core/Private/Dataflow/DataflowConvertNodes.cpp b/Engine/Source/Runtime/Experimental/Dataflow/Core/Private/Dataflow/DataflowConvertNodes.cpp
--- a/Engine/Source/Runtime/Experimental/Dataflow/Core/Private/Dataflow/DataflowConvertNodes.cpp
+++ b/Engine/Source/Runtime/Experimental/Dataflow/Core/Private/Dataflow/DataflowConvertNodes.cpp
@@ -4,6 +4,7 @@
 #include "Dataflow/DataflowNode.h"
 #include "Dataflow/DataflowNodeFactory.h"
 #include "Dataflow/DataflowNodeColorsRegistry.h"
+#include "GeometryCollection/Facades/CollectionTransformSelectionFacade.h"
 
 #include "MathUtil.h"
 
@@ -18,7 +19,8 @@ namespace UE::Dataflow
 		DATAFLOW_NODE_REGISTER_CREATION_FACTORY(FConvertBoolTypesDataflowNode);
 		DATAFLOW_NODE_REGISTER_CREATION_FACTORY(FConvertTransformTypesDataflowNode);
 		DATAFLOW_NODE_REGISTER_CREATION_FACTORY(FConvertStringConvertibleTypesDataflowNode);
-		DATAFLOW_NODE_REGISTER_CREATION_FACTORY(FConvertUObjectConvertibleTypesDataflowNode);
+		// Needs to be fixed
+//		DATAFLOW_NODE_REGISTER_CREATION_FACTORY(FConvertUObjectConvertibleTypesDataflowNode);
 		DATAFLOW_NODE_REGISTER_CREATION_FACTORY(FConvertSelectionTypesDataflowNode);
 		DATAFLOW_NODE_REGISTER_CREATION_FACTORY(FConvertVectorArrayTypesDataflowNode);
 		DATAFLOW_NODE_REGISTER_CREATION_FACTORY(FConvertNumericArrayTypesDataflowNode);
@@ -162,7 +164,9 @@ void FConvertUObjectConvertibleTypesDataflowNode::Evaluate(UE::Dataflow::FContex
 FConvertSelectionTypesDataflowNode::FConvertSelectionTypesDataflowNode(const UE::Dataflow::FNodeParameters& InParam, FGuid InGuid)
 	: FDataflowNode(InParam, InGuid)
 {
+	RegisterInputConnection(&Collection);
 	RegisterInputConnection(&In);
+	RegisterOutputConnection(&Collection, &Collection);
 	RegisterOutputConnection(&Out);
 }
 
@@ -170,8 +174,139 @@ void FConvertSelectionTypesDataflowNode::Evaluate(UE::Dataflow::FContext& Contex
 {
 	if (Output->IsA(&Out))
 	{
+		const FManagedArrayCollection& InCollection = GetValue(Context, &Collection);
 		const FDataflowSelection InValue = GetValue(Context, &In);
-		SetValue(Context, InValue, &Out);
+
+		GeometryCollection::Facades::FCollectionTransformSelectionFacade TransformSelectionFacade(InCollection);
+
+		const FDataflowInput* SelectionInput = FindInput(&In);
+		if (SelectionInput)
+		{
+			if (SelectionInput->IsType<FDataflowTransformSelection>())
+			{
+				if (Output && Output->IsType<FDataflowFaceSelection>())
+				{
+					const TArray<int32>& SelectionArr = TransformSelectionFacade.ConvertTransformSelectionToFaceSelection(InValue.AsArray());
+
+					FDataflowFaceSelection NewSelection;
+					NewSelection.InitFromArray(InCollection, SelectionArr);
+
+					SetValue(Context, MoveTemp(NewSelection), &Out);
+				}
+				else if (Output && Output->IsType<FDataflowVertexSelection>())
+				{
+					const TArray<int32>& SelectionArr = TransformSelectionFacade.ConvertTransformSelectionToVertexSelection(InValue.AsArray());
+
+					FDataflowVertexSelection NewSelection;
+					NewSelection.InitFromArray(InCollection, SelectionArr);
+
+					SetValue(Context, MoveTemp(NewSelection), &Out);
+				}
+				else if (Output && Output->IsType<FDataflowGeometrySelection>())
+				{
+					const TArray<int32>& SelectionArr = TransformSelectionFacade.ConvertTransformSelectionToGeometrySelection(InValue.AsArray());
+
+					FDataflowGeometrySelection NewSelection;
+					NewSelection.InitFromArray(InCollection, SelectionArr);
+
+					SetValue(Context, MoveTemp(NewSelection), &Out);
+				}
+			}
+			else if (SelectionInput->IsType<FDataflowFaceSelection>())
+			{
+				if (Output && Output->IsType<FDataflowTransformSelection>())
+				{
+					const TArray<int32>& SelectionArr = TransformSelectionFacade.ConvertFaceSelectionToTransformSelection(InValue.AsArray(), bAllElementsMustBeSelected);
+
+					FDataflowTransformSelection NewSelection;
+					NewSelection.InitFromArray(InCollection, SelectionArr);
+
+					SetValue(Context, MoveTemp(NewSelection), &Out);
+				}
+				else if (Output && Output->IsType<FDataflowVertexSelection>())
+				{
+					const TArray<int32>& SelectionArr = TransformSelectionFacade.ConvertFaceSelectionToVertexSelection(InValue.AsArray());
+
+					FDataflowVertexSelection NewSelection;
+					NewSelection.InitFromArray(InCollection, SelectionArr);
+
+					SetValue(Context, MoveTemp(NewSelection), &Out);
+				}
+				else if (Output && Output->IsType<FDataflowGeometrySelection>())
+				{
+					const TArray<int32>& SelectionArr = TransformSelectionFacade.ConvertFaceSelectionToGeometrySelection(InValue.AsArray(), bAllElementsMustBeSelected);
+
+					FDataflowGeometrySelection NewSelection;
+					NewSelection.InitFromArray(InCollection, SelectionArr);
+
+					SetValue(Context, MoveTemp(NewSelection), &Out);
+				}
+			}
+			else if (SelectionInput->IsType<FDataflowVertexSelection>())
+			{
+				if (Output && Output->IsType<FDataflowTransformSelection>())
+				{
+					const TArray<int32>& SelectionArr = TransformSelectionFacade.ConvertVertexSelectionToTransformSelection(InValue.AsArray(), bAllElementsMustBeSelected);
+
+					FDataflowTransformSelection NewSelection;
+					NewSelection.InitFromArray(InCollection, SelectionArr);
+
+					SetValue(Context, MoveTemp(NewSelection), &Out);
+				}
+				else if (Output && Output->IsType<FDataflowFaceSelection>())
+				{
+					const TArray<int32>& SelectionArr = TransformSelectionFacade.ConvertVertexSelectionToFaceSelection(InValue.AsArray(), bAllElementsMustBeSelected);
+
+					FDataflowFaceSelection NewSelection;
+					NewSelection.InitFromArray(InCollection, SelectionArr);
+
+					SetValue(Context, MoveTemp(NewSelection), &Out);
+				}
+				else if (Output && Output->IsType<FDataflowGeometrySelection>())
+				{
+					const TArray<int32>& SelectionArr = TransformSelectionFacade.ConvertVertexSelectionToGeometrySelection(InValue.AsArray(), bAllElementsMustBeSelected);
+
+					FDataflowGeometrySelection NewSelection;
+					NewSelection.InitFromArray(InCollection, SelectionArr);
+
+					SetValue(Context, MoveTemp(NewSelection), &Out);
+				}
+			}
+			else if (SelectionInput->IsType<FDataflowGeometrySelection>())
+			{
+				if (Output && Output->IsType<FDataflowTransformSelection>())
+				{
+					const TArray<int32>& SelectionArr = TransformSelectionFacade.ConvertGeometrySelectionToTransformSelection(InValue.AsArray());
+
+					FDataflowTransformSelection NewSelection;
+					NewSelection.InitFromArray(InCollection, SelectionArr);
+
+					SetValue(Context, MoveTemp(NewSelection), &Out);
+				}
+				else if (Output && Output->IsType<FDataflowFaceSelection>())
+				{
+					const TArray<int32>& SelectionArr = TransformSelectionFacade.ConvertGeometrySelectionToFaceSelection(InValue.AsArray());
+
+					FDataflowFaceSelection NewSelection;
+					NewSelection.InitFromArray(InCollection, SelectionArr);
+
+					SetValue(Context, MoveTemp(NewSelection), &Out);
+				}
+				else if (Output && Output->IsType<FDataflowVertexSelection>())
+				{
+					const TArray<int32>& SelectionArr = TransformSelectionFacade.ConvertGeometrySelectionToVertexSelection(InValue.AsArray());
+
+					FDataflowVertexSelection NewSelection;
+					NewSelection.InitFromArray(InCollection, SelectionArr);
+
+					SetValue(Context, MoveTemp(NewSelection), &Out);
+				}
+			}
+		}
+	}
+	else if (Output->IsA(&Collection))
+	{
+		SafeForwardInput(Context, &Collection, &Collection);
 	}
 }
 
diff --git a/Engine/Source/Runtime/Experimental/Dataflow/Core/Private/Dataflow/DataflowSelection.cpp b/Engine/Source/Runtime/Experimental/Dataflow/Core/Private/Dataflow/DataflowSelection.cpp
--- a/Engine/Source/Runtime/Experimental/Dataflow/Core/Private/Dataflow/DataflowSelection.cpp
+++ b/Engine/Source/Runtime/Experimental/Dataflow/Core/Private/Dataflow/DataflowSelection.cpp
@@ -1,6 +1,8 @@
 // Copyright Epic Games, Inc. All Rights Reserved.
 
 #include "Dataflow/DataflowSelection.h"
+#include "GeometryCollection/GeometryCollection.h"
+#include "GeometryCollection/ManagedArrayCollection.h"
 
 #include UE_INLINE_GENERATED_CPP_BY_NAME(DataflowSelection)
 
@@ -147,3 +149,9 @@ void FDataflowSelection::SetSelected(TArray<int32> Indices)
 		}
 	}
 }
+
+void FDataflowSelection::InitFromArray(const FManagedArrayCollection& InCollection, const TArray<int32>& InSelectionArr)
+{
+	Initialize(InCollection.NumElements(GroupName), false);
+	SetFromArray(InSelectionArr);
+}
diff --git a/Engine/Source/Runtime/Experimental/Dataflow/Core/Public/Dataflow/DataflowConvertNodes.h b/Engine/Source/Runtime/Experimental/Dataflow/Core/Public/Dataflow/DataflowConvertNodes.h
--- a/Engine/Source/Runtime/Experimental/Dataflow/Core/Public/Dataflow/DataflowConvertNodes.h
+++ b/Engine/Source/Runtime/Experimental/Dataflow/Core/Public/Dataflow/DataflowConvertNodes.h
@@ -7,6 +7,8 @@
 //#include "Dataflow/DataflowConnection.h"
 //#include "Dataflow/DataflowNodeParameters.h"
 #include "Dataflow/DataflowNode.h"
+#include "GeometryCollection/ManagedArrayCollection.h"
+#include "GeometryCollection/GeometryCollection.h"
 
 #include "DataflowConvertNodes.generated.h"
 
@@ -212,12 +214,23 @@ struct FConvertSelectionTypesDataflowNode : public FDataflowNode
 {
 	GENERATED_USTRUCT_BODY()
 	DATAFLOW_NODE_DEFINE_INTERNAL(FConvertSelectionTypesDataflowNode, "ConvertSelectionTypes", DATAFLOW_CONVERT_NODES_CATEGORY, "")
+	DATAFLOW_NODE_RENDER_TYPE("SurfaceRender", FGeometryCollection::StaticType(), "Collection")
 
 private:
+	/** GeometryCollection for the selection */
+	UPROPERTY(meta = (DataflowInput, DataflowOutput, DataflowPassthrough = "Collection", DataflowIntrinsic))
+	FManagedArrayCollection Collection;
+
 	/** Input value */
 	UPROPERTY(meta = (DataflowInput, DataflowIntrinsic))
 	FDataflowSelectionTypes In;
 
+	/** If true then for converting vertex/face selection to transform/geometry selection all vertex/face must be selected for selecting the associated transform/geometry 
+	or going from vertex to face selection all vertices must be selected to select the face 
+	*/
+	UPROPERTY(EditAnywhere, Category = "Selection")
+	bool bAllElementsMustBeSelected = false;
+
 	/** Output value */
 	UPROPERTY(meta = (DataflowOutput))
 	FDataflowSelectionTypes Out;
diff --git a/Engine/Source/Runtime/Experimental/Dataflow/Core/Public/Dataflow/DataflowSelection.h b/Engine/Source/Runtime/Experimental/Dataflow/Core/Public/Dataflow/DataflowSelection.h
--- a/Engine/Source/Runtime/Experimental/Dataflow/Core/Public/Dataflow/DataflowSelection.h
+++ b/Engine/Source/Runtime/Experimental/Dataflow/Core/Public/Dataflow/DataflowSelection.h
@@ -6,6 +6,8 @@
 #include "UObject/NameTypes.h"
 #include "DataflowSelection.generated.h"
 
+struct FManagedArrayCollection;
+
 USTRUCT()
 struct FDataflowSelection
 {
@@ -35,7 +37,7 @@ struct FDataflowSelection
 	* Sets the selection from a sparse array (it only contains the indices of the selected items)
 	*/
 	DATAFLOWCORE_API void SetFromArray(const TArray<int32>& SelectionArr);
-	
+
 	/**
 	* Sets the selection from a dense array (it contains a true/false element for every item)
 	* (for example from the "Internal" attr from FacesGroup)
@@ -50,6 +52,11 @@ struct FDataflowSelection
 	DATAFLOWCORE_API void SetWithMask(const bool Value, const FDataflowSelection& Mask);
 	DATAFLOWCORE_API const TBitArray<>& GetBitArray() const { return SelectionArray; };
 
+	/**
+	* Initialize from a Collection and an array
+	*/
+	DATAFLOWCORE_API void InitFromArray(const FManagedArrayCollection& InCollection, const TArray<int32>& InSelectionArr);
+
 private:
 	FName GroupName;
 	TBitArray<> SelectionArray;
