Dataflow : fix potential crash when evaluating a subgraph with no outputs

diff --git a/Engine/Source/Runtime/Experimental/Dataflow/Engine/Private/Dataflow/DataflowSubGraphNodes.cpp b/Engine/Source/Runtime/Experimental/Dataflow/Engine/Private/Dataflow/DataflowSubGraphNodes.cpp
--- a/Engine/Source/Runtime/Experimental/Dataflow/Engine/Private/Dataflow/DataflowSubGraphNodes.cpp
+++ b/Engine/Source/Runtime/Experimental/Dataflow/Engine/Private/Dataflow/DataflowSubGraphNodes.cpp
@@ -798,16 +798,18 @@ void FDataflowCallSubGraphNode::EvaluateParentContext(UE::Dataflow::FContext& Co
 
 void FDataflowCallSubGraphNode::Evaluate(UE::Dataflow::FContext& Context, const FDataflowOutput* Out) const
 {
-	check(Out);
-	if (TStrongObjectPtr<UDataflow> DataflowAsset = DataflowAssetWeakPtr.Pin())
+	if (Out)
 	{
-		if (UDataflowSubGraph* SubGraphToCall = DataflowAsset->FindSubGraphByGuid(SubGraphGuid))
+		if (TStrongObjectPtr<UDataflow> DataflowAsset = DataflowAssetWeakPtr.Pin())
 		{
-			if (const FDataflowSubGraphOutputNode* SubGraphOutputNode = SubGraphToCall->GetOutputNode())
+			if (UDataflowSubGraph* SubGraphToCall = DataflowAsset->FindSubGraphByGuid(SubGraphGuid))
 			{
-				// pulling the outputs and when the graph evaluate the remote left end of its graph OnInputValueRequested will be called
-				UE::Dataflow::SubGraphNodes::Private::FContextSubGraph SubGraphContext(Context, *this);
-				SubGraphContext.EvaluateSubGraph(*SubGraphOutputNode, *this, *Out);
+				if (const FDataflowSubGraphOutputNode* SubGraphOutputNode = SubGraphToCall->GetOutputNode())
+				{
+					// pulling the outputs and when the graph evaluate the remote left end of its graph OnInputValueRequested will be called
+					UE::Dataflow::SubGraphNodes::Private::FContextSubGraph SubGraphContext(Context, *this);
+					SubGraphContext.EvaluateSubGraph(*SubGraphOutputNode, *this, *Out);
+				}
 			}
 		}
 	}
