Chaosflesh:
-supported anim curve driven muscle activation:
    -added ReadSkeletalMeshCurves node, updated FMuscleActivationFacade
    -supported reading anim curves in async simulation
    -UDeformableTetrahedralComponent::NewDeformableData() takes muscle indices and activation values
    -Deformable solver can override length-based activation with curves
-Disabled muscle activation with XPBD due to data structure difference (segfault found). Will support this in the future

diff --git a/Engine/Plugins/Experimental/ChaosFlesh/Source/ChaosFleshEngine/Private/ChaosFlesh/ChaosDeformableSolverComponent.cpp b/Engine/Plugins/Experimental/ChaosFlesh/Source/ChaosFleshEngine/Private/ChaosFlesh/ChaosDeformableSolverComponent.cpp
--- a/Engine/Plugins/Experimental/ChaosFlesh/Source/ChaosFleshEngine/Private/ChaosFlesh/ChaosDeformableSolverComponent.cpp
+++ b/Engine/Plugins/Experimental/ChaosFlesh/Source/ChaosFleshEngine/Private/ChaosFlesh/ChaosDeformableSolverComponent.cpp
@@ -253,7 +253,8 @@ void UDeformableSolverComponent::ResetSimulationProperties(const FSolverTimingGr
 			, SolverConstraints.GaussSeidelConstraints.SphereRepulsion.bDoSphereRepulsion
 			, SolverConstraints.GaussSeidelConstraints.SphereRepulsion.SphereRepulsionRadius
 			, SolverConstraints.GaussSeidelConstraints.SphereRepulsion.SphereRepulsionStiffness
-			, SolverMuscleActivation.bDoMuscleActivation
+			, SolverMuscleActivation.GetDoLengthBasedMuscleActivation()
+			, SolverMuscleActivation.bOverrideMuscleActivationWithAnimatedCurves
 			, SolverConstraints.GaussSeidelConstraints.SpringCollision.bCollideWithFullmesh
 			, SolverConstraints.GaussSeidelConstraints.bEnableDynamicSprings
 		));
@@ -296,7 +297,8 @@ void UDeformableSolverComponent::BuildSimulationProxy()
 			, SolverConstraints.GaussSeidelConstraints.SphereRepulsion.bDoSphereRepulsion
 			, SolverConstraints.GaussSeidelConstraints.SphereRepulsion.SphereRepulsionRadius
 			, SolverConstraints.GaussSeidelConstraints.SphereRepulsion.SphereRepulsionStiffness
-			, SolverMuscleActivation.bDoMuscleActivation
+			, SolverMuscleActivation.GetDoLengthBasedMuscleActivation()
+			, SolverMuscleActivation.bOverrideMuscleActivationWithAnimatedCurves
 			, SolverConstraints.GaussSeidelConstraints.SpringCollision.bCollideWithFullmesh
 			, SolverConstraints.GaussSeidelConstraints.bEnableDynamicSprings
 		));
diff --git a/Engine/Plugins/Experimental/ChaosFlesh/Source/ChaosFleshEngine/Private/ChaosFlesh/ChaosDeformableTetrahedralComponent.cpp b/Engine/Plugins/Experimental/ChaosFlesh/Source/ChaosFleshEngine/Private/ChaosFlesh/ChaosDeformableTetrahedralComponent.cpp
--- a/Engine/Plugins/Experimental/ChaosFlesh/Source/ChaosFleshEngine/Private/ChaosFlesh/ChaosDeformableTetrahedralComponent.cpp
+++ b/Engine/Plugins/Experimental/ChaosFlesh/Source/ChaosFleshEngine/Private/ChaosFlesh/ChaosDeformableTetrahedralComponent.cpp
@@ -17,6 +17,7 @@
 #include "Components/SkeletalMeshComponent.h"
 #include "Dataflow/DataflowEngineUtil.h"
 #include "Engine/SkeletalMesh.h"
+#include "GeometryCollection/Facades/CollectionMuscleActivationFacade.h"
 #include "GeometryCollection/Facades/CollectionTransformSourceFacade.h"
 #include "GeometryCollection/Facades/CollectionTetrahedralSkeletalBindingsFacade.h"
 #include "GeometryCollection/GeometryCollectionAlgo.h"
@@ -154,6 +155,8 @@ UDeformablePhysicsComponent::FDataMapValue UDeformableTetrahedralComponent::NewD
 			if (const FFleshCollection* Rest = FleshAsset->GetCollection())
 			{
 				FTransformSource TransformSource(*Rest);
+				FMuscleActivationFacade MuscleActivationFacade(*Rest);
+
 				if (TransformSource.IsValid())
 				{
 					TArray<USkeletalMeshComponent*> SkeletalMeshComponents;
@@ -170,6 +173,8 @@ UDeformablePhysicsComponent::FDataMapValue UDeformableTetrahedralComponent::NewD
 						TArray<FTransform3f> FloatRestTransform = RestTransforms->GetConstArray();
 						TArray<FTransform> AnimationTransforms; ToDoubleLocal(FloatRestTransform, AnimationTransforms);
 						TArray<FTransform> ComponentPose; ToDoubleLocal(FloatRestTransform, ComponentPose);
+						TArray<int32> OverrideMuscleActivationIndices;
+						TArray<float> OverrideMuscleActivationValues;
 
 						// Extract animated transforms from all skeletal meshes.
 						for (const USkeletalMeshComponent* SkeletalMeshComponent : SkeletalMeshComponents)
@@ -206,6 +211,25 @@ UDeformablePhysicsComponent::FDataMapValue UDeformableTetrahedralComponent::NewD
 									}
 								}
 							}
+							
+							if (MuscleActivationFacade.IsValid())
+							{
+#if WITH_EDITOR
+								const FBlendedHeapCurve& AnimCurve = SkeletalMeshComponent->GetAnimationCurves();
+								OverrideMuscleActivationIndices.Reserve(AnimCurve.Num());
+								OverrideMuscleActivationValues.Reserve(AnimCurve.Num());
+								AnimCurve.ForEachElement(
+									[&MuscleActivationFacade, &OverrideMuscleActivationIndices, &OverrideMuscleActivationValues]
+									(const UE::Anim::FCurveElement& Element)
+									{
+										for (int32 MuscleIdx : MuscleActivationFacade.FindMuscleIndexByCurveName(Element.Name.ToString()))
+										{
+											OverrideMuscleActivationIndices.Add(MuscleIdx);
+											OverrideMuscleActivationValues.Add(Element.Value);
+										}
+									});
+#endif
+							}
 						}
 
 						FTransform BoneSpaceXf;
@@ -232,6 +256,8 @@ UDeformablePhysicsComponent::FDataMapValue UDeformableTetrahedralComponent::NewD
 								MassMultiplier,
 								BodyForces.IncompressibilityMultiplier,
 								BodyForces.InflationMultiplier,
+								MoveTemp(OverrideMuscleActivationIndices),
+								MoveTemp(OverrideMuscleActivationValues),
 								this));
 					}
 				}
diff --git a/Engine/Plugins/Experimental/ChaosFlesh/Source/ChaosFleshEngine/Public/ChaosFlesh/ChaosDeformableSolverGroups.h b/Engine/Plugins/Experimental/ChaosFlesh/Source/ChaosFleshEngine/Public/ChaosFlesh/ChaosDeformableSolverGroups.h
--- a/Engine/Plugins/Experimental/ChaosFlesh/Source/ChaosFleshEngine/Public/ChaosFlesh/ChaosDeformableSolverGroups.h
+++ b/Engine/Plugins/Experimental/ChaosFlesh/Source/ChaosFleshEngine/Public/ChaosFlesh/ChaosDeformableSolverGroups.h
@@ -241,8 +241,24 @@ USTRUCT(BlueprintType)
 struct FSolverMuscleActivationGroup
 {
 	GENERATED_BODY()
+	UE_DEPRECATED(5.6, "Direct access to this variable is forbidden.")
+	// Muscle activation is driven by origin-insertion length. Check SetMuscleActivationParameter for more info.
+	UPROPERTY(EditAnywhere, Category = "MuscleActivation", DisplayName = "Do Length Based Muscle Activation")
+	bool bDoMuscleActivation = false;
 
+	// Override muscle activation with animation curves. Check ReadSkeletalMeshCurves node for more info.
 	UPROPERTY(EditAnywhere, Category = "MuscleActivation")
-	bool bDoMuscleActivation = false;
+	bool bOverrideMuscleActivationWithAnimatedCurves = false;
+	
+PRAGMA_DISABLE_DEPRECATION_WARNINGS
+	void SetDoLengthBasedMuscleActivation(bool bInDoMuscleActivation) { bDoMuscleActivation = bInDoMuscleActivation; }
+	bool GetDoLengthBasedMuscleActivation() const { return bDoMuscleActivation; }
+	FSolverMuscleActivationGroup() = default;
+	~FSolverMuscleActivationGroup() = default;
+	FSolverMuscleActivationGroup(const FSolverMuscleActivationGroup&) = default;
+	FSolverMuscleActivationGroup(FSolverMuscleActivationGroup&&) = default;
+	FSolverMuscleActivationGroup& operator=(const FSolverMuscleActivationGroup&) = default;
+	FSolverMuscleActivationGroup& operator=(FSolverMuscleActivationGroup&&) = default;
+PRAGMA_ENABLE_DEPRECATION_WARNINGS
 };
 
diff --git a/Engine/Plugins/Experimental/ChaosFlesh/Source/ChaosFleshNodes/Private/Dataflow/ChaosFleshComputeMuscleActivationNode.cpp b/Engine/Plugins/Experimental/ChaosFlesh/Source/ChaosFleshNodes/Private/Dataflow/ChaosFleshComputeMuscleActivationNode.cpp
--- a/Engine/Plugins/Experimental/ChaosFlesh/Source/ChaosFleshNodes/Private/Dataflow/ChaosFleshComputeMuscleActivationNode.cpp
+++ b/Engine/Plugins/Experimental/ChaosFlesh/Source/ChaosFleshNodes/Private/Dataflow/ChaosFleshComputeMuscleActivationNode.cpp
@@ -5,6 +5,7 @@
 #include "GeometryCollection/Facades/CollectionMeshFacade.h"
 #include "GeometryCollection/Facades/CollectionMuscleActivationFacade.h"
 #include "ChaosFlesh/TetrahedralCollection.h"
+#include "Engine/SkeletalMesh.h"
 
 #include UE_INLINE_GENERATED_CPP_BY_NAME(ChaosFleshComputeMuscleActivationNode)
 
@@ -31,7 +32,10 @@ void FComputeMuscleActivationDataNode_v2::Evaluate(UE::Dataflow::FContext& Conte
 		TArray<int32> InInsertionIndices = GetValue(Context, &InsertionIndicesIn);
 
 		GeometryCollection::Facades::FMuscleActivationFacade FMuscleActivation(InCollection);
-		FMuscleActivation.SetUpMuscleActivation(InOriginIndices, InInsertionIndices);
+		if (!FMuscleActivation.SetUpMuscleActivation(InOriginIndices, InInsertionIndices))
+		{
+			Context.Error(TEXT("Setup failed, please check the Log for more info."), this, Out);
+		}
 		Out->SetValue(MoveTemp(InCollection), Context);
 	}
 }
@@ -123,4 +127,61 @@ bool FSetMuscleActivationParameterNode::ShouldInvalidateOnPropertyChanged(const
 		return false;
 	}
 	return true;
+}
+
+FReadSkeletalMeshCurvesDataflowNode::FReadSkeletalMeshCurvesDataflowNode(const UE::Dataflow::FNodeParameters& InParam, FGuid InGuid)
+	: FDataflowNode(InParam, InGuid)
+	, ImportSKMCurveNames(
+		FDataflowFunctionProperty::FDelegate::CreateLambda([this](UE::Dataflow::FContext& Context)
+			{
+				CurveMuscleNameArray.Empty();
+				if (TObjectPtr<const USkeletalMesh> InSkeletalMesh = GetValue(Context, &SkeletalMesh))
+				{
+					for (const UAssetUserData* AssetUserData : *InSkeletalMesh->GetAssetUserDataArray())
+					{
+						if (const UAnimCurveMetaData* AnimCurveMetaData = Cast<UAnimCurveMetaData>(AssetUserData))
+						{
+							TArray<FName> CurveNamesArray;
+							AnimCurveMetaData->GetCurveMetaDataNames(CurveNamesArray);
+							for (FName Curvename : CurveNamesArray)
+							{
+								CurveMuscleNameArray.Add(FCurveMuscleName(Curvename.ToString(), FString()));
+							}
+						}
+					}
+				}
+			}))
+	, AssignSKMCurveToMuscle(
+		FDataflowFunctionProperty::FDelegate::CreateLambda([this](UE::Dataflow::FContext& Context)
+			{
+				// just triggers node invalidation
+			}))
+{
+	RegisterInputConnection(&Collection);
+	RegisterInputConnection(&SkeletalMesh);
+	RegisterOutputConnection(&Collection, &Collection);
+}
+
+void FReadSkeletalMeshCurvesDataflowNode::Evaluate(UE::Dataflow::FContext& Context, const FDataflowOutput* Out) const
+{
+	if (Out->IsA(&Collection))
+	{
+		FManagedArrayCollection InCollection = GetValue(Context, &Collection);
+		GeometryCollection::Facades::FMuscleActivationFacade MuscleActivationFacade(InCollection);
+		if (MuscleActivationFacade.IsValid())
+		{
+			for (const FCurveMuscleName& CurveMuscleName : CurveMuscleNameArray)
+			{
+				MuscleActivationFacade.AssignCurveName(CurveMuscleName.CurveName, CurveMuscleName.MuscleName);
+			}
+		}
+		Out->SetValue(MoveTemp(InCollection), Context);
+	}
+}
+
+bool FReadSkeletalMeshCurvesDataflowNode::ShouldInvalidateOnPropertyChanged(const FPropertyChangedEvent& InPropertyChangedEvent) const
+{
+	const FName ChangedPropertyName = InPropertyChangedEvent.GetMemberPropertyName();
+	return !(ChangedPropertyName == GET_MEMBER_NAME_CHECKED(FReadSkeletalMeshCurvesDataflowNode, CurveMuscleNameArray) ||
+		ChangedPropertyName == GET_MEMBER_NAME_CHECKED(FReadSkeletalMeshCurvesDataflowNode, ImportSKMCurveNames));
 }
\ No newline at end of file
diff --git a/Engine/Plugins/Experimental/ChaosFlesh/Source/ChaosFleshNodes/Private/Dataflow/ChaosFleshNodesPlugin.cpp b/Engine/Plugins/Experimental/ChaosFlesh/Source/ChaosFleshNodes/Private/Dataflow/ChaosFleshNodesPlugin.cpp
--- a/Engine/Plugins/Experimental/ChaosFlesh/Source/ChaosFleshNodes/Private/Dataflow/ChaosFleshNodesPlugin.cpp
+++ b/Engine/Plugins/Experimental/ChaosFlesh/Source/ChaosFleshNodes/Private/Dataflow/ChaosFleshNodesPlugin.cpp
@@ -65,6 +65,7 @@ void IChaosFleshNodesPlugin::StartupModule()
 	DATAFLOW_NODE_REGISTER_CREATION_FACTORY(FKinematicSkeletonConstraintDataflowNode);
 	DATAFLOW_NODE_REGISTER_CREATION_FACTORY(FKinematicSkeletalMeshInitializationDataflowNode);
 	DATAFLOW_NODE_REGISTER_CREATION_FACTORY(FRadialTetrahedronDataflowNodes);
+	DATAFLOW_NODE_REGISTER_CREATION_FACTORY(FReadSkeletalMeshCurvesDataflowNode);
 	DATAFLOW_NODE_REGISTER_CREATION_FACTORY(FSetCollidableVerticesDataflowNode);
 	DATAFLOW_NODE_REGISTER_CREATION_FACTORY(FSetFleshBonePositionTargetBindingDataflowNode);
 	DATAFLOW_NODE_REGISTER_CREATION_FACTORY(FSetFleshDefaultPropertiesNode);
diff --git a/Engine/Plugins/Experimental/ChaosFlesh/Source/ChaosFleshNodes/Public/Dataflow/ChaosFleshComputeMuscleActivationNode.h b/Engine/Plugins/Experimental/ChaosFlesh/Source/ChaosFleshNodes/Public/Dataflow/ChaosFleshComputeMuscleActivationNode.h
--- a/Engine/Plugins/Experimental/ChaosFlesh/Source/ChaosFleshNodes/Public/Dataflow/ChaosFleshComputeMuscleActivationNode.h
+++ b/Engine/Plugins/Experimental/ChaosFlesh/Source/ChaosFleshNodes/Public/Dataflow/ChaosFleshComputeMuscleActivationNode.h
@@ -156,4 +156,50 @@ public:
 
 	virtual void Evaluate(UE::Dataflow::FContext& Context, const FDataflowOutput* Out) const override;
 	virtual bool ShouldInvalidateOnPropertyChanged(const FPropertyChangedEvent& InPropertyChangedEvent) const override;
+};
+
+/**
+* Struct data structure to store curve names and muscle names
+*/
+USTRUCT()
+struct FCurveMuscleName
+{
+	GENERATED_USTRUCT_BODY()
+	UPROPERTY(EditAnywhere, Category = "Muscle Activation Parameters")
+	FString CurveName;
+	UPROPERTY(EditAnywhere, Category = "Muscle Activation Parameters")
+	FString MuscleName;
+};
+
+USTRUCT(meta = (DataflowFlesh))
+struct FReadSkeletalMeshCurvesDataflowNode : public FDataflowNode
+{
+	GENERATED_USTRUCT_BODY()
+	DATAFLOW_NODE_DEFINE_INTERNAL(FReadSkeletalMeshCurvesDataflowNode, "ReadSkeletalMeshCurves", "Flesh", "")
+	DATAFLOW_NODE_RENDER_TYPE("SurfaceRender", FGeometryCollection::StaticType(), "Collection")
+
+public:
+	FReadSkeletalMeshCurvesDataflowNode(const UE::Dataflow::FNodeParameters& InParam, FGuid InGuid = FGuid::NewGuid());
+
+private:
+	UPROPERTY(meta = (DataflowInput, DataflowOutput, DisplayName = "Collection", DataflowPassthrough = "Collection"))
+	FManagedArrayCollection Collection;
+
+	UPROPERTY(EditAnywhere, Category = "Dataflow", meta = (DataflowInput, DisplayName = "SkeletalMesh"))
+	TObjectPtr<const USkeletalMesh> SkeletalMesh = nullptr;
+
+	/** Click on this button to import curve names from the skeletal mesh. */
+	UPROPERTY(EditAnywhere, Category = "Skeletal Mesh Curves", meta = (ButtonImage = "Icons.Import"))
+	FDataflowFunctionProperty ImportSKMCurveNames;
+
+	/** Click on this button to assign curve to muscles by name. */
+	UPROPERTY(EditAnywhere, Category = "Skeletal Mesh Curves", meta = (ButtonImage = "Icons.Refresh"))
+	FDataflowFunctionProperty AssignSKMCurveToMuscle;
+
+	UPROPERTY(EditAnywhere, Category = "Skeletal Mesh Curves", DisplayName = "Curve Name/Muscle Name")
+	TArray<FCurveMuscleName> CurveMuscleNameArray;
+
+	virtual void Evaluate(UE::Dataflow::FContext& Context, const FDataflowOutput* Out) const override;
+	virtual bool ShouldInvalidateOnPropertyChanged(const FPropertyChangedEvent& InPropertyChangedEvent) const override;
+
 };
\ No newline at end of file
diff --git a/Engine/Plugins/Experimental/Dataflow/Source/DataflowEditor/Private/Dataflow/DataflowSimulationControls.cpp b/Engine/Plugins/Experimental/Dataflow/Source/DataflowEditor/Private/Dataflow/DataflowSimulationControls.cpp
--- a/Engine/Plugins/Experimental/Dataflow/Source/DataflowEditor/Private/Dataflow/DataflowSimulationControls.cpp
+++ b/Engine/Plugins/Experimental/Dataflow/Source/DataflowEditor/Private/Dataflow/DataflowSimulationControls.cpp
@@ -146,12 +146,14 @@ namespace UE::Dataflow
 		}
 	}
 
-	static void FillAnimationDatas(const UAnimSequenceBase* AnimSequence, const float CurrentTime, USkeletalMesh* InSkeletalMesh, TArray<FTransform>& ComponentSpaceTransforms)
+	static void FillAnimationDatas(const UAnimSequenceBase* AnimSequence, const float CurrentTime, USkeletalMeshComponent* SkeletalMeshComponent)
 	{
+		const USkeletalMesh* InSkeletalMesh = SkeletalMeshComponent->GetSkeletalMeshAsset();
 		const FAnimExtractContext ExtractionContext(FMath::Clamp(CurrentTime, 0., AnimSequence->GetPlayLength()));
 
 		if(const FReferenceSkeleton* ReferenceSkeleton = &InSkeletalMesh->GetRefSkeleton())
 		{
+			TArray<FTransform> ComponentSpaceTransforms = SkeletalMeshComponent->GetComponentSpaceTransforms();
 			const int32 NumBones = ReferenceSkeleton->GetNum();
 
 			TArray<FBoneIndexType> BoneIndices;
@@ -177,6 +179,12 @@ namespace UE::Dataflow
 			AnimSequence->GetAnimationPose(AnimationPoseData, ExtractionContext);
 		
 			FAnimationRuntime::FillUpComponentSpaceTransforms(*ReferenceSkeleton, AnimationPoseData.GetPose().GetBones(), ComponentSpaceTransforms);
+			SkeletalMeshComponent->GetEditableComponentSpaceTransforms() = ComponentSpaceTransforms;
+
+			FBlendedHeapCurve BlendedHeapCurve;
+			BlendedHeapCurve.CopyFrom(AnimationPoseData.GetCurve());
+			SkeletalMeshComponent->GetEditableAnimationCurves() = BlendedHeapCurve;
+			SkeletalMeshComponent->ApplyEditedComponentSpaceTransforms();
 		}
 	}
 
@@ -195,11 +203,7 @@ namespace UE::Dataflow
 					{
 						if(const UAnimSequenceBase* AnimSequence = Cast<UAnimSequenceBase>(SingleNodeInstance->GetAnimationAsset()))
 						{
-							TArray<FTransform> ComponentSpaceTransforms = SkeletalMeshComponent->GetComponentSpaceTransforms();
-							FillAnimationDatas(AnimSequence, SimulationTime, SkeletalMeshComponent->GetSkeletalMeshAsset(), ComponentSpaceTransforms);
-							
-							SkeletalMeshComponent->GetEditableComponentSpaceTransforms() = ComponentSpaceTransforms;
-							SkeletalMeshComponent->ApplyEditedComponentSpaceTransforms();
+							FillAnimationDatas(AnimSequence, SimulationTime, SkeletalMeshComponent);
 						}
 					}
 				}
diff --git a/Engine/Source/Runtime/Experimental/Chaos/Private/Chaos/Deformable/ChaosDeformableSolver.cpp b/Engine/Source/Runtime/Experimental/Chaos/Private/Chaos/Deformable/ChaosDeformableSolver.cpp
--- a/Engine/Source/Runtime/Experimental/Chaos/Private/Chaos/Deformable/ChaosDeformableSolver.cpp
+++ b/Engine/Source/Runtime/Experimental/Chaos/Private/Chaos/Deformable/ChaosDeformableSolver.cpp
@@ -65,7 +65,7 @@ DECLARE_CYCLE_STAT(TEXT("Chaos.Deformable.Solver.InitializeSimulationObject"), S
 DECLARE_CYCLE_STAT(TEXT("Chaos.Deformable.Solver.InitializeDeformableParticles"), STAT_ChaosDeformableSolver_InitializeDeformableParticles, STATGROUP_Chaos);
 DECLARE_CYCLE_STAT(TEXT("Chaos.Deformable.Solver.InitializeKinematicParticles"), STAT_ChaosDeformableSolver_InitializeKinematicParticles, STATGROUP_Chaos);
 DECLARE_CYCLE_STAT(TEXT("Chaos.Deformable.Solver.InitializeTetrahedralOrTriangleConstraint"), STAT_ChaosDeformableSolver_InitializeTetrahedralOrTriangleConstraint, STATGROUP_Chaos);
-DECLARE_CYCLE_STAT(TEXT("Chaos.Deformable.Solver.InitializeGidBasedConstraints"), STAT_ChaosDeformableSolver_InitializeGidBasedConstraints, STATGROUP_Chaos);
+DECLARE_CYCLE_STAT(TEXT("Chaos.Deformable.Solver.InitializeGridBasedConstraints"), STAT_ChaosDeformableSolver_InitializeGridBasedConstraints, STATGROUP_Chaos);
 DECLARE_CYCLE_STAT(TEXT("Chaos.Deformable.Solver.InitializeWeakConstraints"), STAT_ChaosDeformableSolver_InitializeWeakConstraints, STATGROUP_Chaos);
 DECLARE_CYCLE_STAT(TEXT("Chaos.Deformable.Solver.InitializeKinematicConstraint"), STAT_ChaosDeformableSolver_InitializeKinematicConstraint, STATGROUP_Chaos);
 DECLARE_CYCLE_STAT(TEXT("Chaos.Deformable.Solver.InitializeCollisionBodies"), STAT_ChaosDeformableSolver_InitializeCollisionBodies, STATGROUP_Chaos);
@@ -185,7 +185,11 @@ namespace Chaos::Softs
 			GSDynamicWeakConstraints.Reset(new FGaussSeidelDynamicWeakConstraints<FSolverReal, FSolverParticles>(GDeformableXPBDWeakConstraintParams));
 			GSSphereRepulsionConstraints.Reset(new FGaussSeidelSphereRepulsionConstraints<FSolverReal, FSolverParticles>(Property.SphereRepulsionRadius, Property.SphereRepulsionStiffness, Evolution->Particles(), GDeformableXPBDWeakConstraintParams));
 		}
-		MuscleActivationConstraints.Reset(new FMuscleActivationConstraints<FSolverReal, FSolverParticles>());
+		if (Property.bDoLengthBasedMuscleActivation || Property.bOverrideMuscleActivationWithAnimatedCurves)
+		{
+			MuscleIndexOffset.Empty();
+			MuscleActivationConstraints.Reset(new FMuscleActivationConstraints<FSolverReal, FSolverParticles>());
+		}
 		AllUnconstrainedSurfaceElementsCorotatedCod.Reset(new TArray<Chaos::TVec3<int32>>());
 		AllUnconstrainedSurfaceElementsSkin.Reset(new TArray<Chaos::TVec3<int32>>());
 
@@ -420,7 +424,7 @@ namespace Chaos::Softs
 					{
 						InitializeGaussSeidelConstraintVariables();
 					}
-					if (Property.bDoMuscleActivation)
+					if (Property.bDoLengthBasedMuscleActivation || Property.bOverrideMuscleActivationWithAnimatedCurves)
 					{
 						InitializeMuscleActivationVariables();
 					}
@@ -469,7 +473,7 @@ namespace Chaos::Softs
 					InitializeWeakConstraint(*Proxy);
 					InitializeMuscleActivation(*Proxy);
 					InitializeTetrahedralOrTriangleConstraint(*Proxy);
-					InitializeGidBasedConstraints(*Proxy);
+					InitializeGridBasedConstraints(*Proxy);
 					InitializeGaussSeidelConstraints(*Proxy);
 				}
 			}
@@ -1239,9 +1243,9 @@ namespace Chaos::Softs
 	}
 
 
-	void FDeformableSolver::InitializeGidBasedConstraints(FFleshThreadingProxy& Proxy)
+	void FDeformableSolver::InitializeGridBasedConstraints(FFleshThreadingProxy& Proxy)
 	{
-		PERF_SCOPE(STAT_ChaosDeformableSolver_InitializeGidBasedConstraints);
+		PERF_SCOPE(STAT_ChaosDeformableSolver_InitializeGridBasedConstraints);
 
 		if (Property.bUseGridBasedConstraints)
 		{
@@ -1912,9 +1916,61 @@ namespace Chaos::Softs
 		int32 InitIndex = Evolution->AddConstraintInitRange(1, true);
 		Evolution->ConstraintInits()[InitIndex] = [this](FSolverParticles& InParticles, const FSolverReal Dt)
 		{
-			this->MuscleActivationConstraints->UpdateLengthBasedMuscleActivation(InParticles);
 			if (Property.bUseGaussSeidelConstraints)
 			{
+				if (Property.bDoLengthBasedMuscleActivation)
+				{
+					this->MuscleActivationConstraints->UpdateLengthBasedMuscleActivation(InParticles);
+				}
+				if (Property.bOverrideMuscleActivationWithAnimatedCurves)
+				{
+					//Override muscle activation with curve data
+					for (const TPair<FThreadingProxy::FKey, TUniquePtr<FThreadingProxy>>& ProxyElem : Proxies)
+					{
+						if (const FFleshThreadingProxy* Proxy = ProxyElem.Value->As<FFleshThreadingProxy>())
+						{
+							if (const UObject* Owner = Proxy->GetOwner())
+							{
+								if (!Proxy->GetIsCached() && MuscleIndexOffset.Contains(Owner))
+								{
+									if (this->CurrentInputPackage && this->CurrentInputPackage->ObjectMap.Contains(Owner))
+									{
+										if (FFleshThreadingProxy::FFleshInputBuffer* FleshInputBuffer = this->CurrentInputPackage->ObjectMap[Owner]->As<FFleshThreadingProxy::FFleshInputBuffer>())
+										{
+											const TArray<int32>& MuscleIndices = FleshInputBuffer->MuscleIndices;
+											const TArray<float>& MuscleActivations = FleshInputBuffer->MuscleActivation;
+											if (ensure(MuscleIndices.Num() == MuscleActivations.Num()))
+											{
+												const int32 Offset = MuscleIndexOffset[Owner];
+												for (int32 Idx = 0; Idx < MuscleIndices.Num(); ++Idx)
+												{
+													this->MuscleActivationConstraints->SetMuscleActivation(Offset + MuscleIndices[Idx], MuscleActivations[Idx]);
+												}
+												if (this->PreviousInputPackage && this->PreviousInputPackage->ObjectMap.Contains(Owner))
+												{
+													if (FFleshThreadingProxy::FFleshInputBuffer* PreviousFleshBuffer = this->PreviousInputPackage->ObjectMap[Owner]->As<FFleshThreadingProxy::FFleshInputBuffer>())
+													{
+														FSolverReal CurrentRatio = FSolverReal(this->Iteration) / FSolverReal(this->Property.NumSolverSubSteps);
+														const TArray<int32>& PrevMuscleIndices = PreviousFleshBuffer->MuscleIndices;
+														const TArray<float>& PrevMuscleActivations = PreviousFleshBuffer->MuscleActivation;
+														for (int32 Idx = 0; Idx < PrevMuscleIndices.Num(); ++Idx)
+														{
+															const int32 MuscleGlobalIdx = Offset + PrevMuscleIndices[Idx];
+															const float CurrentMuscleActivation = this->MuscleActivationConstraints->GetMuscleActivation(MuscleGlobalIdx);
+															this->MuscleActivationConstraints->SetMuscleActivation(MuscleGlobalIdx,
+																CurrentRatio * CurrentMuscleActivation + (1 - CurrentRatio) * PrevMuscleActivations[Idx]);
+														}
+													}
+												}
+											}
+										}
+									}
+								}
+							}
+						}
+					}
+				}
+
 				if (Property.bUseGSNeohookean)
 				{
 					this->MuscleActivationConstraints->ApplyMuscleActivation(*this->GSNeohookeanConstraints);
@@ -1926,7 +1982,9 @@ namespace Chaos::Softs
 			}
 			else
 			{
-				this->MuscleActivationConstraints->ApplyMuscleActivation(*this->CorotatedConstraints[0]);
+				// TODO(flesh): Support muscle activation with XPBD, use an array for GSConstraints
+				//this->MuscleActivationConstraints->ApplyMuscleActivation(*this->CorotatedConstraints[0]);
+				UE_LOG(LogChaosDeformableSolver, Warning, TEXT("Muscle activation is not supported for XPBD now, please switch to Gauss-Seidel."));
 			}
 		};
 	}
@@ -1934,7 +1992,7 @@ namespace Chaos::Softs
 	void FDeformableSolver::InitializeMuscleActivation(FFleshThreadingProxy& Proxy)
 	{
 		PERF_SCOPE(STAT_ChaosDeformableSolver_InitializeMuscleActivation);
-		if (Property.bDoMuscleActivation)
+		if (Property.bUseGaussSeidelConstraints && (Property.bDoLengthBasedMuscleActivation || Property.bOverrideMuscleActivationWithAnimatedCurves))
 		{
 			//Grab MuscleActivationElements, OriginInsertion, FiberDirectionMatrix
 			const FManagedArrayCollection& Rest = Proxy.GetRestCollection();
@@ -1942,7 +2000,9 @@ namespace Chaos::Softs
 			if (MuscleActivationFacade.IsValid())
 			{
 				int32 VertexOffset = Proxy.GetSolverParticleRange().Start;
+				// Only supports Gauss Seidel for now because we are using AllElements
 				int32 ElementOffset = AllElements->Num(); //InitializeMuscleActivation has to go before adding tetrahedrons
+				MuscleIndexOffset.Add(Proxy.GetOwner(), MuscleActivationConstraints->NumMuscles()); 
 				MuscleActivationConstraints->AddMuscles(Evolution->Particles(), MuscleActivationFacade, VertexOffset, ElementOffset);
 			}
 		}
@@ -2024,7 +2084,7 @@ namespace Chaos::Softs
 				{
 					InitializeMuscleActivation(*Proxy);
 					InitializeTetrahedralOrTriangleConstraint(*Proxy);
-					InitializeGidBasedConstraints(*Proxy);
+					InitializeGridBasedConstraints(*Proxy);
 					InitializeGaussSeidelConstraints(*Proxy);
 				}
 			}
diff --git a/Engine/Source/Runtime/Experimental/Chaos/Private/GeometryCollection/Facades/CollectionMuscleActivationFacade.cpp b/Engine/Source/Runtime/Experimental/Chaos/Private/GeometryCollection/Facades/CollectionMuscleActivationFacade.cpp
--- a/Engine/Source/Runtime/Experimental/Chaos/Private/GeometryCollection/Facades/CollectionMuscleActivationFacade.cpp
+++ b/Engine/Source/Runtime/Experimental/Chaos/Private/GeometryCollection/Facades/CollectionMuscleActivationFacade.cpp
@@ -11,6 +11,8 @@
 #include "TriangleTypes.h"
 #include "GeometryCollection/Facades/CollectionMeshFacade.h"
 
+DEFINE_LOG_CATEGORY_STATIC(LogDataflowMuscleActivationFacade, Log, All);
+
 namespace GeometryCollection::Facades
 {
 	// Attributes
@@ -25,7 +27,8 @@ namespace GeometryCollection::Facades
 	const FName FMuscleActivationFacade::MuscleLengthRatioThresholdForMaxActivation("MuscleLengthRatioThresholdForMaxActivation");
 	const FName FMuscleActivationFacade::FiberStreamline("FiberStreamline");
 	const FName FMuscleActivationFacade::FiberStreamlineRestLength("FiberStreamlineRestLength");
-
+	const FName FMuscleActivationFacade::MuscleActivationCurveName("MuscleActivationCurveName");
+	
 	FMuscleActivationFacade::FMuscleActivationFacade(FManagedArrayCollection& InCollection)
 		: ConstCollection(InCollection)
 		, Collection(&InCollection)
@@ -39,6 +42,7 @@ namespace GeometryCollection::Facades
 		, MuscleLengthRatioThresholdForMaxActivationAttribute(InCollection, MuscleLengthRatioThresholdForMaxActivation, GroupName)
 		, FiberStreamlineAttribute(InCollection, FiberStreamline, GroupName)
 		, FiberStreamlineRestLengthAttribute(InCollection, FiberStreamlineRestLength, GroupName)
+		, MuscleActivationCurveNameAttribute(InCollection, MuscleActivationCurveName, GroupName)
 	{
 		DefineSchema();
 	}
@@ -56,6 +60,7 @@ namespace GeometryCollection::Facades
 		, MuscleLengthRatioThresholdForMaxActivationAttribute(InCollection, MuscleLengthRatioThresholdForMaxActivation, GroupName)
 		, FiberStreamlineAttribute(InCollection, FiberStreamline, GroupName)
 		, FiberStreamlineRestLengthAttribute(InCollection, FiberStreamlineRestLength, GroupName)
+		, MuscleActivationCurveNameAttribute(InCollection, MuscleActivationCurveName, GroupName)
 	{
 		
 	}
@@ -64,8 +69,7 @@ namespace GeometryCollection::Facades
 	{
 		return GeometryGroupIndexAttribute.IsValid() && MuscleActivationElementAttribute.IsValid() && OriginInsertionPairAttribute.IsValid() &&
 			OriginInsertionRestLengthAttribute.IsValid() && FiberDirectionMatrixAttribute.IsValid() && ContractionVolumeScaleAttribute.IsValid() &&
-			FiberLengthRatioAtMaxActivationAttribute.IsValid() && MuscleLengthRatioThresholdForMaxActivationAttribute.IsValid() &&
-			FiberStreamlineAttribute.IsValid() && FiberStreamlineRestLengthAttribute.IsValid();
+			FiberLengthRatioAtMaxActivationAttribute.IsValid() && MuscleLengthRatioThresholdForMaxActivationAttribute.IsValid();
 	}
 
 	void FMuscleActivationFacade::DefineSchema()
@@ -81,6 +85,7 @@ namespace GeometryCollection::Facades
 		MuscleLengthRatioThresholdForMaxActivationAttribute.Add();
 		FiberStreamlineAttribute.Add();
 		FiberStreamlineRestLengthAttribute.Add();
+		MuscleActivationCurveNameAttribute.Add();
 	}
 
 	int32 FMuscleActivationFacade::AddMuscleActivationData(const FMuscleActivationData& InputData)
@@ -212,12 +217,43 @@ namespace GeometryCollection::Facades
 	bool FMuscleActivationFacade::SetUpMuscleActivation(const TArray<int32>& InOrigin, const TArray<int32>& Insertion)
 	{
 		// Vertices and fiber field
-		if (!(InOrigin.Num() > 0
-			&& Insertion.Num() > 0
-			&& ConstCollection.FindAttribute<FVector3f>("Vertex", FGeometryCollection::VerticesGroup)
-			&& ConstCollection.FindAttribute<FIntVector4>("Tetrahedron", "Tetrahedral")
-			&& ConstCollection.FindAttribute<FVector3f>("FiberDirection", "Tetrahedral")))
+		if (InOrigin.Num() == 0)
+		{
+			UE_LOG(LogDataflowMuscleActivationFacade, Error, TEXT("No origins given"));
+			return false;
+		}
+		if (Insertion.Num() == 0)
 		{
+			UE_LOG(LogDataflowMuscleActivationFacade, Error, TEXT("No insertions given"));
+			return false;
+		}
+		if (!ConstCollection.FindAttribute<FVector3f>("Vertex", FGeometryCollection::VerticesGroup))
+		{
+			UE_LOG(LogDataflowMuscleActivationFacade, Error, TEXT("Collection has no Vertex attribute"));
+			return false;
+		}
+		if (!ConstCollection.FindAttribute<FIntVector4>("Tetrahedron", "Tetrahedral"))
+		{
+			UE_LOG(LogDataflowMuscleActivationFacade, Error, TEXT("Collection has no Tetrahedron attribute"));
+			return false;
+		}
+		if (!ConstCollection.FindAttribute<FVector3f>("FiberDirection", "Tetrahedral"))
+		{
+			UE_LOG(LogDataflowMuscleActivationFacade, Error, TEXT("Collection has no FiberDirection attribute"));
+			return false;
+		}
+		// Check if origins and insertions have overlaps
+		int32 OverlapCount = 0;
+		for (int32 OriginIdx : InOrigin)
+		{
+			if (Insertion.Contains(OriginIdx))
+			{
+				OverlapCount++;
+			}
+		}
+		if (OverlapCount)
+		{
+			UE_LOG(LogDataflowMuscleActivationFacade, Error, TEXT("Origins and Insertions have %d common indices out of total %d, please check if they are from different sources."), OverlapCount, InOrigin.Num());
 			return false;
 		}
 		TArray<int32> Origin = InOrigin;
@@ -251,7 +287,7 @@ namespace GeometryCollection::Facades
 		{
 			if (!ComponentToMuscleIndex.Contains(ComponentIndex[Insertion[i]]))
 			{
-				ensureMsgf(false, TEXT("No origin in this component"));
+				UE_LOG(LogDataflowMuscleActivationFacade, Warning, TEXT("Some muscle has only insertions but no origins."));
 			}
 			else
 			{
@@ -727,4 +763,33 @@ namespace GeometryCollection::Facades
 		}
 		return MuscleLineSegments;
 	}
-}
+
+	bool FMuscleActivationFacade::AssignCurveName(const FString& CurveName, const FString& MuscleName)
+	{
+		const int32 MuscleIdx = FindMuscleIndexByName(MuscleName);
+		if (MuscleActivationCurveNameAttribute.IsValidIndex(MuscleIdx))
+		{
+			MuscleActivationCurveNameAttribute.Modify()[MuscleIdx] = CurveName;
+			return true;
+		}
+		return false;
+	}
+
+	TArray<int32> FMuscleActivationFacade::FindMuscleIndexByCurveName(const FString& CurveName) const
+	{
+		TArray<int32> MuscleIndices;
+		if (MuscleActivationCurveNameAttribute.IsValid())
+		{
+			for (int32 MuscleIdx = 0; MuscleIdx < MuscleActivationCurveNameAttribute.Num(); ++MuscleIdx)
+			{
+				//Animation curve name becomes lower case somehow after update, IgnoreCase here
+				if (MuscleActivationCurveNameAttribute[MuscleIdx].Equals(CurveName, ESearchCase::IgnoreCase)) 
+				{
+					MuscleIndices.Add(MuscleIdx);
+				}
+			}
+		}
+		return MuscleIndices;
+	}
+} // end namespace GeometryCollection::Facades
+
diff --git a/Engine/Source/Runtime/Experimental/Chaos/Public/Chaos/Deformable/ChaosDeformableSolver.h b/Engine/Source/Runtime/Experimental/Chaos/Public/Chaos/Deformable/ChaosDeformableSolver.h
--- a/Engine/Source/Runtime/Experimental/Chaos/Public/Chaos/Deformable/ChaosDeformableSolver.h
+++ b/Engine/Source/Runtime/Experimental/Chaos/Public/Chaos/Deformable/ChaosDeformableSolver.h
@@ -140,7 +140,7 @@ namespace Chaos::Softs
 		CHAOS_API void PostProcessTransientConstraints();
 		CHAOS_API void InitializeKinematicParticles(FFleshThreadingProxy&);
 		CHAOS_API void InitializeTetrahedralOrTriangleConstraint(FFleshThreadingProxy&);
-		CHAOS_API void InitializeGidBasedConstraints(FFleshThreadingProxy&);
+		CHAOS_API void InitializeGridBasedConstraints(FFleshThreadingProxy&);
 		CHAOS_API void InitializeGaussSeidelConstraints(FFleshThreadingProxy& Proxy);
 		CHAOS_API void InitializeWeakConstraint(FFleshThreadingProxy&);
 		CHAOS_API void InitializeKinematicConstraint();
@@ -238,9 +238,9 @@ namespace Chaos::Softs
 		TMap<int32, TSet<int32>> ParticleTriangleExclusionMap;
 		TUniquePtr <FSolverParticles> ReferencePositions;
 		//Muscle Activation Variables
-		TUniquePtr<Softs::FMuscleActivationConstraints<Softs::FSolverReal, Softs::FSolverParticles>> MuscleActivationConstraints;
+		TUniquePtr<Softs::FMuscleActivationConstraints<Softs::FSolverReal, Softs::FSolverParticles>> MuscleActivationConstraints;		
+		TMap<FThreadingProxy::FKey, int32> MuscleIndexOffset;
 
-		//typedef TMap<int32, TTuple<float, Chaos::Softs::FPAndInvM, FVector3f>> TransientConstraintBufferMap;
 		typedef TMap<int32, TTuple<float, float, FVector3f>> TransientConstraintBufferMap;
 		TransientConstraintBufferMap TransientConstraintBuffer;
 
diff --git a/Engine/Source/Runtime/Experimental/Chaos/Public/Chaos/Deformable/ChaosDeformableSolverProxy.h b/Engine/Source/Runtime/Experimental/Chaos/Public/Chaos/Deformable/ChaosDeformableSolverProxy.h
--- a/Engine/Source/Runtime/Experimental/Chaos/Public/Chaos/Deformable/ChaosDeformableSolverProxy.h
+++ b/Engine/Source/Runtime/Experimental/Chaos/Public/Chaos/Deformable/ChaosDeformableSolverProxy.h
@@ -309,6 +309,40 @@ namespace Chaos::Softs
 				, InflationMultiplier(InInflationMultiplier)
 				, SimSpaceBoneIndex(InSimSpaceBoneIndex)
 			{}
+
+			FFleshInputBuffer(
+				const FManagedArrayCollection&& InSimulationCollection,
+				const FTransform& InWorldToComponentXf,
+				const FTransform& InComponentToBoneXf,
+				const int32 InSimSpaceBoneIndex,
+				const TArray<FTransform>&& InTransforms,
+				const TArray<FTransform>&& InRestTransforms,
+				const bool InbEnableGravity,
+				const float InStiffnessMultiplier,
+				const float InDampingMultiplier,
+				const float InMassMultiplier,
+				const float InInCompressibilityMultiplier,
+				const float InInflationMultiplier,
+				const TArray<int32>&& InMuscleIndices,
+				const TArray<float>&& InMuscleActivation,
+				const UObject* InOwner = nullptr)
+				: Super(InOwner, FFleshThreadingProxy::TypeName())
+				, SimulationCollection(InSimulationCollection)
+				, WorldToComponentXf(InWorldToComponentXf)
+				, ComponentToBoneXf(InComponentToBoneXf)
+				, Transforms(InTransforms)
+				, RestTransforms(InRestTransforms)
+				, bEnableGravity(InbEnableGravity)
+				, StiffnessMultiplier(InStiffnessMultiplier)
+				, DampingMultiplier(InDampingMultiplier)
+				, MassMultiplier(InMassMultiplier)
+				, IncompressibilityMultiplier(InInCompressibilityMultiplier)
+				, InflationMultiplier(InInflationMultiplier)
+				, SimSpaceBoneIndex(InSimSpaceBoneIndex)
+				, MuscleIndices(InMuscleIndices)
+				, MuscleActivation(InMuscleActivation)
+			{}
+
 			virtual ~FFleshInputBuffer() {}
 
 			const FManagedArrayCollection SimulationCollection;
@@ -325,6 +359,8 @@ namespace Chaos::Softs
 			float IncompressibilityMultiplier;
 			float InflationMultiplier;
 			int32 SimSpaceBoneIndex = INDEX_NONE;
+			TArray<int32> MuscleIndices;
+			TArray<float> MuscleActivation;
 		};
 
 		class FFleshOutputBuffer : public FThreadingProxy::FBuffer
diff --git a/Engine/Source/Runtime/Experimental/Chaos/Public/Chaos/Deformable/ChaosDeformableSolverTypes.h b/Engine/Source/Runtime/Experimental/Chaos/Public/Chaos/Deformable/ChaosDeformableSolverTypes.h
--- a/Engine/Source/Runtime/Experimental/Chaos/Public/Chaos/Deformable/ChaosDeformableSolverTypes.h
+++ b/Engine/Source/Runtime/Experimental/Chaos/Public/Chaos/Deformable/ChaosDeformableSolverTypes.h
@@ -15,38 +15,76 @@ namespace Chaos::Softs
 {
 	struct FDeformableSolverProperties
 	{
+		FDeformableSolverProperties() 
+		{
+			FDeformableSolverProperties(
+				/* InNumSolverSubSteps */ 2,
+				/* InNumSolverIterations */ 5,
+				/* InFixTimeStep */ false,
+				/* InTimeStepSize */ (FSolverReal)0.05,
+				/* InCacheToFile */ false,
+				/* InbEnableKinematics */ true,
+				/* InbUseFloor */ true,
+				/* InbUseGridBasedConstraints */ false,
+				/* InGridDx */ (FSolverReal)1.0,
+				/* InbDoQuasistatics */ false,
+				/* InbDoBlended */ false,
+				/* InBlendedZeta */ (FSolverReal)0.1,
+				/* InbEnableGravity */ true,
+				/* InbEnableCorotatedConstraints */ true,
+				/* InbEnablePositionTargets */ true,
+				/* InbUseGaussSeidelConstraints */ false,
+				/* InbUseSOR */ true,
+				/* InOmegaSOR */ (FSolverReal)1.6,
+				/* InbUseGSNeohookean */ false,
+				/* InbDoSpringCollision */ false,
+				/* InbDoInComponentSpringCollision */ false,
+				/* InNRingExcluded */ 1,
+				/* InSpringCollisionSearchRadius */ (FSolverReal)0,
+				/* InSpringCollisionStiffness */ (FSolverReal)500.0,
+				/* InbAllowSliding */ true,
+				/* InbDoSphereRepulsion */ false,
+				/* InSphereRepulsionRadius */ (FSolverReal)0,
+				/* InSphereRepulsionStiffness */ (FSolverReal)500.0,
+				/* InbDoLengthBasedMuscleActivation */ false,
+				/* InbOverrideMuscleActivationWithAnimatedCurves */ false,
+				/* InbCollideWithFullMesh */ false,
+				/* InbEnableDynamicSprings */ true);
+		};
+
 		FDeformableSolverProperties(
-			int32 InNumSolverSubSteps = 2,
-			int32 InNumSolverIterations = 5,
-			bool InFixTimeStep = false,
-			FSolverReal InTimeStepSize = (FSolverReal)0.05,
-			bool InCacheToFile = false,
-			bool InbEnableKinematics = true,
-			bool InbUseFloor = true,
-			bool InbUseGridBasedConstraints = false,
-			FSolverReal InGridDx = (FSolverReal)1.,
-			bool InbDoQuasistatics = false,
-			bool InbDoBlended = false,
-			FSolverReal InBlendedZeta = (FSolverReal).1,
-			bool InbEnableGravity = true,
-			bool InbEnableCorotatedConstraints = true,
-			bool InbEnablePositionTargets = true,
-			bool InbUseGaussSeidelConstraints = false,
-			bool InbUseSOR = true,
-			FSolverReal InOmegaSOR = (FSolverReal)1.6,
-			bool InbUseGSNeohookean = false,
-			bool InbDoSpringCollision = false,
-			bool InbDoInComponentSpringCollision = false,
-			int32 InNRingExcluded = 1,
-			FSolverReal InSpringCollisionSearchRadius = (FSolverReal)0,
-			FSolverReal InSpringCollisionStiffness = (FSolverReal)500.,
-			bool InbAllowSliding = true,
-			bool InbDoSphereRepulsion = false,
-			FSolverReal InSphereRepulsionRadius = (FSolverReal)0,
-			FSolverReal InSphereRepulsionStiffness = (FSolverReal)500.,
-			bool InbDoMuscleActivation = false,
-			bool InbCollideWithFullMesh = false, 
-			bool InbEnableDynamicSprings = true)
+			int32 InNumSolverSubSteps,
+			int32 InNumSolverIterations,
+			bool InFixTimeStep,
+			FSolverReal InTimeStepSize,
+			bool InCacheToFile,
+			bool InbEnableKinematics,
+			bool InbUseFloor,
+			bool InbUseGridBasedConstraints,
+			FSolverReal InGridDx,
+			bool InbDoQuasistatics,
+			bool InbDoBlended,
+			FSolverReal InBlendedZeta,
+			bool InbEnableGravity,
+			bool InbEnableCorotatedConstraints,
+			bool InbEnablePositionTargets,
+			bool InbUseGaussSeidelConstraints,
+			bool InbUseSOR,
+			FSolverReal InOmegaSOR,
+			bool InbUseGSNeohookean,
+			bool InbDoSpringCollision,
+			bool InbDoInComponentSpringCollision,
+			int32 InNRingExcluded,
+			FSolverReal InSpringCollisionSearchRadius,
+			FSolverReal InSpringCollisionStiffness,
+			bool InbAllowSliding,
+			bool InbDoSphereRepulsion,
+			FSolverReal InSphereRepulsionRadius,
+			FSolverReal InSphereRepulsionStiffness,
+			bool InbDoLengthBasedMuscleActivation,
+			bool InbOverrideMuscleActivationWithAnimatedCurves,
+			bool InbCollideWithFullMesh,
+			bool InbEnableDynamicSprings)
 			: NumSolverSubSteps(InNumSolverSubSteps)
 			, NumSolverIterations(InNumSolverIterations)
 			, FixTimeStep(InFixTimeStep)
@@ -75,7 +113,8 @@ namespace Chaos::Softs
 			, bDoSphereRepulsion(InbDoSphereRepulsion)
 			, SphereRepulsionRadius(InSphereRepulsionRadius)
 			, SphereRepulsionStiffness(InSphereRepulsionStiffness)
-			, bDoMuscleActivation(InbDoMuscleActivation)
+			, bDoLengthBasedMuscleActivation(InbDoLengthBasedMuscleActivation)
+			, bOverrideMuscleActivationWithAnimatedCurves(InbOverrideMuscleActivationWithAnimatedCurves)
 			, bCollideWithFullMesh(InbCollideWithFullMesh)
 			, bEnableDynamicSprings(InbEnableDynamicSprings)
 		{}
@@ -108,7 +147,8 @@ namespace Chaos::Softs
 		bool bDoSphereRepulsion = false;
 		FSolverReal SphereRepulsionRadius = (FSolverReal)0;
 		FSolverReal SphereRepulsionStiffness = (FSolverReal)500.;
-		bool bDoMuscleActivation = false;
+		bool bDoLengthBasedMuscleActivation = false;
+		bool bOverrideMuscleActivationWithAnimatedCurves = false;
 		bool bCollideWithFullMesh = false;
 		bool bEnableDynamicSprings = true; 
 	};
diff --git a/Engine/Source/Runtime/Experimental/Chaos/Public/Chaos/Deformable/MuscleActivationConstraints.h b/Engine/Source/Runtime/Experimental/Chaos/Public/Chaos/Deformable/MuscleActivationConstraints.h
--- a/Engine/Source/Runtime/Experimental/Chaos/Public/Chaos/Deformable/MuscleActivationConstraints.h
+++ b/Engine/Source/Runtime/Experimental/Chaos/Public/Chaos/Deformable/MuscleActivationConstraints.h
@@ -58,8 +58,8 @@ namespace Chaos::Softs
 						MuscleActivationData.OriginInsertionPair[0] += VertexOffset;
 						MuscleActivationData.OriginInsertionPair[1] += VertexOffset;
 						OriginInsertionPair.Add(MuscleActivationData.OriginInsertionPair);
-						MuscleRestLength.Add((RestParticles.X(MuscleActivationData.OriginInsertionPair[0]) - RestParticles.X(MuscleActivationData.OriginInsertionPair[1])).Size());
-						MuscleActivation.Add(1.f);
+						MuscleRestLength.Add((RestParticles.GetX(MuscleActivationData.OriginInsertionPair[0]) - RestParticles.GetX(MuscleActivationData.OriginInsertionPair[1])).Size());
+						MuscleActivation.Add(0.f);
 						MuscleVertexOffset.Add(FMuscleActivation.MuscleVertexOffset(MuscleIdx));
 						MuscleVertexCount.Add(FMuscleActivation.NumMuscleVertices(MuscleIdx));
 					}
@@ -76,7 +76,7 @@ namespace Chaos::Softs
 					TEXT("MuscleLengthRatioThresholdForMaxActivation %f of muscle indexed %d is out of range (0,1), please check your setup."), MuscleLengthRatioThresholdForMaxActivation[MuscleIdx], MuscleIdx))
 				{
 					// calculate origin/insertion length
-					const float MuscleLengthRatio = (InParticles.X(OriginInsertionPair[MuscleIdx][0]) - InParticles.X(OriginInsertionPair[MuscleIdx][1])).Size() / MuscleRestLength[MuscleIdx];
+					const float MuscleLengthRatio = (InParticles.P(OriginInsertionPair[MuscleIdx][0]) - InParticles.P(OriginInsertionPair[MuscleIdx][1])).Size() / MuscleRestLength[MuscleIdx];
 					if (MuscleLengthRatio >= 1.f)
 					{
 						// not active
@@ -85,7 +85,7 @@ namespace Chaos::Softs
 					else
 					{
 						// linear muscle activation model: muscle reaches max activation 1 at threshold length
-						MuscleActivation[MuscleIdx] = FMath::Min(1.f, (1.f - MuscleLengthRatio) / (1 - MuscleLengthRatioThresholdForMaxActivation[MuscleIdx]));
+						MuscleActivation[MuscleIdx] = FMath::Clamp((1.f - MuscleLengthRatio) / (1 - MuscleLengthRatioThresholdForMaxActivation[MuscleIdx]), 0.f, 1.f);
 					}
 				}
 			}
@@ -115,7 +115,24 @@ namespace Chaos::Softs
 		int32 GetMuscleVertexCount(int32 MuscleIndex) { return MuscleVertexCount[MuscleIndex]; };
 
 		/* returns muscle activation of the specified muscle */
-		float GetMuscleActivation(int32 MuscleIndex) const { return MuscleActivation[MuscleIndex]; };
+		float GetMuscleActivation(int32 MuscleIndex) const { 
+			if (MuscleActivation.IsValidIndex(MuscleIndex))
+			{
+				return MuscleActivation[MuscleIndex];
+			}
+			return 0.f;  
+		};
+
+		/* Sets muscle activation of the specified muscle and returns bool for success*/
+		bool SetMuscleActivation(int32 MuscleIndex, float InMuscleActivation)
+		{
+			if (MuscleActivation.IsValidIndex(MuscleIndex))
+			{
+				MuscleActivation[MuscleIndex] = FMath::Clamp(InMuscleActivation, 0.f, 1.f);
+				return true;
+			}
+			return false;
+		};
 
 	private:	
 		TArray<TArray<int32>> MuscleActivationElement;
diff --git a/Engine/Source/Runtime/Experimental/Chaos/Public/GeometryCollection/Facades/CollectionMuscleActivationFacade.h b/Engine/Source/Runtime/Experimental/Chaos/Public/GeometryCollection/Facades/CollectionMuscleActivationFacade.h
--- a/Engine/Source/Runtime/Experimental/Chaos/Public/GeometryCollection/Facades/CollectionMuscleActivationFacade.h
+++ b/Engine/Source/Runtime/Experimental/Chaos/Public/GeometryCollection/Facades/CollectionMuscleActivationFacade.h
@@ -40,6 +40,7 @@ namespace GeometryCollection::Facades
 		static CHAOS_API const FName MuscleLengthRatioThresholdForMaxActivation;
 		static CHAOS_API const FName FiberStreamline;
 		static CHAOS_API const FName FiberStreamlineRestLength;
+		static CHAOS_API const FName MuscleActivationCurveName;
 
 		CHAOS_API FMuscleActivationFacade(FManagedArrayCollection& InCollection);
 		CHAOS_API FMuscleActivationFacade(const FManagedArrayCollection& InCollection);
@@ -81,6 +82,8 @@ namespace GeometryCollection::Facades
 			float InGlobalMuscleLengthRatioThresholdForMaxActivation);
 		CHAOS_API TArray<TArray<TArray<FVector3f>>> BuildStreamlines(const TArray<int32>& Origin, const TArray<int32>& Insertion,
 			 int32 NumLinesMultiplier, int32 MaxStreamlineIterations, int32 MaxPointsPerLine);
+		CHAOS_API bool AssignCurveName(const FString& CurveName, const FString& MuscleName);
+		CHAOS_API TArray<int32> FindMuscleIndexByCurveName(const FString& CurveName) const;
 	private:
 		const FManagedArrayCollection& ConstCollection;
 		FManagedArrayCollection* Collection = nullptr;
@@ -94,5 +97,6 @@ namespace GeometryCollection::Facades
 		TManagedArrayAccessor<float> MuscleLengthRatioThresholdForMaxActivationAttribute;
 		TManagedArrayAccessor<TArray<TArray<FVector3f>>> FiberStreamlineAttribute;
 		TManagedArrayAccessor<TArray<float>> FiberStreamlineRestLengthAttribute;
+		TManagedArrayAccessor<FString> MuscleActivationCurveNameAttribute;
 	};
 }
