Change managed array collection shared pointers on attribute to be thread safe

diff --git a/Engine/Source/Runtime/Experimental/Chaos/Private/GeometryCollection/ManagedArrayCollection.cpp b/Engine/Source/Runtime/Experimental/Chaos/Private/GeometryCollection/ManagedArrayCollection.cpp
--- a/Engine/Source/Runtime/Experimental/Chaos/Private/GeometryCollection/ManagedArrayCollection.cpp
+++ b/Engine/Source/Runtime/Experimental/Chaos/Private/GeometryCollection/ManagedArrayCollection.cpp
@@ -876,7 +876,7 @@ void FManagedArrayCollection::FValueType::Serialize(FArchive& Ar)
 	{
 		ensure(Ar.IsLoading());
 		ManagedArray = NewManagedTypedArray(ArrayType);
-		SharedManagedArray = TSharedPtr<FManagedArrayBase, ESPMode::NotThreadSafe>(ManagedArray);
+		SharedManagedArray = TSharedPtr<FManagedArrayBase, ESPMode::ThreadSafe>(ManagedArray);
 	}
 	
 	// Note: We switched to always saving the value here, and use the Saved flag
@@ -913,7 +913,7 @@ FManagedArrayCollection::FValueType::FValueType(const FValueType& Other)
 			ManagedArray->Resize(Other.ManagedArray->Num());
 			ManagedArray->Init(*Other.ManagedArray);
 
-			SharedManagedArray = TSharedPtr<FManagedArrayBase, ESPMode::NotThreadSafe>(ManagedArray);
+			SharedManagedArray = TSharedPtr<FManagedArrayBase, ESPMode::ThreadSafe>(ManagedArray);
 			bExternalValue = false;
 		}
 		else
@@ -964,7 +964,7 @@ void FManagedArrayCollection::FValueType::MakeUniqueForWrite()
 		ManagedArray->Resize(SharedManagedArray->Num());
 		ManagedArray->Init(*SharedManagedArray);
 
-		SharedManagedArray = TSharedPtr<FManagedArrayBase, ESPMode::NotThreadSafe>(ManagedArray);
+		SharedManagedArray = TSharedPtr<FManagedArrayBase, ESPMode::ThreadSafe>(ManagedArray);
 	}
 	check(ManagedArray);
 }
diff --git a/Engine/Source/Runtime/Experimental/Chaos/Public/GeometryCollection/ManagedArrayCollection.h b/Engine/Source/Runtime/Experimental/Chaos/Public/GeometryCollection/ManagedArrayCollection.h
--- a/Engine/Source/Runtime/Experimental/Chaos/Public/GeometryCollection/ManagedArrayCollection.h
+++ b/Engine/Source/Runtime/Experimental/Chaos/Public/GeometryCollection/ManagedArrayCollection.h
@@ -671,9 +671,7 @@ private:
 		bool bPersistent;
 		bool bExternalValue;	//External arrays have external memory management.
 
-		// TODO : save the thread on which the object has been created to call out when unsafe access is done ? 
-
-		TSharedPtr<FManagedArrayBase, ESPMode::NotThreadSafe> SharedManagedArray; // external arrays have this set this to null
+		TSharedPtr<FManagedArrayBase, ESPMode::ThreadSafe> SharedManagedArray; // external arrays have this set this to null
 		FManagedArrayBase* ManagedArray = nullptr;;
 
 	private:
