Stable more linear throttling.

#rnx

diff --git a/Engine/Source/Runtime/Experimental/IoStore/HttpClient/Private/Misc.inl b/Engine/Source/Runtime/Experimental/IoStore/HttpClient/Private/Misc.inl
--- a/Engine/Source/Runtime/Experimental/IoStore/HttpClient/Private/Misc.inl
+++ b/Engine/Source/Runtime/Experimental/IoStore/HttpClient/Private/Misc.inl
@@ -703,30 +703,27 @@ public:
 			FThrottler();
 	void	SetLimit(uint32 KiBPerSec);
 	int32	GetAllowance();
-	void	ReturnUnused(uint32 Unused);
+	void	ReturnUnused(int32 Unused);
 
 private:
 	friend	void ThrottleTest(FAnsiStringView);
-	int32	GetAllowance(uint64 CycleDelta);
-	int32	GetWaitEstimateMs() const;
-	uint64	CycleFreq;
-	uint64	CycleLast;
-	uint64	CycleIdle;
+	int32	GetAllowance(int64 CycleDelta);
+	int64	CycleFreq;
+	int64	CycleLast = 0;
+	int64	CyclePeriod = 0;
 	uint32	Limit = 0;
-	int32	Available = 0;
 
-	enum {
+	enum : uint32 {
 		LIMITLESS	= MAX_int32,
-		THRESHOLD	= 2 << 10,
+		SLICES_POW2	= 5,
 	};
 };
 
 ////////////////////////////////////////////////////////////////////////////////
 FThrottler::FThrottler()
 {
-	CycleFreq = uint64(1.0 / FPlatformTime::GetSecondsPerCycle64());
-	CycleLast = FPlatformTime::Cycles64() - CycleFreq;
-	CycleIdle = CycleFreq * 8;
+	CycleFreq = int64(1.0 / FPlatformTime::GetSecondsPerCycle64());
+	check(CycleFreq >> SLICES_POW2);
 }
 
 ////////////////////////////////////////////////////////////////////////////////
@@ -734,6 +731,10 @@ void FThrottler::SetLimit(uint32 KiBPerSec)
 {
 	// 512MiB/s might as well be limitless.
 	KiBPerSec = (KiBPerSec < (512 << 10)) ? KiBPerSec : 0;
+	if (KiBPerSec)
+	{
+		KiBPerSec = FMath::Max(KiBPerSec, 1u << SLICES_POW2);
+	}
 	Limit = KiBPerSec << 10;
 }
 
@@ -742,67 +743,41 @@ int32 FThrottler::GetAllowance()
 {
 	int64 Cycle = FPlatformTime::Cycles64();
 	int64 CycleDelta = Cycle - CycleLast;
-	if (CycleDelta < 0)
-	{
-		return Limit;
-	}
 	CycleLast = Cycle;
 	return GetAllowance(CycleDelta);
 }
 
 ////////////////////////////////////////////////////////////////////////////////
-int32 FThrottler::GetAllowance(uint64 CycleDelta)
+int32 FThrottler::GetAllowance(int64 CycleDelta)
 {
 	if (Limit == 0)
 	{
 		return LIMITLESS;
 	}
 
-	// If we're idle for too long then reset the throttling
-	if (CycleDelta >= CycleIdle)
-	{
-		Available = 0;
-		return Limit;
-	}
-
-	uint64 Delta = (uint64(Limit) * CycleDelta) / CycleFreq;
-
-	// A gate against lost precision
-	if (Delta == 0)
+	int64 CycleSlice = CycleFreq >> SLICES_POW2;
+	CycleDelta = FMath::Min(CycleDelta, CycleSlice);
+	CyclePeriod -= CycleDelta;
+	if (CyclePeriod > 0)
 	{
-		CycleLast -= CycleDelta;
-		return 0 - GetWaitEstimateMs();
+		return 0 - int32((CyclePeriod * 1000ll) / CycleFreq);
 	}
+	CyclePeriod += CycleSlice;
 
-	// Don't let available run away
-	uint64 Next = FMath::Min<uint64>(uint64(Available) + Delta, Limit * 4);
-
-	Available = uint32(Next);
-
-	// Doesn't make sense to trickle out tiny allowances
-	if (Available < THRESHOLD)
-	{
-		return 0 - GetWaitEstimateMs();
-	}
-
-	int32 Released = Available;
-	Available = 0;
+	int32 Released = Limit >> SLICES_POW2;
 	return Released;
 }
 
 ////////////////////////////////////////////////////////////////////////////////
-void FThrottler::ReturnUnused(uint32 Unused)
+void FThrottler::ReturnUnused(int32 Unused)
 {
-	Available += Unused;
-}
+	if (Limit == 0 || Unused == 0)
+	{
+		return;
+	}
 
-////////////////////////////////////////////////////////////////////////////////
-int32 FThrottler::GetWaitEstimateMs() const
-{
-	// Calculate an approximate time to wait for more allowance
-	int64 Estimate = THRESHOLD - Available;
-	Estimate = (Estimate * 1000ll) / int64(Limit);
-	return FMath::Max(int32(Estimate), 0);
+	int64 CycleReturn = (CycleFreq * Unused) / Limit;
+	CycleLast -= CycleReturn;
 }
 
 // }}}
diff --git a/Engine/Source/Runtime/Experimental/IoStore/HttpClient/Private/Overlapped.inl b/Engine/Source/Runtime/Experimental/IoStore/HttpClient/Private/Overlapped.inl
--- a/Engine/Source/Runtime/Experimental/IoStore/HttpClient/Private/Overlapped.inl
+++ b/Engine/Source/Runtime/Experimental/IoStore/HttpClient/Private/Overlapped.inl
@@ -5,7 +5,7 @@
 namespace UE::IoStore::HTTP
 {
 
-#if 0
+#if 1
 
 using SocketType = SOCKET;
 
diff --git a/Engine/Source/Runtime/Experimental/IoStore/HttpClient/Private/Test.inl b/Engine/Source/Runtime/Experimental/IoStore/HttpClient/Private/Test.inl
--- a/Engine/Source/Runtime/Experimental/IoStore/HttpClient/Private/Test.inl
+++ b/Engine/Source/Runtime/Experimental/IoStore/HttpClient/Private/Test.inl
@@ -156,79 +156,16 @@ static void MiscTest()
 static void ThrottleTest(FAnsiStringView TestUrl)
 {
 	enum { TheMax = 0x7fff'fffful };
-
 	check(FThrottler().GetAllowance() >= TheMax);
 
 	FThrottler Throttler;
 	uint64 OneSecond = Throttler.CycleFreq;
 
-	for (uint32 TargetElapsed : { 1, 5, 7 })
-	{
-		uint64 CycleTest = FPlatformTime::Cycles64();
-		for (uint32 i = 0; i < TargetElapsed; ++i)
-		{
-			FPlatformProcess::Sleep(1.0f);
-		}
-		CycleTest = FPlatformTime::Cycles64() - CycleTest;
-		check((CycleTest + (OneSecond / 8)) / OneSecond == TargetElapsed);
-	}
-
-
-	Throttler.SetLimit(0);
-	check(Throttler.GetAllowance(0) >= TheMax);
-	check(Throttler.GetAllowance()  >= TheMax);
-	check(Throttler.GetAllowance()  >= TheMax);
-
-	for (uint32 i : { 10, 63, 100 })
-	{
-		Throttler.SetLimit(i);
-		check(Throttler.GetAllowance( OneSecond          ) == (i << 10));
-		check(Throttler.GetAllowance((OneSecond + 1) >> 1) == (i << 9));
-		check(Throttler.GetAllowance((OneSecond + 3) >> 2) == (i << 8));
-	}
-
-	uint32 Limit = 17 << 10;
-	Throttler = FThrottler();
-	Throttler.SetLimit(Limit >> 10);
-	Throttler.GetAllowance(OneSecond);
-	Throttler.ReturnUnused(Limit >> 1);
-	check(Throttler.GetAllowance(OneSecond) == (Limit + (Limit >> 1)));
-
-	// runaway
-	check(Throttler.GetAllowance(OneSecond * 3) == Limit * 3);
-	check(Throttler.GetAllowance(OneSecond * 4) == Limit * 4);
-	check(Throttler.GetAllowance(OneSecond * 5) == Limit * 4);
-
-	// idle
-	check(Throttler.GetAllowance(OneSecond * 64) == Limit);
-
-	// threshold
-	Limit = FThrottler::THRESHOLD;
-	Throttler = FThrottler();
-	Throttler.SetLimit(Limit >> 10);
-	for (int64 Counter = OneSecond;;)
-	{
-		int64 Delta = (OneSecond + 15) >> 4;
-		Counter -= Delta;
-		int32 Allowance = Throttler.GetAllowance(Delta);
-		if (Allowance > 0)
-		{
-			check(Allowance == Limit);
-			check(Counter < 10);
-			break;
-		}
-	};
-
-	// overflow(ish)
-	Throttler = FThrottler();
-	Throttler.SetLimit((512 << 10) - 1);
-	Throttler.GetAllowance((OneSecond * 790) / 100);
-
 	// timing test
 	FIoBuffer RecvData;
-	for (uint32 SizeKiB : { 10, 25, 60 })
+	for (uint32 SizeKiB : { 64, 128, 192 })
 	{
-		const uint32 ThrottleKiB = 5;
+		const uint32 ThrottleKiB = 64;
 
 		TAnsiStringBuilder<128> Url;
 		Url << TestUrl;
