PVS complaint

#rnx

diff --git a/Engine/Source/Runtime/Experimental/IoStore/HttpClient/Private/Test.inl b/Engine/Source/Runtime/Experimental/IoStore/HttpClient/Private/Test.inl
--- a/Engine/Source/Runtime/Experimental/IoStore/HttpClient/Private/Test.inl
+++ b/Engine/Source/Runtime/Experimental/IoStore/HttpClient/Private/Test.inl
@@ -184,8 +184,8 @@ static void ThrottleTest(FAnsiStringView TestUrl)
 		});
 
 		int32 Timeout = -1;
-		if (SizeKiB < 25) Timeout = 123;
-		if (SizeKiB > 25) Timeout = 4567;
+		if (SizeKiB < 128) Timeout = 123;
+		if (SizeKiB > 128) Timeout = 4567;
 
 		uint64 Time = FPlatformTime::Cycles64();
 		while (Loop.Tick(Timeout));
