Clang compile fix for warning -Wstring-plus-int.

#rnx

- Clang is concerned that this code was attempting to concat an integer value to a string by using the + operator.
- In this case the code was actually indexing into the string literal to animated the sign of life output and so was correct, but it is easier to change the code to supress the warning.

diff --git a/Engine/Source/Runtime/Experimental/IoStore/OnDemand/Private/Tool/Socks.cpp b/Engine/Source/Runtime/Experimental/IoStore/OnDemand/Private/Tool/Socks.cpp
--- a/Engine/Source/Runtime/Experimental/IoStore/OnDemand/Private/Tool/Socks.cpp
+++ b/Engine/Source/Runtime/Experimental/IoStore/OnDemand/Private/Tool/Socks.cpp
@@ -240,10 +240,10 @@ static void Socks5Server(uint32 Port)
 		}
 		std::printf("\x1b[%dF", uint32(AllStats.Num()) + LineClears + 1);
 
-		const char* SignOfLife = ".oOo" + (Tick++ & 0x3);
+		const char* SignOfLife = ".oOo";
 		std::printf(
-			"[%.1s] n:%d d:%llu u:%llu\n",
-			SignOfLife,
+			"[%c] n:%d d:%llu u:%llu\n",
+			SignOfLife[Tick++ & 0x3],
 			Counter,
 			SoFar[0] + UpDown[0],
 			SoFar[1] + UpDown[1]
