Add missing #include for check().

Previously, check() was transitively #included via AutoRTFM.h which #included AssertionMacros.h. AutoRTFM no longer #includes the AssertionMacros header.

diff --git a/Engine/Source/Runtime/Experimental/Chaos/Public/Chaos/RefCountedObject.h b/Engine/Source/Runtime/Experimental/Chaos/Public/Chaos/RefCountedObject.h
--- a/Engine/Source/Runtime/Experimental/Chaos/Public/Chaos/RefCountedObject.h
+++ b/Engine/Source/Runtime/Experimental/Chaos/Public/Chaos/RefCountedObject.h
@@ -3,6 +3,7 @@
 
 #include <atomic>
 #include "AutoRTFM/AutoRTFM.h"
+#include "Misc/AssertionMacros.h"
 
 namespace Chaos
 {
