Chaosflesh: CIS fix [Issue 798899], unused variable in the constructor

diff --git a/Engine/Source/Runtime/Experimental/Chaos/Public/Chaos/Deformable/ChaosDeformableSolverTypes.h b/Engine/Source/Runtime/Experimental/Chaos/Public/Chaos/Deformable/ChaosDeformableSolverTypes.h
--- a/Engine/Source/Runtime/Experimental/Chaos/Public/Chaos/Deformable/ChaosDeformableSolverTypes.h
+++ b/Engine/Source/Runtime/Experimental/Chaos/Public/Chaos/Deformable/ChaosDeformableSolverTypes.h
@@ -16,8 +16,7 @@ namespace Chaos::Softs
 	struct FDeformableSolverProperties
 	{
 		FDeformableSolverProperties() 
-		{
-			FDeformableSolverProperties(
+			: FDeformableSolverProperties(
 				/* InNumSolverSubSteps */ 2,
 				/* InNumSolverIterations */ 5,
 				/* InFixTimeStep */ false,
@@ -49,8 +48,8 @@ namespace Chaos::Softs
 				/* InbDoLengthBasedMuscleActivation */ false,
 				/* InbOverrideMuscleActivationWithAnimatedCurves */ false,
 				/* InbCollideWithFullMesh */ false,
-				/* InbEnableDynamicSprings */ true);
-		};
+				/* InbEnableDynamicSprings */ true) 
+		{};
 
 		FDeformableSolverProperties(
 			int32 InNumSolverSubSteps,
