Geometry collection : fix potential crash when querying mass and extend for partial clusters

diff --git a/Engine/Source/Runtime/Experimental/GeometryCollectionEngine/Private/GeometryCollection/GeometryCollectionComponent.cpp b/Engine/Source/Runtime/Experimental/GeometryCollectionEngine/Private/GeometryCollection/GeometryCollectionComponent.cpp
--- a/Engine/Source/Runtime/Experimental/GeometryCollectionEngine/Private/GeometryCollection/GeometryCollectionComponent.cpp
+++ b/Engine/Source/Runtime/Experimental/GeometryCollectionEngine/Private/GeometryCollection/GeometryCollectionComponent.cpp
@@ -6968,7 +6968,6 @@ void UGeometryCollectionComponent::GetMassAndExtents(int32 ItemIndex, float& Out
 	OutMass = 0.0f;
 	OutExtents = FBox(EForceInit::ForceInitToZero);
 
-	int32 Level = INDEX_NONE;
 	if (RestCollection && RestCollection->GetGeometryCollection())
 	{
 		const FGeometryCollection& Collection = *RestCollection->GetGeometryCollection();
@@ -6977,8 +6976,7 @@ void UGeometryCollectionComponent::GetMassAndExtents(int32 ItemIndex, float& Out
 			const TManagedArray<FBox>* TransformBoundingBoxes = Collection.FindAttribute<FBox>(BoundingBoxAttributeName, FTransformCollection::TransformGroup);
 			const TManagedArray<FBox>* GeoBoundingBoxes = Collection.FindAttribute<FBox>(BoundingBoxAttributeName, FGeometryCollection::GeometryGroup);
 
-			int32 TransformIndex = INDEX_NONE;
-			FGeometryCollectionItemIndex GCItemIndex = FGeometryCollectionItemIndex::CreateFromExistingItemIndex(ItemIndex);
+			const FGeometryCollectionItemIndex GCItemIndex = FGeometryCollectionItemIndex::CreateFromExistingItemIndex(ItemIndex);
 
 			if (GCItemIndex.IsInternalCluster())
 			{
@@ -6994,7 +6992,7 @@ void UGeometryCollectionComponent::GetMassAndExtents(int32 ItemIndex, float& Out
 						}
 						else if (GeoBoundingBoxes)
 						{
-							const int32 GeometryIndex = Collection.TransformToGeometryIndex[TransformIndex];
+							const int32 GeometryIndex = Collection.TransformToGeometryIndex[ChildTramsformIndex];
 							if (GeoBoundingBoxes->IsValidIndex(GeometryIndex))
 							{
 								OutExtents += (*GeoBoundingBoxes)[GeometryIndex];
@@ -7005,7 +7003,7 @@ void UGeometryCollectionComponent::GetMassAndExtents(int32 ItemIndex, float& Out
 			}
 			else
 			{
-				TransformIndex = GCItemIndex.GetTransformIndex();
+				const int32 TransformIndex = GCItemIndex.GetTransformIndex();
 				OutMass = (*CollectionMass)[TransformIndex];
 				if (TransformBoundingBoxes)
 				{
