Dataflow : Move context cache classes to their own file
- Had to abstract some of the context API with regard to accessing data from the store to avoid circular dependencies

diff --git a/Engine/Source/Runtime/Experimental/Dataflow/Core/Public/Dataflow/DataflowContextCache.h b/Engine/Source/Runtime/Experimental/Dataflow/Core/Public/Dataflow/DataflowContextCache.h
new file mode 100644
--- /dev/null
+++ b/Engine/Source/Runtime/Experimental/Dataflow/Core/Public/Dataflow/DataflowContextCache.h
@@ -0,0 +1,422 @@
+// Copyright Epic Games, Inc. All Rights Reserved.
+
+#pragma once
+
+#include "HAL/Platform.h"
+#include "Templates/Decay.h"
+#include "Templates/UniquePtr.h"
+#include "UObject/UnrealType.h"
+
+namespace UE::Dataflow
+{
+	class FContext;
+	struct FContextCacheElementBase;
+
+	typedef uint32 FContextCacheKey;
+
+	/** Trait used to select the UObject* or TObjectPtr cache element path code. */
+	template <typename T>
+	struct TIsUObjectPtrElement
+	{
+		typedef typename TDecay<T>::Type Type;
+		static constexpr bool Value = (std::is_pointer_v<Type> && std::is_convertible_v<Type, const UObjectBase*>) || TIsTObjectPtr_V<Type>;
+	};
+
+	struct FTimestamp
+	{
+		typedef uint64 Type;
+		Type Value = Type(0);
+
+		FTimestamp(Type InValue) : Value(InValue) {}
+		bool operator>=(const FTimestamp& InTimestamp) const { return Value >= InTimestamp.Value; }
+		bool operator<(const FTimestamp& InTimestamp) const { return Value < InTimestamp.Value; }
+		bool operator==(const FTimestamp& InTimestamp) const { return Value == InTimestamp.Value; }
+		bool IsInvalid() { return Value == Invalid; }
+		bool IsInvalid() const { return Value == Invalid; }
+
+		static DATAFLOWCORE_API Type Current();
+		static DATAFLOWCORE_API Type Invalid; // 0
+	};
+
+	struct IContextCacheStore
+	{
+		virtual const TUniquePtr<FContextCacheElementBase>* FindCacheElement(FContextCacheKey Key) const = 0;
+		virtual bool HasCacheElement(FContextCacheKey Key, FTimestamp InTimestamp = FTimestamp::Invalid) const = 0;
+	};
+
+	//--------------------------------------------------------------------
+	// base class for all context cache entries 
+	//--------------------------------------------------------------------
+	struct FContextCacheElementBase 
+	{
+		enum EType
+		{
+			CacheElementTyped,
+			CacheElementReference,
+			CacheElementNull,
+			CacheElementUObject, // UObjectPtr
+			CacheElementUObjectArray // TArray of UObjectPtr
+		};
+
+		FContextCacheElementBase(EType CacheElementType, FGuid InNodeGuid = FGuid(), const FProperty* InProperty = nullptr, uint32 InNodeHash = 0, FTimestamp InTimestamp = FTimestamp::Invalid)
+			: Type(CacheElementType)
+			, NodeGuid(InNodeGuid)
+			, Property(InProperty)
+			, NodeHash(InNodeHash)
+			, Timestamp(InTimestamp)
+		{}
+		virtual ~FContextCacheElementBase() {}
+
+		// InReferenceDataKey is the key of the cache element this function is called on 
+		virtual TUniquePtr<FContextCacheElementBase> CreateReference(FContextCacheKey InReferenceDataKey) const = 0;
+
+		// clone the cache entry
+		virtual TUniquePtr<FContextCacheElementBase> Clone(const IContextCacheStore& Context) const = 0;
+
+		template<typename T>
+		inline const T& GetTypedData(const IContextCacheStore& Context, const FProperty* PropertyIn, const T& Default) const;
+
+		virtual const void* GetUntypedData(const IContextCacheStore& Context, const FProperty* PropertyIn) const = 0;
+
+		EType GetType() const {	return Type; }
+
+		const FProperty* GetProperty() const { return Property; }
+		const FTimestamp& GetTimestamp() const { return Timestamp; }
+		void SetTimestamp(const FTimestamp& InTimestamp) { Timestamp = InTimestamp; }
+
+		const FGuid& GetNodeGuid() const { return NodeGuid; }
+		const uint32 GetNodeHash() const { return NodeHash; }
+
+		// use this with caution: setting the property of a wrong type may cause problems
+		void SetProperty(const FProperty* NewProperty) { Property = NewProperty; }
+
+	private:
+		friend struct FContextCache;
+
+		EType Type;
+		FGuid NodeGuid;
+		const FProperty* Property = nullptr;
+		uint32 NodeHash = 0;
+		FTimestamp Timestamp = FTimestamp::Invalid;
+	};
+
+	//--------------------------------------------------------------------
+	// Value storing context cache entry - strongly typed
+	//--------------------------------------------------------------------
+	template<class T>
+	struct TContextCacheElement : public FContextCacheElementBase 
+	{
+		TContextCacheElement(FGuid InNodeGuid, const FProperty* InProperty, T&& InData, uint32 InNodeHash, FTimestamp Timestamp)
+			: FContextCacheElementBase(EType::CacheElementTyped, InNodeGuid, InProperty, InNodeHash, Timestamp)
+			, Data(Forward<T>(InData))
+		{}
+
+		TContextCacheElement(const TContextCacheElement<T>& Other)
+			: FContextCacheElementBase(EType::CacheElementTyped, Other.GetNodeGuid(), Other.GetProperty(), Other.GetNodeHash(), Other.GetTimestamp())
+			, Data(Other.Data)
+		{}
+		
+		inline const T& GetData(const IContextCacheStore& Context, const FProperty* PropertyIn, const T& Default) const;
+
+		virtual const void* GetUntypedData(const IContextCacheStore& Context, const FProperty* PropertyIn) const
+		{
+			return &Data;
+		}
+
+		const T& GetDataDirect() const { return Data; }
+
+		inline virtual TUniquePtr<FContextCacheElementBase> CreateReference(FContextCacheKey InReferenceDataKey) const override;
+
+		inline virtual TUniquePtr<FContextCacheElementBase> Clone(const IContextCacheStore& Context) const override;
+
+	private:
+		typedef typename TDecay<T>::Type FDataType;  // Using universal references here means T could be either const& or an rvalue reference
+		const FDataType Data;                        // Decaying T removes any reference and gets the correct underlying storage data type
+	};
+
+	//--------------------------------------------------------------------
+	// Reference to another context cache entry 
+	//--------------------------------------------------------------------
+	template<class T>
+	struct TContextCacheElementReference : public FContextCacheElementBase
+	{
+		TContextCacheElementReference(FGuid InNodeGuid, const FProperty* InProperty, FContextCacheKey InDataKey, uint32 InNodeHash, FTimestamp Timestamp)
+			: FContextCacheElementBase(EType::CacheElementReference, InNodeGuid, InProperty, InNodeHash, Timestamp)
+			, DataKey(InDataKey)
+		{}
+
+		inline const T& GetData(const IContextCacheStore& Context, const FProperty* PropertyIn, const T& Default) const;
+
+		inline virtual const void* GetUntypedData(const IContextCacheStore& Context, const FProperty* PropertyIn) const;
+
+		inline virtual TUniquePtr<FContextCacheElementBase> CreateReference(FContextCacheKey InReferenceDataKey) const override;
+
+		inline virtual TUniquePtr<FContextCacheElementBase> Clone(const IContextCacheStore& Context) const override;
+
+	private:
+		const FContextCacheKey DataKey; // this is a key to another cache element
+	};
+
+	//--------------------------------------------------------------------
+	// Null entry, this will always return a default value 
+	//--------------------------------------------------------------------
+	struct FContextCacheElementNull : public FContextCacheElementBase
+	{
+		//
+		// IMPORTANT: 
+		// Timestamp must be set to (Timestamp.Value - 1) to make sure that this type of entry is always invalid
+		//
+		FContextCacheElementNull(FGuid InNodeGuid, const FProperty* InProperty, FContextCacheKey InDataKey, uint32 InNodeHash, FTimestamp Timestamp)
+			: FContextCacheElementBase(EType::CacheElementNull, InNodeGuid, InProperty, InNodeHash, FTimestamp(Timestamp.Value - 1))
+		{}
+
+		virtual const void* GetUntypedData(const IContextCacheStore& Context, const FProperty* PropertyIn) const
+		{
+			return nullptr;
+		}
+
+		inline virtual TUniquePtr<FContextCacheElementBase> CreateReference(FContextCacheKey InReferenceDataKey) const override;
+
+		inline virtual TUniquePtr<FContextCacheElementBase> Clone(const IContextCacheStore& Context) const override;
+	};
+
+	//--------------------------------------------------------------------
+	// UObject cache element, prevents the object from being garbage collected while in the cache
+	//--------------------------------------------------------------------
+	template<class T>
+	struct TContextCacheElementUObject : public FContextCacheElementBase, public FGCObject
+	{
+		TContextCacheElementUObject(FGuid InNodeGuid, const FProperty* InProperty, T&& InObjectPtr, uint32 InNodeHash, FTimestamp Timestamp)
+			: FContextCacheElementBase(EType::CacheElementUObject, InNodeGuid, InProperty, InNodeHash, Timestamp)
+			, ObjectPtr(InObjectPtr)
+		{}
+
+		TContextCacheElementUObject(const TContextCacheElementUObject<T>& Other)
+			: FContextCacheElementBase(EType::CacheElementUObject, Other.GetNodeGuid(), Other.GetProperty(), Other.GetNodeHash(), Other.GetTimestamp())
+			, ObjectPtr(Other.ObjectPtr)
+		{}
+
+		const T& GetData(const IContextCacheStore& /*Context*/, const FProperty* /*PropertyIn*/, const T& /*Default*/) const
+		{
+			return ObjectPtr;
+		}
+
+		virtual const void* GetUntypedData(const IContextCacheStore& /*Context*/, const FProperty* /*PropertyIn*/) const
+		{
+			return (void*)ObjectPtr;
+		}
+
+		inline virtual TUniquePtr<FContextCacheElementBase> CreateReference(FContextCacheKey InReferenceDataKey) const override;
+
+		inline virtual TUniquePtr<FContextCacheElementBase> Clone(const IContextCacheStore& Context) const override;
+
+		//~ Begin FGCObject interface
+		virtual void AddReferencedObjects(FReferenceCollector& Collector) override { Collector.AddReferencedObject(ObjectPtr); }
+		virtual FString GetReferencerName() const override { return TEXT("TContextCacheElementUObject"); }
+		//~ End FGCObject interface
+
+	private:
+		typedef typename TDecay<T>::Type FDataType;  // Using universal references here means T could be either const& or an rvalue reference
+		FDataType ObjectPtr;                         // Decaying T removes any reference and gets the correct underlying storage data type
+	};
+
+	//--------------------------------------------------------------------
+	// TArray<UObjectPtr> cache element, prevents the object from being garbage collected while in the cache
+	//--------------------------------------------------------------------
+	template<class T>
+	struct TContextCacheElementUObjectArray : public FContextCacheElementBase, public FGCObject
+	{
+		template<typename ArrayT>
+		TContextCacheElementUObjectArray(FGuid InNodeGuid, const FProperty* InProperty, ArrayT&& InObjectPtrArray, uint32 InNodeHash, FTimestamp Timestamp)
+			: FContextCacheElementBase(EType::CacheElementUObjectArray, InNodeGuid, InProperty, InNodeHash, Timestamp)
+			, ObjectPtrArray(Forward<ArrayT>(InObjectPtrArray))
+		{}
+
+
+		TContextCacheElementUObjectArray(const TContextCacheElementUObjectArray<T>& Other)
+			: FContextCacheElementBase(EType::CacheElementUObject, Other.GetNodeGuid(), Other.GetProperty(), Other.GetNodeHash(), Other.GetTimestamp())
+			, ObjectPtrArray(Other.ObjectPtrArray)
+		{}
+
+		const TArray<T>& GetData(const IContextCacheStore& /*Context*/, const FProperty* /*PropertyIn*/, const TArray<T>& /*Default*/) const
+		{
+			return ObjectPtrArray;
+		}
+
+		virtual const void* GetUntypedData(const IContextCacheStore& /*Context*/, const FProperty* /*PropertyIn*/) const
+		{
+			return (void*)&ObjectPtrArray;
+		}
+
+		inline virtual TUniquePtr<FContextCacheElementBase> CreateReference(FContextCacheKey InReferenceDataKey) const override;
+
+		inline virtual TUniquePtr<FContextCacheElementBase> Clone(const IContextCacheStore& Context) const override;
+
+		//~ Begin FGCObject interface
+		virtual void AddReferencedObjects(FReferenceCollector& Collector) override { Collector.AddReferencedObjects(ObjectPtrArray); }
+		virtual FString GetReferencerName() const override { return TEXT("TContextCacheElementUObjectArray"); }
+		//~ End FGCObject interface
+
+	private:
+		typedef typename TDecay<T>::Type FDataType;     // Using universal references here means T could be either const& or an rvalue reference
+		TArray<FDataType> ObjectPtrArray;               // Decaying T removes any reference and gets the correct underlying storage data type
+	};
+
+
+	// cache element method implementation 
+	template<class T>
+	const T& FContextCacheElementBase::GetTypedData(const IContextCacheStore& Context, const FProperty* PropertyIn, const T& Default) const
+	{
+		// check(PropertyIn->IsA<T>()); // @todo(dataflow) compile error for non-class T; find alternatives
+		if (Type == EType::CacheElementTyped)
+		{
+			if constexpr (!TIsUObjectPtrElement<T>::Value)
+			{
+				return static_cast<const TContextCacheElement<T>&>(*this).GetData(Context, PropertyIn, Default);
+			}
+		}
+		if (Type == EType::CacheElementReference)
+		{
+			return static_cast<const TContextCacheElementReference<T>&>(*this).GetData(Context, PropertyIn, Default);
+		}
+		if (Type == EType::CacheElementNull)
+		{
+			return Default; 
+		}
+		if (Type == EType::CacheElementUObjectArray)
+		{
+			if constexpr (TIsTArray<T>::Value)
+			{
+				if constexpr (TIsUObjectPtrElement<typename T::ElementType>::Value)
+				{
+					return static_cast<const TContextCacheElementUObjectArray<typename T::ElementType>&>(*this).GetData(Context, PropertyIn, Default);
+				}
+			}
+		}
+		if (Type == EType::CacheElementUObject)
+		{
+			if constexpr (TIsUObjectPtrElement<T>::Value)
+			{
+				return static_cast<const TContextCacheElementUObject<T>&>(*this).GetData(Context, PropertyIn, Default);
+			}
+		}
+		check(false); // should never happen
+		return Default;
+	}
+
+	struct FContextCache : public TMap<FContextCacheKey, TUniquePtr<FContextCacheElementBase>>
+	{
+		DATAFLOWCORE_API void Serialize(FArchive& Ar);
+	};
+
+	// cache classes implemetation 
+	// this needs to be after the FContext definition because they access its methods
+
+	template<class T>
+	const T& TContextCacheElement<T>::GetData(const IContextCacheStore& Context, const FProperty* PropertyIn, const T& Default) const
+	{
+		return Data;
+	}
+
+	template<class T>
+	TUniquePtr<FContextCacheElementBase> TContextCacheElement<T>::CreateReference(FContextCacheKey InReferenceDataKey) const
+	{
+		return MakeUnique<TContextCacheElementReference<T>>(GetNodeGuid(), GetProperty(), InReferenceDataKey, GetNodeHash(), GetTimestamp());
+	}
+
+	template<class T>
+	TUniquePtr<FContextCacheElementBase> TContextCacheElement<T>::Clone(const IContextCacheStore& Context) const
+	{
+		return MakeUnique<TContextCacheElement<T>>(*this);
+	}
+
+	template<class T>
+	const T& TContextCacheElementReference<T>::GetData(const IContextCacheStore& Context, const FProperty* InProperty, const T& Default) const
+	{
+		if (const TUniquePtr<FContextCacheElementBase>* Cache = Context.FindCacheElement(DataKey))
+		{
+			return (*Cache)->GetTypedData<T>(Context, InProperty, Default);
+		}
+		return Default;
+	}
+
+	template<class T>
+	const void* TContextCacheElementReference<T>::GetUntypedData(const IContextCacheStore& Context, const FProperty* InProperty) const
+	{
+		if (const TUniquePtr<FContextCacheElementBase>* Cache = Context.FindCacheElement(DataKey))
+		{
+			return (*Cache)->GetUntypedData(Context, InProperty);
+		}
+		return nullptr;
+	}
+
+	template<class T>
+	TUniquePtr<FContextCacheElementBase> TContextCacheElementReference<T>::CreateReference(FContextCacheKey InReferenceDataKey) const
+	{
+		return MakeUnique<TContextCacheElementReference<T>>(GetNodeGuid(), GetProperty(), InReferenceDataKey, GetNodeHash(), GetTimestamp());
+	}
+
+	template<class T>
+	TUniquePtr<FContextCacheElementBase> TContextCacheElementReference<T>::Clone(const IContextCacheStore& Context) const
+	{
+		if (const TUniquePtr<FContextCacheElementBase>* ReferencedCacheEntry = Context.FindCacheElement(DataKey))
+		{
+			if (ReferencedCacheEntry && ReferencedCacheEntry->IsValid())
+			{
+				return (*ReferencedCacheEntry)->Clone(Context);
+			}
+		}
+		return MakeUnique<FContextCacheElementNull>(GetNodeGuid(), GetProperty(), FContextCacheKey(), GetNodeHash(), GetTimestamp());
+	}
+
+	inline TUniquePtr<FContextCacheElementBase> FContextCacheElementNull::CreateReference(FContextCacheKey InReferenceDataKey) const
+	{
+		// a null entry will always return a null entry as a reference
+		return MakeUnique<FContextCacheElementNull>(GetNodeGuid(), GetProperty(), InReferenceDataKey, GetNodeHash(), GetTimestamp());
+	}
+
+
+	TUniquePtr<FContextCacheElementBase> FContextCacheElementNull::Clone(const IContextCacheStore& Context) const
+	{
+		return MakeUnique<FContextCacheElementNull>(GetNodeGuid(), GetProperty(), FContextCacheKey(), GetNodeHash(), GetTimestamp());
+	}
+
+	template<class T>
+	TUniquePtr<FContextCacheElementBase> TContextCacheElementUObject<T>::CreateReference(FContextCacheKey InReferenceDataKey) const
+	{
+		return MakeUnique<TContextCacheElementReference<T>>(GetNodeGuid(), GetProperty(), InReferenceDataKey, GetNodeHash(), GetTimestamp());
+	}
+
+	template<class T>
+	TUniquePtr<FContextCacheElementBase> TContextCacheElementUObject<T>::Clone(const IContextCacheStore& Context) const
+	{
+		return MakeUnique<TContextCacheElementUObject<T>>(*this);
+	}
+
+	template<class T>
+	TUniquePtr<FContextCacheElementBase> TContextCacheElementUObjectArray<T>::CreateReference(FContextCacheKey InReferenceDataKey) const
+	{
+		return MakeUnique<TContextCacheElementReference<T>>(GetNodeGuid(), GetProperty(), InReferenceDataKey, GetNodeHash(), GetTimestamp());
+	}
+
+	template<class T>
+	TUniquePtr<FContextCacheElementBase> TContextCacheElementUObjectArray<T>::Clone(const IContextCacheStore& Context) const
+	{
+		return MakeUnique<TContextCacheElementUObjectArray<T>>(*this);
+	}
+};
+
+inline FArchive& operator<<(FArchive& Ar, UE::Dataflow::FTimestamp& ValueIn)
+{
+	Ar << ValueIn.Value;
+	Ar << ValueIn.Invalid;
+	return Ar;
+}
+
+inline FArchive& operator<<(FArchive& Ar, UE::Dataflow::FContextCache& ValueIn)
+{
+	ValueIn.Serialize(Ar);
+	return Ar;
+}
+
+
diff --git a/Engine/Source/Runtime/Experimental/Dataflow/Core/Public/Dataflow/DataflowNodeParameters.h b/Engine/Source/Runtime/Experimental/Dataflow/Core/Public/Dataflow/DataflowNodeParameters.h
--- a/Engine/Source/Runtime/Experimental/Dataflow/Core/Public/Dataflow/DataflowNodeParameters.h
+++ b/Engine/Source/Runtime/Experimental/Dataflow/Core/Public/Dataflow/DataflowNodeParameters.h
@@ -7,6 +7,7 @@
 #include "UObject/UnrealType.h"
 #include "Templates/UniquePtr.h"
 #include "GenericPlatform/GenericPlatformCriticalSection.h"
+#include "Dataflow/DataflowContextCache.h"
 #include "Serialization/Archive.h"
 
 class  UDataflow;
@@ -18,35 +19,8 @@ struct FDataflowConnection;
 
 namespace UE::Dataflow
 {
-	class FContext;
-
-	typedef uint32 FContextCacheKey;
-
-	/** Trait used to select the UObject* or TObjectPtr cache element path code. */
-	template <typename T>
-	struct TIsUObjectPtrElement
-	{
-		typedef typename TDecay<T>::Type Type;
-		static constexpr bool Value = (std::is_pointer_v<Type> && std::is_convertible_v<Type, const UObjectBase*>) || TIsTObjectPtr_V<Type>;
-	};
-
-	struct FTimestamp
+	struct FRenderingParameter 
 	{
-		typedef uint64 Type;
-		Type Value = Type(0);
-
-		FTimestamp(Type InValue) : Value(InValue) {}
-		bool operator>=(const FTimestamp& InTimestamp) const { return Value >= InTimestamp.Value; }
-		bool operator<(const FTimestamp& InTimestamp) const { return Value < InTimestamp.Value; }
-		bool operator==(const FTimestamp& InTimestamp) const { return Value == InTimestamp.Value; }
-		bool IsInvalid() { return Value == Invalid; }
-		bool IsInvalid() const { return Value == Invalid; }
-
-		static DATAFLOWCORE_API Type Current();
-		static DATAFLOWCORE_API Type Invalid; // 0
-	};
-
-	struct FRenderingParameter {
 		FRenderingParameter() {}
 		FRenderingParameter(FString InRenderName, FName InTypeName, const TArray<FName>& InOutputs)
 			: Name(InRenderName), Type(InTypeName), Outputs(InOutputs) {}
@@ -60,287 +34,11 @@ namespace UE::Dataflow
 		TArray<FName> Outputs;
 	};
 
-	//--------------------------------------------------------------------
-	// base class for all context cache entries 
-	//--------------------------------------------------------------------
-	struct FContextCacheElementBase 
-	{
-		enum EType
-		{
-			CacheElementTyped,
-			CacheElementReference,
-			CacheElementNull,
-			CacheElementUObject, // UObjectPtr
-			CacheElementUObjectArray // TArray of UObjectPtr
-		};
-
-		FContextCacheElementBase(EType CacheElementType, FGuid InNodeGuid = FGuid(), const FProperty* InProperty = nullptr, uint32 InNodeHash = 0, FTimestamp InTimestamp = FTimestamp::Invalid)
-			: Type(CacheElementType)
-			, NodeGuid(InNodeGuid)
-			, Property(InProperty)
-			, NodeHash(InNodeHash)
-			, Timestamp(InTimestamp)
-		{}
-		virtual ~FContextCacheElementBase() {}
-
-		// InReferenceDataKey is the key of the cache element this function is called on 
-		virtual TUniquePtr<FContextCacheElementBase> CreateReference(FContextCacheKey InReferenceDataKey) const = 0;
-
-		// clone the cache entry
-		virtual TUniquePtr<FContextCacheElementBase> Clone(const FContext& Context) const = 0;
-
-		template<typename T>
-		inline const T& GetTypedData(const FContext& Context, const FProperty* PropertyIn, const T& Default) const;
-
-		virtual const void* GetUntypedData(const FContext& Context, const FProperty* PropertyIn) const = 0;
-
-		EType GetType() const {	return Type; }
-
-		const FProperty* GetProperty() const { return Property; }
-		const FTimestamp& GetTimestamp() const { return Timestamp; }
-		void SetTimestamp(const FTimestamp& InTimestamp) { Timestamp = InTimestamp; }
-
-		const FGuid& GetNodeGuid() const { return NodeGuid; }
-		const uint32 GetNodeHash() const { return NodeHash; }
-
-		// use this with caution: setting the property of a wrong type may cause problems
-		void SetProperty(const FProperty* NewProperty) { Property = NewProperty; }
-
-	private:
-		friend struct FContextCache;
-
-		EType Type;
-		FGuid NodeGuid;
-		const FProperty* Property = nullptr;
-		uint32 NodeHash = 0;
-		FTimestamp Timestamp = FTimestamp::Invalid;
-	};
-
-	//--------------------------------------------------------------------
-	// Value storing context cache entry - strongly typed
-	//--------------------------------------------------------------------
-	template<class T>
-	struct TContextCacheElement : public FContextCacheElementBase 
-	{
-		TContextCacheElement(FGuid InNodeGuid, const FProperty* InProperty, T&& InData, uint32 InNodeHash, FTimestamp Timestamp)
-			: FContextCacheElementBase(EType::CacheElementTyped, InNodeGuid, InProperty, InNodeHash, Timestamp)
-			, Data(Forward<T>(InData))
-		{}
-
-		TContextCacheElement(const TContextCacheElement<T>& Other)
-			: FContextCacheElementBase(EType::CacheElementTyped, Other.GetNodeGuid(), Other.GetProperty(), Other.GetNodeHash(), Other.GetTimestamp())
-			, Data(Other.Data)
-		{}
-		
-		inline const T& GetData(const FContext& Context, const FProperty* PropertyIn, const T& Default) const;
-
-		virtual const void* GetUntypedData(const FContext& Context, const FProperty* PropertyIn) const
-		{
-			return &Data;
-		}
-
-		const T& GetDataDirect() const { return Data; }
-
-		inline virtual TUniquePtr<FContextCacheElementBase> CreateReference(FContextCacheKey InReferenceDataKey) const override;
-
-		inline virtual TUniquePtr<FContextCacheElementBase> Clone(const FContext& Context) const override;
-
-	private:
-		typedef typename TDecay<T>::Type FDataType;  // Using universal references here means T could be either const& or an rvalue reference
-		const FDataType Data;                        // Decaying T removes any reference and gets the correct underlying storage data type
-	};
-
-	//--------------------------------------------------------------------
-	// Reference to another context cache entry 
-	//--------------------------------------------------------------------
-	template<class T>
-	struct TContextCacheElementReference : public FContextCacheElementBase
-	{
-		TContextCacheElementReference(FGuid InNodeGuid, const FProperty* InProperty, FContextCacheKey InDataKey, uint32 InNodeHash, FTimestamp Timestamp)
-			: FContextCacheElementBase(EType::CacheElementReference, InNodeGuid, InProperty, InNodeHash, Timestamp)
-			, DataKey(InDataKey)
-		{}
-
-		inline const T& GetData(const FContext& Context, const FProperty* PropertyIn, const T& Default) const;
-
-		inline virtual const void* GetUntypedData(const FContext& Context, const FProperty* PropertyIn) const;
-
-		inline virtual TUniquePtr<FContextCacheElementBase> CreateReference(FContextCacheKey InReferenceDataKey) const override;
-
-		inline virtual TUniquePtr<FContextCacheElementBase> Clone(const FContext& Context) const override;
-
-	private:
-		const FContextCacheKey DataKey; // this is a key to another cache element
-	};
-
-	//--------------------------------------------------------------------
-	// Null entry, this will always return a default value 
-	//--------------------------------------------------------------------
-	struct FContextCacheElementNull : public FContextCacheElementBase
-	{
-		//
-		// IMPORTANT: 
-		// Timestamp must be set to (Timestamp.Value - 1) to make sure that this type of entry is always invalid
-		//
-		FContextCacheElementNull(FGuid InNodeGuid, const FProperty* InProperty, FContextCacheKey InDataKey, uint32 InNodeHash, FTimestamp Timestamp)
-			: FContextCacheElementBase(EType::CacheElementNull, InNodeGuid, InProperty, InNodeHash, FTimestamp(Timestamp.Value - 1))
-		{}
-
-		virtual const void* GetUntypedData(const FContext& Context, const FProperty* PropertyIn) const
-		{
-			return nullptr;
-		}
-
-		inline virtual TUniquePtr<FContextCacheElementBase> CreateReference(FContextCacheKey InReferenceDataKey) const override;
-
-		inline virtual TUniquePtr<FContextCacheElementBase> Clone(const FContext& Context) const override;
-	};
-
-	//--------------------------------------------------------------------
-	// UObject cache element, prevents the object from being garbage collected while in the cache
-	//--------------------------------------------------------------------
-	template<class T>
-	struct TContextCacheElementUObject : public FContextCacheElementBase, public FGCObject
-	{
-		TContextCacheElementUObject(FGuid InNodeGuid, const FProperty* InProperty, T&& InObjectPtr, uint32 InNodeHash, FTimestamp Timestamp)
-			: FContextCacheElementBase(EType::CacheElementUObject, InNodeGuid, InProperty, InNodeHash, Timestamp)
-			, ObjectPtr(InObjectPtr)
-		{}
-
-		TContextCacheElementUObject(const TContextCacheElementUObject<T>& Other)
-			: FContextCacheElementBase(EType::CacheElementUObject, Other.GetNodeGuid(), Other.GetProperty(), Other.GetNodeHash(), Other.GetTimestamp())
-			, ObjectPtr(Other.ObjectPtr)
-		{}
-
-		const T& GetData(const FContext& /*Context*/, const FProperty* /*PropertyIn*/, const T& /*Default*/) const
-		{
-			return ObjectPtr;
-		}
-
-		virtual const void* GetUntypedData(const FContext& /*Context*/, const FProperty* /*PropertyIn*/) const
-		{
-			return (void*)ObjectPtr;
-		}
-
-		inline virtual TUniquePtr<FContextCacheElementBase> CreateReference(FContextCacheKey InReferenceDataKey) const override;
-
-		inline virtual TUniquePtr<FContextCacheElementBase> Clone(const FContext& Context) const override;
-
-		//~ Begin FGCObject interface
-		virtual void AddReferencedObjects(FReferenceCollector& Collector) override { Collector.AddReferencedObject(ObjectPtr); }
-		virtual FString GetReferencerName() const override { return TEXT("TContextCacheElementUObject"); }
-		//~ End FGCObject interface
-
-	private:
-		typedef typename TDecay<T>::Type FDataType;  // Using universal references here means T could be either const& or an rvalue reference
-		FDataType ObjectPtr;                         // Decaying T removes any reference and gets the correct underlying storage data type
-	};
-
-	//--------------------------------------------------------------------
-	// TArray<UObjectPtr> cache element, prevents the object from being garbage collected while in the cache
-	//--------------------------------------------------------------------
-	template<class T>
-	struct TContextCacheElementUObjectArray : public FContextCacheElementBase, public FGCObject
-	{
-		template<typename ArrayT>
-		TContextCacheElementUObjectArray(FGuid InNodeGuid, const FProperty* InProperty, ArrayT&& InObjectPtrArray, uint32 InNodeHash, FTimestamp Timestamp)
-			: FContextCacheElementBase(EType::CacheElementUObjectArray, InNodeGuid, InProperty, InNodeHash, Timestamp)
-			, ObjectPtrArray(Forward<ArrayT>(InObjectPtrArray))
-		{}
-
-
-		TContextCacheElementUObjectArray(const TContextCacheElementUObjectArray<T>& Other)
-			: FContextCacheElementBase(EType::CacheElementUObject, Other.GetNodeGuid(), Other.GetProperty(), Other.GetNodeHash(), Other.GetTimestamp())
-			, ObjectPtrArray(Other.ObjectPtrArray)
-		{}
-
-		const TArray<T>& GetData(const FContext& /*Context*/, const FProperty* /*PropertyIn*/, const TArray<T>& /*Default*/) const
-		{
-			return ObjectPtrArray;
-		}
-
-		virtual const void* GetUntypedData(const FContext& /*Context*/, const FProperty* /*PropertyIn*/) const
-		{
-			return (void*)&ObjectPtrArray;
-		}
+	//////////////////////////////////////////////////////////////////////////////////
 
-		inline virtual TUniquePtr<FContextCacheElementBase> CreateReference(FContextCacheKey InReferenceDataKey) const override;
-
-		inline virtual TUniquePtr<FContextCacheElementBase> Clone(const FContext& Context) const override;
-
-		//~ Begin FGCObject interface
-		virtual void AddReferencedObjects(FReferenceCollector& Collector) override { Collector.AddReferencedObjects(ObjectPtrArray); }
-		virtual FString GetReferencerName() const override { return TEXT("TContextCacheElementUObjectArray"); }
-		//~ End FGCObject interface
-
-	private:
-		typedef typename TDecay<T>::Type FDataType;     // Using universal references here means T could be either const& or an rvalue reference
-		TArray<FDataType> ObjectPtrArray;               // Decaying T removes any reference and gets the correct underlying storage data type
-	};
-
-
-	// cache element method implementation 
-	template<class T>
-	const T& FContextCacheElementBase::GetTypedData(const FContext& Context, const FProperty* PropertyIn, const T& Default) const
-	{
-		// check(PropertyIn->IsA<T>()); // @todo(dataflow) compile error for non-class T; find alternatives
-		if (Type == EType::CacheElementTyped)
-		{
-			if constexpr (!TIsUObjectPtrElement<T>::Value)
-			{
-				return static_cast<const TContextCacheElement<T>&>(*this).GetData(Context, PropertyIn, Default);
-			}
-		}
-		if (Type == EType::CacheElementReference)
-		{
-			return static_cast<const TContextCacheElementReference<T>&>(*this).GetData(Context, PropertyIn, Default);
-		}
-		if (Type == EType::CacheElementNull)
-		{
-			return Default; 
-		}
-		if (Type == EType::CacheElementUObjectArray)
-		{
-			if constexpr (TIsTArray<T>::Value)
-			{
-				if constexpr (TIsUObjectPtrElement<typename T::ElementType>::Value)
-				{
-					return static_cast<const TContextCacheElementUObjectArray<typename T::ElementType>&>(*this).GetData(Context, PropertyIn, Default);
-				}
-			}
-		}
-		if (Type == EType::CacheElementUObject)
-		{
-			if constexpr (TIsUObjectPtrElement<T>::Value)
-			{
-				return static_cast<const TContextCacheElementUObject<T>&>(*this).GetData(Context, PropertyIn, Default);
-			}
-		}
-		check(false); // should never happen
-		return Default;
-	}
-
-	struct FContextCache : public TMap<FContextCacheKey, TUniquePtr<FContextCacheElementBase>>
-	{
-		DATAFLOWCORE_API void Serialize(FArchive& Ar);
-	};
-};
-
-inline FArchive& operator<<(FArchive& Ar, UE::Dataflow::FTimestamp& ValueIn)
-{
-	Ar << ValueIn.Value;
-	Ar << ValueIn.Invalid;
-	return Ar;
-}
-
-inline FArchive& operator<<(FArchive& Ar, UE::Dataflow::FContextCache& ValueIn)
-{
-	ValueIn.Serialize(Ar);
-	return Ar;
-}
-
-namespace UE::Dataflow
-{
+	/*
+	* Performance data per connection 
+	*/
 	struct FContextPerfData
 	{
 		void Reset();
@@ -350,11 +48,17 @@ namespace UE::Dataflow
 		{
 			float InclusiveTimeMs = 0;
 			float ExclusiveTimeMs = 0;
-			UE::Dataflow::FTimestamp LastTimestamp = UE::Dataflow::FTimestamp(0); // timestamp of the owner node when data was last collected
+			FTimestamp LastTimestamp = FTimestamp(0); // timestamp of the owner node when data was last collected
 		};
 		TMap<FGuid, FData> DataPerOutput;
 	};
 
+	//////////////////////////////////////////////////////////////////////////////////
+
+	/*
+	* Connection context callstack
+	* Used to detected loops and use for error handling
+	*/
 	struct FContextCallstack
 	{
 		void Push(const FDataflowConnection* Connection);
@@ -379,7 +83,12 @@ namespace UE::Dataflow
 		TArray<FEntry> Callstack;
 	};
 
-	class FContext
+	//////////////////////////////////////////////////////////////////////////////////
+
+	/*
+	* Dataflow context base class
+	*/
+	class FContext: public IContextCacheStore
 	{
 	protected:
 		FContext(FContext&&) = default;
@@ -521,6 +230,17 @@ namespace UE::Dataflow
 			Ar << DataStore;
 		}
 
+		// Begin - IContextCacheStore
+		virtual const TUniquePtr<FContextCacheElementBase>* FindCacheElement(FContextCacheKey InKey) const override
+		{
+			return GetDataImpl(InKey);
+		}
+		virtual bool HasCacheElement(FContextCacheKey InKey, FTimestamp InTimestamp = FTimestamp::Invalid) const override
+		{
+			return HasDataImpl(InKey, InTimestamp);
+		}
+		// End - IContextCacheStore
+
 		DATAFLOWCORE_API FTimestamp GetTimestamp(FContextCacheKey Key) const;
 
 		virtual void Evaluate(const FDataflowNode* Node, const FDataflowOutput* Output) = 0;
@@ -682,94 +402,6 @@ namespace UE::Dataflow
 		DATAFLOWCORE_API virtual bool Evaluate(const FDataflowOutput& Connection) override;
 
 	};
-
-	// cache classes implemetation 
-	// this needs to be after the FContext definition because they access its methods
-
-	template<class T>
-	const T& TContextCacheElement<T>::GetData(const FContext& Context, const FProperty* PropertyIn, const T& Default) const
-	{
-		return Data;
-	}
-
-	template<class T>
-	TUniquePtr<FContextCacheElementBase> TContextCacheElement<T>::CreateReference(FContextCacheKey InReferenceDataKey) const
-	{
-		return MakeUnique<TContextCacheElementReference<T>>(GetNodeGuid(), GetProperty(), InReferenceDataKey, GetNodeHash(), GetTimestamp());
-	}
-
-	template<class T>
-	TUniquePtr<FContextCacheElementBase> TContextCacheElement<T>::Clone(const FContext& Context) const
-	{
-		return MakeUnique<TContextCacheElement<T>>(*this);
-	}
-
-	template<class T>
-	const T& TContextCacheElementReference<T>::GetData(const FContext& Context, const FProperty* PropertyIn, const T& Default) const
-	{
-		return Context.GetData(DataKey, PropertyIn, Default);
-	}
-
-	template<class T>
-	const void* TContextCacheElementReference<T>::GetUntypedData(const FContext& Context, const FProperty* PropertyIn) const
-	{
-		return Context.GetUntypedData(DataKey, PropertyIn);
-	}
-
-	template<class T>
-	TUniquePtr<FContextCacheElementBase> TContextCacheElementReference<T>::CreateReference(FContextCacheKey InReferenceDataKey) const
-	{
-		return MakeUnique<TContextCacheElementReference<T>>(GetNodeGuid(), GetProperty(), InReferenceDataKey, GetNodeHash(), GetTimestamp());
-	}
-
-	template<class T>
-	TUniquePtr<FContextCacheElementBase> TContextCacheElementReference<T>::Clone(const FContext& Context) const
-	{
-		if (const TUniquePtr<FContextCacheElementBase>* ReferencedCacheEntry = Context.GetDataImpl(DataKey))
-		{
-			if (ReferencedCacheEntry && ReferencedCacheEntry->IsValid())
-			{
-				return (*ReferencedCacheEntry)->Clone(Context);
-			}
-		}
-		return MakeUnique<FContextCacheElementNull>(GetNodeGuid(), GetProperty(), FContextCacheKey(), GetNodeHash(), GetTimestamp());
-	}
-
-	inline TUniquePtr<FContextCacheElementBase> FContextCacheElementNull::CreateReference(FContextCacheKey InReferenceDataKey) const
-	{
-		// a null entry will always return a null entry as a reference
-		return MakeUnique<FContextCacheElementNull>(GetNodeGuid(), GetProperty(), InReferenceDataKey, GetNodeHash(), GetTimestamp());
-	}
-
-
-	TUniquePtr<FContextCacheElementBase> FContextCacheElementNull::Clone(const FContext& Context) const
-	{
-		return MakeUnique<FContextCacheElementNull>(GetNodeGuid(), GetProperty(), FContextCacheKey(), GetNodeHash(), GetTimestamp());
-	}
-
-	template<class T>
-	TUniquePtr<FContextCacheElementBase> TContextCacheElementUObject<T>::CreateReference(FContextCacheKey InReferenceDataKey) const
-	{
-		return MakeUnique<TContextCacheElementReference<T>>(GetNodeGuid(), GetProperty(), InReferenceDataKey, GetNodeHash(), GetTimestamp());
-	}
-
-	template<class T>
-	TUniquePtr<FContextCacheElementBase> TContextCacheElementUObject<T>::Clone(const FContext& Context) const
-	{
-		return MakeUnique<TContextCacheElementUObject<T>>(*this);
-	}
-
-	template<class T>
-	TUniquePtr<FContextCacheElementBase> TContextCacheElementUObjectArray<T>::CreateReference(FContextCacheKey InReferenceDataKey) const
-	{
-		return MakeUnique<TContextCacheElementReference<T>>(GetNodeGuid(), GetProperty(), InReferenceDataKey, GetNodeHash(), GetTimestamp());
-	}
-
-	template<class T>
-	TUniquePtr<FContextCacheElementBase> TContextCacheElementUObjectArray<T>::Clone(const FContext& Context) const
-	{
-		return MakeUnique<TContextCacheElementUObjectArray<T>>(*this);
-	}
 }
 
 
