[NetPhysics] Physics replication and resim improvements
- Populate PhysicsReplicationCache at the end of the solve and mark that as the physics state for the next physics frame, fixes issue where the cache could populate before or after other SimCallbackObjects (due to execution order) that might alter the physics objects state.
- Allow linking custom Input and State history in RewindData to a specific particle, also register which particle(s) are requesting a resimulation and take those particles custom input and state histories into account when finding a valid rewind frame.

diff --git a/Engine/Source/Runtime/Engine/Private/PhysicsEngine/Experimental/PhysScene_Chaos.cpp b/Engine/Source/Runtime/Engine/Private/PhysicsEngine/Experimental/PhysScene_Chaos.cpp
--- a/Engine/Source/Runtime/Engine/Private/PhysicsEngine/Experimental/PhysScene_Chaos.cpp
+++ b/Engine/Source/Runtime/Engine/Private/PhysicsEngine/Experimental/PhysScene_Chaos.cpp
@@ -183,7 +183,7 @@ public:
 						ResimFrame = (ResimFrame == INDEX_NONE) ? PendingCommand.PhysicsStep :
 							FMath::Min(ResimFrame, PendingCommand.PhysicsStep);
 
-						RewindData->SetResimFrame(ResimFrame);
+						RewindData->RequestResimulation(ResimFrame);
 					}
 				}
 				else if (PendingCommand.PhysicsStep == CurrentFrame)
diff --git a/Engine/Source/Runtime/Engine/Private/PhysicsEngine/NetworkPhysicsComponent.cpp b/Engine/Source/Runtime/Engine/Private/PhysicsEngine/NetworkPhysicsComponent.cpp
--- a/Engine/Source/Runtime/Engine/Private/PhysicsEngine/NetworkPhysicsComponent.cpp
+++ b/Engine/Source/Runtime/Engine/Private/PhysicsEngine/NetworkPhysicsComponent.cpp
@@ -461,14 +461,6 @@ void UNetworkPhysicsComponent::InitPhysics()
 		bCompareInputToTriggerRewind = CVarCompareInputToTriggerRewind->GetBool();
 	}
 
-	if (AActor* Owner = GetOwner())
-	{
-		if (UPrimitiveComponent* RootPrimComp = Cast<UPrimitiveComponent>(Owner->GetRootComponent()))
-		{
-			RootPhysicsObject = RootPrimComp->GetPhysicsObjectByName(NAME_None);
-		}
-	}
-
 	/** NOTE:
 	* If the NetworkPhysicsComponent is added as a SubObject after the actor has processed bAutoActivate and bWantsInitializeComponent
 	* SetActive(true) and InitializeComponent() needs to be called manually for the component to function properly. */
@@ -522,6 +514,14 @@ void UNetworkPhysicsComponent::InitializeComponent()
 				ReplicatedStates.History->ResizeDataHistory(StatesToNetwork);
 			}
 		}
+
+		if (!PhysicsObject)
+		{
+			if (UPrimitiveComponent* RootPrimComp = Cast<UPrimitiveComponent>(Owner->GetRootComponent()))
+			{
+				SetPhysicsObject(RootPrimComp->GetPhysicsObjectByName(NAME_None));
+			}
+		}
 	}
 
 	if (UWorld* World = GetWorld())
@@ -532,7 +532,7 @@ void UNetworkPhysicsComponent::InitializeComponent()
 			{
 				// Create async component to run on Physics Thread
 				NetworkPhysicsComponent_Internal = Solver->CreateAndRegisterSimCallbackObject_External<FAsyncNetworkPhysicsComponent>();
-				NetworkPhysicsComponent_Internal->RootPhysicsObject = RootPhysicsObject;
+				NetworkPhysicsComponent_Internal->PhysicsObject = PhysicsObject;
 				NetworkPhysicsComponent_Internal->InputsToNetwork_OwnerDefault = InputsToNetwork_OwnerDefault;
 				NetworkPhysicsComponent_Internal->InputsToNetwork_Simulated = InputsToNetwork_Simulated;
 				NetworkPhysicsComponent_Internal->StatesToNetwork = StatesToNetwork;
@@ -795,6 +795,12 @@ void UNetworkPhysicsComponent::OnRep_SetReplicatedDeltaSourceInput()
 		ensureMsgf(false, TEXT("OnRep_SetReplicatedDeltaSourceInput failed delta serialization, should not happen."));
 		return;
 	}
+
+	if (!ReplicatedDeltaSourceInput.History)
+	{
+		return;
+	}
+
 	AddDeltaSourceInput();
 }
 
@@ -867,6 +873,12 @@ void UNetworkPhysicsComponent::OnRep_SetReplicatedDeltaSourceState()
 		ensureMsgf(false, TEXT("OnRep_SetReplicatedDeltaSourceState failed delta serialization, should not happen."));
 		return;
 	}
+
+	if (!ReplicatedDeltaSourceState.History)
+	{
+		return;
+	}
+
 	AddDeltaSourceState();
 }
 
@@ -1315,6 +1327,25 @@ APlayerController* UNetworkPhysicsComponent::GetPlayerController() const
 	return nullptr;
 }
 
+void UNetworkPhysicsComponent::SetPhysicsObject(Chaos::FConstPhysicsObjectHandle InPhysicsObject)
+{
+	if (PhysicsObject == InPhysicsObject)
+	{
+		return;
+	}
+
+	PhysicsObject = InPhysicsObject;
+
+	// Marshal data from Game Thread to Physics Thread
+	if (NetworkPhysicsComponent_Internal)
+	{
+		if (FAsyncNetworkPhysicsComponentInput* AsyncInput = NetworkPhysicsComponent_Internal->GetProducerInputData_External())
+		{
+			AsyncInput->PhysicsObject = InPhysicsObject;
+		}
+	}
+}
+
 void UNetworkPhysicsComponent::UpdateAsyncComponent(const bool bFullUpdate)
 {
 	// Marshal data from Game Thread to Physics Thread
@@ -1444,7 +1475,7 @@ FAsyncNetworkPhysicsComponent::FAsyncNetworkPhysicsComponent() : TSimCallbackObj
 	, NetworkPhysicsTickOffset(0)
 	, PhysicsReplicationMode(EPhysicsReplicationMode::Default)
 	, SettingsComponent(nullptr)
-	, RootPhysicsObject(nullptr)
+	, PhysicsObject(nullptr)
 	, bCompareStateToTriggerRewind(false)
 	, bCompareStateToTriggerRewindIncludeSimProxies(false)
 	, bCompareInputToTriggerRewind(false)
@@ -1535,6 +1566,14 @@ void FAsyncNetworkPhysicsComponent::ConsumeAsyncInput(const int32 PhysicsStep)
 			{
 				ActorComponent = *AsyncInput->ActorComponent;
 			}
+			if (AsyncInput->PhysicsObject.IsSet())
+			{
+				if (PhysicsObject == nullptr || PhysicsObject != *AsyncInput->PhysicsObject)
+				{
+					PhysicsObject = *AsyncInput->PhysicsObject;
+					RegisterDataHistoryInRewindData();
+				}
+			}
 			if (AsyncInput->ActorName.IsSet())
 			{
 				ActorName = *AsyncInput->ActorName;
@@ -2102,19 +2141,11 @@ void FAsyncNetworkPhysicsComponent::TriggerResimulation(int32 ResimFrame)
 		{
 			if (Chaos::FRewindData* RewindData = RigidSolver->GetRewindData())
 			{
-				// Mark particle/island as resim
 				Chaos::FReadPhysicsObjectInterface_Internal Interface = Chaos::FPhysicsObjectInternalInterface::GetRead();
-				if (Chaos::FPBDRigidParticleHandle* POHandle = Interface.GetRigidParticle(RootPhysicsObject))
-				{
-					if (Chaos::FPBDRigidsEvolution* Evolution = GetEvolution())
-					{
-						Evolution->GetIslandManager().SetParticleResimFrame(POHandle, ResimFrame);
-					}
-				}
+				Chaos::FPBDRigidParticleHandle* Particle = Interface.GetRigidParticle(PhysicsObject);
 
 				// Set resim frame in rewind data
-				ResimFrame = (RewindData->GetResimFrame() == INDEX_NONE) ? ResimFrame : FMath::Min(ResimFrame, RewindData->GetResimFrame());
-				RewindData->SetResimFrame(ResimFrame);
+				RewindData->RequestResimulation(ResimFrame, Particle);
 			}
 		}
 	}
@@ -2218,10 +2249,16 @@ void FAsyncNetworkPhysicsComponent::RegisterDataHistoryInRewindData()
 		{
 			UnregisterDataHistoryFromRewindData();
 
-			RewindData->AddInputHistory(InputHistory);
+			Chaos::FReadPhysicsObjectInterface_Internal Interface = Chaos::FPhysicsObjectInternalInterface::GetRead();
+			Chaos::FGeometryParticleHandle* Particle = Interface.GetParticle(PhysicsObject);
+
+			if (InputHistory)
+			{
+				RewindData->AddInputHistory(InputHistory, Particle);
+			}
 			if (StateHistory)
 			{
-				RewindData->AddStateHistory(StateHistory);
+				RewindData->AddStateHistory(StateHistory, Particle);
 			}
 		}
 	}
@@ -2233,8 +2270,11 @@ void FAsyncNetworkPhysicsComponent::UnregisterDataHistoryFromRewindData()
 	{
 		if (Chaos::FRewindData* RewindData = BaseSolver->GetRewindData())
 		{
-			RewindData->RemoveInputHistory(InputHistory);
-			RewindData->RemoveStateHistory(StateHistory);
+			Chaos::FReadPhysicsObjectInterface_Internal Interface = Chaos::FPhysicsObjectInternalInterface::GetRead();
+			Chaos::FGeometryParticleHandle* Particle = Interface.GetParticle(PhysicsObject);
+
+			RewindData->RemoveInputHistory(InputHistory, Particle);
+			RewindData->RemoveStateHistory(StateHistory, Particle);
 		}
 	}
 }
diff --git a/Engine/Source/Runtime/Engine/Private/PhysicsEngine/PhysicsReplication.cpp b/Engine/Source/Runtime/Engine/Private/PhysicsEngine/PhysicsReplication.cpp
--- a/Engine/Source/Runtime/Engine/Private/PhysicsEngine/PhysicsReplication.cpp
+++ b/Engine/Source/Runtime/Engine/Private/PhysicsEngine/PhysicsReplication.cpp
@@ -2203,12 +2203,8 @@ bool FPhysicsReplicationAsync::ResimulationReplication(Chaos::FPBDRigidParticleH
 
 	if (bShouldTriggerResim && Target.TickCount == 0 && LocalFrame > RewindData->GetBlockedResimFrame())
 	{
-		// Trigger resimulation
-		RigidsSolver->GetEvolution()->GetIslandManager().SetParticleResimFrame(Handle, LocalFrame);
-
-		int32 ResimFrame = RewindData->GetResimFrame();
-		ResimFrame = (ResimFrame == INDEX_NONE) ? LocalFrame : FMath::Min(ResimFrame, LocalFrame);
-		RewindData->SetResimFrame(ResimFrame);
+		// Request resimulation
+		RewindData->RequestResimulation(LocalFrame, Handle);
 	}
 	else if (SettingsCurrent.ResimulationSettings.GetRuntimeCorrectionEnabled())
 	{
diff --git a/Engine/Source/Runtime/Engine/Private/PhysicsEngine/PhysicsReplicationCache.cpp b/Engine/Source/Runtime/Engine/Private/PhysicsEngine/PhysicsReplicationCache.cpp
--- a/Engine/Source/Runtime/Engine/Private/PhysicsEngine/PhysicsReplicationCache.cpp
+++ b/Engine/Source/Runtime/Engine/Private/PhysicsEngine/PhysicsReplicationCache.cpp
@@ -162,6 +162,10 @@ void FPhysicsReplicationCache::UnregisterForReplicationCache(UPrimitiveComponent
 void FPhysicsReplicationCacheAsync::OnPreSimulate_Internal()
 {
 	ProcessAsyncInputs();
+}
+
+void FPhysicsReplicationCacheAsync::OnPostSolve_Internal()
+{
 	PopulateReplicationCache_Internal();
 }
 
@@ -218,7 +222,9 @@ void FPhysicsReplicationCacheAsync::PopulateReplicationCache_Internal()
 	}
 
 	FPhysicsReplicationCacheAsyncOutput& AsyncOutput = GetProducerOutputData_Internal();
-	AsyncOutput.SolverFrame = RigidSolver->GetCurrentFrame();
+	
+	// +1 Because we cache this in PostSolve which has the end result of the physics solve which is equal to the starting state of next physics frame which is what this state correspond to
+	AsyncOutput.SolverFrame = RigidSolver->GetCurrentFrame() + 1;
 
 	// Setup the arrays to marshal replication cache with correct allocation size
 	AsyncOutput.ReplicationCache_Key_Marshal.SetNum(ReplicationCache_Internal.Num());
@@ -251,8 +257,8 @@ void FPhysicsReplicationCacheAsync::PopulateReplicationCache_Internal()
 			if (Chaos::FPBDRigidParticleHandle* Handle = Interface.GetRigidParticle(PhysicsObject))
 			{
 				FRigidBodyState& ReplicationState = ReplicationCacheData.GetState();
-				ReplicationState.Position = Handle->GetX();
-				ReplicationState.Quaternion = Handle->GetR();
+				ReplicationState.Position = Handle->GetP();
+				ReplicationState.Quaternion = Handle->GetQ();
 				ReplicationState.LinVel = Handle->GetV();
 				ReplicationState.AngVel = FMath::RadiansToDegrees(Handle->GetW());
 				ReplicationState.Flags = Handle->ObjectState() == Chaos::EObjectStateType::Sleeping ? ERigidBodyFlags::Sleeping : 0;
diff --git a/Engine/Source/Runtime/Engine/Public/Physics/NetworkPhysicsComponent.h b/Engine/Source/Runtime/Engine/Public/Physics/NetworkPhysicsComponent.h
--- a/Engine/Source/Runtime/Engine/Public/Physics/NetworkPhysicsComponent.h
+++ b/Engine/Source/Runtime/Engine/Public/Physics/NetworkPhysicsComponent.h
@@ -974,6 +974,10 @@ public:
 	/** Get FAsyncNetworkPhysicsComponent on the Physics Thread */
 	FAsyncNetworkPhysicsComponent* GetNetworkPhysicsComponent_Internal() { return NetworkPhysicsComponent_Internal; }
 
+	/** Set the physics object that gets affected by this NetworkPhysicsComponents inputs / states, default is the physics object from the root primitive component.
+	* Used to link a desynced physics object that needs resimulation to its input and state history to be able to ensure there is valid data to rewind to. */
+	ENGINE_API void SetPhysicsObject(Chaos::FConstPhysicsObjectHandle InPhysicsObject);
+
 protected: 
 
 	// repnotify for input, used for delta serialization
@@ -1189,7 +1193,7 @@ public:
 
 private:
 	// Root components physics object
-	Chaos::FPhysicsObjectHandle RootPhysicsObject;
+	Chaos::FConstPhysicsObjectHandle PhysicsObject;
 
 	// Locally relayed inputs makes this component act as if it's a locally controlled pawn.
 	bool bIsRelayingLocalInputs = false;
@@ -1325,6 +1329,7 @@ struct FAsyncNetworkPhysicsComponentInput : public Chaos::FSimCallbackInput
 	TOptional<uint16> InputsToNetwork_Owner;
 	TOptional<EPhysicsReplicationMode> PhysicsReplicationMode;
 	TOptional<TWeakObjectPtr<UActorComponent>> ActorComponent;
+	TOptional<Chaos::FConstPhysicsObjectHandle> PhysicsObject;
 	TOptional<FString> ActorName;
 	TOptional<TUniquePtr<FNetworkPhysicsDataHelper>> InputHelper;
 	TOptional<TUniquePtr<FNetworkPhysicsDataHelper>> StateHelper;
@@ -1350,6 +1355,7 @@ struct FAsyncNetworkPhysicsComponentInput : public Chaos::FSimCallbackInput
 		InputsToNetwork_Owner.Reset();
 		PhysicsReplicationMode.Reset();
 		ActorComponent.Reset();
+		PhysicsObject.Reset();
 		ActorName.Reset();
 		InputHelper.Reset();
 		StateHelper.Reset();
@@ -1503,7 +1509,7 @@ private:
 	TWeakObjectPtr<UActorComponent> ActorComponent;
 
 	// Root components physics object
-	Chaos::FPhysicsObjectHandle RootPhysicsObject;
+	Chaos::FConstPhysicsObjectHandle PhysicsObject;
 
 	// States history uses to rewind simulation 
 	TSharedPtr<Chaos::FBaseRewindHistory> StateHistory;
diff --git a/Engine/Source/Runtime/Engine/Public/Physics/PhysicsReplicationCache.h b/Engine/Source/Runtime/Engine/Public/Physics/PhysicsReplicationCache.h
--- a/Engine/Source/Runtime/Engine/Public/Physics/PhysicsReplicationCache.h
+++ b/Engine/Source/Runtime/Engine/Public/Physics/PhysicsReplicationCache.h
@@ -100,9 +100,10 @@ private:
 class FPhysicsReplicationCacheAsync : public Chaos::TSimCallbackObject<
 	FPhysicsReplicationCacheAsyncInput,
 	FPhysicsReplicationCacheAsyncOutput,
-	Chaos::ESimCallbackOptions::Presimulate | Chaos::ESimCallbackOptions::PhysicsObjectUnregister>
+	Chaos::ESimCallbackOptions::Presimulate | Chaos::ESimCallbackOptions::PostSolve | Chaos::ESimCallbackOptions::PhysicsObjectUnregister>
 {
 	virtual void OnPreSimulate_Internal() override;
+	virtual void OnPostSolve_Internal() override;
 
 	/** Callback when a physics object is removed */
 	virtual void OnPhysicsObjectUnregistered_Internal(Chaos::FConstPhysicsObjectHandle PhysicsObject) override;
diff --git a/Engine/Source/Runtime/Experimental/Chaos/Private/RewindData.cpp b/Engine/Source/Runtime/Experimental/Chaos/Private/RewindData.cpp
--- a/Engine/Source/Runtime/Experimental/Chaos/Private/RewindData.cpp
+++ b/Engine/Source/Runtime/Experimental/Chaos/Private/RewindData.cpp
@@ -773,15 +773,20 @@ void FRewindData::AdvanceFrameImp(IResimCacheBase* ResimCache)
 				auto Handle = Info.GetObjectPtr();
 				Info.bResimAsFollower = Handle->ResimType() == EResimType::ResimAsFollower;
 
-				if (!bRewindDataOptimization && IsResim() && !Info.bResimAsFollower)
+				if (IsResim())
 				{
-					DesyncIfNecessary</*bSkipDynamics=*/false>(Info, FrameAndPhase);
-				}
+					if (!bRewindDataOptimization && !Info.bResimAsFollower)
+					{
+						DesyncIfNecessary</*bSkipDynamics=*/false>(Info, FrameAndPhase);
+					}
 
-				if (IsResim() && Handle->SyncState() != ESyncState::InSync && !SkipDesyncTest)
-				{
-					Handle->SetEnabledDuringResim(true);	//for now just mark anything out of sync as resim enabled. TODO: use bubble
-					DesyncFunc(Handle);
+					if (Handle->SyncState() != ESyncState::InSync && !SkipDesyncTest)
+					{
+						Handle->SetEnabledDuringResim(true);	//for now just mark anything out of sync as resim enabled. TODO: use bubble
+						DesyncFunc(Handle);
+					}
+
+					Info.bNeedsResim = false;
 				}
 
 				AdvanceDirtyFunc(Info, Handle);
@@ -1195,16 +1200,21 @@ enum class EResimFrameValidation : int32
 	IslandValidation = 1, // Validate dirty particles inside the islands that have resim trigger particles in them
 	TriggerParticleValidation = 2 // Only validate the resim triggering particle(s)
 };
-CHAOS_API int32 ResimFrameValidation = (int32)EResimFrameValidation::IslandValidation;
+CHAOS_API int32 ResimFrameValidation = static_cast<int32>(EResimFrameValidation::IslandValidation);
 FAutoConsoleVariableRef CVarResimFrameValidationLeniency(TEXT("p.Resim.ResimFrameValidation"), ResimFrameValidation, TEXT("0 = no leniency, all dirty particles need a valid target. 1 = Island leniency, all particles in resim islands need a valid target. 2 = Full leniency, only the particle triggering the resim need a valid target."));
 CHAOS_API bool bResimIncompleteHistory = false;
 FAutoConsoleVariableRef CVarResimIncompleteHistory(TEXT("p.Resim.IncompleteHistory"), bResimIncompleteHistory, TEXT("If a valid resim frame can't be found, use the requested resim frame and perform a resimulation with incomplete data."));
+CHAOS_API bool bFindValidInputHistory = true;
+FAutoConsoleVariableRef CVarResimFindValidInputHistory(TEXT("p.Resim.FindValidInputHistory"), bFindValidInputHistory, TEXT("If the particle that needs resimulation has custom input history, find a valid resim frame where inputs are available."));
+CHAOS_API bool bFindValidStateHistory = true;
+FAutoConsoleVariableRef CVarResimFindValidStateHistory(TEXT("p.Resim.FindValidStateHistory"), bFindValidStateHistory, TEXT("If the particle that needs resimulation has custom state history, find a valid resim frame where states are available."));
 CHAOS_API bool bUseParticleResimAsFollowerDuringTargetValidation = false;
 FAutoConsoleVariableRef CVarUseParticleResimAsFollowerDuringTargetValidation(TEXT("p.Resim.UseParticleResimAsFollowerDuringTargetValidation"), bUseParticleResimAsFollowerDuringTargetValidation, TEXT("If disabled, do not use the particle's ResimAsFollower flag when trying to find a valid resim frame."));
 
 int32 FRewindData::FindValidResimFrame(const int32 RequestedFrame)
 {
 	int32 ValidFrame = INDEX_NONE;
+	int32 ValidTargetFrame = INDEX_NONE;
 
 	if (RequestedFrame <= BlockResimFrame)
 	{
@@ -1240,10 +1250,45 @@ int32 FRewindData::FindValidResimFrame(const int32 RequestedFrame)
 		return bValid;
 	};
 
+	auto CustomDataFinderHelper = [&](FDirtyParticleInfo* DirtyParticleInfo, const FFrameAndPhase FrameAndPhase) -> bool
+	{
+		if (!DirtyParticleInfo->bNeedsResim)
+		{
+			true;
+		}
+
+		bool bValid = true;
+		FGeometryParticleHandle* Handle = DirtyParticleInfo->GetObjectPtr();
+		
+		if (bFindValidInputHistory)
+		{
+			if (const TWeakPtr<FBaseRewindHistory>* InputHistory = InputParticleHistories.Find(Handle))
+			{
+				if (InputHistory->IsValid() && !InputHistory->Pin().Get()->HasValidData(FrameAndPhase.Frame))
+				{
+					bValid = false;
+				}
+			}
+		}
+
+		if (bValid && bFindValidStateHistory)
+		{
+			if (const TWeakPtr<FBaseRewindHistory>* StateHistory = StateParticleHistories.Find(Handle))
+			{
+				if (StateHistory->IsValid() && !StateHistory->Pin().Get()->HasValidData(FrameAndPhase.Frame))
+				{
+					bValid = false;
+				}
+			}
+		}
+
+		return bValid;
+	};
+
 	Private::FPBDIslandManager& IslandManager = Solver->GetEvolution()->GetIslandManager();
 
 	// Cache all particles in islands that have a resim triggering particle
-	if (ResimFrameValidation == (int32)EResimFrameValidation::IslandValidation)
+	if (ResimFrameValidation == static_cast<int32>(EResimFrameValidation::IslandValidation))
 	{
 		IslandValidationIslands.Reset();
 		IslandValidationIslandParticles.Reset();
@@ -1254,25 +1299,34 @@ int32 FRewindData::FindValidResimFrame(const int32 RequestedFrame)
 			{
 				IslandManager.FindParticleIslands(Handle, OUT IslandValidationIslands);
 			}
+
+			/** If the particle needs resim, add it to the IslandValidationIslandParticles array to ensure it will be processed
+			* particles that are not in contact with anything doesn't have an island and would be missed */
+			if (DirtyParticleInfo.bNeedsResim)
+			{
+				IslandValidationIslandParticles.AddUnique(Handle);
+			}
 		}
 
 		IslandManager.FindParticlesInIslands(IslandValidationIslands, OUT IslandValidationIslandParticles);
 	}
 
-	// First frame of the history datas
+	// First frame of the history data
 	const int32 EarliestFrame = FMath::Max(GetEarliestFrame_Internal(), BlockResimFrame);
 	bool bHasTargetHistory = false;
+	bool bHasCustomDataHistory = false;
 
-	for (ValidFrame = RequestedFrame; ValidFrame > EarliestFrame; ValidFrame--)
+	for (int32 CheckFrame = RequestedFrame; CheckFrame > EarliestFrame; CheckFrame--)
 	{
-		const FFrameAndPhase FrameAndPhase{ ValidFrame, FFrameAndPhase::PostPushData };
+		const FFrameAndPhase FrameAndPhase{ CheckFrame, FFrameAndPhase::PostPushData };
 		bHasTargetHistory = true;
+		bHasCustomDataHistory = true;
 		
 #if DEBUG_REWIND_DATA
-		UE_LOG(LogChaos, Log, TEXT("CLIENT | PT | FindValidResimFrame | Processing resim particles | History Frame: %d | Total Particle Count: %d | ResimIslands Particle Count: %d | ResimFrameValidation: %d"), ValidFrame, DirtyParticles.Num(), IslandValidationIslandParticles.Num(), ResimFrameValidation);
+		UE_LOG(LogChaos, Log, TEXT("CLIENT | PT | FindValidResimFrame | Processing resim particles | Check Frame: %d | Total Particle Count: %d | ResimIslands Particle Count: %d | ResimFrameValidation: %d | ValidTargetFrame: %d"), CheckFrame, DirtyParticles.Num(), IslandValidationIslandParticles.Num(), ResimFrameValidation, ValidTargetFrame);
 #endif
 
-		if ((EResimFrameValidation)ResimFrameValidation == EResimFrameValidation::IslandValidation)
+		if (ResimFrameValidation == static_cast<int32>(EResimFrameValidation::IslandValidation))
 		{
 			// Iterate over islands previously found having resim particles in them and check if the particles in the islands have targets
 			for (const FGeometryParticleHandle* IslandParticle : IslandValidationIslandParticles)
@@ -1285,6 +1339,12 @@ int32 FRewindData::FindValidResimFrame(const int32 RequestedFrame)
 						bHasTargetHistory = false;
 						break;
 					}
+
+					if (!CustomDataFinderHelper(DirtyParticleInfo, FrameAndPhase))
+					{
+						bHasCustomDataHistory = false;
+						break;
+					}
 				}
 			}
 		}
@@ -1293,7 +1353,7 @@ int32 FRewindData::FindValidResimFrame(const int32 RequestedFrame)
 			for (FDirtyParticleInfo& DirtyParticleInfo : DirtyParticles)
 			{
 				// If running validation leniency, check if the particle is marked for resimulation else don't bother checking for valid target states.
-				if ((EResimFrameValidation)ResimFrameValidation == EResimFrameValidation::TriggerParticleValidation)
+				if (ResimFrameValidation == static_cast<int32>(EResimFrameValidation::TriggerParticleValidation))
 				{
 					FGeometryParticleHandle* Handle = DirtyParticleInfo.GetObjectPtr();
 					if(IslandManager.GetParticleResimFrame(Handle) == INDEX_NONE)
@@ -1307,50 +1367,53 @@ int32 FRewindData::FindValidResimFrame(const int32 RequestedFrame)
 					bHasTargetHistory = false;
 					break;
 				}
-			}
-		}
-		
-/*
-* todo: Only check input and state histories for actors that will actively resimulate
 
-		if (bHasTargetHistory)
-		{
-			for (TWeakPtr<FBaseRewindHistory>& InputHistory : InputHistories)
-			{
-				if (InputHistory.IsValid() && !InputHistory.Pin().Get()->HasValidData(ValidFrame))
+				if (!CustomDataFinderHelper(&DirtyParticleInfo, FrameAndPhase))
 				{
-					bHasTargetHistory = false;
+					bHasCustomDataHistory = false;
 					break;
 				}
 			}
 		}
 
-		if (bHasTargetHistory)
+		if (bHasTargetHistory && bHasCustomDataHistory)
 		{
-			for (TWeakPtr<FBaseRewindHistory>& StateHistory : StateHistories)
-			{
-				if (StateHistory.IsValid() && !StateHistory.Pin().Get()->HasValidData(ValidFrame))
-				{
-					bHasTargetHistory = false;
-					break;
-				}
-			}
+			ValidFrame = CheckFrame;
+			break;
 		}
-*/
-
-		if (bHasTargetHistory)
+		else if (bHasTargetHistory && ValidTargetFrame == INDEX_NONE)
 		{
-			break;
+			// If we have a valid frame with targets from the server but no custom data, cache the frame number to use if we don't find any frame with valid custom data
+			ValidTargetFrame = CheckFrame;
 		}
 	}
 	
-	if (!bHasTargetHistory)
+	// Check if no valid frame was found with both target state and custom data
+	if (ValidFrame == INDEX_NONE)
 	{
-		ValidFrame = bResimIncompleteHistory ? RequestedFrame : INDEX_NONE;
+		// Check if a valid target frame was found
+		if (ValidTargetFrame != INDEX_NONE)
+		{
+			ValidFrame = ValidTargetFrame;
+		}
+		else // No valid frame found
+		{
+			ValidFrame = bResimIncompleteHistory ? RequestedFrame : INDEX_NONE;
 
-#if DEBUG_REWIND_DATA
-		UE_LOG(LogChaos, Warning, TEXT("CLIENT | PT | FindValidResimFrame | No valid resim frame found | RequestedFrame: %d | ValidFrame: %d | EarliestFrame: %d | HasTargetHistory: %d | EarliestHistoryFrame: %d | CurrentFrame: %d | FramesSaved: %d | ResimFrameValidation: %d"), RequestedFrame, ValidFrame, EarliestFrame, bHasTargetHistory, GetEarliestFrame_Internal(), CurrentFrame(), FramesSaved, ResimFrameValidation);
-#endif
+			// If we can't perform a resim, clear bNeedsResim flags on dirty physics objects, else they will get cleared during AdvanceFrame if we are resimulating
+			for (FDirtyParticleInfo& DirtyParticleInfo : DirtyParticles)
+			{
+				DirtyParticleInfo.bNeedsResim = false;
+			}
+			for (FDirtyJointInfo& DirtyJointInfo : DirtyJoints)
+			{
+				DirtyJointInfo.bNeedsResim = false;
+			}
+
+	#if DEBUG_REWIND_DATA
+			UE_LOG(LogChaos, Warning, TEXT("CLIENT | PT | FindValidResimFrame | No valid resim frame found | RequestedFrame: %d | ValidFrame: %d | ValidTargetFrame: %d | EarliestFrame: %d | HasTargetHistory: %d | HasCustomDataHistory: %d | EarliestHistoryFrame: %d | CurrentFrame: %d | FramesSaved: %d | ResimFrameValidation: %d"), RequestedFrame, ValidFrame, ValidTargetFrame, EarliestFrame, bHasTargetHistory, bHasCustomDataHistory, GetEarliestFrame_Internal(), CurrentFrame(), FramesSaved, ResimFrameValidation);
+	#endif
+		}
 	}
 
 	return ValidFrame;
@@ -1573,6 +1636,25 @@ void FRewindData::SetTargetStateAtFrame(FGeometryParticleHandle& Handle, const i
 		Position, Quaternion, LinVelocity, AngVelocity, bShouldSleep);
 }
 
+void FRewindData::RequestResimulation(int32 Frame, Chaos::FGeometryParticleHandle* Particle)
+{
+	// Update ResimFrame but don't allow to set a newer frame than already set
+	ResimFrame = (ResimFrame == INDEX_NONE) ? Frame : FMath::Min(ResimFrame, Frame);
+
+	if (Particle && ensure(Solver))
+	{
+		if (FDirtyParticleInfo* DirtyParticleInfo = FindDirtyObj(*Particle))
+		{
+			DirtyParticleInfo->bNeedsResim = true;
+		}
+
+		if (Chaos::FPBDRigidsEvolution* Evolution = Solver->GetEvolution())
+		{
+			Evolution->GetIslandManager().SetParticleResimFrame(Particle, ResimFrame);
+		}
+	}
+}
+
 void FRewindData::BlockResim()
 {
 	if (LatestFrame > BlockResimFrame)
diff --git a/Engine/Source/Runtime/Experimental/Chaos/Public/RewindData.h b/Engine/Source/Runtime/Experimental/Chaos/Public/RewindData.h
--- a/Engine/Source/Runtime/Experimental/Chaos/Public/RewindData.h
+++ b/Engine/Source/Runtime/Experimental/Chaos/Public/RewindData.h
@@ -1700,6 +1700,26 @@ public:
 		InputHistories.Remove(InputHistory.ToWeakPtr());
 	}
 
+	/** Add input history for particle to the rewind data */
+	void AddInputHistory(const TSharedPtr<FBaseRewindHistory>& InputHistory, Chaos::FGeometryParticleHandle* Particle)
+	{
+		AddInputHistory(InputHistory);
+		if (Particle != nullptr)
+		{
+			InputParticleHistories.Add(Particle, InputHistory.ToWeakPtr());
+		}
+	}
+
+	/** Remove input history for particle from the rewind data */
+	void RemoveInputHistory(const TSharedPtr<FBaseRewindHistory>& InputHistory, Chaos::FGeometryParticleHandle* Particle)
+	{
+		RemoveInputHistory(InputHistory);
+		if (Particle != nullptr)
+		{
+			InputParticleHistories.Remove(Particle);
+		}
+	}
+
 	/** Add state history to the rewind data for future use while rewinding */
 	void AddStateHistory(const TSharedPtr<FBaseRewindHistory>& StateHistory)
 	{
@@ -1712,6 +1732,26 @@ public:
 		StateHistories.Remove(StateHistory.ToWeakPtr());
 	}
 
+	/** Add state history for particle to the rewind data */
+	void AddStateHistory(const TSharedPtr<FBaseRewindHistory>& StateHistory, Chaos::FGeometryParticleHandle* Particle)
+	{
+		AddInputHistory(StateHistory);
+		if (Particle != nullptr)
+		{
+			StateParticleHistories.Add(Particle, StateHistory.ToWeakPtr());
+		}
+	}
+
+	/** Remove state history for particle from the rewind data */
+	void RemoveStateHistory(const TSharedPtr<FBaseRewindHistory>& StateHistory, Chaos::FGeometryParticleHandle* Particle)
+	{
+		RemoveInputHistory(StateHistory);
+		if (Particle != nullptr)
+		{
+			StateParticleHistories.Remove(Particle);
+		}
+	}
+
 	/** Apply inputs for specified frame from rewind data */
 	UE_DEPRECATED(5.6, "Deprecated, ApplyInputs is no longer viable. Any custom states can be applied during IRewindCallback::ProcessInputs_Internal during resimulation. Example FNetworkPhysicsCallback")
 		void ApplyInputs(const int32 ApplyFrame, const bool bResetSolver);
@@ -1733,6 +1773,10 @@ public:
 	const int32 GetResimFrame() const { return ResimFrame; }
 	void SetResimFrame(int32 Frame) { ResimFrame = Frame; }
 
+	/** Request a resimulation by setting a requested frame to rewind to
+	* @param Particle can be used to mark a specific particle as the one triggering the resimulation */
+	void CHAOS_API RequestResimulation(int32 RequestedFrame, Chaos::FGeometryParticleHandle* Particle = nullptr);
+
 	/** This blocks any future resimulation to rewind back past the frame this is called on */
 	void BlockResim();
 
@@ -1779,6 +1823,7 @@ private:
 		int32 LastDirtyFrame;	//Track how recently this was made dirty
 		int32 InitializedOnStep = INDEX_NONE;	//if not INDEX_NONE, it indicates we saw initialization during rewind history window
 		bool bResimAsFollower = true;	//Indicates the particle will always resim in the exact same way from game thread data
+		bool bNeedsResim = false;	//This particle needs resimulation, should have a higher priority when checking for valid rewind frames
 
 		TDirtyObjectInfo(FDirtyPropertiesPool& InPropertiesPool, TObj& InObj, const int32 CurFrame, const int32 NumFrames, const bool bCacheOnePhase)
 			: History(NumFrames, bCacheOnePhase)
@@ -1795,6 +1840,7 @@ private:
 			, LastDirtyFrame(Other.LastDirtyFrame)
 			, InitializedOnStep(Other.InitializedOnStep)
 			, bResimAsFollower(Other.bResimAsFollower)
+			, bNeedsResim(Other.bNeedsResim)
 		{
 			Other.PropertiesPool = nullptr;
 		}
@@ -1932,8 +1978,11 @@ private:
 	TDirtyObjects<FDirtyJointInfo> DirtyJoints;
 	TDirtyObjects<FDirtyParticleErrorInfo> DirtyParticleErrors;
 
-	TArray<TWeakPtr<FBaseRewindHistory>> InputHistories;
-	TArray<TWeakPtr<FBaseRewindHistory>> StateHistories;
+	TArray<TWeakPtr<FBaseRewindHistory>> InputHistories; // Todo, deprecate in favor of Particle->Input map?
+	TArray<TWeakPtr<FBaseRewindHistory>> StateHistories; // Todo, deprecate in favor of Particle->State map?
+
+	TMap<Chaos::FGeometryParticleHandle*, TWeakPtr<FBaseRewindHistory>> InputParticleHistories;
+	TMap<Chaos::FGeometryParticleHandle*, TWeakPtr<FBaseRewindHistory>> StateParticleHistories;
 
 	FPBDRigidsSolver* Solver;
 	int32 CurFrame;
