[NetPhysics] Fix missing return in recent change

diff --git a/Engine/Source/Runtime/Experimental/Chaos/Private/RewindData.cpp b/Engine/Source/Runtime/Experimental/Chaos/Private/RewindData.cpp
--- a/Engine/Source/Runtime/Experimental/Chaos/Private/RewindData.cpp
+++ b/Engine/Source/Runtime/Experimental/Chaos/Private/RewindData.cpp
@@ -1254,7 +1254,7 @@ int32 FRewindData::FindValidResimFrame(const int32 RequestedFrame)
 	{
 		if (!DirtyParticleInfo->bNeedsResim)
 		{
-			true;
+			return true;
 		}
 
 		bool bValid = true;
