Dataflow : add the ability to read current cached/frozen value on an output without causing an evaluation of the graph

diff --git a/Engine/Source/Runtime/Experimental/Dataflow/Core/Public/Dataflow/DataflowInputOutput.h b/Engine/Source/Runtime/Experimental/Dataflow/Core/Public/Dataflow/DataflowInputOutput.h
--- a/Engine/Source/Runtime/Experimental/Dataflow/Core/Public/Dataflow/DataflowInputOutput.h
+++ b/Engine/Source/Runtime/Experimental/Dataflow/Core/Public/Dataflow/DataflowInputOutput.h
@@ -236,6 +236,31 @@ public:
 	/** Set a null value : this means that the connected input will get a default value and that this output will not be set re-evaluate next time */
 	DATAFLOWCORE_API void SetNullValue(UE::Dataflow::FContext& Context) const;
 
+	/*
+	* Read value returns the cache value on the output without causing an evaluation of the corresponding node
+	* as a result it does not cause a cascading evaluation of the graph
+	* if there's no cached value this will return the default value passed as a parameter 
+	* @see GetValue
+	*/
+	template<class T>
+	const T& ReadValue(UE::Dataflow::FContext& Context, const T& Default) const
+	{
+		if (HasFrozenValue())
+		{
+			return GetFrozenValue(Default);
+		}
+		if (Context.HasData(CacheKey()))
+		{
+			return Context.GetData(CacheKey(), Property, Default);
+		}
+		return Default;
+	}
+
+	/*
+	* Get most updated value of the output
+	* if the value is cache or frozen return it otherwise evaluates the node with potentially cascading evaluation of the graph 
+	* @see GetValue
+	*/
 	template<class T>
 	const T& GetValue(UE::Dataflow::FContext& Context, const T& Default) const
 	{
