Restore an analytic value I accidently deleted in CL 39630837.

#rnx

diff --git a/Engine/Source/Runtime/Experimental/IoStore/OnDemand/Private/Statistics.cpp b/Engine/Source/Runtime/Experimental/IoStore/OnDemand/Private/Statistics.cpp
--- a/Engine/Source/Runtime/Experimental/IoStore/OnDemand/Private/Statistics.cpp
+++ b/Engine/Source/Runtime/Experimental/IoStore/OnDemand/Private/Statistics.cpp
@@ -54,6 +54,137 @@ static FAutoConsoleVariableRef CVar_DisplayOnScreenStatistics(
 static int32 BytesToApproxMB(uint64 Bytes) { return int32(Bytes >> 20); }
 static int32 BytesToApproxKB(uint64 Bytes) { return int32(Bytes >> 10); }
 
+/**
+ * Code taken from SummarizeTraceCommandlet.cpp pending discussion on moving it
+ * somewhere for general use.
+ * Currently not thread safe!
+ */
+class FIncrementalVariance
+{
+public:
+	FIncrementalVariance()
+		: Count(0)
+		, Mean(0.0)
+		, VarianceAccumulator(0.0)
+	{
+
+	}
+
+	uint64 GetCount() const
+	{
+		return Count;
+	}
+
+	double GetMean() const
+	{
+		return Mean;
+	}
+
+	/**
+	* Compute the variance given Welford's accumulator and the overall count
+	*
+	* @return The variance in sample units squared
+	*/
+	double GetVariance() const
+	{
+		double Result = 0.0;
+
+		if (Count > 1)
+		{
+			// Welford's final step, dependent on sample count
+			Result = VarianceAccumulator / double(Count - 1);
+		}
+
+		return Result;
+	}
+
+	/**
+	* Compute the standard deviation given Welford's accumulator and the overall count
+	*
+	* @return The standard deviation in sample units
+	*/
+	double GetDeviation() const
+	{
+		double Result = 0.0;
+
+		if (Count > 1)
+		{
+			// Welford's final step, dependent on sample count
+			double DeviationSqrd = VarianceAccumulator / double(Count - 1);
+
+			// stddev is sqrt of variance, to restore to units (vs. units squared)
+			Result = sqrt(DeviationSqrd);
+		}
+
+		return Result;
+	}
+
+	/**
+	* Perform an increment of work for Welford's variance, from which we can compute variation and standard deviation
+	*
+	* @param InSample	The new sample value to operate on
+	*/
+	void Increment(const double InSample)
+	{
+		Count++;
+		const double OldMean = Mean;
+		Mean += ((InSample - Mean) / double(Count));
+		VarianceAccumulator += ((InSample - Mean) * (InSample - OldMean));
+	}
+
+	/**
+	* Merge with another IncrementalVariance series in progress
+	*
+	* @param Other	The other variance incremented from another mutually exclusive population of analogous data.
+	*/
+	void Merge(const FIncrementalVariance& Other)
+	{
+		// empty other, nothing to do
+		if (Other.Count == 0)
+		{
+			return;
+		}
+
+		// empty this, just copy other
+		if (Count == 0)
+		{
+			Count = Other.Count;
+			Mean = Other.Mean;
+			VarianceAccumulator = Other.VarianceAccumulator;
+			return;
+		}
+
+		const double TotalPopulation = static_cast<double>(Count + Other.Count);
+		const double MeanDifference = Mean - Other.Mean;
+		const double A = (double(Count - 1) * GetVariance()) + (double(Other.Count - 1) * Other.GetVariance());
+		const double B = (MeanDifference) * (MeanDifference) * (double(Count) * double(Other.Count) / TotalPopulation);
+		const double MergedVariance = (A + B) / (TotalPopulation - 1);
+
+		const uint64 NewCount = Count + Other.Count;
+		const double NewMean = ((Mean * double(Count)) + (Other.Mean * double(Other.Count))) / double(NewCount);
+		const double NewVarianceAccumulator = MergedVariance * double(NewCount - 1);
+
+		Count = NewCount;
+		Mean = NewMean;
+		VarianceAccumulator = NewVarianceAccumulator;
+	}
+
+	/**
+	* Reset state back to initialized.
+	*/
+	void Reset()
+	{
+		Count = 0;
+		Mean = 0.0;
+		VarianceAccumulator = 0.0;
+	}
+
+private:
+	uint64 Count;
+	double Mean;
+	double VarianceAccumulator;
+};
+
 class FDeltaTracking
 {
 public:
@@ -163,6 +294,7 @@ int64					GHttpHistoryBytes[GHttpHistoryCount] = {};
 int64					GHttpHistoryTotalDuration = 0;
 int64					GHttpHistoryTotalBytes = 0;
 int64 					GHttpHistoryIndex = 0;
+FIncrementalVariance	GHttpAvgDuration; // Duration of the http requests, in milliseconds
 static uint32			GHttpCdnCacheHit = 0;
 static uint32			GHttpCdnCacheMiss = 0;
 static uint32			GHttpCdnCacheUnknown = 0;
@@ -418,6 +550,8 @@ void FOnDemandIoBackendStats::ReportEndPointAnalytics(TArray<FAnalyticsEventAttr
 			,TRACK_DELTA("IasHttpRetryCount", GHttpRetryCount.Get())
 			,TRACK_DELTA("IasHttpGetCount", GHttpGetCount.Get())
 			,TRACK_DELTA("IasHttpDownloadedBytes", GHttpDownloadedBytes.Get())
+			,TEXT("IasHttpDurationMeanAvg"), GHttpAvgDuration.GetMean()
+			,TEXT("IasHttpDurationStdDev"), GHttpAvgDuration.GetDeviation()
 		);
 	}
 
@@ -444,6 +578,9 @@ void FOnDemandIoBackendStats::ReportEndPointAnalytics(TArray<FAnalyticsEventAttr
 			,TRACK_DELTA("IasCacheReadBytes", GCacheReadBytes.Get())
 			,TRACK_DELTA("IasCacheRejectBytes", GCacheRejectBytes.Get())
 		);
+
+		// These values we can just reset as they are only being used with analytics
+		GHttpAvgDuration.Reset();
 	}
 
 #undef TRACK_DELTA
@@ -478,7 +615,8 @@ private:
 		Http.RetryCount.Stop();
 		Http.GetCount.Stop();
 		Http.DownloadedBytes.Stop();
-
+		Http.TotalDuration.Stop();
+		
 		Cache.ErrorCount.Stop();
 		Cache.DecodeErrorCount.Stop();
 		Cache.GetCount.Stop();
@@ -504,6 +642,7 @@ private:
 				, TEXT("IasHttpRetryCount"), Http.RetryCount.GetValue()
 				, TEXT("IasHttpGetCount"), Http.GetCount.GetValue()
 				, TEXT("IasHttpDownloadedBytes"), Http.DownloadedBytes.GetValue()
+				, TEXT("IasHttpTotalDuration"), Http.TotalDuration.GetValue()
 			);
 		}
 
@@ -592,6 +731,8 @@ private:
 		FTrackedValue<FCounterInt, GHttpRetryCount>				RetryCount;
 		FTrackedValue<FCounterInt, GHttpGetCount>				GetCount;
 		FTrackedValue<FCounterInt, GHttpDownloadedBytes>		DownloadedBytes;
+		FTrackedValue<int64, GHttpDurationMsSum>				TotalDuration;
+		
 	} Http;
 
 	/** Cache Stats */
@@ -755,6 +896,8 @@ void FOnDemandIoBackendStats::OnHttpGet(uint64 SizeBytes, uint64 DurationMs)
 	}
 
 	GHttpHistoryIndex = (GHttpHistoryIndex + 1) % GHttpHistoryCount;
+
+	GHttpAvgDuration.Increment(static_cast<double>(DurationMs));
 }
 
 void FOnDemandIoBackendStats::OnHttpCancel()
