Restore IAD http logging

#rnx

-During recent refactors the verbostiy of logging for FMultiEndpointHttpClient::Log was accidently hard coded to VeryVerbose, which is the verbosity that IAS commonly uses but IAD used Log.
- FMultiEndpointHttpClientConfig now accepts a ElogVerbosity (defaulting to Log) which IAS overrides to VeryVerbose.
- Now both IAS and IAD are restored to their previous behavior.

diff --git a/Engine/Source/Runtime/Experimental/IoStore/OnDemand/Private/OnDemandHttpClient.cpp b/Engine/Source/Runtime/Experimental/IoStore/OnDemand/Private/OnDemandHttpClient.cpp
--- a/Engine/Source/Runtime/Experimental/IoStore/OnDemand/Private/OnDemandHttpClient.cpp
+++ b/Engine/Source/Runtime/Experimental/IoStore/OnDemand/Private/OnDemandHttpClient.cpp
@@ -414,7 +414,7 @@ FMultiEndpointHttpClient::FConnection* FMultiEndpointHttpClient::FindConnection(
 void FMultiEndpointHttpClient::Log(const FRequest& Request) const
 {
 #if !NO_LOGGING
-	if (Config.LogCategory == nullptr || Config.LogCategory->IsSuppressed(ELogVerbosity::VeryVerbose))
+	if (Config.LogCategory == nullptr || Config.LogCategory->IsSuppressed(Config.LogVerbosity))
 	{
 		return;
 	}
@@ -422,7 +422,7 @@ void FMultiEndpointHttpClient::Log(const FRequest& Request) const
 	const uint64 DurationMs = (uint64)FPlatformTime::ToMilliseconds64(FPlatformTime::Cycles64() - Request.StartTime);
 	const uint64 Size = Request.Body.GetSize() >> 10;
 
-	FMsg::Logf(__FILE__, __LINE__, Config.LogCategory->GetCategoryName(), ELogVerbosity::VeryVerbose,
+	FMsg::Logf(__FILE__, __LINE__, Config.LogCategory->GetCategoryName(), Config.LogVerbosity,
 		TEXT("http-%3u: %5" UINT64_FMT "ms %5" UINT64_FMT "KiB %s (Attempt %d/%d)"),
 		Request.StatusCode,
 		DurationMs,
diff --git a/Engine/Source/Runtime/Experimental/IoStore/OnDemand/Private/OnDemandHttpClient.h b/Engine/Source/Runtime/Experimental/IoStore/OnDemand/Private/OnDemandHttpClient.h
--- a/Engine/Source/Runtime/Experimental/IoStore/OnDemand/Private/OnDemandHttpClient.h
+++ b/Engine/Source/Runtime/Experimental/IoStore/OnDemand/Private/OnDemandHttpClient.h
@@ -46,6 +46,7 @@ struct FMultiEndpointHttpClientConfig
 	bool						bAllowChunkedTransfer = true;
 	/** Logging will be disabled if this is set to a nullptr, it is up to the calling system to assign a LogCategory */
 	const UE_LOG_CATEGORY_TYPE*	LogCategory = nullptr;
+	ELogVerbosity::Type			LogVerbosity = ELogVerbosity::Log;
 };
 
 struct FMultiEndpointHttpClientResponse
diff --git a/Engine/Source/Runtime/Experimental/IoStore/OnDemand/Private/OnDemandIoDispatcherBackend.cpp b/Engine/Source/Runtime/Experimental/IoStore/OnDemand/Private/OnDemandIoDispatcherBackend.cpp
--- a/Engine/Source/Runtime/Experimental/IoStore/OnDemand/Private/OnDemandIoDispatcherBackend.cpp
+++ b/Engine/Source/Runtime/Experimental/IoStore/OnDemand/Private/OnDemandIoDispatcherBackend.cpp
@@ -1645,7 +1645,8 @@ uint32 FOnDemandIoBackend::Run()
 			.TimeoutMs = GIasHttpFailTimeOutMs,
 			.bEnableThreadSafetyChecks = true,
 			.bAllowChunkedTransfer = GIasHttpAllowChunkedXfer,
-			.LogCategory = &LogIas
+			.LogCategory = &LogIas,
+			.LogVerbosity = ELogVerbosity::VeryVerbose
 		});
 	check(HttpClient.IsValid());
 
