Dataflow - Added CVar to disable node Freezing until the bugs are ironed out.
Added missing UE:: in front of Dataflow namespace.

#rnx

diff --git a/Engine/Plugins/Experimental/Dataflow/Source/DataflowEditor/Private/Dataflow/DataflowSNode.cpp b/Engine/Plugins/Experimental/Dataflow/Source/DataflowEditor/Private/Dataflow/DataflowSNode.cpp
--- a/Engine/Plugins/Experimental/Dataflow/Source/DataflowEditor/Private/Dataflow/DataflowSNode.cpp
+++ b/Engine/Plugins/Experimental/Dataflow/Source/DataflowEditor/Private/Dataflow/DataflowSNode.cpp
@@ -26,7 +26,7 @@
 
 #define LOCTEXT_NAMESPACE "SDataflowEdNode"
 
-namespace Dataflow::Private
+namespace UE::Dataflow::Private
 {
 	static bool bShowPerfData = true;
 	static FAutoConsoleVariableRef CVarChaosBoxCalcBoundsISPCEnabled(TEXT("p.Dataflow.ShowPerfData"), bShowPerfData, TEXT("Whether to display perf data on top of dataflow nodes"));
@@ -234,7 +234,7 @@ TArray<FOverlayWidgetInfo> SDataflowEdNode::GetOverlayWidgets(bool bSelected, co
 			Widgets.Add(FreezeImageInfo);
 		}
 
-		if (Dataflow::Private::bShowPerfData)
+		if (UE::Dataflow::Private::bShowPerfData)
 		{
 			if (TSharedPtr<UE::Dataflow::FContext> Context = DataflowInterface->GetDataflowContext())
 			{
diff --git a/Engine/Source/Runtime/Experimental/Dataflow/Core/Private/Dataflow/DataflowInputOutput.cpp b/Engine/Source/Runtime/Experimental/Dataflow/Core/Private/Dataflow/DataflowInputOutput.cpp
--- a/Engine/Source/Runtime/Experimental/Dataflow/Core/Private/Dataflow/DataflowInputOutput.cpp
+++ b/Engine/Source/Runtime/Experimental/Dataflow/Core/Private/Dataflow/DataflowInputOutput.cpp
@@ -9,6 +9,11 @@
 
 #include UE_INLINE_GENERATED_CPP_BY_NAME(DataflowInputOutput)
 
+namespace UE::Dataflow::Private
+{
+	static bool bDisableFreezingState = true;
+	static FAutoConsoleVariableRef CVar(TEXT("p.Dataflow.DisableFreezingState"), bDisableFreezingState, TEXT("Disable the freezing state of Dataflow nodes. Useful for recovering from crashes caused by frozen nodes."));
+}
 
 FDataflowInput FDataflowInput::NoOpInput = FDataflowInput();
 FDataflowOutput FDataflowOutput::NoOpOutput = FDataflowOutput();
@@ -336,7 +341,7 @@ void FDataflowOutput::Freeze(UE::Dataflow::FContext& Context, FInstancedProperty
 
 bool FDataflowOutput::HasFrozenValue() const
 {
-	return GetOwningNode() && GetOwningNode()->IsFrozen();
+	return GetOwningNode() && GetOwningNode()->IsFrozen() && !UE::Dataflow::Private::bDisableFreezingState;
 }
 
 const uint8* FDataflowOutput::GetFrozenPropertyValue(const uint8* Default) const
