StringBuilder: Deprecate the const ToString() and operator*() because they can mutate the string

diff --git a/Engine/Source/Developer/AssetTools/Private/AssetHeaderPatcher.cpp b/Engine/Source/Developer/AssetTools/Private/AssetHeaderPatcher.cpp
--- a/Engine/Source/Developer/AssetTools/Private/AssetHeaderPatcher.cpp
+++ b/Engine/Source/Developer/AssetTools/Private/AssetHeaderPatcher.cpp
@@ -3510,7 +3510,7 @@ void FAssetHeaderPatcherInner::DumpState(FStringView OutputDirectory)
 							Builder.Append(TEXT("\t"));
 						}
 					}
-					const TCHAR* ToString() const
+					const TCHAR* ToString()
 					{
 						return Builder.ToString();
 					}
diff --git a/Engine/Source/Developer/AssetTools/Private/AssetRenameManager.cpp b/Engine/Source/Developer/AssetTools/Private/AssetRenameManager.cpp
--- a/Engine/Source/Developer/AssetTools/Private/AssetRenameManager.cpp
+++ b/Engine/Source/Developer/AssetTools/Private/AssetRenameManager.cpp
@@ -1444,7 +1444,7 @@ void FAssetRenameManager::SetupPublicAssets(TArray<FAssetRenameDataWithReference
 		FString ReferencingAssetInDifferentMountPoint;
 		if (RenameData.bCreateRedirector)
 		{
-			ReferencingAssetInDifferentMountPoint = OldPackageName.ToString();
+			ReferencingAssetInDifferentMountPoint = OldPackageName;
 		}
 		if (ReferencingAssetInDifferentMountPoint.IsEmpty())
 		{
@@ -1454,7 +1454,7 @@ void FAssetRenameManager::SetupPublicAssets(TArray<FAssetRenameDataWithReference
 				const FStringView ReferencingPackagMountPoint = FPathViews::GetMountPointNameFromPath(ReferencingPackageName);
 				if (!ReferencingPackagMountPoint.Equals(NewPackageMountPoint, ESearchCase::IgnoreCase))
 				{
-					ReferencingAssetInDifferentMountPoint = ReferencingPackageName.ToString();
+					ReferencingAssetInDifferentMountPoint = ReferencingPackageName;
 					break;
 				}
 			}
diff --git a/Engine/Source/Developer/CollectionManager/Private/CollectionContainer.cpp b/Engine/Source/Developer/CollectionManager/Private/CollectionContainer.cpp
--- a/Engine/Source/Developer/CollectionManager/Private/CollectionContainer.cpp
+++ b/Engine/Source/Developer/CollectionManager/Private/CollectionContainer.cpp
@@ -2499,7 +2499,7 @@ void FCollectionContainer::LoadCollections()
 			const ECollectionShareType::Type ShareType = ECollectionShareType::Type(CacheIdx);
 			const bool bUseSCC = ShouldUseSCC(ShareType);
 			const FString& CollectionFolder = CollectionSource->GetCollectionFolder(static_cast<ECollectionShareType::Type>(CacheIdx));
-			const TStringBuilder<256> WildCard{InPlace, CollectionFolder, TEXTVIEW("/*."), CollectionExtension};
+			TStringBuilder<256> WildCard{InPlace, CollectionFolder, TEXTVIEW("/*."), CollectionExtension};
 
 			TArray<FString> Filenames;
 			IFileManager::Get().FindFiles(Filenames, *WildCard, true, false);
diff --git a/Engine/Source/Editor/UnrealEd/Private/Commandlets/GatherTextFromAssetsCommandlet.cpp b/Engine/Source/Editor/UnrealEd/Private/Commandlets/GatherTextFromAssetsCommandlet.cpp
--- a/Engine/Source/Editor/UnrealEd/Private/Commandlets/GatherTextFromAssetsCommandlet.cpp
+++ b/Engine/Source/Editor/UnrealEd/Private/Commandlets/GatherTextFromAssetsCommandlet.cpp
@@ -1151,7 +1151,7 @@ void UGatherTextFromAssetsCommandlet::ProcessAndRemoveCachedPackages(TMap<FName,
 		TMap<FString, FName> AssignedPackageLocalizationIds;
 		PackagesPendingGather.RemoveAll([this, &OutExternalActorsWithStaleOrMissingCaches, &AssignedPackageLocalizationIds, &NumPackagesProcessed, InitialNumPackagesToGather](const FPackagePendingGather& PackagePendingGather) -> bool
 		{
-			const FNameBuilder PackageNameStr(PackagePendingGather.PackageName);
+			FNameBuilder PackageNameStr(PackagePendingGather.PackageName);
 			const int32 CurrentPackageNum = ++NumPackagesProcessed;
 			const float PercentageComplete = static_cast<float>(CurrentPackageNum) / static_cast<float>(InitialNumPackagesToGather) * 100.0f;
 
@@ -1283,7 +1283,7 @@ bool UGatherTextFromAssetsCommandlet::LoadAndProcessUncachedPackages(TArray<FNam
 	while (PackagesPendingGather.Num() > 0)
 	{
 		const FPackagePendingGather PackagePendingGather = PackagesPendingGather.Pop(EAllowShrinking::No);
-		const FNameBuilder PackageNameStr(PackagePendingGather.PackageName);
+		FNameBuilder PackageNameStr(PackagePendingGather.PackageName);
 
 		const int32 CurrentPackageNum = ++NumPackagesProcessed;
 		const float PercentageComplete = static_cast<float>(CurrentPackageNum) / static_cast<float>(PackageCount) * 100.0f;
diff --git a/Engine/Source/Runtime/AssetRegistry/Private/AssetRegistry.cpp b/Engine/Source/Runtime/AssetRegistry/Private/AssetRegistry.cpp
--- a/Engine/Source/Runtime/AssetRegistry/Private/AssetRegistry.cpp
+++ b/Engine/Source/Runtime/AssetRegistry/Private/AssetRegistry.cpp
@@ -7043,7 +7043,7 @@ void FAssetRegistryImpl::DependencyDataGathered(TMultiMap<FName, FPackageDepende
 						DependsNode->SetIsReferencersSorted(ShouldSortReferencers());
 
 						// This was newly created, see if we need to read the script package Guid
-						const FNameBuilder DependencyPackageNameStr(DependencyPackageName);
+						FNameBuilder DependencyPackageNameStr(DependencyPackageName);
 
 						if (FPackageName::IsScriptPackage(DependencyPackageNameStr))
 						{
diff --git a/Engine/Source/Runtime/Core/Public/Misc/StringBuilder.h b/Engine/Source/Runtime/Core/Public/Misc/StringBuilder.h
--- a/Engine/Source/Runtime/Core/Public/Misc/StringBuilder.h
+++ b/Engine/Source/Runtime/Core/Public/Misc/StringBuilder.h
@@ -115,7 +115,19 @@ public:
 		return Base;
 	}
 
-	/** Returns a pointer to a null-terminated string that is valid until the builder is mutated. */
+	/**
+	 * Prefer operator*() for a pointer to a null-terminated string.
+	 *
+	 * Construct FString and other string types directly from the string builder: FString Str(Builder).
+	 */
+	inline const CharType* ToString() UE_LIFETIMEBOUND
+	{
+		EnsureNulTerminated();
+		return Base;
+	}
+
+	/** Prefer operator*() for a pointer to a null-terminated string. */
+	UE_DEPRECATED(5.6, "Call the non-const operator*().")
 	inline const CharType* ToString() const UE_LIFETIMEBOUND
 	{
 		EnsureNulTerminated();
@@ -123,9 +135,18 @@ public:
 	}
 
 	/** Returns a pointer to a null-terminated string that is valid until the builder is mutated. */
+	inline const CharType* operator*() UE_LIFETIMEBOUND
+	{
+		EnsureNulTerminated();
+		return Base;
+	}
+
+	/** Returns a pointer to a null-terminated string that is valid until the builder is mutated. */
+	UE_DEPRECATED(5.6, "Call the non-const operator*().")
 	inline const CharType* operator*() const UE_LIFETIMEBOUND
 	{
-		return ToString();
+		EnsureNulTerminated();
+		return Base;
 	}
 
 	/** Returns a view of the string that is valid until the builder is mutated. */
diff --git a/Engine/Source/Runtime/Experimental/Iris/Core/Private/Iris/ReplicationState/ReplicationStateDescriptorBuilder.cpp b/Engine/Source/Runtime/Experimental/Iris/Core/Private/Iris/ReplicationState/ReplicationStateDescriptorBuilder.cpp
--- a/Engine/Source/Runtime/Experimental/Iris/Core/Private/Iris/ReplicationState/ReplicationStateDescriptorBuilder.cpp
+++ b/Engine/Source/Runtime/Experimental/Iris/Core/Private/Iris/ReplicationState/ReplicationStateDescriptorBuilder.cpp
@@ -511,7 +511,8 @@ void FPropertyReplicationStateDescriptorBuilder::ValidateIsFullyPushModelClassDe
 	}
 
 	// A class and its descriptors is expected to be validated at most once. We want to report all of the failing descriptors because the user has explicitly requested us to.
-	ensureAlwaysMsgf(false, TEXT("ReplicationStateDescriptor %s for class %s does not have the HasPushBasedDirtiness trait. This prevents the class from being fully push based. See log for additional details."), ToCStr(Descriptor->DebugName), (Context.BuildParams.PathName ? Context.BuildParams.PathName->ToString() : TEXT("UNKNOWN")));
+	ensureAlwaysMsgf(false, TEXT("ReplicationStateDescriptor %s for class %s does not have the HasPushBasedDirtiness trait. This prevents the class from being fully push based. See log for additional details."),
+		ToCStr(Descriptor->DebugName), *WriteToString<128>(Context.BuildParams.PathName ? Context.BuildParams.PathName->ToView() : TEXTVIEW("UNKNOWN")));
 
 	// Try to figure out the cause of the lack of push based trait
 	if (!Members.IsEmpty())
@@ -541,7 +542,8 @@ void FPropertyReplicationStateDescriptorBuilder::ValidateIsFullyPushModelClassDe
 	}
 	else
 	{
-		UE_LOG_DESCRIPTORBUILDER_WARNING(TEXT("ReplicationStateDescriptor %s for class %s does not have the HasPushBasedDirtiness trait. No additional info."), ToCStr(Descriptor->DebugName), Context.BuildParams.PathName ? Context.BuildParams.PathName->ToString() : TEXT(""));
+		UE_LOG_DESCRIPTORBUILDER_WARNING(TEXT("ReplicationStateDescriptor %s for class %s does not have the HasPushBasedDirtiness trait. No additional info."),
+			ToCStr(Descriptor->DebugName), *WriteToString<128>(Context.BuildParams.PathName ? Context.BuildParams.PathName->ToView() : TEXTVIEW("")));
 	}
 }
 
@@ -1937,7 +1939,7 @@ FPropertyReplicationStateDescriptorBuilder::Build(const FString& StateName, FRep
 		UpperCaseNameForHash.ToUpperInline();
 		Descriptor->DescriptorIdentifier.Value = CityHash64WithSeed((const char*)UpperCaseNameForHash.GetCharArray().GetData(), sizeof(TCHAR) * UpperCaseNameForHash.GetCharArray().Num(), Seed);
 
-		UE_LOG_DESCRIPTORBUILDER(Verbose, TEXT("Generated identifier(0x%" UINT64_x_FMT ") From %s:%s"), Descriptor->DescriptorIdentifier.Value, PathLen ? BuildParams.PathName->ToString() : TEXT(""), *StateName);
+		UE_LOG_DESCRIPTORBUILDER(Verbose, TEXT("Generated identifier(0x%" UINT64_x_FMT ") From %s:%s"), Descriptor->DescriptorIdentifier.Value, *WriteToString<128>(PathLen ? BuildParams.PathName->ToView() : TEXTVIEW("")), *StateName);
 
 		if (Descriptor->InternalSize && BuildParams.DefaultStateSourceData)
 		{
