Dataflow : make context more general with the ability to change to be threaded at runtime if needed

diff --git a/Engine/Source/Runtime/Experimental/Dataflow/Core/Private/Dataflow/DataflowNodeParameters.cpp b/Engine/Source/Runtime/Experimental/Dataflow/Core/Private/Dataflow/DataflowNodeParameters.cpp
--- a/Engine/Source/Runtime/Experimental/Dataflow/Core/Private/Dataflow/DataflowNodeParameters.cpp
+++ b/Engine/Source/Runtime/Experimental/Dataflow/Core/Private/Dataflow/DataflowNodeParameters.cpp
@@ -91,6 +91,64 @@ namespace UE::Dataflow
 
 	//////////////////////////////////////////////////////////////////////
 
+	void FContext::SetThreaded(bool bValue)
+	{
+		if (bValue != IsThreaded())
+		{
+			if (bValue)
+			{
+				DataLock = MakeUnique<FCriticalSection>();
+			}
+			else
+			{
+				DataLock.Reset();
+			}
+		}
+	}
+
+	int32 FContext::GetKeys(TSet<FContextCacheKey>& InKeys) const
+	{ 
+		FScopedOptionalLock Lock(DataLock.Get());
+		return DataStore.GetKeys(InKeys);
+	}
+
+	void FContext::SetDataImpl(FContextCacheKey Key, TUniquePtr<FContextCacheElementBase>&& DataStoreEntry)
+	{
+		FScopedOptionalLock Lock(DataLock.Get());
+		if (IsThreaded())
+		{
+			// Threaded evaluation can only set an output once per context evaluation. Otherwise
+				// downstream nodes that are extracting the data will get corrupted store entries. 
+			const TUniquePtr<FContextCacheElementBase>* const CurrentData = DataStore.Find(Key);
+			if (!CurrentData || !(*CurrentData) || (*CurrentData)->GetTimestamp() < DataStoreEntry->GetTimestamp())
+			{
+				DataStore.Emplace(Key, MoveTemp(DataStoreEntry));
+			}
+		}
+		else
+		{
+			DataStore.Emplace(Key, MoveTemp(DataStoreEntry));
+		}
+	}
+
+	const TUniquePtr<FContextCacheElementBase>* FContext::GetDataImpl(FContextCacheKey Key) const
+	{
+		FScopedOptionalLock Lock(DataLock.Get());
+		return DataStore.Find(Key);
+	}
+
+	bool FContext::HasDataImpl(FContextCacheKey Key, FTimestamp InTimestamp) const
+	{
+		FScopedOptionalLock Lock(DataLock.Get());
+		return DataStore.Contains(Key) && DataStore[Key]->GetTimestamp() >= InTimestamp;
+	}
+
+	bool FContext::IsEmptyImpl() const
+	{
+		FScopedOptionalLock Lock(DataLock.Get());
+		return DataStore.IsEmpty();
+	}
+
 	FTimestamp FContext::GetTimestamp(FContextCacheKey Key) const
 	{
 		if (const TUniquePtr<FContextCacheElementBase>* const Cache = const_cast<FContext*>(this)->GetDataImpl(Key))
@@ -283,6 +341,27 @@ namespace UE::Dataflow
 		Context.PopFromCallstack(Connection);
 	}
 
+	void FContext::CheckIntrinsicInputs(const FDataflowOutput& Connection)
+	{
+#if WITH_EDITOR
+		//
+		// Check if all DataflowIntrinsic input(s) connected
+		// 
+		if (const FDataflowNode* Node = Connection.GetOwningNode())
+		{
+			TArray<FDataflowInput*> Inputs = Node->GetInputs();
+			for (const FDataflowInput* Input : Inputs)
+			{
+				if (Input && Input->IsRequired() && !Input->IsConnected())
+				{
+					const FString WarningMsg = FString::Printf(TEXT("Input %s must be connected"), *Input->GetName().ToString());
+					Warning(WarningMsg, Node);
+				}
+			}
+		}
+#endif // EDITOR_ONLY
+	}
+
 	void FContext::BeginContextEvaluation(const FDataflowNode * Node, const FDataflowOutput * Output)
 	{
 		if (Output)
@@ -307,17 +386,17 @@ namespace UE::Dataflow
 				//       The purpose of the node evaluation function is to evaluate outputs.
 				//       Therefore if a node has no outputs, then it shouldn't need any evaluation.
 
-				if (this->OnNodeBeginEvaluateMulticast.IsBound())
+				if (!IsThreaded() && OnNodeBeginEvaluateMulticast.IsBound())
 				{
-					this->OnNodeBeginEvaluateMulticast.Broadcast(Node, Output);
+					OnNodeBeginEvaluateMulticast.Broadcast(Node, Output);
 				}
 
 				UE_LOG(LogChaosDataflow, Verbose, TEXT("FDataflowNode::Evaluate(): Node [%s], Output [nullptr], NodeTimestamp [%lu]"), *Node->GetName().ToString(), Node->GetTimestamp().Value);
 				Node->Evaluate(*this, nullptr);
 
-				if (this->OnNodeFinishEvaluateMulticast.IsBound())
+				if (!IsThreaded() && OnNodeFinishEvaluateMulticast.IsBound())
 				{
-					this->OnNodeFinishEvaluateMulticast.Broadcast(Node, Output);
+					OnNodeFinishEvaluateMulticast.Broadcast(Node, Output);
 				}
 			}
 		}
@@ -327,59 +406,34 @@ namespace UE::Dataflow
 		}
 	}
 
-	void FContextSingle::Evaluate(const FDataflowNode* Node, const FDataflowOutput* Output)
+	void FContext::Evaluate(const FDataflowNode* Node, const FDataflowOutput* Output)
 	{
 		BeginContextEvaluation(Node, Output);
 	}
 
-	bool FContextSingle::Evaluate(const FDataflowOutput& Connection)
+	bool FContext::Evaluate(const FDataflowOutput& Connection)
 	{
-#if WITH_EDITOR
-		//
-		// Check if all DataflowIntrinsic input(s) connected
-		// 
-		const FDataflowNode* Node = Connection.GetOwningNode();
-		TArray<FDataflowInput*> Inputs = Node->GetInputs();
-		for (FDataflowInput* Input : Inputs)
+		CheckIntrinsicInputs(Connection);
+
+		if (!IsThreaded() && OnNodeBeginEvaluateMulticast.IsBound())
 		{
-			if (Input->IsRequired() && Input->GetConnection() == nullptr)
-			{
-				const FString InputNameStr = Input->GetName().ToString();
-				const FString WarningMsg = FString::Printf(TEXT("Missing connection to required input : %s"), *InputNameStr);
-				this->Warning(WarningMsg, Node);
-			}
+			OnNodeBeginEvaluateMulticast.Broadcast(Connection.GetOwningNode(), &Connection);
 		}
 
-		UE_LOG(LogChaosDataflow, VeryVerbose, TEXT("FContextSingle::Evaluate(): Node [%s], Output [%s]"), *Connection.GetOwningNode()->GetName().ToString(), *Connection.GetName().ToString());
-#endif // EDITOR_ONLY
-
-		if (OnNodeBeginEvaluateMulticast.IsBound())
+		UE_LOG(LogChaosDataflow, VeryVerbose, TEXT("FContext::Evaluate(): Node [%s], Output [%s]"), *Connection.GetOwningNode()->GetName().ToString(), *Connection.GetName().ToString());
+		bool RetValue = false;
 		{
-			this->OnNodeBeginEvaluateMulticast.Broadcast(Connection.GetOwningNode(), &Connection);
+			FScopedOptionalLock Lock(IsThreaded() ? Connection.OutputLock.Get() : nullptr);
+			RetValue = Connection.EvaluateImpl(*this);
 		}
-
-		const bool RetValue = Connection.EvaluateImpl(*this);
-
-		if (OnNodeFinishEvaluateMulticast.IsBound())
+		if (!IsThreaded() && OnNodeFinishEvaluateMulticast.IsBound())
 		{
-			this->OnNodeFinishEvaluateMulticast.Broadcast(Connection.GetOwningNode(), &Connection);
+			OnNodeFinishEvaluateMulticast.Broadcast(Connection.GetOwningNode(), &Connection);
 		}
 
 		return !GetNumErrors() && RetValue;
 	}
 
-	void FContextThreaded::Evaluate(const FDataflowNode* Node, const FDataflowOutput* Output)
-	{
-		BeginContextEvaluation(Node, Output);
-	}
-
-	bool FContextThreaded::Evaluate(const FDataflowOutput& Connection)
-	{
-		UE_LOG(LogChaosDataflow, VeryVerbose, TEXT("FContextThreaded::Evaluate(): Node [%s], Output [%s]"), *Connection.GetOwningNode()->GetName().ToString(), *Connection.GetName().ToString());
-		Connection.OutputLock->Lock(); ON_SCOPE_EXIT{ Connection.OutputLock->Unlock(); };
-		return Connection.EvaluateImpl(*this);
-	}
-
 	void FContextCache::Serialize(FArchive& Ar)
 	{
 		if (Ar.IsSaving())
@@ -446,7 +500,7 @@ namespace UE::Dataflow
 					NewElement->NodeGuid = NodeGuid;
 					NewElement->NodeHash = NodeHash;
 					NewElement->Timestamp = Timestamp;
-					this->Add(InKey, TUniquePtr<FContextCacheElementBase>(NewElement));
+					Add(InKey, TUniquePtr<FContextCacheElementBase>(NewElement));
 				}
 				DATAFLOW_OPTIONAL_BLOCK_READ_ELSE()
 				{
diff --git a/Engine/Source/Runtime/Experimental/Dataflow/Core/Public/Dataflow/DataflowNodeParameters.h b/Engine/Source/Runtime/Experimental/Dataflow/Core/Public/Dataflow/DataflowNodeParameters.h
--- a/Engine/Source/Runtime/Experimental/Dataflow/Core/Public/Dataflow/DataflowNodeParameters.h
+++ b/Engine/Source/Runtime/Experimental/Dataflow/Core/Public/Dataflow/DataflowNodeParameters.h
@@ -83,6 +83,42 @@ namespace UE::Dataflow
 		TArray<FEntry> Callstack;
 	};
 
+	//////////////////////////////////////////////////////////////////////////////////
+	class FScopedOptionalLock final
+	{
+	public:
+		UE_NODISCARD_CTOR explicit FScopedOptionalLock(FCriticalSection* InCriticalSection)
+			: CriticalSection(InCriticalSection)
+		{
+			if (CriticalSection)
+			{
+				CriticalSection->Lock();
+			}
+		}
+
+		~FScopedOptionalLock()
+		{
+			Unlock();
+		}
+
+		void Unlock()
+		{
+			if (CriticalSection)
+			{
+				CriticalSection->Unlock();
+				CriticalSection = nullptr;
+			}
+		}
+
+	private:
+		FScopedOptionalLock() = delete;
+		FScopedOptionalLock(const FScopedOptionalLock& InScopeLock) = delete;
+		FScopedOptionalLock& operator=(const FScopedOptionalLock& InScopeLock) = delete;
+
+	private:
+		FCriticalSection* CriticalSection;
+	};
+
 	//////////////////////////////////////////////////////////////////////////////////
 
 	/*
@@ -98,6 +134,7 @@ namespace UE::Dataflow
 		FContext& operator=(const FContext&) = delete;
 
 		FContextCache DataStore;
+		TUniquePtr<FCriticalSection> DataLock;
 
 	public:
 		FContext() = default;
@@ -109,8 +146,6 @@ namespace UE::Dataflow
 
 		virtual FName GetType() const { return FContext::StaticType(); }
 
-		virtual int32 GetKeys(TSet<FContextCacheKey>& InKeys) const { return DataStore.GetKeys(InKeys); }
-
 		template<class T>
 		const T* AsType() const
 		{
@@ -131,8 +166,18 @@ namespace UE::Dataflow
 			return nullptr;
 		}
 
-		virtual void SetDataImpl(FContextCacheKey Key, TUniquePtr<FContextCacheElementBase>&& DataStoreEntry) = 0;
-		
+		DATAFLOWCORE_API bool IsThreaded() const { return DataLock.IsValid(); }
+		DATAFLOWCORE_API void SetThreaded(bool bValue);
+
+		DATAFLOWCORE_API virtual void SetDataImpl(FContextCacheKey Key, TUniquePtr<FContextCacheElementBase>&& DataStoreEntry);
+		DATAFLOWCORE_API virtual const TUniquePtr<FContextCacheElementBase>* GetDataImpl(FContextCacheKey Key) const;
+		DATAFLOWCORE_API virtual bool HasDataImpl(FContextCacheKey Key, FTimestamp InTimestamp = FTimestamp::Invalid) const;
+		DATAFLOWCORE_API virtual bool IsEmptyImpl() const;
+		DATAFLOWCORE_API virtual int32 GetKeys(TSet<FContextCacheKey>& InKeys) const;
+
+		DATAFLOWCORE_API virtual void Evaluate(const FDataflowNode* Node, const FDataflowOutput* Output);
+		DATAFLOWCORE_API virtual bool Evaluate(const FDataflowOutput& Connection);
+
 		template<typename T>
 		void SetData(FContextCacheKey InKey, const FProperty* InProperty, T&& InValue, const FGuid& InNodeGuid, uint32 InNodeHash, const FTimestamp& InTimestamp)
 		{
@@ -187,8 +232,6 @@ namespace UE::Dataflow
 			SetDataImpl(InKey, MoveTemp(CacheNullElement));
 		}
 
-		virtual const TUniquePtr<FContextCacheElementBase>* GetDataImpl(FContextCacheKey Key) const = 0;
-
 		template<class T>
 		const T& GetData(FContextCacheKey Key, const FProperty* InProperty, const T& Default = T()) const
 		{
@@ -208,16 +251,12 @@ namespace UE::Dataflow
 			return nullptr;
 		}
 
-		virtual bool HasDataImpl(FContextCacheKey Key, FTimestamp InTimestamp = FTimestamp::Invalid) const = 0;
-		
 		bool HasData(FContextCacheKey Key, FTimestamp InTimestamp = FTimestamp::Invalid) const
 		{
 			FContextCacheKey IntKey = (FContextCacheKey)Key;
 			return HasDataImpl(Key, InTimestamp);
 		}
 
-		virtual bool IsEmptyImpl() const = 0;
-
 		bool IsEmpty() const
 		{
 			return IsEmptyImpl();
@@ -243,9 +282,6 @@ namespace UE::Dataflow
 
 		DATAFLOWCORE_API FTimestamp GetTimestamp(FContextCacheKey Key) const;
 
-		virtual void Evaluate(const FDataflowNode* Node, const FDataflowOutput* Output) = 0;
-		virtual bool Evaluate(const FDataflowOutput& Connection) = 0;
-
 		DATAFLOWCORE_API void PushToCallstack(const FDataflowConnection* Connection);
 		DATAFLOWCORE_API void PopFromCallstack(const FDataflowConnection* Connection);
 		DATAFLOWCORE_API bool IsInCallstack(const FDataflowConnection* Connection) const;
@@ -288,6 +324,9 @@ namespace UE::Dataflow
 		DATAFLOWCORE_API void BeginContextEvaluation(const FDataflowNode* Node, const FDataflowOutput* Output);
 
 	private:
+
+		void CheckIntrinsicInputs(const FDataflowOutput& Connection);
+
 #if DATAFLOW_EDITOR_EVALUATION
 		FContextCallstack Callstack;
 		FContextPerfData PerfData;
@@ -320,87 +359,20 @@ namespace UE::Dataflow
 
 	class FContextSingle : public FContext
 	{
-
 	public:
 		DATAFLOW_CONTEXT_INTERNAL(FContext, FContextSingle);
-
-		FContextSingle() = default;
-
-		virtual void SetDataImpl(FContextCacheKey Key, TUniquePtr<FContextCacheElementBase>&& DataStoreEntry) override
-		{
-			DataStore.Emplace(Key, MoveTemp(DataStoreEntry));
-		}
-
-		virtual const TUniquePtr<FContextCacheElementBase>* GetDataImpl(FContextCacheKey Key) const override
-		{
-			return DataStore.Find(Key);
-		}
-
-		virtual bool HasDataImpl(FContextCacheKey Key, FTimestamp InTimestamp = FTimestamp::Invalid) const override
-		{
-			return DataStore.Contains(Key) && (DataStore[Key]->GetTimestamp() >= InTimestamp);
-		}
-
-		virtual bool IsEmptyImpl() const override
-		{
-			return DataStore.IsEmpty();
-		}
-
-		DATAFLOWCORE_API virtual void Evaluate(const FDataflowNode* Node, const FDataflowOutput* Output) override;
-		DATAFLOWCORE_API virtual bool Evaluate(const FDataflowOutput& Connection) override;
-
 	};
 	
 	class FContextThreaded : public FContext
 	{
-		TSharedPtr<FCriticalSection> CacheLock;
-
 	public:
 		DATAFLOW_CONTEXT_INTERNAL(FContext, FContextThreaded);
 
-
 		FContextThreaded()
 			: FContext()
 		{
-			CacheLock = MakeShared<FCriticalSection>();
+			SetThreaded(true);
 		}
-
-		virtual void SetDataImpl(FContextCacheKey Key, TUniquePtr<FContextCacheElementBase>&& DataStoreEntry) override
-		{
-			CacheLock->Lock(); ON_SCOPE_EXIT { CacheLock->Unlock(); };
-			
-
-			// Threaded evaluation can only set an output once per context evaluation. Otherwise
-			// downstream nodes that are extracting the data will get currupted store entries. 
-			TUniquePtr<FContextCacheElementBase>* CurrentData = DataStore.Find(Key);
-			if (!CurrentData || !(*CurrentData) || (*CurrentData)->GetTimestamp() < DataStoreEntry->GetTimestamp())
-			{
-				DataStore.Emplace(Key, MoveTemp(DataStoreEntry));
-			}
-		}
-
-		virtual const TUniquePtr<FContextCacheElementBase>* GetDataImpl(FContextCacheKey Key) const override
-		{
-			CacheLock->Lock(); ON_SCOPE_EXIT { CacheLock->Unlock(); };
-			
-			return DataStore.Find(Key);
-		}
-
-		virtual bool HasDataImpl(FContextCacheKey Key, FTimestamp InTimestamp = FTimestamp::Invalid) const override
-		{
-			CacheLock->Lock(); ON_SCOPE_EXIT { CacheLock->Unlock(); };
-			
-			return DataStore.Contains(Key) && DataStore[Key]->GetTimestamp() >= InTimestamp;
-		}
-
-		virtual bool IsEmptyImpl() const override
-		{
-			return DataStore.IsEmpty();
-		}
-
-		DATAFLOWCORE_API virtual void Evaluate(const FDataflowNode* Node, const FDataflowOutput* Output) override;
-		DATAFLOWCORE_API virtual bool Evaluate(const FDataflowOutput& Connection) override;
-
 	};
 }
 
