Dataflow : move DataflowDynamicConnections structure out in its own file

diff --git a/Engine/Source/Runtime/Experimental/Dataflow/Engine/Private/Dataflow/DataflowDynamicConnections.cpp b/Engine/Source/Runtime/Experimental/Dataflow/Engine/Private/Dataflow/DataflowDynamicConnections.cpp
new file mode 100644
--- /dev/null
+++ b/Engine/Source/Runtime/Experimental/Dataflow/Engine/Private/Dataflow/DataflowDynamicConnections.cpp
@@ -0,0 +1,252 @@
+// Copyright Epic Games, Inc. All Rights Reserved.
+
+#include "Dataflow/DataflowDynamicConnections.h"
+
+#include "Dataflow/DataflowObject.h"
+
+FDataflowDynamicConnections::FDataflowDynamicConnections()
+	: PinDirection(UE::Dataflow::FPin::EDirection::NONE)
+	, OwnerInterface(nullptr)
+	, DataflowAssetWeakPtr(nullptr)
+{}
+
+FDataflowDynamicConnections::FDataflowDynamicConnections(UE::Dataflow::FPin::EDirection InPinDirection, IOwnerInterface* InOwnerInterface, UDataflow* InDataflowAsset)
+	: PinDirection(InPinDirection)
+	, OwnerInterface(InOwnerInterface)
+	, DataflowAssetWeakPtr(InDataflowAsset)
+{}
+
+
+FDataflowDynamicConnections::FConnectionReference FDataflowDynamicConnections::GetConnectionReference(int32 Index) const
+{
+	return { &DynamicProperties[Index], Index, &DynamicProperties };
+}
+
+void FDataflowDynamicConnections::Refresh()
+{
+	using namespace UE::Dataflow;
+
+	if (!OwnerInterface)
+	{
+		return;
+	}
+
+	FDataflowNode* OwnerNode = OwnerInterface->GetOwner(this);
+	if (!OwnerNode)
+	{
+		return;
+	}
+
+	TStrongObjectPtr<UDataflow> DataflowAsset = DataflowAssetWeakPtr.Pin();
+	if (!DataflowAsset.IsValid())
+	{
+		return;
+	}
+
+	TSharedPtr<UE::Dataflow::FGraph, ESPMode::ThreadSafe> DataflowGraph = DataflowAsset->GetDataflow();
+	if (!DataflowGraph.IsValid())
+	{
+		return;
+	}
+
+	// save the remote connections from the dynamic connections
+	using FRemoteConnections = TArray<FDataflowConnection*>;
+	TMap<FGuid, FRemoteConnections> RemoteConnectionsByPropertyId;
+
+	for (FDataflowConnection* DynamicConnection : GetDynamicConnections())
+	{
+		if (DynamicConnection)
+		{
+			const FGuid* PropertyId = ConnectionNameToPropertyId.Find(DynamicConnection->GetName());
+			if (PropertyId)
+			{
+				FRemoteConnections& RemoteConnections = RemoteConnectionsByPropertyId.Emplace(*PropertyId);
+				DynamicConnection->GetConnections(RemoteConnections);
+				DataflowGraph->ClearConnections(DynamicConnection);
+			}
+		}
+	}
+	// clear all the dynamic connections
+	ClearDynamicConnections();
+
+	ConnectionNameToPropertyId.Reset();
+
+	// Go through all the properties of the property bag recreate the output and reconnect them if needed
+	const TConnectionReference<FDataflowAllTypes> NullReference(nullptr);
+	const FInstancedPropertyBag& PropertyBag = OwnerInterface->GetPropertyBag(this);
+	if (const UPropertyBag* PropertyBagStruct = PropertyBag.GetPropertyBagStruct())
+	{
+		for (const FPropertyBagPropertyDesc& PropertyDesc : PropertyBagStruct->GetPropertyDescs())
+		{
+			// we can only add the output if it is a supported type
+			if (IsSupportedType(PropertyDesc))
+			{
+				const FName CppType = GetCppTypeFromPropertyDesc(PropertyDesc);
+				if (FDataflowConnection* NewConnection = AddNewConnectionFromPropertyDesc(PropertyDesc))
+				{
+					// now try to reconnect what was connected before
+					if (FRemoteConnections* RemoteConnections = RemoteConnectionsByPropertyId.Find(PropertyDesc.ID))
+					{
+						for (FDataflowConnection* RemoteConnection : *RemoteConnections)
+						{
+							DataflowGraph->Connect(NewConnection, RemoteConnection);
+						}
+					}
+				}
+			}
+		}
+		// Todo : need to also remove the one that are no longer necessary
+		// ( need to disconnect if needed ? or keep the one that are still connected ?)
+		OwnerNode->Invalidate();
+
+		DataflowAsset->RefreshEdNodeByGuid(OwnerNode->GetGuid());
+	}
+}
+
+FDataflowConnection* FDataflowDynamicConnections::AddNewConnectionFromPropertyDesc(const FPropertyBagPropertyDesc& PropertyDesc)
+{
+	using namespace UE::Dataflow;
+
+	if (FDataflowNode* OwnerNode = OwnerInterface ? OwnerInterface->GetOwner(this) : nullptr)
+	{
+		const int32 Index = DynamicProperties.AddDefaulted();
+		TConnectionReference<FDataflowAllTypes> ConnectionReference = GetConnectionReference(Index);
+		if (FDataflowConnection* NewConnection = CreateConnection(ConnectionReference))
+		{
+			SetConnectionTypeFromPropertyDesc(*NewConnection, PropertyDesc);
+			if (PinDirection == UE::Dataflow::FPin::EDirection::INPUT)
+			{
+				OwnerNode->TryRenameInput(ConnectionReference, PropertyDesc.Name);
+			}
+			else if (PinDirection == UE::Dataflow::FPin::EDirection::OUTPUT)
+			{
+				OwnerNode->TryRenameOutput(ConnectionReference, PropertyDesc.Name);
+			}
+			NewConnection->LockType();
+
+			ConnectionNameToPropertyId.Add(NewConnection->GetName(), PropertyDesc.ID);
+			return NewConnection;
+		}
+	}
+	return nullptr;
+}
+
+bool FDataflowDynamicConnections::IsSupportedType(const FPropertyBagPropertyDesc& Desc)
+{
+	switch (Desc.ValueType)
+	{
+	case EPropertyBagPropertyType::Bool:
+	case EPropertyBagPropertyType::Byte:
+	case EPropertyBagPropertyType::Int32:
+	case EPropertyBagPropertyType::Int64:
+	case EPropertyBagPropertyType::Float:
+	case EPropertyBagPropertyType::Double:
+	case EPropertyBagPropertyType::Name:
+	case EPropertyBagPropertyType::String:
+	case EPropertyBagPropertyType::Text:
+	case EPropertyBagPropertyType::Object:
+	case EPropertyBagPropertyType::Struct:
+		return true;
+
+		// unsupported type for now 
+	case EPropertyBagPropertyType::Enum:
+	case EPropertyBagPropertyType::SoftObject:
+	case EPropertyBagPropertyType::Class:
+	case EPropertyBagPropertyType::SoftClass:
+	case EPropertyBagPropertyType::UInt32:	// Type not fully supported at UI, will work with restrictions to type editing
+	case EPropertyBagPropertyType::UInt64: // Type not fully supported at UI, will work with restrictions to type editing
+	default:
+		return false;
+	}
+}
+
+FName FDataflowDynamicConnections::GetCppTypeFromPropertyDesc(const FPropertyBagPropertyDesc& Desc)
+{
+	return FDataflowConnection::GetTypeNameFromProperty(Desc.CachedProperty);
+}
+
+bool FDataflowDynamicConnections::SetConnectionTypeFromPropertyDesc(FDataflowConnection& Connection, const FPropertyBagPropertyDesc& Desc)
+{
+	if (FDataflowNode* OwnerNode = OwnerInterface ? OwnerInterface->GetOwner(this) : nullptr)
+	{
+		FName ConcreteType(GetCppTypeFromPropertyDesc(Desc));
+
+		switch (Desc.ValueType)
+		{
+		case EPropertyBagPropertyType::Double:
+		case EPropertyBagPropertyType::Float:
+			// the UI shows float but behind the scene set a double property let only deal with float 
+			return Connection.SetConcreteType(TDataflowPolicyTypeName<float>::GetName());
+
+		case EPropertyBagPropertyType::Bool:
+		case EPropertyBagPropertyType::Byte:
+		case EPropertyBagPropertyType::Int32:
+		case EPropertyBagPropertyType::Int64:
+		case EPropertyBagPropertyType::Name:
+		case EPropertyBagPropertyType::String:
+		case EPropertyBagPropertyType::Text:
+		case EPropertyBagPropertyType::Object:
+		case EPropertyBagPropertyType::Struct:
+			return Connection.SetConcreteType(ConcreteType);
+
+			// unsupported type for now 
+		case EPropertyBagPropertyType::Enum:
+		case EPropertyBagPropertyType::SoftObject:
+		case EPropertyBagPropertyType::Class:
+		case EPropertyBagPropertyType::SoftClass:
+		case EPropertyBagPropertyType::UInt32:	// Type not fully supported at UI, will work with restrictions to type editing
+		case EPropertyBagPropertyType::UInt64: // Type not fully supported at UI, will work with restrictions to type editing
+		default:
+			return false;
+		}
+	}
+	return false;
+}
+
+FDataflowConnection* FDataflowDynamicConnections::CreateConnection(FDataflowDynamicConnections::FConnectionReference ConnectionReference)
+{
+	if (FDataflowNode* OwnerNode = OwnerInterface ? OwnerInterface->GetOwner(this) : nullptr)
+	{
+		if (PinDirection == UE::Dataflow::FPin::EDirection::INPUT)
+		{
+			return &OwnerNode->RegisterInputArrayConnection(ConnectionReference);
+		}
+		else if (PinDirection == UE::Dataflow::FPin::EDirection::OUTPUT)
+		{
+			return &OwnerNode->RegisterOutputArrayConnection(ConnectionReference);
+		}
+	}
+	return nullptr;
+}
+
+TArray<FDataflowConnection*> FDataflowDynamicConnections::GetDynamicConnections() const
+{
+	TArray<FDataflowConnection*> Connections;
+	if (FDataflowNode* OwnerNode = OwnerInterface ? OwnerInterface->GetOwner(this) : nullptr)
+	{
+		if (PinDirection == UE::Dataflow::FPin::EDirection::INPUT)
+		{
+			Connections.Append(OwnerNode->GetInputs());
+		}
+		else if (PinDirection == UE::Dataflow::FPin::EDirection::OUTPUT)
+		{
+			Connections.Append(OwnerNode->GetOutputs());
+		}
+	}
+	return Connections;
+}
+
+void FDataflowDynamicConnections::ClearDynamicConnections()
+{
+	if (FDataflowNode* OwnerNode = OwnerInterface ? OwnerInterface->GetOwner(this) : nullptr)
+	{
+		if (PinDirection == UE::Dataflow::FPin::EDirection::INPUT)
+		{
+			OwnerNode->ClearInputs();
+		}
+		else if (PinDirection == UE::Dataflow::FPin::EDirection::OUTPUT)
+		{
+			OwnerNode->ClearOutputs();
+		}
+	}
+}
\ No newline at end of file
diff --git a/Engine/Source/Runtime/Experimental/Dataflow/Engine/Private/Dataflow/DataflowSubGraphNodes.cpp b/Engine/Source/Runtime/Experimental/Dataflow/Engine/Private/Dataflow/DataflowSubGraphNodes.cpp
--- a/Engine/Source/Runtime/Experimental/Dataflow/Engine/Private/Dataflow/DataflowSubGraphNodes.cpp
+++ b/Engine/Source/Runtime/Experimental/Dataflow/Engine/Private/Dataflow/DataflowSubGraphNodes.cpp
@@ -241,255 +241,6 @@ namespace UE::Dataflow
 
 //////////////////////////////////////////////////////////////////////////////////////////////////////////
 
-FDataflowDynamicConnections::FDataflowDynamicConnections()
-	: PinDirection(UE::Dataflow::FPin::EDirection::NONE)
-	, OwnerInterface(nullptr)
-	, DataflowAssetWeakPtr(nullptr)
-{}
-
-FDataflowDynamicConnections::FDataflowDynamicConnections(UE::Dataflow::FPin::EDirection InPinDirection, IOwnerInterface* InOwnerInterface, UDataflow* InDataflowAsset)
-	: PinDirection(InPinDirection)
-	, OwnerInterface(InOwnerInterface)
-	, DataflowAssetWeakPtr(InDataflowAsset)
-{}
-
-
-FDataflowDynamicConnections::FConnectionReference FDataflowDynamicConnections::GetConnectionReference(int32 Index) const
-{
-	return { &DynamicProperties[Index], Index, &DynamicProperties };
-}
-
-void FDataflowDynamicConnections::Refresh()
-{
-	using namespace UE::Dataflow;
-
-	if (!OwnerInterface)
-	{
-		return;
-	}
-
-	FDataflowNode* OwnerNode = OwnerInterface->GetOwner(this);
-	if (!OwnerNode)
-	{
-		return;
-	}
-
-	TStrongObjectPtr<UDataflow> DataflowAsset = DataflowAssetWeakPtr.Pin();
-	if (!DataflowAsset.IsValid())
-	{
-		return;
-	}
-
-	TSharedPtr<UE::Dataflow::FGraph, ESPMode::ThreadSafe> DataflowGraph = DataflowAsset->GetDataflow();
-	if (!DataflowGraph.IsValid())
-	{
-		return;
-	}
-
-	// save the remote connections from the dynamic connections
-	using FRemoteConnections = TArray<FDataflowConnection*>;
-	TMap<FGuid, FRemoteConnections> RemoteConnectionsByPropertyId;
-
-	for (FDataflowConnection* DynamicConnection : GetDynamicConnections())
-	{
-		if (DynamicConnection)
-		{
-			const FGuid* PropertyId = ConnectionNameToPropertyId.Find(DynamicConnection->GetName());
-			if (PropertyId)
-			{
-				FRemoteConnections& RemoteConnections = RemoteConnectionsByPropertyId.Emplace(*PropertyId);
-				DynamicConnection->GetConnections(RemoteConnections);
-				DataflowGraph->ClearConnections(DynamicConnection);
-			}
-		}
-	}
-	// clear all the dynamic connections
-	ClearDynamicConnections();
-
-	ConnectionNameToPropertyId.Reset();
-
-	// Go through all the properties of the property bag recreate the output and reconnect them if needed
-	const TConnectionReference<FDataflowAllTypes> NullReference(nullptr);
-	const FInstancedPropertyBag& PropertyBag = OwnerInterface->GetPropertyBag(this);
-	if (const UPropertyBag* PropertyBagStruct = PropertyBag.GetPropertyBagStruct())
-	{
-		for (const FPropertyBagPropertyDesc& PropertyDesc : PropertyBagStruct->GetPropertyDescs())
-		{
-			// we can only add the output if it is a supported type
-			if (IsSupportedType(PropertyDesc))
-			{
-				const FName CppType = GetCppTypeFromPropertyDesc(PropertyDesc);
-				if (FDataflowConnection* NewConnection = AddNewConnectionFromPropertyDesc(PropertyDesc))
-				{
-					// now try to reconnect what was connected before
-					if (FRemoteConnections* RemoteConnections = RemoteConnectionsByPropertyId.Find(PropertyDesc.ID))
-					{
-						for (FDataflowConnection* RemoteConnection : *RemoteConnections)
-						{
-							DataflowGraph->Connect(NewConnection, RemoteConnection);
-						}
-					}
-				}
-			}
-		}
-		// Todo : need to also remove the one that are no longer necessary
-		// ( need to disconnect if needed ? or keep the one that are still connected ?)
-		OwnerNode->Invalidate();
-
-		DataflowAsset->RefreshEdNodeByGuid(OwnerNode->GetGuid());
-	}
-}
-
-FDataflowConnection* FDataflowDynamicConnections::AddNewConnectionFromPropertyDesc(const FPropertyBagPropertyDesc& PropertyDesc)
-{
-	using namespace UE::Dataflow;
-
-	if (FDataflowNode* OwnerNode = OwnerInterface ? OwnerInterface->GetOwner(this) : nullptr)
-	{
-		const int32 Index = DynamicProperties.AddDefaulted();
-		TConnectionReference<FDataflowAllTypes> ConnectionReference = GetConnectionReference(Index);
-		if (FDataflowConnection* NewConnection = CreateConnection(ConnectionReference))
-		{
-			SetConnectionTypeFromPropertyDesc(*NewConnection, PropertyDesc);
-			if (PinDirection == UE::Dataflow::FPin::EDirection::INPUT)
-			{
-				OwnerNode->TryRenameInput(ConnectionReference, PropertyDesc.Name);
-			}
-			else if (PinDirection == UE::Dataflow::FPin::EDirection::OUTPUT)
-			{
-				OwnerNode->TryRenameOutput(ConnectionReference, PropertyDesc.Name);
-			}
-			NewConnection->LockType();
-
-			ConnectionNameToPropertyId.Add(NewConnection->GetName(), PropertyDesc.ID);
-			return NewConnection;
-		}
-	}
-	return nullptr;
-}
-
-bool FDataflowDynamicConnections::IsSupportedType(const FPropertyBagPropertyDesc& Desc)
-{
-	switch (Desc.ValueType)
-	{
-	case EPropertyBagPropertyType::Bool:
-	case EPropertyBagPropertyType::Byte:
-	case EPropertyBagPropertyType::Int32:
-	case EPropertyBagPropertyType::Int64:
-	case EPropertyBagPropertyType::Float:
-	case EPropertyBagPropertyType::Double:
-	case EPropertyBagPropertyType::Name:
-	case EPropertyBagPropertyType::String:
-	case EPropertyBagPropertyType::Text:
-	case EPropertyBagPropertyType::Object:
-	case EPropertyBagPropertyType::Struct:
-		return true;
-
-		// unsupported type for now 
-	case EPropertyBagPropertyType::Enum:
-	case EPropertyBagPropertyType::SoftObject:
-	case EPropertyBagPropertyType::Class:
-	case EPropertyBagPropertyType::SoftClass:
-	case EPropertyBagPropertyType::UInt32:	// Type not fully supported at UI, will work with restrictions to type editing
-	case EPropertyBagPropertyType::UInt64: // Type not fully supported at UI, will work with restrictions to type editing
-	default:
-		return false;
-	}
-}
-
-FName FDataflowDynamicConnections::GetCppTypeFromPropertyDesc(const FPropertyBagPropertyDesc& Desc)
-{
-	return FDataflowConnection::GetTypeNameFromProperty(Desc.CachedProperty);
-}
-
-bool FDataflowDynamicConnections::SetConnectionTypeFromPropertyDesc(FDataflowConnection& Connection, const FPropertyBagPropertyDesc& Desc)
-{
-	if (FDataflowNode* OwnerNode = OwnerInterface ? OwnerInterface->GetOwner(this) : nullptr)
-	{
-		FName ConcreteType(GetCppTypeFromPropertyDesc(Desc));
-
-		switch (Desc.ValueType)
-		{
-		case EPropertyBagPropertyType::Double:
-		case EPropertyBagPropertyType::Float:
-			// the UI shows float but behind the scene set a double property let only deal with float 
-			return Connection.SetConcreteType(TDataflowPolicyTypeName<float>::GetName());
-
-		case EPropertyBagPropertyType::Bool:	
-		case EPropertyBagPropertyType::Byte:
-		case EPropertyBagPropertyType::Int32:
-		case EPropertyBagPropertyType::Int64:
-		case EPropertyBagPropertyType::Name:
-		case EPropertyBagPropertyType::String:
-		case EPropertyBagPropertyType::Text:
-		case EPropertyBagPropertyType::Object:
-		case EPropertyBagPropertyType::Struct:
-			return Connection.SetConcreteType(ConcreteType);
-
-		// unsupported type for now 
-		case EPropertyBagPropertyType::Enum:
-		case EPropertyBagPropertyType::SoftObject:
-		case EPropertyBagPropertyType::Class:
-		case EPropertyBagPropertyType::SoftClass:
-		case EPropertyBagPropertyType::UInt32:	// Type not fully supported at UI, will work with restrictions to type editing
-		case EPropertyBagPropertyType::UInt64: // Type not fully supported at UI, will work with restrictions to type editing
-		default:
-			return false;
-		}
-	}
-	return false;
-}
-
-FDataflowConnection* FDataflowDynamicConnections::CreateConnection(FDataflowDynamicConnections::FConnectionReference ConnectionReference)
-{
-	if (FDataflowNode* OwnerNode = OwnerInterface ? OwnerInterface->GetOwner(this) : nullptr)
-	{
-		if (PinDirection == UE::Dataflow::FPin::EDirection::INPUT)
-		{
-			return &OwnerNode->RegisterInputArrayConnection(ConnectionReference);
-		}
-		else if (PinDirection == UE::Dataflow::FPin::EDirection::OUTPUT)
-		{
-			return &OwnerNode->RegisterOutputArrayConnection(ConnectionReference);
-		}
-	}
-	return nullptr;
-}
-
-TArray<FDataflowConnection*> FDataflowDynamicConnections::GetDynamicConnections() const
-{
-	TArray<FDataflowConnection*> Connections;
-	if (FDataflowNode* OwnerNode = OwnerInterface ? OwnerInterface->GetOwner(this) : nullptr)
-	{
-		if (PinDirection == UE::Dataflow::FPin::EDirection::INPUT)
-		{
-			Connections.Append(OwnerNode->GetInputs());
-		}
-		else if (PinDirection == UE::Dataflow::FPin::EDirection::OUTPUT)
-		{
-			Connections.Append(OwnerNode->GetOutputs());
-		}
-	}
-	return Connections;
-}
-
-void FDataflowDynamicConnections::ClearDynamicConnections()
-{
-	if (FDataflowNode* OwnerNode = OwnerInterface ? OwnerInterface->GetOwner(this) : nullptr)
-	{
-		if (PinDirection == UE::Dataflow::FPin::EDirection::INPUT)
-		{
-			OwnerNode->ClearInputs();
-		}
-		else if (PinDirection == UE::Dataflow::FPin::EDirection::OUTPUT)
-		{
-			OwnerNode->ClearOutputs();
-		}
-	}
-}
-
-//////////////////////////////////////////////////////////////////////////////////////////////////////////
-
 FDataflowSubGraphInputNode::FDataflowSubGraphInputNode(const UE::Dataflow::FNodeParameters& InParam, FGuid InGuid)
 	: FDataflowNode(InParam, InGuid)
 	, DynamicConnections(UE::Dataflow::FPin::EDirection::OUTPUT, this, Cast<UDataflow>(InParam.OwningObject))
@@ -591,8 +342,6 @@ const FDataflowConnection* FDataflowSubGraphOutputNode::OnDropConnectionOnNode(c
 
 //////////////////////////////////////////////////////////////////////////////////////////////////////////
 
-#define NEW_METHOD 1
-
 FDataflowCallSubGraphNode::FDataflowCallSubGraphNode(const UE::Dataflow::FNodeParameters& InParam, FGuid InGuid)
 	: FDataflowNode(InParam, InGuid)
 	, DynamicInputs(UE::Dataflow::FPin::EDirection::INPUT, this, Cast<UDataflow>(InParam.OwningObject))
diff --git a/Engine/Source/Runtime/Experimental/Dataflow/Engine/Public/Dataflow/DataflowDynamicConnections.h b/Engine/Source/Runtime/Experimental/Dataflow/Engine/Public/Dataflow/DataflowDynamicConnections.h
new file mode 100644
--- /dev/null
+++ b/Engine/Source/Runtime/Experimental/Dataflow/Engine/Public/Dataflow/DataflowDynamicConnections.h
@@ -0,0 +1,56 @@
+// Copyright Epic Games, Inc. All Rights Reserved.
+
+#pragma once
+
+#include "Dataflow/DataflowAnyType.h"
+#include "Dataflow/DataflowConnection.h"
+
+#include "StructUtils/PropertyBag.h"
+
+#include "DataflowDynamicConnections.generated.h"
+
+struct FDataflowNode;
+
+/*
+* Dynamic connection object allow a node to add dynamic inputs or outputs backed by a strongly typed properties via a property bag 
+*/
+USTRUCT()
+struct FDataflowDynamicConnections
+{
+	GENERATED_USTRUCT_BODY()
+
+	struct IOwnerInterface
+	{
+		virtual FDataflowNode* GetOwner(const FDataflowDynamicConnections* Caller) = 0;
+		virtual const FInstancedPropertyBag& GetPropertyBag(const FDataflowDynamicConnections* Caller) = 0;
+	};
+
+public:
+	FDataflowDynamicConnections();
+	FDataflowDynamicConnections(UE::Dataflow::FPin::EDirection PInDirection, IOwnerInterface* OwnerInterface, UDataflow* DataflowAsset);
+
+	void Refresh();
+
+private:
+	using FConnectionReference = UE::Dataflow::TConnectionReference<FDataflowAllTypes>;
+	FConnectionReference GetConnectionReference(int32 Index) const;
+
+	bool IsSupportedType(const FPropertyBagPropertyDesc& PropertyDesc);
+	FDataflowConnection* CreateConnection(FConnectionReference Connectionreference);
+	TArray<FDataflowConnection*> GetDynamicConnections() const;
+	void ClearDynamicConnections();
+
+	bool SetConnectionTypeFromPropertyDesc(FDataflowConnection& Connection, const FPropertyBagPropertyDesc& PropertyDesc);
+	FName GetCppTypeFromPropertyDesc(const FPropertyBagPropertyDesc& PropertyDesc);
+	FDataflowConnection* AddNewConnectionFromPropertyDesc(const FPropertyBagPropertyDesc& PropertyDesc);
+
+private:
+	UPROPERTY()
+	TArray<FDataflowAllTypes> DynamicProperties;
+
+	TMap<FName, FGuid> ConnectionNameToPropertyId;
+
+	UE::Dataflow::FPin::EDirection PinDirection = UE::Dataflow::FPin::EDirection::NONE;
+	IOwnerInterface* OwnerInterface;
+	TWeakObjectPtr<UDataflow> DataflowAssetWeakPtr;
+};
\ No newline at end of file
diff --git a/Engine/Source/Runtime/Experimental/Dataflow/Engine/Public/Dataflow/DataflowSubGraphNodes.h b/Engine/Source/Runtime/Experimental/Dataflow/Engine/Public/Dataflow/DataflowSubGraphNodes.h
--- a/Engine/Source/Runtime/Experimental/Dataflow/Engine/Public/Dataflow/DataflowSubGraphNodes.h
+++ b/Engine/Source/Runtime/Experimental/Dataflow/Engine/Public/Dataflow/DataflowSubGraphNodes.h
@@ -5,6 +5,7 @@
 #include "Dataflow/DataflowConnection.h"
 #include "Dataflow/DataflowNodeParameters.h"
 #include "Dataflow/DataflowNode.h"
+#include "Dataflow/DataflowDynamicConnections.h"
 #include "StructUtils/PropertyBag.h"
 
 #include "DataflowSubGraphNodes.generated.h"
@@ -24,46 +25,7 @@ namespace UE::Dataflow
 	};
 }
 
-USTRUCT()
-struct FDataflowDynamicConnections
-{
-	GENERATED_USTRUCT_BODY()
-
-	struct IOwnerInterface
-	{
-		virtual FDataflowNode* GetOwner(const FDataflowDynamicConnections* Caller) = 0;
-		virtual const FInstancedPropertyBag& GetPropertyBag(const FDataflowDynamicConnections* Caller) = 0;
-	};
-
-public:
-	FDataflowDynamicConnections();
-	FDataflowDynamicConnections(UE::Dataflow::FPin::EDirection PInDirection, IOwnerInterface* OwnerInterface, UDataflow* DataflowAsset);
-
-	void Refresh();
-
-private:
-	using FConnectionReference = UE::Dataflow::TConnectionReference<FDataflowAllTypes>;
-	FConnectionReference GetConnectionReference(int32 Index) const;
-
-	bool IsSupportedType(const FPropertyBagPropertyDesc& PropertyDesc);
-	FDataflowConnection* CreateConnection(FConnectionReference Connectionreference);
-	TArray<FDataflowConnection*> GetDynamicConnections() const;
-	void ClearDynamicConnections();
-
-	bool SetConnectionTypeFromPropertyDesc(FDataflowConnection& Connection, const FPropertyBagPropertyDesc& PropertyDesc);
-	FName GetCppTypeFromPropertyDesc(const FPropertyBagPropertyDesc& PropertyDesc);
-	FDataflowConnection* AddNewConnectionFromPropertyDesc(const FPropertyBagPropertyDesc& PropertyDesc);
 
-private:
-	UPROPERTY()
-	TArray<FDataflowAllTypes> DynamicProperties;
-
-	TMap<FName, FGuid> ConnectionNameToPropertyId;
-
-	UE::Dataflow::FPin::EDirection PinDirection = UE::Dataflow::FPin::EDirection::NONE;
-	IOwnerInterface* OwnerInterface;
-	TWeakObjectPtr<UDataflow> DataflowAssetWeakPtr;
-};
 
 /**
 * This node represent the inputs of a dataflow subgraph 
