Dataflow : fix support for arrays of float in dynamic connections

diff --git a/Engine/Source/Runtime/Experimental/Dataflow/Engine/Private/Dataflow/DataflowDynamicConnections.cpp b/Engine/Source/Runtime/Experimental/Dataflow/Engine/Private/Dataflow/DataflowDynamicConnections.cpp
--- a/Engine/Source/Runtime/Experimental/Dataflow/Engine/Private/Dataflow/DataflowDynamicConnections.cpp
+++ b/Engine/Source/Runtime/Experimental/Dataflow/Engine/Private/Dataflow/DataflowDynamicConnections.cpp
@@ -169,14 +169,20 @@ bool FDataflowDynamicConnections::SetConnectionTypeFromPropertyDesc(FDataflowCon
 {
 	if (FDataflowNode* OwnerNode = OwnerInterface ? OwnerInterface->GetOwner(this) : nullptr)
 	{
-		FName ConcreteType(GetCppTypeFromPropertyDesc(Desc));
+		const FName PropertyCppType(GetCppTypeFromPropertyDesc(Desc));
+		const bool bIsArray = (Desc.ContainerTypes.GetFirstContainerType() == EPropertyBagContainerType::Array);
 
+		FName ConcreteType = NAME_None;
 		switch (Desc.ValueType)
 		{
 		case EPropertyBagPropertyType::Double:
 		case EPropertyBagPropertyType::Float:
 			// the UI shows float but behind the scene set a double property let only deal with float 
-			return Connection.SetConcreteType(TDataflowPolicyTypeName<float>::GetName());
+			if (bIsArray)
+			{
+				return Connection.SetConcreteType(UE::Dataflow::GetTypeName<TArray<float>>());
+			}
+			return Connection.SetConcreteType(UE::Dataflow::GetTypeName<float>());
 
 		case EPropertyBagPropertyType::Bool:
 		case EPropertyBagPropertyType::Byte:
@@ -187,9 +193,9 @@ bool FDataflowDynamicConnections::SetConnectionTypeFromPropertyDesc(FDataflowCon
 		case EPropertyBagPropertyType::Text:
 		case EPropertyBagPropertyType::Object:
 		case EPropertyBagPropertyType::Struct:
-			return Connection.SetConcreteType(ConcreteType);
+			return Connection.SetConcreteType(PropertyCppType);
 
-			// unsupported type for now 
+		// unsupported type for now 
 		case EPropertyBagPropertyType::Enum:
 		case EPropertyBagPropertyType::SoftObject:
 		case EPropertyBagPropertyType::Class:
