Geometry collection : expose automatic partial cluster crumbling to users for remove on sleep feature (bAutomaticCrumblePartialClusters)

diff --git a/Engine/Source/Runtime/Experimental/GeometryCollectionEngine/Private/GeometryCollection/GeometryCollectionComponent.cpp b/Engine/Source/Runtime/Experimental/GeometryCollectionEngine/Private/GeometryCollection/GeometryCollectionComponent.cpp
--- a/Engine/Source/Runtime/Experimental/GeometryCollectionEngine/Private/GeometryCollection/GeometryCollectionComponent.cpp
+++ b/Engine/Source/Runtime/Experimental/GeometryCollectionEngine/Private/GeometryCollection/GeometryCollectionComponent.cpp
@@ -6522,6 +6522,7 @@ void UGeometryCollectionComponent::IncrementSleepTimer(float DeltaTime)
 			FGeometryCollectionDecayContext DecayContext(*PhysicsProxy, DecayFacade);
 
 			const TManagedArray<int32>& OriginalParents = RestCollection->GetGeometryCollection()->Parent;
+			const bool bAutomaticCrumblePartialClusters = RestCollection->bAutomaticCrumblePartialClusters;
 
 			const int32 NumTransforms = OriginalParents.Num();
 			for (int32 TransformIdx = 0; TransformIdx < NumTransforms; ++TransformIdx)
@@ -6534,7 +6535,7 @@ void UGeometryCollectionComponent::IncrementSleepTimer(float DeltaTime)
 					const bool HasDynamicInternalClusterParent = DynamicStateFacade.HasDynamicInternalClusterParent(TransformIdx);
 					if (OriginalParentIdx > INDEX_NONE && HasDynamicInternalClusterParent && RemoveOnSleepFacade.IsRemovalActive(OriginalParentIdx))
 					{
-						const bool UseClusterCrumbling = true; // with sleep removal : internal clusters always crumble - this will change when we merge the removal feature together
+						const bool UseClusterCrumbling = bAutomaticCrumblePartialClusters;
 						const float UpdatedBreakDecay = UE_SMALL_NUMBER; // since we crumble we can only pass a timy number since this will be ignore ( but need to be >0 to ake sure Update Decay works properly )
 						UpdateDecay(TransformIdx, UpdatedBreakDecay, UseClusterCrumbling, HasDynamicInternalClusterParent, DecayContext);
 					}
diff --git a/Engine/Source/Runtime/Experimental/GeometryCollectionEngine/Private/GeometryCollection/GeometryCollectionObject.cpp b/Engine/Source/Runtime/Experimental/GeometryCollectionEngine/Private/GeometryCollection/GeometryCollectionObject.cpp
--- a/Engine/Source/Runtime/Experimental/GeometryCollectionEngine/Private/GeometryCollection/GeometryCollectionObject.cpp
+++ b/Engine/Source/Runtime/Experimental/GeometryCollectionEngine/Private/GeometryCollection/GeometryCollectionObject.cpp
@@ -133,6 +133,7 @@ UGeometryCollection::UGeometryCollection(const FObjectInitializer& ObjectInitial
 	, bOptimizeConvexes(Chaos::CVars::bChaosConvexSimplifyUnion)
 	, bScaleOnRemoval(true)
 	, bRemoveOnMaxSleep(false)
+	, bAutomaticCrumblePartialClusters(true)
 	, MaximumSleepTime(5.0, 10.0)
 	, RemovalDuration(2.5, 5.0)
 	, bSlowMovingAsSleeping(true)
diff --git a/Engine/Source/Runtime/Experimental/GeometryCollectionEngine/Public/GeometryCollection/GeometryCollectionObject.h b/Engine/Source/Runtime/Experimental/GeometryCollectionEngine/Public/GeometryCollection/GeometryCollectionObject.h
--- a/Engine/Source/Runtime/Experimental/GeometryCollectionEngine/Public/GeometryCollection/GeometryCollectionObject.h
+++ b/Engine/Source/Runtime/Experimental/GeometryCollectionEngine/Public/GeometryCollection/GeometryCollectionObject.h
@@ -757,7 +757,11 @@ public:
 	/** Remove particle from simulation and dissolve rendered geometry once sleep threshold has been exceeded. */
 	UPROPERTY(EditAnywhere, BlueprintReadOnly, Category = Removal, meta = (DisplayName = "Remove on Sleep"))
 	bool bRemoveOnMaxSleep;
-	
+
+	/** When using remove on sleep , automatically crumble partial cluster into smaller pieces when possible */
+	UPROPERTY(EditAnywhere, BlueprintReadOnly, Category = Removal, meta = (EditCondition = "bRemoveOnMaxSleep"))
+	bool bAutomaticCrumblePartialClusters;
+
 	/** How long may the particle sleep before initiating removal (in seconds). */
 	UPROPERTY(EditAnywhere, BlueprintReadOnly, Category = Removal, meta = (DisplayName = "Sleep Min Max", EditCondition="bRemoveOnMaxSleep"))
 	FVector2D MaximumSleepTime;
