[NetTokens] Implement TStructNetTokenDataStore and Serialization Helpers.

- Implements a templated StructNetTokenDataStore and the supporting helpers to simplify implementation/use of NetTokens
- Implements native serialization techniques, supporting automatic replacement of NetTokenBase types with their respective NetToken
- Implements Generic Struct support and Replay Support for those times
-- Still has the same issues with Replay Support that FInstancedStruct does!
- Implements Tests cases for Iris test suite

diff --git a/Engine/Plugins/Runtime/ReplicationSystemTestPlugin/Source/Private/Tests/ReplicationSystem/TestStructNetTokenDataStore.cpp b/Engine/Plugins/Runtime/ReplicationSystemTestPlugin/Source/Private/Tests/ReplicationSystem/TestStructNetTokenDataStore.cpp
new file mode 100644
--- /dev/null
+++ b/Engine/Plugins/Runtime/ReplicationSystemTestPlugin/Source/Private/Tests/ReplicationSystem/TestStructNetTokenDataStore.cpp
@@ -0,0 +1,230 @@
+// Copyright Epic Games, Inc. All Rights Reserved.
+
+#include "TestStructNetTokenDataStore.h"
+#include "ReplicationSystemServerClientTestFixture.h"
+#include "Iris/ReplicationSystem/ReplicationFragmentUtil.h"
+#include "Iris/ReplicationSystem/StructNetTokenDataStore.h"
+#include "Iris/ReplicationState/PropertyNetSerializerInfoRegistry.h"
+#include "Logging/LogScopedVerbosityOverride.h"
+
+// Note: this is in engine code.
+#include "Net/Iris/ReplicationSystem/StructNetTokenDataStoreHelper.h"
+#include "Net/UnrealNetwork.h"
+
+namespace UE::Net
+{
+	
+FTestMessage& operator<<(FTestMessage& Message, const FTestStructNetTokenStoreStruct& StructValue)
+{
+	return Message << TEXT("StructValue: A:") << StructValue.IntA << TEXT("B:") << StructValue.IntB;
+}
+	
+} // end of namespace UE::Net;
+
+UE_NET_IMPLEMENT_NAMED_NETTOKEN_STRUCT_SERIALIZERS(TestStructNetTokenStoreStruct)
+
+UE_NET_IMPLEMENT_NAMED_STRUCT_LASTRESORT_NETSERIALIZER_AND_REGISTRY_DELEGATES(TestStructNetTokenStoreStructDerived);
+
+////////////////////////////////////////////////
+// Tests
+////////////////////////////////////////////////
+namespace UE::Net::Private
+{
+
+class FTestStructNetTokensFixture : public FReplicationSystemServerClientTestFixture
+{
+public:
+
+	typedef UE::Net::TStructNetTokenDataStore<FTestStructNetTokenStoreStruct> FTestStructNetTokenDataStore;
+	using DataType = FTestStructNetTokenDataStore::DataType;
+
+	FTestStructNetTokenDataStore* ServerStructTokenStore = nullptr;
+	FTestStructNetTokenDataStore* ClientStructTokenStore = nullptr;
+
+	FReplicationSystemTestClient* Client = nullptr;
+	
+	const FNetTokenStoreState* ClientRemoteNetTokenState;
+	const FNetTokenStoreState* ServerRemoteNetTokenState;
+
+
+	FNetToken CreateAndExportNetToken(const DataType& Struct)
+	{
+		FNetToken Token = ServerStructTokenStore->GetOrCreateToken(Struct);
+		UNetTokenDataStream* NetTokenDataStream = Cast<UNetTokenDataStream>(Server->GetReplicationSystem()->GetDataStream(Client->ConnectionIdOnServer, FName("NetToken")));
+		if (NetTokenDataStream)
+		{
+			NetTokenDataStream->AddNetTokenForExplicitExport(Token);
+		}
+
+		return Token;
+	}
+
+	FNetToken CreateAndExportNetTokenOnClient(const DataType& Struct)
+	{
+		FNetToken Token = ClientStructTokenStore->GetOrCreateToken(Struct);
+		UNetTokenDataStream* NetTokenDataStream = Cast<UNetTokenDataStream>(Client->GetReplicationSystem()->GetDataStream(Client->LocalConnectionId, FName("NetToken")));
+		if (NetTokenDataStream)
+		{
+			NetTokenDataStream->AddNetTokenForExplicitExport(Token);
+		}
+
+		return Token;
+	}
+
+	virtual void SetUp() override
+	{
+		FReplicationSystemServerClientTestFixture::SetUp();
+		NetTokenDataStoreUtil.AddNetTokenStoreTypeIdPair(FTestStructNetTokenStoreStruct::GetTokenStoreName().ToString(), 3);
+
+		Client = CreateClient();
+		{
+			FNetTokenStore* ServerTokenStore = Server->GetReplicationSystem()->GetNetTokenStore();
+
+			// Create server token store
+			ServerTokenStore->CreateAndRegisterDataStore<FTestStructNetTokenDataStore>();
+			ServerStructTokenStore = ServerTokenStore->GetDataStore<FTestStructNetTokenDataStore>();		
+			ServerRemoteNetTokenState = ServerTokenStore->GetRemoteNetTokenStoreState(Client->ConnectionIdOnServer);
+		}
+		{
+			FNetTokenStore* ClientTokenStore = Client->GetReplicationSystem()->GetNetTokenStore();
+
+			// Create client token store
+			ClientTokenStore->CreateAndRegisterDataStore<FTestStructNetTokenDataStore>();
+			ClientStructTokenStore = ClientTokenStore->GetDataStore<FTestStructNetTokenDataStore>();		
+			ClientRemoteNetTokenState = ClientTokenStore->GetRemoteNetTokenStoreState(Client->LocalConnectionId);
+		}
+	}
+
+	virtual void TearDown()
+	{
+		FReplicationSystemServerClientTestFixture::TearDown();
+	}	
+};
+
+// Test that we can send a FTestStructNetTokenStoreStruct as a NetToken (explicit)
+UE_NET_TEST_FIXTURE(FTestStructNetTokensFixture, StructAsNetToken)
+{
+	// Create token
+	FTestStructNetTokenStoreStruct StructA;
+	StructA.IntA = 1;
+	StructA.IntB = 2;
+
+	FNetToken StructTokenA = CreateAndExportNetToken(StructA);
+
+	// Verify that we can resolve the token on server.
+	UE_NET_ASSERT_EQ(StructA, ServerStructTokenStore->ResolveToken(StructTokenA));
+
+	{
+		LOG_SCOPE_VERBOSITY_OVERRIDE(LogNetToken, ELogVerbosity::Fatal);
+		// Verify that we cannot resolve the token on the client
+		UE_NET_ASSERT_NE(StructA, ClientStructTokenStore->ResolveRemoteToken(StructTokenA, *ClientRemoteNetTokenState));
+	}
+
+	// Send and deliver packet
+	Server->PreSendUpdate();
+	Server->SendAndDeliverTo(Client, true);
+	Server->PostSendUpdate();
+
+	// Verify that we can resolve the token on the client
+	UE_NET_ASSERT_EQ(StructA, ClientStructTokenStore->ResolveRemoteToken(StructTokenA, *ClientRemoteNetTokenState));
+}
+
+// Test that we can send a FTestStructNetTokenStoreStruct as a NetToken but this time exported as a replicated property using an Iris NetSerializer
+UE_NET_TEST_FIXTURE(FTestStructNetTokensFixture, StructAsNetTokenPropertyWithIris)
+{
+	// Create token struct
+	FTestStructNetTokenStoreStruct StructA;
+	StructA.IntA = 1;
+	StructA.IntB = 2;
+
+	UTestStructAsNetTokenObject* ServerObject = Server->CreateObject<UTestStructAsNetTokenObject>();
+	ServerObject->NetTokenStoreStruct = StructA;
+
+	// Send and deliver packet
+	Server->PreSendUpdate();
+	Server->SendAndDeliverTo(Client, true);
+	Server->PostSendUpdate();
+
+	UTestStructAsNetTokenObject* ClientObject = Cast<UTestStructAsNetTokenObject>(Client->GetReplicationBridge()->GetReplicatedObject(ServerObject->NetRefHandle));
+
+	// Verify that we replicated object
+	UE_NET_ASSERT_NE(ClientObject, nullptr);
+
+	// Verify that we replicated struct correctly
+	UE_NET_ASSERT_EQ(StructA, ClientObject->NetTokenStoreStruct);
+
+	// Get token on server
+	FNetToken ServerToken = ServerStructTokenStore->GetOrCreateToken(StructA);
+
+	// Get token on client
+	FNetToken ClientToken = ClientStructTokenStore->GetOrCreateToken(ClientObject->NetTokenStoreStruct);
+
+	// They should match
+	UE_NET_ASSERT_EQ(ServerToken, ClientToken);
+}
+
+
+// Test that we can send a FTestStructNetTokenStoreStruct as a NetToken but this time exported through a property using the LastResortNetserializer, calling into ::NetSerialize
+UE_NET_TEST_FIXTURE(FTestStructNetTokensFixture, StructAsNetTokenUsingLastResortNetSerializer)
+{
+	// Create token struct
+	FTestStructNetTokenStoreStructDerived StructA;
+	StructA.IntA = 1;
+	StructA.IntB = 2;
+
+	UTestStructAsNetTokenObject* ServerObject = Server->CreateObject<UTestStructAsNetTokenObject>();
+	ServerObject->DerivedNetTokenStoreStruct = StructA;
+
+	// Send and deliver packet
+	Server->PreSendUpdate();
+	Server->SendAndDeliverTo(Client, true);
+	Server->PostSendUpdate();
+
+	UTestStructAsNetTokenObject* ClientObject = Cast<UTestStructAsNetTokenObject>(Client->GetReplicationBridge()->GetReplicatedObject(ServerObject->NetRefHandle));
+
+	// Verify that we replicated object
+	UE_NET_ASSERT_NE(ClientObject, nullptr);
+
+	// Verify that we replicated struct correctly
+	UE_NET_ASSERT_EQ(StructA, ClientObject->DerivedNetTokenStoreStruct);
+
+	// Get token on server
+	FNetToken ServerToken = ServerStructTokenStore->GetOrCreateToken(StructA);
+
+	// Get token on client
+	FNetToken ClientToken = ClientStructTokenStore->GetOrCreateToken(ClientObject->DerivedNetTokenStoreStruct);
+
+	// They should match
+	UE_NET_ASSERT_EQ(ServerToken, ClientToken);
+}
+
+} // end namespace UE::Net::Private
+
+
+//////////////////////////////////////////////////////////////////////////
+// Implementation for UTestStructAsNetTokenObject
+//////////////////////////////////////////////////////////////////////////
+UTestStructAsNetTokenObject::UTestStructAsNetTokenObject()
+: UReplicatedTestObject()
+{
+}
+
+void UTestStructAsNetTokenObject::GetLifetimeReplicatedProps( TArray< class FLifetimeProperty > & OutLifetimeProps ) const
+{
+	FDoRepLifetimeParams LifetimeParams;
+	LifetimeParams.bIsPushBased = false;
+	DOREPLIFETIME_WITH_PARAMS(ThisClass, NetTokenStoreStruct, LifetimeParams);
+}
+
+void UTestStructAsNetTokenObject::RegisterReplicationFragments(UE::Net::FFragmentRegistrationContext& Context, UE::Net::EFragmentRegistrationFlags RegistrationFlags)
+{
+	// Base object owns the fragment in this case
+	{
+		this->ReplicationFragments.Reset();
+		UE::Net::FReplicationFragmentUtil::CreateAndRegisterFragmentsForObject(this, Context, RegistrationFlags, &this->ReplicationFragments);
+	}
+}
+
+
+
+
diff --git a/Engine/Plugins/Runtime/ReplicationSystemTestPlugin/Source/Private/Tests/ReplicationSystem/TestStructNetTokenDataStore.h b/Engine/Plugins/Runtime/ReplicationSystemTestPlugin/Source/Private/Tests/ReplicationSystem/TestStructNetTokenDataStore.h
new file mode 100644
--- /dev/null
+++ b/Engine/Plugins/Runtime/ReplicationSystemTestPlugin/Source/Private/Tests/ReplicationSystem/TestStructNetTokenDataStore.h
@@ -0,0 +1,70 @@
+// Copyright Epic Games, Inc. All Rights Reserved.
+
+#pragma once
+
+#include "ReplicatedTestObject.h"
+#include "Iris/Serialization/NetSerializer.h"
+#include "Iris/ReplicationSystem/NetTokenStructDefines.h"
+#include "TestStructNetTokenDataStore.generated.h"
+
+USTRUCT()
+struct FTestStructNetTokenStoreStruct
+{
+	GENERATED_BODY()
+
+	UE_NET_NETTOKEN_GENERATED_BODY(TestStructNetTokenStoreStruct, REPLICATIONSYSTEMTESTPLUGIN_API)
+
+	uint64 GetUniqueKey() const
+	{
+		return HashCombineFast(IntA, IntB);
+	}
+
+	UPROPERTY(Transient)
+	int32 IntA;
+
+	UPROPERTY(Transient)
+	int32 IntB;
+};
+	
+UE_NET_DECLARE_NAMED_NETTOKEN_STRUCT_SERIALIZERS(TestStructNetTokenStoreStruct, REPLICATIONSYSTEMTESTPLUGIN_API)
+
+// We will use a LastResortNetSerializer for this struct just to test NetSerialize using util.
+USTRUCT()
+struct FTestStructNetTokenStoreStructDerived : public FTestStructNetTokenStoreStruct 
+{
+	GENERATED_BODY()
+};
+
+//Since we are only testing last resort serializer for this Derived, we only want to declare the default NetTokenStruct ops for the RepGraph version.
+template<> struct TStructOpsTypeTraits<FTestStructNetTokenStoreStructDerived> : public UE::Net::TNetTokenStructOpsTypeTraits<FTestStructNetTokenStoreStructDerived> {};
+
+
+/**
+ *  A test class for Replication that itself uses Property based replication but also has "components" that uses a mix of property based replication and native ReplicationStates
+ */
+UCLASS()
+class UTestStructAsNetTokenObject : public UReplicatedTestObject
+{
+	GENERATED_BODY()
+
+public:
+	UTestStructAsNetTokenObject();
+
+	// Network interface must be part of base.
+	virtual void RegisterReplicationFragments(UE::Net::FFragmentRegistrationContext& Fragments, UE::Net::EFragmentRegistrationFlags RegistrationFlags) override;
+
+public:
+	
+	// This will replicate using iris NetSerializer
+	UPROPERTY(Transient, Replicated)
+	FTestStructNetTokenStoreStruct NetTokenStoreStruct;
+
+	// This will replicate using a LastResortNetSerializer and thus will call into old replication path but do exports
+	// using iris exports.
+	UPROPERTY(Transient, Replicated)
+	FTestStructNetTokenStoreStructDerived DerivedNetTokenStoreStruct;
+
+public:
+	// Network data only for test
+	TArray<UE::Net::FReplicationFragment*> ReplicationFragments;
+};
\ No newline at end of file
diff --git a/Engine/Source/Runtime/Engine/Classes/Engine/PackageMapClient.h b/Engine/Source/Runtime/Engine/Classes/Engine/PackageMapClient.h
--- a/Engine/Source/Runtime/Engine/Classes/Engine/PackageMapClient.h
+++ b/Engine/Source/Runtime/Engine/Classes/Engine/PackageMapClient.h
@@ -43,16 +43,30 @@ enum class EChannelGetAdditionalRequiredBunchesFlags : uint32;
 class FNetFieldExport
 {
 public:
-	FNetFieldExport() : Handle( 0 ), CompatibleChecksum( 0 ), bExported( false ), bDirtyForReplay( true ), bIncompatible( false )
+	FNetFieldExport() : Blob ( {} ), Handle( 0 ), CompatibleChecksum( 0 ), bExported( false ),  bDirtyForReplay( true ), bExportBlob( false ), bIncompatible( false )
 	{
 	}
 
 	FNetFieldExport( const uint32 InHandle, const uint32 InCompatibleChecksum, const FName& InName ) :
+		Blob ( {} ),
 		Handle( InHandle ),
 		CompatibleChecksum( InCompatibleChecksum ),
 		ExportName( InName ),
 		bExported( false ),
 		bDirtyForReplay( true ),
+		bExportBlob( false ),
+		bIncompatible( false )
+	{
+	}
+
+	FNetFieldExport( const uint32 InHandle, const uint32 InCompatibleChecksum, const FName& InName, const TArray<uint8>& InBlob ) :
+		Blob ( InBlob ),
+		Handle( InHandle ),
+		CompatibleChecksum( InCompatibleChecksum ),
+		ExportName( InName ),
+		bExported( false ),
+		bDirtyForReplay( true ),
+		bExportBlob( true ),
 		bIncompatible( false )
 	{
 	}
@@ -60,12 +74,15 @@ public:
 	friend FArchive& operator<<(FArchive& Ar, FNetFieldExport& C);
 
 	ENGINE_API void CountBytes(FArchive& Ar) const;
-
+	
+	TArray<uint8>	Blob;
 	uint32			Handle;
 	uint32			CompatibleChecksum;
 	FName			ExportName;
 	bool			bExported;
 	bool			bDirtyForReplay;
+	bool			bExportBlob;
+	
 
 	// Transient properties
 	mutable bool	bIncompatible;		// If true, we've already determined that this property isn't compatible. We use this to curb warning spam.
diff --git a/Engine/Source/Runtime/Engine/Private/PackageMapClient.cpp b/Engine/Source/Runtime/Engine/Private/PackageMapClient.cpp
--- a/Engine/Source/Runtime/Engine/Private/PackageMapClient.cpp
+++ b/Engine/Source/Runtime/Engine/Private/PackageMapClient.cpp
@@ -2258,7 +2258,7 @@ FOutBunch* UPackageMapClient::CreateCustomExportsBunch(const FOutBunch& Outgoing
 
 			// Write token 
 			NetTokenResolveContext.NetTokenStore->WriteNetToken(*ExportsBunch, NetToken);
-			NetTokenResolveContext.NetTokenStore->WriteTokenData(*ExportsBunch, NetToken);
+			NetTokenResolveContext.NetTokenStore->WriteTokenData(*ExportsBunch, NetToken, this);
 
 			// Track what exports we exported until we know what packet we actually did commit to.
 			ExportsBunch->NetTokensPendingExport.Add(NetToken);
@@ -2293,7 +2293,7 @@ void UPackageMapClient::ReceiveCustomExportsBunch(FInBunch &InBunch)
 	{
 		FNetToken ImportedNetToken = NetTokenStore->ReadNetToken(InBunch);
 
-		NetTokenStore->ReadTokenData(InBunch, ImportedNetToken, *RemoteNetTokenStoreState);
+		NetTokenStore->ReadTokenData(InBunch, ImportedNetToken, *RemoteNetTokenStoreState, this);
 
 		UE_LOG(LogNetToken, VeryVerbose, TEXT("%hs: Imported token %s"), __func__, *ImportedNetToken.ToString());
 
@@ -4309,13 +4309,14 @@ void FNetFieldExport::CountBytes(FArchive& Ar) const
 
 FArchive& operator<<(FArchive& Ar, FNetFieldExport& C)
 {
-	uint8 Flags = C.bExported ? 1 : 0;
-
+	uint8 Flags = 0;
+	Flags |= C.bExported	? (1 << 0) : 0;
+	Flags |= C.bExportBlob	? (1 << 1) : 0;
 	Ar << Flags;
-
 	if (Ar.IsLoading())
 	{
-		C.bExported = (Flags == 1);
+		C.bExported =	Flags & (1 << 0);
+		C.bExportBlob = Flags & (1 << 1);
 	}
 
 	if (C.bExported)
@@ -4344,6 +4345,11 @@ FArchive& operator<<(FArchive& Ar, FNetFieldExport& C)
 				UPackageMap::StaticSerializeName(Ar, C.ExportName);
 			}
 		}
+
+		if (C.bExportBlob)
+		{
+			SafeNetSerializeTArray_Default<4096>(Ar,C.Blob);
+		}
 	}
 
 	return Ar;
diff --git a/Engine/Source/Runtime/Engine/Public/Net/Iris/ReplicationSystem/StructNetTokenDataStoreHelper.h b/Engine/Source/Runtime/Engine/Public/Net/Iris/ReplicationSystem/StructNetTokenDataStoreHelper.h
new file mode 100644
--- /dev/null
+++ b/Engine/Source/Runtime/Engine/Public/Net/Iris/ReplicationSystem/StructNetTokenDataStoreHelper.h
@@ -0,0 +1,210 @@
+// Copyright Epic Games, Inc. All Rights Reserved.
+
+#pragma once
+
+#include "Iris/ReplicationSystem/NetTokenStore.h"
+#include "Iris/ReplicationSystem/StructNetTokenDataStore.h"
+#include "Engine/NetConnection.h"
+#include "Engine/PackageMapClient.h"
+#include "Net/RepLayout.h"
+#include "Net/Core/NetToken/NetTokenExportContext.h"
+#include "Iris/ReplicationState/PropertyNetSerializerInfoRegistry.h"
+
+namespace UE::Net
+{
+	
+// Helper for Serializing the NetTokenStructs 
+template <typename T>
+class TStructNetTokenDataStoreHelper
+{
+public:
+
+	using DataStoreType = TStructNetTokenDataStore<T>;
+
+	static void SerializeHelper(T& InOutData, FArchive& Ar, UPackageMap* Map)
+	{
+		UPackageMapClient* MapClient = Cast<UPackageMapClient>(Map);
+		check(::IsValid(MapClient));
+
+		UNetConnection* NetConnection = MapClient->GetConnection();
+		check(::IsValid(NetConnection));
+		check(::IsValid(NetConnection->GetDriver()));
+
+		UScriptStruct* Struct = StaticStruct<T>();
+		const TSharedPtr<FRepLayout> RepLayout = NetConnection->GetDriver()->GetStructRepLayout(Struct);
+		check(RepLayout.IsValid());
+
+		bool bHasUnmapped = false;
+		RepLayout->SerializePropertiesForStruct(Struct, static_cast<FBitArchive&>(Ar), Map, &InOutData, bHasUnmapped);
+	}
+public:
+	static bool NetSerializeAndExportToken(FArchive& Ar, UPackageMap* Map, T& InOutStateData)
+	{
+		if (Ar.IsSaving())
+		{
+			FNetTokenExportContext* ExportContext = FNetTokenExportContext::GetNetTokenExportContext(Ar);
+			FNetTokenStore* NetTokenStore = ExportContext ? ExportContext->GetNetTokenStore() : nullptr;
+			DataStoreType* StateTokenStore = NetTokenStore ? NetTokenStore->GetDataStore<DataStoreType>() : nullptr;
+			if (!StateTokenStore)
+			{
+				return false;
+			}
+
+			if (!DataStoreType::NetSerializeScriptDelegate.IsBound())
+			{
+				DataStoreType::NetSerializeScriptDelegate.BindStatic(&SerializeHelper);
+			}
+
+			UPackageMapClient* PackageMapClient = Cast<UPackageMapClient>(Map);
+			if (PackageMapClient && PackageMapClient->GetConnection() && PackageMapClient->GetConnection()->IsReplay())
+			{
+				return NetSerializeStateDataAsReplayData(StateTokenStore, InOutStateData, Ar, *PackageMapClient);
+			}
+			
+			FNetToken StateToken = StateTokenStore->GetOrCreateToken(InOutStateData);
+			StateTokenStore->WriteNetToken(Ar, StateToken);
+			ExportContext->AddNetTokenPendingExport(StateToken);
+		}
+		else
+		{
+			const FNetTokenResolveContext* NetTokenResolveContext = Map ? Map->GetNetTokenResolveContext() : nullptr;
+			FNetTokenStore* NetTokenStore = NetTokenResolveContext ? NetTokenResolveContext->NetTokenStore : nullptr;
+			DataStoreType* StateTokenStore = NetTokenStore ? NetTokenStore->GetDataStore<DataStoreType>() : nullptr;
+			if (!StateTokenStore)
+			{
+				return false;
+			}
+
+			UPackageMapClient* PackageMapClient = Cast<UPackageMapClient>(Map);
+			if (PackageMapClient && PackageMapClient->GetConnection() && PackageMapClient->GetConnection()->IsReplay())
+			{
+				return NetSerializeStateDataAsReplayData(StateTokenStore, InOutStateData, Ar, *PackageMapClient);
+			}
+			
+			FNetToken StateToken = StateTokenStore->ReadNetToken(Ar);
+			if (Ar.IsError())
+			{
+				return false;
+			}
+			InOutStateData = StateTokenStore->ResolveRemoteToken(StateToken, *NetTokenResolveContext->RemoteNetTokenStoreState);
+		}
+		return true;
+	}
+
+	// Helpers for Replay Support
+	static bool NetSerializeStateDataAsReplayData(DataStoreType* StateTokenStore, T& InOutStateData, FArchive& Ar, UPackageMapClient& PackageMapClient)
+	{
+		TSharedPtr<FNetFieldExportGroup> NetFieldExportGroup = PackageMapClient.GetNetFieldExportGroup(DataStoreType::GetTokenStoreName().ToString());
+		if (Ar.IsSaving())
+		{
+			if (!NetFieldExportGroup.IsValid())
+			{
+				SetupReplayExportGroup(StateTokenStore, PackageMapClient);
+			}
+			FNetToken StateToken = StateTokenStore->GetOrCreateToken(InOutStateData);
+			StateTokenStore->WriteNetToken(Ar, StateToken);
+
+			if (!NetFieldExportGroup->NetFieldExports.IsValidIndex(StateToken.GetIndex()))
+			{
+				NetFieldExportGroup->NetFieldExports.SetNum(StateToken.GetIndex() + 1, EAllowShrinking::No);
+				NetFieldExportGroup->bDirtyForReplay = true;
+			}
+			FNetFieldExport& NetFieldExport = NetFieldExportGroup->NetFieldExports[StateToken.GetIndex()];
+
+			AddStateToReplayExportGroup(StateTokenStore, InOutStateData, NetFieldExportGroup, PackageMapClient);
+			PackageMapClient.TrackNetFieldExport(NetFieldExportGroup.Get(), StateToken.GetIndex());
+
+			return true;
+		}
+		else if (Ar.IsLoading())
+		{
+			FNetToken StateToken = StateTokenStore->ReadNetToken(Ar);
+			if (Ar.IsError())
+			{
+				return false;
+			}
+			if (!NetFieldExportGroup.IsValid())
+			{
+				return false;
+			}
+
+			if (StateToken.IsValid())
+			{
+				if(GetStateFromReplayExportGroup(StateTokenStore, InOutStateData, StateToken, NetFieldExportGroup, PackageMapClient))
+				{
+					return true;
+				}
+				InOutStateData = DataStoreType::GetInvalidState();
+				return false;
+			}
+		}
+		return false;
+	}
+
+	static bool GetStateFromReplayExportGroup(DataStoreType* StateTokenStore, T& OutData, FNetToken& StateToken, TSharedPtr<FNetFieldExportGroup> NetFieldExportGroup, UPackageMapClient& Map)
+	{
+		if (!NetFieldExportGroup.IsValid())
+		{
+			return false;
+		}
+		const uint32 TagTokenIndex = StateToken.GetIndex();
+		if (static_cast<int32>(TagTokenIndex) >= NetFieldExportGroup->NetFieldExports.Num())
+		{
+			return false;
+		}
+		FNetFieldExport& NetFieldExport = NetFieldExportGroup->NetFieldExports[StateToken.GetIndex()];
+		FNetBitReader TempAr;
+		TArray<uint8>& Data = NetFieldExportGroup->NetFieldExports[TagTokenIndex].Blob;
+		TempAr.SetData(Data.GetData(),Data.NumBytes()*8);
+		StateTokenStore->ReadTokenData(TempAr, StateToken, &Map);
+		OutData = StateTokenStore->ResolveToken(StateToken);
+		return true;
+	}
+	
+	static bool AddStateToReplayExportGroup(DataStoreType* StateTokenStore, T& InData, TSharedPtr<FNetFieldExportGroup> NetFieldExportGroup, UPackageMapClient& Map)
+	{
+		if (!NetFieldExportGroup.IsValid())
+		{
+			return false;
+		}
+		FNetToken StateToken = StateTokenStore->GetOrCreateToken(InData);
+		if (static_cast<int32>(StateToken.GetIndex()) >= NetFieldExportGroup->NetFieldExports.Num())
+		{
+			NetFieldExportGroup->NetFieldExports.SetNum(StateToken.GetIndex() + 1, EAllowShrinking::No);
+		}
+		FNetFieldExport& NetFieldExport = NetFieldExportGroup->NetFieldExports[StateToken.GetIndex()];
+		if (NetFieldExport.bExported)
+		{
+			return false;
+		}
+		
+		FNetBitWriter TempAr;
+		StateTokenStore->WriteTokenData(TempAr, StateTokenStore->GetOrCreatePersistentState(InData), &Map);
+
+		NetFieldExport = FNetFieldExport(StateToken.GetIndex(),0,NAME_None,*TempAr.GetBuffer());
+		NetFieldExportGroup->bDirtyForReplay = true;
+		UE_LOG(LogNetToken, Verbose, TEXT("Replay> Exported for %s as NetFieldIndex %u"), *DataStoreType::GetTokenStoreName().ToString(), StateToken.GetIndex());
+		return true;
+	}
+
+	static bool SetupReplayExportGroup(DataStoreType* StateTokenStore, UPackageMapClient& PackageMapClient)
+	{
+		const FString StoreName = DataStoreType::GetTokenStoreName().ToString();
+		TSharedPtr<FNetFieldExportGroup> NetFieldExportGroup = PackageMapClient.GetNetFieldExportGroup(StoreName);
+		if (NetFieldExportGroup.IsValid())
+		{
+			return true;
+		}
+		NetFieldExportGroup = MakeShared<FNetFieldExportGroup>();
+		NetFieldExportGroup->PathName = StoreName;
+		NetFieldExportGroup->NetFieldExports.SetNum(StateTokenStore->StoredStates.Num());
+		for (auto& Pair : StateTokenStore->StoredStates)
+		{
+			AddStateToReplayExportGroup(StateTokenStore, Pair.Value, NetFieldExportGroup, PackageMapClient);
+		}
+		PackageMapClient.AddNetFieldExportGroup(StoreName, NetFieldExportGroup);
+		return true;
+	}
+};
+	
+}
\ No newline at end of file
diff --git a/Engine/Source/Runtime/Experimental/Iris/Core/Private/Iris/Core/IrisMemoryTracker.cpp b/Engine/Source/Runtime/Experimental/Iris/Core/Private/Iris/Core/IrisMemoryTracker.cpp
--- a/Engine/Source/Runtime/Experimental/Iris/Core/Private/Iris/Core/IrisMemoryTracker.cpp
+++ b/Engine/Source/Runtime/Experimental/Iris/Core/Private/Iris/Core/IrisMemoryTracker.cpp
@@ -7,9 +7,11 @@ DECLARE_LLM_MEMORY_STAT(TEXT("Iris"), STAT_IrisLLM, STATGROUP_LLMFULL);
 DECLARE_LLM_MEMORY_STAT(TEXT("IrisState"), STAT_IrisStateLLM, STATGROUP_LLMFULL);
 DECLARE_LLM_MEMORY_STAT(TEXT("IrisInitialization"), STAT_IrisInitializationLLM, STATGROUP_LLMFULL);
 DECLARE_LLM_MEMORY_STAT(TEXT("IrisConnection"), STAT_IrisConnectionLLM, STATGROUP_LLMFULL);
+DECLARE_LLM_MEMORY_STAT(TEXT("NetTokenStructState"), STAT_NetTokenStructStateLLM, STATGROUP_LLMFULL);
 
 // Tag name, display name, parent tag name, stat name, summary stat name
 LLM_DEFINE_TAG(Iris, NAME_None, NAME_None, GET_STATFNAME(STAT_IrisLLM), GET_STATFNAME(STAT_NetworkingSummaryLLM));
 LLM_DEFINE_TAG(IrisState, "State", "Iris", GET_STATFNAME(STAT_IrisStateLLM), GET_STATFNAME(STAT_NetworkingSummaryLLM));
 LLM_DEFINE_TAG(IrisInitialization, "Initialization", "Iris", GET_STATFNAME(STAT_IrisInitializationLLM), GET_STATFNAME(STAT_NetworkingSummaryLLM));
 LLM_DEFINE_TAG(IrisConnection, "Connection", "Iris", GET_STATFNAME(STAT_IrisConnectionLLM), GET_STATFNAME(STAT_NetworkingSummaryLLM));
+LLM_DEFINE_TAG(NetTokenStructState, "NetTokenStructState", "Iris", GET_STATFNAME(STAT_NetTokenStructStateLLM), GET_STATFNAME(STAT_NetworkingSummaryLLM));
\ No newline at end of file
diff --git a/Engine/Source/Runtime/Experimental/Iris/Core/Private/Iris/Serialization/StructNetSerializerUtil.cpp b/Engine/Source/Runtime/Experimental/Iris/Core/Private/Iris/Serialization/StructNetSerializerUtil.cpp
new file mode 100644
--- /dev/null
+++ b/Engine/Source/Runtime/Experimental/Iris/Core/Private/Iris/Serialization/StructNetSerializerUtil.cpp
@@ -0,0 +1,69 @@
+// Copyright Epic Games, Inc. All Rights Reserved.
+
+#include "Iris/Serialization/StructNetSerializerUtil.h"
+
+#include "Iris/Serialization/NetSerializers.h"
+#include "Iris/Serialization/NetSerializerArrayStorage.h"
+
+namespace UE::Net
+{
+	void WriteStruct(FNetSerializationContext& Context, NetSerializerValuePointer InValue, const FReplicationStateDescriptor* Descriptor)
+	{
+		if (!ensureAlwaysMsgf(Descriptor, TEXT("Replication State Descriptor cannot be null")))
+		{
+			return;
+		}
+		const FNetSerializer& Serializer = UE_NET_GET_SERIALIZER(FStructNetSerializer);
+		FStructNetSerializerConfig StructConfig;
+		StructConfig.StateDescriptor = Descriptor;
+		
+		FNetSerializerAlignedStorage QuantizedStorage;
+		QuantizedStorage.AdjustSize(Context, StructConfig.StateDescriptor->InternalSize, StructConfig.StateDescriptor->InternalAlignment);
+		{
+			FNetQuantizeArgs QuantizeArgs;
+			QuantizeArgs.Source = InValue;
+			QuantizeArgs.Target = NetSerializerValuePointer(QuantizedStorage.GetData());
+			QuantizeArgs.NetSerializerConfig = NetSerializerConfigParam(&StructConfig);
+			Serializer.Quantize(Context, QuantizeArgs);
+		}
+
+		{
+			FNetSerializeArgs SerializeArgs;
+			SerializeArgs.Version = Serializer.Version;
+			SerializeArgs.Source = NetSerializerValuePointer(QuantizedStorage.GetData());
+			SerializeArgs.NetSerializerConfig = NetSerializerConfigParam(&StructConfig);
+			Serializer.Serialize(Context, SerializeArgs);
+		}
+		QuantizedStorage.Free(Context);
+	}
+	
+	void ReadStruct(FNetSerializationContext& Context, NetSerializerValuePointer OutValue, const FReplicationStateDescriptor* Descriptor)
+	{
+		if (!ensureAlwaysMsgf(Descriptor, TEXT("Replication State Descriptor cannot be null")))
+		{
+			return;
+		}
+		const FNetSerializer& Serializer = UE_NET_GET_SERIALIZER(FStructNetSerializer);
+		FStructNetSerializerConfig StructConfig;
+		StructConfig.StateDescriptor = Descriptor;
+		FNetSerializerAlignedStorage QuantizedStorage;
+		QuantizedStorage.AdjustSize(Context, StructConfig.StateDescriptor->InternalSize, StructConfig.StateDescriptor->InternalAlignment);
+		{
+			FNetDeserializeArgs SerializeArgs;
+			SerializeArgs.Version = Serializer.Version;
+			SerializeArgs.Target = NetSerializerValuePointer(QuantizedStorage.GetData());
+			SerializeArgs.NetSerializerConfig = NetSerializerConfigParam(&StructConfig);
+			Serializer.Deserialize(Context, SerializeArgs);
+		}
+		
+		{
+			FNetDequantizeArgs DequantizeArgs;
+			DequantizeArgs.Source = NetSerializerValuePointer(QuantizedStorage.GetData());
+			DequantizeArgs.Target = OutValue;
+			DequantizeArgs.NetSerializerConfig = NetSerializerConfigParam(&StructConfig);
+			Serializer.Dequantize(Context, DequantizeArgs);
+		}
+		
+		QuantizedStorage.Free(Context);
+	}
+}
diff --git a/Engine/Source/Runtime/Experimental/Iris/Core/Public/Iris/Core/IrisMemoryTracker.h b/Engine/Source/Runtime/Experimental/Iris/Core/Public/Iris/Core/IrisMemoryTracker.h
--- a/Engine/Source/Runtime/Experimental/Iris/Core/Public/Iris/Core/IrisMemoryTracker.h
+++ b/Engine/Source/Runtime/Experimental/Iris/Core/Public/Iris/Core/IrisMemoryTracker.h
@@ -8,3 +8,5 @@ LLM_DECLARE_TAG_API(Iris, IRISCORE_API);
 LLM_DECLARE_TAG_API(IrisState, IRISCORE_API);
 LLM_DECLARE_TAG_API(IrisInitialization, IRISCORE_API);
 LLM_DECLARE_TAG_API(IrisConnection, IRISCORE_API);
+LLM_DECLARE_TAG_API(NetToken, IRISCORE_API);
+LLM_DECLARE_TAG_API(NetTokenStructState, IRISCORE_API);
\ No newline at end of file
diff --git a/Engine/Source/Runtime/Experimental/Iris/Core/Public/Iris/ReplicationSystem/NetTokenStructDefines.h b/Engine/Source/Runtime/Experimental/Iris/Core/Public/Iris/ReplicationSystem/NetTokenStructDefines.h
new file mode 100644
--- /dev/null
+++ b/Engine/Source/Runtime/Experimental/Iris/Core/Public/Iris/ReplicationSystem/NetTokenStructDefines.h
@@ -0,0 +1,64 @@
+// Copyright Epic Games, Inc. All Rights Reserved.
+
+#pragma once
+
+#include "UObject/Class.h"
+
+namespace UE::Net
+{
+
+template<typename T>
+struct TNetTokenStructOpsTypeTraits : public TStructOpsTypeTraitsBase2<T>
+{ 
+	enum
+	{ 
+		WithNetSerializer = true,
+		WithIdenticalViaEquality = true,
+		WithNetSharedSerialization = true,
+	};
+};
+
+};
+	
+// Optional, declares a default Native NetSerializer StructOpts and IRIS NetSerializer methods that serializes the struct as a NetToken by default.
+#define UE_NET_DECLARE_NAMED_NETTOKEN_STRUCT_SERIALIZERS(NAME, API) \
+	template<> struct TStructOpsTypeTraits<F##NAME> : public UE::Net::TNetTokenStructOpsTypeTraits<F##NAME> {}; \
+	namespace UE::Net \
+	{ \
+		UE_NET_DECLARE_SERIALIZER(F##NAME##NetSerializer, API);\
+	}
+
+// Optional, declares a default native NetSerializer method, GetTokenStoreName method and Equality operators (for equality via identity) for a NetToken Struct type
+#define UE_NET_NETTOKEN_GENERATED_BODY(NAME, API) \
+	inline static FName TokenStoreName = TEXT( PREPROCESSOR_TO_STRING(F##NAME) );\
+	API static FName GetTokenStoreName() \
+	{ \
+		return TokenStoreName; \
+	} \
+	API bool NetSerialize(FArchive& Ar, UPackageMap* Map, bool& bOutSuccess); \
+	API bool operator==(const F##NAME& Other) const; \
+	API bool operator!=(const F##NAME& Other) const; 
+
+// Optional, implements a default NetSerializer method and Native IRIS NetSerializer methods that works for NetToken types by default
+#define UE_NET_IMPLEMENT_NAMED_NETTOKEN_STRUCT_SERIALIZERS(NAME) \
+	bool F##NAME::NetSerialize(FArchive& Ar, UPackageMap* Map, bool& bOutSuccess) \
+	{ \
+		bOutSuccess = bOutSuccess && UE::Net::TStructNetTokenDataStoreHelper<F##NAME>::NetSerializeAndExportToken(Ar, Map, *static_cast<F##NAME*>(this)); \
+		return true; \
+	} \
+	bool F##NAME::operator==(const F##NAME& Other) const \
+	{ \
+		return GetUniqueKey() == Other.GetUniqueKey(); \
+	}; \
+	bool F##NAME::operator!=(const F##NAME& Other) const \
+	{ \
+		return GetUniqueKey() != Other.GetUniqueKey(); \
+	} \
+	namespace UE::Net { \
+		struct F##NAME##NetSerializer : public TStructAsNetTokenNetSerializerImpl<F##NAME> { \
+			static const uint32 Version = 0; \
+			static inline const ConfigType DefaultConfig = ConfigType(); \
+		};\
+		UE_NET_IMPLEMENT_SERIALIZER(F##NAME##NetSerializer); \
+		UE_NET_IMPLEMENT_FORWARDING_NETSERIALIZER_AND_REGISTRY_DELEGATES(NAME, F##NAME##NetSerializer); \
+	}
\ No newline at end of file
diff --git a/Engine/Source/Runtime/Experimental/Iris/Core/Public/Iris/ReplicationSystem/StructNetTokenDataStore.h b/Engine/Source/Runtime/Experimental/Iris/Core/Public/Iris/ReplicationSystem/StructNetTokenDataStore.h
new file mode 100644
--- /dev/null
+++ b/Engine/Source/Runtime/Experimental/Iris/Core/Public/Iris/ReplicationSystem/StructNetTokenDataStore.h
@@ -0,0 +1,490 @@
+// Copyright Epic Games, Inc. All Rights Reserved.
+
+#pragma once
+
+#include "Iris/ReplicationState/ReplicationStateDescriptorBuilder.h"
+#include "Iris/ReplicationSystem/NetTokenStore.h"
+#include "Iris/Serialization/NetSerializer.h"
+#include "Iris/Serialization/NetSerializerDelegates.h"
+#include "Iris/Serialization/StructNetSerializerUtil.h"
+#include "Iris/Core/IrisLog.h"
+#include "Iris/Core/IrisMemoryTracker.h"
+
+namespace UE::Net
+{
+	/*
+	 * This is a simplification of the process of building novel NetToken types and NetTokenDataStores for those types.
+	 *
+	 * The "General Idea" here is that you can define a USTRUCT that can be easily serialized as a NetToken instead of
+	 * serializing the entirety of the data structure every time. This is typically useful for large data structures where the data changes infrequently
+	 *
+	 * or
+	 *
+	 * where the data is commonly one of a small-ish set of permutations of that data. There are probably other useful cases as well, but these uses are readily apparent.
+	 *
+	 * Each NetToken refers to single, immutable, instance of the Struct you pass in/out. Every instance of the Struct that returns the same GetUniqueKey has the same NetToken.
+	 * Conceptually this means that you can't store pointers or references in them to sub pieces of data that change.
+	 *
+	 * You can't store references to other NetToken types, or other reference types... i.e. no properties that also are exported as NetTokens or ObjectReferences (NetGuids/NetRefHandles)
+	 * You should generally consider that you can only store primitive types in the NetToken struct.
+	 *
+	 * This means no: FNames, FGameplayTags or anything with a NetGuid.
+	 *
+	 * Additionally, this means that the data permutations that you are going to iterate are finite, or effectively finite. This isn't magic and it is not a compression algorithm.
+	 * We also aren't doing anything "smart" here about storing off the ShadowCopies of the input Structs. We very naively keep an entire copy of the struct in memory and there is no age out mechanism.
+	 * If you try to put "big" pieces of data in here we will just be making copies of that data. It also means complicated hashing mechanics for GetUniqueKey will eat up a lot of time.
+	 *
+	 * Again, its a simplification, not magic. You still want to be judicious in the data you are choosing to replicate.
+	 *
+	 * Most/All of the "Hard Bits" are implemented for you in a 'reasonable' default. You will need to implement custom versions if you need more concrete control.
+	 * - Replay Compatibility
+	 * - Default Serialization/Quantization
+	 * - Bookkeeping
+	 * - Maintaining Shadow Copies of the required data
+	 *
+	 * Example: Pretend that StringTokenStore didn't exist and I have a RPC that I want to call with a couple of FStrings that come from a set of known Strings.
+	 *
+	 *  UPROPERTY(Replicated)
+	 *  TArray<FString> KnownStrings =
+	 *  {
+	 *		TEXT("First string"),
+	 *		TEXT("Second string, this string is super long and would take up just a bunch of bytes to transmit. Just pretend its a huge string ok, thanks."),
+	 *		....
+	 *		TEXT("Last String")
+	 *  };
+	 *
+	 *  When asked to set this up we have the naive implementation:
+	 *  USTRUCT()
+	 *	struct FTwoStrings
+	 *  {
+	 *		UPROPERTY()
+	 *		FString Left;
+	 *		UPROPERTY()
+	 *		FString Right;	
+	 *  };
+	 *
+	 *  ...
+	 *  UFUCNTION(Server,Reliable)
+	 *  void SendSelectionToServer(FTwoStrings Data);
+	 *
+	 *	...
+	 *  void OnTick(double DeltaTime)
+	 *  {
+	 *		...
+	 *		SendSelectionToServer({Left,Right});
+	 *	}
+	 *
+	 *  Given this naive implementation we can (hopefully immediately) see the multitude of problems.
+	 *  In this simplistic example we could probably come up with several immediate fixes.
+	 *  Let's pretend we have to actually send this data on every tick.
+	 *
+	 *  We can spend a good deal of time coming up with custom serializer that would use significantly less data or some other representation for the data...
+	 *  For example:
+	*   USTRUCT()
+	 *	struct FTwoStringIdx
+	 *  {
+	 *		UPROPERTY()
+	 *		int32 Left;
+	 *		UPROPERTY()
+	 *		int32 Right;	
+	 *  };
+	 *
+	 *  But this suffers from a fundamental problem that isn't immediately obvious. KnownStrings is mutable and might be changing...
+	 *  And now we have a synchronization problem, where we are trying to maintain the state of KnownStrings, the index and the versioning of both.
+	 *  This rapidly grows in complexity beyond a simple "how do we make this use less data". We have to figure out a strategy on determining agreement on state and ordering and more.
+	 *  Ex: Which updates first? Do we need to maintain a history of known states in the event of delayed packets,etc,etc...
+	 *
+	 *  Even though the list of possible strings is known/finite, the mutable nature of the KnownStrings array in conjunction with the indeterminate ordering of network communication makes this a much more difficult problem.
+	 *
+	 *  This problem can be solved like this.
+	 *
+	 *	--------------
+	 *  TwoStrings.h
+	 *  --------------
+	 *  USTRUCT()
+	 *	struct FTwoStrings : FNetTokenStructBase
+	 *  {
+	 *      GENERATED_BODY()
+	 *      UE_NET_NETTOKEN_GENERATED_BODY(TwoStrings) // Note this is the struct name minus the 'F'
+	 *		UPROPERTY()
+	 *		FString Left;
+	 *		UPROPERTY()
+	 *		FString Right;
+	 *		uint64 GetUniqueKey()
+	 *		{
+	 *			FString Combo = Left+Right;
+	 *			return CityHash64((const char*)Combo, Combo.Num()*sizeof(TCHAR));
+	 *		}
+	 *  };
+	 *
+	 *  UE_NET_DECLARE_NAMED_NETTOKEN_STRUCT_SERIALIZERS(TwoStrings, YOUR_MODULES_API)
+	 *	--------------
+	 *  TwoStrings.cpp
+	 *  --------------
+	 *  UE_NET_IMPLEMENT_NAMED_NETTOKEN_STRUCT_SERIALIZERS(TwoStrings)
+	 *
+	 *	--------------
+	 *  DefaultEngine.ini
+	 *  --------------
+	 *  +ReservedTypeIds=(StoreTypeName="TwoStrings", TypeID=4) # Don't forget to use an unused TypeID.
+	 *  
+	 *	Having a UPROPERTY() or RPC UFUNCTION() will automatically serialize the type as the NetToken.
+	 *
+	 *  It is also possible to explicitly declare the DataStore, and use it in a custom NetSerialize.
+	 *  class FTwoStringDataStore : public UE::Net::TStructNetTokenDataStore<FTwoStrings>
+	 *  {
+	 *  public:
+	 *		explicit FTwoStringDataStore(UE::Net::FNetTokenStore& InTokenStore)
+	 *			: UE::Net::TStructNetTokenDataStore<FTwoStrings>(InTokenStore)
+	 *		{
+	 *		}
+	 *	};
+	 *
+	 *  It is possible to use these in custom serializers like this:
+	 *	FTwoStrings Data{Left,Right};
+	 *	FTwoStringDataStore::NetSerializeAndExportToken(Ar,Map,Data);
+	 */
+
+template <typename T>
+class TStructNetTokenDataStoreHelper;
+
+template <typename T>
+class TStructNetTokenDataStore : public FNetTokenDataStore
+{
+	using FNetTokenStoreKey = FNetTokenDataStore::FNetTokenStoreKey;
+	
+	UE_NONCOPYABLE(TStructNetTokenDataStore<T>);
+public:
+	using DataType = T;
+	
+	explicit TStructNetTokenDataStore<T>(FNetTokenStore& InTokenStore)
+	: FNetTokenDataStore(InTokenStore)
+	{
+	}
+	
+	// Create a token for input struct
+	FNetToken GetOrCreateToken(const T& InData)
+	{
+		FNetToken Result;
+		const FNetTokenStoreKey Key = GetOrCreatePersistentState(InData);
+		if (Key.IsValid())
+		{
+			Result = GetNetTokenFromKey(Key);
+			if (!Result.IsValid())
+			{
+				Result = CreateAndStoreTokenForKey(Key);
+				UE_LOG(LogNetToken, VeryVerbose, TEXT("TStructNetTokenDataStore::GetOrCreateToken %s CreatedToken %s."), *GetTokenStoreName().ToString(), *Result.ToString());
+			}
+		}
+		UE_LOG(LogNetToken, VeryVerbose, TEXT("TStructNetTokenDataStore::GetOrCreateToken %s GetOrCreateToken %s."), *GetTokenStoreName().ToString(), *Result.ToString());
+		return Result;
+	}
+
+	// Resolve NetToken, to resolve remote tokens RemoteTokenStoreState must be valid
+	const T& ResolveToken(FNetToken Token, const FNetTokenStoreState* RemoteTokenStoreState = nullptr) const
+	{
+		const FNetTokenStoreState* TokenStoreState = TokenStore.IsLocalToken(Token) ? TokenStore.GetLocalNetTokenStoreState() : RemoteTokenStoreState;
+		UE_LOG(LogNetToken, VeryVerbose, TEXT("TStructNetTokenDataStore::ResolveToken Starting up %s - %s. Local: %d, TokenStoreState: %p"), *GetTokenStoreName().ToString(), *Token.ToString(), TokenStore.IsLocalToken(Token), TokenStoreState);
+		if (Token.IsValid() && ensureMsgf(TokenStoreState, TEXT("TStructNetTokenDataStore::ResolveToken Needs valid TokenStoreState to resolve %s"), *Token.ToString()))
+		{
+			const FNetTokenStoreKey StoreKey = GetTokenKey(Token, *TokenStoreState);
+			UE_LOG(LogNetToken, VeryVerbose, TEXT("TStructNetTokenDataStore::ResolveToken Got Token Key %s - %s. StoreKeyIsValid: %d, StoreKeyID: %d, NumStoredStates: %d"), *GetTokenStoreName().ToString(), *Token.ToString(), StoreKey.IsValid(), StoreKey.GetKeyIndex(), StoredStates.Num());
+			if (StoreKey.IsValid() && StoredStates.Contains(StoreKey.GetKeyIndex()))
+			{
+				UE_LOG(LogNetToken, VeryVerbose, TEXT("TStructNetTokenDataStore::ResolveToken %s Succeeded %s."), *GetTokenStoreName().ToString(), *Token.ToString());
+				return StoredStates[StoreKey.GetKeyIndex()];
+			}
+			else
+			{
+				UE_LOG(LogNetToken, Error, TEXT("TStructNetTokenDataStore::ResolveToken %s failed to resolve %s."), *GetTokenStoreName().ToString(), *Token.ToString());
+			}
+		}
+		return InvalidState;
+	}
+
+	// Resolve a token received from remote
+	const T& ResolveRemoteToken(FNetToken Token, const FNetTokenStoreState& NetTokenStoreState) const
+	{ 
+		return ResolveToken(Token, &NetTokenStoreState);
+	}
+	
+	static FName GetTokenStoreName()
+	{
+		return T::GetTokenStoreName();
+	}
+
+	static const T& GetInvalidState()
+	{
+		return InvalidState;
+	}
+
+	// Need to declare as a delegate that is bound elsewhere to reduce the includes. We can't eliminate them, since the specialization module needs to include the engine parts necessary to make this work
+	// and we can't circular depend on the engine module to implement these details. We could write the helpers somewhere in the engine, though there isn't a great place to do this.
+	DECLARE_DELEGATE_ThreeParams(TNetSerializeTokenType, T&, FArchive&, UPackageMap*);
+	inline static TNetSerializeTokenType NetSerializeScriptDelegate;
+protected:
+	friend TStructNetTokenDataStoreHelper<T>;
+	
+	// Serialize data for a token, note there is not validation in this function
+	virtual void WriteTokenData(FNetSerializationContext& Context, FNetTokenStoreKey TokenStoreKey) const override
+	{
+		if (!TokenStoreKey.IsValid() || !StoredStates.Contains(TokenStoreKey.GetKeyIndex()))
+		{
+			UE_LOG(LogNetToken, Error, TEXT("TStructNetTokenDataStore::WriteTokenData %s KeyIndex: %d FAILED"),*GetTokenStoreName().ToString(),TokenStoreKey.GetKeyIndex());
+			Context.SetError(GNetError_InvalidValue);
+			return;
+		}
+		
+		if (!Descriptor.IsValid())
+		{
+			UE_LOG(LogNetToken, Error, TEXT("TStructNetTokenDataStore::WriteTokenData %s Failed. Serialization Descriptor Invalid."),*GetTokenStoreName().ToString());
+			Context.SetError(GNetError_InternalError);
+			return;
+		}
+		
+		UE_LOG(LogNetToken, VeryVerbose, TEXT("TStructNetTokenDataStore::WriteTokenData %s KeyIndex: %d Serializing"),*GetTokenStoreName().ToString(),TokenStoreKey.GetKeyIndex());
+		T TempValue = StoredStates.FindRef(TokenStoreKey.GetKeyIndex());
+		
+		WriteStruct(Context, NetSerializerValuePointer(&TempValue), Descriptor);
+		
+		if (Context.HasError())
+		{
+			UE_LOG(LogNetToken, Error, TEXT("TStructNetTokenDataStore::WriteTokenData, %s, FAILED"), *GetTokenStoreName().ToString());
+		}
+	}
+	
+	virtual void WriteTokenData(FArchive& Ar, FNetTokenStoreKey TokenStoreKey, UPackageMap* Map = nullptr) const override
+	{
+		if (!TokenStoreKey.IsValid()
+				|| !StoredStates.Contains(TokenStoreKey.GetKeyIndex()))
+		{
+			UE_LOG(LogNetToken, Error, TEXT("TStructNetTokenDataStore::WriteTokenData %s KeyIndex: %d FAILED"),*GetTokenStoreName().ToString(),TokenStoreKey.GetKeyIndex());
+			return;
+		}
+		UE_LOG(LogNetToken, VeryVerbose, TEXT("TStructNetTokenDataStore::WriteTokenData %s KeyIndex: %d Serializing"),*GetTokenStoreName().ToString(),TokenStoreKey.GetKeyIndex());
+		T TempValue = StoredStates.FindRef(TokenStoreKey.GetKeyIndex());
+		NetSerializeScriptDelegate.ExecuteIfBound(TempValue, Ar, Map);
+		if (Ar.IsError())
+		{
+			UE_LOG(LogNetToken, Error, TEXT("TStructNetTokenDataStore::WriteTokenData, %s, FAILED"), *GetTokenStoreName().ToString());
+		}	
+	}
+
+	// Read data for a token, returns a valid StoreKey if successful read
+	virtual FNetTokenStoreKey ReadTokenData(FNetSerializationContext& Context, const FNetToken& NetToken) override
+	{
+		if (!Descriptor.IsValid())
+		{
+			UE_LOG(LogNetToken, Error, TEXT("TStructNetTokenDataStore::ReadTokenData %s Failed. Serialization Descriptor Invalid."),*GetTokenStoreName().ToString());
+			return FNetTokenStoreKey();
+		}
+		
+		UE_LOG(LogNetToken, VeryVerbose, TEXT("TStructNetTokenDataStore::ReadTokenData %s"), *GetTokenStoreName().ToString());
+		T Value;
+		ReadStruct( Context, NetSerializerValuePointer(&Value), Descriptor);
+		if (!Context.HasErrorOrOverflow())
+		{
+			UE_LOG(LogNetToken, VeryVerbose, TEXT("TStructNetTokenDataStore::ReadTokenData %s, Succeeded"),*GetTokenStoreName().ToString());
+			return GetOrCreatePersistentState(Value);
+		}
+		UE_LOG(LogNetToken, Error, TEXT("TStructNetTokenDataStore::ReadTokenData, %s, FAILED"), *GetTokenStoreName().ToString());
+		return FNetTokenStoreKey();
+	}
+
+	virtual FNetTokenStoreKey ReadTokenData(FArchive& Ar, const FNetToken& NetToken, UPackageMap* Map = nullptr) override
+	{
+		T Value;
+		UE_LOG(LogNetToken, VeryVerbose, TEXT("TStructNetTokenDataStore::ReadTokenData %s"), *GetTokenStoreName().ToString());
+		NetSerializeScriptDelegate.ExecuteIfBound(Value, Ar, Map);
+		if (!Ar.IsError())
+		{
+			UE_LOG(LogNetToken, VeryVerbose, TEXT("TStructNetTokenDataStore::ReadTokenData %s, Succeeded"),*GetTokenStoreName().ToString());
+			return GetOrCreatePersistentState(Value);
+		}
+		UE_LOG(LogNetToken, Error, TEXT("TStructNetTokenDataStore::ReadTokenData, %s, FAILED"), *GetTokenStoreName().ToString());
+		return FNetTokenStoreKey();
+	}
+	
+	// Creates a persistent copy of the input struct 
+	FNetTokenStoreKey GetOrCreatePersistentState(const T& InNetTokenData)
+	{
+		LLM_SCOPE_BYTAG(NetTokenStructState);
+		const uint64 HashKey = InNetTokenData.GetUniqueKey();
+		if (HashToKey.Contains(HashKey))
+		{
+			UE_LOG(LogNetToken, VeryVerbose, TEXT("TStructNetTokenDataStore::GetOrCreatePersistentState %s FoundToken, Hash: %llu KeyIndex: %d"), *GetTokenStoreName().ToString(), HashKey, HashToKey[HashKey].GetKeyIndex());
+			return HashToKey[HashKey];
+		}
+
+		const FNetTokenStoreKey NewKey = GetNextNetTokenStoreKey();
+		if (NewKey.IsValid())
+		{
+			HashToKey.Add(HashKey, NewKey);
+			StoredStates.Add(NewKey.GetKeyIndex(), InNetTokenData);
+			UE_LOG(LogNetToken, VeryVerbose, TEXT("TStructNetTokenDataStore::GetOrCreatePersistentState %s Adding New, Hash: %llu KeyIndex: %d"), *GetTokenStoreName().ToString(), HashKey, HashToKey[HashKey].GetKeyIndex());
+			return NewKey;
+		}
+		UE_LOG(LogNetToken, VeryVerbose, TEXT("TStructNetTokenDataStore::GetOrCreatePersistentState %s ERROR, Hash: %llu"), *GetTokenStoreName().ToString(), HashKey);
+		return FNetTokenStoreKey();
+	}
+
+private:
+	inline static UScriptStruct* Struct;
+	inline static TRefCountPtr<const FReplicationStateDescriptor> Descriptor;
+	TMap<uint64, FNetTokenStoreKey> HashToKey;
+	TMap<uint32, T> StoredStates;
+	static inline T InvalidState = T();
+	
+	// Helper for NetSerializerDescriptor setup
+	class FNetSerializerRegistryDelegates final : private UE::Net::FNetSerializerRegistryDelegates
+	{
+	private:
+		virtual void OnPostFreezeNetSerializerRegistry() override
+		{
+			if (!Struct)
+			{
+				Struct = StaticStruct<T>();
+			}
+			if (!Descriptor.IsValid())
+			{				
+				FReplicationStateDescriptorBuilder::FParameters Params;
+
+				// TODO: this need flags to not pickup serializer using NetToken to write the data struct.
+				Params.SkipCheckForCustomNetSerializerForStruct = 1U;
+
+				Descriptor = FReplicationStateDescriptorBuilder::CreateDescriptorForStruct(Struct, Params);
+			}
+		}
+	};
+
+	inline static TStructNetTokenDataStore<T>::FNetSerializerRegistryDelegates NetSerializerRegistryDelegates;
+};
+	
+// Helper to implement a NetSerializer for a struct that should serialize using a NetTokenStore
+template <typename T, typename NetTokenDataStoreT = TStructNetTokenDataStore<T>>
+class TStructAsNetTokenNetSerializerImpl
+{
+public:
+	// Version
+	static const uint32 Version = 0;
+
+	typedef T SourceType;
+	typedef FNetToken QuantizedType;
+	typedef struct FNetSerializerConfig ConfigType;
+	static inline const ConfigType DefaultConfig = ConfigType();
+
+	// Iris
+	static void Serialize(FNetSerializationContext&, const FNetSerializeArgs&);
+	static void Deserialize(FNetSerializationContext&, const FNetDeserializeArgs&);
+	static void Quantize(FNetSerializationContext&, const FNetQuantizeArgs&);
+	static void Dequantize(FNetSerializationContext&, const FNetDequantizeArgs&);
+	static bool IsEqual(FNetSerializationContext&, const FNetIsEqualArgs&);
+};
+
+template <typename T, typename NetTokenDataStoreT>
+void TStructAsNetTokenNetSerializerImpl<T, NetTokenDataStoreT>::Serialize(FNetSerializationContext& Context, const FNetSerializeArgs& Args)
+{
+	const FNetToken& NetToken = *reinterpret_cast<FNetToken*>(Args.Source);
+	
+	// Tokens will differ, so we cannot store them in the default statehash.
+	if (Context.IsInitializingDefaultState())
+	{
+		return;
+	}
+
+	// Write token without type
+	Context.GetNetTokenStore()->WriteNetTokenWithKnownType<NetTokenDataStoreT>(Context, NetToken);
+
+	// Add to pending exports for later export
+	FNetTokenStore::AppendExport(Context, NetToken);
+}
+
+template <typename T, typename NetTokenDataStoreT>
+void TStructAsNetTokenNetSerializerImpl<T, NetTokenDataStoreT>::Deserialize(FNetSerializationContext& Context, const FNetDeserializeArgs& Args)
+{
+	FNetToken& NetToken = *reinterpret_cast<FNetToken*>(Args.Target);
+	NetToken = Context.GetNetTokenStore()->ReadNetTokenWithKnownType<NetTokenDataStoreT>(Context);
+}
+
+template <typename T, typename NetTokenDataStoreT>
+void TStructAsNetTokenNetSerializerImpl<T, NetTokenDataStoreT>::Quantize(FNetSerializationContext& Context, const FNetQuantizeArgs& Args)
+{
+	const SourceType& SourceValue = *reinterpret_cast<const SourceType*>(Args.Source);
+	FNetToken& TargetValue = *reinterpret_cast<FNetToken*>(Args.Target);
+
+	TargetValue = FNetToken();
+
+	if (NetTokenDataStoreT* NetTokenDataStore = Context.GetNetTokenStore()->GetDataStore<NetTokenDataStoreT>())
+	{
+		TargetValue = NetTokenDataStore->GetOrCreateToken(SourceValue);
+	}
+	else
+	{
+		UE_LOG(LogIris, Error, TEXT("TStructAsNetTokenNetSerializerImpl<T>::Quantize Could not find required FNetTokenDataStore %s"), *(typename NetTokenDataStoreT::GetTokenStoreName().ToString()));
+		ensure(false);
+	}
+}
+
+template <typename T, typename NetTokenDataStoreT>
+void TStructAsNetTokenNetSerializerImpl<T, NetTokenDataStoreT>::Dequantize(FNetSerializationContext& Context, const FNetDequantizeArgs& Args)
+{
+	const FNetToken& Source = *reinterpret_cast<const FNetToken*>(Args.Source);
+	SourceType& Target = *reinterpret_cast<SourceType*>(Args.Target);
+
+	if (NetTokenDataStoreT* NetTokenDataStore = Context.GetNetTokenStore()->GetDataStore<NetTokenDataStoreT>())
+	{
+		Target = NetTokenDataStore->ResolveToken(Source, Context.GetRemoteNetTokenStoreState());
+	}
+	else
+	{
+		UE_LOG(LogIris, Error, TEXT("TStructAsNetTokenNetSerializerImpl<T>::Deqquantize Could not find required FNetTokenDataStore %s"), *NetTokenDataStoreT::GetTokenStoreName().ToString());
+		ensure(false);
+
+		Target = SourceType();
+	}
+}
+template <typename T, typename NetTokenDataStoreT>
+bool TStructAsNetTokenNetSerializerImpl<T, NetTokenDataStoreT>::IsEqual(FNetSerializationContext& Context, const FNetIsEqualArgs& Args)
+{
+	if (Args.bStateIsQuantized)
+	{
+		const FNetToken& Value0 = *reinterpret_cast<const FNetToken*>(Args.Source0);
+		const FNetToken& Value1 = *reinterpret_cast<const FNetToken*>(Args.Source1);
+
+		// Need to compare actual Tags to properly compare non-auth and auth token
+		if (Value0.IsAssignedByAuthority() != Value1.IsAssignedByAuthority())
+		{
+			NetTokenDataStoreT* NetTokenDataStore = Context.GetNetTokenStore()->GetDataStore<NetTokenDataStoreT>();
+			const UE::Net::FNetTokenStoreState* RemoteNetTokenStoreState = Context.GetRemoteNetTokenStoreState();
+	
+			const SourceType Source0 = NetTokenDataStore->ResolveToken(Value0, RemoteNetTokenStoreState);
+			const SourceType Source1 = NetTokenDataStore->ResolveToken(Value1, RemoteNetTokenStoreState);
+				
+			if (Source0 != Source1)
+			{
+				return false;
+			}
+		}
+		else if (Value0 != Value1)
+		{
+			return false;
+		}
+		return true;
+	}
+	else
+	{
+		const SourceType& Value0 = *reinterpret_cast<SourceType*>(Args.Source0);
+		const SourceType& Value1 = *reinterpret_cast<SourceType*>(Args.Source1);
+
+		return Value0 == Value1;
+	}
+}
+
+} // End of namespace
+
+
+
+
+
+
+
diff --git a/Engine/Source/Runtime/Experimental/Iris/Core/Public/Iris/Serialization/StructNetSerializerUtil.h b/Engine/Source/Runtime/Experimental/Iris/Core/Public/Iris/Serialization/StructNetSerializerUtil.h
new file mode 100644
--- /dev/null
+++ b/Engine/Source/Runtime/Experimental/Iris/Core/Public/Iris/Serialization/StructNetSerializerUtil.h
@@ -0,0 +1,14 @@
+// Copyright Epic Games, Inc. All Rights Reserved.
+
+#pragma once
+
+#include "Iris/Serialization/NetSerializer.h"
+
+namespace UE::Net
+{
+	struct FReplicationStateDescriptor;
+	// Write a Struct value to the provided context using the provided Descriptor.
+	IRISCORE_API void WriteStruct(FNetSerializationContext& Context, NetSerializerValuePointer InValue, const FReplicationStateDescriptor* Descriptor);
+	// Read a Struct value to OutValue from the provided context using the provided Descriptor.
+	IRISCORE_API void ReadStruct(FNetSerializationContext& Context, NetSerializerValuePointer OutValue, const FReplicationStateDescriptor* Descriptor);
+}
