add optional proximity filtering for dataflow convex hull merging nodes

diff --git a/Engine/Plugins/Experimental/GeometryCollectionPlugin/Source/GeometryCollectionNodes/Private/Dataflow/GeometryCollectionUtilityNodes.cpp b/Engine/Plugins/Experimental/GeometryCollectionPlugin/Source/GeometryCollectionNodes/Private/Dataflow/GeometryCollectionUtilityNodes.cpp
--- a/Engine/Plugins/Experimental/GeometryCollectionPlugin/Source/GeometryCollectionNodes/Private/Dataflow/GeometryCollectionUtilityNodes.cpp
+++ b/Engine/Plugins/Experimental/GeometryCollectionPlugin/Source/GeometryCollectionNodes/Private/Dataflow/GeometryCollectionUtilityNodes.cpp
@@ -333,6 +333,8 @@ void FGenerateClusterConvexHullsFromLeafHullsDataflowNode::Evaluate(UE::Dataflow
 			HullMergeSettings.AllowMergesMethod = AllowMerges;
 			HullMergeSettings.bAllowMergingLeafHulls = GetValue(Context, &bAllowMergingLeafHulls);
 			HullMergeSettings.EmptySpace = bHasNegativeSpace ? &NegativeSpace : nullptr;
+			HullMergeSettings.ProximityFilter = MergeProximityFilter;
+			HullMergeSettings.ProximityDistanceThreshold = MergeProximityDistanceThreshold;
 
 			if (bHasSelectionFilter)
 			{
@@ -459,6 +461,8 @@ void FGenerateClusterConvexHullsFromChildrenHullsDataflowNode::Evaluate(UE::Data
 			HullMergeSettings.AllowMergesMethod = EAllowConvexMergeMethod::Any; // Note: Only 'Any' is supported for this node currently
 			HullMergeSettings.EmptySpace = bHasNegativeSpace ? &NegativeSpace : nullptr;
 			HullMergeSettings.bAllowMergingLeafHulls = GetValue(Context, &bAllowMergingLeafHulls);
+			HullMergeSettings.ProximityFilter = MergeProximityFilter;
+			HullMergeSettings.ProximityDistanceThreshold = MergeProximityDistanceThreshold;
 
 			if (bHasSelectionFilter)
 			{
@@ -615,6 +619,8 @@ void FMergeConvexHullsDataflowNode::Evaluate(UE::Dataflow::FContext& Context, co
 		HullMergeSettings.ErrorToleranceInCm = InErrorToleranceInCm;
 		HullMergeSettings.MaxConvexCount = InMaxConvexCount;
 		HullMergeSettings.ComputeEmptySpacePerBoneSettings = (bInProtectNegativeSpace && bComputeNegativeSpacePerBone) ? &NegativeSpaceSettings : nullptr;
+		HullMergeSettings.ProximityFilter = MergeProximityFilter;
+		HullMergeSettings.ProximityDistanceThreshold = MergeProximityDistanceThreshold;
 
 		UE::Geometry::FSphereCovering UsedNegativeSpace;
 		FGeometryCollectionConvexUtility::MergeHullsOnTransforms(InCollection, HullMergeSettings, bHasSelectionFilter, SelectionArray, &UsedNegativeSpace);
diff --git a/Engine/Plugins/Experimental/GeometryCollectionPlugin/Source/GeometryCollectionNodes/Public/Dataflow/GeometryCollectionUtilityNodes.h b/Engine/Plugins/Experimental/GeometryCollectionPlugin/Source/GeometryCollectionNodes/Public/Dataflow/GeometryCollectionUtilityNodes.h
--- a/Engine/Plugins/Experimental/GeometryCollectionPlugin/Source/GeometryCollectionNodes/Public/Dataflow/GeometryCollectionUtilityNodes.h
+++ b/Engine/Plugins/Experimental/GeometryCollectionPlugin/Source/GeometryCollectionNodes/Public/Dataflow/GeometryCollectionUtilityNodes.h
@@ -304,10 +304,18 @@ private:
 	UPROPERTY(EditAnywhere, Category = "Convex")
 	bool bPreferExternalCollisionShapes = true;
 
-	/** Method to determine which convex hull pairs can potentially be merged */
-	UPROPERTY(EditAnywhere, Category = "Convex")
+	/** Method to determine if the convex hulls from two separate bones can potentially be merged */
+	UPROPERTY(EditAnywhere, Category = "Convex", meta = (DisplayName = "Allow Merges Between Bones"))
 	EAllowConvexMergeMethod AllowMerges = EAllowConvexMergeMethod::ByProximity;
 
+	/** Filter to optionally only consider spatially close convex hulls for merges */
+	UPROPERTY(EditAnywhere, Category = "Convex")
+	EConvexHullProximityFilter MergeProximityFilter = EConvexHullProximityFilter::None;
+
+	/** If applying a convex hull proximity filter, the distance threshold to use for determining that two convex hulls are close enough to merge */
+	UPROPERTY(EditAnywhere, Category = "Convex", meta = (Units = cm, EditCondition = "MergeProximityFilter != EConvexHullProximityFilter::None"))
+	float MergeProximityDistanceThreshold = .1f;
+
 	/** Optional transform selection to compute cluster hulls on -- if not provided, all cluster hulls will be computed. */
 	UPROPERTY(meta = (DataflowInput))
 	FDataflowTransformSelection OptionalSelectionFilter;
@@ -408,6 +416,14 @@ private:
 	UPROPERTY(meta = (DataflowInput))
 	FDataflowTransformSelection OptionalSelectionFilter;
 
+	/** Filter to optionally only consider spatially close convex hulls for merges */
+	UPROPERTY(EditAnywhere, Category = "Convex")
+	EConvexHullProximityFilter MergeProximityFilter = EConvexHullProximityFilter::None;
+
+	/** If applying a convex hull proximity filter, the distance threshold to use for determining that two convex hulls are close enough to merge */
+	UPROPERTY(EditAnywhere, Category = "Convex", meta = (Units = cm, EditCondition = "MergeProximityFilter != EConvexHullProximityFilter::None"))
+	float MergeProximityDistanceThreshold = .1f;
+
 	/** Also allow the same hull merging process to run on leaf hulls (merging hulls on leaves in the selection) */
 	UPROPERTY(EditAnywhere, Category = "Convex", meta = (DataflowInput))
 	bool bAllowMergingLeafHulls = false;
@@ -524,6 +540,14 @@ public:
 	UPROPERTY(meta = (DataflowInput))
 	FDataflowTransformSelection OptionalSelectionFilter;
 
+	/** Filter to optionally only consider spatially close convex hulls for merges */
+	UPROPERTY(EditAnywhere, Category = "Convex")
+	EConvexHullProximityFilter MergeProximityFilter = EConvexHullProximityFilter::None;
+
+	/** If applying a convex hull proximity filter, the distance threshold to use for determining that two convex hulls are close enough to merge */
+	UPROPERTY(EditAnywhere, Category = "Convex", meta = (Units = cm, EditCondition = "MergeProximityFilter != EConvexHullProximityFilter::None"))
+	float MergeProximityDistanceThreshold = .1f;
+
 	/** Whether to use a sphere cover to define negative space that should not be covered by convex hulls */
 	UPROPERTY(EditAnywhere, Category = NegativeSpace, meta = (DataflowInput))
 	bool bProtectNegativeSpace = false;
diff --git a/Engine/Source/Runtime/Experimental/Chaos/Private/GeometryCollection/GeometryCollectionConvexUtility.cpp b/Engine/Source/Runtime/Experimental/Chaos/Private/GeometryCollection/GeometryCollectionConvexUtility.cpp
--- a/Engine/Source/Runtime/Experimental/Chaos/Private/GeometryCollection/GeometryCollectionConvexUtility.cpp
+++ b/Engine/Source/Runtime/Experimental/Chaos/Private/GeometryCollection/GeometryCollectionConvexUtility.cpp
@@ -1589,6 +1589,51 @@ bool CopyHulls(
 	return true;
 }
 
+// enum to pass to AddHullProximity's GetTransformBFn to indicate that all hulls use the same constant FTransform, so we can skip transforms and keep hulls in the shared space
+enum class EUseConstantTransform
+{
+	ConstantTransform
+};
+
+// Helper to add hull proximities for one hull to a set of neighboring hulls
+template<typename GetOtherTransformFn, typename HullIndexMapContainer = TArray<int32>>
+static void AddHullProximity(TArray<TPair<int32, int32>>& HullProximity, const EConvexHullProximityFilter ProximityFilter, const float ProximityDistanceThreshold,
+	TFunctionRef<const Chaos::FConvex* (int32)> GetConvex, const FTransform& TransformA, int32 ConvexA,
+	GetOtherTransformFn GetTransformBFn, int32 ConvexBIdxStart, int32 ConvexBIdxEnd, const HullIndexMapContainer* ConvexBIdxMap = nullptr)
+{
+	Chaos::FConvex::FAABB3Type BoxA;
+	if (ProximityFilter == EConvexHullProximityFilter::BoundingBox)
+	{
+		BoxA = GetConvex(ConvexA)->GetLocalBoundingBox();
+		if constexpr (!std::is_same_v<GetOtherTransformFn, EUseConstantTransform>)
+		{
+			BoxA = BoxA.TransformedAABB(TransformA);
+			BoxA.Thicken(ProximityDistanceThreshold);
+		}
+		else
+		{
+			float ThickenScaleFactor = 1.0f / FMath::Max(FMathf::ZeroTolerance, (float)TransformA.GetScale3D().GetAbsMax());
+			BoxA.Thicken(ProximityDistanceThreshold * ThickenScaleFactor);
+		}
+	}
+	for (int32 ConvexBIdx = ConvexBIdxStart; ConvexBIdx < ConvexBIdxEnd; ++ConvexBIdx)
+	{
+		const int32 ConvexB = ConvexBIdxMap ? (*ConvexBIdxMap)[ConvexBIdx] : ConvexBIdx;
+		if (ProximityFilter == EConvexHullProximityFilter::BoundingBox)
+		{
+			Chaos::FConvex::FAABB3Type BoxB = GetConvex(ConvexB)->GetLocalBoundingBox();
+			if constexpr (!std::is_same_v<GetOtherTransformFn, EUseConstantTransform>)
+			{
+				BoxB = BoxB.TransformedAABB(GetTransformBFn(ConvexB));
+			}
+			if (!BoxA.Intersects(BoxB))
+			{
+				continue;
+			}
+		}
+		HullProximity.Emplace(ConvexA, ConvexB);
+	}
+}
 
 }
 
@@ -2122,10 +2167,11 @@ void FGeometryCollectionConvexUtility::GenerateClusterConvexHullsFromLeafOrChild
 								{
 									for (int32 SourceHullIdxA : (*HullIndicesA))
 									{
-										for (int32 SourceHullIdxB : (*HullIndicesB))
-										{
-											HullProximity.Emplace(SourceHullIdxA, SourceHullIdxB);
-										}
+
+										AddHullProximity(HullProximity, Settings.ProximityFilter, Settings.ProximityDistanceThreshold,
+											[&Hulls](int32 Idx) -> const Chaos::FConvex* { return Hulls[Idx].Convex; }, Hulls[SourceHullIdxA].Transform,
+											SourceHullIdxA, [&Hulls](int32 Idx) -> const FTransform& { return Hulls[Idx].Transform; },
+											0, HullIndicesB->Num(), HullIndicesB);
 									}
 								}
 							}
@@ -2135,11 +2181,9 @@ void FGeometryCollectionConvexUtility::GenerateClusterConvexHullsFromLeafOrChild
 								for (int32 HullSubIdx = 0; HullSubIdx + 1 < HullIndicesA->Num(); ++HullSubIdx)
 								{
 									const int32 HullA = (*HullIndicesA)[HullSubIdx];
-									for (int32 ToHull = HullSubIdx + 1; ToHull < HullIndicesA->Num(); ++ToHull)
-									{
-										const int32 HullB = (*HullIndicesA)[ToHull];
-										HullProximity.Emplace(HullA, HullB);
-									}
+									AddHullProximity(HullProximity, Settings.ProximityFilter, Settings.ProximityDistanceThreshold,
+										[&Hulls](int32 Idx) -> const Chaos::FConvex* { return Hulls[Idx].Convex; }, Hulls[HullA].Transform,
+										HullA, EUseConstantTransform::ConstantTransform, HullSubIdx + 1, HullIndicesA->Num(), HullIndicesA);
 								}
 							}
 						}
@@ -2161,10 +2205,11 @@ void FGeometryCollectionConvexUtility::GenerateClusterConvexHullsFromLeafOrChild
 								{
 									for (int32 SourceHullIdx : SourceHullIndices)
 									{
-										for (int32 NeighborHullIdx : (*NeighborHullIndices))
-										{
-											HullProximity.Emplace(SourceHullIdx, NeighborHullIdx);
-										}
+										AddHullProximity(HullProximity, Settings.ProximityFilter, Settings.ProximityDistanceThreshold,
+											[&Hulls](int32 Idx) -> const Chaos::FConvex* { return Hulls[Idx].Convex; }, Hulls[SourceHullIdx].Transform,
+											SourceHullIdx, [&Hulls](int32 Idx) -> const FTransform& { return Hulls[Idx].Transform; },
+											0, NeighborHullIndices->Num(), NeighborHullIndices);
+
 									}
 								}
 							}
@@ -2174,11 +2219,9 @@ void FGeometryCollectionConvexUtility::GenerateClusterConvexHullsFromLeafOrChild
 								for (int32 HullSubIdx = 0; HullSubIdx + 1 < SourceHullIndices.Num(); ++HullSubIdx)
 								{
 									const int32 HullA = SourceHullIndices[HullSubIdx];
-									for (int32 ToHull = HullSubIdx + 1; ToHull < SourceHullIndices.Num(); ++ToHull)
-									{
-										const int32 HullB = SourceHullIndices[ToHull];
-										HullProximity.Emplace(HullA, HullB);
-									}
+									AddHullProximity(HullProximity, Settings.ProximityFilter, Settings.ProximityDistanceThreshold,
+										[&Hulls](int32 Idx) -> const Chaos::FConvex* { return Hulls[Idx].Convex; }, Hulls[HullA].Transform,
+										HullA, EUseConstantTransform::ConstantTransform, HullSubIdx + 1, SourceHullIndices.Num(), &SourceHullIndices);
 								}
 							}
 						}
@@ -2288,10 +2331,9 @@ void FGeometryCollectionConvexUtility::MergeHullsOnTransforms(FManagedArrayColle
 		TArray<TPair<int32, int32>> HullProximity;
 		for (int32 ConvexA = 0; ConvexA < InitialNumConvex; ++ConvexA)
 		{
-			for (int32 ConvexB = ConvexA + 1; ConvexB < InitialNumConvex; ++ConvexB)
-			{
-				HullProximity.Emplace(ConvexA, ConvexB);
-			}
+			AddHullProximity(HullProximity, Settings.ProximityFilter, Settings.ProximityDistanceThreshold,
+					[&Hulls](int32 Idx) -> const Chaos::FConvex* { return Hulls[Idx]; }, GlobalTransforms[TransformIndex],
+					ConvexA, EUseConstantTransform::ConstantTransform, ConvexA + 1, InitialNumConvex);
 		}
 
 		// try to merge the leaf convex into a simpler set of convex
diff --git a/Engine/Source/Runtime/Experimental/Chaos/Public/GeometryCollection/GeometryCollectionConvexUtility.h b/Engine/Source/Runtime/Experimental/Chaos/Public/GeometryCollection/GeometryCollectionConvexUtility.h
--- a/Engine/Source/Runtime/Experimental/Chaos/Public/GeometryCollection/GeometryCollectionConvexUtility.h
+++ b/Engine/Source/Runtime/Experimental/Chaos/Public/GeometryCollection/GeometryCollectionConvexUtility.h
@@ -40,15 +40,27 @@ enum class EGenerateConvexMethod : uint8
 	IntersectExternalWithComputed
 };
 
-UENUM()
+// Controls which bones are allowed to merge their convex hulls
+UENUM(meta=(DisplayName="Allow Convex Bone Merge Method"))
 enum class EAllowConvexMergeMethod : uint8
 {
-	// Only allow merging convex hulls that are in proximity
-	ByProximity,
+	// Only allow merging convex hulls of bones that are in proximity
+	ByProximity UMETA(DisplayName = "By Bone Proximity"),
 	// Allow any pair of convex hulls to merge
 	Any
 };
 
+// Controls which individual convex hulls are considered for merges
+UENUM()
+enum class EConvexHullProximityFilter : uint8
+{
+	// Do not require convex hull geometry to be in proximity to merge
+	// Note: May still use bone proximity to limit possible merges in the cluster merge case, via the Allow Convex Bone Merge Method options
+	None,
+	// Only consider convex hull merges if their bounding boxes are in proximity
+	BoundingBox
+};
+
 namespace UE::GeometryCollectionConvexUtility
 {
 	// Used to pass computed convex hulls along with metadata
@@ -126,7 +138,12 @@ public:
 		bool bUseExternalCollisionIfAvailable = true;
 		// Option to apply the same hull merging settings to leaf hulls
 		bool bAllowMergingLeafHulls = false;
+		// How to allow merges between hulls on different bones
 		EAllowConvexMergeMethod AllowMergesMethod = EAllowConvexMergeMethod::ByProximity;
+		// How to filter by proximity for all merges (including within bones)
+		EConvexHullProximityFilter ProximityFilter = EConvexHullProximityFilter::None;
+		float ProximityDistanceThreshold = 0.1f;
+
 		UE::Geometry::FSphereCovering* EmptySpace = nullptr;
 	};
 
@@ -144,6 +161,9 @@ public:
 
 		// Optional settings to compute targeted empty space per-bone
 		UE::Geometry::FNegativeSpaceSampleSettings* ComputeEmptySpacePerBoneSettings = nullptr;
+
+		EConvexHullProximityFilter ProximityFilter = EConvexHullProximityFilter::None;
+		float ProximityDistanceThreshold = 0.1f;
 	};
 	// Merge convex hulls that are currently on each (selected) transform. If convex hulls are not present, does nothing.
 	// @params OptionalSphereCoveringOut		If non-null, will be filled with spheres from all used sphere covering.
