[NetTokens] Remove extra typename use for compatibility with clang compiler.

diff --git a/Engine/Source/Runtime/Experimental/Iris/Core/Public/Iris/ReplicationSystem/StructNetTokenDataStore.h b/Engine/Source/Runtime/Experimental/Iris/Core/Public/Iris/ReplicationSystem/StructNetTokenDataStore.h
--- a/Engine/Source/Runtime/Experimental/Iris/Core/Public/Iris/ReplicationSystem/StructNetTokenDataStore.h
+++ b/Engine/Source/Runtime/Experimental/Iris/Core/Public/Iris/ReplicationSystem/StructNetTokenDataStore.h
@@ -420,7 +420,7 @@ void TStructAsNetTokenNetSerializerImpl<T, NetTokenDataStoreT>::Quantize(FNetSer
 	}
 	else
 	{
-		UE_LOG(LogIris, Error, TEXT("TStructAsNetTokenNetSerializerImpl<T>::Quantize Could not find required FNetTokenDataStore %s"), *(typename NetTokenDataStoreT::GetTokenStoreName().ToString()));
+		UE_LOG(LogIris, Error, TEXT("TStructAsNetTokenNetSerializerImpl<T>::Quantize Could not find required FNetTokenDataStore %s"), *NetTokenDataStoreT::GetTokenStoreName().ToString());
 		ensure(false);
 	}
 }
