Fix Chaos solver actor missing export

diff --git a/Engine/Source/Runtime/Experimental/ChaosSolverEngine/Public/Chaos/ChaosSolverActor.h b/Engine/Source/Runtime/Experimental/ChaosSolverEngine/Public/Chaos/ChaosSolverActor.h
--- a/Engine/Source/Runtime/Experimental/ChaosSolverEngine/Public/Chaos/ChaosSolverActor.h
+++ b/Engine/Source/Runtime/Experimental/ChaosSolverEngine/Public/Chaos/ChaosSolverActor.h
@@ -163,15 +163,15 @@ public:
 	CHAOSSOLVERENGINE_API virtual void SetSolverActive(bool bActive);
 
 	//~ Begin IDataflowPhysicsSolverInterface interface
-	virtual FString GetSimulationName() const override {return GetName();};
-	virtual FDataflowSimulationAsset& GetSimulationAsset() override {return SimulationAsset;};
-	virtual const FDataflowSimulationAsset& GetSimulationAsset() const override {return SimulationAsset;};
-	virtual FDataflowSimulationProxy* GetSimulationProxy() override {return &RigidSolverProxy;}
-	virtual const FDataflowSimulationProxy* GetSimulationProxy() const  override {return &RigidSolverProxy;}
-	virtual void BuildSimulationProxy() override;
-	virtual void ResetSimulationProxy() override;
-	virtual void WriteToSimulation(const float DeltaTime, const bool bAsyncTask) override;
-	virtual void ReadFromSimulation(const float DeltaTime, const bool bAsyncTask) override;
+	CHAOSSOLVERENGINE_API virtual FString GetSimulationName() const override {return GetName();};
+	CHAOSSOLVERENGINE_API virtual FDataflowSimulationAsset& GetSimulationAsset() override {return SimulationAsset;};
+	CHAOSSOLVERENGINE_API virtual const FDataflowSimulationAsset& GetSimulationAsset() const override {return SimulationAsset;};
+	CHAOSSOLVERENGINE_API virtual FDataflowSimulationProxy* GetSimulationProxy() override {return &RigidSolverProxy;}
+	CHAOSSOLVERENGINE_API virtual const FDataflowSimulationProxy* GetSimulationProxy() const  override {return &RigidSolverProxy;}
+	CHAOSSOLVERENGINE_API virtual void BuildSimulationProxy() override;
+	CHAOSSOLVERENGINE_API virtual void ResetSimulationProxy() override;
+	CHAOSSOLVERENGINE_API virtual void WriteToSimulation(const float DeltaTime, const bool bAsyncTask) override;
+	CHAOSSOLVERENGINE_API virtual void ReadFromSimulation(const float DeltaTime, const bool bAsyncTask) override;
 	//~ End IDataflowPhysicsSolverInterface interface
 
 	/*
