CIS fix: IslandManager shadow warning

diff --git a/Engine/Source/Runtime/Experimental/Chaos/Private/Chaos/Island/IslandManager.cpp b/Engine/Source/Runtime/Experimental/Chaos/Private/Chaos/Island/IslandManager.cpp
--- a/Engine/Source/Runtime/Experimental/Chaos/Private/Chaos/Island/IslandManager.cpp
+++ b/Engine/Source/Runtime/Experimental/Chaos/Private/Chaos/Island/IslandManager.cpp
@@ -2781,9 +2781,9 @@ namespace Chaos::Private
 	{
 		if (!ParticlesToWake.IsEmpty())
 		{
-			const auto UpdateShapeTransforms = [](const TArray<FGeometryParticleHandle*>& Particles)
+			const auto UpdateShapeTransforms = [](const TArray<FGeometryParticleHandle*>& ParticlesToWake)
 				{
-					for (FGeometryParticleHandle* Particle : Particles)
+					for (FGeometryParticleHandle* Particle : ParticlesToWake)
 					{
 						for (const TUniquePtr<FPerShapeData>& Shape : Particle->ShapesArray())
 						{
