Chaos Cloth Asset - Made the MaxDistance Low/High values animatable.

[RN] Chaos Cloth Asset - Made the MaxDistance Low/High values animatable using the ClothOutfitInteractor.

diff --git a/Engine/Plugins/ChaosClothAssetEditor/Source/ChaosClothAssetDataflowNodes/Public/ChaosClothAsset/SimulationMaxDistanceConfigNode.h b/Engine/Plugins/ChaosClothAssetEditor/Source/ChaosClothAssetDataflowNodes/Public/ChaosClothAsset/SimulationMaxDistanceConfigNode.h
--- a/Engine/Plugins/ChaosClothAssetEditor/Source/ChaosClothAssetDataflowNodes/Public/ChaosClothAsset/SimulationMaxDistanceConfigNode.h
+++ b/Engine/Plugins/ChaosClothAssetEditor/Source/ChaosClothAssetDataflowNodes/Public/ChaosClothAsset/SimulationMaxDistanceConfigNode.h
@@ -22,7 +22,7 @@ public:
 	 * Otherwise all particles are considered to have a zero weight, and only the Low value is meaningful.
 	 */
 	UPROPERTY(EditAnywhere, Category = "MaxDistance Properties", Meta = (UIMin = "0", UIMax = "100", ClampMin = "0", InteractorName = "MaxDistance"))
-	FChaosClothAssetWeightedValue MaxDistance = { true, 0.f, 100.f, TEXT("MaxDistance") };
+	FChaosClothAssetWeightedValue MaxDistance = { true, 0.f, 100.f, TEXT("MaxDistance"), true };
 
 	/**
 	 * Selection set of SimVertices3D that should be made kinematic regardless of the MaxDistance map value.
diff --git a/Engine/Source/Runtime/Experimental/Chaos/Public/Chaos/CollectionPropertyFacade.h b/Engine/Source/Runtime/Experimental/Chaos/Public/Chaos/CollectionPropertyFacade.h
--- a/Engine/Source/Runtime/Experimental/Chaos/Public/Chaos/CollectionPropertyFacade.h
+++ b/Engine/Source/Runtime/Experimental/Chaos/Public/Chaos/CollectionPropertyFacade.h
@@ -18,7 +18,7 @@ namespace Chaos::Softs
 		Animatable = 1 << 1,  /** Whether this property needs to be set at every frame. This flag is ignored when the Intrinsic flag is also set. */
 		Legacy = 1 << 2,  /** Whether this property has been set by a legacy system predating the property collection. Can be useful for overriding/upgrading some properties post conversion. */
 		Interpolable = 1 << 3,  /** Whether this property can be interpolated. Used to allow interpolation when merging float properties */
-		Intrinsic = 1 << 4,  /** Whether this property is intrinsically built into the simulated object model, rather than affecting the simulation itself (see Animatable in this case). Changing this property requires a re-construction of the simulated object model to be effective. Implies non Animatable. */
+		Intrinsic = 1 << 4,  /** Whether this property is intrinsically built into the simulated object model, rather than only affecting the simulation itself (see Animatable in this case). Changing this property requires a re-construction of the simulated object model to be fully effective. */
 		//~ Add new flags above this line
 		StringDirty = 1 << 6,  /** Whether this property's string has changed and needs to be updated at the next frame. */
 		Dirty = 1 << 7  /** Whether this property's value has changed and needs to be updated at the next frame. */
@@ -82,7 +82,7 @@ namespace Chaos::Softs
 		ECollectionPropertyFlags GetFlags(int32 KeyIndex) const { return (ECollectionPropertyFlags)FlagsArray[KeyIndex]; }
 
 		bool IsEnabled(int32 KeyIndex) const { return HasAnyFlags(KeyIndex, ECollectionPropertyFlags::Enabled); }
-		bool IsAnimatable(int32 KeyIndex) const { return HasAnyFlags(KeyIndex, ECollectionPropertyFlags::Animatable) && !HasAnyFlags(KeyIndex, ECollectionPropertyFlags::Intrinsic); }
+		bool IsAnimatable(int32 KeyIndex) const { return HasAnyFlags(KeyIndex, ECollectionPropertyFlags::Animatable); }
 		bool IsLegacy(int32 KeyIndex) const { return HasAnyFlags(KeyIndex, ECollectionPropertyFlags::Legacy); }
 		bool IsIntrinsic(int32 KeyIndex) const { return HasAnyFlags(KeyIndex, ECollectionPropertyFlags::Intrinsic); }
 		bool IsStringDirty(int32 KeyIndex) const { return HasAnyFlags(KeyIndex, ECollectionPropertyFlags::StringDirty); }
