[ChaosModularVehicles] Revert nettokens from Input until dependency chain can be resolved.

diff --git a/Engine/Plugins/Experimental/ChaosModularVehicle/Source/ChaosModularVehicleEngine/Private/ChaosModularVehicle/ChaosSimModuleManager.cpp b/Engine/Plugins/Experimental/ChaosModularVehicle/Source/ChaosModularVehicleEngine/Private/ChaosModularVehicle/ChaosSimModuleManager.cpp
--- a/Engine/Plugins/Experimental/ChaosModularVehicle/Source/ChaosModularVehicleEngine/Private/ChaosModularVehicle/ChaosSimModuleManager.cpp
+++ b/Engine/Plugins/Experimental/ChaosModularVehicle/Source/ChaosModularVehicleEngine/Private/ChaosModularVehicle/ChaosSimModuleManager.cpp
@@ -5,7 +5,7 @@
 #include "PBDRigidsSolver.h"
 #include "GameFramework/HUD.h" // for ShowDebugInfo
 #include "Physics/Experimental/PhysScene_Chaos.h"
-#include "SimModule/ModuleInputTokenStore.h"
+
 TMap<FPhysScene*, FChaosSimModuleManager*> FChaosSimModuleManager::SceneToModuleManagerMap;
 
 FDelegateHandle FChaosSimModuleManager::OnPostWorldInitializationHandle;
@@ -120,12 +120,7 @@ void FChaosSimModuleManager::RegisterNetTokenDataStores(UNetDriver* InNetDriver)
 	{
 		if (UE::Net::FNetTokenStore* TokenStore = InNetDriver->GetNetTokenStore())
 		{
-			typedef UE::Net::TStructNetTokenDataStore<FModuleInputNetTokenData> FModuleInputNetTokenStore;
 			typedef UE::Net::TStructNetTokenDataStore<FNetworkModularVehicleStateNetTokenData> FNetworkModularVehicleStateNetTokenStore;
-			if (!TokenStore->GetDataStore<FModuleInputNetTokenStore>())
-			{
-				TokenStore->CreateAndRegisterDataStore<FModuleInputNetTokenStore>();
-			}
 			if (!TokenStore->GetDataStore<FNetworkModularVehicleStateNetTokenStore>())
 			{
 				TokenStore->CreateAndRegisterDataStore<FNetworkModularVehicleStateNetTokenStore>();
diff --git a/Engine/Plugins/Experimental/ChaosModularVehicle/Source/ChaosModularVehicleEngine/Private/ChaosModularVehicle/ChaosSimModuleManagerAsyncCallback.cpp b/Engine/Plugins/Experimental/ChaosModularVehicle/Source/ChaosModularVehicleEngine/Private/ChaosModularVehicle/ChaosSimModuleManagerAsyncCallback.cpp
--- a/Engine/Plugins/Experimental/ChaosModularVehicle/Source/ChaosModularVehicleEngine/Private/ChaosModularVehicle/ChaosSimModuleManagerAsyncCallback.cpp
+++ b/Engine/Plugins/Experimental/ChaosModularVehicle/Source/ChaosModularVehicleEngine/Private/ChaosModularVehicle/ChaosSimModuleManagerAsyncCallback.cpp
@@ -258,17 +258,10 @@ bool FNetworkModularVehicleInputs::NetSerialize(FArchive& Ar, class UPackageMap*
 
 	Ar.SerializeBits(&VehicleInputs.Reverse, 1);
 	Ar.SerializeBits(&VehicleInputs.KeepAwake, 1);
-
-	FNetworkModularVehicleInputs* DeltaSource = static_cast<FNetworkModularVehicleInputs*>(DeltaSourceData);
-	if (DeltaSource && DeltaSource->ServerFrame != -1 && ChaosModularVehicleCVars::bEnableStateReducedBandwidth)
-	{
-		VehicleInputs.Container.DeltaNetSerialize(Ar, Map, bOutSuccess,DeltaSource->VehicleInputs.Container);
-	}
-	else
+	
 	{
 		VehicleInputs.Container.Serialize(Ar, Map, bOutSuccess);
 	}
-	
 
 	bOutSuccess = true;
 	return bOutSuccess;
diff --git a/Engine/Source/Runtime/Experimental/ChaosVehicles/ChaosVehiclesCore/ChaosVehiclesCore.Build.cs b/Engine/Source/Runtime/Experimental/ChaosVehicles/ChaosVehiclesCore/ChaosVehiclesCore.Build.cs
--- a/Engine/Source/Runtime/Experimental/ChaosVehicles/ChaosVehiclesCore/ChaosVehiclesCore.Build.cs
+++ b/Engine/Source/Runtime/Experimental/ChaosVehicles/ChaosVehiclesCore/ChaosVehiclesCore.Build.cs
@@ -10,14 +10,11 @@ namespace UnrealBuildTool.Rules
                 new string[] {
                 "Core",
                 "CoreUObject",
-				"Chaos",
-				"NetCore",
-				"Engine"
+				"Chaos"
 				}
             ); 
 
 			PrivateDefinitions.Add("CHAOS_INCLUDE_LEVEL_1=1");
-			SetupIrisSupport(Target);
 		}
     }
 }
diff --git a/Engine/Source/Runtime/Experimental/ChaosVehicles/ChaosVehiclesCore/Private/SimModule/ModuleInput.cpp b/Engine/Source/Runtime/Experimental/ChaosVehicles/ChaosVehiclesCore/Private/SimModule/ModuleInput.cpp
--- a/Engine/Source/Runtime/Experimental/ChaosVehicles/ChaosVehiclesCore/Private/SimModule/ModuleInput.cpp
+++ b/Engine/Source/Runtime/Experimental/ChaosVehicles/ChaosVehiclesCore/Private/SimModule/ModuleInput.cpp
@@ -1,22 +1,12 @@
 // Copyright Epic Games, Inc. All Rights Reserved.
 
 #include "SimModule/ModuleInput.h"
-#include "SimModule/ModuleInputTokenStore.h"
-#include "Net/Iris/ReplicationSystem/StructNetTokenDataStoreHelper.h"
 #include "Net/Core/NetBitArray.h"
 
 TArray<struct FModuleInputSetup>* FScopedModuleInputInitializer::InitSetupData = nullptr;
 
 DEFINE_LOG_CATEGORY(LogModularInput);
 
-namespace ChaosModularVehicleCVars
-{
-	bool bEnableInputReducedBandwidth = false;
-	FAutoConsoleVariableRef EnableInputReducedBandwidth(TEXT("p.ModularVehicle.EnableInputReducedBandwidth"), bEnableInputReducedBandwidth, TEXT("Enable/Disable NetTokens and DeltaSerialization path for Modular Vehicle Input. Default: false"));
-};
-
-UE_NET_IMPLEMENT_NAMED_NETTOKEN_STRUCT_SERIALIZERS(ModuleInputNetTokenData)
-
 bool FModuleInputValue::IsNonZero(float Tolerance) const
 {
 	switch (ValueType)
@@ -122,13 +112,12 @@ void FModuleInputValue::SetMagnitude(float NewSize)
 
 void FModuleInputValue::Serialize(FArchive& Ar, UPackageMap* Map, bool& bOutSuccess)
 {
+	// #TODO: send only value changes/deltas
+
 	bool bIsNonZero = Ar.IsLoading() ? false : IsQuantizedNonZero();
 	Ar.SerializeBits(&bIsNonZero, 1);
-	if (!ChaosModularVehicleCVars::bEnableInputReducedBandwidth)
-	{
-		Ar.SerializeBits(&ValueType, 3);
-		Ar.SerializeBits(&bApplyInputDecay, 1);
-	}
+	Ar.SerializeBits(&ValueType, 3);
+	Ar.SerializeBits(&bApplyInputDecay, 1);
 
 	if (bIsNonZero)
 	{
@@ -171,31 +160,6 @@ void FModuleInputValue::Serialize(FArchive& Ar, UPackageMap* Map, bool& bOutSucc
 	bOutSuccess = true;
 }
 
-void FModuleInputValue::DeltaNetSerialize(FArchive& Ar, UPackageMap* Map, bool& bOutSuccess, const FModuleInputValue& PreviousInputValue)
-{
-	bool bHasSamePreviousValue = Value == PreviousInputValue.Value;
-	bool bHasSamePreviousType = ValueType == PreviousInputValue.ValueType;
-	bool bHasSamePreviousValueInt = ValueInt == PreviousInputValue.ValueInt;
-	bool bIsSameAsPrevious = bHasSamePreviousValue && bHasSamePreviousType && bHasSamePreviousValueInt;
-	Ar.SerializeBits(&bIsSameAsPrevious, 1);
-	if (bIsSameAsPrevious)
-	{
-		bOutSuccess = true;
-		if (Ar.IsSaving())
-		{
-			return;
-		}
-
-		ConvertToType(PreviousInputValue.ValueType);
-		Value = PreviousInputValue.Value;
-		ValueInt = PreviousInputValue.ValueInt;
-		bApplyInputDecay = PreviousInputValue.bApplyInputDecay;
-		return;
-	}
-
-	Serialize(Ar, Map, bOutSuccess);
-}
-
 void FModuleInputValue::Lerp(const FModuleInputValue& Min, const FModuleInputValue& Max, float Alpha)
 {
 	switch (GetValueType())
@@ -368,10 +332,6 @@ void FModuleInputContainer::ZeroValues()
 
 void FModuleInputContainer::Serialize(FArchive& Ar, UPackageMap* Map, bool& bOutSuccess)
 {
-	if (ChaosModularVehicleCVars::bEnableInputReducedBandwidth)
-	{
-		return SerializeReduced(Ar, Map, bOutSuccess);
-	}
 	using namespace UE::Net;
 
 	bOutSuccess = true;
@@ -392,65 +352,6 @@ void FModuleInputContainer::Serialize(FArchive& Ar, UPackageMap* Map, bool& bOut
 	return;
 }
 
-void FModuleInputContainer::SerializeReduced(FArchive& Ar, UPackageMap* Map, bool& bOutSuccess)
-{
-	using namespace UE::Net;
-    bOutSuccess = true;
-    
-    FModuleInputNetTokenData InputStateData;
-    InputStateData.Init(InputValues);
-    TStructNetTokenDataStoreHelper<FModuleInputNetTokenData>::NetSerializeAndExportToken(Ar,Map,InputStateData);
-
-    uint32 Number = InputStateData.Types.Num();
-    //Ar.SerializeIntPacked(Number);
-
-    if (Ar.IsLoading())
-    {
-    	InputValues.SetNum(Number);
-    }
-
-    for (uint32 I = 0; I < Number; I++)
-    {
-    	InputValues[I].ConvertToType(static_cast<EModuleInputValueType>(InputStateData.Types[I]));
-    	InputValues[I].SetApplyInputDecay(InputStateData.DecayValues[I]);
-    	InputValues[I].Serialize(Ar, Map, bOutSuccess);
-    }
-
-    return;
-}
-
-void FModuleInputContainer::DeltaNetSerialize(FArchive& Ar, UPackageMap* Map, bool& bOutSuccess, FModuleInputContainer& PreviousInputContainer)
-{
-	using namespace UE::Net;
-	bOutSuccess = true;
-	
-	FModuleInputNetTokenData InputStateData;
-	InputStateData.Init(InputValues);
-	TStructNetTokenDataStoreHelper<FModuleInputNetTokenData>::NetSerializeAndExportToken(Ar,Map,InputStateData);
-
-	uint32 Number = InputStateData.Types.Num();
-	if (Ar.IsLoading())
-	{
-		InputValues.SetNum(Number);
-	}
-
-	for (uint32 I = 0; I < Number; I++)
-	{
-		InputValues[I].ConvertToType(static_cast<EModuleInputValueType>(InputStateData.Types[I]));
-		InputValues[I].SetApplyInputDecay(InputStateData.DecayValues[I]);
-		if (PreviousInputContainer.InputValues.Num() == InputValues.Num())
-		{
-			InputValues[I].DeltaNetSerialize(Ar, Map, bOutSuccess, PreviousInputContainer.InputValues[I]);
-		}
-		else
-		{
-			InputValues[I].Serialize(Ar, Map, bOutSuccess);
-		}
-		
-	}
-
-}
-
 int FModuleInputContainer::AddInput(EModuleInputValueType Type, TSubclassOf<UDefaultModularVehicleInputModifier>& InputModifierClass)
 {
 	switch(Type)
diff --git a/Engine/Source/Runtime/Experimental/ChaosVehicles/ChaosVehiclesCore/Public/SimModule/ModuleInput.h b/Engine/Source/Runtime/Experimental/ChaosVehicles/ChaosVehiclesCore/Public/SimModule/ModuleInput.h
--- a/Engine/Source/Runtime/Experimental/ChaosVehicles/ChaosVehiclesCore/Public/SimModule/ModuleInput.h
+++ b/Engine/Source/Runtime/Experimental/ChaosVehicles/ChaosVehiclesCore/Public/SimModule/ModuleInput.h
@@ -3,6 +3,7 @@
 #pragma once
 
 #include "CoreMinimal.h"
+//#include "Engine/NetSerialization.h"
 #include "Templates/SubclassOf.h"
 #include "ModuleInput.generated.h"
 
@@ -102,65 +103,43 @@ namespace ModularQuantize
 	}
 
 	template<int32 MaxValue, uint32 NumBits, typename T UE_REQUIRES(std::is_floating_point_v<T>&& NumBits < 32)>
-	bool WriteCompressedFloat(const T Value, FArchive& Ar, bool bAutoReduceZeroValues = false)
+	bool WriteCompressedFloat(const T Value, FArchive& Ar)
 	{
 		using Details = ModularQuantize::TCompressedFloatDetails<MaxValue, NumBits>;
 
 		uint32 CompressedValue;
 		bool clamp = ModularQuantize::ToCompressedFloat<MaxValue, NumBits>(Value, CompressedValue);
-		
-		if(bAutoReduceZeroValues)
-		{
-			bool IsZero = CompressedValue == 0;
-			Ar.SerializeBits(&IsZero,1);
-			if(!IsZero)
-			{
-				Ar.SerializeBits(&CompressedValue, NumBits);
-			}
-		}
-		else
-		{
-			Ar.SerializeBits(&CompressedValue, NumBits);
-		}
+
+		Ar.SerializeInt(CompressedValue, Details::SerIntMax);
 
 		return !clamp;
 	}
 
 	template<int32 MaxValue, uint32 NumBits, typename T UE_REQUIRES(std::is_floating_point_v<T>&& NumBits < 32)>
-	bool ReadCompressedFloat(T& Value, FArchive& Ar, bool bAutoReduceZeroValues = false)
+	bool ReadCompressedFloat(T& Value, FArchive& Ar)
 	{
 		using Details = ModularQuantize::TCompressedFloatDetails<MaxValue, NumBits>;
-		uint32 CompressedValue = 0;
-		if(bAutoReduceZeroValues)
-		{
-			bool IsZero;
-			Ar.SerializeBits(&IsZero,1);
-			if(!IsZero)
-			{
-				Ar.SerializeBits(&CompressedValue, NumBits);
-			}
-		}
-		else
-		{
-			Ar.SerializeBits(&CompressedValue,NumBits);
-		}
+
+		uint32 CompressedValue;
+		Ar.SerializeInt(CompressedValue, Details::SerIntMax);
 
 		ModularQuantize::FromCompressedFloat<MaxValue, NumBits>(CompressedValue, Value);
+
 		return true;
 	}
 
 	// Required because we seialize quantized vector in seperate parts depending on input type
 	template<int32 MaxValue, uint32 NumBits>
-	bool SerializeFixedFloat(double& InOutValue, FArchive& Ar, bool bAutoReduceZeroValues = false)
+	bool SerializeFixedFloat(double& InOutValue, FArchive& Ar)
 	{
 		if (Ar.IsSaving())
 		{
 			bool success = true;
-			success &= ModularQuantize::WriteCompressedFloat<MaxValue, NumBits>(InOutValue, Ar, bAutoReduceZeroValues);
+			success &= ModularQuantize::WriteCompressedFloat<MaxValue, NumBits>(InOutValue, Ar);
 			return success;
 		}
 
-		ModularQuantize::ReadCompressedFloat<MaxValue, NumBits>(InOutValue, Ar, bAutoReduceZeroValues);
+		ModularQuantize::ReadCompressedFloat<MaxValue, NumBits>(InOutValue, Ar);
 		return true;
 	}
 
@@ -171,18 +150,6 @@ namespace ModularQuantize
 		ModularQuantize::ToCompressedFloat<MaxValue, NumBits>(Value, CompressedValue);
 		ModularQuantize::FromCompressedFloat<MaxValue, NumBits>(CompressedValue, Value);
 	}
-	
-	template<int32 MaxValue, uint32 NumBits, typename T UE_REQUIRES(std::is_floating_point_v<T>&& NumBits < 32)>
-    bool QuantizedIsNearlyEqual(const T& Left,const T& Right)
-    {
-    	uint32 CompressedValue = 0;
-    	T QuantLeft, QuantRight;
-    	ModularQuantize::ToCompressedFloat<MaxValue, NumBits>(Left, CompressedValue);
-    	ModularQuantize::FromCompressedFloat<MaxValue, NumBits>(CompressedValue, QuantLeft);
-    	ModularQuantize::ToCompressedFloat<MaxValue, NumBits>(Right, CompressedValue);
-    	ModularQuantize::FromCompressedFloat<MaxValue, NumBits>(CompressedValue, QuantRight);
-    	return FMath::IsNearlyEqual(QuantLeft, QuantRight);
-    }
 
 }
 
@@ -359,6 +326,8 @@ public:
 		return OutValue;
 	}
 
+
+
 	template<typename T>
 	inline T Get() const { static_assert(sizeof(T) == 0, "Unsupported conversion for type"); }
 
@@ -387,7 +356,6 @@ public:
 
 	// Serialize values
 	void Serialize(FArchive& Ar, UPackageMap* Map, bool& bOutSuccess);
-	void DeltaNetSerialize(FArchive& Ar, UPackageMap* Map, bool& bOutSuccess,const FModuleInputValue& PreviousInputValue);
 
 	void Lerp(const FModuleInputValue& Min, const FModuleInputValue& Max, float Alpha);
 
@@ -400,7 +368,7 @@ public:
 
 	/** During physics resimulation, apply decay for this input while it's being extrapolated */
 	void SetApplyInputDecay(const bool bInApplyInputDecay) { bApplyInputDecay = bInApplyInputDecay; }
-	bool ShouldApplyInputDecay() const { return  bApplyInputDecay; }
+	bool ShouldApplyInputDecay() { return  bApplyInputDecay; }
 
 protected:
 	void SetMagnitude(float NewSize);
@@ -595,6 +563,7 @@ private:
 	static TArray<struct FModuleInputSetup>* InitSetupData;
 };
 
+
 USTRUCT()
 struct CHAOSVEHICLESCORE_API FModuleInputContainer
 {
@@ -647,8 +616,6 @@ public:
 	void ZeroValues();
 
 	void Serialize(FArchive& Ar, UPackageMap* Map, bool& bOutSuccess);
-	void SerializeReduced(FArchive& Ar, UPackageMap* Map, bool& bOutSuccess);
-	void DeltaNetSerialize(FArchive& Ar, UPackageMap* Map, bool& bOutSuccess, FModuleInputContainer& PreviousInputContainer);
 
 	int AddInput(EModuleInputValueType Type, TSubclassOf<UDefaultModularVehicleInputModifier>& InputModifierClass);
 
@@ -667,9 +634,11 @@ private:
 
 };
 
+
 class CHAOSVEHICLESCORE_API FInputInterface
 {
 public:
+
 	using FInputNameMap = TMap<FName, int>;
 
 	FInputInterface(const FInputNameMap& InNameMap, FModuleInputContainer& InValueContainer)
diff --git a/Engine/Source/Runtime/Experimental/ChaosVehicles/ChaosVehiclesCore/Public/SimModule/ModuleInputTokenStore.h b/Engine/Source/Runtime/Experimental/ChaosVehicles/ChaosVehiclesCore/Public/SimModule/ModuleInputTokenStore.h
deleted file mode 100644
