[DSTM] Updating the world bounds of a particle regardless of its type.

This fixing an issue when if you apply the state of a particle that was serialized while sleeping, the new particle will never re-calculate its bounds until is awaken.

This was not an issue for awake particles because their bounds get re-calculated on Integrate stage of the evolution.

diff --git a/Engine/Source/Runtime/Experimental/Chaos/Private/Chaos/Serialization/SolverSerializer.cpp b/Engine/Source/Runtime/Experimental/Chaos/Private/Chaos/Serialization/SolverSerializer.cpp
--- a/Engine/Source/Runtime/Experimental/Chaos/Private/Chaos/Serialization/SolverSerializer.cpp
+++ b/Engine/Source/Runtime/Experimental/Chaos/Private/Chaos/Serialization/SolverSerializer.cpp
@@ -115,11 +115,8 @@ namespace Chaos
 		constexpr bool bWakeUp = false;
 		Evolution.SetParticleTransform(InParticleHandle, InParticleState.ParticlePositionRotation.MX, InParticleState.ParticlePositionRotation.MR, bIsTeleport, bWakeUp);
 
-		if (InParticleState.Type == EChaosVDParticleType::Static)
-		{
-			const FVec3 BoundsExpansion = FVec3(0);
-			InParticleHandle->UpdateWorldSpaceState(InParticleHandle->GetTransformXR(), BoundsExpansion);
-		}
+		const FVec3 BoundsExpansion = FVec3(0);
+		InParticleHandle->UpdateWorldSpaceState(InParticleHandle->GetTransformXR(), BoundsExpansion);
 
 		Evolution.SetParticleVelocities(InParticleHandle, InParticleState.ParticleVelocities.MV, InParticleState.ParticleVelocities.MW);
 
