Replace use of DefaultHostGroup with the host group that the chunk will actually use.

#rnx

diff --git a/Engine/Source/Runtime/Experimental/IoStore/OnDemand/Private/OnDemandIoDispatcherBackend.cpp b/Engine/Source/Runtime/Experimental/IoStore/OnDemand/Private/OnDemandIoDispatcherBackend.cpp
--- a/Engine/Source/Runtime/Experimental/IoStore/OnDemand/Private/OnDemandIoDispatcherBackend.cpp
+++ b/Engine/Source/Runtime/Experimental/IoStore/OnDemand/Private/OnDemandIoDispatcherBackend.cpp
@@ -1175,7 +1175,7 @@ bool FOnDemandIoBackend::Resolve(FIoRequestImpl* Request)
 
 	FChunkRequestParams RequestParams = FChunkRequestParams::Create(Request, ChunkInfo);
 
-	if (BackendStatus.IsHttpEnabled(Request->ChunkId.GetChunkType()) == false || !DefaultHostGroup.IsConnected())
+	if (BackendStatus.IsHttpEnabled(Request->ChunkId.GetChunkType()) == false || ChunkInfo.HostGroup().IsConnected() == false)
 	{ 
 		// If the cache is not readonly the chunk may get evicted before the request is completed
 		if (BackendStatus.IsCacheReadOnly() == false || Cache->ContainsChunk(RequestParams.ChunkKey) == false)
