Check the http status rather than DefaultHostGroup when deciding if we should send IAS analytics or not.

#rnx

diff --git a/Engine/Source/Runtime/Experimental/IoStore/OnDemand/Private/OnDemandIoDispatcherBackend.cpp b/Engine/Source/Runtime/Experimental/IoStore/OnDemand/Private/OnDemandIoDispatcherBackend.cpp
--- a/Engine/Source/Runtime/Experimental/IoStore/OnDemand/Private/OnDemandIoDispatcherBackend.cpp
+++ b/Engine/Source/Runtime/Experimental/IoStore/OnDemand/Private/OnDemandIoDispatcherBackend.cpp
@@ -1471,7 +1471,7 @@ void FOnDemandIoBackend::ReportAnalytics(TArray<FAnalyticsEventAttribute>& OutAn
 
 	Stats.ReportGeneralAnalytics(OutAnalyticsArray);
 
-	if (DefaultHostGroup.IsConnected())
+	if (BackendStatus.IsHttpEnabled())
 	{
 		Stats.ReportEndPointAnalytics(OutAnalyticsArray);
 	}
