Dataflow debug draw extension

diff --git a/Engine/Source/Runtime/Experimental/Dataflow/Core/Public/Dataflow/DataflowDebugDrawInterface.h b/Engine/Source/Runtime/Experimental/Dataflow/Core/Public/Dataflow/DataflowDebugDrawInterface.h
--- a/Engine/Source/Runtime/Experimental/Dataflow/Core/Public/Dataflow/DataflowDebugDrawInterface.h
+++ b/Engine/Source/Runtime/Experimental/Dataflow/Core/Public/Dataflow/DataflowDebugDrawInterface.h
@@ -42,6 +42,8 @@ public:
 
 	DATAFLOWCORE_API virtual void DrawBox(const FVector& Extents, const FQuat& Rotation, const FVector& Center, double UniformScale) const = 0;
 	DATAFLOWCORE_API virtual void DrawSphere(const FVector& Center, double Radius) const = 0;
+	DATAFLOWCORE_API virtual void DrawCapsule(const FVector& Center, const double& Radius, const double& HalfHeight, const FVector& XAxis, const FVector& YAxis, const FVector &ZAxis) const = 0;
+
 
 	// TODO: More draw functions
 
diff --git a/Engine/Source/Runtime/Experimental/Dataflow/Engine/Private/Dataflow/DataflowDebugDraw.cpp b/Engine/Source/Runtime/Experimental/Dataflow/Engine/Private/Dataflow/DataflowDebugDraw.cpp
--- a/Engine/Source/Runtime/Experimental/Dataflow/Engine/Private/Dataflow/DataflowDebugDraw.cpp
+++ b/Engine/Source/Runtime/Experimental/Dataflow/Engine/Private/Dataflow/DataflowDebugDraw.cpp
@@ -183,6 +183,11 @@ void FDataflowDebugDraw::DrawSphere(const FVector& Center, double Radius) const
 	DebugRenderSceneProxy->Spheres.Add(FDebugRenderSceneProxy::FSphere(Radius, Center, ColorWithTranslucency, DrawType));
 }
 
+void FDataflowDebugDraw::DrawCapsule(const FVector& Center, const double& Radius, const double& HalfHeight, const FVector& XAxis, const FVector& YAxis, const FVector &ZAxis) const
+{
+	DebugRenderSceneProxy->Capsules.Add(FDebugRenderSceneProxy::FCapsule(Center, Radius, XAxis, YAxis, ZAxis, HalfHeight, ColorWithTranslucency, DrawType));
+}
+
 /* ----------------------------------------------------------------------------------------------------------------------- */
 
 void FDataflowNodeDebugDrawSettings::SetDebugDrawSettings(IDataflowDebugDrawInterface& DataflowRenderingInterface) const
@@ -205,3 +210,4 @@ void FDataflowNodeDebugDrawSettings::SetDebugDrawSettings(IDataflowDebugDrawInte
 	DataflowRenderingInterface.SetColor(Color);
 }
 
+
diff --git a/Engine/Source/Runtime/Experimental/Dataflow/Engine/Public/Dataflow/DataflowDebugDraw.h b/Engine/Source/Runtime/Experimental/Dataflow/Engine/Public/Dataflow/DataflowDebugDraw.h
--- a/Engine/Source/Runtime/Experimental/Dataflow/Engine/Public/Dataflow/DataflowDebugDraw.h
+++ b/Engine/Source/Runtime/Experimental/Dataflow/Engine/Public/Dataflow/DataflowDebugDraw.h
@@ -33,6 +33,7 @@ private:
 
 	virtual void DrawBox(const FVector& Extents, const FQuat& Rotation, const FVector& Center, double UniformScale) const override;
 	virtual void DrawSphere(const FVector& Center, double Radius) const override;
+	virtual void DrawCapsule(const FVector& Center, const double& Radius, const double& HalfHeight, const FVector& XAxis, const FVector& YAxis, const FVector &ZAxis) const override;
 
 	class FDataflowDebugRenderSceneProxy* DebugRenderSceneProxy = nullptr;
 
