Dataflow : variable node improvement and bug fixes
- Better support for arrays
- Automatic disconnect / reconnection when type changes

diff --git a/Engine/Source/Runtime/Experimental/Dataflow/Core/Public/Dataflow/DataflowTypePolicy.h b/Engine/Source/Runtime/Experimental/Dataflow/Core/Public/Dataflow/DataflowTypePolicy.h
--- a/Engine/Source/Runtime/Experimental/Dataflow/Core/Public/Dataflow/DataflowTypePolicy.h
+++ b/Engine/Source/Runtime/Experimental/Dataflow/Core/Public/Dataflow/DataflowTypePolicy.h
@@ -20,6 +20,17 @@ namespace UE::Dataflow
 {
 	template<typename T>
 	FName GetTypeName() { return TDataflowPolicyTypeName<T>::GetName(); }
+
+	template<typename T>
+	FName GetTypeName(bool bAsArray)
+	{ 
+		if (bAsArray)
+		{
+			return TDataflowPolicyTypeName<TArray<T>>::GetName();
+		}
+		return TDataflowPolicyTypeName<T>::GetName();
+	}
+
 }
 
 
diff --git a/Engine/Source/Runtime/Experimental/Dataflow/Engine/Private/Dataflow/DataflowVariableNodes.cpp b/Engine/Source/Runtime/Experimental/Dataflow/Engine/Private/Dataflow/DataflowVariableNodes.cpp
--- a/Engine/Source/Runtime/Experimental/Dataflow/Engine/Private/Dataflow/DataflowVariableNodes.cpp
+++ b/Engine/Source/Runtime/Experimental/Dataflow/Engine/Private/Dataflow/DataflowVariableNodes.cpp
@@ -120,8 +120,33 @@ void FGetDataflowVariableNode::SetVariable(UDataflow* Dataflow, FName InVariable
 	}
 }
 
+void FGetDataflowVariableNode::ChangeOutputType(FDataflowOutput* Output, FName NewType)
+{
+	// make sure to disconnect
+	if (Output && Output->GetType() != NewType)
+	{
+		TStrongObjectPtr<UDataflow> DataflowAsset = WeakDataflowPtr.Pin();
+		const TArray<FDataflowInput*> InputsToReconnect = Output->GetConnections();
+		if (DataflowAsset)
+		{
+			DataflowAsset->GetDataflow()->ClearConnections(Output);
+		}
+		
+		Output->UnlockType();
+		Output->SetConcreteType(NewType);
+		Output->LockType();
+
+		if (DataflowAsset)
+		{
+			for (FDataflowInput* InputToReconncted : InputsToReconnect)
+			DataflowAsset->GetDataflow()->Connect(InputToReconncted, Output);
+		}
+	}
+}
+
 void FGetDataflowVariableNode::UpdateOutputTypes(const FPropertyBagPropertyDesc& Desc)
 {
+	using namespace UE::Dataflow;
 	// compute type
 	if (!Desc.CachedProperty)
 	{
@@ -131,52 +156,52 @@ void FGetDataflowVariableNode::UpdateOutputTypes(const FPropertyBagPropertyDesc&
 	FDataflowOutput* Output = FindOutput(&Value);
 	check(Output);
 
-	Output->UnlockType();
-
 	ContainerType = Desc.ContainerTypes.GetFirstContainerType();
 
+	const bool bIsArrayType = (ContainerType == EPropertyBagContainerType::Array);
 	const FName ConcreteType = FDataflowConnection::GetTypeNameFromProperty(Desc.CachedProperty);
+
 	switch (Desc.ValueType)
 	{
 	case EPropertyBagPropertyType::Bool:
-		SetOutputConcreteType(&Value, TDataflowSingleTypePolicy<bool>::TypeName);
+		ChangeOutputType(Output, GetTypeName<bool>(bIsArrayType));
 		EvaluateFunction = &EvaluateBool;
 		break;
 	case EPropertyBagPropertyType::Byte:
-		SetOutputConcreteType(&Value, TDataflowSingleTypePolicy<uint8>::TypeName);
+		ChangeOutputType(Output, GetTypeName<uint8>(bIsArrayType));
 		EvaluateFunction = &EvaluateByte;
 		break;
 	case EPropertyBagPropertyType::Int32:
-		SetOutputConcreteType(&Value, TDataflowSingleTypePolicy<int32>::TypeName);
+		ChangeOutputType(Output, GetTypeName<int32>(bIsArrayType));
 		EvaluateFunction = &EvaluateInt32;
 		break;
 	case EPropertyBagPropertyType::Int64:
-		SetOutputConcreteType(&Value, TDataflowSingleTypePolicy<int64>::TypeName);
+		ChangeOutputType(Output, GetTypeName<int64>(bIsArrayType));
 		EvaluateFunction = &EvaluateInt64;
 		break;
 	case EPropertyBagPropertyType::Float:
-		SetOutputConcreteType(&Value, TDataflowSingleTypePolicy<float>::TypeName);
+		ChangeOutputType(Output, GetTypeName<float>(bIsArrayType));
 		EvaluateFunction = &EvaluateFloat;
 		break;
 	case EPropertyBagPropertyType::Double:
 		// the UI shows float but behind the scene set a double property let only deal with float 
-		SetOutputConcreteType(&Value, TDataflowSingleTypePolicy<float>::TypeName);
+		ChangeOutputType(Output, GetTypeName<float>(bIsArrayType));
 		EvaluateFunction = &EvaluateDouble;
 		break;
 	case EPropertyBagPropertyType::Name:
-		SetOutputConcreteType(&Value, TDataflowSingleTypePolicy<FName>::TypeName);
+		ChangeOutputType(Output, GetTypeName<FName>(bIsArrayType));
 		EvaluateFunction = &EvaluateName;
 		break;
 	case EPropertyBagPropertyType::String:
-		SetOutputConcreteType(&Value, TDataflowSingleTypePolicy<FString>::TypeName);
+		ChangeOutputType(Output, GetTypeName<FString>(bIsArrayType));
 		EvaluateFunction = &EvaluateString;
 		break;
 	case EPropertyBagPropertyType::Text:
-		SetOutputConcreteType(&Value, TDataflowSingleTypePolicy<FString>::TypeName);
+		ChangeOutputType(Output, GetTypeName<FText>(bIsArrayType));
 		EvaluateFunction = &EvaluateText;
 		break;
 	case EPropertyBagPropertyType::Object:
-		SetOutputConcreteType(&Value, ConcreteType);
+		ChangeOutputType(Output, ConcreteType);
 		EvaluateFunction = &EvaluateObject;
 		break;
 		// unsupported type for now 
@@ -192,8 +217,6 @@ void FGetDataflowVariableNode::UpdateOutputTypes(const FPropertyBagPropertyDesc&
 		break;
 	}
 
-	Output->LockType();
-
 	if (TStrongObjectPtr<UDataflow> DataflowAsset = WeakDataflowPtr.Pin())
 	{
 		DataflowAsset->RefreshEdNodeByGuid(GetGuid());
diff --git a/Engine/Source/Runtime/Experimental/Dataflow/Engine/Public/Dataflow/DataflowVariableNodes.h b/Engine/Source/Runtime/Experimental/Dataflow/Engine/Public/Dataflow/DataflowVariableNodes.h
--- a/Engine/Source/Runtime/Experimental/Dataflow/Engine/Public/Dataflow/DataflowVariableNodes.h
+++ b/Engine/Source/Runtime/Experimental/Dataflow/Engine/Public/Dataflow/DataflowVariableNodes.h
@@ -50,7 +50,7 @@ private:
 	static void EvaluateObject(const FGetDataflowVariableNode& Node, const FInstancedPropertyBag& Variables, UE::Dataflow::FContext& Context, const FDataflowOutput& Out);
 
 	void UpdateOutputTypes(const FPropertyBagPropertyDesc& Desc);
-
+	void ChangeOutputType(FDataflowOutput* Output, FName NewType);
 
 private:
 	UPROPERTY(EditAnywhere, Category = Variables, meta = (FixedLayout, ShowOnlyInnerProperties))
