[ChaosVehicleCore] Resolve bug with flipped default condition state on serializing input type information.

diff --git a/Engine/Source/Runtime/Experimental/ChaosVehicles/ChaosVehiclesCore/Private/SimModule/ModuleInput.cpp b/Engine/Source/Runtime/Experimental/ChaosVehicles/ChaosVehiclesCore/Private/SimModule/ModuleInput.cpp
--- a/Engine/Source/Runtime/Experimental/ChaosVehicles/ChaosVehiclesCore/Private/SimModule/ModuleInput.cpp
+++ b/Engine/Source/Runtime/Experimental/ChaosVehicles/ChaosVehiclesCore/Private/SimModule/ModuleInput.cpp
@@ -124,7 +124,7 @@ void FModuleInputValue::Serialize(FArchive& Ar, UPackageMap* Map, bool& bOutSucc
 
 	bool bIsNonZero = Ar.IsLoading() ? false : IsQuantizedNonZero();
 	Ar.SerializeBits(&bIsNonZero, 1);
-	if (!ChaosModularVehicleCVars::bEnableInputSendingTypeInfo)
+	if (ChaosModularVehicleCVars::bEnableInputSendingTypeInfo)
 	{
 		Ar.SerializeBits(&ValueType, 3);
 		Ar.SerializeBits(&bApplyInputDecay, 1);
