Dataflow get variable node :fix for loop not using scope

diff --git a/Engine/Source/Runtime/Experimental/Dataflow/Engine/Private/Dataflow/DataflowVariableNodes.cpp b/Engine/Source/Runtime/Experimental/Dataflow/Engine/Private/Dataflow/DataflowVariableNodes.cpp
--- a/Engine/Source/Runtime/Experimental/Dataflow/Engine/Private/Dataflow/DataflowVariableNodes.cpp
+++ b/Engine/Source/Runtime/Experimental/Dataflow/Engine/Private/Dataflow/DataflowVariableNodes.cpp
@@ -138,8 +138,10 @@ void FGetDataflowVariableNode::ChangeOutputType(FDataflowOutput* Output, FName N
 
 		if (DataflowAsset)
 		{
-			for (FDataflowInput* InputToReconncted : InputsToReconnect)
-			DataflowAsset->GetDataflow()->Connect(InputToReconncted, Output);
+			for (FDataflowInput* InputToReconnect : InputsToReconnect)
+			{
+				DataflowAsset->GetDataflow()->Connect(InputToReconnect, Output);
+			}
 		}
 	}
 }
