Dataflow : add IsTypeLocked method to dataflow connections

diff --git a/Engine/Source/Runtime/Experimental/Dataflow/Core/Public/Dataflow/DataflowConnection.h b/Engine/Source/Runtime/Experimental/Dataflow/Core/Public/Dataflow/DataflowConnection.h
--- a/Engine/Source/Runtime/Experimental/Dataflow/Core/Public/Dataflow/DataflowConnection.h
+++ b/Engine/Source/Runtime/Experimental/Dataflow/Core/Public/Dataflow/DataflowConnection.h
@@ -231,6 +231,7 @@ public:
 
 	void LockType() { bLockType = true; }
 	void UnlockType() { bLockType = false; }
+	bool IsTypeLocked() const { return bLockType; }
 
 	FName GetType() const { return Type; }
 	FName GetOriginalType() const { return OriginalType; }
