Dataflow : add support for screenspace point debug draw

diff --git a/Engine/Source/Runtime/Experimental/Dataflow/Core/Public/Dataflow/DataflowDebugDrawInterface.h b/Engine/Source/Runtime/Experimental/Dataflow/Core/Public/Dataflow/DataflowDebugDrawInterface.h
--- a/Engine/Source/Runtime/Experimental/Dataflow/Core/Public/Dataflow/DataflowDebugDrawInterface.h
+++ b/Engine/Source/Runtime/Experimental/Dataflow/Core/Public/Dataflow/DataflowDebugDrawInterface.h
@@ -14,6 +14,7 @@ public:
 
 	// State management
 	DATAFLOWCORE_API virtual void SetColor(const FLinearColor& InColor) = 0;
+	DATAFLOWCORE_API virtual void SetPointSize(float Size) = 0;
 	DATAFLOWCORE_API virtual void SetLineWidth(double Width) = 0;
 	DATAFLOWCORE_API virtual void SetWireframe(bool bInWireframe) = 0;
 	DATAFLOWCORE_API virtual void SetShaded(bool bInShaded) = 0;
@@ -22,6 +23,7 @@ public:
 	DATAFLOWCORE_API virtual void SetWorldPriority() = 0;
 	DATAFLOWCORE_API virtual void ResetAllState() = 0;
 
+	DATAFLOWCORE_API virtual void DrawPoint(const FVector& Position) = 0;
 	DATAFLOWCORE_API virtual void DrawLine(const FVector& Start, const FVector& End) const = 0;
 
 	struct IDebugDrawMesh
diff --git a/Engine/Source/Runtime/Experimental/Dataflow/Engine/Private/Dataflow/DataflowDebugDraw.cpp b/Engine/Source/Runtime/Experimental/Dataflow/Engine/Private/Dataflow/DataflowDebugDraw.cpp
--- a/Engine/Source/Runtime/Experimental/Dataflow/Engine/Private/Dataflow/DataflowDebugDraw.cpp
+++ b/Engine/Source/Runtime/Experimental/Dataflow/Engine/Private/Dataflow/DataflowDebugDraw.cpp
@@ -27,6 +27,11 @@ void FDataflowDebugDraw::SetColor(const FLinearColor& InColor)
 	}
 }
 
+void FDataflowDebugDraw::SetPointSize(float InSize)
+{
+	PointSize = InSize;
+}
+
 void FDataflowDebugDraw::SetLineWidth(double InWidth)
 {
 	LineWidth = InWidth;
@@ -104,6 +109,17 @@ void FDataflowDebugDraw::ResetAllState()
 	ColorWithTranslucency = Color;
 }
 
+void FDataflowDebugDraw::DrawPoint(const FVector& Position)
+{
+	FDataflowDebugRenderSceneProxy::FDebugPoint NewPoint
+	{
+		.Position = Position,
+		.Size = PointSize,
+		.Color = ColorWithTranslucency.ToFColor(true),
+	};
+	DebugRenderSceneProxy->AddPoint(NewPoint);
+}
+
 void FDataflowDebugDraw::DrawLine(const FVector& Start, const FVector& End) const
 {
 	DebugRenderSceneProxy->Lines.Add(FDebugRenderSceneProxy::FDebugLine(Start, End, ColorWithTranslucency.ToFColor(true), LineWidth));
diff --git a/Engine/Source/Runtime/Experimental/Dataflow/Engine/Private/Dataflow/DataflowDebugDrawComponent.cpp b/Engine/Source/Runtime/Experimental/Dataflow/Engine/Private/Dataflow/DataflowDebugDrawComponent.cpp
--- a/Engine/Source/Runtime/Experimental/Dataflow/Engine/Private/Dataflow/DataflowDebugDrawComponent.cpp
+++ b/Engine/Source/Runtime/Experimental/Dataflow/Engine/Private/Dataflow/DataflowDebugDrawComponent.cpp
@@ -1,7 +1,8 @@
 // Copyright Epic Games, Inc. All Rights Reserved.
 
 #include "Dataflow/DataflowDebugDrawComponent.h"
-
+#include "MeshElementCollector.h"
+#include "PrimitiveDrawInterface.h"
 
 FDebugRenderSceneProxy* UDataflowDebugDrawComponent::CreateDebugSceneProxy()
 {
@@ -76,6 +77,10 @@ FBoxSphereBounds UDataflowDebugDrawComponent::CalcBounds(const FTransform& Local
 		{
 			BoundsBuilder += FSphere(CoordinateSystem.AxisLoc, CoordinateSystem.Scale + CoordinateSystem.Thickness);
 		}
+		for (const FDataflowDebugRenderSceneProxy::FDebugPoint& Point : DebugSceneProxy->Points)
+		{
+			BoundsBuilder += Point.Position;
+		}
 	}
 
 	FBoxSphereBounds ReturnBounds(BoundsBuilder);
@@ -105,6 +110,7 @@ void FDataflowDebugRenderSceneProxy::ClearAll()
 	Meshes.Reset();
 	Capsules.Reset();
 	CoordinateSystems.Reset();
+	Points.Reset();
 }
 
 
@@ -118,3 +124,20 @@ FPrimitiveViewRelevance FDataflowDebugRenderSceneProxy::GetViewRelevance(const F
 	return Result;
 }
 
+void FDataflowDebugRenderSceneProxy::AddPoint(const FDebugPoint& Point)
+{
+	Points.Add(Point);
+}
+
+void FDataflowDebugRenderSceneProxy::GetDynamicMeshElementsForView(const FSceneView* View, const int32 ViewIndex, const FSceneViewFamily& ViewFamily, const uint32 VisibilityMap, FMeshElementCollector& Collector, FMaterialCache& DefaultMaterialCache, FMaterialCache& SolidMeshMaterialCache) const
+{
+	FDebugRenderSceneProxy::GetDynamicMeshElementsForView(View, ViewIndex, ViewFamily, VisibilityMap, Collector, DefaultMaterialCache, SolidMeshMaterialCache);
+
+	FPrimitiveDrawInterface* PDI = Collector.GetPDI(ViewIndex);
+
+	// Draw Lines
+	for (const FDebugPoint& Point : Points)
+	{
+		PDI->DrawPoint(Point.Position, Point.Color, Point.Size, SDPG_World);
+	}
+}
\ No newline at end of file
diff --git a/Engine/Source/Runtime/Experimental/Dataflow/Engine/Public/Dataflow/DataflowDebugDraw.h b/Engine/Source/Runtime/Experimental/Dataflow/Engine/Public/Dataflow/DataflowDebugDraw.h
--- a/Engine/Source/Runtime/Experimental/Dataflow/Engine/Public/Dataflow/DataflowDebugDraw.h
+++ b/Engine/Source/Runtime/Experimental/Dataflow/Engine/Public/Dataflow/DataflowDebugDraw.h
@@ -20,6 +20,7 @@ public:
 private:
 
 	virtual void SetColor(const FLinearColor& InColor) override;
+	virtual void SetPointSize(float Size) override;
 	virtual void SetLineWidth(double Width) override;
 	virtual void SetWireframe(bool bInWireframe) override;
 	virtual void SetShaded(bool bInShaded) override;
@@ -28,6 +29,7 @@ private:
 	virtual void SetWorldPriority() override;
 	virtual void ResetAllState() override;
 
+	virtual void DrawPoint(const FVector& Position) override;
 	virtual void DrawLine(const FVector& Start, const FVector& End) const override;
 	virtual void DrawMesh(const IDebugDrawMesh& Mesh) const override;
 
@@ -39,6 +41,7 @@ private:
 
 	FLinearColor Color = FLinearColor::White;
 	double LineWidth = 1.0;
+	float PointSize = 5.0;
 	bool bWireframe = true;
 	bool bShaded = false;
 	bool bTranslucent = false;
diff --git a/Engine/Source/Runtime/Experimental/Dataflow/Engine/Public/Dataflow/DataflowDebugDrawComponent.h b/Engine/Source/Runtime/Experimental/Dataflow/Engine/Public/Dataflow/DataflowDebugDrawComponent.h
--- a/Engine/Source/Runtime/Experimental/Dataflow/Engine/Public/Dataflow/DataflowDebugDrawComponent.h
+++ b/Engine/Source/Runtime/Experimental/Dataflow/Engine/Public/Dataflow/DataflowDebugDrawComponent.h
@@ -11,7 +11,6 @@ class UDataflowDebugDrawComponent : public UDebugDrawComponent
 	GENERATED_BODY()
 
 private:
-
 	virtual FDebugRenderSceneProxy* CreateDebugSceneProxy() override;
 	virtual FBoxSphereBounds CalcBounds(const FTransform& LocalToWorld) const override;
 };
@@ -26,7 +25,20 @@ public:
 	/** Remove all primitives stored on this proxy */
 	void ClearAll();
 
+	struct FDebugPoint
+	{
+		FVector Position;
+		float	Size = 0;
+		FLinearColor Color = FLinearColor::White;
+	};
+
+	void AddPoint(const FDebugPoint& Point);
+
 private:
+	friend class UDataflowDebugDrawComponent;
+	TArray<FDebugPoint> Points;
 
+private:
+	virtual void GetDynamicMeshElementsForView(const FSceneView* View, const int32 ViewIndex, const FSceneViewFamily& ViewFamily, const uint32 VisibilityMap, FMeshElementCollector& Collector, FMaterialCache& DefaultMaterialCache, FMaterialCache& SolidMeshMaterialCache) const override;
 	virtual FPrimitiveViewRelevance GetViewRelevance(const FSceneView* View) const override;
 };
