Replace use of DefaultHostGroup with the host group that the chunk will actually use when recovering from a failed disk cache get.

#rnx

diff --git a/Engine/Source/Runtime/Experimental/IoStore/OnDemand/Private/OnDemandIoDispatcherBackend.cpp b/Engine/Source/Runtime/Experimental/IoStore/OnDemand/Private/OnDemandIoDispatcherBackend.cpp
--- a/Engine/Source/Runtime/Experimental/IoStore/OnDemand/Private/OnDemandIoDispatcherBackend.cpp
+++ b/Engine/Source/Runtime/Experimental/IoStore/OnDemand/Private/OnDemandIoDispatcherBackend.cpp
@@ -1151,9 +1151,9 @@ void FOnDemandIoBackend::CompleteCacheRead(FChunkRequest* ChunkRequest)
 		break;
 	}
 
-	if (bWasCancelled || (!DefaultHostGroup.IsConnected() || !BackendStatus.IsHttpEnabled()))
+	if (bWasCancelled || (!BackendStatus.IsHttpEnabled() || !ChunkRequest->HostGroup.IsConnected()))
 	{
-		UE_CLOG(!DefaultHostGroup.IsConnected() || !BackendStatus.IsHttpEnabled(), LogIas, Log, TEXT("Chunk was not found in the cache and HTTP is disabled"));
+		UE_CLOG(!bWasCancelled, LogIas, Log, TEXT("Chunk was not found in the cache and HTTP is disabled"));
 		CompleteRequest(ChunkRequest);
 		return;
 	}
