Gather all of the DefaultHostGroup creation code to the same scope at the start of FOnDemandIoBackend::Run and improve some insights trace markers.

#rnx

diff --git a/Engine/Source/Runtime/Experimental/IoStore/OnDemand/Private/OnDemandIoDispatcherBackend.cpp b/Engine/Source/Runtime/Experimental/IoStore/OnDemand/Private/OnDemandIoDispatcherBackend.cpp
--- a/Engine/Source/Runtime/Experimental/IoStore/OnDemand/Private/OnDemandIoDispatcherBackend.cpp
+++ b/Engine/Source/Runtime/Experimental/IoStore/OnDemand/Private/OnDemandIoDispatcherBackend.cpp
@@ -1546,6 +1546,7 @@ void FOnDemandIoBackend::ProcessHttpRequests(FMultiEndpointHttpClient& HttpClien
 						[this, &NextChunkRequest, ChunkRequest, &NumConcurrentRequests]
 						(FMultiEndpointHttpClientResponse&& HttpResponse)
 						{
+							TRACE_CPUPROFILER_EVENT_SCOPE(IasBackend::RequestCallback);
 							NumConcurrentRequests--;
 
 							if (HttpResponse.RetryCount > 0)
@@ -1632,9 +1633,28 @@ uint32 FOnDemandIoBackend::Run()
 {
 	LLM_SCOPE_BYTAG(Ias);
 
-	if (!CreateDefaultHostGroup())
 	{
-		return 0;
+		TRACE_CPUPROFILER_EVENT_SCOPE(IasBackend::ResolveDefaultHostGroup);
+
+		const double InitStart = FPlatformTime::Seconds();
+
+		if (!CreateDefaultHostGroup())
+		{
+			return 0;
+		}
+
+		if (ConnectionTest(DefaultHostGroup.GetPrimaryHostUrl(), DefaultHostGroup.GetTestPath(), GIasHttpTimeOutMs))
+		{
+			Stats.OnHttpConnected();
+
+		}
+		else
+		{
+			DefaultHostGroup.Disconnect();
+		}
+
+		const double InitTime = FPlatformTime::Seconds() - InitStart;
+		UE_LOG(LogIas, Display, TEXT("HostGroup init took %.3f seconds"), InitTime);
 	}
 
 	TUniquePtr<FMultiEndpointHttpClient> HttpClient = FMultiEndpointHttpClient::Create(FMultiEndpointHttpClientConfig
@@ -1650,27 +1670,12 @@ uint32 FOnDemandIoBackend::Run()
 		});
 	check(HttpClient.IsValid());
 
-	const double InitStart = FPlatformTime::Seconds();
-
-	if (ConnectionTest(DefaultHostGroup.GetPrimaryHostUrl(), DefaultHostGroup.GetTestPath(), GIasHttpTimeOutMs))
-	{
-		Stats.OnHttpConnected();
-		
-	}
-	else
-	{
-		DefaultHostGroup.Disconnect();
-	}
-
 	BackendStatus.SetHttpEnabled(true);
 
-	const double InitTime = FPlatformTime::Seconds() - InitStart;
-	UE_LOG(LogIas, Display, TEXT("HostGroup init took %.3f seconds"), InitTime);
-
 	while (!bStopRequested)
 	{
 		UpdateThreadPriorityIfNeeded();
-		 
+
 		ProcessHttpRequests(*HttpClient, FMath::Min(GIasHttpConcurrentRequests, 32));
 
 		if (!bStopRequested)
