Add support for retrieving soft references from on-demand package store.

#rnx

diff --git a/Engine/Source/Runtime/Experimental/IoStore/OnDemand/Private/OnDemandPackageStoreBackend.cpp b/Engine/Source/Runtime/Experimental/IoStore/OnDemand/Private/OnDemandPackageStoreBackend.cpp
--- a/Engine/Source/Runtime/Experimental/IoStore/OnDemand/Private/OnDemandPackageStoreBackend.cpp
+++ b/Engine/Source/Runtime/Experimental/IoStore/OnDemand/Private/OnDemandPackageStoreBackend.cpp
@@ -18,6 +18,8 @@ namespace UE::IoStore
 class FOnDemandPackageStoreBackend final
 	: public IOnDemandPackageStoreBackend
 {
+	using FSoftPackageReferenceMap = TMap<FPackageId, const FFilePackageStoreEntrySoftReferences*>;
+
 	struct FContainer
 	{
 		FContainer(FString&& ContainerName, FSharedContainerHeader ContainerHeader)
@@ -25,8 +27,9 @@ class FOnDemandPackageStoreBackend final
 			, Header(MoveTemp(ContainerHeader))
 		{ }
 
-		FString					Name;
-		FSharedContainerHeader	Header;
+		FString						Name;
+		FSharedContainerHeader		Header;
+		FSoftPackageReferenceMap	SoftRefs;
 	};
 
 	using FSharedBackendContext = TSharedPtr<const FPackageStoreBackendContext>;
@@ -58,6 +61,8 @@ public:
 							FName PackageName,
 							FPackageStoreEntry& OutPackageStoreEntry) override;
 
+	virtual TConstArrayView<uint32> GetSoftReferences(FPackageId PackageId, TConstArrayView<FPackageId>& OutPackageIds) override;
+
 private:
 	void				UpdateLookupTables();
 
@@ -191,6 +196,19 @@ bool FOnDemandPackageStoreBackend::GetPackageRedirectInfo(
 	return false;
 }
 
+TConstArrayView<uint32> FOnDemandPackageStoreBackend::GetSoftReferences(FPackageId PackageId, TConstArrayView<FPackageId>& OutPackageIds) 
+{
+	for (const FUniqueContainer& Container : Containers)
+	{
+		if (const FFilePackageStoreEntrySoftReferences* SoftRefs = Container->SoftRefs.FindRef(PackageId))
+		{
+			OutPackageIds = Container->Header->SoftPackageReferences.PackageIds;
+			return TConstArrayView<uint32>(SoftRefs->Indices.Data(), SoftRefs->Indices.Num());
+		}
+	}
+	return TConstArrayView<uint32>();
+}
+
 void FOnDemandPackageStoreBackend::UpdateLookupTables()
 {
 	if (!bNeedsUpdate)
@@ -214,17 +232,37 @@ void FOnDemandPackageStoreBackend::UpdateLookupTables()
 
 	for (const FUniqueContainer& Container : Containers)
 	{
+		Container->SoftRefs.Reset();
+
 		const FIoContainerHeader& Hdr = *Container->Header;
 		TConstArrayView<FFilePackageStoreEntry> Entries(
 			reinterpret_cast<const FFilePackageStoreEntry*>(Hdr.StoreEntries.GetData()),
 			Hdr.PackageIds.Num());
 
-		int32 Idx = 0;
-		for (const FFilePackageStoreEntry& Entry : Entries)
+		TConstArrayView<FFilePackageStoreEntrySoftReferences> AllSoftReferences;
+		if (Hdr.SoftPackageReferences.bContainsSoftPackageReferences)
+		{
+			AllSoftReferences = MakeArrayView<const FFilePackageStoreEntrySoftReferences>(
+				reinterpret_cast<const FFilePackageStoreEntrySoftReferences*>(Hdr.SoftPackageReferences.PackageIndices.GetData()),
+				Hdr.PackageIds.Num());
+			Container->SoftRefs.Reserve(Hdr.PackageIds.Num());
+		}
+
+		for (int32 Idx = 0; const FFilePackageStoreEntry& Entry : Entries)
 		{
-			const FPackageId PkgId = Hdr.PackageIds[Idx++];
+			const FPackageId PkgId = Hdr.PackageIds[Idx];
 			EntryMap.Add(PkgId, &Entry);
+			if (AllSoftReferences.IsEmpty() == false)
+			{
+				const FFilePackageStoreEntrySoftReferences& SoftRefs = AllSoftReferences[Idx];
+				if (SoftRefs.Indices.Num() > 0)
+				{
+					Container->SoftRefs.Add(PkgId, &SoftRefs);
+				}
+			}
+			++Idx;
 		}
+		Container->SoftRefs.Shrink();
 
 		for (const FIoContainerHeaderLocalizedPackage& Localized : Hdr.LocalizedPackages)
 		{
@@ -248,6 +286,10 @@ void FOnDemandPackageStoreBackend::UpdateLookupTables()
 			}
 		}
 	}
+
+	EntryMap.Shrink();
+	LocalizedMap.Shrink();
+	RedirectMap.Shrink();
 }
 
 ///////////////////////////////////////////////////////////////////////////////
