Use package store for traversing dependencies when installing individual packages.

- updated on-demand I/O store to use the package store when resolving dependencies. This enables installing dependencies
for base game installed content.
- updated on-demand I/O store to add containers to package store backend when mounting containers instead of when installing content

#rnx

diff --git a/Engine/Source/Runtime/Experimental/IoStore/OnDemand/Private/OnDemandIoStore.cpp b/Engine/Source/Runtime/Experimental/IoStore/OnDemand/Private/OnDemandIoStore.cpp
--- a/Engine/Source/Runtime/Experimental/IoStore/OnDemand/Private/OnDemandIoStore.cpp
+++ b/Engine/Source/Runtime/Experimental/IoStore/OnDemand/Private/OnDemandIoStore.cpp
@@ -13,7 +13,6 @@
 #include "Algo/AnyOf.h"
 #include "Algo/RemoveIf.h"
 #include "Algo/Transform.h"
-#include "Async/ManualResetEvent.h"
 #include "Async/UniqueLock.h"
 #include "Containers/RingBuffer.h"
 #include "Containers/StringConv.h"
@@ -171,9 +170,6 @@ static TIoStatusOr<FSharedContainerHeader> DeserializeContainerHeader(const FOnD
 }
 
 ///////////////////////////////////////////////////////////////////////////////
-using FPackageStoreEntryMap		= TMap<FPackageId, const FFilePackageStoreEntry*>;
-using FSoftPackageReferenceMap	= TMap<FPackageId, const FFilePackageStoreEntrySoftReferences*>;
-
 enum class EChunkCacheState : uint8
 {
 	None = 0,
@@ -183,8 +179,6 @@ enum class EChunkCacheState : uint8
 
 struct FContainerInstallData
 {
-	FPackageStoreEntryMap			PackageStoreEntries;
-	FSoftPackageReferenceMap		SoftPackageReferences;
 	TSet<FPackageId>				ResolvedPackageIds;
 	TMap<uint32, EChunkCacheState>	ResolvedChunks;
 	uint64							TotalSize = 0;
@@ -206,6 +200,8 @@ static FIoStatus BuildInstallData(
 	FInstallData& OutInstallData,
 	TSet<FPackageId>& OutMissing)
 {
+	TRACE_CPUPROFILER_EVENT_SCOPE(FOnDemandIoStore::BuildInstallData);
+
 	OutInstallData.Reserve(Containers.Num());
 
 	// For now we always download these required chunks
@@ -236,176 +232,122 @@ static FIoStatus BuildInstallData(
 		return EIoErrorCode::Ok;
 	}
 
-	// Setup the package information for each container
-	for (const FSharedOnDemandContainer& Container : Containers)
+	// Traverse dependencies for each package ID
+	TArray<FPackageId> ResolvedPackageIds;
 	{
-		FContainerInstallData& Data = OutInstallData.FindOrAdd(Container);
+		TRACE_CPUPROFILER_EVENT_SCOPE(FOnDemandIoStore::ResolveDependencies);
 
-		if (Container->Header.IsValid() == false || Container->Header->PackageIds.IsEmpty())
-		{
-			// The container contains no package data
-			continue;
-		}
-
-		const FIoContainerHeader& Header = *Container->Header;
-		TConstArrayView<FFilePackageStoreEntry> Entries(
-			reinterpret_cast<const FFilePackageStoreEntry*>(Header.StoreEntries.GetData()),
-			Header.PackageIds.Num());
-		
-		Data.PackageStoreEntries.Reserve(Header.PackageIds.Num());
+		TRingBuffer<FPackageId> Queue;
+		TSet<FPackageId>		Visitied;
 
-		TConstArrayView<FFilePackageStoreEntrySoftReferences> SoftReferences;
-		if (Header.SoftPackageReferences.bContainsSoftPackageReferences)
+		Visitied.Reserve(PackageIds.Num());
+		Queue.Reserve(PackageIds.Num());
+		for (const FPackageId& PackageId : PackageIds)
 		{
-			Data.PackageStoreEntries.Reserve(Header.PackageIds.Num());
-			SoftReferences = MakeArrayView<const FFilePackageStoreEntrySoftReferences>(
-				reinterpret_cast<const FFilePackageStoreEntrySoftReferences*>(Header.SoftPackageReferences.PackageIndices.GetData()),
-				Header.PackageIds.Num());
+			Queue.Add(PackageId);
 		}
 
-		for (int32 PackageIndex = 0; const FFilePackageStoreEntry& Entry : Entries)
+		FPackageStoreReadScope _(FPackageStore::Get());
+
+		while (!Queue.IsEmpty())
 		{
-			const FPackageId PackageId = Header.PackageIds[PackageIndex];
-			Data.PackageStoreEntries.Add(PackageId, &Entry);
+			FPackageId PackageId = Queue.PopFrontValue();
 
-			if (SoftReferences.IsEmpty() == false)
+			bool bIsAlreadyInSet = false;
+			Visitied.Add(PackageId, &bIsAlreadyInSet);
+			if (bIsAlreadyInSet)
 			{
-				Data.SoftPackageReferences.Add(PackageId, &SoftReferences[PackageIndex]);
+				continue;
 			}
 
-			++PackageIndex;
-		}
-	}
-
-	// Traverse dependencies for each package ID
-	using FQueue = TRingBuffer<FPackageId>;
-
-	FQueue Queue;
-	TSet<FPackageId> Visitied;
-
-	Visitied.Reserve(PackageIds.Num());
-	Queue.Reserve(PackageIds.Num());
-
-	//Algo::Copy(PackageIds, Queue);
-	for (const FPackageId& PackageId : PackageIds)
-	{
-		Queue.Add(PackageId);
-	}
-
-	while (!Queue.IsEmpty())
-	{
-		FPackageId PackageId = Queue.PopFrontValue();
-
-		bool bIsAlreadyInSet = false;
-		Visitied.Add(PackageId, &bIsAlreadyInSet);
-		if (bIsAlreadyInSet)
-		{
-			continue;
-		}
-
-		bool bFound = false;
-		for (TPair<FSharedOnDemandContainer, FContainerInstallData>& Kv : OutInstallData)
-		{
-			FContainerInstallData& Data = Kv.Value;
-			if (const FFilePackageStoreEntry* Entry = Data.PackageStoreEntries.FindRef(PackageId))
+			FPackageStoreEntry PackageStoreEntry;
+			const EPackageStoreEntryStatus EntryStatus = FPackageStore::Get().GetPackageStoreEntry(PackageId, NAME_None, PackageStoreEntry);
+			if (EntryStatus != EPackageStoreEntryStatus::Missing)
 			{
-				Data.ResolvedPackageIds.Add(PackageId);
-
-				// Add hard references 
-				for (const FPackageId& ImportedPackageId : Entry->ImportedPackages)
+				ResolvedPackageIds.Add(PackageId);
+				for (const FPackageId& ImportedPackageId : PackageStoreEntry.ImportedPackageIds)
 				{
 					if (!Visitied.Contains(ImportedPackageId))
 					{
 						Queue.Add(ImportedPackageId);
 					}
-				}
 
-				if (EnumHasAnyFlags(Options, EBuildInstallDataOptions::IncludeSoftReferences))
-				{
-					// Add soft references
-					if (const FFilePackageStoreEntrySoftReferences* SoftRefs = Data.SoftPackageReferences.FindRef(PackageId))
+					if (EnumHasAnyFlags(Options, EBuildInstallDataOptions::IncludeSoftReferences))
 					{
-						const FIoContainerHeader& Header = *Kv.Key->Header;
-						for (uint32 Index : SoftRefs->Indices)
+						TConstArrayView<FPackageId> SoftReferences;
+						TConstArrayView<uint32> Indices = FPackageStore::Get().GetSoftReferences(PackageId, SoftReferences);
+						for (uint32 Idx : Indices)
 						{
-							const FPackageId SoftPackageReference = Header.SoftPackageReferences.PackageIds[int32(Index)];
-							if (!Visitied.Contains(SoftPackageReference))
+							const FPackageId& SoftRef = SoftReferences[Idx];
+							if (!Visitied.Contains(SoftRef))
 							{
-								Queue.Add(SoftPackageReference);
+								Queue.Add(SoftRef);
 							}
 						}
 					}
 				}
-
-				bFound = true;
-				break;
+			}
+			else
+			{
+				OutMissing.Add(PackageId);
 			}
 		}
+	}
 
-		if (!bFound)
+	auto FindChunkEntry = [&OutInstallData](const FIoChunkId& ChunkId, int32& OutIndex) -> TPair<FSharedOnDemandContainer, FContainerInstallData>*
+	{
+		for (TPair<FSharedOnDemandContainer, FContainerInstallData>& Kv : OutInstallData)
 		{
-			OutMissing.Add(PackageId);
+			if (OutIndex = Kv.Key->FindChunkEntryIndex(ChunkId); OutIndex != INDEX_NONE)
+			{
+				return &Kv;
+			}
 		}
-	}
 
-	// Resolve all chunk ID's from the package ID's
-	for (TPair<FSharedOnDemandContainer, FContainerInstallData>& Kv : OutInstallData)
+		return nullptr; 
+	};
+
+	// Resolve all chunk ID's from the resolved package ID's
+	for (const FPackageId& PackageId : ResolvedPackageIds)
 	{
-		const FOnDemandContainer& Container = *Kv.Key;
-		FContainerInstallData& Data			= Kv.Value;
+		const FIoChunkId PackageChunkId								= CreatePackageDataChunkId(PackageId);
+		int32 EntryIndex											= INDEX_NONE; 
+		TPair<FSharedOnDemandContainer, FContainerInstallData>*	Kv	= FindChunkEntry(PackageChunkId, EntryIndex);
 
-		if (Data.ResolvedPackageIds.IsEmpty())
+		if (Kv == nullptr) 
 		{
+			// The chunk resides in a base game container
 			continue;
 		}
 
-		for (const FPackageId& PackageId : Data.ResolvedPackageIds)
-		{
-			const FIoChunkId PackageChunkId	= CreatePackageDataChunkId(PackageId);
-			int32 EntryIndex				= Container.FindChunkEntryIndex(PackageChunkId);
+		check(EntryIndex != INDEX_NONE);
+		FOnDemandContainer& Container		= *Kv->Key;
+		FContainerInstallData& InstallData	= Kv->Value;
 
-			if (EntryIndex == INDEX_NONE)
-			{
-				UE_LOG(LogIoStoreOnDemand, Warning, TEXT("Missing package data"));
-				continue;
-			}
+		InstallData.ResolvedPackageIds.Add(PackageId);
+		InstallData.ResolvedChunks.Add(EntryIndex, EChunkCacheState::None);
+		InstallData.TotalSize += Container.ChunkEntries[EntryIndex].EncodedSize;
 
-			Data.ResolvedChunks.Add(EntryIndex, EChunkCacheState::None);
-			Data.TotalSize += Container.ChunkEntries[EntryIndex].EncodedSize;
+		// TODO: Installing optional bulkdata should probably be an install argument
+		const EIoChunkType AdditionalPackageChunkTypes[] =
+		{
+			EIoChunkType::BulkData,
+			EIoChunkType::OptionalBulkData,
+			EIoChunkType::MemoryMappedBulkData 
+		};
 
-			// TODO: Installing optional bulkdata should probably be an install argument
-			const EIoChunkType AdditionalPackageChunkTypes[] =
+		for (EIoChunkType ChunkType : AdditionalPackageChunkTypes)
+		{
+			// TODO: For Mutable we need to traverse all possible bulk data chunk indices?
+			const FIoChunkId ChunkId = CreateBulkDataIoChunkId(PackageId.Value(), 0, 0, ChunkType);
+			if (Kv = FindChunkEntry(ChunkId, EntryIndex); Kv != nullptr)
 			{
-				EIoChunkType::BulkData,
-				EIoChunkType::OptionalBulkData,
-				EIoChunkType::MemoryMappedBulkData 
-			};
+				check(EntryIndex != INDEX_NONE);
+				FOnDemandContainer& OtherContainer		= *Kv->Key;
+				FContainerInstallData& OtherInstallData	= Kv->Value;
 
-			for (EIoChunkType ChunkType : AdditionalPackageChunkTypes)
-			{
-				const FIoChunkId ChunkId = CreateBulkDataIoChunkId(PackageId.Value(), 0, 0, ChunkType);
-				if (EntryIndex = Container.FindChunkEntryIndex(ChunkId); EntryIndex != INDEX_NONE)
-				{
-					Data.ResolvedChunks.Add(EntryIndex, EChunkCacheState::None);
-					Data.TotalSize += Container.ChunkEntries[EntryIndex].EncodedSize;
-				}
-				else
-				{
-					// Some chunk types resides in other containers
-					for (TPair<FSharedOnDemandContainer, FContainerInstallData>& OtherKv : OutInstallData)
-					{
-						const FOnDemandContainer& OtherContainer = *OtherKv.Key;
-						if (&Container != &OtherContainer)
-						{
-							if (EntryIndex = OtherContainer.FindChunkEntryIndex(ChunkId); EntryIndex != INDEX_NONE)
-							{
-								FContainerInstallData& OtherData = OtherKv.Value;
-								OtherData.ResolvedChunks.Add(EntryIndex, EChunkCacheState::None);
-								OtherData.TotalSize += OtherContainer.ChunkEntries[EntryIndex].EncodedSize;
-							}
-						}
-					}
-				}
+				OtherInstallData.ResolvedChunks.Add(EntryIndex, EChunkCacheState::None);
+				OtherInstallData.TotalSize += OtherContainer.ChunkEntries[EntryIndex].EncodedSize;
 			}
 		}
 	}
@@ -1273,6 +1215,10 @@ FIoStatus FOnDemandIoStore::TickMountRequest(FMountRequest& MountRequest)
 				}
 
 				Container->Header = Header.ConsumeValueOrDie(); 
+				if (FIoStatus MountStatus = PackageStoreBackend->Mount(Container->UniqueName(), Container->Header); !MountStatus.IsOk())
+				{
+					return MountStatus;
+				}
 			}
 		}
 
@@ -1318,6 +1264,8 @@ void FOnDemandIoStore::CompleteMountRequest(FMountRequest& Request, FOnDemandMou
 
 FOnDemandInstallResult FOnDemandIoStore::TickInstallRequest(FOnDemandInternalInstallRequest& InstallRequest)
 {
+	TRACE_CPUPROFILER_EVENT_SCOPE(FOnDemandIoStore::InstallContent);
+
 	UE_LOG(LogIoStoreOnDemand, Log, TEXT("Ticking install request, ContentHandle='%s'"),
 		*LexToString(InstallRequest.Args.ContentHandle));
 
@@ -1418,6 +1366,8 @@ FOnDemandInstallResult FOnDemandIoStore::TickInstallRequest(FOnDemandInternalIns
 
 	// Purge
 	{
+		TRACE_CPUPROFILER_EVENT_SCOPE(FOnDemandIoStore::PurgeInstallCache);
+
 		TMap<FIoHash, uint64> ChunksToInstall;
 		for (const auto& Kv : InstallData)
 		{
@@ -1602,18 +1552,11 @@ FOnDemandInstallResult FOnDemandIoStore::TickInstallRequest(FOnDemandInternalIns
 		ContentHandle->IoStore = AsWeak(); 
 	}
 
-	// TODO: Only mount what has been installed
 	for (const auto& Kv : InstallData)
 	{
 		FOnDemandContainer& Container				= *Kv.Key;
 		const Private::FContainerInstallData& Data	= Kv.Value;
-		if (!Data.ResolvedPackageIds.IsEmpty())
-		{
-			check(Container.Header.IsValid());
-			const FIoStatus MountStatus = PackageStoreBackend->Mount(Container.UniqueName(), Container.Header);
-			check(MountStatus.IsOk());
-		}
-
+		// TODO: Mark installed packages in package store backend 
 		{
 			UE::TUniqueLock Lock(ContainerMutex);
 			FOnDemandChunkEntryReferences& References = Container.FindOrAddChunkEntryReferences(*ContentHandle);
@@ -1760,7 +1703,15 @@ void FOnDemandIoStore::OnEncryptionKeyAdded(const FGuid& Id, const FAES::FAESKey
 				if (ensure(Header.IsOk()))
 				{
 					Container->Header = Header.ConsumeValueOrDie(); 
-					EnumAddFlags(Container->Flags, EOnDemandContainerFlags::Mounted);
+					const FIoStatus MountStatus = PackageStoreBackend->Mount(Container->UniqueName(), Container->Header);
+					if (MountStatus.IsOk())
+					{
+						EnumAddFlags(Container->Flags, EOnDemandContainerFlags::Mounted);
+					}
+					else
+					{
+						UE_LOG(LogIoStoreOnDemand, Error, TEXT("Failed to add container '%s' to package store backend"), *Container->UniqueName()); 
+					}
 				}
 				else
 				{
