Dataflow : concrete type change on connection now properly disconnect or reconnect existing connections
- Dataflow node now hold a weak ptr to its owning graph
- Added a SetConnectionConcreteType method to FDataflowNode
- All existing Set*ConcreteType method now use SetConnectionConcreteType
- Changed get variable node to use the new SetConnectionConcreteType method

diff --git a/Engine/Source/Runtime/Experimental/Dataflow/Core/Private/Dataflow/DataflowNode.cpp b/Engine/Source/Runtime/Experimental/Dataflow/Core/Private/Dataflow/DataflowNode.cpp
--- a/Engine/Source/Runtime/Experimental/Dataflow/Core/Private/Dataflow/DataflowNode.cpp
+++ b/Engine/Source/Runtime/Experimental/Dataflow/Core/Private/Dataflow/DataflowNode.cpp
@@ -5,6 +5,7 @@
 #include "ChaosLog.h"
 #include "Dataflow/DataflowInputOutput.h"
 #include "Dataflow/DataflowArchive.h"
+#include "Dataflow/DataflowGraph.h"
 #include "Serialization/ObjectWriter.h"
 #include "Serialization/ObjectReader.h"
 #include "Templates/TypeHash.h"
@@ -109,6 +110,24 @@ static TArray<FString> GetPinMetaDataFromProperty(const FProperty* Property)
 }
 };
 
+FDataflowNode::FDataflowNode()
+	: Guid(FGuid())
+	, Name("Invalid")
+	, LastModifiedTimestamp(UE::Dataflow::FTimestamp::Invalid)
+{
+}
+
+FDataflowNode::FDataflowNode(const UE::Dataflow::FNodeParameters& Param, FGuid InGuid)
+	: Guid(InGuid)
+	, Name(Param.Name)
+	, LastModifiedTimestamp(UE::Dataflow::FTimestamp::Invalid)
+{
+	if (IDataflowGraphInterface* Interface = Cast<IDataflowGraphInterface>(Param.OwningObject))
+	{
+		WeakDataflowGraph = Interface->GetDataflowGraph().ToWeakPtr();
+	}
+}
+
 FDataflowConnection* FDataflowNode::FindConnection(const UE::Dataflow::FConnectionKey& Key)
 {
 	FDataflowConnection* Connection = FindInput(Key);
@@ -1399,6 +1418,11 @@ void FDataflowNode::CopyNodeProperties(const TSharedPtr<FDataflowNode> CopyFromD
 	this->SerializeInternal(ArReader);
 }
 
+TSharedPtr<UE::Dataflow::FGraph> FDataflowNode::GetDataflowGraph() const
+{
+	return WeakDataflowGraph.Pin();
+}
+
 void FDataflowNode::ForwardInput(UE::Dataflow::FContext& Context, const UE::Dataflow::FConnectionReference& InputReference, const UE::Dataflow::FConnectionReference& Reference) const
 {
 	if (const FDataflowOutput* Output = FindOutput(Reference))
@@ -1425,9 +1449,11 @@ bool FDataflowNode::TrySetConnectionType(FDataflowConnection* Connection, FName
 	{
 		if (Connection->IsAnyType() && Connection->GetType() != NewType && !FDataflowConnection::IsAnyType(NewType))
 		{
-			Connection->SetConcreteType(NewType);
-			NotifyConnectionTypeChanged(Connection);
-			return true;
+			if (SetConnectionConcreteType(Connection, NewType))
+			{
+				NotifyConnectionTypeChanged(Connection);
+				return true;
+			}
 		}
 	}
 	return false;
@@ -1472,13 +1498,11 @@ bool FDataflowNode::OnConnectionTypeChanged(const FDataflowConnection& Connectio
 	bool bTypeChanged = false;
 	const FName NewType = Connection.GetType();
 	ForEachConnection(
-		[&bTypeChanged, &DependencyGroup, &NewType](FDataflowConnection* OtherConnection)
+		[this, &bTypeChanged, &DependencyGroup, &NewType](FDataflowConnection* OtherConnection)
 		{
-			if (OtherConnection
-				&& (OtherConnection->GetTypeDependencyGroup() == DependencyGroup)
-				&& (OtherConnection->GetType() != NewType))
+			if (OtherConnection && (OtherConnection->GetTypeDependencyGroup() == DependencyGroup))
 			{
-				OtherConnection->SetConcreteType(NewType);
+				SetConnectionConcreteType(OtherConnection, NewType);
 				bTypeChanged = true;
 			}
 			return true;
@@ -1489,38 +1513,17 @@ bool FDataflowNode::OnConnectionTypeChanged(const FDataflowConnection& Connectio
 
 bool FDataflowNode::SetConnectionConcreteType(const UE::Dataflow::FConnectionKey& ConnectionKey, FName NewType)
 {
-	if (FDataflowConnection* Connection = FindConnection(ConnectionKey))
-	{
-		if (Connection->GetType() != NewType)
-		{
-			return Connection->SetConcreteType(NewType);
-		}
-	}
-	return false;
+	return SetConnectionConcreteType(FindConnection(ConnectionKey), NewType);
 }
 
 bool FDataflowNode::SetInputConcreteType(const UE::Dataflow::FConnectionReference& InputReference, FName NewType)
 {
-	if (FDataflowInput* Input = FindInput(InputReference))
-	{
-		if (Input->GetType() != NewType)
-		{
-			return Input->SetConcreteType(NewType);
-		}
-	}
-	return false;
+	return SetConnectionConcreteType(FindInput(InputReference), NewType);
 }
 
 bool FDataflowNode::SetOutputConcreteType(const UE::Dataflow::FConnectionReference& OutputReference, FName NewType)
 {
-	if (FDataflowOutput* Output = FindOutput(OutputReference))
-	{
-		if (Output->GetType() != NewType)
-		{
-			return Output->SetConcreteType(NewType);
-		}
-	}
-	return false;
+	return SetConnectionConcreteType(FindOutput(OutputReference), NewType);
 }
 
 bool FDataflowNode::SetAllConnectionConcreteType(FName NewType)
@@ -1528,23 +1531,52 @@ bool FDataflowNode::SetAllConnectionConcreteType(FName NewType)
 	bool bChanged = false;
 	for (const TPair<UE::Dataflow::FConnectionKey, FDataflowInput*>& InputEntry : ExpandedInputs)
 	{
-		FDataflowInput* const Input = InputEntry.Value;
-		if (Input && Input->GetType() != NewType)
-		{
-			bChanged |= Input->SetConcreteType(NewType);
-		}
+		bChanged |= SetConnectionConcreteType(InputEntry.Value, NewType);
 	}
 
 	for (const TPair<UE::Dataflow::FConnectionKey, FDataflowOutput*>& OutputEntry : ExpandedOutputs)
 	{
-		FDataflowOutput* const Output = OutputEntry.Value;
-		if (Output && Output->GetType() != NewType)
+		bChanged |= SetConnectionConcreteType(OutputEntry.Value, NewType);
+	}
+	return bChanged;
+}
+
+bool FDataflowNode::SetConnectionConcreteType(FDataflowConnection* Connection, FName NewType)
+{
+	bool bSuccess = false;
+	if (Connection)
+	{
+		if (Connection->GetType() != NewType)
 		{
-			bChanged |= Output->SetConcreteType(NewType);
+			if (TSharedPtr<UE::Dataflow::FGraph> DataflowGraph = GetDataflowGraph())
+			{
+				// first save connections  and disconnect fomr them 
+				TArray<FDataflowConnection*> RemoteConnections;
+				Connection->GetConnections(RemoteConnections);
+				DataflowGraph->ClearConnections(Connection);
+
+				// try changing the type
+				const bool bWasTypeLocked = Connection->IsTypeLocked();
+				Connection->UnlockType();
+
+				bSuccess = Connection->SetConcreteType(NewType);
+
+				// we need to set the lock state before reconnecting because the connect logic will use it 
+				// to know which side of the connection allows changes
+				if (bWasTypeLocked)
+				{
+					Connection->LockType();
+				}
+
+				// now try to reconnect (if the types are no longer compatible, connection will be dropped)
+				for (FDataflowConnection* RemoteConnection : RemoteConnections)
+				{
+					DataflowGraph->Connect(Connection, RemoteConnection);
+				}
+			}
 		}
 	}
-
-	return bChanged;
+	return bSuccess;
 }
 
 bool FDataflowNode::IsAnytypeDependencyConnected(FName DependencyGroup, const FDataflowConnection* IgnoreConnection) const
diff --git a/Engine/Source/Runtime/Experimental/Dataflow/Core/Public/Dataflow/DataflowGraph.h b/Engine/Source/Runtime/Experimental/Dataflow/Core/Public/Dataflow/DataflowGraph.h
--- a/Engine/Source/Runtime/Experimental/Dataflow/Core/Public/Dataflow/DataflowGraph.h
+++ b/Engine/Source/Runtime/Experimental/Dataflow/Core/Public/Dataflow/DataflowGraph.h
@@ -3,16 +3,38 @@
 #pragma once
 
 #include "CoreMinimal.h"
+#include "UObject/Interface.h"
 #include "Chaos/ChaosArchive.h"
 #include "Dataflow/DataflowNode.h"
 #include "Dataflow/DataflowTerminalNode.h"
 #include "Serialization/Archive.h"
 
+#include "DataflowGraph.generated.h"
+
 struct FDataflowConnection;
 
 namespace UE::Dataflow
 {
+	class FGraph;
+}
 
+/** Interface for objects that owns a dataflow graph */
+UINTERFACE(MinimalAPI)
+class UDataflowGraphInterface : public UInterface
+{
+	GENERATED_BODY()
+};
+
+class IDataflowGraphInterface
+{
+	GENERATED_BODY()
+
+public:
+	virtual TSharedPtr<UE::Dataflow::FGraph> GetDataflowGraph() const = 0;
+};
+
+namespace UE::Dataflow
+{
 	struct FLink {
 		FGuid InputNode;
 		FGuid Input;
diff --git a/Engine/Source/Runtime/Experimental/Dataflow/Core/Public/Dataflow/DataflowNode.h b/Engine/Source/Runtime/Experimental/Dataflow/Core/Public/Dataflow/DataflowNode.h
--- a/Engine/Source/Runtime/Experimental/Dataflow/Core/Public/Dataflow/DataflowNode.h
+++ b/Engine/Source/Runtime/Experimental/Dataflow/Core/Public/Dataflow/DataflowNode.h
@@ -72,19 +72,8 @@ PRAGMA_ENABLE_DEPRECATION_WARNINGS
 	UPROPERTY(EditAnywhere, Category = "Dataflow", meta = (EditCondition="bOverrideColor"))
 	FLinearColor OverrideColor = FLinearColor::White;
 
-	FDataflowNode()
-		: Guid(FGuid())
-		, Name("Invalid")
-		, LastModifiedTimestamp(UE::Dataflow::FTimestamp::Invalid)
-	{
-	}
-
-	FDataflowNode(const UE::Dataflow::FNodeParameters& Param, FGuid InGuid = FGuid::NewGuid())
-		: Guid(InGuid)
-		, Name(Param.Name)
-		, LastModifiedTimestamp(UE::Dataflow::FTimestamp::Invalid)
-	{
-	}
+	DATAFLOWCORE_API FDataflowNode();
+	DATAFLOWCORE_API FDataflowNode(const UE::Dataflow::FNodeParameters& Param, FGuid InGuid = FGuid::NewGuid());
 
 	PRAGMA_DISABLE_DEPRECATION_WARNINGS
 	virtual ~FDataflowNode() { ClearInputs(); ClearOutputs(); }
@@ -757,6 +746,8 @@ protected:
 	* */
 	DATAFLOWCORE_API bool SetConnectionConcreteType(const UE::Dataflow::FConnectionKey& ConnectionKey, FName NewType);
 
+	DATAFLOWCORE_API bool SetConnectionConcreteType(FDataflowConnection* Connection, FName NewType);
+
 	// returns true if any of the types was changed successfully
 	DATAFLOWCORE_API bool SetAllConnectionConcreteType(FName NewType);
 
@@ -765,6 +756,8 @@ protected:
 	DATAFLOWCORE_API FDataflowInput& RegisterInputArrayConnectionInternal(const UE::Dataflow::FConnectionReference& Reference, const FName& ElementPropertyName = NAME_None, const FName& ArrayPropertyName = NAME_None);
 	DATAFLOWCORE_API FDataflowOutput& RegisterOutputArrayConnectionInternal(const UE::Dataflow::FConnectionReference& Reference, const FName& ElementPropertyName = NAME_None, const FName& ArrayPropertyName = NAME_None);
 
+	DATAFLOWCORE_API TSharedPtr<UE::Dataflow::FGraph> GetDataflowGraph() const;
+
 private:
 
 	void InitConnectionParametersFromPropertyReference(const FStructOnScope& StructOnScope, const void* PropertyRef, const FName& PropertyName, UE::Dataflow::FConnectionParameters& OutParams);
@@ -809,6 +802,8 @@ private:
 	UPROPERTY()
 	bool bIsFrozen = false;
 
+	TWeakPtr<UE::Dataflow::FGraph> WeakDataflowGraph;
+
 protected:
 	bool bPauseInvalidations = false;
 	UE::Dataflow::FTimestamp PausedModifiedTimestamp = UE::Dataflow::FTimestamp::Invalid; // When unpausing invalidations, Invalidate will be called with this timestamp.
diff --git a/Engine/Source/Runtime/Experimental/Dataflow/Engine/Private/Dataflow/DataflowVariableNodes.cpp b/Engine/Source/Runtime/Experimental/Dataflow/Engine/Private/Dataflow/DataflowVariableNodes.cpp
--- a/Engine/Source/Runtime/Experimental/Dataflow/Engine/Private/Dataflow/DataflowVariableNodes.cpp
+++ b/Engine/Source/Runtime/Experimental/Dataflow/Engine/Private/Dataflow/DataflowVariableNodes.cpp
@@ -122,27 +122,10 @@ void FGetDataflowVariableNode::SetVariable(UDataflow* Dataflow, FName InVariable
 
 void FGetDataflowVariableNode::ChangeOutputType(FDataflowOutput* Output, FName NewType)
 {
-	// make sure to disconnect
 	if (Output && Output->GetType() != NewType)
 	{
-		TStrongObjectPtr<UDataflow> DataflowAsset = WeakDataflowPtr.Pin();
-		const TArray<FDataflowInput*> InputsToReconnect = Output->GetConnections();
-		if (DataflowAsset)
-		{
-			DataflowAsset->GetDataflow()->ClearConnections(Output);
-		}
-		
-		Output->UnlockType();
-		Output->SetConcreteType(NewType);
+		SetConnectionConcreteType(Output, NewType);
 		Output->LockType();
-
-		if (DataflowAsset)
-		{
-			for (FDataflowInput* InputToReconnect : InputsToReconnect)
-			{
-				DataflowAsset->GetDataflow()->Connect(InputToReconnect, Output);
-			}
-		}
 	}
 }
 
diff --git a/Engine/Source/Runtime/Experimental/Dataflow/Engine/Public/Dataflow/DataflowObject.h b/Engine/Source/Runtime/Experimental/Dataflow/Engine/Public/Dataflow/DataflowObject.h
--- a/Engine/Source/Runtime/Experimental/Dataflow/Engine/Public/Dataflow/DataflowObject.h
+++ b/Engine/Source/Runtime/Experimental/Dataflow/Engine/Public/Dataflow/DataflowObject.h
@@ -98,7 +98,7 @@ enum class EDataflowType : uint8
 *
 */
 UCLASS(BlueprintType, customconstructor, MinimalAPI)
-class UDataflow : public UEdGraph, public IMeshDeformerProducer
+class UDataflow : public UEdGraph, public IMeshDeformerProducer, public IDataflowGraphInterface
 {
 	GENERATED_UCLASS_BODY()
 
@@ -130,6 +130,10 @@ public:
 	DATAFLOWENGINE_API virtual FMeshDeformerBeginDestroyEvent& OnBeginDestroy() override {return BeginDestroyEvent;};
 	//~ End IMeshDeformerProducer interface
 
+	//~ Begin IDataflowGraphInterface interface
+	virtual TSharedPtr<UE::Dataflow::FGraph> GetDataflowGraph() const override { return Dataflow; }
+	//~ End IDataflowGraphInterface interface
+
 public:
 	UPROPERTY(EditAnywhere, Category = "Evaluation")
 	bool bActive = true;
