Dataflow : Add ReadValue method for anytypes

diff --git a/Engine/Source/Runtime/Experimental/Dataflow/Core/Public/Dataflow/DataflowInputOutput.h b/Engine/Source/Runtime/Experimental/Dataflow/Core/Public/Dataflow/DataflowInputOutput.h
--- a/Engine/Source/Runtime/Experimental/Dataflow/Core/Public/Dataflow/DataflowInputOutput.h
+++ b/Engine/Source/Runtime/Experimental/Dataflow/Core/Public/Dataflow/DataflowInputOutput.h
@@ -5,6 +5,7 @@
 #include "CoreMinimal.h"
 #include "Dataflow/DataflowNodeParameters.h"
 #include "Dataflow/DataflowConnection.h"
+#include "Dataflow/DataflowAnyType.h"
 #include "Templates/Function.h"
 #include "Async/Async.h"
 #include "GenericPlatform/GenericPlatformCriticalSection.h"
@@ -242,7 +243,7 @@ public:
 	* if there's no cached value this will return the default value passed as a parameter 
 	* @see GetValue
 	*/
-	template<class T>
+	template<class T, typename = std::enable_if_t<!std::is_base_of_v<FDataflowAnyType, T>>>
 	const T& ReadValue(UE::Dataflow::FContext& Context, const T& Default) const
 	{
 		if (HasFrozenValue())
@@ -256,6 +257,32 @@ public:
 		return Default;
 	}
 
+	template<typename TAnyType, typename = std::enable_if_t<std::is_base_of_v<FDataflowAnyType, TAnyType>>>
+	typename TAnyType::FStorageType ReadValue(UE::Dataflow::FContext& Context, const typename TAnyType::FStorageType& Default) const
+	{
+		if (HasFrozenValue())
+		{
+			return GetFrozenValue(Default);
+		}
+		if (const TUniquePtr<UE::Dataflow::FContextCacheElementBase>* CacheEntry = Context.GetDataImpl(CacheKey()))
+		{
+			if (*CacheEntry)
+			{
+				typename TAnyType::FStorageType ReturnValue = Default;
+				TAnyType::FPolicyType::VisitPolicyByType(GetType(),
+					[this, &Context, &CacheEntry, &ReturnValue](auto SingleTypePolicy)
+					{
+						using FSingleType = typename decltype(SingleTypePolicy)::FType;
+						FSingleType Default{};
+						const FSingleType& CachedValue = (*CacheEntry)->GetTypedData<FSingleType>(Context, nullptr, Default);
+						FDataflowConverter<typename TAnyType::FStorageType>::From(CachedValue, ReturnValue);
+					});
+				return ReturnValue;
+			}
+		}
+		return Default;
+	}
+
 	/*
 	* Get most updated value of the output
 	* if the value is cache or frozen return it otherwise evaluates the node with potentially cascading evaluation of the graph 
