Dataflow : had method to check if an output has any valid data

diff --git a/Engine/Source/Runtime/Experimental/Dataflow/Core/Private/Dataflow/DataflowInputOutput.cpp b/Engine/Source/Runtime/Experimental/Dataflow/Core/Private/Dataflow/DataflowInputOutput.cpp
--- a/Engine/Source/Runtime/Experimental/Dataflow/Core/Private/Dataflow/DataflowInputOutput.cpp
+++ b/Engine/Source/Runtime/Experimental/Dataflow/Core/Private/Dataflow/DataflowInputOutput.cpp
@@ -251,6 +251,19 @@ void FDataflowOutput::Invalidate(const UE::Dataflow::FTimestamp& ModifiedTimesta
 	}
 }
 
+bool FDataflowOutput::HasValidData(UE::Dataflow::FContext& Context) const
+{
+	if (HasFrozenValue())
+	{
+		return true;
+	}
+	if (OwningNode && Context.HasData(CacheKey(), OwningNode->GetTimestamp()))
+	{
+		return true;
+	}
+	return false;
+}
+
 bool FDataflowOutput::Evaluate(UE::Dataflow::FContext& Context) const
 {
 	check(OwningNode);
@@ -264,7 +277,6 @@ bool FDataflowOutput::Evaluate(UE::Dataflow::FContext& Context) const
 		ForwardInput(PassthroughInput, Context);
 		return true;
 	}
-
 	return false;
 }
 
diff --git a/Engine/Source/Runtime/Experimental/Dataflow/Core/Public/Dataflow/DataflowInputOutput.h b/Engine/Source/Runtime/Experimental/Dataflow/Core/Public/Dataflow/DataflowInputOutput.h
--- a/Engine/Source/Runtime/Experimental/Dataflow/Core/Public/Dataflow/DataflowInputOutput.h
+++ b/Engine/Source/Runtime/Experimental/Dataflow/Core/Public/Dataflow/DataflowInputOutput.h
@@ -329,6 +329,7 @@ public:
 	DATAFLOWCORE_API void ForwardInput(const UE::Dataflow::FConnectionReference& InputReference, UE::Dataflow::FContext& Context) const;
 	DATAFLOWCORE_API void ForwardInput(const FDataflowInput* Input, UE::Dataflow::FContext& Context) const;
 
+	DATAFLOWCORE_API bool HasValidData(UE::Dataflow::FContext& Context) const;
 
 	DATAFLOWCORE_API bool EvaluateImpl(UE::Dataflow::FContext& Context) const;
 	
