Update how FMultiEndpointHttpClientConfig::MaxRetryCount is used.

#rnx

- Previously this was the only way top set a retry count for a http client and that value would be applied to all host groups. However usually we want the retry count to be the same as the number of entries in a host group.
- Now that MaxRetryCount logic is as follows
-- A negative number means that the client should just use the number of urls as the retry count.
-- A zero value means that requests should never be retried.
-- A positive value will be the number of times a request will be retried.
- As before the retry count does not include the initial attempts. So if a retry count of 3 is used then each request can have potentially up to 4 requests.
- This does not change the retry logic, the first retry will be for the original url and subsequent retried will try the next url in the host group.
- This removes the last link that we had between DefaultHostGroup and the http client in FOnDemandIoBackend.

diff --git a/Engine/Source/Runtime/Experimental/IoStore/OnDemand/Private/OnDemandHttpClient.cpp b/Engine/Source/Runtime/Experimental/IoStore/OnDemand/Private/OnDemandHttpClient.cpp
--- a/Engine/Source/Runtime/Experimental/IoStore/OnDemand/Private/OnDemandHttpClient.cpp
+++ b/Engine/Source/Runtime/Experimental/IoStore/OnDemand/Private/OnDemandHttpClient.cpp
@@ -125,7 +125,8 @@ TIoStatusOr<FMultiEndpointHttpClientResponse> FMultiEndpointHttpClient::Get(FAns
 	FIoBuffer Body;
 	TStringBuilder<128> Reason;
 	uint32 StatusCode = 0;
-	const uint32 MaxAttempts = FMath::Min(Config.MaxRetryCount, 3u);
+
+	const uint32 MaxAttempts = Config.MaxRetryCount == -1 ? 3u : static_cast<uint32>(Config.MaxRetryCount);
 
 	const uint64 StartTime = FPlatformTime::Cycles64();
 
@@ -287,7 +288,7 @@ void FMultiEndpointHttpClient::IssueRequest(FRequest&& Request)
 				const bool bError		= TicketStatus.GetId() == FTicketStatus::EId::Error;
 				const bool bServerError = Request.StatusCode > 499 && Request.StatusCode < 600;
 
-				if ((bError || bServerError) && Request.RetryCount < Config.MaxRetryCount)
+				if ((bError || bServerError) && Request.RetryCount < GetRetryLimitForRequest(Request))
 				{
 					RetryRequest(MoveTemp(Request));
 				}
@@ -411,6 +412,11 @@ FMultiEndpointHttpClient::FConnection* FMultiEndpointHttpClient::FindConnection(
 	return nullptr;
 }
 
+uint32 FMultiEndpointHttpClient::GetRetryLimitForRequest(const FRequest& Request) const
+{
+	return Config.MaxRetryCount == -1 ? Request.Connection.HostGroup.Hosts().Num() : Config.MaxRetryCount;
+}
+
 void FMultiEndpointHttpClient::Log(const FRequest& Request) const
 {
 #if !NO_LOGGING
@@ -423,13 +429,13 @@ void FMultiEndpointHttpClient::Log(const FRequest& Request) const
 	const uint64 Size = Request.Body.GetSize() >> 10;
 
 	FMsg::Logf(__FILE__, __LINE__, Config.LogCategory->GetCategoryName(), Config.LogVerbosity,
-		TEXT("http-%3u: %5" UINT64_FMT "ms %5" UINT64_FMT "KiB %s (Attempt %d/%d)"),
+		TEXT("http-%3u: %5" UINT64_FMT "ms %5" UINT64_FMT "KiB %s (Attempt %u/%u)"),
 		Request.StatusCode,
 		DurationMs,
 		Size,
 		WriteToString<512>(Request.Connection.HostGroup.Host(Request.Host), Request.RelativeUrl).ToString(),
 		Request.RetryCount,
-		Config.MaxRetryCount
+		GetRetryLimitForRequest(Request)
 	);
 
 #endif //!NO_LOGGING
diff --git a/Engine/Source/Runtime/Experimental/IoStore/OnDemand/Private/OnDemandHttpClient.h b/Engine/Source/Runtime/Experimental/IoStore/OnDemand/Private/OnDemandHttpClient.h
--- a/Engine/Source/Runtime/Experimental/IoStore/OnDemand/Private/OnDemandHttpClient.h
+++ b/Engine/Source/Runtime/Experimental/IoStore/OnDemand/Private/OnDemandHttpClient.h
@@ -39,7 +39,9 @@ struct FMultiEndpointHttpClientConfig
 {
 	int32						MaxConnectionCount = 4;
 	int32						ReceiveBufferSize = -1;
-	uint32						MaxRetryCount = 0;
+	int32						MaxRetryCount = -1;	// Positive: The number of times to retry a failed request
+													// Zero: Failed requests will not be retried
+													// Negative: A failed request will retry once per provided host url
 	int32						TimeoutMs = 0;
 	EHttpRedirects				Redirects = EHttpRedirects::Follow;
 	bool						bEnableThreadSafetyChecks = false;
@@ -122,6 +124,8 @@ private:
 	FConnection&						GetConnection(const FOnDemandHostGroup& HostGroup);
 	FConnection*						FindConnection(const FOnDemandHostGroup& HostGroup);
 
+	uint32								GetRetryLimitForRequest(const FRequest& Request) const;
+
 	void								Log(const FRequest& Request) const;
 
 	FMultiEndpointHttpClientConfig		Config;
diff --git a/Engine/Source/Runtime/Experimental/IoStore/OnDemand/Private/OnDemandIoDispatcherBackend.cpp b/Engine/Source/Runtime/Experimental/IoStore/OnDemand/Private/OnDemandIoDispatcherBackend.cpp
--- a/Engine/Source/Runtime/Experimental/IoStore/OnDemand/Private/OnDemandIoDispatcherBackend.cpp
+++ b/Engine/Source/Runtime/Experimental/IoStore/OnDemand/Private/OnDemandIoDispatcherBackend.cpp
@@ -108,16 +108,11 @@ static FAutoConsoleVariableRef CVar_IasHttpFailTimeOutMs(
 	TEXT("Fail infinite network waits that take longer than this (in ms, a value of zero will use the default)")
 );
 
-/**
- *This is only applied when the connection was made to a single ServiceUrl rather than a DistributedUrl.
- * In the latter case we will make two attempts on the primary CDN followed by a single attempt for the
- * remaining CDN's to be tried in the order provided by the distributed endpoint.
- */
-static int32 GIasHttpRetryCount = 2;
+static int32 GIasHttpRetryCount = -1;
 static FAutoConsoleVariableRef CVar_IasHttpRetryCount(
 	TEXT("ias.HttpRetryCount"),
 	GIasHttpRetryCount,
-	TEXT("Number of HTTP request retries before failing the request (if connected to a service url rather than distributed endpoints).")
+	TEXT("Number of times that a request should be retried before being considered failed. A negative value will use the default behaviour, which is one retry per host url provided.")
 );
 
 int32 GIasHttpTimeOutMs = 10 * 1000;
@@ -1661,7 +1656,7 @@ uint32 FOnDemandIoBackend::Run()
 		{
 			.MaxConnectionCount = GIasHttpConnectionCount,
 			.ReceiveBufferSize = GIasHttpRecvBufKiB >= 0 ? GIasHttpRecvBufKiB << 10 : -1,
-			.MaxRetryCount = (uint32)FMath::Max(DefaultHostGroup.GetHostUrls().Num(), GIasHttpRetryCount),
+			.MaxRetryCount = GIasHttpRetryCount,
 			.TimeoutMs = GIasHttpFailTimeOutMs,
 			.bEnableThreadSafetyChecks = true,
 			.bAllowChunkedTransfer = GIasHttpAllowChunkedXfer,
diff --git a/Engine/Source/Runtime/Experimental/IoStore/OnDemand/Private/OnDemandToc.cpp b/Engine/Source/Runtime/Experimental/IoStore/OnDemand/Private/OnDemandToc.cpp
--- a/Engine/Source/Runtime/Experimental/IoStore/OnDemand/Private/OnDemandToc.cpp
+++ b/Engine/Source/Runtime/Experimental/IoStore/OnDemand/Private/OnDemandToc.cpp
@@ -657,7 +657,7 @@ TIoStatusOr<FOnDemandToc> FOnDemandToc::LoadFromFile(const FString& FilePath, bo
 	return Toc; 
 }
 
-TIoStatusOr<FOnDemandToc> FOnDemandToc::LoadFromUrl(FAnsiStringView Url, uint32 RetryCount, bool bFollowRedirects)
+TIoStatusOr<FOnDemandToc> FOnDemandToc::LoadFromUrl(FAnsiStringView Url, int32 RetryCount, bool bFollowRedirects)
 {
 	const EHttpRedirects Redirects = bFollowRedirects ? EHttpRedirects::Follow : EHttpRedirects::Disabled;
 	TIoStatusOr<FMultiEndpointHttpClientResponse> Response =  FMultiEndpointHttpClient::Get(Url, FMultiEndpointHttpClientConfig
@@ -685,7 +685,7 @@ TIoStatusOr<FOnDemandToc> FOnDemandToc::LoadFromUrl(FAnsiStringView Url, uint32
 	return Toc; 
 }
 
-TIoStatusOr<FOnDemandToc> FOnDemandToc::LoadFromUrl(FStringView Url, uint32 RetryCount, bool bFollowRedirects)
+TIoStatusOr<FOnDemandToc> FOnDemandToc::LoadFromUrl(FStringView Url, int32 RetryCount, bool bFollowRedirects)
 {
 	auto AnsiUrl = StringCast<ANSICHAR>(Url.GetData(), Url.Len());
 	return LoadFromUrl(AnsiUrl, RetryCount, bFollowRedirects);
diff --git a/Engine/Source/Runtime/Experimental/IoStore/OnDemand/Public/IO/OnDemandToc.h b/Engine/Source/Runtime/Experimental/IoStore/OnDemand/Public/IO/OnDemandToc.h
--- a/Engine/Source/Runtime/Experimental/IoStore/OnDemand/Public/IO/OnDemandToc.h
+++ b/Engine/Source/Runtime/Experimental/IoStore/OnDemand/Public/IO/OnDemandToc.h
@@ -187,8 +187,8 @@ struct FOnDemandToc
 	static FGuid VersionGuid;
 
 	UE_API static TIoStatusOr<FOnDemandToc> LoadFromFile(const FString& FilePath, bool bValidate);
-	UE_API static TIoStatusOr<FOnDemandToc> LoadFromUrl(FAnsiStringView Url, uint32 RetryCount = 0, bool bFollowRedirects = false);
-	UE_API static TIoStatusOr<FOnDemandToc> LoadFromUrl(FStringView Url, uint32 RetryCount = 0, bool bFollowRedirects = false);
+	UE_API static TIoStatusOr<FOnDemandToc> LoadFromUrl(FAnsiStringView Url, int32 RetryCount = 0, bool bFollowRedirects = false);
+	UE_API static TIoStatusOr<FOnDemandToc> LoadFromUrl(FStringView Url, int32 RetryCount = 0, bool bFollowRedirects = false);
 };
 
 UE_API bool LoadFromCompactBinary(FCbFieldView Field, FOnDemandToc& OutToc);
