Add a new method FHostGroupManager::ForEachHostGroup that allows a calling system to iterate over all active host group's in the OnDemand system.

#rnx

diff --git a/Engine/Source/Runtime/Experimental/IoStore/OnDemand/Private/IasHostGroup.cpp b/Engine/Source/Runtime/Experimental/IoStore/OnDemand/Private/IasHostGroup.cpp
--- a/Engine/Source/Runtime/Experimental/IoStore/OnDemand/Private/IasHostGroup.cpp
+++ b/Engine/Source/Runtime/Experimental/IoStore/OnDemand/Private/IasHostGroup.cpp
@@ -346,9 +346,20 @@ FIASHostGroup FHostGroupManager::Find(FName Name)
 	}
 }
 
+void FHostGroupManager::ForEachHostGroup(TFunctionRef<void(const FIASHostGroup&)> Callback) const
+{
+	UE::TUniqueLock _(Mutex);
+
+	for (const FIASHostGroup& Host : HostGroups)
+	{
+		Callback(Host);
+	}
+}
+
 void FHostGroupManager::Tick(uint32 TimeoutMs, std::atomic_bool& CancellationToken)
 {
 	UE::TUniqueLock _(Mutex);
+
 	for (FIASHostGroup& HostGroup : HostGroups)
 	{
 		if (HostGroup.AttemptReconnection(TimeoutMs, CancellationToken) == FIASHostGroup::EReconnectionResult::Reconnected)
diff --git a/Engine/Source/Runtime/Experimental/IoStore/OnDemand/Private/IasHostGroup.h b/Engine/Source/Runtime/Experimental/IoStore/OnDemand/Private/IasHostGroup.h
--- a/Engine/Source/Runtime/Experimental/IoStore/OnDemand/Private/IasHostGroup.h
+++ b/Engine/Source/Runtime/Experimental/IoStore/OnDemand/Private/IasHostGroup.h
@@ -5,6 +5,7 @@
 #include "Async/Mutex.h"
 #include "IO/IoStatus.h"
 #include "OnDemandBackendStatus.h"
+#include "Templates/Function.h"
 #include "Templates/SharedPointer.h"
 
 namespace UE::IoStore
@@ -79,7 +80,9 @@ public:
 	TIoStatusOr<FIASHostGroup> Register(FName Name, FAnsiStringView TestPath); // <- Unresolved
 	TIoStatusOr<FIASHostGroup> Register(FName Name, TConstArrayView<FAnsiString> HostUrls); // <- Resolved
 
-	FIASHostGroup				Find(FName Name);
+	FIASHostGroup Find(FName Name);
+
+	void ForEachHostGroup(TFunctionRef<void (const FIASHostGroup&)> Callback) const;
 
 	void Tick(uint32 TimeoutMs, std::atomic_bool& CancellationToken);
 
