IAS statistics now accesses the current host groups via FHostGroupManager and no longer maintains it's own list.

#rnx

diff --git a/Engine/Source/Runtime/Experimental/IoStore/OnDemand/Private/OnDemandIoDispatcherBackend.cpp b/Engine/Source/Runtime/Experimental/IoStore/OnDemand/Private/OnDemandIoDispatcherBackend.cpp
--- a/Engine/Source/Runtime/Experimental/IoStore/OnDemand/Private/OnDemandIoDispatcherBackend.cpp
+++ b/Engine/Source/Runtime/Experimental/IoStore/OnDemand/Private/OnDemandIoDispatcherBackend.cpp
@@ -1395,8 +1395,6 @@ bool FOnDemandIoBackend::CreateDefaultHostGroup()
 		return false;
 	}
 
-	Stats.AddHostGroup(DefaultHostGroup);
-
 	return true;
 }
 
diff --git a/Engine/Source/Runtime/Experimental/IoStore/OnDemand/Private/Statistics.cpp b/Engine/Source/Runtime/Experimental/IoStore/OnDemand/Private/Statistics.cpp
--- a/Engine/Source/Runtime/Experimental/IoStore/OnDemand/Private/Statistics.cpp
+++ b/Engine/Source/Runtime/Experimental/IoStore/OnDemand/Private/Statistics.cpp
@@ -420,7 +420,7 @@ FOnDemandIoBackendStats::FOnDemandIoBackendStats(FBackendStatus& InStatus)
 				Out.Add(FCoreDelegates::EOnScreenMessageSeverity::Info, FText::FromStringView(BackendStatusText.ToView()));
 			}
 
-			for (const FIASHostGroup& HostGroup : HostGroups)
+			FHostGroupManager::Get().ForEachHostGroup([&Out](const FIASHostGroup& HostGroup)
 			{
 				FCoreDelegates::EOnScreenMessageSeverity Verbosity = FCoreDelegates::EOnScreenMessageSeverity::Info;
 				TStringBuilder<256> Text;
@@ -442,7 +442,7 @@ FOnDemandIoBackendStats::FOnDemandIoBackendStats(FBackendStatus& InStatus)
 				}
 
 				Out.Add(Verbosity, FText::FromStringView(Text.ToView()));
-			}
+			});
 
 			TStringBuilder<256> CachingText;
 			if (BackendStatus.IsCacheEnabled())
@@ -517,12 +517,6 @@ FOnDemandIoBackendStats* FOnDemandIoBackendStats::Get()
 	return GStatistics;
 }
 
-void FOnDemandIoBackendStats::AddHostGroup(const FIASHostGroup& NewHostGroup)
-{
-	// TODO: Maybe check for duplicates?
-	HostGroups.Add(NewHostGroup);
-}
-
 void FOnDemandIoBackendStats::ReportGeneralAnalytics(TArray<FAnalyticsEventAttribute>& OutAnalyticsArray) const
 {
 	// Note that this analytics section is not optional, if we are reporting analytics then we report this section
@@ -592,9 +586,8 @@ class FAnalyticsRecording : public IAnalyticsRecording
 public:
 	FAnalyticsRecording() = delete;
 
-	FAnalyticsRecording(const TArray<FIASHostGroup>& InHostGroups, const FBackendStatus& InBackendStatus)
-		: HostGroups(InHostGroups)
-		, BackendStatus(InBackendStatus)
+	FAnalyticsRecording(const FBackendStatus& InBackendStatus)
+		: BackendStatus(InBackendStatus)
 	{
 
 	}
@@ -748,14 +741,13 @@ private:
 		FTrackedValue<FCounterAtomicInt, GCacheRejectBytes>			RejectBytes;
 	} Cache;
 
-	const TArray<FIASHostGroup>& HostGroups;
 	const FBackendStatus& BackendStatus;
 	bool bRecording = true;
 };
 
 TUniquePtr<IAnalyticsRecording> FOnDemandIoBackendStats::StartAnalyticsRecording() const
 {
-	TUniquePtr<FAnalyticsRecording> Recording = MakeUnique<FAnalyticsRecording>(HostGroups, BackendStatus);
+	TUniquePtr<FAnalyticsRecording> Recording = MakeUnique<FAnalyticsRecording>(BackendStatus);
 	return Recording;
 }
 
diff --git a/Engine/Source/Runtime/Experimental/IoStore/OnDemand/Private/Statistics.h b/Engine/Source/Runtime/Experimental/IoStore/OnDemand/Private/Statistics.h
--- a/Engine/Source/Runtime/Experimental/IoStore/OnDemand/Private/Statistics.h
+++ b/Engine/Source/Runtime/Experimental/IoStore/OnDemand/Private/Statistics.h
@@ -20,7 +20,6 @@ LLM_DECLARE_TAG(Ias);
 namespace UE::IoStore
 {
 
-class FIASHostGroup;
 class IAnalyticsRecording;
 struct FAvailableEndPoints;
 
@@ -38,8 +37,6 @@ public:
 	FOnDemandIoBackendStats(FBackendStatus& InBackendStatus) IAS_STATISTICS_IMPL()
 	~FOnDemandIoBackendStats() IAS_STATISTICS_IMPL()
 
-	void AddHostGroup(const FIASHostGroup& HostGroup)  IAS_STATISTICS_IMPL()
-
 	/** Report analytics not directly associated with a specific endpoint */
 	void ReportGeneralAnalytics(TArray<FAnalyticsEventAttribute>& OutAnalyticsArray) const IAS_STATISTICS_IMPL()
 	/** Report analytics for the current endpoint */
@@ -80,7 +77,6 @@ private:
 
 #if IAS_WITH_STATISTICS
 	const FBackendStatus& BackendStatus;
-	TArray<FIASHostGroup> HostGroups;
 #endif //IAS_WITH_STATISTICS
 };
 
