Dataflow - Fixed crash when remapping managed arrays of sets containing INDEX_NONE indices.

#rnx

diff --git a/Engine/Source/Runtime/Experimental/Chaos/Public/GeometryCollection/ManagedArray.h b/Engine/Source/Runtime/Experimental/Chaos/Public/GeometryCollection/ManagedArray.h
--- a/Engine/Source/Runtime/Experimental/Chaos/Public/GeometryCollection/ManagedArray.h
+++ b/Engine/Source/Runtime/Experimental/Chaos/Public/GeometryCollection/ManagedArray.h
@@ -1301,7 +1301,7 @@ public:
 
 			for (int32 StaleEntry : OldSet)
 			{
-				const int32 NewEntry = StaleEntry - Offsets[StaleEntry];
+				const int32 NewEntry = StaleEntry >= 0 ? StaleEntry - Offsets[StaleEntry] : StaleEntry;	//only remap if valid
 				NewSet.Add(NewEntry);
 			}
 		}
