Convert IAS test to automation framework test

#rnx

diff --git a/Engine/Source/Runtime/Experimental/IoStore/OnDemand/Private/OnDemandIoDispatcherBackend.cpp b/Engine/Source/Runtime/Experimental/IoStore/OnDemand/Private/OnDemandIoDispatcherBackend.cpp
--- a/Engine/Source/Runtime/Experimental/IoStore/OnDemand/Private/OnDemandIoDispatcherBackend.cpp
+++ b/Engine/Source/Runtime/Experimental/IoStore/OnDemand/Private/OnDemandIoDispatcherBackend.cpp
@@ -41,6 +41,12 @@
 
 #include <atomic>
 
+#if defined(WITH_IOSTORE_ONDEMAND_TESTS)
+#include "TestHarness.h"
+#include "TestMacros/Assertions.h"
+#include <catch2/generators/catch_generators.hpp>
+#endif
+
 /** When enabled the IAS system can add additional debug console commands for development use */
 #define UE_IAS_DEBUG_CONSOLE_CMDS (1 && !NO_CVARS && !UE_BUILD_SHIPPING)
 
@@ -1763,18 +1769,12 @@ TSharedPtr<IOnDemandIoDispatcherBackend> MakeOnDemandIoDispatcherBackend(
 		new FOnDemandIoBackend(Config, IoStore, MoveTemp(Cache)));
 }
 
-#if UE_ENABLE_IAS_TESTING
-
-namespace Tests
-{
+///////////////////////////////////////////////////////////////////////////////
+#if defined(WITH_IOSTORE_ONDEMAND_TESTS)
 
-/***
- * We should port these tests to the automation framework or lowlevel tests once we decide which way we are going with testing
- * for this module.
- */
-struct FTests
+TEST_CASE("IoStore::OnDemand::Ias::Misc", "[IoStoreOnDemand][Ias]")
 {
-	void ChunkRequestLinkedListTest()
+	SECTION("AddRemoveDispatcherRequest")
 	{
 		// We are not allowed to create FIoRequestImpl directly due to it's api, however we don't need the full functionality
 		// in order to test FChunkRequest so we can fake it with a malloced version of the struct, mem set to zero.
@@ -1787,45 +1787,38 @@ struct FTests
 		ChunkRequest.AddDispatcherRequest(IoSecondRequest);
 		ChunkRequest.AddDispatcherRequest(IoThirdRequest);
 
-		check(ChunkRequest.RequestHead == IoFirstRequest);
-		check(ChunkRequest.RequestHead->NextRequest == IoSecondRequest);
-		check(ChunkRequest.RequestHead->NextRequest->NextRequest == IoThirdRequest);
-		check(ChunkRequest.RequestTail == IoThirdRequest);
+		CHECK(ChunkRequest.RequestHead == IoFirstRequest);
+		CHECK(ChunkRequest.RequestHead->NextRequest == IoSecondRequest);
+		CHECK(ChunkRequest.RequestHead->NextRequest->NextRequest == IoThirdRequest);
+		CHECK(ChunkRequest.RequestTail == IoThirdRequest);
 
-		check(ChunkRequest.RequestCount == 3);
+		CHECK(ChunkRequest.RequestCount == 3);
 
 		ChunkRequest.RemoveDispatcherRequest(IoThirdRequest);
 		
-		check(ChunkRequest.RequestCount == 2);
-		check(ChunkRequest.RequestHead == IoFirstRequest);
-		check(ChunkRequest.RequestTail == IoSecondRequest);
+		CHECK(ChunkRequest.RequestCount == 2);
+		CHECK(ChunkRequest.RequestHead == IoFirstRequest);
+		CHECK(ChunkRequest.RequestTail == IoSecondRequest);
 
 		ChunkRequest.RemoveDispatcherRequest(IoSecondRequest);
 		
-		check(ChunkRequest.RequestCount == 1);
-		check(ChunkRequest.RequestHead == IoFirstRequest);
-		check(ChunkRequest.RequestTail == IoFirstRequest);
+		CHECK(ChunkRequest.RequestCount == 1);
+		CHECK(ChunkRequest.RequestHead == IoFirstRequest);
+		CHECK(ChunkRequest.RequestTail == IoFirstRequest);
 
 		ChunkRequest.RemoveDispatcherRequest(IoFirstRequest);
-		check(ChunkRequest.RequestCount == 0);
+		CHECK(ChunkRequest.RequestCount == 0);
 
-		check(ChunkRequest.RequestHead == nullptr);
-		check(ChunkRequest.RequestTail == nullptr);
+		CHECK(ChunkRequest.RequestHead == nullptr);
+		CHECK(ChunkRequest.RequestTail == nullptr);
 
 		FMemory::Free(IoFirstRequest);
 		FMemory::Free(IoSecondRequest);
 		FMemory::Free(IoThirdRequest);
 	}
+}
 
-	FTests()
-	{
-		ChunkRequestLinkedListTest();
-	}
-} RunTests;
-
-} // namespace Tests
-
-#endif // UE_ENABLE_IAS_TESTING
+#endif // WITH_IOSTORE_ONDEMAND_TESTS
 
 } // namespace UE::IoStore
 
