[AutoRTFM] Use a transactionally-safe critical section inside the tracked geometry manager.

This is reachable from Verse script when starting movie playback.

diff --git a/Engine/Source/Runtime/Experimental/Chaos/Private/Chaos/TrackedGeometryManager.cpp b/Engine/Source/Runtime/Experimental/Chaos/Private/Chaos/TrackedGeometryManager.cpp
--- a/Engine/Source/Runtime/Experimental/Chaos/Private/Chaos/TrackedGeometryManager.cpp
+++ b/Engine/Source/Runtime/Experimental/Chaos/Private/Chaos/TrackedGeometryManager.cpp
@@ -55,13 +55,13 @@ namespace Chaos
 
 	void FTrackedGeometryManager::AddGeometry(TSerializablePtr<FImplicitObject> Geometry, const FString& DebugInfo)
 	{
-		FScopeLock Lock(&CriticalSection);
+		UE::TScopeLock Lock(CriticalSection);
 		SharedGeometry.Add(Geometry, DebugInfo);
 	}
 
 	void FTrackedGeometryManager::RemoveGeometry(const FImplicitObject* Geometry)
 	{
-		FScopeLock Lock(&CriticalSection);
+		UE::TScopeLock Lock(CriticalSection);
 		TSerializablePtr<FImplicitObject> Dummy;
 		Dummy.SetFromRawLowLevel(Geometry);
 		SharedGeometry.Remove(Dummy);
diff --git a/Engine/Source/Runtime/Experimental/Chaos/Public/Chaos/TrackedGeometryManager.h b/Engine/Source/Runtime/Experimental/Chaos/Public/Chaos/TrackedGeometryManager.h
--- a/Engine/Source/Runtime/Experimental/Chaos/Public/Chaos/TrackedGeometryManager.h
+++ b/Engine/Source/Runtime/Experimental/Chaos/Public/Chaos/TrackedGeometryManager.h
@@ -3,6 +3,7 @@
 
 #include "Chaos/Serializable.h"
 #include "Chaos/ImplicitObject.h"
+#include "Misc/TransactionallySafeCriticalSection.h"
 #include "Serialization/MemoryWriter.h"
 
 class FOutputDevice;
@@ -18,7 +19,7 @@ namespace Chaos
 		
 	private:
 		TMap<TSerializablePtr<FImplicitObject>, FString> SharedGeometry;
-		FCriticalSection CriticalSection;
+		FTransactionallySafeCriticalSection CriticalSection;
 		
 		friend FImplicitObject;
 		
