I/O store on-demand install cache snapshot support.

- a snapshot stores the current state of the cache as a single unit
- when the cache journal file gets too big a snapshot is created when the game shuts down
- when the game starts the cache state is re-created from the snapshot and the journal

diff --git a/Engine/Source/Runtime/Experimental/IoStore/OnDemand/Private/GenericHash.h b/Engine/Source/Runtime/Experimental/IoStore/OnDemand/Private/GenericHash.h
--- a/Engine/Source/Runtime/Experimental/IoStore/OnDemand/Private/GenericHash.h
+++ b/Engine/Source/Runtime/Experimental/IoStore/OnDemand/Private/GenericHash.h
@@ -8,6 +8,7 @@
 #include "HAL/UnrealMemory.h"
 #include "Math/UnrealMathUtility.h"
 #include "Memory/MemoryView.h"
+#include "Serialization/Archive.h"
 
 struct FIoHash;
 
@@ -68,6 +69,12 @@ struct THash
 		return GenericHash::ToHex(H.GetView(), Out);
 	}
 
+	friend FArchive& operator<<(FArchive& Ar, THash& H)
+	{
+		Ar.Serialize(const_cast<uint8*>(H.GetData()), H.GetSize());
+		return Ar;
+	}
+
 private:
 	FMutableMemoryView GetMutableView() { return FMutableMemoryView(Hash, Size); }
 
diff --git a/Engine/Source/Runtime/Experimental/IoStore/OnDemand/Private/OnDemandConfig.cpp b/Engine/Source/Runtime/Experimental/IoStore/OnDemand/Private/OnDemandConfig.cpp
--- a/Engine/Source/Runtime/Experimental/IoStore/OnDemand/Private/OnDemandConfig.cpp
+++ b/Engine/Source/Runtime/Experimental/IoStore/OnDemand/Private/OnDemandConfig.cpp
@@ -404,6 +404,24 @@ TIoStatusOr<FOnDemandInstallCacheConfig> TryParseInstallCacheConfig(const TCHAR*
 		OutConfig.DiskQuota = DiskQuota;
 	}
 
+	const int64 JournalMaxSize = [CommandLine]
+	{
+		if (FString ParamValue; FParse::Value(CommandLine, TEXT("-Iad.FileCache.JournalMaxSize="), ParamValue))
+		{
+			return Config::ParseSizeParam(ParamValue);
+		}
+		else if (FString ValueStr; GConfig->GetString(TEXT("OnDemandInstall"), TEXT("FileCache.JournalMaxSize"), ValueStr, GEngineIni))
+		{
+			return Config::ParseSizeParam(ValueStr);
+		}
+		return int64(-1);
+	}();
+
+	if (JournalMaxSize > 0)
+	{
+		OutConfig.JournalMaxSize = JournalMaxSize;
+	}
+
 	OutConfig.RootDirectory = Config::GetInstallCacheDirectory(CommandLine);
 #if !UE_BUILD_SHIPPING
 	OutConfig.bDropCache = FParse::Param(CommandLine, TEXT("Iad.DropCache"));
diff --git a/Engine/Source/Runtime/Experimental/IoStore/OnDemand/Private/OnDemandInstallCache.cpp b/Engine/Source/Runtime/Experimental/IoStore/OnDemand/Private/OnDemandInstallCache.cpp
--- a/Engine/Source/Runtime/Experimental/IoStore/OnDemand/Private/OnDemandInstallCache.cpp
+++ b/Engine/Source/Runtime/Experimental/IoStore/OnDemand/Private/OnDemandInstallCache.cpp
@@ -17,6 +17,7 @@
 #include "IO/IoChunkId.h"
 #include "IO/IoChunkEncoding.h"
 #include "Misc/DateTime.h"
+#include "Misc/Paths.h"
 #include "Misc/PathViews.h"
 #include "Misc/ScopeExit.h"
 #include "Misc/StringBuilder.h"
@@ -26,6 +27,7 @@
 #include "ProfilingDebugging/IoStoreTrace.h"
 
 #if defined(WITH_IOSTORE_ONDEMAND_TESTS)
+#include "Algo/Find.h"
 #include "TestHarness.h"
 #include "TestMacros/Assertions.h"
 #include <catch2/generators/catch_generators.hpp>
@@ -85,6 +87,12 @@ struct FCasBlockId
 		return GetTypeHash(BlockId.Id);
 	}
 
+	friend FArchive& operator<<(FArchive& Ar, FCasBlockId& BlockId)
+	{
+		Ar << BlockId.Id;
+		return Ar;
+	}
+
 	static const FCasBlockId Invalid;
 
 	uint32 Id = 0;
@@ -97,6 +105,23 @@ struct FCasLocation
 {
 	bool IsValid() const { return BlockId.IsValid() && BlockOffset != MAX_uint32; }
 
+	friend inline bool operator==(FCasLocation LHS, FCasLocation RHS)
+	{
+		return LHS.BlockId == RHS.BlockId && LHS.BlockOffset == RHS.BlockOffset;
+	}
+
+	friend inline uint32 GetTypeHash(FCasLocation Loc)
+	{
+		return HashCombine(GetTypeHash(Loc.BlockId), GetTypeHash(Loc.BlockOffset));
+	}
+
+	friend FArchive& operator<<(FArchive& Ar, FCasLocation& Loc)
+	{
+		Ar << Loc.BlockId; 
+		Ar << Loc.BlockOffset;
+		return Ar;
+	}
+
 	static const FCasLocation Invalid;
 
 	FCasBlockId	BlockId;
@@ -836,6 +861,262 @@ void FCasJournal::FTransaction::BlockAccess(FCasBlockId BlockId, int64 UtcTicks)
 	};
 }
 
+///////////////////////////////////////////////////////////////////////////////
+struct FCasSnapshot
+{
+	enum class EVersion : uint32
+	{
+		Invalid = 0,
+		Initial,
+
+		LatestPlusOne,
+		Latest = LatestPlusOne - 1
+	};
+
+	struct FHeader
+	{
+		static const inline uint8 MagicSequence[16] = {'+', 'S', 'N', 'A', 'P', 'S', 'H', 'O', 'T', 'H', 'E', 'A', 'D', 'E', 'R', '+'};
+
+		bool		IsValid() const;
+
+		uint8		Magic[16] = {0};
+		EVersion	Version = EVersion::Invalid;
+		uint8		Pad[12] = {0};
+	};
+	static_assert(sizeof(FHeader) == 32);
+
+	struct FFooter
+	{
+		static const inline uint8 MagicSequence[16] = {'+', 'S', 'N', 'A', 'P', 'S', 'H', 'O', 'T', 'F', 'O', 'O', 'T', 'E', 'R', '+'};
+
+		bool IsValid() const;
+
+		uint8 Magic[16] = {0};
+	};
+	static_assert(sizeof(FFooter) == 16);
+
+	struct FBlock
+	{
+		friend FArchive& operator<<(FArchive& Ar, FBlock& Block)
+		{
+			Ar << Block.BlockId;
+			Ar << Block.LastAccess;
+			return Ar;
+		}
+
+		FCasBlockId BlockId;
+		int64		LastAccess;
+	};
+
+	using FChunkLocation = TPair<FCasAddr, FCasLocation>;
+
+	static TIoStatusOr<FCasSnapshot>	FromJournal(const FString& JournalFile);
+	static TIoStatusOr<FCasSnapshot>	Load(const FString& SnapShotFile, int64* OutFileSize = nullptr);
+	static TIoStatusOr<int64>			Save(const FCasSnapshot& Snapshot, const FString& SnapshotFile);
+
+	TArray<FBlock>						Blocks;
+	TArray<FChunkLocation>				ChunkLocations;
+	FCasBlockId							CurrentBlockId;
+};
+
+///////////////////////////////////////////////////////////////////////////////
+bool FCasSnapshot::FHeader::IsValid() const
+{
+	if (FMemory::Memcmp(&Magic, &FHeader::MagicSequence, sizeof(FHeader::MagicSequence)) != 0)
+	{
+		return false;
+	}
+
+	if (static_cast<uint32>(Version) > static_cast<uint32>(EVersion::Latest))
+	{
+		return false;
+	}
+
+	return true;
+}
+
+bool FCasSnapshot::FFooter::IsValid() const
+{
+	return FMemory::Memcmp(Magic, FFooter::MagicSequence, sizeof(FFooter::MagicSequence)) == 0;
+}
+
+TIoStatusOr<FCasSnapshot> FCasSnapshot::FromJournal(const FString& JournalFile)
+{
+	TMap<FCasAddr, FCasLocation>	CasLookup;
+	TMap<FCasBlockId, int64>		LastAccess;
+	TSet<FCasBlockId>				BlockIds;
+	FCasBlockId						CurrentBlockId;
+
+	FIoStatus ReplayStatus = FCasJournal::Replay(
+		JournalFile,
+		[&CasLookup, &LastAccess, &BlockIds, &CurrentBlockId](const FCasJournal::FEntry& JournalEntry)
+		{
+			switch(JournalEntry.Type())
+			{
+			case FCasJournal::FEntry::EType::ChunkLocation:
+			{
+				const FCasJournal::FEntry::FChunkLocation& ChunkLocation = JournalEntry.ChunkLocation;
+				if (ChunkLocation.CasLocation.IsValid())
+				{
+					CasLookup.FindOrAdd(ChunkLocation.CasAddr, ChunkLocation.CasLocation);
+				}
+				else
+				{
+					CasLookup.Remove(ChunkLocation.CasAddr);
+				}
+				break;
+			}
+			case FCasJournal::FEntry::EType::BlockCreated:
+			{
+				const FCasJournal::FEntry::FBlockOperation& Op = JournalEntry.BlockOperation;
+				CurrentBlockId = Op.BlockId;
+				BlockIds.Add(Op.BlockId);
+				break;
+			}
+			case FCasJournal::FEntry::EType::BlockDeleted:
+			{
+				const FCasJournal::FEntry::FBlockOperation& Op = JournalEntry.BlockOperation;
+				BlockIds.Remove(Op.BlockId);
+				if (CurrentBlockId == Op.BlockId)
+				{
+					CurrentBlockId = FCasBlockId::Invalid; 
+				}
+				break;
+			}
+			case FCasJournal::FEntry::EType::BlockAccess:
+			{
+				const FCasJournal::FEntry::FBlockOperation& Op = JournalEntry.BlockOperation;
+				LastAccess.Add(Op.BlockId, Op.UtcTicks);
+				break;
+			}
+			};
+		});
+
+	FCasSnapshot Snapshot;
+	Snapshot.Blocks.Reserve(BlockIds.Num());
+	for (FCasBlockId BlockId : BlockIds)
+	{
+		const int64* AccessTime = LastAccess.Find(BlockId);
+		Snapshot.Blocks.Add(FBlock
+		{
+			.BlockId	= BlockId,
+			.LastAccess = AccessTime != nullptr ? *AccessTime : FDateTime::UtcNow().GetTicks()
+		});
+	}
+
+	Snapshot.ChunkLocations = CasLookup.Array();
+	Snapshot.CurrentBlockId = CurrentBlockId;
+
+	return Snapshot;
+}
+
+TIoStatusOr<int64> FCasSnapshot::Save(const FCasSnapshot& Snapshot, const FString& SnapshotFile)
+{
+	IFileManager& Ifm = IFileManager::Get();
+
+	const FString TmpSnapshotFile = FPaths::ChangeExtension(SnapshotFile, TEXT(".snptmp"));
+
+	TUniquePtr<FArchive> Ar(Ifm.CreateFileWriter(*TmpSnapshotFile));
+	if (Ar.IsValid() == false)
+	{
+		return FIoStatus(EIoErrorCode::FileNotOpen);
+	}
+
+	FHeader Header;
+	FMemory::Memcpy(&Header.Magic, &FHeader::MagicSequence, sizeof(FHeader::MagicSequence));
+	Header.Version = EVersion::Latest;
+
+	Ar->Serialize(reinterpret_cast<uint8*>(&Header), sizeof(FHeader));
+	if (Ar->IsError())
+	{
+		Ar.Reset();
+		Ifm.Delete(*TmpSnapshotFile);
+		return FIoStatus(EIoErrorCode::WriteError);
+	}
+
+	FCasSnapshot& NonConst = *const_cast<FCasSnapshot*>(&Snapshot);
+	*Ar << NonConst.Blocks;
+	*Ar << NonConst.ChunkLocations;
+	*Ar << NonConst.CurrentBlockId;
+
+	if (Ar->IsError())
+	{
+		Ar.Reset();
+		Ifm.Delete(*TmpSnapshotFile);
+		return FIoStatus(EIoErrorCode::WriteError);
+	}
+
+	FFooter Footer;
+	FMemory::Memcpy(&Footer.Magic, &FFooter::MagicSequence, sizeof(FFooter::MagicSequence));
+	Ar->Serialize(reinterpret_cast<uint8*>(&Footer), sizeof(FFooter));
+	if (Ar->IsError())
+	{
+		Ar.Reset();
+		Ifm.Delete(*TmpSnapshotFile);
+		return FIoStatus(EIoErrorCode::WriteError);
+	}
+
+	const int64 FileSize = Ar->TotalSize();
+	if (Ar->Close() == false)
+	{
+		Ar.Reset();
+		Ifm.Delete(*TmpSnapshotFile);
+		return FIoStatus(EIoErrorCode::WriteError);
+	}
+
+	if (Ifm.Move(*SnapshotFile, *TmpSnapshotFile) == false)
+	{
+		Ifm.Delete(*TmpSnapshotFile);
+		return FIoStatus(EIoErrorCode::WriteError);
+	}
+
+	return FileSize; 
+}
+
+TIoStatusOr<FCasSnapshot> FCasSnapshot::Load(const FString& SnapshotFile, int64* OutFileSize)
+{
+	IFileManager& Ifm = IFileManager::Get();
+
+	TUniquePtr<FArchive> Ar(Ifm.CreateFileReader(*SnapshotFile));
+	if (Ar.IsValid() == false)
+	{
+		return FIoStatus(EIoErrorCode::NotFound);
+	}
+
+	FHeader Header;
+	Ar->Serialize(reinterpret_cast<uint8*>(&Header), sizeof(FHeader));
+	if (Ar->IsError() || Header.IsValid() == false)
+	{
+		FIoStatus Status = FIoStatusBuilder(EIoErrorCode::ReadError)
+			<< TEXT("Failed to validate snapshot header '")
+			<< SnapshotFile 
+			<< TEXT("'");
+		return Status;
+	}
+
+	FCasSnapshot Snapshot;
+	*Ar << Snapshot.Blocks;
+	*Ar << Snapshot.ChunkLocations;
+	*Ar << Snapshot.CurrentBlockId;
+
+	FFooter Footer;
+	Ar->Serialize(reinterpret_cast<uint8*>(&Footer), sizeof(FFooter));
+	if (Ar->IsError() || Footer.IsValid() == false)
+	{
+		FIoStatus Status = FIoStatusBuilder(EIoErrorCode::ReadError)
+			<< TEXT("Failed to validate snapshot footer '")
+			<< SnapshotFile 
+			<< TEXT("'");
+		return Status;
+	}
+
+	if (OutFileSize != nullptr)
+	{
+		*OutFileSize = Ar->Tell();
+	}
+	return Snapshot;
+}
+
 ///////////////////////////////////////////////////////////////////////////////
 class FOnDemandInstallCache final
 	: public IOnDemandInstallCache 
@@ -976,6 +1257,7 @@ private:
 	FIoStatus									FlushPendingChunks(FPendingChunks& Chunks);
 	FIoStatus									FlushPendingChunksImpl(FPendingChunks& Chunks);
 	FString										GetJournalFilename() const { return CacheDirectory / TEXT("cas.jrn"); }
+	FString										GetSnapshotFilename() const { return CacheDirectory / TEXT("cas.snp"); }
 
 	FOnDemandIoStore&			IoStore;
 	FString						CacheDirectory;
@@ -986,6 +1268,7 @@ private:
 	FIoRequestList				CompletedRequests;
 	UE::FMutex					Mutex;
 	uint64						MaxCacheSize{ 0 };
+	uint64						MaxJournalSize;
 
 #if UE_ONDEMANDINSTALLCACHE_EXCLUSIVE_WRITE
 	UE::Tasks::FPipe			ExclusivePipe{ UE_SOURCE_LOCATION };
@@ -997,10 +1280,11 @@ FOnDemandInstallCache::FOnDemandInstallCache(const FOnDemandInstallCacheConfig&
 	: IoStore(InIoStore)
 	, CacheDirectory(Config.RootDirectory)
 	, MaxCacheSize(Config.DiskQuota)
+	, MaxJournalSize(Config.JournalMaxSize)
 {
 
-	UE_LOG(LogIoStoreOnDemand, Log, TEXT("Initializing install cache, MaxCacheSize=%.2lf MiB"),
-		ToMiB(MaxCacheSize));
+	UE_LOG(LogIoStoreOnDemand, Log, TEXT("Initializing install cache, MaxCacheSize=%.2lf MiB, MaxJournalSize=%.2lf KiB"),
+		ToMiB(MaxCacheSize), ToKiB(MaxJournalSize));
 
 	FIoStatus Status = Cas.Initialize(CacheDirectory);
 	if (Status.IsOk() == false)
@@ -1009,8 +1293,32 @@ FOnDemandInstallCache::FOnDemandInstallCache(const FOnDemandInstallCacheConfig&
 		return;
 	}
 
-	// Replay the journal to get the current state 
-	// TODO: Purge the journal or create snapshot when the journal gets too big
+	{
+		const FString SnapshotFile = GetSnapshotFilename();
+		int64 SnapshotSize = -1;
+		TIoStatusOr<FCasSnapshot> SnapshotStatus = FCasSnapshot::Load(SnapshotFile, &SnapshotSize);
+		if (SnapshotStatus.IsOk())
+		{
+			FCasSnapshot Snapshot = SnapshotStatus.ConsumeValueOrDie();
+			UE_LOG(LogIoStoreOnDemand, Log, TEXT("Loaded CAS snapshot '%s' %.2lf KiB with %d blocks and %d chunk locations"),
+					*SnapshotFile, ToKiB(SnapshotSize), Snapshot.Blocks.Num(), Snapshot.ChunkLocations.Num());
+
+			Cas.Lookup.Reserve(Snapshot.ChunkLocations.Num());
+			for (TPair<FCasAddr, FCasLocation>& Kv : Snapshot.ChunkLocations)
+			{
+				Cas.Lookup.Add(MoveTemp(Kv));
+			}
+
+			Cas.BlockIds.Reserve(Snapshot.Blocks.Num());
+			Cas.LastAccess.Reserve(Snapshot.Blocks.Num());
+			for (const FCasSnapshot::FBlock& Block : Snapshot.Blocks)
+			{
+				Cas.BlockIds.Add(Block.BlockId);
+				Cas.LastAccess.Add(Block.BlockId, Block.LastAccess);
+			}
+		}
+	}
+
 	const FString JournalFile = GetJournalFilename();
 	Status = FCasJournal::Replay(JournalFile, [this](const FCasJournal::FEntry& JournalEntry)
 	{
@@ -1114,12 +1422,51 @@ void FOnDemandInstallCache::Shutdown()
 		LastAccess = MoveTemp(Cas.LastAccess);
 	}
 
-	FCasJournal::FTransaction Transaction = FCasJournal::Begin(GetJournalFilename());
+	const FString JournalFile = GetJournalFilename();
+	FCasJournal::FTransaction Transaction = FCasJournal::Begin(JournalFile);
 	for (const TPair<FCasBlockId, int64>& Kv : LastAccess)
 	{
 		Transaction.BlockAccess(Kv.Key, Kv.Value);
 	}
-	FCasJournal::Commit(MoveTemp(Transaction));
+
+	if (FIoStatus Status = FCasJournal::Commit(MoveTemp(Transaction)); !Status.IsOk())
+	{
+		UE_LOG(LogIoStoreOnDemand, Error, TEXT("Failed to commit journal entries to '%s'"), *JournalFile);
+		//TODO: Delete journal and start over? 
+	}
+
+	IPlatformFile& Ipf = FPlatformFileManager::Get().GetPlatformFile();
+	if (int64 JournalSize = Ipf.FileSize(*JournalFile); JournalSize > int64(MaxJournalSize))
+	{
+		TIoStatusOr<FCasSnapshot> SnapshotStatus = FCasSnapshot::FromJournal(JournalFile);
+		if (SnapshotStatus.IsOk() == false)
+		{
+			UE_LOG(LogIoStoreOnDemand, Error, TEXT("Failed to create CAS snapshot from journal '%s', reason '%s'"),
+				*JournalFile, *SnapshotStatus.Status().ToString());
+			return;
+		}
+
+		FCasSnapshot Snapshot = SnapshotStatus.ConsumeValueOrDie();
+		const FString SnapshotFile = GetSnapshotFilename();
+		TIoStatusOr<int64> SaveStatus = FCasSnapshot::Save(Snapshot, SnapshotFile);
+		if (SaveStatus.IsOk() == false)
+		{
+			UE_LOG(LogIoStoreOnDemand, Error, TEXT("Failed to save CAS snapshot '%s', reason '%s'"),
+				*SnapshotFile, *SaveStatus.Status().ToString());
+			return;
+		}
+
+		if (Ipf.DeleteFile(*JournalFile) == false)
+		{
+			UE_LOG(LogIoStoreOnDemand, Error, TEXT("Failed to delete CAS journal '%s'"), *JournalFile);
+			if (Ipf.DeleteFile(*SnapshotFile) == false)
+			{
+				UE_LOG(LogIoStoreOnDemand, Error, TEXT("Failed to delete CAS snapshot '%s'"), *JournalFile);
+			}
+			return;
+		}
+		UE_LOG(LogIoStoreOnDemand, Log, TEXT("Saved CAS snapshot '%s' %.2lf KiB"), *SnapshotFile, ToKiB(JournalSize));
+	}
 }
 
 void FOnDemandInstallCache::ResolveIoRequests(FIoRequestList Requests, FIoRequestList& OutUnresolved)
@@ -1791,6 +2138,11 @@ private:
 	FString Dir;
 };
 
+FCasAddr CreateCasTestAddr(uint64 Value)
+{
+	return FCasAddr::From(reinterpret_cast<const uint8*>(&Value), sizeof(uint64));
+}
+
 TEST_CASE("IoStore::OnDemand::InstallCache::Journal", "[IoStoreOnDemand][InstallCache]")
 {
 	const FString TestBaseDir = TEXT("TestTmpDir");
@@ -1980,6 +2332,129 @@ TEST_CASE("IoStore::OnDemand::InstallCache::Journal", "[IoStoreOnDemand][Install
 	}
 }
 
+TEST_CASE("IoStore::OnDemand::InstallCache::Snapshot", "[IoStoreOnDemand][InstallCache]")
+{
+	const FString TestBaseDir = TEXT("TestTmpDir");
+
+	SECTION("SaveLoadRoundtrip")
+	{
+		// Arrange
+		FCasSnapshot ExpectedSnapshot;
+
+		for (uint32 Id = 1; Id <= 10; ++Id)
+		{
+			ExpectedSnapshot.Blocks.Add(FCasSnapshot::FBlock
+			{
+				.BlockId	= FCasBlockId(Id),
+				.LastAccess = FDateTime::UtcNow().GetTicks()
+			});
+
+			for (uint32 Idx = 1; Idx <= 10; ++Idx)
+			{
+				FCasAddr CasAddr = CreateCasTestAddr(Idx);
+				FCasLocation Loc = FCasLocation
+				{
+					.BlockId		= FCasBlockId(Id),
+					.BlockOffset	= Idx * 256
+				};
+				ExpectedSnapshot.ChunkLocations.Emplace(CasAddr, Loc);
+			}
+		}
+		ExpectedSnapshot.CurrentBlockId = FCasBlockId(1);
+
+		// Act
+		FTmpDirectoryScope _(TestBaseDir);
+		const FString SnapshotFile = TestBaseDir / TEXT("test.snp");
+		TIoStatusOr<int64> Status = FCasSnapshot::Save(ExpectedSnapshot, SnapshotFile);
+		CHECK(Status.IsOk());
+		const FCasSnapshot Snapshot = FCasSnapshot::Load(SnapshotFile).ConsumeValueOrDie();
+
+		// Assert
+		CHECK(Snapshot.Blocks.Num() == ExpectedSnapshot.Blocks.Num());
+		for (int32 Idx = 0; Idx < Snapshot.Blocks.Num(); ++Idx)
+		{
+			CHECK(Snapshot.Blocks[Idx].BlockId == ExpectedSnapshot.Blocks[Idx].BlockId);
+			CHECK(Snapshot.Blocks[Idx].LastAccess == ExpectedSnapshot.Blocks[Idx].LastAccess);
+		}
+		CHECK(Snapshot.ChunkLocations.Num() == ExpectedSnapshot.ChunkLocations.Num());
+		for (int32 Idx = 0; Idx < Snapshot.ChunkLocations.Num(); ++Idx)
+		{
+			CHECK(Snapshot.ChunkLocations[Idx].Get<0>() == ExpectedSnapshot.ChunkLocations[Idx].Get<0>());
+			CHECK(Snapshot.ChunkLocations[Idx].Get<1>() == ExpectedSnapshot.ChunkLocations[Idx].Get<1>());
+		}
+		CHECK(Snapshot.CurrentBlockId == ExpectedSnapshot.CurrentBlockId);
+	}
+
+	SECTION("CreateFromJournal")
+	{
+		// Arrange
+		FTmpDirectoryScope _(TestBaseDir);
+		const FString JournalFile = TestBaseDir / TEXT("test.jrn");
+		const FCasBlockId ExpectedCurrentBlockId(2);
+
+		FIoStatus Status = FCasJournal::Create(JournalFile);
+		CHECK(Status.IsOk());
+
+		FCasJournal::FTransaction Tx = FCasJournal::Begin(JournalFile);
+
+		// Add a block and some chunk locations
+		Tx.BlockCreated(FCasBlockId(1));
+		for (int32 Idx = 1; Idx <= 10; ++Idx)
+		{
+			Tx.ChunkLocation(FCasLocation
+			{
+				.BlockId		= FCasBlockId(1),
+				.BlockOffset	= 256
+			},
+			CreateCasTestAddr(uint64(Idx) << 32 | 1ull));
+		}
+
+		// Remove the block and the corresponding chunk locations
+		for (int32 Idx = 1; Idx <= 10; ++Idx)
+		{
+			Tx.ChunkLocation(FCasLocation::Invalid, CreateCasTestAddr(uint64(Idx) << 32 | 1ull));
+		}
+		Tx.BlockDeleted(FCasBlockId(1));
+
+		// Add a second block and some chunk locations
+		Tx.BlockCreated(ExpectedCurrentBlockId);
+		for (int32 Idx = 1; Idx <= 10; ++Idx)
+		{
+			Tx.ChunkLocation(FCasLocation
+			{
+				.BlockId		= ExpectedCurrentBlockId,
+				.BlockOffset	= uint32(Idx) * 256
+			},
+			CreateCasTestAddr(Idx));
+		}
+
+		Status = FCasJournal::Commit(MoveTemp(Tx));
+		CHECK(Status.IsOk());
+
+		// Act
+		const FCasSnapshot Snapshot = FCasSnapshot::FromJournal(JournalFile).ConsumeValueOrDie();
+
+		// Assert
+		CHECK(Snapshot.CurrentBlockId == ExpectedCurrentBlockId);
+		CHECK(Snapshot.Blocks.Num() == 1); 
+		CHECK(Snapshot.ChunkLocations.Num() == 10); 
+		for (int32 Idx = 1; Idx < Snapshot.ChunkLocations.Num(); ++Idx)
+		{
+			const FCasAddr Addr = CreateCasTestAddr(Idx);
+			const FCasSnapshot::FChunkLocation* Loc =
+				Algo::FindByPredicate(
+					Snapshot.ChunkLocations,
+					[&Addr](const FCasSnapshot::FChunkLocation& L) { return L.Get<0>() == Addr; });
+			CHECK(Loc != nullptr);
+			if (Loc != nullptr)
+			{
+				CHECK(Loc->Get<1>().BlockId == ExpectedCurrentBlockId);
+				CHECK(Loc->Get<1>().BlockOffset == uint32(Idx) * 256); 
+			}
+		}
+	}
+}
+
 #endif // WITH_IOSTORE_ONDEMAND_TESTS
 
 } // namespace UE::IoStore
diff --git a/Engine/Source/Runtime/Experimental/IoStore/OnDemand/Private/OnDemandInstallCache.h b/Engine/Source/Runtime/Experimental/IoStore/OnDemand/Private/OnDemandInstallCache.h
--- a/Engine/Source/Runtime/Experimental/IoStore/OnDemand/Private/OnDemandInstallCache.h
+++ b/Engine/Source/Runtime/Experimental/IoStore/OnDemand/Private/OnDemandInstallCache.h
@@ -38,8 +38,9 @@ public:
 struct FOnDemandInstallCacheConfig
 {
 	FString RootDirectory;
-	uint64	DiskQuota = 1ull << 30;
-	bool	bDropCache = false;
+	uint64	DiskQuota		= 1ull << 30;
+	uint64	JournalMaxSize	= 2ull << 20;
+	bool	bDropCache		= false;
 };
 
 TSharedPtr<IOnDemandInstallCache> MakeOnDemandInstallCache(
