Island Manager: Always assign edge levels if partial island sleeping enabled

diff --git a/Engine/Source/Runtime/Experimental/Chaos/Private/Chaos/Island/IslandManager.cpp b/Engine/Source/Runtime/Experimental/Chaos/Private/Chaos/Island/IslandManager.cpp
--- a/Engine/Source/Runtime/Experimental/Chaos/Private/Chaos/Island/IslandManager.cpp
+++ b/Engine/Source/Runtime/Experimental/Chaos/Private/Chaos/Island/IslandManager.cpp
@@ -552,6 +552,12 @@ namespace Chaos::Private
 		return NextVisitEpoch++;
 	}
 
+	bool FPBDIslandManager::ShouldAssignLevels() const
+	{
+		// Partial island sleeping always needs levels
+		return bAssignLevels || CVars::bChaosSolverPartialIslandSleep;
+	}
+
 	void FPBDIslandManager::SetIsDeterministic(const bool bInIsDeterministic)
 	{
 		bIsDeterministic = bInIsDeterministic;
@@ -1955,8 +1961,8 @@ namespace Chaos::Private
 			// Levels will only change when awake and when something has been added or removed
 			if (!Island->Flags.bIsSleeping && (!!Island->Flags.bItemsAdded || !!Island->Flags.bItemsRemoved))
 			{
-				// We only need to assign levels if shock propagation is enabled (but we always sort)
-				if (bAssignLevels)
+				// We only need to assign levels if shock propagation or partial island sleeping is enabled (but we always sort)
+				if (ShouldAssignLevels())
 				{
 					AssignIslandLevels(Island);
 				}
@@ -2619,9 +2625,8 @@ namespace Chaos::Private
 								if (!Edge->Flags.bWasProcessed)
 								{
 									FPBDIslandParticle* OtherNode = (Node == Edge->Nodes[0]) ? Edge->Nodes[1] : Edge->Nodes[0];
-									// Wake up the OtherNode if the node level is greater than or equal to the edge level (if level information available).
-									if ((OtherNode != nullptr) && OtherNode->Flags.bIsDynamic
-										&& bAssignLevels && (OtherNode->Level >= Edge->Level))
+									// Wake up the OtherNode if the node level is greater than or equal to the edge level.
+									if ((OtherNode != nullptr) && OtherNode->Flags.bIsDynamic && (OtherNode->Level >= Edge->Level))
 									{
 										ProcessConnectedParticles(Edge, ParticlesToWakeUp, true);
 										NodeQueue.Add(OtherNode);
diff --git a/Engine/Source/Runtime/Experimental/Chaos/Public/Chaos/Island/IslandManager.h b/Engine/Source/Runtime/Experimental/Chaos/Public/Chaos/Island/IslandManager.h
--- a/Engine/Source/Runtime/Experimental/Chaos/Public/Chaos/Island/IslandManager.h
+++ b/Engine/Source/Runtime/Experimental/Chaos/Public/Chaos/Island/IslandManager.h
@@ -615,6 +615,7 @@ namespace Chaos::Private
 		CHAOS_API void UpdateAllParticleSleepMetrics(const FRealSingle Dt);
 
 		CHAOS_API int32 GetNextVisitEpoch();
+		CHAOS_API bool ShouldAssignLevels() const;
 
 		// Ideally the island manager would not know about the particles, materials etc, but
 		// it is currently required for particle sleep management
