Dataflow : fix issue with null cache entries causing issue with reevaluation when connecting to an output
- especially was causing issues with reroute nodes

diff --git a/Engine/Source/Runtime/Experimental/Dataflow/Core/Public/Dataflow/DataflowContextCache.h b/Engine/Source/Runtime/Experimental/Dataflow/Core/Public/Dataflow/DataflowContextCache.h
--- a/Engine/Source/Runtime/Experimental/Dataflow/Core/Public/Dataflow/DataflowContextCache.h
+++ b/Engine/Source/Runtime/Experimental/Dataflow/Core/Public/Dataflow/DataflowContextCache.h
@@ -167,7 +167,7 @@ namespace UE::Dataflow
 		// Timestamp must be set to (Timestamp.Value - 1) to make sure that this type of entry is always invalid
 		//
 		FContextCacheElementNull(FGuid InNodeGuid, const FProperty* InProperty, FContextCacheKey InDataKey, uint32 InNodeHash, FTimestamp Timestamp)
-			: FContextCacheElementBase(EType::CacheElementNull, InNodeGuid, InProperty, InNodeHash, FTimestamp(Timestamp.Value - 1))
+			: FContextCacheElementBase(EType::CacheElementNull, InNodeGuid, InProperty, InNodeHash, FTimestamp((Timestamp.Value == 0)? 0: (Timestamp.Value - 1)))
 		{}
 
 		virtual const void* GetUntypedData(const IContextCacheStore& Context, const FProperty* PropertyIn) const
