Unit tests for IAS priority queue.

#rnx

diff --git a/Engine/Source/Runtime/Experimental/IoStore/OnDemand/Private/OnDemandIoDispatcherBackend.cpp b/Engine/Source/Runtime/Experimental/IoStore/OnDemand/Private/OnDemandIoDispatcherBackend.cpp
--- a/Engine/Source/Runtime/Experimental/IoStore/OnDemand/Private/OnDemandIoDispatcherBackend.cpp
+++ b/Engine/Source/Runtime/Experimental/IoStore/OnDemand/Private/OnDemandIoDispatcherBackend.cpp
@@ -1818,6 +1818,169 @@ TEST_CASE("IoStore::OnDemand::Ias::Misc", "[IoStoreOnDemand][Ias]")
 	}
 }
 
+TEST_CASE("IoStore::OnDemand::Ias::PriorityQueue", "[IoStoreOnDemand][Ias]")
+{
+	struct FTestRequest
+	{
+		explicit FTestRequest(int32 InValue)
+			: Value(InValue)
+		{ }
+		explicit FTestRequest(int32 InValue, int32 InPriority)
+			: Value(InValue)
+			, Priority(InPriority)
+		{ }
+		int32			Value;
+		int32			Priority	= 0;
+		FTestRequest*	NextRequest	= nullptr;
+	};
+
+	using FTestRequestAlloc	= TSingleThreadedSlabAllocator<FTestRequest>;
+	using FTestRequestQueue	= TThreadSafeIntrusiveQueue<FTestRequest>;
+
+	auto Dequeue = [](FTestRequestQueue& Queue) -> TArray<FTestRequest*>
+	{
+		TArray<FTestRequest*> Out; 
+		FTestRequest* NextRequest = Queue.Dequeue();
+		while (NextRequest != nullptr)
+		{
+			FTestRequest* Request	= NextRequest; 
+			NextRequest				= Request->NextRequest;
+			Request->NextRequest	= nullptr;
+			Out.Add(Request);
+		}
+		return Out;
+	};
+
+	auto Destroy = [](TArray<FTestRequest*>& Requests, FTestRequestAlloc& Alloc)
+	{
+		for (FTestRequest* R : Requests)
+		{
+			Alloc.Destroy(R);
+		}
+		Requests.Empty();
+	};
+
+	SECTION("EnqueueDequeue")
+	{
+		// Arrange
+		FTestRequestAlloc Alloc;
+		FTestRequestQueue Queue;
+		const int32 ExpectedCount = 100;
+
+		// Act
+		for (int32 Idx = 0; Idx < ExpectedCount; ++Idx)
+		{
+			Queue.Enqueue(Alloc.Construct(Idx));
+		}
+
+		TArray<FTestRequest*> Requests = Dequeue(Queue);
+
+		// Assert
+		CHECK(Requests.Num() == ExpectedCount);
+		for (int32 Idx = 0; Idx < ExpectedCount; ++Idx)
+		{
+			CHECK(Requests[Idx]->Value == Idx);
+			Alloc.Destroy(Requests[Idx]);
+		}
+	}
+
+	SECTION("EnqueueByPriority")
+	{
+		// Arrange
+		FTestRequestAlloc Alloc;
+		FTestRequestQueue Queue;
+		const int32 ExpectedCount = 100;
+
+		// Act
+		for (int32 Idx = 0; Idx < ExpectedCount; ++Idx)
+		{
+			Queue.EnqueueByPriority(Alloc.Construct(Idx, Idx));
+		}
+
+		TArray<FTestRequest*> Requests = Dequeue(Queue);
+
+		// Assert
+		CHECK(Requests.Num() == ExpectedCount);
+		for (int32 Idx = 0; Idx < ExpectedCount; ++Idx)
+		{
+			CHECK(Requests[Idx]->Value == (ExpectedCount - Idx - 1));
+			Alloc.Destroy(Requests[Idx]);
+		}
+	}
+
+	SECTION("UpdateWithHighestPriority")
+	{
+		// Arrange
+		FTestRequestAlloc Alloc;
+		FTestRequestQueue Queue;
+
+		// Act
+		FTestRequest* Request = Alloc.Construct(1, 1);
+		Queue.EnqueueByPriority(Request);
+		Queue.EnqueueByPriority(Alloc.Construct(2, 2));
+		Queue.EnqueueByPriority(Alloc.Construct(3, 3));
+		Request->Priority = 4;
+		Queue.Reprioritize(Request);
+
+		TArray<FTestRequest*> Requests = Dequeue(Queue);
+
+		// Assert
+		CHECK(Requests.Num() == 3); 
+		CHECK(Requests[0]->Value == 1);
+		CHECK(Requests[1]->Value == 3);
+		CHECK(Requests[2]->Value == 2);
+		Destroy(Requests, Alloc);
+	}
+
+	SECTION("UpdateWithHigherPriority")
+	{
+		// Arrange
+		FTestRequestAlloc Alloc;
+		FTestRequestQueue Queue;
+
+		// Act
+		FTestRequest* Request = Alloc.Construct(1, 1);
+		Queue.EnqueueByPriority(Request);
+		Queue.EnqueueByPriority(Alloc.Construct(2, 2));
+		Queue.EnqueueByPriority(Alloc.Construct(3, 4));
+		Request->Priority = 3;
+		Queue.Reprioritize(Request);
+
+		TArray<FTestRequest*> Requests = Dequeue(Queue);
+
+		// Assert
+		CHECK(Requests.Num() == 3); 
+		CHECK(Requests[0]->Value == 3);
+		CHECK(Requests[1]->Value == 1);
+		CHECK(Requests[2]->Value == 2);
+		Destroy(Requests, Alloc);
+	}
+
+	SECTION("UpdateWithLowestPriority")
+	{
+		// Arrange
+		FTestRequestAlloc Alloc;
+		FTestRequestQueue Queue;
+
+		// Act
+		Queue.EnqueueByPriority(Alloc.Construct(1, 1));
+		Queue.EnqueueByPriority(Alloc.Construct(2, 2));
+		FTestRequest* Request = Alloc.Construct(3, 3);
+		Queue.EnqueueByPriority(Request);
+		Request->Priority = 0;
+		Queue.Reprioritize(Request);
+
+		TArray<FTestRequest*> Requests = Dequeue(Queue);
+
+		// Assert
+		CHECK(Requests.Num() == 3); 
+		CHECK(Requests[0]->Value == 2);
+		CHECK(Requests[1]->Value == 1);
+		CHECK(Requests[2]->Value == 3);
+		Destroy(Requests, Alloc);
+	}
+}
+
 #endif // WITH_IOSTORE_ONDEMAND_TESTS
 
 } // namespace UE::IoStore
