CDN cache status tracking is now returned to the caller as part of the FMultiEndpointHttpClientResponse rather than being applied to the statistics directly by the client.

#rnx

- Since both IAD and IAS share the same client, this will allow us the collect each systems statistics independantly as previously IAD CDN cache hit/miss values were being collected by an IAS debug display.
- Additionally, the request logging that FMultiEndpointHttpClient can do now logs the cache hit/miss/??? after the request duration and size. This makes it easier to track what exactly happened to each request when using verbose logging.
- TrackCdnCacheStats now returns the status rather than passing it directly to the statistics.

diff --git a/Engine/Source/Runtime/Experimental/IoStore/OnDemand/Private/OnDemandHttpClient.cpp b/Engine/Source/Runtime/Experimental/IoStore/OnDemand/Private/OnDemandHttpClient.cpp
--- a/Engine/Source/Runtime/Experimental/IoStore/OnDemand/Private/OnDemandHttpClient.cpp
+++ b/Engine/Source/Runtime/Experimental/IoStore/OnDemand/Private/OnDemandHttpClient.cpp
@@ -1,7 +1,6 @@
 // Copyright Epic Games, Inc. All Rights Reserved.
 
 #include "OnDemandHttpClient.h"
-#include "Statistics.h"
 
 #include "Containers/StringConv.h"
 #include "HAL/FileManager.h"
@@ -20,7 +19,7 @@
 namespace UE::IoStore
 {
 
-static void TrackCdnCacheStats(const HTTP::FResponse& Response)
+static int8 TrackCdnCacheStats(const HTTP::FResponse& Response)
 {
 #if !UE_BUILD_TEST
 	for (FAnsiStringView Header : { "X-Cache", "CF-Cache-Status" })
@@ -30,14 +29,19 @@ static void TrackCdnCacheStats(const HTTP::FResponse& Response)
 		{
 			continue;
 		}
-		int32 Reply = (Value.Find("HIT", 0, ESearchCase::IgnoreCase) >= 0);
-		FOnDemandIoBackendStats::Get()->OnHttpCdnCacheReply(Reply);
-		return;
+
+		return Value.Find("HIT", 0, ESearchCase::IgnoreCase) >= 0 ? 1 : 0;
 	}
-	FOnDemandIoBackendStats::Get()->OnHttpCdnCacheReply(-1);
+	
+	return -1;
 #endif
 }
 
+static const TCHAR* CDNCacheStatusToString(int8 Status)
+{
+	return Status > 0 ? TEXT("HIT") : Status == 0 ? TEXT("MISS") : TEXT("???");
+}
+
 FIoStatus LoadDefaultHttpCertificates(bool& bWasLoaded)
 {
 	using namespace UE::IoStore::HTTP;
@@ -276,7 +280,7 @@ void FMultiEndpointHttpClient::IssueRequest(FRequest&& Request)
 				FResponse& HttpResponse = TicketStatus.GetResponse();
 				Request.StatusCode		= HttpResponse.GetStatusCode();
 				HttpResponse.SetDestination(&Request.Body);
-				TrackCdnCacheStats(HttpResponse);
+				Request.CDNCacheStatus = TrackCdnCacheStats(HttpResponse);
 				break;
 			}
 			case FTicketStatus::EId::Content:
@@ -300,7 +304,8 @@ void FMultiEndpointHttpClient::IssueRequest(FRequest&& Request)
 						.DurationMilliseconds	= DurationMs,
 						.StatusCode				= Request.StatusCode,
 						.RetryCount				= Request.RetryCount,
-						.HostIndex				= Request.Host
+						.HostIndex				= Request.Host,
+						.CDNCacheStatus = Request.CDNCacheStatus
 					};
 
 					if (bError)
@@ -429,10 +434,11 @@ void FMultiEndpointHttpClient::Log(const FRequest& Request) const
 	const uint64 Size = Request.Body.GetSize() >> 10;
 
 	FMsg::Logf(__FILE__, __LINE__, Config.LogCategory->GetCategoryName(), Config.LogVerbosity,
-		TEXT("http-%3u: %5" UINT64_FMT "ms %5" UINT64_FMT "KiB %s (Attempt %u/%u)"),
+		TEXT("http-%3u: %5" UINT64_FMT "ms %5" UINT64_FMT "KiB [%4s] %s (Attempt %u/%u)"),
 		Request.StatusCode,
 		DurationMs,
 		Size,
+		CDNCacheStatusToString(Request.CDNCacheStatus),
 		WriteToString<512>(Request.Connection.HostGroup.Host(Request.Host), Request.RelativeUrl).ToString(),
 		Request.RetryCount,
 		GetRetryLimitForRequest(Request)
diff --git a/Engine/Source/Runtime/Experimental/IoStore/OnDemand/Private/OnDemandHttpClient.h b/Engine/Source/Runtime/Experimental/IoStore/OnDemand/Private/OnDemandHttpClient.h
--- a/Engine/Source/Runtime/Experimental/IoStore/OnDemand/Private/OnDemandHttpClient.h
+++ b/Engine/Source/Runtime/Experimental/IoStore/OnDemand/Private/OnDemandHttpClient.h
@@ -61,6 +61,7 @@ struct FMultiEndpointHttpClientResponse
 	uint32		StatusCode = 0;
 	uint32		RetryCount = 0;
 	int32		HostIndex = INDEX_NONE;
+	int8		CDNCacheStatus = -1;
 };
 
 /**
@@ -114,6 +115,7 @@ private:
 		uint32				RetryCount = 0;
 		uint32				StatusCode = 0;
 		int32				Host = INDEX_NONE;
+		int8				CDNCacheStatus = -1;
 	};
 
 										FMultiEndpointHttpClient(const FMultiEndpointHttpClientConfig& Config);
diff --git a/Engine/Source/Runtime/Experimental/IoStore/OnDemand/Private/OnDemandIoDispatcherBackend.cpp b/Engine/Source/Runtime/Experimental/IoStore/OnDemand/Private/OnDemandIoDispatcherBackend.cpp
--- a/Engine/Source/Runtime/Experimental/IoStore/OnDemand/Private/OnDemandIoDispatcherBackend.cpp
+++ b/Engine/Source/Runtime/Experimental/IoStore/OnDemand/Private/OnDemandIoDispatcherBackend.cpp
@@ -1553,6 +1553,8 @@ void FOnDemandIoBackend::ProcessHttpRequests(FMultiEndpointHttpClient& HttpClien
 								Stats.OnHttpRetry();
 							}
 
+							Stats.OnHttpCdnCacheReply(HttpResponse.CDNCacheStatus);
+
 							if (HttpResponse.IsOk())
 							{
 								ChunkRequest->HostGroup.OnSuccessfulResponse();
