Use the lastest task graph API in collision detection

diff --git a/Engine/Source/Runtime/Experimental/Chaos/Private/Chaos/Collision/SpatialAccelerationBroadPhase.cpp b/Engine/Source/Runtime/Experimental/Chaos/Private/Chaos/Collision/SpatialAccelerationBroadPhase.cpp
--- a/Engine/Source/Runtime/Experimental/Chaos/Private/Chaos/Collision/SpatialAccelerationBroadPhase.cpp
+++ b/Engine/Source/Runtime/Experimental/Chaos/Private/Chaos/Collision/SpatialAccelerationBroadPhase.cpp
@@ -1,7 +1,5 @@
 // Copyright Epic Games, Inc. All Rights Reserved.
 #include "Chaos/Collision/SpatialAccelerationBroadPhase.h"
-#include "Chaos/PBDCollisionConstraints.h"
-
 
 namespace Chaos
 {
@@ -10,115 +8,4 @@ namespace Chaos
 		int32 NumWorkerCollisionFactor = 2;
 		FAutoConsoleVariableRef CVarNumWorkerCollisionFactor(TEXT("p.Chaos.NumWorkerCollisionFactor"), NumWorkerCollisionFactor, TEXT("Set the number of tasks created for collision detection per worker."));
 	}
-	namespace Private
-	{
-		FAssignMidPhaseTask::FAssignMidPhaseTask(FSpatialAccelerationBroadPhase& InSpatialAccelerationBroadPhase,Private::FBroadPhaseContext& InBroadPhaseContext)
-		: SpatialAccelerationBroadPhase(InSpatialAccelerationBroadPhase)
-		, BroadPhaseContext(InBroadPhaseContext)
-		{
-			CVD_GET_CURRENT_CONTEXT(CVDContext);
-		}
-
-		ENamedThreads::Type FAssignMidPhaseTask::GetDesiredThread()
-		{
-			return ENamedThreads::AnyHiPriThreadHiPriTask;
-		}
-		void FAssignMidPhaseTask::Process()
-		{
-			CVD_SCOPE_CONTEXT(CVDContext);
-			SpatialAccelerationBroadPhase.AssignMidPhases(BroadPhaseContext);
-		}
-
-		FProcessMidPhaseTask::FProcessMidPhaseTask(FSpatialAccelerationBroadPhase& InSpatialAccelerationBroadPhase, Private::FBroadPhaseContext& InBroadPhaseContext, FReal InDt)
-		: SpatialAccelerationBroadPhase(InSpatialAccelerationBroadPhase)
-		, BroadPhaseContext(InBroadPhaseContext)
-		, Dt(InDt)
-		{
-			CVD_GET_CURRENT_CONTEXT(CVDContext);
-		}
-
-		ENamedThreads::Type FProcessMidPhaseTask::GetDesiredThread()
-		{
-			return ENamedThreads::AnyHiPriThreadHiPriTask;
-		}
-		void FProcessMidPhaseTask::Process()
-		{
-			CVD_SCOPE_CONTEXT(CVDContext);
-			SpatialAccelerationBroadPhase.ProcessMidPhases(Dt, BroadPhaseContext);
-		}
-
-		FSortConstraintTask::FSortConstraintTask(Private::FBroadPhaseContext& InBroadPhaseContext)
-			: BroadPhaseContext(InBroadPhaseContext)
-		{
-			CVD_GET_CURRENT_CONTEXT(CVDContext);
-		}
-
-		ENamedThreads::Type FSortConstraintTask::GetDesiredThread()
-		{
-			return ENamedThreads::AnyNormalThreadNormalTask;
-		}
-
-		void FSortConstraintTask::Process()
-		{
-			QUICK_SCOPE_CYCLE_COUNTER(STAT_SortConstraintBatch);
-			CVD_SCOPE_CONTEXT(CVDContext);
-			BroadPhaseContext.CollisionContext.Allocator->NewActiveConstraints.Sort(
-				[](const FPBDCollisionConstraintHandle& L, const FPBDCollisionConstraintHandle& R)
-				{
-					return L.GetContact().GetCollisionSortKey() < R.GetContact().GetCollisionSortKey();
-				});
-		}
-
-		FMergeConstraintsTask::FMergeConstraintsTask(Private::FBroadPhaseContext& InBroadPhaseContext, Private::FBroadPhaseContext& InOtherBroadPhaseContext, bool bInNeedSorting)
-			: BroadPhaseContext(InBroadPhaseContext)
-			, OtherBroadPhaseContext(InOtherBroadPhaseContext)
-			, bNeedSorting(bInNeedSorting)
-		{
-			CVD_GET_CURRENT_CONTEXT(CVDContext);
-		}
-
-		ENamedThreads::Type FMergeConstraintsTask::GetDesiredThread()
-		{
-			return ENamedThreads::AnyNormalThreadNormalTask;
-		}
-		void FMergeConstraintsTask::Process()
-		{
-			QUICK_SCOPE_CYCLE_COUNTER(STAT_MergeConstraints);
-			CVD_SCOPE_CONTEXT(CVDContext);
-
-			TArray<FPBDCollisionConstraint*>& ThisArray = BroadPhaseContext.CollisionContext.Allocator->NewActiveConstraints;
-			TArray<FPBDCollisionConstraint*>& OtherArray = OtherBroadPhaseContext.CollisionContext.Allocator->NewActiveConstraints;
-
-			if (bNeedSorting)
-			{
-				TArray<FPBDCollisionConstraint*> ResultArray;
-				const int32 NumThis = ThisArray.Num();
-				const int32 NumOther = OtherArray.Num();
-				ResultArray.Reserve(NumThis + NumOther);
-
-				int32 ThisIndex = 0;
-				int32 OtherIndex = 0;
-
-				while (ThisIndex < NumThis || OtherIndex < NumOther)
-				{
-					if (ThisIndex < NumThis && (OtherIndex >= NumOther || ThisArray[ThisIndex]->GetContact().GetCollisionSortKey() < OtherArray[OtherIndex]->GetContact().GetCollisionSortKey()))
-					{
-						ResultArray.Add(ThisArray[ThisIndex++]);
-					}
-					else
-					{
-						check(OtherIndex < NumOther);
-						ResultArray.Add(OtherArray[OtherIndex++]);
-					}
-				}
-
-				ThisArray = MoveTemp(ResultArray);
-			}
-			else
-			{
-				ThisArray.Append(OtherArray);
-			}
-			OtherArray.Reset();
-		}
-	}
 }
\ No newline at end of file
diff --git a/Engine/Source/Runtime/Experimental/Chaos/Public/Chaos/Collision/SpatialAccelerationBroadPhase.h b/Engine/Source/Runtime/Experimental/Chaos/Public/Chaos/Collision/SpatialAccelerationBroadPhase.h
--- a/Engine/Source/Runtime/Experimental/Chaos/Public/Chaos/Collision/SpatialAccelerationBroadPhase.h
+++ b/Engine/Source/Runtime/Experimental/Chaos/Public/Chaos/Collision/SpatialAccelerationBroadPhase.h
@@ -14,6 +14,7 @@
 #include "Chaos/EvolutionResimCache.h"
 #include "Chaos/GeometryParticlesfwd.h"
 #include "Chaos/AABBTree.h"
+#include "Tasks/Task.h"
 
 // Set to 1 to enable some sanity chacks on the filtered broadphase results
 #ifndef CHAOS_CHECK_BROADPHASE
@@ -294,88 +295,6 @@ namespace Chaos
 			/** Handle to be stored to retrieve the payload on which we are querying the acceleration structure*/
 			FAccelerationStructureHandle AccelerationHandle;
 		};
-
-		template<bool bOnlyRigid, typename THandle, typename LambdaHandle, typename SpatialAccelerationType>
-		class FProduceOverlapTask;
-
-		class FAssignMidPhaseTask
-		{
-		public:
-			FAssignMidPhaseTask(FSpatialAccelerationBroadPhase& InSpatialAccelerationBroadPhase, Private::FBroadPhaseContext& InBroadPhaseContext);
-
-			TStatId GetStatId() const { RETURN_QUICK_DECLARE_CYCLE_STAT(FAssignMidPhaseTask, STATGROUP_TaskGraphTasks); }
-			static CHAOS_API ENamedThreads::Type GetDesiredThread();
-			static ESubsequentsMode::Type GetSubsequentsMode() { return ESubsequentsMode::TrackSubsequents; }
-			void DoTask(ENamedThreads::Type CurrentThread, const FGraphEventRef& MyCompletionGraphEvent) { Process(); }
-
-		private:
-			void Process();
-			FSpatialAccelerationBroadPhase& SpatialAccelerationBroadPhase;
-			Private::FBroadPhaseContext& BroadPhaseContext;
-
-#if WITH_CHAOS_VISUAL_DEBUGGER
-			FChaosVDContext CVDContext;
-#endif
-		};
-
-		class FProcessMidPhaseTask
-		{
-		public:
-			FProcessMidPhaseTask(FSpatialAccelerationBroadPhase& InSpatialAccelerationBroadPhase, Private::FBroadPhaseContext& InBroadPhaseContext, FReal InDt);
-
-			TStatId GetStatId() const { RETURN_QUICK_DECLARE_CYCLE_STAT(FProcessMidPhaseTask, STATGROUP_TaskGraphTasks); }
-			static CHAOS_API ENamedThreads::Type GetDesiredThread();
-			static ESubsequentsMode::Type GetSubsequentsMode() { return ESubsequentsMode::TrackSubsequents; }
-			void DoTask(ENamedThreads::Type CurrentThread, const FGraphEventRef& MyCompletionGraphEvent) { Process(); }
-		private:
-			void Process();
-			FSpatialAccelerationBroadPhase& SpatialAccelerationBroadPhase;
-			Private::FBroadPhaseContext& BroadPhaseContext;
-			FReal Dt;
-
-#if WITH_CHAOS_VISUAL_DEBUGGER
-			FChaosVDContext CVDContext;
-#endif
-		};
-
-		class FSortConstraintTask
-		{
-		public:
-			FSortConstraintTask(Private::FBroadPhaseContext& InBroadPhaseContext);
-
-			TStatId GetStatId() const { RETURN_QUICK_DECLARE_CYCLE_STAT(FSortConstraintTask, STATGROUP_TaskGraphTasks); }
-			static CHAOS_API ENamedThreads::Type GetDesiredThread();
-			static ESubsequentsMode::Type GetSubsequentsMode() { return ESubsequentsMode::TrackSubsequents; }
-			void DoTask(ENamedThreads::Type CurrentThread, const FGraphEventRef& MyCompletionGraphEvent) { Process(); }
-		private:
-			void Process();
-			Private::FBroadPhaseContext& BroadPhaseContext;
-
-#if WITH_CHAOS_VISUAL_DEBUGGER
-			FChaosVDContext CVDContext;
-#endif
-		};
-
-		class FMergeConstraintsTask
-		{
-		public:
-			FMergeConstraintsTask(Private::FBroadPhaseContext& InBroadPhaseContext, Private::FBroadPhaseContext& InOtherBroadPhaseContext, bool bNeedSorting);
-
-			TStatId GetStatId() const { RETURN_QUICK_DECLARE_CYCLE_STAT(FMergeConstraintsTask, STATGROUP_TaskGraphTasks); }
-			static CHAOS_API ENamedThreads::Type GetDesiredThread();
-			static ESubsequentsMode::Type GetSubsequentsMode() { return ESubsequentsMode::TrackSubsequents; }
-			void DoTask(ENamedThreads::Type CurrentThread, const FGraphEventRef& MyCompletionGraphEvent) { Process(); }
-		private:
-			void Process();
-			Private::FBroadPhaseContext& BroadPhaseContext;
-			Private::FBroadPhaseContext& OtherBroadPhaseContext;
-			bool bNeedSorting;
-
-#if WITH_CHAOS_VISUAL_DEBUGGER
-			FChaosVDContext CVDContext;
-#endif
-		};
-
 	}
 
 	/**
@@ -452,18 +371,21 @@ namespace Chaos
 				FChaosVDContextWrapper CVDContext;
 				CVD_GET_WRAPPED_CURRENT_CONTEXT(CVDContext);
 
-				FGraphEventArray NextPendingTasks;
+				TArray<UE::Tasks::FTask> NextPendingTasks;
 				const int32 NumContextTask = PendingTasks.Num();
 				check(NumContextTask <= NumActiveBroadphaseContexts);
 				for (int32 ContextIndex = 0; ContextIndex < NumContextTask; ContextIndex++)
 				{
-					const FGraphEventArray Prereq = { PendingTasks[ContextIndex] };
-					const FGraphEventRef PendingTask = TGraphTask<Private::FAssignMidPhaseTask>::CreateTask(&Prereq).ConstructAndDispatchWhenReady(*this, BroadphaseContexts[ContextIndex]);
+					Private::FBroadPhaseContext& BroadPhaseContext = BroadphaseContexts[ContextIndex];
+
+					UE::Tasks::FTask PendingTask = UE::Tasks::Launch(UE_SOURCE_LOCATION, [this, &BroadPhaseContext](){
+						AssignMidPhases(BroadPhaseContext);
+						}, PendingTasks[ContextIndex]);
 					NextPendingTasks.Add(PendingTask);
 				}
 				// It is required to wait for the middle phase modifiers, which is a callback to modify middle phases.
 				// This could be removed by adding dependent tasks for the middle phase modifiers.
-				FTaskGraphInterface::Get().WaitUntilTasksComplete(NextPendingTasks);
+				UE::Tasks::Wait(NextPendingTasks);
 			}
 
 			// Run some error checks in non-shipping builds
@@ -493,16 +415,18 @@ namespace Chaos
 			// MidPhase runs the NarrowPhase in ProcessMidPhase
 			RedistributeMidPhasesInContexts();
 			const int32 NumContextTask = PendingTasks.Num();
-			FGraphEventArray ProcessMidPhaseTasks;
 			PendingTasks.Reset();
 			ContextsWithConstraints.Reset();
 			check(NumContextTask <= NumActiveBroadphaseContexts);
 			for (int32 ContextIndex = 0; ContextIndex < NumContextTask; ContextIndex++)
 			{
+				Private::FBroadPhaseContext& BroadphaseContext = BroadphaseContexts[ContextIndex];
 				// Accessing the overlaps array is safe here, all overlaps have been finished to be computed at this point
-				if (!BroadphaseContexts[ContextIndex].Overlaps.IsEmpty())
+				if (!BroadphaseContext.Overlaps.IsEmpty())
 				{
-					const FGraphEventRef ProcessMidPhaseTask = TGraphTask<Private::FProcessMidPhaseTask>::CreateTask().ConstructAndDispatchWhenReady(*this, BroadphaseContexts[ContextIndex], Dt);
+					UE::Tasks::FTask ProcessMidPhaseTask = UE::Tasks::Launch(UE_SOURCE_LOCATION, [this, &BroadphaseContext, Dt]() {
+						ProcessMidPhases(Dt, BroadphaseContext);
+						}, UE::Tasks::ETaskPriority::High);
 					PendingTasks.Add(ProcessMidPhaseTask);
 					ContextsWithConstraints.Add(ContextIndex);
 				}
@@ -513,11 +437,18 @@ namespace Chaos
 		{
 			if (bIsDeterministic)
 			{
-				FGraphEventArray ProcessSorting;
+				TArray<UE::Tasks::FTask> ProcessSorting;
 				for (int32 ContextIndex = 0; ContextIndex < PendingTasks.Num(); ContextIndex++)
 				{
-					FGraphEventArray Prereq = { PendingTasks[ContextIndex] };
-					const FGraphEventRef SortConstraintTask = TGraphTask<Private::FSortConstraintTask>::CreateTask(&Prereq).ConstructAndDispatchWhenReady(BroadphaseContexts[ContextsWithConstraints[ContextIndex]]);
+					Private::FBroadPhaseContext& BroadphaseContext = BroadphaseContexts[ContextsWithConstraints[ContextIndex]];
+					UE::Tasks::FTask SortConstraintTask = UE::Tasks::Launch(UE_SOURCE_LOCATION, [this, &BroadphaseContext]() {
+						QUICK_SCOPE_CYCLE_COUNTER(STAT_SortConstraintBatch);
+						BroadphaseContext.CollisionContext.Allocator->NewActiveConstraints.Sort(
+							[](const FPBDCollisionConstraintHandle& L, const FPBDCollisionConstraintHandle& R)
+							{
+								return L.GetContact().GetCollisionSortKey() < R.GetContact().GetCollisionSortKey();
+							});
+						}, PendingTasks[ContextIndex]);
 					ProcessSorting.Add(SortConstraintTask);
 				}
 				PendingTasks = MoveTemp(ProcessSorting);
@@ -527,12 +458,50 @@ namespace Chaos
 			while (PendingTasks.Num() >= 2)
 			{
 				int32 NumProcessSorting = PendingTasks.Num();
-				FGraphEventArray ProcessMerging;
+				TArray<UE::Tasks::FTask> ProcessMerging;
 				// We merge data in the context that have a power of 2 index, so for first layer 0, 2, 4..., for the second layer 0, 4, 8... and so on.
 				for (int32 ProcessIndex = 0, ContextIndex = 0; ProcessIndex < NumProcessSorting - 1; ProcessIndex += 2, ContextIndex += 2*NeighborStep)
 				{
-					FGraphEventArray Prereq = { PendingTasks[ProcessIndex], PendingTasks[ProcessIndex + 1] };
-					const FGraphEventRef MergeConstraintTask = TGraphTask<Private::FMergeConstraintsTask>::CreateTask(&Prereq).ConstructAndDispatchWhenReady(BroadphaseContexts[ContextsWithConstraints[ContextIndex]], BroadphaseContexts[ContextsWithConstraints[ContextIndex + NeighborStep]], bIsDeterministic);
+					TArray<UE::Tasks::FTask> Preriq{ PendingTasks[ProcessIndex], PendingTasks[ProcessIndex + 1] };
+					Private::FBroadPhaseContext& Context = BroadphaseContexts[ContextsWithConstraints[ContextIndex]];
+					Private::FBroadPhaseContext& OtherContext = BroadphaseContexts[ContextsWithConstraints[ContextIndex + NeighborStep]];
+					UE::Tasks::FTask MergeConstraintTask = UE::Tasks::Launch(UE_SOURCE_LOCATION, [this, &Context, &OtherContext, bIsDeterministic]() 
+					{
+						QUICK_SCOPE_CYCLE_COUNTER(STAT_MergeConstraints);
+
+						TArray<FPBDCollisionConstraint*>& ThisArray = Context.CollisionContext.Allocator->NewActiveConstraints;
+						TArray<FPBDCollisionConstraint*>& OtherArray = OtherContext.CollisionContext.Allocator->NewActiveConstraints;
+
+						if (bIsDeterministic)
+						{
+							TArray<FPBDCollisionConstraint*> ResultArray;
+							const int32 NumThis = ThisArray.Num();
+							const int32 NumOther = OtherArray.Num();
+							ResultArray.Reserve(NumThis + NumOther);
+
+							int32 ThisIndex = 0;
+							int32 OtherIndex = 0;
+
+							while (ThisIndex < NumThis || OtherIndex < NumOther)
+							{
+								if (ThisIndex < NumThis && (OtherIndex >= NumOther || ThisArray[ThisIndex]->GetContact().GetCollisionSortKey() < OtherArray[OtherIndex]->GetContact().GetCollisionSortKey()))
+								{
+									ResultArray.Add(ThisArray[ThisIndex++]);
+								}
+								else
+								{
+									check(OtherIndex < NumOther);
+									ResultArray.Add(OtherArray[OtherIndex++]);
+								}
+							}
+							ThisArray = MoveTemp(ResultArray);
+						}
+						else
+						{
+							ThisArray.Append(OtherArray);
+						}
+						OtherArray.Reset();
+						}, Preriq);
 					ProcessMerging.Add(MergeConstraintTask);
 				}
 				// Add the last one for odd process number
@@ -545,7 +514,7 @@ namespace Chaos
 			}
 			ensure(PendingTasks.Num() <= 1);
 
-			FTaskGraphInterface::Get().WaitUntilTasksComplete(PendingTasks);
+			UE::Tasks::Wait(PendingTasks);
 
 			// Lets move all constraints in the first allocator to be picked up by the CollisionConstraintAllocator
 			if (!ContextsWithConstraints.IsEmpty() && ContextsWithConstraints[0] != 0)
@@ -678,10 +647,12 @@ namespace Chaos
 		}
 
 	private:
-
 		template<bool bOnlyRigid, typename SpatialAccelerationType, typename THandle, typename LambdaType>
 		void DispatchTasks(int32& ContextIndex, int32 NumParticles, const SpatialAccelerationType& InSpatialAcceleration, FReal Dt, LambdaType ReturnHandle)
 		{
+			using THandleBase = typename THandle::THandleBase;
+			using TTransientHandle = typename THandle::TTransientHandle;
+
 			const int32 ParticleByContext = FMath::Max(FMath::DivideAndRoundUp(NumParticles, NumActiveBroadphaseContexts), SmallBatchSize);
 			const int32 NumBatch = FMath::DivideAndRoundUp(NumParticles, ParticleByContext);
 
@@ -692,23 +663,48 @@ namespace Chaos
 					ContextIndex = 0;
 				}
 
-				int32 LastIndex = (BatchIndex + 1) * ParticleByContext;
-				LastIndex = BatchIndex == NumBatch - 1 ? FMath::Min(NumParticles, LastIndex) : LastIndex;
+				int32 EndIndex = (BatchIndex + 1) * ParticleByContext;
+				EndIndex = BatchIndex == NumBatch - 1 ? FMath::Min(NumParticles, EndIndex) : EndIndex;
 
 				if (ContextIndex >= PendingTasks.Num())
 				{
-					FGraphEventRef PendingTask = TGraphTask<Private::FProduceOverlapTask<bOnlyRigid, THandle, LambdaType, SpatialAccelerationType>>::CreateTask().
-						ConstructAndDispatchWhenReady(*this, BroadphaseContexts[ContextIndex],
-							InSpatialAcceleration, Dt, ReturnHandle, BatchIndex * ParticleByContext, LastIndex);
+					Private::FBroadPhaseContext& BroadPhaseContext = BroadphaseContexts[ContextIndex];
+					const int32 StartIndex = BatchIndex * ParticleByContext;
+					UE::Tasks::FTask PendingTask = UE::Tasks::Launch(UE_SOURCE_LOCATION, [this, &BroadPhaseContext, &InSpatialAcceleration, Dt, ReturnHandle, StartIndex, EndIndex]()
+					{
+						SCOPE_CYCLE_COUNTER(STAT_Collisions_ParticlePairBroadPhase);
+						for (int32 Index = StartIndex; Index < EndIndex; Index++)
+						{
+							THandleBase Handle = ReturnHandle(Index);
+							TTransientHandle& TransientHandle = static_cast<TTransientHandle&>(Handle);
+							if (!TransientHandle.LightWeightDisabled())
+							{
+								ProduceParticleOverlaps<bOnlyRigid>(Dt, TransientHandle.Handle(), InSpatialAcceleration, BroadPhaseContext);
+							}
+						}
+					});
 					PendingTasks.Add(PendingTask);
 				}
 				else
 				{
+					
 					// If all context threads has been filled, create other tasks and link them to the previous one, 
-					FGraphEventArray Prereq = { PendingTasks[ContextIndex] };
-					FGraphEventRef PendingTask = TGraphTask<Private::FProduceOverlapTask<bOnlyRigid, THandle, LambdaType, SpatialAccelerationType>>::CreateTask(&Prereq).
-						ConstructAndDispatchWhenReady(*this, BroadphaseContexts[ContextIndex],
-							InSpatialAcceleration, Dt, ReturnHandle, BatchIndex * ParticleByContext, LastIndex);
+					Private::FBroadPhaseContext& BroadPhaseContext = BroadphaseContexts[ContextIndex];
+					const int32 StartIndex = BatchIndex * ParticleByContext;
+					UE::Tasks::FTask PendingTask = UE::Tasks::Launch(UE_SOURCE_LOCATION, [this, &BroadPhaseContext, &InSpatialAcceleration, Dt, ReturnHandle, StartIndex, EndIndex]()
+					{
+						SCOPE_CYCLE_COUNTER(STAT_Collisions_ParticlePairBroadPhase);
+						for (int32 Index = StartIndex; Index < EndIndex; Index++)
+						{
+							THandleBase Handle = ReturnHandle(Index);
+							TTransientHandle& TransientHandle = static_cast<TTransientHandle&>(Handle);
+							if (!TransientHandle.LightWeightDisabled())
+							{
+								ProduceParticleOverlaps<bOnlyRigid>(Dt, TransientHandle.Handle(), InSpatialAcceleration, BroadPhaseContext);
+							}
+						}
+					}, PendingTasks[ContextIndex]);
+
 					PendingTasks[ContextIndex] = PendingTask;
 				}
 			}
@@ -950,62 +946,8 @@ namespace Chaos
 		int32 NumBroadPhasePairs;
 		int32 NumMidPhases;
 
-		FGraphEventArray PendingTasks;
+		TArray<UE::Tasks::FTask> PendingTasks;
 		// Contexts indexes in which there are constraints
 		TArray<int32> ContextsWithConstraints;
-
-		template<bool bOnlyRigid, typename THandle, typename LambdaHandle, typename SpatialAccelerationType>
-		friend class Private::FProduceOverlapTask;
-		friend class Private::FAssignMidPhaseTask;
-		friend class Private::FProcessMidPhaseTask;
 	};
-	namespace Private
-	{
-		template<bool bOnlyRigid, typename THandle, typename LambdaHandle, typename SpatialAccelerationType>
-		class FProduceOverlapTask
-		{
-		public:
-
-			FProduceOverlapTask(FSpatialAccelerationBroadPhase& InSpatialAccelerationBroadPhase, Private::FBroadPhaseContext& InBroadPhaseContext,
-				const SpatialAccelerationType& InSpatialAcceleration, FReal InDt, LambdaHandle InReturnHandle, int32 InStartParticleIndex, int32 InEndParticleIndex)
-				: SpatialAccelerationBroadPhase(&InSpatialAccelerationBroadPhase)
-				, BroadPhaseContext(InBroadPhaseContext)
-				, SpatialAcceleration(InSpatialAcceleration)
-				, Dt(InDt)
-				, ReturnHandle(InReturnHandle)
-				, StartParticleIndex(InStartParticleIndex)
-				, EndParticleIndex(InEndParticleIndex)
-			{}
-
-			TStatId GetStatId() const { RETURN_QUICK_DECLARE_CYCLE_STAT(FProduceOverlapTask, STATGROUP_TaskGraphTasks); }
-			static CHAOS_API ENamedThreads::Type GetDesiredThread() { return ENamedThreads::AnyHiPriThreadHiPriTask; }
-			static ESubsequentsMode::Type GetSubsequentsMode() { return ESubsequentsMode::TrackSubsequents; }
-			void DoTask(ENamedThreads::Type CurrentThread, const FGraphEventRef& MyCompletionGraphEvent) { Process(); }
-
-		private:
-			using THandleBase = typename THandle::THandleBase;
-			using TTransientHandle = typename THandle::TTransientHandle;
-			void Process()
-			{
-				SCOPE_CYCLE_COUNTER(STAT_Collisions_ParticlePairBroadPhase);
-				for (int32 Index = StartParticleIndex; Index < EndParticleIndex; Index++)
-				{
-					THandleBase Handle = ReturnHandle(Index);
-					TTransientHandle& TransientHandle = static_cast<TTransientHandle&>(Handle);
-					if (!TransientHandle.LightWeightDisabled())
-					{
-						SpatialAccelerationBroadPhase->ProduceParticleOverlaps<bOnlyRigid>(Dt, TransientHandle.Handle(), SpatialAcceleration, BroadPhaseContext);
-					}
-				}
-			}
-
-			FSpatialAccelerationBroadPhase* SpatialAccelerationBroadPhase;
-			Private::FBroadPhaseContext& BroadPhaseContext;
-			const SpatialAccelerationType& SpatialAcceleration;
-			FReal Dt;
-			LambdaHandle ReturnHandle;
-			int32 StartParticleIndex;
-			int32 EndParticleIndex;
-		};
-	}
 }
