[HWRT] Fix FGeometryCollectionSceneProxyBase::GetDynamicRayTracingInstances(...) using FRHICommandListImmediate::Get() instead of Collector.GetRHICommandList().

diff --git a/Engine/Source/Runtime/Experimental/GeometryCollectionEngine/Private/GeometryCollection/GeometryCollectionSceneProxy.cpp b/Engine/Source/Runtime/Experimental/GeometryCollectionEngine/Private/GeometryCollection/GeometryCollectionSceneProxy.cpp
--- a/Engine/Source/Runtime/Experimental/GeometryCollectionEngine/Private/GeometryCollection/GeometryCollectionSceneProxy.cpp
+++ b/Engine/Source/Runtime/Experimental/GeometryCollectionEngine/Private/GeometryCollection/GeometryCollectionSceneProxy.cpp
@@ -445,10 +445,8 @@ void FGeometryCollectionSceneProxyBase::UpdateSkinnedPositions(FRHICommandListBa
 }
 
 #if RHI_RAYTRACING
-void FGeometryCollectionSceneProxyBase::UpdatingRayTracingGeometry_RenderingThread(TArray<FGeometryCollectionMeshElement> const& InSectionArray)
+void FGeometryCollectionSceneProxyBase::UpdatingRayTracingGeometry_RenderingThread(FRHICommandListBase& RHICmdList, TConstArrayView<FGeometryCollectionMeshElement> InSectionArray)
 {
-	FRHICommandList& RHICmdList = FRHICommandListImmediate::Get();
-
 	if (bRayTracingGeometryDirty)
 	{
 		RayTracingGeometry.Initializer.Segments.Empty();
@@ -506,7 +504,7 @@ void FGeometryCollectionSceneProxyBase::GetDynamicRayTracingInstances(FRayTracin
 
 	TArray<FGeometryCollectionMeshElement> const& SectionArray = MeshDescription.Sections;
 
-	UpdatingRayTracingGeometry_RenderingThread(SectionArray);
+	UpdatingRayTracingGeometry_RenderingThread(Collector.GetRHICommandList(), SectionArray);
 
 	if (SectionArray.Num() && RayTracingGeometry.IsValid())
 	{
@@ -537,7 +535,6 @@ void FGeometryCollectionSceneProxyBase::GetDynamicRayTracingInstances(FRayTracin
 		{
 			const FGeometryCollectionMeshElement& Section = SectionArray[SectionIndex];
 
-			// Draw the mesh
 			FMeshBatch& Mesh = RayTracingInstance.Materials.AddDefaulted_GetRef();
 			Mesh.bWireframe = bWireframe;
 			Mesh.SegmentIndex = SectionIndex;
diff --git a/Engine/Source/Runtime/Experimental/GeometryCollectionEngine/Private/GeometryCollection/GeometryCollectionSceneProxy.h b/Engine/Source/Runtime/Experimental/GeometryCollectionEngine/Private/GeometryCollection/GeometryCollectionSceneProxy.h
--- a/Engine/Source/Runtime/Experimental/GeometryCollectionEngine/Private/GeometryCollection/GeometryCollectionSceneProxy.h
+++ b/Engine/Source/Runtime/Experimental/GeometryCollectionEngine/Private/GeometryCollection/GeometryCollectionSceneProxy.h
@@ -173,7 +173,7 @@ protected:
 	}
 
 #if RHI_RAYTRACING
-	void UpdatingRayTracingGeometry_RenderingThread(TArray<FGeometryCollectionMeshElement> const& InSectionArray);
+	void UpdatingRayTracingGeometry_RenderingThread(FRHICommandListBase& RHICmdList, TConstArrayView<FGeometryCollectionMeshElement> InSectionArray);
 
 	bool IsRayTracingRelevant() const { return true; }
 	bool IsRayTracingStaticRelevant() const { return false; }
