Iris
* Added Net.Iris.PrintDebugInfoOfNetHandle to print alot of internal Iris information about a given replicated object and a connection

#rnx

diff --git a/Engine/Source/Runtime/Experimental/Iris/Core/Private/Iris/ReplicationSystem/DeltaCompression/DeltaCompressionBaselineManager.cpp b/Engine/Source/Runtime/Experimental/Iris/Core/Private/Iris/ReplicationSystem/DeltaCompression/DeltaCompressionBaselineManager.cpp
--- a/Engine/Source/Runtime/Experimental/Iris/Core/Private/Iris/ReplicationSystem/DeltaCompression/DeltaCompressionBaselineManager.cpp
+++ b/Engine/Source/Runtime/Experimental/Iris/Core/Private/Iris/ReplicationSystem/DeltaCompression/DeltaCompressionBaselineManager.cpp
@@ -727,6 +727,20 @@ void FDeltaCompressionBaselineManager::DestroyBaseline(uint32 ConnId, uint32 Obj
 	AdjustBaselineCount(ObjectInfo, ObjectInfoIndex, BaselineCountAdjustment);
 }
 
+FString FDeltaCompressionBaselineManager::PrintDeltaCompressionStatus(uint32 ConnectionId, FInternalNetRefIndex ObjectIndex) const
+{
+	if (const FPerObjectInfo* PerObjectInfo = GetPerObjectInfoForObject(ObjectIndex))
+	{
+		//TODO: Print per-connection info ?
+		//FObjectBaselineInfo& BaselineInfo = ObjectInfo->BaselinesForConnections[ConnId];
+		return FString::Printf(TEXT("PrevBaselineCreationFrame: %u | ChangeMaskStride: %u"), PerObjectInfo->PrevBaselineCreationFrame, PerObjectInfo->ChangeMaskStride);
+	}
+	else
+	{
+		return FString::Printf(TEXT("No delta compression info for %u"), ObjectIndex);
+	}
+}
+
 }
 
 #undef UE_NET_DC_BASELINE_CHECK
diff --git a/Engine/Source/Runtime/Experimental/Iris/Core/Private/Iris/ReplicationSystem/DeltaCompression/DeltaCompressionBaselineManager.h b/Engine/Source/Runtime/Experimental/Iris/Core/Private/Iris/ReplicationSystem/DeltaCompression/DeltaCompressionBaselineManager.h
--- a/Engine/Source/Runtime/Experimental/Iris/Core/Private/Iris/ReplicationSystem/DeltaCompression/DeltaCompressionBaselineManager.h
+++ b/Engine/Source/Runtime/Experimental/Iris/Core/Private/Iris/ReplicationSystem/DeltaCompression/DeltaCompressionBaselineManager.h
@@ -100,6 +100,8 @@ public:
 	 */
 	FDeltaCompressionBaseline GetBaseline(uint32 ConnId, uint32 ObjectIndex, uint32 BaselineIndex) const;
 
+	[[nodiscard]] FString PrintDeltaCompressionStatus(uint32 ConnectionId, FInternalNetRefIndex ObjectIndex) const;
+
 private:
 	using ObjectInfoIndexType = uint16;
 
diff --git a/Engine/Source/Runtime/Experimental/Iris/Core/Private/Iris/ReplicationSystem/Filtering/NetObjectGridFilter.cpp b/Engine/Source/Runtime/Experimental/Iris/Core/Private/Iris/ReplicationSystem/Filtering/NetObjectGridFilter.cpp
--- a/Engine/Source/Runtime/Experimental/Iris/Core/Private/Iris/ReplicationSystem/Filtering/NetObjectGridFilter.cpp
+++ b/Engine/Source/Runtime/Experimental/Iris/Core/Private/Iris/ReplicationSystem/Filtering/NetObjectGridFilter.cpp
@@ -525,11 +525,6 @@ FString UNetObjectGridFilter::PrintDebugInfoForObject(const FDebugInfoParams& Pa
 
 	const FObjectLocationInfo& ObjectLocationInfo = static_cast<const FObjectLocationInfo&>(Params.FilteringInfos[ObjectIndex]);
 
-	if (ObjectLocationInfo.GetInfoIndex() == 0)
-	{
-		return TEXT("NoGridFilter Info");
-	}
-
 	const FPerObjectInfo& PerObjectInfo = ObjectInfos[ObjectLocationInfo.GetInfoIndex()];
 
 	const FPerConnectionInfo* const PerConnectionInfo = PerConnectionInfos.IsValidIndex(Params.ConnectionId) ? &PerConnectionInfos[Params.ConnectionId] : nullptr;
@@ -545,7 +540,6 @@ FString UNetObjectGridFilter::PrintDebugInfoForObject(const FDebugInfoParams& Pa
 		DistZ = FMath::Min(FMath::Abs(PerObjectInfo.Position.Z - View.Pos.Z), DistZ);
 	}
 	
-
 	return FString::Printf(TEXT("[GridFilter] Dist2D: %.2f, CullDistance: %.2f, CulledFrameCount: %u, DistZ: %.2f, Pos: %s"),
 		Dist2d, PerObjectInfo.GetCullDistance(), CulledFrameCount, DistZ, *PerObjectInfo.Position.ToCompactString()
 	);
diff --git a/Engine/Source/Runtime/Experimental/Iris/Core/Private/Iris/ReplicationSystem/Filtering/ReplicationFiltering.cpp b/Engine/Source/Runtime/Experimental/Iris/Core/Private/Iris/ReplicationSystem/Filtering/ReplicationFiltering.cpp
--- a/Engine/Source/Runtime/Experimental/Iris/Core/Private/Iris/ReplicationSystem/Filtering/ReplicationFiltering.cpp
+++ b/Engine/Source/Runtime/Experimental/Iris/Core/Private/Iris/ReplicationSystem/Filtering/ReplicationFiltering.cpp
@@ -2744,7 +2744,7 @@ FString FReplicationFiltering::PrintFilterObjectInfo(FInternalNetRefIndex Object
 	DebugParams.ConnectionId = ConnectionId;
 	DebugParams.View = Connections->GetReplicationView(ConnectionId);
 
-	return FilterInfo.Filter->PrintDebugInfoForObject(DebugParams, ObjectIndex);
+	return FString::Printf(TEXT("FilterName: %s (%s) | %s"), *FilterInfo.Name.ToString(), *FilterInfo.Filter->GetClass()->GetName(), *FilterInfo.Filter->PrintDebugInfoForObject(DebugParams, ObjectIndex));
 }
 
 void FReplicationFiltering::BuildObjectsInFilterList(FNetBitArrayView OutObjectsInFilter, FName FilterName) const
diff --git a/Engine/Source/Runtime/Experimental/Iris/Core/Private/Iris/ReplicationSystem/ObjectReplicationBridgeDebugging.cpp b/Engine/Source/Runtime/Experimental/Iris/Core/Private/Iris/ReplicationSystem/ObjectReplicationBridgeDebugging.cpp
--- a/Engine/Source/Runtime/Experimental/Iris/Core/Private/Iris/ReplicationSystem/ObjectReplicationBridgeDebugging.cpp
+++ b/Engine/Source/Runtime/Experimental/Iris/Core/Private/Iris/ReplicationSystem/ObjectReplicationBridgeDebugging.cpp
@@ -11,12 +11,20 @@
 #include "Iris/ReplicationSystem/ReplicationSystemTypes.h"
 #include "Iris/ReplicationSystem/ReplicationSystemInternal.h"
 #include "Iris/ReplicationSystem/ReplicationOperations.h"
+#include "Iris/ReplicationSystem/ReplicationReader.h"
+#include "Iris/ReplicationSystem/ReplicationWriter.h"
+
+#include "Iris/ReplicationSystem/DeltaCompression/DeltaCompressionBaselineManager.h"
 
 #include "Iris/ReplicationSystem/Filtering/NetObjectFilterDefinitions.h"
 #include "Iris/ReplicationSystem/Filtering/ReplicationFiltering.h"
+
+#include "Iris/ReplicationSystem/Prioritization/ReplicationPrioritization.h"
+
 #include "Iris/ReplicationSystem/ObjectReplicationBridgeConfig.h"
 
 #include "Net/Core/NetBitArrayPrinter.h"
+#include "Net/Core/NetHandle/NetHandleManager.h"
 #include "Net/Core/Trace/NetDebugName.h"
 
 #include "Templates/UnrealTemplate.h"
@@ -36,6 +44,9 @@
 namespace UE::Net::Private::ObjectBridgeDebugging
 {
 
+/** 
+ * Generic traits to control what and how we log debug information
+ */
 enum class EPrintTraits : uint32
 {
 	Default					= 0x0000,
@@ -48,6 +59,7 @@ enum class EPrintTraits : uint32
 };
 ENUM_CLASS_FLAGS(EPrintTraits);
 
+/** Extracts the print traits from console arguments */
 EPrintTraits FindPrintTraitsFromArgs(const TArray<FString>& Args)
 {
 	EPrintTraits Traits = EPrintTraits::Default;
@@ -69,6 +81,9 @@ EPrintTraits FindPrintTraitsFromArgs(const TArray<FString>& Args)
 	return Traits;
 }
 
+/**
+ * Print traits to control the dynamic filter debug commands
+ */
 enum class EDynamicFilterPrintTraits : uint32
 {
 	Default							= 0x0000,
@@ -84,6 +99,7 @@ enum class EDynamicFilterPrintTraits : uint32
 };
 ENUM_CLASS_FLAGS(EDynamicFilterPrintTraits);
 
+/** Extract dynamic filter print traits from console arguments */
 EDynamicFilterPrintTraits FindDynamicFilterPrintTraitsFromArgs(const TArray<FString>& Args)
 {
 	EDynamicFilterPrintTraits Traits = EDynamicFilterPrintTraits::Default;
@@ -176,6 +192,7 @@ void SortViaTrait(TArray<FRootObjectData>& OutArray, EPrintTraits ArgTraits, EPr
 	}
 }
 
+/** Print all the protocols of the default state of an object (so it's CDO/Archetype baseline) */
 void PrintDefaultNetObjectState(UReplicationSystem* ReplicationSystem, uint32 ConnectionId, const FReplicationFragments& RegisteredFragments, FStringBuilderBase& StringBuilder)
 {
 	FReplicationSystemInternal* ReplicationSystemInternal = ReplicationSystem->GetReplicationSystemInternal();
@@ -212,7 +229,9 @@ void RemoteProtocolMismatchDetected(TMap<FObjectKey, bool>& ArchetypesAlreadyPri
 		}
 	}
 }
-
+/**
+ * Find the replication system from console arguments: 'RepSystemId=0...'
+ */
 UReplicationSystem* FindReplicationSystemFromArg(const TArray<FString>& Args)
 {
 	uint32 RepSystemId = 0;
@@ -294,6 +313,44 @@ void LogViaTrait(FLogContext& LogContext, EPrintTraits ArgTraits, EPrintTraits D
 	LogRootObjectList(LogContext, bLogSubObjects);	
 }
 
+/** 
+ * Returns a list of NetRefHandles of replicated objects from console arguments
+ */
+TArray<FNetRefHandle> FindNetRefHandlesFromArg(UReplicationSystem* RepSystem, const TArray<FString>& Args)
+{
+	//TODO: Add more ways to find objects FindByInternalIndex, FindByPtr, FindByName, etc.
+	TArray<FNetRefHandle> NetRefHandles;
+	 
+	const FNetRefHandleManager&  NetRefHandleManager = RepSystem->GetReplicationSystemInternal()->GetNetRefHandleManager();
+
+	if (const FString* FindByClass = Args.FindByPredicate([](const FString& Str) { return Str.Contains(TEXT("FindByClass=")); }))
+	{
+		FString ClassNameToFind;
+		FParse::Value(**FindByClass, TEXT("FindByClass="), ClassNameToFind);
+
+		for (const auto& MapIt : NetRefHandleManager.GetReplicatedHandles())
+		{
+			const FInternalNetRefIndex InternalIndex = NetRefHandleManager.GetInternalIndex(MapIt.Key);
+			UObject* ReplicatedObject = NetRefHandleManager.GetReplicatedObjectInstance(InternalIndex);
+			if (IsValid(ReplicatedObject))
+			{
+				for (const UClass* Class = ReplicatedObject->GetClass(); Class; Class = Class->GetSuperClass())
+				{
+					const FString ClassName = Class->GetPathName();
+					if (ClassName.Contains(ClassNameToFind))
+					{
+						// Found it
+						NetRefHandles.Emplace(MapIt.Key);
+						break;
+					}
+				}
+			}
+		}
+	}
+
+	return NetRefHandles;
+}
+
 } // end namespace UE::Net::Private::ObjectBridgeDebugging
 
 // --------------------------------------------------------------------------------------------------------------------------------------------
@@ -1112,3 +1169,203 @@ void UObjectReplicationBridge::PrintPushBasedStatuses() const
 	UE_LOG(LogIrisBridge, Display, TEXT(""));
 	UE_LOG(LogIrisBridge, Display, TEXT("################ Stop Printing Push-Based Statuses ################"));
 }
+
+//-----------------------------------------------
+FAutoConsoleCommand ObjectBridgePrintDebugInfoOfNetHandle(
+	TEXT("Net.Iris.PrintDebugInfoOfNetHandle"),
+	TEXT("Prints all relevant information about a NetHandle and a connection."),
+FConsoleCommandWithArgsDelegate::CreateLambda([](const TArray< FString >& Args)
+{
+	using namespace UE::Net;
+	using namespace UE::Net::Private;
+	using namespace UE::Net::Private::ObjectBridgeDebugging;
+
+	UReplicationSystem* RepSystem = FindReplicationSystemFromArg(Args);
+	if (!RepSystem)
+	{
+		UE_LOG(LogIrisBridge, Error, TEXT("Could not find ReplicationSystem."));
+		return;
+	}
+
+	UObjectReplicationBridge* ObjectBridge = RepSystem->GetReplicationBridgeAs<UObjectReplicationBridge>();
+	if (!ObjectBridge)
+	{
+		UE_LOG(LogIrisBridge, Error, TEXT("Could not find ObjectReplicationBridge."));
+		return;
+	}
+
+	const TArray<FNetRefHandle> NetHandles = FindNetRefHandlesFromArg(RepSystem, Args);
+
+	UObjectReplicationBridge::EPrintDebugInfoTraits PrintTraits = UObjectReplicationBridge::EPrintDebugInfoTraits::Default;
+	ObjectBridge->PrintDebugInfoForNetRefHandlesAndConnections(NetHandles, Args, PrintTraits);
+	
+}));
+
+void UObjectReplicationBridge::PrintDebugInfoForNetRefHandlesAndConnections(const TArray<FNetRefHandle>& NetHandles, const TArray<FString>& Args, EPrintDebugInfoTraits PrintTraits)
+{
+	using namespace UE::Net;
+	using namespace UE::Net::Private;
+	using namespace UE::Net::Private::ObjectBridgeDebugging;
+
+	FReplicationConnections& Connections = GetReplicationSystem()->GetReplicationSystemInternal()->GetConnections();
+	const FNetBitArray& ValidConnections = Connections.GetValidConnections();
+
+	// Default to all connections
+	FNetBitArray ConnectionsToPrint(ValidConnections);
+
+	// Filter down the list if specific connections were specified
+	{
+		const TArray<uint32> RequestedConnectionList = FindConnectionsFromArgs(Args);
+		if (!RequestedConnectionList.IsEmpty())
+		{
+			ConnectionsToPrint.ClearAllBits();
+
+			for (uint32 ConnectionId : RequestedConnectionList)
+			{
+				ConnectionsToPrint.SetBitValue(ConnectionId, ValidConnections.GetBit(ConnectionId));
+			}
+		}
+	}
+
+	EPrintDebugInfoTraits DebugPrintTraits = (EPrintDebugInfoTraits)PrintTraits;
+
+	ConnectionsToPrint.ForAllSetBits([&](uint32 ConnectionId)
+	{
+		for (const FNetRefHandle& NetHandle : NetHandles)
+		{
+			PrintDebugInfoForNetRefHandle(NetHandle, ConnectionId, DebugPrintTraits);
+		}
+
+		// Do not print the object state for the other connections since there is only one global state for all
+		EnumAddFlags(DebugPrintTraits, EPrintDebugInfoTraits::NoProtocolState);
+	});
+}
+
+
+void UObjectReplicationBridge::PrintDebugInfoForNetRefHandle(const FNetRefHandle NetRefHandle, uint32 ConnectionId, EPrintDebugInfoTraits PrintTraits) const
+{
+	using namespace UE::Net::Private;
+	using namespace UE::Net;
+	using namespace UE::Net::Private::ObjectBridgeDebugging;
+
+	FReplicationSystemInternal* ReplicationSystemInternal = GetReplicationSystem()->GetReplicationSystemInternal();
+
+	const FWorldLocations& WorldLocations = ReplicationSystemInternal->GetWorldLocations();
+	const FReplicationConnections& Connections = ReplicationSystemInternal->GetConnections();
+	const FReplicationPrioritization& Prioritization = ReplicationSystemInternal->GetPrioritization();
+	const FDeltaCompressionBaselineManager& DeltaManager = ReplicationSystemInternal->GetDeltaCompressionBaselineManager();
+	const FReplicationFiltering& Filtering = ReplicationSystemInternal->GetFiltering();
+
+	const FInternalNetRefIndex ObjectIndex = NetRefHandleManager->GetInternalIndex(NetRefHandle);
+
+	UE_LOG(LogIris, Display, TEXT("################ Start Printing debug info for: %s ################"), *NetRefHandleManager->PrintObjectFromNetRefHandle(NetRefHandle));
+	UE_LOG(LogIris, Display, TEXT(""));
+	UE_LOG(LogIris, Display, TEXT("For Connection: %s"), *PrintConnectionInfo(ConnectionId))
+	UE_LOG(LogIris, Display, TEXT(""));
+
+#if !UE_BUILD_SHIPPING
+	if (!EnumHasAnyFlags(PrintTraits, EPrintDebugInfoTraits::NoProtocolState))
+	{
+		TStringBuilder<4096> StringBuilder;
+		IrisDebugHelper::NetObjectStateToString(StringBuilder, NetRefHandle);
+
+		// Protocol state
+		UE_LOG(LogIris, Display, TEXT("\tNetObjectState: %s"), *StringBuilder);
+	}
+#endif
+
+	if (GetReplicationSystem()->IsServer())
+	{
+		// World location
+		if (WorldLocations.HasInfoForObject(ObjectIndex))
+		{
+			const FWorldLocations::FWorldInfo WorldInfo = WorldLocations.GetWorldInfo(ObjectIndex);
+			UE_LOG(LogIris, Display, TEXT("\tWorldLocation: %s | CullDistance: %f"), *WorldInfo.WorldLocation.ToCompactString(), WorldInfo.CullDistance);
+		}
+		else
+		{
+			UE_LOG(LogIris, Display, TEXT("\tWorldLocation: None"));
+		}
+
+		// Filter Info
+		UE_LOG(LogIris, Display, TEXT("\tFilterInfo: %s"), *Filtering.PrintFilterObjectInfo(ObjectIndex, ConnectionId));
+
+		// Connection info
+		if (const FReplicationConnection* Connection = Connections.GetConnection(ConnectionId))
+		{
+			if (Filtering.GetRelevantObjectsInScope(ConnectionId).IsBitSet(ObjectIndex))
+			{
+				if (const FReplicationWriter* ReplicationWriter = Connection->ReplicationWriter)
+				{
+					// Replication Writer Info
+					UE_LOG(LogIris, Display, TEXT("\tReplicationWriter: %s"), *ReplicationWriter->PrintObjectInfo(ObjectIndex));
+				}
+				else
+				{
+					UE_LOG(LogIris, Display, TEXT("\tReplicationWriter: not found ?"));
+				}
+
+				// Prioritizer
+				UE_LOG(LogIris, Display, TEXT("\tPrioritizer: Priority: %f"), Prioritization.GetObjectPriorityForConnection(ConnectionId, ObjectIndex));
+
+				// DeltaCompression info
+				if (DeltaManager.GetDeltaCompressionStatus(ObjectIndex) == ENetObjectDeltaCompressionStatus::Allow)
+				{
+					UE_LOG(LogIris, Display, TEXT("\tDeltaCompression: %s"), *DeltaManager.PrintDeltaCompressionStatus(ConnectionId, ObjectIndex));
+				}
+				else
+				{
+					UE_LOG(LogIris, Display, TEXT("\tDeltaCompression: not enabled"));
+				}
+			}
+			else
+			{
+				UE_LOG(LogIris, Display, TEXT("\tNot relevant to connection %d"), ConnectionId);
+			}
+		}
+		else
+		{
+			UE_LOG(LogIris, Display, TEXT("\tNo valid connection with id: %u"), ConnectionId);
+		}
+
+		// Dependents
+		{
+			UE_LOG(LogIris, Display, TEXT("\tChildDependents: %d objects"), NetRefHandleManager->GetDependentObjectInfos(ObjectIndex).Num());
+			UE_LOG(LogIris, Display, TEXT("\tParentDependents: %d objects"), NetRefHandleManager->GetDependentObjectParents(ObjectIndex).Num());
+		}
+	}
+	else
+	{
+		if (const FReplicationConnection* Connection = Connections.GetConnection(ConnectionId))
+		{
+			if (const FReplicationReader* ReplicationReader = Connection->ReplicationReader)
+			{
+				UE_LOG(LogIris, Display, TEXT("\tReplicationReader: %s"), *ReplicationReader->PrintObjectInfo(ObjectIndex, NetRefHandle));
+			}
+			else
+			{
+				UE_LOG(LogIris, Display, TEXT("\tReplicationReader not found ?"));
+			}
+		}
+		else
+		{
+			UE_LOG(LogIris, Display, TEXT("\tNo valid connection with id: %u"), ConnectionId);
+		}
+	}
+	
+	UE_LOG(LogIris, Display, TEXT(""));
+	UE_LOG(LogIris, Display, TEXT("################ Stopped Printing debug info for: %s ################"), *NetRefHandleManager->PrintObjectFromNetRefHandle(NetRefHandle));
+
+	// Log dependents if requested
+	if (GetReplicationSystem()->IsServer() && !EnumHasAnyFlags(PrintTraits, EPrintDebugInfoTraits::NoDependents))
+	{
+		const TArrayView<const FDependentObjectInfo> ChildDependents = NetRefHandleManager->GetDependentObjectInfos(ObjectIndex);
+		for (const FDependentObjectInfo& DependentInfo : ChildDependents)
+		{
+			UE_LOG(LogIris, Display, TEXT(""));
+			UE_LOG(LogIris, Display, TEXT("################ Start printing dependent %s of parent %s ################"), *NetRefHandleManager->PrintObjectFromIndex(DependentInfo.NetRefIndex), *NetRefHandleManager->PrintObjectFromIndex(ObjectIndex));
+			UE_LOG(LogIris, Display, TEXT(""));
+			PrintDebugInfoForNetRefHandle(NetRefHandleManager->GetNetRefHandleFromInternalIndex(DependentInfo.NetRefIndex), ConnectionId, PrintTraits);
+		}
+	}
+}
\ No newline at end of file
diff --git a/Engine/Source/Runtime/Experimental/Iris/Core/Private/Iris/ReplicationSystem/Prioritization/ReplicationPrioritization.h b/Engine/Source/Runtime/Experimental/Iris/Core/Private/Iris/ReplicationSystem/Prioritization/ReplicationPrioritization.h
--- a/Engine/Source/Runtime/Experimental/Iris/Core/Private/Iris/ReplicationSystem/Prioritization/ReplicationPrioritization.h
+++ b/Engine/Source/Runtime/Experimental/Iris/Core/Private/Iris/ReplicationSystem/Prioritization/ReplicationPrioritization.h
@@ -58,6 +58,11 @@ public:
 	void AddConnection(uint32 ConnectionId);
 	void RemoveConnection(uint32 ConnectionId);
 
+	float GetObjectPriorityForConnection(uint32 ConnectionId, FInternalNetRefIndex InternalIndex) const
+	{
+		return GetPrioritiesForConnection(ConnectionId)[InternalIndex];
+	}
+
 private:
 	class FPrioritizerBatchHelper;
 	class FUpdateDirtyObjectsBatchHelper;
diff --git a/Engine/Source/Runtime/Experimental/Iris/Core/Private/Iris/ReplicationSystem/ReplicationReader.cpp b/Engine/Source/Runtime/Experimental/Iris/Core/Private/Iris/ReplicationSystem/ReplicationReader.cpp
--- a/Engine/Source/Runtime/Experimental/Iris/Core/Private/Iris/ReplicationSystem/ReplicationReader.cpp
+++ b/Engine/Source/Runtime/Experimental/Iris/Core/Private/Iris/ReplicationSystem/ReplicationReader.cpp
@@ -591,6 +591,11 @@ FReplicationReader::FReplicatedObjectInfo* FReplicationReader::GetReplicatedObje
 	return ReplicatedObjects.Find(InternalIndex);
 }
 
+const FReplicationReader::FReplicatedObjectInfo* FReplicationReader::GetReplicatedObjectInfo(uint32 InternalIndex) const
+{
+	return ReplicatedObjects.Find(InternalIndex);
+}
+
 void FReplicationReader::CleanupObjectData(FReplicatedObjectInfo& ObjectInfo)
 {
 	// Remove from pending resolve
@@ -2960,4 +2965,36 @@ bool FReplicationReader::ReadSentinel(FNetSerializationContext& Context, const T
 	return bSuccess;
 }
 
+FString FReplicationReader::PrintObjectInfo(FInternalNetRefIndex ObjectIndex, FNetRefHandle NetRefHandle) const
+{
+	const FReplicatedObjectInfo* ObjectInfo = GetReplicatedObjectInfo(ObjectIndex);
+	if (!ObjectInfo)
+	{
+		return FString::Printf(TEXT("No object info for (InternalIndex: %u)"), ObjectIndex);
+	}
+
+	TStringBuilder<512> InfoBuilder;
+
+	const FNetRefHandleManager::FReplicatedObjectData& ObjectData = NetRefHandleManager->GetReplicatedObjectDataNoCheck(ObjectIndex);
+	const bool bIsSubObject = ObjectData.SubObjectRootIndex != FNetRefHandleManager::InvalidInternalIndex;
+	InfoBuilder.Appendf(TEXT("Status info: 0x%llx (%s)"), ObjectInfo->Value, bIsSubObject?TEXT("SubObject"):TEXT("RootObject"));
+
+	const bool bHasPendingBatches = PendingBatches.GetHasPendingBatches();
+	const FInternalNetRefIndex RootInternalIndex = bIsSubObject ? ObjectData.SubObjectRootIndex : ObjectIndex;
+	if (const FPendingBatchData* PendingBatchData = bHasPendingBatches ? PendingBatches.Find(NetRefHandleManager->GetNetRefHandleFromInternalIndex(RootInternalIndex)) : nullptr)
+	{
+		InfoBuilder.Appendf(TEXT("| PendingBatches: QueuedChunks: %d | PendingReferences(%d): "), PendingBatchData->QueuedDataChunks.Num(), PendingBatchData->PendingMustBeMappedReferences.Num());
+		for (const FNetRefHandle& PendingHandle : PendingBatchData->PendingMustBeMappedReferences)
+		{
+			InfoBuilder << PendingHandle.ToString() << ", ";
+		}
+	}
+	else
+	{
+		InfoBuilder << TEXT("| NoPendingBatches");
+	}
+
+	return InfoBuilder.ToString();
+}
+
 } // end namespace UE::Net::Private
diff --git a/Engine/Source/Runtime/Experimental/Iris/Core/Private/Iris/ReplicationSystem/ReplicationReader.h b/Engine/Source/Runtime/Experimental/Iris/Core/Private/Iris/ReplicationSystem/ReplicationReader.h
--- a/Engine/Source/Runtime/Experimental/Iris/Core/Private/Iris/ReplicationSystem/ReplicationReader.h
+++ b/Engine/Source/Runtime/Experimental/Iris/Core/Private/Iris/ReplicationSystem/ReplicationReader.h
@@ -63,6 +63,8 @@ public:
 	// Process queued batches
 	void ProcessQueuedBatches();
 
+	[[nodiscard]] FString PrintObjectInfo(FInternalNetRefIndex ObjectIndex, FNetRefHandle NetRefHandle) const;
+
 private:
 	// ChangeMaskOffset -> FNetRefHandle
 	enum EConstants : uint32
@@ -190,6 +192,7 @@ private:
 
 	// Lookup the tracking info for the object with IntnernalIndex
 	FReplicatedObjectInfo* GetReplicatedObjectInfo(uint32 InternalIndex);
+	const FReplicatedObjectInfo* GetReplicatedObjectInfo(uint32 InternalIndex) const;
 
 	// Update reference tracking maps for the current object
 	void UpdateObjectReferenceTracking(FReplicatedObjectInfo* ReplicationInfo, FNetBitArrayView ChangeMask, bool bIncludeInitState, FResolvedNetRefHandlesArray& OutNewResolvedRefHandles, const FObjectReferenceTracker& NewUnresolvedReferences, const FObjectReferenceTracker& NewMappedDynamicReferences);
diff --git a/Engine/Source/Runtime/Experimental/Iris/Core/Private/Iris/ReplicationSystem/ReplicationWriter.cpp b/Engine/Source/Runtime/Experimental/Iris/Core/Private/Iris/ReplicationSystem/ReplicationWriter.cpp
--- a/Engine/Source/Runtime/Experimental/Iris/Core/Private/Iris/ReplicationSystem/ReplicationWriter.cpp
+++ b/Engine/Source/Runtime/Experimental/Iris/Core/Private/Iris/ReplicationSystem/ReplicationWriter.cpp
@@ -4037,6 +4037,56 @@ void FReplicationWriter::WriteSentinel(FNetBitStreamWriter* Writer, const TCHAR*
 #endif
 }
 
+FString FReplicationWriter::PrintObjectInfo(FInternalNetRefIndex ObjectIndex) const
+{
+	const FReplicationInfo& ObjectInfo = ReplicatedObjects[ObjectIndex];
+
+	TStringBuilder<512> InfoBuilder;
+
+	InfoBuilder.Appendf(TEXT("Status info: 0x%llx (%s) | %s | SchedulingPriority: %f | LastAckedBaselineIndex: %u | Flags"), 
+		ObjectInfo.Value, 
+		ObjectInfo.IsSubObject?TEXT("SubObject"):TEXT("RootObject"),
+		LexToString(ObjectInfo.GetState()), 
+		SchedulingPriorities[ObjectIndex],
+		ObjectInfo.LastAckedBaselineIndex);
+
+#define FREPLICATIONWRITER_APPENDFLAG(FlagName) \
+	if (ObjectInfo.FlagName) \
+	{ \
+		InfoBuilder << TEXT(" | " #FlagName); \
+	}
+	// Relevant flags
+	FREPLICATIONWRITER_APPENDFLAG(HasDirtyChangeMask);
+	FREPLICATIONWRITER_APPENDFLAG(HasDirtySubObjects);
+	FREPLICATIONWRITER_APPENDFLAG(HasAttachments);
+	FREPLICATIONWRITER_APPENDFLAG(IsCreationConfirmed);
+	FREPLICATIONWRITER_APPENDFLAG(TearOff);
+	FREPLICATIONWRITER_APPENDFLAG(SubObjectPendingDestroy);
+	FREPLICATIONWRITER_APPENDFLAG(IsDeltaCompressionEnabled);
+	FREPLICATIONWRITER_APPENDFLAG(HasDirtyConditionals);
+
+#undef FREPLICATIONWRITER_APPENDFLAG
+
+	if (ObjectsWithDirtyChanges.GetBit(ObjectIndex))
+	{
+		InfoBuilder << TEXT(" | IsInDirtyChanges");
+	}
+	if (ObjectsPendingDestroy.GetBit(ObjectIndex))
+	{
+		InfoBuilder << TEXT(" | IsInPendingDestroy");
+	}
+	if (ObjectsInScope.GetBit(ObjectIndex))
+	{
+		InfoBuilder << TEXT(" | IsInScope");
+	}
+	if (IsActiveHugeObject(ObjectIndex))
+	{
+		InfoBuilder << TEXT(" | IsActiveHugeObject");
+	}
+
+	return InfoBuilder.ToString();
+}
+
 FReplicationWriter::FHugeObjectContext::FHugeObjectContext() = default;
 
 FReplicationWriter::FHugeObjectContext::~FHugeObjectContext() = default;
diff --git a/Engine/Source/Runtime/Experimental/Iris/Core/Private/Iris/ReplicationSystem/ReplicationWriter.h b/Engine/Source/Runtime/Experimental/Iris/Core/Private/Iris/ReplicationSystem/ReplicationWriter.h
--- a/Engine/Source/Runtime/Experimental/Iris/Core/Private/Iris/ReplicationSystem/ReplicationWriter.h
+++ b/Engine/Source/Runtime/Experimental/Iris/Core/Private/Iris/ReplicationSystem/ReplicationWriter.h
@@ -195,6 +195,8 @@ public:
 
 	bool AreAllReliableAttachmentsSentAndAcked() const;
 
+	[[nodiscard]] FString PrintObjectInfo(FInternalNetRefIndex ObjectIndex) const;
+
 private:
 	// Various types
 
@@ -582,7 +584,6 @@ private:
 	// Writes a sentinel naking it easier for the receiver to detect bitstream errors.
 	void WriteSentinel(FNetBitStreamWriter* Writer, const TCHAR* DebugName);
 
-
 private:
 	// Replication parameters
 	FReplicationParameters Parameters;
diff --git a/Engine/Source/Runtime/Experimental/Iris/Core/Public/Iris/ReplicationSystem/ObjectReplicationBridge.h b/Engine/Source/Runtime/Experimental/Iris/Core/Public/Iris/ReplicationSystem/ObjectReplicationBridge.h
--- a/Engine/Source/Runtime/Experimental/Iris/Core/Public/Iris/ReplicationSystem/ObjectReplicationBridge.h
+++ b/Engine/Source/Runtime/Experimental/Iris/Core/Public/Iris/ReplicationSystem/ObjectReplicationBridge.h
@@ -3,15 +3,21 @@
 #pragma once
 
 #include "HAL/Platform.h"
-#include "Net/Core/NetBitArray.h"
-#include "Net/Core/NetHandle/NetHandle.h"
+
 #include "Iris/ReplicationSystem/ReplicationBridge.h"
-#include "Delegates/IDelegateInstance.h"
 #include "Iris/ReplicationSystem/ReplicationSystemTypes.h"
 #include "Iris/ReplicationSystem/NetObjectFactory.h"
 #include "Iris/ReplicationSystem/NetObjectFactoryRegistry.h"
+
 #include "Containers/ArrayView.h"
 
+#include "Delegates/IDelegateInstance.h"
+
+#include "Misc/EnumClassFlags.h"
+
+#include "Net/Core/NetBitArray.h"
+#include "Net/Core/NetHandle/NetHandle.h"
+
 #include "ObjectReplicationBridge.generated.h"
 
 DECLARE_LOG_CATEGORY_EXTERN(LogIrisFilterConfig, Log, All);
@@ -246,6 +252,16 @@ public:
 
 public:
 
+	/** Print traits to control what gets logged in PrintIrisDebugInfoForNetHandle */
+	enum class EPrintDebugInfoTraits : uint32
+	{
+		Default = 0x0000,
+		// Do not print the entire protocol state when this is set 
+		NoProtocolState = 0x0001,
+		// Do not automatically print the debug information of child dependents
+		NoDependents = 0x0002,
+	};
+
 	// Debug functions that are triggered via console commands declared in ObjectReplicationBridgeDebugging.cpp
 	void PrintDynamicFilterClassConfig(uint32 ArgTraits);
 	void PrintReplicatedObjects(uint32 ArgTraits) const;
@@ -255,6 +271,9 @@ public:
 	void PrintNetCullDistances(const TArray<FString>& Args) const;
 	void PrintPushBasedStatuses() const;
 
+	void PrintDebugInfoForNetRefHandlesAndConnections(const TArray<FNetRefHandle>& NetHandles, const TArray<FString>& Args, EPrintDebugInfoTraits PrintTraits);
+	IRISCORE_API void PrintDebugInfoForNetRefHandle(const FNetRefHandle NetHandle, uint32 ConnectionId, EPrintDebugInfoTraits PrintTraits=EPrintDebugInfoTraits::Default) const;
+
 protected:
 	IRISCORE_API virtual ~UObjectReplicationBridge();
 
@@ -511,3 +530,5 @@ protected:
 private:
 	TMap<FObjectKey, bool> ArchetypesAlreadyPrinted;
 };
+
+ENUM_CLASS_FLAGS(UObjectReplicationBridge::EPrintDebugInfoTraits);
\ No newline at end of file
