[ChaosVD] Fix for an issue in CVD where an ensure is triggered and collision data could be missing, caused by CL 40753150 submitted this morning.

The issue is that now we do these collision detection tasks in a worker thread using the tasks API where we implemnt the work as a lambda instead of a task object. This means we need to copy the current CVD context to the lamda so it can be pushed in the corresponding worked thread. (This was done by the previous task object that was removed).

Note: Submitting pre review because the change is trivial, and it is affecting multiple CVD workflows. Review will be done post commit

diff --git a/Engine/Source/Runtime/Experimental/Chaos/Public/Chaos/Collision/SpatialAccelerationBroadPhase.h b/Engine/Source/Runtime/Experimental/Chaos/Public/Chaos/Collision/SpatialAccelerationBroadPhase.h
--- a/Engine/Source/Runtime/Experimental/Chaos/Public/Chaos/Collision/SpatialAccelerationBroadPhase.h
+++ b/Engine/Source/Runtime/Experimental/Chaos/Public/Chaos/Collision/SpatialAccelerationBroadPhase.h
@@ -378,8 +378,10 @@ namespace Chaos
 				{
 					Private::FBroadPhaseContext& BroadPhaseContext = BroadphaseContexts[ContextIndex];
 
-					UE::Tasks::FTask PendingTask = UE::Tasks::Launch(UE_SOURCE_LOCATION, [this, &BroadPhaseContext](){
-						AssignMidPhases(BroadPhaseContext);
+					UE::Tasks::FTask PendingTask = UE::Tasks::Launch(UE_SOURCE_LOCATION, [this, &BroadPhaseContext, CVDContext]()
+						{
+							CVD_SCOPE_CONTEXT(CVDContext.Context);
+							AssignMidPhases(BroadPhaseContext);
 						}, PendingTasks[ContextIndex]);
 					NextPendingTasks.Add(PendingTask);
 				}
@@ -424,8 +426,10 @@ namespace Chaos
 				// Accessing the overlaps array is safe here, all overlaps have been finished to be computed at this point
 				if (!BroadphaseContext.Overlaps.IsEmpty())
 				{
-					UE::Tasks::FTask ProcessMidPhaseTask = UE::Tasks::Launch(UE_SOURCE_LOCATION, [this, &BroadphaseContext, Dt]() {
-						ProcessMidPhases(Dt, BroadphaseContext);
+					UE::Tasks::FTask ProcessMidPhaseTask = UE::Tasks::Launch(UE_SOURCE_LOCATION, [this, &BroadphaseContext, Dt, CVDContext]()
+						{
+							CVD_SCOPE_CONTEXT(CVDContext.Context);
+							ProcessMidPhases(Dt, BroadphaseContext);
 						}, UE::Tasks::ETaskPriority::High);
 					PendingTasks.Add(ProcessMidPhaseTask);
 					ContextsWithConstraints.Add(ContextIndex);
