Complete I/O store install request when cancel is being called before the request has been started.

#rnx

diff --git a/Engine/Source/Runtime/Experimental/IoStore/OnDemand/Private/OnDemandIoStore.cpp b/Engine/Source/Runtime/Experimental/IoStore/OnDemand/Private/OnDemandIoStore.cpp
--- a/Engine/Source/Runtime/Experimental/IoStore/OnDemand/Private/OnDemandIoStore.cpp
+++ b/Engine/Source/Runtime/Experimental/IoStore/OnDemand/Private/OnDemandIoStore.cpp
@@ -1031,14 +1031,18 @@ bool FOnDemandIoStore::Tick()
 	// Tick install request(s)
 	if (InstallRequest.IsValid())
 	{
-		FOnDemandInstallResult Result = TickInstallRequest(*InstallRequest);
-
-		// If the error code is not unknown the the request has already been completed when cancelled 
+		// Process the install request if the error code is unknown otherwise the request has been cancelled and completed
 		EIoErrorCode Expected = EIoErrorCode::Unknown;
-		if (InstallRequest->ErrorCode.compare_exchange_strong(Expected, Result.Status.GetErrorCode()))
+		if (InstallRequest->ErrorCode.compare_exchange_strong(Expected, EIoErrorCode::Ok))
 		{
+			FOnDemandInstallResult Result = TickInstallRequest(*InstallRequest);
 			CompleteInstallRequest(InstallRequest, MoveTemp(Result));
 		}
+		else
+		{
+			UE_LOG(LogIoStoreOnDemand, Log, TEXT("Skipping cancelled install request, %s"), *LexToString(*InstallRequest));
+		}
+
 		bTicked = true;
 	}
 
@@ -1995,14 +1999,18 @@ void FOnDemandIoStore::CancelInstallRequest(FSharedInternalInstallRequest Instal
 {
 	check(InstallRequest.IsValid());
 
+	// Cancel the request if the error code is unknown otherwise the request will be completed in the tick loop 
 	EIoErrorCode Expected = EIoErrorCode::Unknown;
-	if (InstallRequest->ErrorCode.compare_exchange_strong(Expected, EIoErrorCode::Cancelled) == false) 
+	if (InstallRequest->ErrorCode.compare_exchange_strong(Expected, EIoErrorCode::Cancelled))
 	{
-		// If the error code is not unknown the request will be completed in the tick loop 
-		return;
+		UE_LOG(LogIoStoreOnDemand, Log, TEXT("Cancelling install request, %s"), *LexToString(*InstallRequest));
+		FOnDemandInstallResult InstallResult
+		{
+			.Status				= FIoStatus(EIoErrorCode::Cancelled),
+			.DurationInSeconds 	= FPlatformTime::Seconds() - InstallRequest->StartTime
+		};
+		CompleteInstallRequest(InstallRequest, MoveTemp(InstallResult));
 	}
-
-	UE_LOG(LogIoStoreOnDemand, Log, TEXT("Cancelling install request, %s"), *LexToString(*InstallRequest));
 }
 
 void FOnDemandIoStore::UpdateInstallRequestPriority(FSharedInternalInstallRequest InstallRequest, int32 NewPriority)
