[HWRT] Support using one BLAS per Geometry Collection part instead of a combined BLAS when not using WPO.
- Avoids having to update a combined BLAS whenever part transforms change.
- Can be controlled using r.RayTracing.Geometry.GeometryCollection.CombinedBLAS
- Also removed some incorrect assumptions in RHI wrt to Initializer.TotalPrimitiveCount.

diff --git a/Engine/Source/Runtime/D3D12RHI/Private/D3D12RayTracing.cpp b/Engine/Source/Runtime/D3D12RHI/Private/D3D12RayTracing.cpp
--- a/Engine/Source/Runtime/D3D12RHI/Private/D3D12RayTracing.cpp
+++ b/Engine/Source/Runtime/D3D12RHI/Private/D3D12RayTracing.cpp
@@ -3294,8 +3294,6 @@ void TranslateRayTracingGeometryDescs(const FRayTracingGeometryInitializer& Init
 			checkf(Segment.VertexBufferStride >= (2 * sizeof(FVector3f)), TEXT("Procedural geometry vertex buffer must contain at least 2xFloat3 that defines 3D bounding boxes of primitives."));
 		}
 
-		check(Segment.FirstPrimitive + Segment.NumPrimitives <= Initializer.TotalPrimitiveCount);
-
 		if (Initializer.IndexBuffer)
 		{
 			uint32 IndexStride = Initializer.IndexBuffer->GetStride();
@@ -3869,14 +3867,14 @@ void FD3D12RayTracingGeometry::AllocateBufferSRVs(uint32 InGPUIndex)
 
 		FD3D12Buffer* IndexBuffer = FD3D12DynamicRHI::ResourceCast(Initializer.IndexBuffer.GetReference());
 
-		// Initializer.TotalPrimitiveCount currently just accumulated the num primitives of the segments
-		// but can be invalid if multiple segments overlap/share the same index buffer range
+		// Initializer.TotalPrimitiveCount is the accumulated num primitives of the segments
+		// The highest indexed entry can be higher due to Segment.FirstPrimitive or it can be lower if segments overlap
+		// So here we calculate the highest indexed entry by looping over the segments
 		uint32 MaxPrimitiveCount = 0;
 		for (const FRayTracingGeometrySegment& Segment : Initializer.Segments)
 		{
 			MaxPrimitiveCount = FMath::Max(MaxPrimitiveCount, Segment.FirstPrimitive + Segment.NumPrimitives);
 		}
-		ensure(MaxPrimitiveCount <= Initializer.TotalPrimitiveCount);
 
 		D3D12_SHADER_RESOURCE_VIEW_DESC SRVDesc = {};
 		SRVDesc.ViewDimension = D3D12_SRV_DIMENSION_BUFFER;
diff --git a/Engine/Source/Runtime/Experimental/GeometryCollectionEngine/Private/GeometryCollection/GeometryCollectionSceneProxy.cpp b/Engine/Source/Runtime/Experimental/GeometryCollectionEngine/Private/GeometryCollection/GeometryCollectionSceneProxy.cpp
--- a/Engine/Source/Runtime/Experimental/GeometryCollectionEngine/Private/GeometryCollection/GeometryCollectionSceneProxy.cpp
+++ b/Engine/Source/Runtime/Experimental/GeometryCollectionEngine/Private/GeometryCollection/GeometryCollectionSceneProxy.cpp
@@ -81,7 +81,18 @@ FAutoConsoleVariableRef CVarRayTracingGeometryCollectionProxyMeshes(
 static TAutoConsoleVariable<int32> CVarRayTracingGeometryCollectionForceUpdate(
 	TEXT("r.RayTracing.Geometry.GeometryCollection.ForceUpdate"),
 	0,
-	TEXT("Forces ray tracing representation for geometry collections meshes to be updated every frame"));
+	TEXT("Forces ray tracing representation for geometry collections meshes to be updated every frame."),
+	ECVF_RenderThreadSafe
+);
+
+static TAutoConsoleVariable<int32> CVarRayTracingGeometryCollectionCombinedBLAS(
+	TEXT("r.RayTracing.Geometry.GeometryCollection.CombinedBLAS"),
+	0,
+	TEXT("Whether to always use a combined BLAS instead of one instance per collection part.\n")
+	TEXT("A combined BLAS needs to be fully rebuilt whenever any transform changes.\n")
+	TEXT("This is automatically enabled for geometry collections using WPO since BLAS must be updated anyway."),
+	ECVF_RenderThreadSafe
+);
 
 
 #if !defined(CHAOS_GEOMETRY_COLLECTION_SET_DYNAMIC_DATA_ISPC_ENABLED_DEFAULT)
@@ -211,23 +222,6 @@ void FGeometryCollectionSceneProxyBase::CreateRenderThreadResources(FRHICommandL
 		SetupVertexFactory(RHICmdList, VertexFactory);
 	}
 
-#if RHI_RAYTRACING
-	if (bHasRayTracingRepresentation)
-	{
-		FRayTracingGeometryInitializer Initializer;
-		Initializer.DebugName = FName(TEXT("GeometryCollection"));
-		Initializer.GeometryType = RTGT_Triangles;
-		Initializer.bFastBuild = true;
-		Initializer.bAllowUpdate = false;
-		Initializer.TotalPrimitiveCount = 0;
-
-		RayTracingGeometry.SetInitializer(Initializer);
-		RayTracingGeometry.InitResource(RHICmdList);
-
-		bRayTracingGeometryDirty = true;
-	}
-#endif
-
 	bRenderResourcesCreated = true;
 }
 
@@ -256,6 +250,12 @@ void FGeometryCollectionSceneProxyBase::DestroyRenderThreadResources()
 #if RHI_RAYTRACING
 	if (bHasRayTracingRepresentation)
 	{
+		for (FRayTracingGeometry& PartRayTracingGeometry : PartRayTracingGeometries)
+		{
+			PartRayTracingGeometry.ReleaseResource();
+		}
+		PartRayTracingGeometries.Empty();
+
 		RayTracingGeometry.ReleaseResource();
 		RayTracingDynamicVertexBuffer.Release();
 	}
@@ -364,9 +364,9 @@ void FGeometryCollectionSceneProxyBase::SetDynamicData_RenderThread(FRHICommandL
 	}
 
 #if RHI_RAYTRACING
-	if (bHasRayTracingRepresentation)
+	if (RayTracingGeometry.IsInitialized())
 	{
-		bRayTracingGeometryDirty = true;
+		RayTracingGeometry.SetRequiresBuild(true);
 	}
 #endif
 }
@@ -445,33 +445,95 @@ void FGeometryCollectionSceneProxyBase::UpdateSkinnedPositions(FRHICommandListBa
 }
 
 #if RHI_RAYTRACING
-void FGeometryCollectionSceneProxyBase::UpdatingRayTracingGeometry_RenderingThread(FRHICommandListBase& RHICmdList, TConstArrayView<FGeometryCollectionMeshElement> InSectionArray)
+void FGeometryCollectionSceneProxyBase::UpdatingRayTracingGeometry_RenderingThread(FRHICommandListBase& RHICmdList, TConstArrayView<FGeometryCollectionMeshElement> InSectionArray, bool bRayTracingGeometryPerSection)
 {
-	if (bRayTracingGeometryDirty)
+	// TODO: Could use SectionsNoInternal when geometry collection is undamaged?
+
+	if (bRayTracingGeometryPerSection)
 	{
-		RayTracingGeometry.Initializer.Segments.Empty();
-		RayTracingGeometry.Initializer.TotalPrimitiveCount = 0;
+		// release combined geometry since will use part geometries
+		RayTracingGeometry.ReleaseResource();
+
+		// TODO: Combine sections using the same transform into a single BLAS
 
-		for (int SectionIndex = 0; SectionIndex < InSectionArray.Num(); ++SectionIndex)
+		if (PartRayTracingGeometries.IsEmpty())
 		{
-			const FGeometryCollectionMeshElement& Section = InSectionArray[SectionIndex];
-			FRayTracingGeometrySegment Segment;
-			Segment.FirstPrimitive = Section.TriangleStart;
-			Segment.VertexBuffer = MeshResource.PositionVertexBuffer.VertexBufferRHI;
-			Segment.NumPrimitives = Section.TriangleCount;
-			Segment.MaxVertices = Section.VertexEnd;
-			RayTracingGeometry.Initializer.Segments.Add(Segment);
-			RayTracingGeometry.Initializer.TotalPrimitiveCount += Section.TriangleCount;
+			PartRayTracingGeometries.SetNum(InSectionArray.Num());
+
+			for (int32 SectionIndex = 0; SectionIndex < InSectionArray.Num(); ++SectionIndex)
+			{
+				FRayTracingGeometryInitializer Initializer;
+				Initializer.DebugName = FDebugName(FName(TEXT("GeometryCollectionPart")), SectionIndex);
+				Initializer.GeometryType = RTGT_Triangles;
+				Initializer.bFastBuild = true;
+				Initializer.bAllowUpdate = false;
+				Initializer.TotalPrimitiveCount = 0;
+				Initializer.IndexBuffer = MeshResource.IndexBuffer.IndexBufferRHI;
+
+				{
+					const FGeometryCollectionMeshElement& Section = InSectionArray[SectionIndex];
+
+					FRayTracingGeometrySegment Segment;
+					Segment.FirstPrimitive = Section.TriangleStart;
+					Segment.VertexBuffer = MeshResource.PositionVertexBuffer.VertexBufferRHI;
+					Segment.NumPrimitives = Section.TriangleCount;
+					Segment.MaxVertices = Section.VertexEnd;
+
+					Initializer.Segments.Add(Segment);
+
+					Initializer.TotalPrimitiveCount += Section.TriangleCount;
+				}
+
+				PartRayTracingGeometries[SectionIndex].SetInitializer(MoveTemp(Initializer));
+				PartRayTracingGeometries[SectionIndex].InitResource(RHICmdList);
+			}
+		}
+	}
+	else
+	{
+		// release part geometries since will use combined geometry
+		for (FRayTracingGeometry& PartRayTracingGeometry : PartRayTracingGeometries)
+		{
+			PartRayTracingGeometry.ReleaseResource();
 		}
+		PartRayTracingGeometries.Empty();
 
-		if (RayTracingGeometry.Initializer.TotalPrimitiveCount > 0)
+		// initialize combined geometry if necessary
+		if (!RayTracingGeometry.IsInitialized())
 		{
-			RayTracingGeometry.Initializer.IndexBuffer = MeshResource.IndexBuffer.IndexBufferRHI;
-			// Create the ray tracing geometry but delay the acceleration structure build.
+			FRayTracingGeometryInitializer Initializer;
+			Initializer.DebugName = FName(TEXT("GeometryCollection"));
+			Initializer.GeometryType = RTGT_Triangles;
+			Initializer.bFastBuild = true;
+			Initializer.bAllowUpdate = false;
+			Initializer.TotalPrimitiveCount = 0;
+			Initializer.IndexBuffer = MeshResource.IndexBuffer.IndexBufferRHI;
+
+			RayTracingGeometry.SetInitializer(Initializer);
+
+			// InitResource before initializing segments to avoid requesting an unnecessary build
+			RayTracingGeometry.InitResource(RHICmdList);
+
+			for (int32 SectionIndex = 0; SectionIndex < InSectionArray.Num(); ++SectionIndex)
+			{
+				const FGeometryCollectionMeshElement& Section = InSectionArray[SectionIndex];
+
+				FRayTracingGeometrySegment Segment;
+				Segment.FirstPrimitive = Section.TriangleStart;
+				Segment.VertexBuffer = MeshResource.PositionVertexBuffer.VertexBufferRHI;
+				Segment.NumPrimitives = Section.TriangleCount;
+				Segment.MaxVertices = Section.VertexEnd;
+
+				Initializer.Segments.Add(Segment);
+
+				Initializer.TotalPrimitiveCount += Section.TriangleCount;
+			}
+
+			RayTracingGeometry.SetInitializer(MoveTemp(Initializer));
+
+			// Build will be requested later using dynamic geometry update path
 			RayTracingGeometry.CreateRayTracingGeometry(RHICmdList, ERTAccelerationStructureBuildPriority::Skip);
 		}
-
-		bRayTracingGeometryDirty = false;
 	}
 }
 
@@ -484,6 +546,11 @@ void FGeometryCollectionSceneProxyBase::GetDynamicRayTracingInstances(FRayTracin
 		return;
 	}
 
+	if (MeshDescription.Sections.IsEmpty())
+	{
+		return;
+	}
+
 	QUICK_SCOPE_CYCLE_COUNTER(STAT_GeometryCollectionSceneProxyBase_GetDynamicRayTracingInstances);
 
 	const uint32 LODIndex = 0;
@@ -494,42 +561,44 @@ void FGeometryCollectionSceneProxyBase::GetDynamicRayTracingInstances(FRayTracin
 	CollectorResources = &Collector.AllocateOneFrameResource<FGeometryCollectionMeshCollectorResources>(FeatureLevel);
 	FGeometryCollectionVertexFactory& GeometryCollectionVertexFactory = CollectorResources->GetVertexFactory();
 
-	// Render dynamic objects
-	if (!GeometryCollectionVertexFactory.GetType()->SupportsRayTracingDynamicGeometry())
-	{
-		return;
-	}
-
 	SetupVertexFactory(Collector.GetRHICommandList(), GeometryCollectionVertexFactory);
 
-	TArray<FGeometryCollectionMeshElement> const& SectionArray = MeshDescription.Sections;
+	const bool bUseSubSections = MeshDescription.SubSections.Num() && !bAnyMaterialHasWorldPositionOffset && CVarRayTracingGeometryCollectionCombinedBLAS.GetValueOnRenderThread() == 0;
 
-	UpdatingRayTracingGeometry_RenderingThread(Collector.GetRHICommandList(), SectionArray);
+	TArray<FGeometryCollectionMeshElement> const& SectionArray = bUseSubSections ? MeshDescription.SubSections : MeshDescription.Sections;
 
-	if (SectionArray.Num() && RayTracingGeometry.IsValid())
+	UpdatingRayTracingGeometry_RenderingThread(Collector.GetRHICommandList(), SectionArray, bUseSubSections);
+
+	// Grab the material proxies we'll be using for each section
+	TArray<FMaterialRenderProxy*, TInlineAllocator<32>> MaterialProxies;
+
+	for (int32 SectionIndex = 0; SectionIndex < SectionArray.Num(); ++SectionIndex)
 	{
-		FRayTracingInstance RayTracingInstance;
-		RayTracingInstance.Geometry = &RayTracingGeometry;
-		RayTracingInstance.InstanceTransforms.Emplace(LocalToWorld);
+		const FGeometryCollectionMeshElement& Section = SectionArray[SectionIndex];
 
-		// Grab the material proxies we'll be using for each section
-		TArray<FMaterialRenderProxy*, TInlineAllocator<32>> MaterialProxies;
+		//TODO: Add BoneColor support in Path/Ray tracing?
+		FMaterialRenderProxy* MaterialProxy = Materials[Section.MaterialIndex]->GetRenderProxy();
 
-		for (int32 SectionIndex = 0; SectionIndex < SectionArray.Num(); ++SectionIndex)
+		if (MaterialProxy == nullptr)
 		{
-			const FGeometryCollectionMeshElement& Section = SectionArray[SectionIndex];
-
-			//TODO: Add BoneColor support in Path/Ray tracing?
-			FMaterialRenderProxy* MaterialProxy = Materials[Section.MaterialIndex]->GetRenderProxy();
+			MaterialProxy = UMaterial::GetDefaultMaterial(MD_Surface)->GetRenderProxy();
+		}
 
-			if (MaterialProxy == nullptr)
-			{
-				MaterialProxy = UMaterial::GetDefaultMaterial(MD_Surface)->GetRenderProxy();
-			}
+		MaterialProxies.Add(MaterialProxy);
+	}
 
-			MaterialProxies.Add(MaterialProxy);
+	if (RayTracingGeometry.IsValid())
+	{
+		// Render dynamic objects
+		if (!GeometryCollectionVertexFactory.GetType()->SupportsRayTracingDynamicGeometry())
+		{
+			return;
 		}
 
+		FRayTracingInstance RayTracingInstance;
+		RayTracingInstance.Geometry = &RayTracingGeometry;
+		RayTracingInstance.InstanceTransforms.Emplace(LocalToWorld);
+
 		uint32 MaxVertexIndex = 0;
 		for (int32 SectionIndex = 0; SectionIndex < SectionArray.Num(); ++SectionIndex)
 		{
@@ -563,7 +632,6 @@ void FGeometryCollectionSceneProxyBase::GetDynamicRayTracingInstances(FRayTracin
 		const bool bAlwaysUpdate = bAnyMaterialHasWorldPositionOffset || (CVarRayTracingGeometryCollectionForceUpdate.GetValueOnRenderThread() != 0);
 		const bool bNeedsUpdate = bAlwaysUpdate
 			|| (!bAlwaysUpdate && RayTracingGeometry.DynamicGeometrySharedBufferGenerationID != FRayTracingGeometry::NonSharedVertexBuffers) // was using shared VB but won't use it anymore so update once
-			|| !RayTracingGeometry.IsValid()
 			|| RayTracingGeometry.IsEvicted()
 			|| RayTracingGeometry.GetRequiresBuild();
 
@@ -596,17 +664,61 @@ void FGeometryCollectionSceneProxyBase::GetDynamicRayTracingInstances(FRayTracin
 
 		Collector.AddRayTracingInstance(MoveTemp(RayTracingInstance));
 	}
+
+	for (int32 SectionIndex = 0; SectionIndex < PartRayTracingGeometries.Num(); ++SectionIndex)
+	{
+		const FGeometryCollectionMeshElement& Section = SectionArray[SectionIndex];
+
+		FRayTracingInstance RayTracingInstance;
+		RayTracingInstance.Geometry = &PartRayTracingGeometries[SectionIndex];
+		RayTracingInstance.InstanceTransforms.Emplace(FMatrix(DynamicData->Transforms[Section.TransformIndex]) * LocalToWorld);
+
+		{
+			FMeshBatch& Mesh = RayTracingInstance.Materials.AddDefaulted_GetRef();
+			Mesh.bWireframe = bWireframe;
+			Mesh.SegmentIndex = 0;
+			Mesh.VertexFactory = &GeometryCollectionVertexFactory;
+			Mesh.MaterialRenderProxy = MaterialProxies[SectionIndex];
+			Mesh.LODIndex = LODIndex;
+			Mesh.bDisableBackfaceCulling = true;
+			Mesh.Type = PT_TriangleList;
+			Mesh.DepthPriorityGroup = SDPG_World;
+			Mesh.bCanApplyViewModeOverrides = true;
+
+			FMeshBatchElement& BatchElement = Mesh.Elements[0];
+			BatchElement.IndexBuffer = &MeshResource.IndexBuffer;
+			BatchElement.PrimitiveUniformBuffer = UniformBuffer;
+			BatchElement.FirstIndex = Section.TriangleStart * 3;
+			BatchElement.NumPrimitives = Section.TriangleCount;
+			BatchElement.MinVertexIndex = Section.VertexStart;
+			BatchElement.MaxVertexIndex = Section.VertexEnd;
+			BatchElement.NumInstances = 1;
+
+			//#TODO: bone color, bone selection and render bound?
+		}
+
+		Collector.AddRayTracingInstance(MoveTemp(RayTracingInstance));
+	}
 }
 #endif
 
 uint32 FGeometryCollectionSceneProxyBase::GetAllocatedSize() const
 {
+#if RHI_RAYTRACING
+	uint32 RayTracingGeometryAllocatedSize = RayTracingGeometry.RawData.GetAllocatedSize();
+
+	for (const FRayTracingGeometry& PartRayTracingGeometry : PartRayTracingGeometries)
+	{
+		RayTracingGeometryAllocatedSize += PartRayTracingGeometry.RawData.GetAllocatedSize();
+	}
+#endif
+
 	return Materials.GetAllocatedSize()
 		+ MeshDescription.Sections.GetAllocatedSize()
 		+ MeshDescription.SubSections.GetAllocatedSize()
 		+ (SkinnedPositionVertexBuffer.GetAllowCPUAccess() ? SkinnedPositionVertexBuffer.GetStride() * SkinnedPositionVertexBuffer.GetNumVertices() : 0)
 #if RHI_RAYTRACING
-		+ RayTracingGeometry.RawData.GetAllocatedSize()
+		+ RayTracingGeometryAllocatedSize
 #endif
 		;
 }
diff --git a/Engine/Source/Runtime/Experimental/GeometryCollectionEngine/Private/GeometryCollection/GeometryCollectionSceneProxy.h b/Engine/Source/Runtime/Experimental/GeometryCollectionEngine/Private/GeometryCollection/GeometryCollectionSceneProxy.h
--- a/Engine/Source/Runtime/Experimental/GeometryCollectionEngine/Private/GeometryCollection/GeometryCollectionSceneProxy.h
+++ b/Engine/Source/Runtime/Experimental/GeometryCollectionEngine/Private/GeometryCollection/GeometryCollectionSceneProxy.h
@@ -173,7 +173,7 @@ protected:
 	}
 
 #if RHI_RAYTRACING
-	void UpdatingRayTracingGeometry_RenderingThread(FRHICommandListBase& RHICmdList, TConstArrayView<FGeometryCollectionMeshElement> InSectionArray);
+	void UpdatingRayTracingGeometry_RenderingThread(FRHICommandListBase& RHICmdList, TConstArrayView<FGeometryCollectionMeshElement> InSectionArray, bool bRayTracingGeometryPerSection);
 
 	bool IsRayTracingRelevant() const { return true; }
 	bool IsRayTracingStaticRelevant() const { return false; }
@@ -215,9 +215,10 @@ private:
 	const bool bHasRayTracingRepresentation = false;
 
 #if RHI_RAYTRACING
-	bool bRayTracingGeometryDirty = true;
 	FRayTracingGeometry RayTracingGeometry;
 	FRWBuffer RayTracingDynamicVertexBuffer;
+
+	TArray<FRayTracingGeometry> PartRayTracingGeometries;
 #endif
 
 	bool bRenderResourcesCreated = false;
