Added an iterator to Chaos' handle array.
Also added a visualizer.

diff --git a/Engine/Source/Programs/HeadlessChaos/Private/HeadlessChaosApp.cpp b/Engine/Source/Programs/HeadlessChaos/Private/HeadlessChaosApp.cpp
--- a/Engine/Source/Programs/HeadlessChaos/Private/HeadlessChaosApp.cpp
+++ b/Engine/Source/Programs/HeadlessChaos/Private/HeadlessChaosApp.cpp
@@ -291,6 +291,7 @@ TEST(Handles, FrameworkTests)
 	ChaosTest::Handles::HandleArrayTest();
 	ChaosTest::Handles::HandleHeapTest();
 	ChaosTest::Handles::HandleSerializeTest();
+	ChaosTest::Handles::HandleArrayIteratorTest();
 }
 
 TEST(TriangleMesh, TriangleMeshTests) {
diff --git a/Engine/Source/Programs/HeadlessChaos/Private/HeadlessChaosTestHandles.cpp b/Engine/Source/Programs/HeadlessChaos/Private/HeadlessChaosTestHandles.cpp
--- a/Engine/Source/Programs/HeadlessChaos/Private/HeadlessChaosTestHandles.cpp
+++ b/Engine/Source/Programs/HeadlessChaos/Private/HeadlessChaosTestHandles.cpp
@@ -234,5 +234,64 @@ namespace ChaosTest
 			EXPECT_EQ(AsHeap.Num(), AsHeapRead.Num());
 			EXPECT_EQ(AsHeap.GetNumActive(), AsHeapRead.GetNumActive());
 		}
+
+		void HandleArrayIteratorTest()
+		{
+			THandleArray<int> HandleArray(10);
+			using FHandle = THandleArray<int>::FHandle;
+			TArray<FHandle> Handles;
+
+			for (int32 I = 0; I < 10; ++I)
+			{
+				Handles.Add(HandleArray.Create(I));
+			}
+
+			HandleArray.Destroy(Handles[0]);
+			HandleArray.Destroy(Handles[1]);
+			HandleArray.Destroy(Handles[5]);
+			HandleArray.Destroy(Handles[8]);
+			HandleArray.Destroy(Handles[9]);
+
+			const TArray<int> ExpectedValues
+			{
+				2,
+				3,
+				4,
+				6,
+				7,
+			};
+			const TArray<FHandle> ExpectedHandles
+			{
+				Handles[2],
+				Handles[3],
+				Handles[4],
+				Handles[6],
+				Handles[7],
+			};
+
+			TArray<int> ActualValues;
+			TArray<FHandle> ActualHandles;
+			for (THandleArray<int>::TRangedForIterator It = HandleArray.begin(); It != HandleArray.end(); ++It)
+			{
+				ActualHandles.Add(It.GetHandle());
+				int& Value = *It;
+				ActualValues.Add(Value);
+			}
+			EXPECT_EQ(ActualHandles, ExpectedHandles);
+			EXPECT_EQ(ActualValues, ExpectedValues);
+
+			// Test the const iterator
+			ActualHandles.Reset();
+			ActualValues.Reset();
+			const THandleArray<int>& ConstHandleArray = HandleArray;
+			for (THandleArray<int>::TRangedForConstIterator It = ConstHandleArray.begin(); It != ConstHandleArray.end(); ++It)
+			{
+				ActualHandles.Add(It.GetHandle());
+				const int& Value = *It;
+				ActualValues.Add(Value);
+			}
+			EXPECT_EQ(ActualHandles, ExpectedHandles);
+			EXPECT_EQ(ActualValues, ExpectedValues);
+		}
 	}
 }
\ No newline at end of file
diff --git a/Engine/Source/Programs/HeadlessChaos/Public/HeadlessChaosTestHandles.h b/Engine/Source/Programs/HeadlessChaos/Public/HeadlessChaosTestHandles.h
--- a/Engine/Source/Programs/HeadlessChaos/Public/HeadlessChaosTestHandles.h
+++ b/Engine/Source/Programs/HeadlessChaos/Public/HeadlessChaosTestHandles.h
@@ -11,5 +11,7 @@ namespace ChaosTest
 		void HandleHeapTest();
 
 		void HandleSerializeTest();
+
+		void HandleArrayIteratorTest();
 	}
 }
\ No newline at end of file
diff --git a/Engine/Source/Runtime/Experimental/Chaos/Public/Chaos/Framework/Handles.h b/Engine/Source/Runtime/Experimental/Chaos/Public/Chaos/Framework/Handles.h
--- a/Engine/Source/Runtime/Experimental/Chaos/Public/Chaos/Framework/Handles.h
+++ b/Engine/Source/Runtime/Experimental/Chaos/Public/Chaos/Framework/Handles.h
@@ -58,6 +58,11 @@ namespace Chaos
 			return Index << GenerationWidth | Generation;
 		}
 
+		uint32 GetIndex() const
+		{
+			return Index;
+		}
+
 		void FromUint(uint32 InUint)
 		{
 			Generation = InUint << IndexWidth;
@@ -253,7 +258,7 @@ namespace Chaos
 			return nullptr;
 		}
 
-		FHandle GetHandle(uint32 InIndex)
+		FHandle GetHandle(uint32 InIndex) const
 		{
 			// Only return a handle if we have an entry and it hasn't been retired
 			if(HandleEntries.IsValidIndex(InIndex) && HandleEntries[InIndex].Generation != 0)
@@ -264,7 +269,7 @@ namespace Chaos
 			return FHandle();
 		}
 
-		FConstHandle GetConstHandle(uint32 InIndex)
+		FConstHandle GetConstHandle(uint32 InIndex) const
 		{
 			return FConstHandle(GetHandle(InIndex));
 		}
@@ -391,6 +396,87 @@ namespace Chaos
 			return Ar;
 		}
 
+		template<bool bConst>
+		struct TIterator
+		{
+			typedef std::conditional_t<bConst, const THandleArray, THandleArray> ArrayType;
+			typedef std::conditional_t<bConst, const ElementType, ElementType> ItElementType;
+
+			explicit TIterator(ArrayType& InArray, int32 InIndex)
+				: Array(InArray)
+				, It(InArray.Validity, InIndex)
+			{
+			}
+
+			FORCEINLINE TIterator& operator++()
+			{
+				// Iterate to the next set allocation flag.
+				++It;
+				return *this;
+			}
+
+			[[nodiscard]] FORCEINLINE bool operator==(const TIterator& Rhs) const
+			{
+				return It == Rhs.It && &Array == &Rhs.Array;
+			}
+
+			[[nodiscard]] FORCEINLINE bool operator!=(const TIterator& Rhs) const
+			{
+				return It != Rhs.It || &Array != &Rhs.Array;
+			}
+
+			FConstHandle GetConstHandle() const
+			{
+				const int32 Index = It.GetIndex();
+				return Array.GetConstHandle(Index);
+			}
+
+			FHandle GetHandle() const
+			{
+				const int32 Index = It.GetIndex();
+				return Array.GetHandle(Index);
+			}
+
+			[[nodiscard]] FORCEINLINE ItElementType& operator*() const
+			{
+				const FHandle Handle = GetHandle();
+				ItElementType* Element = Array.Get(Handle);
+				return *Element;
+			}
+
+			[[nodiscard]] FORCEINLINE ItElementType* operator->() const
+			{
+				const FHandle Handle = GetHandle();
+				return Array.Get(Handle);
+			}
+
+			ArrayType& Array;
+			TConstSetBitIterator<> It;
+		};
+
+		using TRangedForIterator = TIterator<false>;
+		using TRangedForConstIterator = TIterator<true>;
+
+		TRangedForIterator begin()
+		{
+			return TRangedForIterator(*this, 0);
+		}
+
+		TRangedForConstIterator begin() const
+		{
+			return TRangedForConstIterator(*this, 0);
+		}
+
+		TRangedForIterator end()
+		{
+			return TRangedForIterator(*this, (int32)Num());
+		}
+
+		TRangedForConstIterator end() const
+		{
+			return TRangedForConstIterator(*this, (int32)Num());
+		}
+
 	private:
 
 		void DestroyItems()
