Fix VerifyGraph in Chaos GraphColoring to correctly handle when an edge has duplicate nodes.

diff --git a/Engine/Source/Runtime/Experimental/Chaos/Private/Chaos/GraphColoring.cpp b/Engine/Source/Runtime/Experimental/Chaos/Private/Chaos/GraphColoring.cpp
--- a/Engine/Source/Runtime/Experimental/Chaos/Private/Chaos/GraphColoring.cpp
+++ b/Engine/Source/Runtime/Experimental/Chaos/Private/Chaos/GraphColoring.cpp
@@ -17,7 +17,8 @@ static bool VerifyGraph(const TArray<TArray<int32>>& ColorGraph, const TArray<Ch
 			for (int32 NIndex = 0; NIndex < N; ++NIndex)
 			{
 				const int32 Node = Graph[Edge][NIndex];
-				if (const int32* const ExistingEdge = ColorNodesToEdges.Find(Node))
+				const int32* const ExistingEdge = ColorNodesToEdges.Find(Node);
+				if (ExistingEdge && *ExistingEdge != Edge)
 				{
 					UE_LOG(LogChaos, Error, TEXT("Color %d has duplicate Node %d. First added for Edge %d, and now found for Edge %d"), Color, Node, *ExistingEdge, Edge);
 					return false;
