[DSTM] Added support to serialize particle data from the game thread

Also, added VSmooth and Kinematic Target data to the Physics Thread state serialized data

Lastly, both are now also exposed in CVD

diff --git a/Engine/Plugins/ChaosVD/Source/ChaosVD/Private/Visualizers/ChaosVDJointConstraintsDataComponentVisualizer.cpp b/Engine/Plugins/ChaosVD/Source/ChaosVD/Private/Visualizers/ChaosVDJointConstraintsDataComponentVisualizer.cpp
--- a/Engine/Plugins/ChaosVD/Source/ChaosVD/Private/Visualizers/ChaosVDJointConstraintsDataComponentVisualizer.cpp
+++ b/Engine/Plugins/ChaosVD/Source/ChaosVD/Private/Visualizers/ChaosVDJointConstraintsDataComponentVisualizer.cpp
@@ -225,7 +225,7 @@ void FChaosVDJointConstraintsDataComponentVisualizer::DrawJointConstraint(const
 
 	if (!VisualizationContext.IsVisualizationFlagEnabled(EChaosVDJointsDataVisualizationFlags::DrawDisabled))
 	{
-		if (InJointConstraintData.JointState.bDisabled)
+		if (InJointConstraintData.PhysicsThreadJointState.bDisabled)
 		{
 			return;
 		}
@@ -391,7 +391,7 @@ void FChaosVDJointConstraintsDataComponentVisualizer::DrawJointConstraint(const
 	if (VisualizationContext.IsVisualizationFlagEnabled(EChaosVDJointsDataVisualizationFlags::PushOut))
 	{
 		FLinearColor PushOutImpulseColor = FColor(0, 250, 250);
-		FChaosVDDebugDrawUtils::DrawLine(PDI, Xa, Xa + VisualizationContext.DebugDrawSettings->LinearImpulseScale * VisualizationContext.SpaceTransform.TransformVectorNoScale(InJointConstraintData.JointState.LinearImpulse), GenerateSelectionAwareDebugColor(PushOutImpulseColor, bIsSelected), FText::GetEmpty(), VisualizationContext.DebugDrawSettings->DepthPriority, LineThickness);
+		FChaosVDDebugDrawUtils::DrawLine(PDI, Xa, Xa + VisualizationContext.DebugDrawSettings->LinearImpulseScale * VisualizationContext.SpaceTransform.TransformVectorNoScale(InJointConstraintData.PhysicsThreadJointState.LinearImpulse), GenerateSelectionAwareDebugColor(PushOutImpulseColor, bIsSelected), FText::GetEmpty(), VisualizationContext.DebugDrawSettings->DepthPriority, LineThickness);
 	}
 
 	//TODO: Should we draw the Angular Impulse
diff --git a/Engine/Source/Runtime/Engine/Classes/Components/SkeletalMeshComponent.h b/Engine/Source/Runtime/Engine/Classes/Components/SkeletalMeshComponent.h
--- a/Engine/Source/Runtime/Engine/Classes/Components/SkeletalMeshComponent.h
+++ b/Engine/Source/Runtime/Engine/Classes/Components/SkeletalMeshComponent.h
@@ -1997,101 +1997,6 @@ protected:
 	 */
 	ENGINE_API virtual FTransform GetComponentTransformFromBodyInstance(FBodyInstance* UseBI) override;
 
-private:
-	struct UE_INTERNAL FMigratedPhysicsState
-	{
-		void ReadElementDataIntoBuffer(Chaos::FSerializedDataBuffer& TargetBuffer);
-
-		template<typename ElementType>
-		void AppendDataFromElement(ElementType& InElement)
-		{
-			FDataEntryTag DataTag;
-			DataTag.DataOffset = GetSize();
-			
-			InElement.SerializeState(MigratedStateAsBytes);
-
-			DataTag.DataSize = GetSize() - DataTag.DataOffset;
-
-			Header.DataTagPerElementIndex.Emplace(DataTag);
-		}
-
-		void Serialize(FArchive& Ar);
-
-		int32 GetSize() const
-		{
-			return MigratedStateAsBytes.GetSize();
-		}
-
-private:
-
-		Chaos::FSerializedDataBuffer MigratedStateAsBytes;
-
-		struct FDataEntryTag
-		{
-			int32 DataOffset = INDEX_NONE;
-			int32 DataSize = INDEX_NONE;
-
-			friend FArchive& operator<<(FArchive& Ar, FDataEntryTag& Data)
-			{
-				Ar << Data.DataOffset;
-				Ar << Data.DataSize;
-				return Ar;
-			}
-
-			bool IsValid() const
-			{
-				return DataOffset != INDEX_NONE && DataSize != INDEX_NONE;
-			}
-		};
-
-		struct FHeader
-		{
-			int32 NumBodies = 0;
-			int32 NumConstraints = 0;
-
-			TArray<FDataEntryTag> DataTagPerElementIndex;
-			
-			friend FArchive& operator<<(FArchive& Ar, FHeader& Data)
-			{
-				Ar << Data.NumBodies;
-				Ar << Data.NumConstraints;
-				Ar << Data.DataTagPerElementIndex;
-
-				return Ar;
-			}
-
-			bool IsValid() const
-			{
-				return NumBodies + NumConstraints == DataTagPerElementIndex.Num();
-			}
-		};
-		
-		FHeader Header;
-
-	public:
-
-		void SetNumOfBodies(int32 NumBodies)
-		{
-			Header.NumBodies = NumBodies;
-		}
-
-		void SetNumOfConstraints(int32 NumConstraints)
-		{
-			Header.NumConstraints = NumConstraints;
-		}
-
-		const FHeader& GetHeader() const
-		{
-			return Header;
-		}
-
-	private:
-
-		int32 CurrentReadElementIndex = 0;
-	};
-
-protected:
-
 #if UE_WITH_REMOTE_OBJECT_HANDLE
 	ENGINE_API virtual void MigratePhysicsState(FArchive& Ar) override;
 	ENGINE_API virtual void ApplyMigratedPhysicsState(Chaos::FSerializedDataBufferPtr MigratedStateData = nullptr) override;
diff --git a/Engine/Source/Runtime/Engine/Classes/PhysicsEngine/BodyInstance.h b/Engine/Source/Runtime/Engine/Classes/PhysicsEngine/BodyInstance.h
--- a/Engine/Source/Runtime/Engine/Classes/PhysicsEngine/BodyInstance.h
+++ b/Engine/Source/Runtime/Engine/Classes/PhysicsEngine/BodyInstance.h
@@ -1471,10 +1471,10 @@ private:
 public:
 
 	/** Serialized out the physics thread state of this body as an array of bytes */
-	void SerializeState(Chaos::FSerializedDataBuffer& OutSerializedState);
+	UE_INTERNAL ENGINE_API void SerializeState(Chaos::FSerializedDataBuffer& OutSerializedState, Chaos::ESerializedDataContext DataContext);
 
 	/** Deserializes the provided array of bytes representing a saved physics thread state and applies it to this body */
-	void ApplySerializedState(Chaos::FSerializedDataBufferPtr&& InSerializedState);
+	UE_INTERNAL ENGINE_API void ApplySerializedState(Chaos::FSerializedDataBufferPtr&& InSerializedState, Chaos::ESerializedDataContext DataContext);
 
 private:
 
diff --git a/Engine/Source/Runtime/Engine/Classes/PhysicsEngine/ConstraintInstance.h b/Engine/Source/Runtime/Engine/Classes/PhysicsEngine/ConstraintInstance.h
--- a/Engine/Source/Runtime/Engine/Classes/PhysicsEngine/ConstraintInstance.h
+++ b/Engine/Source/Runtime/Engine/Classes/PhysicsEngine/ConstraintInstance.h
@@ -332,10 +332,10 @@ public:
 	ENGINE_API const FPhysicsConstraintHandle& GetPhysicsConstraintRef() const;
 
 	/** Serialized out the physics thread state of this constraint as an array of bytes */
-	void SerializeState(Chaos::FSerializedDataBuffer& OutSerializedState);
+	UE_INTERNAL void SerializeState(Chaos::FSerializedDataBuffer& OutSerializedState, Chaos::ESerializedDataContext DataContext);
 
 	/** Deserializes the provided array of bytes representing a saved physics thread state and applies it to this constraint */
-	void ApplySerializedState(Chaos::FSerializedDataBufferPtr&& InSerializedState);
+	UE_INTERNAL void ApplySerializedState(Chaos::FSerializedDataBufferPtr&& InSerializedState, Chaos::ESerializedDataContext DataContext);
 
 	FChaosUserData UserData;
 
diff --git a/Engine/Source/Runtime/Engine/Private/Components/PrimitiveComponent.cpp b/Engine/Source/Runtime/Engine/Private/Components/PrimitiveComponent.cpp
--- a/Engine/Source/Runtime/Engine/Private/Components/PrimitiveComponent.cpp
+++ b/Engine/Source/Runtime/Engine/Private/Components/PrimitiveComponent.cpp
@@ -69,6 +69,7 @@
 #endif // WITH_EDITOR
 
 #include "Chaos/Serialization/SerializedDataBuffer.h"
+#include "Chaos/Serialization/SerializedMultiPhysicsState.h"
 
 #if DO_CHECK
 #include "Engine/StaticMesh.h"
@@ -1252,20 +1253,6 @@ void UPrimitiveComponent::Serialize(FArchive& Ar)
 	// For remote objects we need non-UProperties to be transferred
 	if (Ar.IsMigratingRemoteObjects())
 	{
-		const FUObjectMigrationContext* MigrationContext = Ar.GetMigrationContext();
-		check(MigrationContext);
-
-		if (Ar.IsLoading())
-		{
-			EObjectMigrationRecvType ObjMigrationRecvType = MigrationContext->GetObjectMigrationRecvType(this);
-			UE_LOG(LogTemp, Log, TEXT("%hs: %s ObjMigrationRecvType: %s"), __func__, *GetName(), ToString(ObjMigrationRecvType));
-		}
-		else if (Ar.IsSaving())
-		{
-			EObjectMigrationSendType ObjMigrationSendType = MigrationContext->GetObjectMigrationSendType(this);
-			UE_LOG(LogTemp, Log, TEXT("%hs: %s ObjMigrationSendType: %s"), __func__, *GetName(), ToString(ObjMigrationSendType));
-		}
-
 		{
 	    	int32 NumOverlappingComponents = OverlappingComponents.Num();
     		Ar << NumOverlappingComponents;
@@ -1313,15 +1300,25 @@ void UPrimitiveComponent::Serialize(FArchive& Ar)
 #if UE_WITH_REMOTE_OBJECT_HANDLE
 void UPrimitiveComponent::MigratePhysicsState(FArchive& Ar)
 {
+	const FUObjectMigrationContext* MigrationContext = Ar.GetMigrationContext();
+	check(MigrationContext);
+	
 	if (Ar.IsLoading())
 	{
 		Chaos::FSerializedDataBufferPtr MigratedStateDataBuffer = MakeUnique<Chaos::FSerializedDataBuffer>();
 		
 		Ar << (*MigratedStateDataBuffer);
+		
+		const EObjectMigrationRecvType ObjMigrationRecvType = MigrationContext->GetObjectMigrationRecvType(this);
 
 		if (MigratedStateDataBuffer->GetDataAsByteArrayRef().IsEmpty())
 		{
-			UE_LOG(LogPrimitiveComponent, Verbose, TEXT("Migrated an empty physics state (this might be expected if during migration no physics state was available) | [%s] will use its default physics state."), *GetPathNameSafe(this));
+			UE_LOG(LogPrimitiveComponent, Verbose, TEXT("Migrated an empty physics state (this can be expected) | [%s] will use its default physics state or the one already had | Migration Context [%s]."), *GetPathNameSafe(this), ToString(ObjMigrationRecvType));
+			return;
+		}
+
+		if (ObjMigrationRecvType == EObjectMigrationRecvType::ReturnedLoan)
+		{
 			return;
 		}
 
@@ -1336,28 +1333,40 @@ void UPrimitiveComponent::MigratePhysicsState(FArchive& Ar)
 		}
 	}
 	else if (Ar.IsSaving())
-	{
-		if (HasValidPhysicsState())
+	{		
+		const EObjectMigrationSendType ObjMigrationSendType = MigrationContext->GetObjectMigrationSendType(this);
+		const bool bCanSendPhysicsState = ObjMigrationSendType != EObjectMigrationSendType::ReturnBorrowed && HasValidPhysicsState();
+
+		if (bCanSendPhysicsState)
 		{
-			Chaos::FSerializedDataBuffer SerializedStateData;
-			BodyInstance.SerializeState(SerializedStateData);
+			const Chaos::ESerializedDataContext PhysicsDataContext = ObjMigrationSendType == EObjectMigrationSendType::ReassignOwnership ? Chaos::ESerializedDataContext::Internal : Chaos::ESerializedDataContext::External;
+			Chaos::FSerializedMultiPhysicsState MigratedState(PhysicsDataContext);
+
+			MigratedState.SetNumOfBodies(1);
+			MigratedState.SetNumOfConstraints(0);
+
+			MigratedState.AppendDataFromElement(BodyInstance);
 
-			Ar << SerializedStateData;
+			// We need re-serialize the state as a chaos data buffer, so we can use it as is on de-serialization
+			Chaos::FSerializedDataBuffer MigratedStateAsBytes;
+			FMemoryWriter MigratedStateWriter(MigratedStateAsBytes.GetDataAsByteArrayRef());
+			MigratedState.Serialize(MigratedStateWriter);
+
+			Ar << MigratedStateAsBytes;
 		}
 		else
 		{
 			// This object might be being migrated before we had a chance to process a previous migration
 			FPhysScene* PhysicsScene = FPhysicsInterface::GetCurrentSceneForComponent(this);
-			if (Chaos::FSerializedDataBufferPtr RawMigratedData = PhysicsScene ? PhysicsScene->PopPendingSerializedStateForObject(this) : nullptr)
+			Chaos::FSerializedDataBufferPtr RawMigratedData = PhysicsScene ? PhysicsScene->PopPendingSerializedStateForObject(this) : nullptr;
+			if (!ensure(!RawMigratedData))
 			{
-				Ar << (*RawMigratedData);
-			}
-			else
-			{
-				// We need to serialize an empty buffer as array, so we can handle this case during the load phase post migration
-				Chaos::FSerializedDataBuffer EmptyPhysicsState;
-				Ar << EmptyPhysicsState;
+				UE_LOG(LogPrimitiveComponent, Error, TEXT("[%hs] Attempted to migrated an object physics state, before processing a previous pending physics state| Discarding pending physics state and sending an empty one..."), __func__);
 			}
+
+			// We need to serialize an empty buffer as array, so we can handle this case during the load phase post migration
+			Chaos::FSerializedDataBuffer EmptyPhysicsState;
+			Ar << EmptyPhysicsState;
 		}
 	}
 }
@@ -1369,17 +1378,38 @@ void UPrimitiveComponent::ApplyMigratedPhysicsState(Chaos::FSerializedDataBuffer
 		return;
 	}
 
+	Chaos::FSerializedDataBufferPtr RawMigratedData;
+
 	if (MigratedStateData)
 	{
-		BodyInstance.ApplySerializedState(MoveTemp(MigratedStateData));
+		RawMigratedData = MoveTemp(MigratedStateData);
 	}
-	else
+	else if (FPhysScene* PhysicsScene = FPhysicsInterface::GetCurrentSceneForComponent(this))
 	{
-		FPhysScene* PhysicsScene = FPhysicsInterface::GetCurrentSceneForComponent(this);
-		if (Chaos::FSerializedDataBufferPtr RawMigratedData = PhysicsScene ? PhysicsScene->PopPendingSerializedStateForObject(this) : nullptr)
+		RawMigratedData = PhysicsScene->PopPendingSerializedStateForObject(this);
+	}
+
+	if (RawMigratedData)
+	{
+		FMemoryReader MultiBodyStateReader(RawMigratedData->GetDataAsByteArrayRef());
+
+		Chaos::FSerializedMultiPhysicsState MigratedState;
+		MigratedState.Serialize(MultiBodyStateReader);
+		
+		const Chaos::FSerializedMultiPhysicsState::FHeader& MigratedStateHeader = MigratedState.GetHeader();
+
+		if (MigratedStateHeader.IsValid())
 		{
-			BodyInstance.ApplySerializedState(MoveTemp(RawMigratedData));
-		}
+			UE_LOG(LogSkeletalMesh, Verbose, TEXT("[%hs] Applying migrated physics state | Data context [%s]"), __func__, Chaos::LexToString(MigratedStateHeader.DataContext));
+
+			Chaos::FSerializedDataBufferPtr DataBuffer = MakeUnique<Chaos::FSerializedDataBuffer>();
+			MigratedState.ReadElementDataIntoBuffer(*DataBuffer);
+
+			// If we received an Internal physics state, we want to apply it to both, Internal and External so it is available immediatly in the game thread API
+			Chaos::ESerializedDataContext SerializationDataContext = MigratedStateHeader.DataContext == Chaos::ESerializedDataContext::Internal ?  Chaos::ESerializedDataContext::Both : MigratedStateHeader.DataContext;
+
+			BodyInstance.ApplySerializedState(MoveTemp(DataBuffer), SerializationDataContext);
+		}	
 	}
 }
 #endif
diff --git a/Engine/Source/Runtime/Engine/Private/Components/SkeletalMeshComponent.cpp b/Engine/Source/Runtime/Engine/Private/Components/SkeletalMeshComponent.cpp
--- a/Engine/Source/Runtime/Engine/Private/Components/SkeletalMeshComponent.cpp
+++ b/Engine/Source/Runtime/Engine/Private/Components/SkeletalMeshComponent.cpp
@@ -59,6 +59,9 @@
 #endif
 
 #include "Chaos/Serialization/SerializedDataBuffer.h"
+#include "Chaos/Serialization/SerializedMultiPhysicsState.h"
+#include "PhysicsProxy/SingleParticlePhysicsProxy.h"
+#include "UObject/UObjectMigrationContext.h"
 
 #include UE_INLINE_GENERATED_CPP_BY_NAME(SkeletalMeshComponent)
 
@@ -3835,54 +3838,29 @@ void USkeletalMeshComponent::GetResourceSizeEx(FResourceSizeEx& CumulativeResour
 	}
 }
 
-void USkeletalMeshComponent::FMigratedPhysicsState::ReadElementDataIntoBuffer(Chaos::FSerializedDataBuffer& TargetBuffer)
-{
-	FDataEntryTag DataTag = Header.DataTagPerElementIndex.IsValidIndex(CurrentReadElementIndex) ? Header.DataTagPerElementIndex[CurrentReadElementIndex] : FDataEntryTag();
-
-	TArray<uint8>& AsByteArray = MigratedStateAsBytes.GetDataAsByteArrayRef();
-
-	const int32 BufferSize = GetSize();
-
-	UE_LOG(LogSkeletalMesh, Verbose, TEXT("[Physics State Migration] Attempted to read an empty data entry. This might means we failed to obtain data to serialize a specific body or constraint."));
-
-	// If we have entries past the available buffer size, but they were recorded as empty, the serialization code is ok to continue (it will just be a no-op)
-	const bool bHasValidStartingOffset = DataTag.DataOffset < BufferSize || DataTag.DataSize == 0;
-	const bool bThereIsEnoughDataToCopyFromBuffer = DataTag.DataOffset + DataTag.DataSize <= BufferSize;
-
-	if (ensure(DataTag.IsValid() && bHasValidStartingOffset && bThereIsEnoughDataToCopyFromBuffer))
-	{
-		if (DataTag.DataSize > 0)
-		{
-			TargetBuffer.GetDataAsByteArrayRef().Append(&AsByteArray[DataTag.DataOffset], DataTag.DataSize);
-		}
-	}
-	else
-	{
-		UE_LOG(LogSkeletalMesh, Error, TEXT("[Physics State Migration] Attempted to read data out of bounds! | Buffer size [%d] | Start Pos [%d]| SizeToCopy [%d]"), AsByteArray.Num(), DataTag.DataOffset, DataTag.DataSize);
-	}
-	
-	CurrentReadElementIndex++;
-}
-
-void USkeletalMeshComponent::FMigratedPhysicsState::Serialize(FArchive& Ar)
-{
-	Ar << Header;
-	Ar << MigratedStateAsBytes;
-}
-
 #if UE_WITH_REMOTE_OBJECT_HANDLE
 
 void USkeletalMeshComponent::MigratePhysicsState(FArchive& Ar)
 {
+	const FUObjectMigrationContext* MigrationContext = Ar.GetMigrationContext();
+	check(MigrationContext);
+
 	if (Ar.IsLoading())
 	{
 		Chaos::FSerializedDataBufferPtr MigratedStateData = MakeUnique<Chaos::FSerializedDataBuffer>();
 
 		Ar << (*MigratedStateData);
 
+		const EObjectMigrationRecvType ObjMigrationRecvType = MigrationContext->GetObjectMigrationRecvType(this);
+
 		if (MigratedStateData->GetDataAsByteArrayRef().IsEmpty())
 		{
-			UE_LOG(LogSkeletalMesh, Verbose, TEXT("Migrated an empty physics state (this might be expected if during migration no physics state was available) | [%s] will use its default physics state."), *GetPathNameSafe(this));
+			UE_LOG(LogSkeletalMesh, Verbose, TEXT("Migrated an empty physics state (this can be expected) | [%s] will use its default physics state or the one already had | Migration Context [%s]."), *GetPathNameSafe(this), ToString(ObjMigrationRecvType));
+			return;
+		}
+	
+		if (ObjMigrationRecvType == EObjectMigrationRecvType::ReturnedLoan)
+		{
 			return;
 		}
 
@@ -3905,9 +3883,12 @@ void USkeletalMeshComponent::MigratePhysicsState(FArchive& Ar)
 	}
 	else if (Ar.IsSaving())
 	{
-		if (HasValidPhysicsState())
+		const EObjectMigrationSendType ObjMigrationSendType = MigrationContext->GetObjectMigrationSendType(this);
+		const bool bCanSendPhysicsState = ObjMigrationSendType != EObjectMigrationSendType::ReturnBorrowed && HasValidPhysicsState();
+
+		if (bCanSendPhysicsState)
 		{
-			FMigratedPhysicsState MigratedState;
+			Chaos::FSerializedMultiPhysicsState MigratedState(ObjMigrationSendType == EObjectMigrationSendType::ReassignOwnership ? Chaos::ESerializedDataContext::Internal : Chaos::ESerializedDataContext::External);
 			MigratedState.SetNumOfBodies(Bodies.Num());
 			MigratedState.SetNumOfConstraints(Constraints.Num());
 
@@ -3934,16 +3915,15 @@ void USkeletalMeshComponent::MigratePhysicsState(FArchive& Ar)
 		{
 			// This object might be being migrated before we had a chance to process a previous migration
 			FPhysScene* PhysicsScene = FPhysicsInterface::GetCurrentSceneForComponent(this);
-			if (Chaos::FSerializedDataBufferPtr RawMigratedData = PhysicsScene ? PhysicsScene->PopPendingSerializedStateForObject(this) : nullptr)
-			{
-				Ar << (*RawMigratedData);
-			}
-			else
+			Chaos::FSerializedDataBufferPtr RawMigratedData = PhysicsScene ? PhysicsScene->PopPendingSerializedStateForObject(this) : nullptr;
+			if (!ensure(!RawMigratedData))
 			{
-				// We need to serialize an empty buffer as array, so we can handle this case during the load phase post migration
-				Chaos::FSerializedDataBuffer EmptyPhysicsState;
-				Ar << EmptyPhysicsState;
+				UE_LOG(LogSkeletalMesh, Error, TEXT("[%hs] Attempted to migrated an object physics state, before processing a previous pending physics state| Discarding pending physics state and sending an empty one..."), __func__);
 			}
+			
+			// We need to serialize an empty buffer as array, so we can handle this case during the load phase post migration
+			Chaos::FSerializedDataBuffer EmptyPhysicsState;
+			Ar << EmptyPhysicsState;
 		}
 	}
 }
@@ -3965,15 +3945,21 @@ void USkeletalMeshComponent::ApplyMigratedPhysicsState(Chaos::FSerializedDataBuf
 	{
 		FMemoryReader MultiBodyStateReader(RawMigratedData->GetDataAsByteArrayRef());
 
-		FMigratedPhysicsState MigratedState;
+		Chaos::FSerializedMultiPhysicsState MigratedState;
 		MigratedState.Serialize(MultiBodyStateReader);
 
-		const bool bIsValidMigratedState = MigratedState.GetHeader().IsValid() && MigratedState.GetHeader().NumBodies == Bodies.Num() || MigratedState.GetHeader().NumConstraints == Constraints.Num();
+		const Chaos::FSerializedMultiPhysicsState::FHeader& MigratedStateHeader = MigratedState.GetHeader();
+		const bool bIsValidMigratedState = MigratedStateHeader.IsValid() && MigratedStateHeader.NumBodies == Bodies.Num() || MigratedStateHeader.NumConstraints == Constraints.Num();
 		if (!ensure(bIsValidMigratedState))
 		{
 			return;
 		}
 
+		UE_LOG(LogSkeletalMesh, Verbose, TEXT("[%hs] Applying migrated physics state | Data context [%s]"), __func__, Chaos::LexToString( MigratedState.GetHeader().DataContext));
+
+		// If we received an Internal physics state, we want to apply it to both, Internal and External so it is available immediatly in the game thread API
+		Chaos::ESerializedDataContext SerializationDataContext = MigratedStateHeader.DataContext == Chaos::ESerializedDataContext::Internal ? Chaos::ESerializedDataContext::Both : MigratedStateHeader.DataContext;
+
 		for (FBodyInstance* BodyInstance : Bodies)
 		{
 			Chaos::FSerializedDataBufferPtr DataBuffer = MakeUnique<Chaos::FSerializedDataBuffer>();
@@ -3981,7 +3967,7 @@ void USkeletalMeshComponent::ApplyMigratedPhysicsState(Chaos::FSerializedDataBuf
 
 			if (ensure(BodyInstance))
 			{
-				BodyInstance->ApplySerializedState(MoveTemp(DataBuffer));
+				BodyInstance->ApplySerializedState(MoveTemp(DataBuffer), SerializationDataContext);
 			}
 		}
 
@@ -3992,7 +3978,7 @@ void USkeletalMeshComponent::ApplyMigratedPhysicsState(Chaos::FSerializedDataBuf
 
 			if (ensure(ConstraintInstance))
 			{
-				ConstraintInstance->ApplySerializedState(MoveTemp(DataBuffer));
+				ConstraintInstance->ApplySerializedState(MoveTemp(DataBuffer), SerializationDataContext);
 			}
 		}
 	}
diff --git a/Engine/Source/Runtime/Engine/Private/PhysicsEngine/BodyInstance.cpp b/Engine/Source/Runtime/Engine/Private/PhysicsEngine/BodyInstance.cpp
--- a/Engine/Source/Runtime/Engine/Private/PhysicsEngine/BodyInstance.cpp
+++ b/Engine/Source/Runtime/Engine/Private/PhysicsEngine/BodyInstance.cpp
@@ -43,6 +43,9 @@
 
 #define LOCTEXT_NAMESPACE "BodyInstance"
 
+#include "Chaos/ChaosEngineInterface.h"
+#include "Physics/Experimental/ChaosScopedSceneLock.h"
+
 #include UE_INLINE_GENERATED_CPP_BY_NAME(BodyInstance)
 
 DECLARE_CYCLE_STAT(TEXT("Init Body"), STAT_InitBody, STATGROUP_Physics);
@@ -3444,20 +3447,22 @@ void FBodyInstance::UpdateDebugRendering()
 #endif
 }
 
-void FBodyInstance::ApplySerializedState(Chaos::FSerializedDataBufferPtr&& InSerializedState)
+void FBodyInstance::ApplySerializedState(Chaos::FSerializedDataBufferPtr&& InSerializedState, Chaos::ESerializedDataContext DataContext)
 {
-	FPhysicsCommand::ExecuteWrite(ActorHandle, [&InSerializedState](const FPhysicsActorHandle& Actor)
+	const EPhysicsInterfaceScopedThreadContext LockContext = DataContext == Chaos::ESerializedDataContext::Internal ? EPhysicsInterfaceScopedThreadContext::Internal : EPhysicsInterfaceScopedThreadContext::External;
+	FPhysicsCommand::ExecuteWrite(ActorHandle, [&InSerializedState, DataContext](const FPhysicsActorHandle& Actor)
 	{
-		FChaosEngineInterface::ApplySerializedStateToParticle(Actor, MoveTemp(InSerializedState));
-	});
+		FChaosEngineInterface::ApplySerializedStateToParticle(Actor, MoveTemp(InSerializedState), DataContext);
+	}, LockContext);
 }
 
-void FBodyInstance::SerializeState(Chaos::FSerializedDataBuffer& OutSerializedState)
+void FBodyInstance::SerializeState(Chaos::FSerializedDataBuffer& OutSerializedState, Chaos::ESerializedDataContext DataContext)
 {
-	FPhysicsCommand::ExecuteRead(ActorHandle, [&OutSerializedState](const FPhysicsActorHandle& Actor)
+	const EPhysicsInterfaceScopedThreadContext LockContext = DataContext == Chaos::ESerializedDataContext::Internal ? EPhysicsInterfaceScopedThreadContext::Internal : EPhysicsInterfaceScopedThreadContext::External;
+	FPhysicsCommand::ExecuteRead(ActorHandle, [&OutSerializedState, DataContext](const FPhysicsActorHandle& Actor)
 	{
-		FChaosEngineInterface::SerializeParticleStateToBuffer(Actor, OutSerializedState);
-	});
+		FChaosEngineInterface::SerializeParticleStateToBuffer(Actor, OutSerializedState, DataContext);
+	}, LockContext);
 }
 
 void FBodyInstance::UpdateDampingProperties()
diff --git a/Engine/Source/Runtime/Engine/Private/PhysicsEngine/ConstraintInstance.cpp b/Engine/Source/Runtime/Engine/Private/PhysicsEngine/ConstraintInstance.cpp
--- a/Engine/Source/Runtime/Engine/Private/PhysicsEngine/ConstraintInstance.cpp
+++ b/Engine/Source/Runtime/Engine/Private/PhysicsEngine/ConstraintInstance.cpp
@@ -713,20 +713,22 @@ const FPhysicsConstraintHandle& FConstraintInstance::GetPhysicsConstraintRef() c
 	return ConstraintHandle;
 }
 
-void FConstraintInstance::SerializeState(Chaos::FSerializedDataBuffer& OutSerializedState)
+void FConstraintInstance::SerializeState(Chaos::FSerializedDataBuffer& OutSerializedState, Chaos::ESerializedDataContext DataContext)
 {
-	FPhysicsCommand::ExecuteWrite(ConstraintHandle, [&OutSerializedState](const FPhysicsConstraintHandle& Constraint)
+	const EPhysicsInterfaceScopedThreadContext LockContext = DataContext == Chaos::ESerializedDataContext::Internal ? EPhysicsInterfaceScopedThreadContext::Internal : EPhysicsInterfaceScopedThreadContext::External;
+	FPhysicsCommand::ExecuteWrite(ConstraintHandle, [&OutSerializedState, DataContext](const FPhysicsConstraintHandle& Constraint)
 	{
-		FChaosEngineInterface::SerializeConstraintStateToBuffer(Constraint, OutSerializedState);
-	});
+		FChaosEngineInterface::SerializeConstraintStateToBuffer(Constraint, OutSerializedState, DataContext);
+	}, LockContext);
 }
 
-void FConstraintInstance::ApplySerializedState(Chaos::FSerializedDataBufferPtr&& InSerializedState)
+void FConstraintInstance::ApplySerializedState(Chaos::FSerializedDataBufferPtr&& InSerializedState, Chaos::ESerializedDataContext DataContext)
 {
-	FPhysicsCommand::ExecuteRead(ConstraintHandle, [&InSerializedState, this](const FPhysicsConstraintHandle& Constraint)
+	const EPhysicsInterfaceScopedThreadContext LockContext = DataContext == Chaos::ESerializedDataContext::Internal ? EPhysicsInterfaceScopedThreadContext::Internal : EPhysicsInterfaceScopedThreadContext::External;
+	FPhysicsCommand::ExecuteRead(ConstraintHandle, [&InSerializedState, this, DataContext](const FPhysicsConstraintHandle& Constraint)
 	{
-		FChaosEngineInterface::ApplySerializedStateToConstraint(Constraint, MoveTemp(InSerializedState));
-	});
+		FChaosEngineInterface::ApplySerializedStateToConstraint(Constraint, MoveTemp(InSerializedState), DataContext);
+	}, LockContext);
 }
 
 void FConstraintInstance::CopyConstraintGeometryFrom(const FConstraintInstance* FromInstance)
diff --git a/Engine/Source/Runtime/Engine/Private/PhysicsEngine/Experimental/ChaosScopedSceneLock.cpp b/Engine/Source/Runtime/Engine/Private/PhysicsEngine/Experimental/ChaosScopedSceneLock.cpp
--- a/Engine/Source/Runtime/Engine/Private/PhysicsEngine/Experimental/ChaosScopedSceneLock.cpp
+++ b/Engine/Source/Runtime/Engine/Private/PhysicsEngine/Experimental/ChaosScopedSceneLock.cpp
@@ -74,8 +74,8 @@ FScopedSceneLock_Chaos::FScopedSceneLock_Chaos(FPhysicsActorHandle const* InActo
 }
 
 FScopedSceneLock_Chaos::FScopedSceneLock_Chaos(FPhysicsConstraintHandle const* InConstraintHandle, EPhysicsInterfaceScopedLockType InLockType)
-	: Solver(nullptr)
-	, LockType(InLockType)
+	: Solver(nullptr),
+	LockType(InLockType)
 {
 	if (InConstraintHandle)
 	{
@@ -137,8 +137,8 @@ FScopedSceneLock_Chaos::FScopedSceneLock_Chaos(Chaos::FPhysicsObjectHandle InObj
 }
 
 FScopedSceneLock_Chaos::FScopedSceneLock_Chaos(FChaosScene* InScene, EPhysicsInterfaceScopedLockType InLockType)
-	: Solver(InScene ? InScene->GetSolver() : nullptr)
-	, LockType(InLockType)
+	: Solver(InScene ? InScene->GetSolver() : nullptr),
+	LockType(InLockType)
 {
 	LockScene();
 }
@@ -152,6 +152,7 @@ FScopedSceneLock_Chaos& FScopedSceneLock_Chaos::operator=(FScopedSceneLock_Chaos
 {
 	Solver = Other.Solver;
 	LockType = Other.LockType;
+	ThreadContext = Other.ThreadContext;
 	bHasLock = Other.bHasLock;
 
 	Other.bHasLock = false;
@@ -174,7 +175,8 @@ void FScopedSceneLock_Chaos::Release()
 
 void FScopedSceneLock_Chaos::LockScene()
 {
-	if (!Solver)
+	Chaos::FPhysSceneLock* SceneLock = GetSolverLock(Solver);
+	if (!SceneLock)
 	{
 		return;
 	}
@@ -182,10 +184,10 @@ void FScopedSceneLock_Chaos::LockScene()
 	switch (LockType)
 	{
 	case EPhysicsInterfaceScopedLockType::Read:
-		Solver->GetExternalDataLock_External().ReadLock();
+		SceneLock->ReadLock();
 		break;
 	case EPhysicsInterfaceScopedLockType::Write:
-		Solver->GetExternalDataLock_External().WriteLock();
+		SceneLock->WriteLock();
 		break;
 	}
 
@@ -194,7 +196,8 @@ void FScopedSceneLock_Chaos::LockScene()
 
 void FScopedSceneLock_Chaos::UnlockScene()
 {
-	if (!Solver)
+	Chaos::FPhysSceneLock* SceneLock = GetSolverLock(Solver);
+	if (!SceneLock)
 	{
 		return;
 	}
@@ -202,10 +205,10 @@ void FScopedSceneLock_Chaos::UnlockScene()
 	switch (LockType)
 	{
 	case EPhysicsInterfaceScopedLockType::Read:
-		Solver->GetExternalDataLock_External().ReadUnlock();
+		SceneLock->ReadUnlock();
 		break;
 	case EPhysicsInterfaceScopedLockType::Write:
-		Solver->GetExternalDataLock_External().WriteUnlock();
+		SceneLock->WriteUnlock();
 		break;
 	}
 
@@ -237,3 +240,39 @@ FChaosScene* FScopedSceneLock_Chaos::GetSceneForActor(FPhysicsConstraintHandle c
 
 	return nullptr;
 }
+
+Chaos::FPhysSceneLock* FScopedSceneLock_Chaos::GetSolverLock(Chaos::FPBDRigidsSolver* InSolver)
+{
+	if (!InSolver)
+	{
+		return nullptr;
+	}
+
+#if UE_WITH_REMOTE_OBJECT_HANDLE
+	switch (ThreadContext)
+	{
+		case EPhysicsInterfaceScopedThreadContext::Internal:
+			return &InSolver->GetInternalDataLock();
+		case EPhysicsInterfaceScopedThreadContext::External:
+			return &InSolver->GetExternalDataLock_External();
+		default:
+			{
+				UE_LOG(LogChaos, Fatal, TEXT("Unsupported thread context used."))
+				return nullptr;
+			}
+	}
+#else
+	return &InSolver->GetExternalDataLock_External();
+#endif
+}
+
+FScopedSceneLockWithContext_Chaos::FScopedSceneLockWithContext_Chaos(const FPhysicsActorHandle& InActorHandle, EPhysicsInterfaceScopedLockType InLockType, EPhysicsInterfaceScopedThreadContext InThreadContext) : FScopedSceneLock_Chaos(InActorHandle, InLockType)
+{
+	ThreadContext = InThreadContext;
+}
+
+
+FScopedSceneLockWithContext_Chaos::FScopedSceneLockWithContext_Chaos(FPhysicsConstraintHandle const* InConstraintHandle, EPhysicsInterfaceScopedLockType InLockType, EPhysicsInterfaceScopedThreadContext InThreadContext) : FScopedSceneLock_Chaos(InConstraintHandle, InLockType)
+{
+	ThreadContext = InThreadContext;
+}
diff --git a/Engine/Source/Runtime/Engine/Private/PhysicsEngine/Experimental/PhysInterface_Chaos.cpp b/Engine/Source/Runtime/Engine/Private/PhysicsEngine/Experimental/PhysInterface_Chaos.cpp
--- a/Engine/Source/Runtime/Engine/Private/PhysicsEngine/Experimental/PhysInterface_Chaos.cpp
+++ b/Engine/Source/Runtime/Engine/Private/PhysicsEngine/Experimental/PhysInterface_Chaos.cpp
@@ -673,6 +673,55 @@ void FPhysInterface_Chaos::ExecuteShapeWrite(FBodyInstance* InInstance, FPhysics
 }
 
 
+bool FPhysInterface_Chaos::ExecuteWrite(FPhysicsActorHandle& InActorReference, TFunctionRef<void(FPhysicsActorHandle& Actor)> InCallable, EPhysicsInterfaceScopedThreadContext ThreadContext)
+{
+	if (InActorReference)
+	{
+		FScopedSceneLockWithContext_Chaos SceneLock(InActorReference, EPhysicsInterfaceScopedLockType::Write, ThreadContext);
+		InCallable(InActorReference);
+		return true;
+	}
+
+	return false;
+}
+
+
+bool FPhysInterface_Chaos::ExecuteRead(FPhysicsActorHandle& InActorReference, TFunctionRef<void(FPhysicsActorHandle& Actor)> InCallable, EPhysicsInterfaceScopedThreadContext ThreadContext)
+{
+	if (InActorReference)
+	{
+		FScopedSceneLockWithContext_Chaos SceneLock(InActorReference, EPhysicsInterfaceScopedLockType::Read, ThreadContext);
+		InCallable(InActorReference);
+		return true;
+	}
+
+	return false;
+}
+
+bool FPhysInterface_Chaos::ExecuteWrite(const FPhysicsConstraintHandle& InConstraintRef, TFunctionRef<void(const FPhysicsConstraintHandle& Constraint)> InCallable, EPhysicsInterfaceScopedThreadContext ThreadContext)
+{
+	if (InConstraintRef.IsValid())
+	{
+		FScopedSceneLockWithContext_Chaos SceneLock(&InConstraintRef, EPhysicsInterfaceScopedLockType::Write, ThreadContext);
+		InCallable(InConstraintRef);
+		return true;
+	}
+
+	return false;
+}
+
+bool FPhysInterface_Chaos::ExecuteRead(const FPhysicsConstraintHandle& InConstraintRef, TFunctionRef<void(const FPhysicsConstraintHandle& Constraint)> InCallable, EPhysicsInterfaceScopedThreadContext ThreadContext)
+{
+	if (InConstraintRef.IsValid())
+	{
+		FScopedSceneLockWithContext_Chaos SceneLock(&InConstraintRef, EPhysicsInterfaceScopedLockType::Read, ThreadContext);
+		InCallable(InConstraintRef);
+		return true;
+	}
+
+	return false;
+}
+
 FPhysicsShapeHandle FPhysInterface_Chaos::CreateShape(physx::PxGeometry* InGeom, bool bSimulation, bool bQuery, UPhysicalMaterial* InSimpleMaterial, TArray<UPhysicalMaterial*>* InComplexMaterials)
 {
 	// #todo : Implement
diff --git a/Engine/Source/Runtime/Engine/Public/Physics/Experimental/ChaosScopedSceneLock.h b/Engine/Source/Runtime/Engine/Public/Physics/Experimental/ChaosScopedSceneLock.h
--- a/Engine/Source/Runtime/Engine/Public/Physics/Experimental/ChaosScopedSceneLock.h
+++ b/Engine/Source/Runtime/Engine/Public/Physics/Experimental/ChaosScopedSceneLock.h
@@ -4,6 +4,7 @@
 
 #include "PhysicsInterfaceDeclaresCore.h"
 #include "Chaos/PhysicsObject.h"
+#include "Framework/Threading.h"
 
 enum class EPhysicsInterfaceScopedLockType : uint8
 {
@@ -11,6 +12,12 @@ enum class EPhysicsInterfaceScopedLockType : uint8
 	Write
 };
 
+enum class EPhysicsInterfaceScopedThreadContext : uint8
+{
+	External,
+	Internal
+};
+
 class USkeletalMeshComponent;
 class FChaosScene;
 
@@ -48,7 +55,19 @@ private:
 	ENGINE_API FChaosScene* GetSceneForActor(const FPhysicsActorHandle& InActorHandle);
 	ENGINE_API FChaosScene* GetSceneForActor(FPhysicsConstraintHandle const* InConstraintHandle);
 
+	Chaos::FPhysSceneLock* GetSolverLock(Chaos::FPBDRigidsSolver* InSolver);
+
 	bool bHasLock = false;
 	Chaos::FPBDRigidsSolver* Solver;
 	EPhysicsInterfaceScopedLockType LockType;
+
+protected:
+	EPhysicsInterfaceScopedThreadContext ThreadContext = EPhysicsInterfaceScopedThreadContext::External;
 };
+
+struct UE_INTERNAL FScopedSceneLockWithContext_Chaos : FScopedSceneLock_Chaos
+{
+	ENGINE_API FScopedSceneLockWithContext_Chaos(const FPhysicsActorHandle& InActorHandle, EPhysicsInterfaceScopedLockType InLockType, EPhysicsInterfaceScopedThreadContext InThreadContext);
+	ENGINE_API FScopedSceneLockWithContext_Chaos(FPhysicsConstraintHandle const* InConstraintHandle, EPhysicsInterfaceScopedLockType InLockType, EPhysicsInterfaceScopedThreadContext InThreadContext);
+};
+
diff --git a/Engine/Source/Runtime/Engine/Public/Physics/Experimental/PhysInterface_Chaos.h b/Engine/Source/Runtime/Engine/Public/Physics/Experimental/PhysInterface_Chaos.h
--- a/Engine/Source/Runtime/Engine/Public/Physics/Experimental/PhysInterface_Chaos.h
+++ b/Engine/Source/Runtime/Engine/Public/Physics/Experimental/PhysInterface_Chaos.h
@@ -15,6 +15,7 @@
 
 //NOTE: Do not include Chaos headers directly as it means recompiling all of engine. This should be reworked to avoid allocations
 
+enum class EPhysicsInterfaceScopedThreadContext : uint8;
 static int32 NextBodyIdValue = 0;
 static int32 NextConstraintIdValue = 0;
 
@@ -86,6 +87,11 @@ public:
 
     static ENGINE_API void ExecuteShapeWrite(FBodyInstance* InInstance, FPhysicsShapeHandle& InShape, TFunctionRef<void(FPhysicsShapeHandle& InShape)> InCallable);
 
+	UE_INTERNAL static ENGINE_API bool ExecuteWrite(FPhysicsActorHandle& InActorReference, TFunctionRef<void(FPhysicsActorHandle& Actor)> InCallable, EPhysicsInterfaceScopedThreadContext ThreadContext);
+	UE_INTERNAL static ENGINE_API bool ExecuteRead(FPhysicsActorHandle& InActorReference, TFunctionRef<void(FPhysicsActorHandle& Actor)> InCallable, EPhysicsInterfaceScopedThreadContext ThreadContext);
+	UE_INTERNAL static ENGINE_API bool ExecuteWrite(const FPhysicsConstraintHandle& InConstraintRef, TFunctionRef<void(const FPhysicsConstraintHandle& Constraint)> InCallable, EPhysicsInterfaceScopedThreadContext ThreadContext);
+	UE_INTERNAL static ENGINE_API bool ExecuteRead(const FPhysicsConstraintHandle& InConstraintRef, TFunctionRef<void(const FPhysicsConstraintHandle& Constraint)> InCallable, EPhysicsInterfaceScopedThreadContext ThreadContext);
+
 	// Misc
 
 	static ENGINE_API bool ExecPhysCommands(const TCHAR* Cmd, FOutputDevice* Ar, UWorld* InWorld);
diff --git a/Engine/Source/Runtime/Experimental/Chaos/Private/Chaos/Framework/PhysicsSolverBase.cpp b/Engine/Source/Runtime/Experimental/Chaos/Private/Chaos/Framework/PhysicsSolverBase.cpp
--- a/Engine/Source/Runtime/Experimental/Chaos/Private/Chaos/Framework/PhysicsSolverBase.cpp
+++ b/Engine/Source/Runtime/Experimental/Chaos/Private/Chaos/Framework/PhysicsSolverBase.cpp
@@ -166,6 +166,10 @@ namespace Chaos
 		FPhysicsThreadContextScope Scope(/*IsPhysicsThreadContext=*/true);
 #endif
 
+#if UE_WITH_REMOTE_OBJECT_HANDLE
+		FPhysicsSceneGuardScopedWrite ScopedWrite(Solver.GetInternalDataLock());
+#endif
+
 		// StepFraction: how much of the remaining time this step represents while substepping, used to interpolate kinematic targets
 		// E.g., for 4 steps this will be: 1/4, 1/3, 1/2, 1
 		const FReal PseudoFraction = PushData->bSolverSubstepped ? (FReal)1 / (FReal)(PushData->IntervalNumSteps - PushData->IntervalStep) : 1.0f;
@@ -233,6 +237,9 @@ namespace Chaos
 		, bPaused_External(false)
 		, Owner(InOwner)
 		, ExternalDataLock_External(new FPhysSceneLock())
+#if UE_WITH_REMOTE_OBJECT_HANDLE
+		, InternalDataLock(new FPhysSceneLock())
+#endif
 		, bIsShuttingDown(false)
 		, AsyncDt(InAsyncDt)
 		, AccumulatedTime(0)
diff --git a/Engine/Source/Runtime/Experimental/Chaos/Private/Chaos/PBDJointConstraintData.cpp b/Engine/Source/Runtime/Experimental/Chaos/Private/Chaos/PBDJointConstraintData.cpp
--- a/Engine/Source/Runtime/Experimental/Chaos/Private/Chaos/PBDJointConstraintData.cpp
+++ b/Engine/Source/Runtime/Experimental/Chaos/Private/Chaos/PBDJointConstraintData.cpp
@@ -13,6 +13,14 @@ namespace Chaos
 	{
 	}
 
+	void FJointConstraint::SetJointSettings(const FPBDJointSettings& NewSettings)
+	{
+		JointSettings.Modify(/*bInvalidate=*/true, DirtyFlags, Proxy, [&NewSettings](auto& Data)
+		{
+			Data = NewSettings;
+		});
+	}
+
 	void FJointConstraint::SetKinematicEndPoint(FSingleParticlePhysicsProxy* InDummyParticle, FPBDRigidsSolver* Solver)
 	{
 		ensure(KinematicEndPoint == nullptr);
diff --git a/Engine/Source/Runtime/Experimental/Chaos/Private/Chaos/Serialization/SerializedMultiPhysicsState.cpp b/Engine/Source/Runtime/Experimental/Chaos/Private/Chaos/Serialization/SerializedMultiPhysicsState.cpp
new file mode 100644
--- /dev/null
+++ b/Engine/Source/Runtime/Experimental/Chaos/Private/Chaos/Serialization/SerializedMultiPhysicsState.cpp
@@ -0,0 +1,44 @@
+// Copyright Epic Games, Inc. All Rights Reserved.
+
+#include "Chaos/Serialization/SerializedMultiPhysicsState.h"
+
+#include "ChaosLog.h"
+
+namespace Chaos
+{
+	void FSerializedMultiPhysicsState::ReadElementDataIntoBuffer(FSerializedDataBuffer& TargetBuffer)
+	{
+		FDataEntryTag DataTag = Header.DataTagPerElementIndex.IsValidIndex(CurrentReadElementIndex) ? Header.DataTagPerElementIndex[CurrentReadElementIndex] : FDataEntryTag();
+
+		TArray<uint8>& AsByteArray = MigratedStateAsBytes.GetDataAsByteArrayRef();
+
+		const int32 BufferSize = GetSize();
+
+		UE_LOG(LogChaos, Verbose, TEXT("[MultiStateSerialization] Attempted to read an empty data entry. This might means we failed to obtain data to serialize a specific body or constraint."));
+
+		// If we have entries past the available buffer size, but they were recorded as empty, the serialization code is ok to continue (it will just be a no-op)
+		const bool bHasValidStartingOffset = DataTag.DataOffset < BufferSize || DataTag.DataSize == 0;
+		const bool bThereIsEnoughDataToCopyFromBuffer = DataTag.DataOffset + DataTag.DataSize <= BufferSize;
+
+		if (ensure(DataTag.IsValid() && bHasValidStartingOffset && bThereIsEnoughDataToCopyFromBuffer))
+		{
+			if (DataTag.DataSize > 0)
+			{
+				TargetBuffer.GetDataAsByteArrayRef().Append(&AsByteArray[DataTag.DataOffset], DataTag.DataSize);
+			}
+		}
+		else
+		{
+			UE_LOG(LogChaos, Error, TEXT("[MultiStateSerialization] Attempted to read data out of bounds! | Buffer size [%d] | Start Pos [%d]| SizeToCopy [%d]"), AsByteArray.Num(), DataTag.DataOffset, DataTag.DataSize);
+		}
+		
+		CurrentReadElementIndex++;
+	}
+
+
+	void FSerializedMultiPhysicsState::Serialize(FArchive& Ar)
+	{
+		Ar << Header;
+		Ar << MigratedStateAsBytes;
+	}
+}
diff --git a/Engine/Source/Runtime/Experimental/Chaos/Private/Chaos/Serialization/SerializedSolverScene.cpp b/Engine/Source/Runtime/Experimental/Chaos/Private/Chaos/Serialization/SerializedSolverScene.cpp
new file mode 100644
--- /dev/null
+++ b/Engine/Source/Runtime/Experimental/Chaos/Private/Chaos/Serialization/SerializedSolverScene.cpp
@@ -0,0 +1,3 @@
+// Copyright Epic Games, Inc. All Rights Reserved.
+
+#include "SerializedSolverScene.h"
diff --git a/Engine/Source/Runtime/Experimental/Chaos/Private/Chaos/Serialization/SerializedSolverScene.h b/Engine/Source/Runtime/Experimental/Chaos/Private/Chaos/Serialization/SerializedSolverScene.h
new file mode 100644
--- /dev/null
+++ b/Engine/Source/Runtime/Experimental/Chaos/Private/Chaos/Serialization/SerializedSolverScene.h
@@ -0,0 +1,40 @@
+// Copyright Epic Games, Inc. All Rights Reserved.
+
+#pragma once
+#include "Chaos/ImplicitFwd.h"
+#include "DataWrappers/ChaosVDImplicitObjectDataWrapper.h"
+#include "SerializedSolverScene.generated.h"
+
+namespace Chaos
+{
+	class FChaosArchive;
+}
+
+struct FChaosVDParticlePairMidPhase;
+struct FChaosVDCharacterGroundConstraint;
+struct FChaosVDCharacterGroundConstraintDataDataWrapper;
+struct FChaosVDJointConstraint;
+struct FChaosVDParticleDataWrapper;
+
+using FSerializedImplicitObject = FChaosVDImplicitObjectDataWrapper<Chaos::FImplicitObjectPtr, Chaos::FChaosArchive>;
+
+/**
+ * Structure holding the serialized state of a rigid solver
+ * Used to save it to disk or to re-hydrate a new solver instance
+ */
+USTRUCT()
+struct UE_INTERNAL FSerializedSolverScene
+{
+	GENERATED_BODY()
+
+	UPROPERTY()
+	TArray<FChaosVDParticleDataWrapper> ParticleData;
+	UPROPERTY()
+	TArray<FChaosVDJointConstraint> JointConstraintData;
+	UPROPERTY()
+	TArray<FChaosVDCharacterGroundConstraint> CharacterGroundConstraintData;
+	UPROPERTY()
+	TArray<FChaosVDParticlePairMidPhase> CollisionMidPhaseData;
+	
+	TArray<FSerializedImplicitObject> ImplicitObjectData;
+};
\ No newline at end of file
diff --git a/Engine/Source/Runtime/Experimental/Chaos/Private/Chaos/Serialization/SolverSerializer.cpp b/Engine/Source/Runtime/Experimental/Chaos/Private/Chaos/Serialization/SolverSerializer.cpp
--- a/Engine/Source/Runtime/Experimental/Chaos/Private/Chaos/Serialization/SolverSerializer.cpp
+++ b/Engine/Source/Runtime/Experimental/Chaos/Private/Chaos/Serialization/SolverSerializer.cpp
@@ -4,6 +4,7 @@
 
 #include "PBDRigidsSolver.h"
 #include "Chaos/PBDJointConstraints.h"
+#include "Chaos/Character/CharacterGroundConstraint.h"
 #include "Chaos/Character/CharacterGroundConstraintContainer.h"
 #include "Chaos/Serialization/SerializationUtils.h"
 #include "Chaos/Serialization/SerializedDataBuffer.h"
@@ -14,6 +15,96 @@
 
 namespace Chaos
 {
+	namespace Private
+	{
+		FChaosVDParticleDataWrapper ExtractParticleDataFromBuffer(FSerializedDataBuffer& InSerializedData)
+		{
+			FChaosVDParticleDataWrapper ParticleDataWrapper;
+	
+			using namespace Chaos::Serialization::Private;
+
+			using namespace Chaos::Serialization::Private;
+			if (!ensure(!InSerializedData.GetDataAsByteArrayRef().IsEmpty()))
+			{
+				return MoveTemp(ParticleDataWrapper);
+			}
+
+			FMemoryReader MemReader(InSerializedData.GetDataAsByteArrayRef());
+			MemReader.SetShouldSkipUpdateCustomVersion(true);
+		
+			FastStructSerialize(MemReader, &ParticleDataWrapper);
+
+			return MoveTemp(ParticleDataWrapper);
+		}
+
+		template<typename ParticleType>
+		void WriteParticleDataToBuffer(ParticleType* InParticle, FSerializedDataBuffer& OutSerializedData)
+		{
+			using namespace Chaos::Serialization::Private;
+
+			if (!ensure(InParticle))
+			{
+				UE_LOG(LogChaos, Warning, TEXT("[%hs]: Failed to serialize particle | Invalid Particle"), __func__);
+				return;
+			}
+
+			// This allows us to serialize multiple particles in a single buffer, if the provided buffer was already used
+			constexpr bool bSetOffset = true;
+			constexpr bool bIsPersistent = false;
+
+			FMemoryWriter MemWriter(OutSerializedData.GetDataAsByteArrayRef(), bIsPersistent, bSetOffset);
+			MemWriter.SetShouldSkipUpdateCustomVersion(true);
+
+			FChaosVDParticleDataWrapper ParticleDataWrapper;
+
+			static_assert(std::is_base_of_v<FGeometryParticleHandle, ParticleType> || std::is_base_of_v<FGeometryParticle, ParticleType>, "WriteParticleDataToBuffer Only supports FGeometryParticleHandle and FGeometryParticle types.");
+
+			// TODO: Handle Geometry and shape data serialization for the case that need it
+			if constexpr(std::is_same_v<FGeometryParticleHandle, ParticleType>)
+			{
+				 ParticleDataWrapper = FChaosVDDataWrapperUtils::BuildParticleDataWrapperFromParticle(InParticle);
+			}
+			else if (std::is_same_v<FGeometryParticle, ParticleType>)
+			{
+				ParticleDataWrapper = FChaosVDDataWrapperUtils::BuildParticleDataWrapperFromGTParticle(InParticle);
+			}
+
+			FastStructSerialize(MemWriter, &ParticleDataWrapper);
+		}
+
+		FChaosVDJointConstraint ExtractJointConstraintDataFromBuffer(FSerializedDataBuffer& InSerializedData)
+		{
+			FChaosVDJointConstraint JointData;
+
+			if (!ensure(!InSerializedData.GetDataAsByteArrayRef().IsEmpty()))
+			{
+				return MoveTemp(JointData);
+			}
+		
+			FMemoryReader Reader(InSerializedData.GetDataAsByteArrayRef());
+			Reader.SetShouldSkipUpdateCustomVersion(true);
+			Serialization::Private::FastStructSerialize(Reader, &JointData);
+
+			return MoveTemp(JointData);
+		}
+	}
+
+	const TCHAR* LexToString(ESerializedDataContext Value)
+	{
+		switch (Value)
+		{
+			case ESerializedDataContext::Internal:
+				return TEXT("ESerializedDataContext:Internal");
+			case ESerializedDataContext::External:
+				return TEXT("ESerializedDataContext:External");
+			case ESerializedDataContext::Both:
+				return TEXT("ESerializedDataContext:Both");
+			case ESerializedDataContext::Invalid:
+			default:
+				return TEXT("ESerializedDataContext:Invalid");
+		}
+	}
+
 	FSolverSerializer::FSolverSerializer(FPBDRigidsSolver* InSolver) : SolverInstance(InSolver)
 	{
 	}
@@ -30,25 +121,12 @@ namespace Chaos
 
 	void FSolverSerializer::SerializeParticleStateToBuffer(FGeometryParticleHandle* InParticleHandle, FSerializedDataBuffer& OutSerializedData)
 	{
-		using namespace Chaos::Serialization::Private;
-
-		if (!ensure(InParticleHandle))
-		{
-			UE_LOG(LogChaos, Warning, TEXT("[%hs]: Failed to serialize particle | Invalid Handle"), __func__);
-			return;
-		}
-
-		// This allows us to serialize multiple particles in a single buffer, if the provided buffer was already used
-		constexpr bool bSetOffset = true;
-		constexpr bool bIsPersistent = false;
-
-		FMemoryWriter MemWriter(OutSerializedData.GetDataAsByteArrayRef(), bIsPersistent, bSetOffset);
-		MemWriter.SetShouldSkipUpdateCustomVersion(true);
-
-		// TODO: Handle Geometry and shape data serialization for the case that need it
-		FChaosVDParticleDataWrapper ParticleDataWrapper = FChaosVDDataWrapperUtils::BuildParticleDataWrapperFromParticle(InParticleHandle);
+		Private::WriteParticleDataToBuffer(InParticleHandle, OutSerializedData);
+	}
 
-		FastStructSerialize(MemWriter, &ParticleDataWrapper);
+	void FSolverSerializer::SerializeParticleStateToBuffer(FGeometryParticle* GTParticle, FSerializedDataBuffer& OutSerializedData)
+	{
+		Private::WriteParticleDataToBuffer(GTParticle, OutSerializedData);
 	}
 
 	void FSolverSerializer::SerializeConstraintStateToBuffer(FConstraintHandle* ConstraintHandlePtr, FSerializedDataBuffer& OutSerializedData)
@@ -84,27 +162,92 @@ namespace Chaos
 		}
 	}
 
-	void FSolverSerializer::ApplySerializedStateToParticle(FGeometryParticleHandle* InParticleHandle, FSerializedDataBuffer& InSerializedData)
+
+	void FSolverSerializer::SerializeConstraintStateToBuffer(FConstraintBase* GTConstraintPtr, FSerializedDataBuffer& OutSerializedData)
 	{
 		using namespace Chaos::Serialization::Private;
-		if (!ensure(!InSerializedData.GetDataAsByteArrayRef().IsEmpty()))
+
+		if (!ensure(GTConstraintPtr))
 		{
+			UE_LOG(LogChaos, Warning, TEXT("[%hs]: Failed to serialize constraint | Invalid Constraint"), __func__);
 			return;
 		}
 
-		FMemoryReader MemReader(InSerializedData.GetDataAsByteArrayRef());
-		MemReader.SetShouldSkipUpdateCustomVersion(true);
-		
-		FChaosVDParticleDataWrapper ParticleDataWrapper;
+		// This allows us to serialize multiple particles in a single buffer, if the provided buffer was already used
+		constexpr bool bSetOffset = true;
+		constexpr bool bIsPersistent = false;
 
-		FastStructSerialize(MemReader, &ParticleDataWrapper);
+		FMemoryWriter MemWriter(OutSerializedData.GetDataAsByteArrayRef(), bIsPersistent, bSetOffset);
+		MemWriter.SetShouldSkipUpdateCustomVersion(true);
+		switch (EConstraintType ConstraintType = GTConstraintPtr->GetType())
+		{
+			case EConstraintType::JointConstraintType:
+				{
+					FJointConstraint* AsJoint = static_cast<FJointConstraint*>(GTConstraintPtr);	
+					FChaosVDJointConstraint JointData = FChaosVDDataWrapperUtils::BuildGTJointDataWrapper(AsJoint);
+
+					FastStructSerialize(MemWriter, &JointData);
+					break;
+				}
+			case EConstraintType::SpringConstraintType:
+			case EConstraintType::SuspensionConstraintType:
+			case EConstraintType::CharacterGroundConstraintType:
+			case EConstraintType::NoneType:
+			default:
+				{
+					ensureMsgf(false, TEXT("Attempted to Serialize a Constraint type not supported yet | Constraint Type [%d]"), ConstraintType);
+					break;
+				}
+		}
+	}
+
+
+	void FSolverSerializer::ApplySerializedStateToParticle(FGeometryParticleHandle* InParticleHandle, FSerializedDataBuffer& InSerializedData)
+	{
+		FChaosVDParticleDataWrapper ParticleDataWrapper = Private::ExtractParticleDataFromBuffer(InSerializedData);
 
 		ApplySerializedStateToParticle(InParticleHandle, ParticleDataWrapper);
 	}
 
+	void FSolverSerializer::ApplySerializedStateToParticle(FGeometryParticle* GTParticle, FSerializedDataBuffer& InSerializedData)
+	{
+		FChaosVDParticleDataWrapper ParticleDataWrapper = Private::ExtractParticleDataFromBuffer(InSerializedData);
+	
+		ApplySerializedStateToParticle(GTParticle, ParticleDataWrapper);
+	}
+
+	void FSolverSerializer::ApplySerializedStateToParticle(FGeometryParticle* GTParticle, const FChaosVDParticleDataWrapper& InParticleState)
+	{
+		if (!ensure(InParticleState.HasValidData() && GTParticle && GTParticle->ObjectType() == static_cast<EParticleType>(InParticleState.Type)))
+		{
+			return;
+		}
+
+		InParticleState.ParticlePositionRotation.CopyTo<FGeometryParticle, FChaosVDParticlePositionRotation::EAccessorType::XR>(*GTParticle);
+
+		if (TKinematicGeometryParticle<FReal, 3>* KinematicParticle = GTParticle->CastToKinematicParticle())
+		{
+			FKinematicTarget NewKT;
+			InParticleState.ParticleKinematicTarget.CopyTo<FKinematicTarget, EKinematicTargetMode>(NewKT);
+			KinematicParticle->SetKinematicTarget(NewKT);
+			
+			InParticleState.ParticleVelocities.CopyTo(*KinematicParticle);
+		}
+
+		if (TPBDRigidParticle<FReal, 3>* RigidParticle = GTParticle->CastToRigidParticle())
+		{
+			InParticleState.ParticleDynamics.CopyTo(*RigidParticle);
+			InParticleState.ParticleMassProps.CopyTo(*RigidParticle);
+			InParticleState.ParticleDynamicsMisc.CopyWithoutStateTo<TPBDRigidParticle<FReal, 3>, FRigidParticleControlFlags, ESleepType>(*RigidParticle);
+			RigidParticle->SetObjectState(static_cast<EObjectStateType>(InParticleState.ParticleDynamicsMisc.MObjectState));
+		}
+
+		// TODO: Add support for Cluster unions and Geometry Collections
+	}
+
 	void FSolverSerializer::ApplySerializedStateToParticle(FGeometryParticleHandle* InParticleHandle, const FChaosVDParticleDataWrapper& InParticleState)
 	{
-		if (!ensure(InParticleHandle->Type == static_cast<EParticleType>(InParticleState.Type)))
+		if (!ensure(InParticleState.HasValidData() && InParticleHandle && InParticleHandle->Type == static_cast<EParticleType>(InParticleState.Type)))
 		{
 			return;
 		}
@@ -120,14 +263,22 @@ namespace Chaos
 
 		Evolution.SetParticleVelocities(InParticleHandle, InParticleState.ParticleVelocities.MV, InParticleState.ParticleVelocities.MW);
 
+		if (TKinematicGeometryParticleHandleImp<FReal, 3, true>* KinematicParticle = InParticleHandle->CastToKinematicParticle())
+		{
+			FKinematicTarget NewKT;
+			InParticleState.ParticleKinematicTarget.CopyTo<FKinematicTarget, EKinematicTargetMode>(NewKT);
+			KinematicParticle->SetKinematicTarget(NewKT);
+		}
+
 		if (TPBDRigidParticleHandleImp<FReal, 3, true>* RigidParticle = InParticleHandle->CastToRigidParticle())
 		{
 			InParticleState.ParticleDynamics.CopyTo(*RigidParticle);
 			InParticleState.ParticleMassProps.CopyTo(*RigidParticle);
-			InParticleState.ParticleDynamicsMisc.CopyTo<TPBDRigidParticleHandleImp<FReal, 3, true>, FRigidParticleControlFlags, EObjectStateType, ESleepType>(*RigidParticle);
+			InParticleState.ParticleDynamicsMisc.CopyWithoutStateTo<TPBDRigidParticleHandleImp<FReal, 3, true>, FRigidParticleControlFlags, ESleepType>(*RigidParticle);
+			InParticleState.ParticleVWSmooth.CopyTo<TPBDRigidParticleHandleImp<FReal, 3, true>>(*RigidParticle);
 
 			// Make sure the particle ends up in the correct SoA view
-			Evolution.SetParticleObjectState(RigidParticle, RigidParticle->ObjectState());
+			Evolution.SetParticleObjectState(RigidParticle, static_cast<EObjectStateType>(InParticleState.ParticleDynamicsMisc.MObjectState));
 			Evolution.SetParticleSleepType(RigidParticle, RigidParticle->SleepType());
 		}
 
@@ -146,31 +297,68 @@ namespace Chaos
 		}
 	}
 
-	void FSolverSerializer::ApplySerializedStateToJointConstraint(Chaos::FPBDJointConstraintHandle* ConstraintHandlePtr, FSerializedDataBuffer& InSerializedData)
+	void FSolverSerializer::ApplySerializedStateToJointConstraint(FPBDJointConstraintHandle* ConstraintHandlePtr, FSerializedDataBuffer& InSerializedData)
 	{
-		using namespace Chaos::Serialization::Private;
+		FChaosVDDataWrapperUtils::ApplyJointDataWrapperToHandle(ConstraintHandlePtr, Private::ExtractJointConstraintDataFromBuffer(InSerializedData));	
+	}
+
+	void FSolverSerializer::ApplySerializedStateToJointConstraint(FJointConstraint* ConstraintPtr, FSerializedDataBuffer& InSerializedData)
+	{
+		FChaosVDDataWrapperUtils::ApplyJointDataWrapperGTConstraint(ConstraintPtr, Private::ExtractJointConstraintDataFromBuffer(InSerializedData));
+	}
 
-		if (!ConstraintHandlePtr)
+	void FSolverSerializer::ApplySerializedStateToConstraint(FConstraintHandle* ConstraintHandlePtr, FSerializedDataBuffer& InSerializedData)
+	{
+		if (!ensure(ConstraintHandlePtr))
 		{
 			return;
 		}
 
-		FMemoryReader Reader(InSerializedData.GetDataAsByteArrayRef());
-		Reader.SetShouldSkipUpdateCustomVersion(true);
-		FChaosVDJointConstraint JointData;
-		FastStructSerialize(Reader, &JointData);
-	
-		FChaosVDDataWrapperUtils::ApplyJointDataWrapperToHandle(ConstraintHandlePtr, JointData);	
+		if (ConstraintHandlePtr->GetType().IsA(FPBDJointConstraintHandle::StaticType()))
+		{
+			FPBDJointConstraintHandle* AsJointHandle = static_cast<FPBDJointConstraintHandle*>(ConstraintHandlePtr);
+			ApplySerializedStateToJointConstraint(AsJointHandle, InSerializedData);
+		}
+		else
+		{
+			UE_LOG(LogChaos, Warning, TEXT("Attempted to apply a serialized state to an unsupported constraint type | Type [%s]. The data will be discarded"), *ConstraintHandlePtr->GetType().ToString())
+		}
+	}
+
+	void FSolverSerializer::ApplySerializedStateToConstraint(FConstraintBase* ConstraintPtr, FSerializedDataBuffer& InSerializedData)
+	{
+		if (!ConstraintPtr)
+		{
+			return;
+		}
+
+		switch (EConstraintType Type = ConstraintPtr->GetType())
+		{
+			case EConstraintType::JointConstraintType:
+				{
+					ApplySerializedStateToJointConstraint(static_cast<FJointConstraint*>(ConstraintPtr), InSerializedData);
+					break;
+				}
+			case EConstraintType::NoneType:
+			case EConstraintType::SpringConstraintType:
+			case EConstraintType::SuspensionConstraintType:
+			case EConstraintType::CharacterGroundConstraintType:
+			default:
+				{
+					ensureMsgf(false, TEXT("Attempted to serialize an unsupported constraint type | Type [%d]"), Type);
+					break;
+				}
+		}
 	}
 
-	void FSolverSerializer::PushPendingSerializedStateForProxy(IPhysicsProxyBase* Proxy, FSerializedDataBufferPtr&& InState)
+	void FSolverSerializer::PushPendingInternalSerializedStateForProxy(IPhysicsProxyBase* Proxy, FSerializedDataBufferPtr&& InState)
 	{
 		PendingMigratedPhysicsStateByProxy.Emplace(Proxy, MoveTemp(InState));
 	}
 
-	FSerializedDataBufferPtr FSolverSerializer::PopPendingSerializedStateForProxy(IPhysicsProxyBase* Proxy)
+	FSerializedDataBufferPtr FSolverSerializer::PopPendingInternalSerializedStateForProxy(IPhysicsProxyBase* Proxy)
 	{
-		if (PendingMigratedPhysicsStateByProxy.IsEmpty())
+		if (PendingMigratedPhysicsStateByProxy.Num() == 0)
 		{
 			return nullptr;
 		}
diff --git a/Engine/Source/Runtime/Experimental/Chaos/Private/ChaosVisualDebugger/ChaosVDDataWrapperUtils.cpp b/Engine/Source/Runtime/Experimental/Chaos/Private/ChaosVisualDebugger/ChaosVDDataWrapperUtils.cpp
--- a/Engine/Source/Runtime/Experimental/Chaos/Private/ChaosVisualDebugger/ChaosVDDataWrapperUtils.cpp
+++ b/Engine/Source/Runtime/Experimental/Chaos/Private/ChaosVisualDebugger/ChaosVDDataWrapperUtils.cpp
@@ -91,6 +91,8 @@ FChaosVDParticleDataWrapper FChaosVDDataWrapperUtils::BuildParticleDataWrapperFr
 
 	FChaosVDParticleDataWrapper WrappedParticleData;
 
+	WrappedParticleData.ParticleContext = EChaosVDParticleContext::PhysicsThread;
+
 	WrappedParticleData.ParticleIndex = ParticleHandlePtr->UniqueIdx().Idx;
 	WrappedParticleData.Type =  static_cast<EChaosVDParticleType>(ParticleHandlePtr->Type);
 
@@ -101,19 +103,24 @@ FChaosVDParticleDataWrapper FChaosVDDataWrapperUtils::BuildParticleDataWrapperFr
 	WrappedParticleData.DebugNamePtr = ParticleHandlePtr->DebugName().IsValid() ? ParticleHandlePtr->DebugName().Get() : nullptr;
 #endif
 
-	WrappedParticleData.ParticlePositionRotation.CopyFrom(Chaos::FConstGenericParticleHandle(ParticleHandlePtr));
-
 	if (const Chaos::TKinematicGeometryParticleHandleImp<Chaos::FReal, 3, true>* KinematicParticle = ParticleHandlePtr->CastToKinematicParticle())
 	{
 		WrappedParticleData.ParticleVelocities.CopyFrom(*KinematicParticle);
+		WrappedParticleData.ParticleKinematicTarget.CopyFrom(KinematicParticle->KinematicTarget());
 	}
 
 	if (const Chaos::TPBDRigidParticleHandleImp<Chaos::FReal, 3, true>* RigidParticle = ParticleHandlePtr->CastToRigidParticle())
 	{
+		WrappedParticleData.ParticlePositionRotation.CopyFrom<Chaos::TPBDRigidParticleHandleImp<Chaos::FReal, 3, true>, FChaosVDParticlePositionRotation::EAccessorType::PQ>(*RigidParticle);
+		WrappedParticleData.ParticleVWSmooth.CopyFrom(*RigidParticle);
 		WrappedParticleData.ParticleDynamics.CopyFrom(*RigidParticle);
 		WrappedParticleData.ParticleDynamicsMisc.CopyFrom(*RigidParticle);
 		WrappedParticleData.ParticleMassProps.CopyFrom(*RigidParticle);
 	}
+	else
+	{
+		WrappedParticleData.ParticlePositionRotation.CopyFrom<Chaos::FGeometryParticleHandle, FChaosVDParticlePositionRotation::EAccessorType::XR>(*ParticleHandlePtr);
+	}
 
 	if (const Chaos::TPBDRigidClusteredParticleHandleImp<Chaos::FReal, 3, true>* ClusteredParticle = ParticleHandlePtr->CastToClustered())
 	{
@@ -125,6 +132,46 @@ FChaosVDParticleDataWrapper FChaosVDDataWrapperUtils::BuildParticleDataWrapperFr
 	return MoveTemp(WrappedParticleData);
 }
 
+FChaosVDParticleDataWrapper FChaosVDDataWrapperUtils::BuildParticleDataWrapperFromGTParticle(const Chaos::FGeometryParticle* GTParticle)
+{
+	FChaosVDParticleDataWrapper WrappedParticleData;
+
+	WrappedParticleData.ParticleContext = EChaosVDParticleContext::GameThread;
+
+	WrappedParticleData.ParticleIndex = GTParticle->UniqueIdx().Idx;
+	WrappedParticleData.Type =  static_cast<EChaosVDParticleType>(GTParticle->ObjectType());
+	
+#if CHAOS_DEBUG_NAME
+	// Passing it as a Ptr because from here until it is serialized right after this function ends this string does not change
+	// Passing it as a sharedptr has an additional 20% cost as it has to increment the reference counter, which adds up
+	// TODO: We should switch to FName to take advantage of the new CVD Serializable name table so they can be de-duplicated, but to do so we need to change how we create our debug names to not be unique strings
+	WrappedParticleData.DebugNamePtr = GTParticle->DebugName().Get();
+#endif
+
+	WrappedParticleData.ParticlePositionRotation.CopyFrom<Chaos::FGeometryParticle, FChaosVDParticlePositionRotation::EAccessorType::XR>(*GTParticle);
+
+	if (const Chaos::TKinematicGeometryParticle<Chaos::FReal, 3>* KinematicParticle = GTParticle->CastToKinematicParticle())
+	{
+		WrappedParticleData.ParticleVelocities.CopyFrom(*KinematicParticle);
+		WrappedParticleData.ParticleKinematicTarget.CopyFrom(KinematicParticle->KinematicTarget());
+	}
+
+	if (const Chaos::TPBDRigidParticle<Chaos::FReal, 3>* RigidParticle = GTParticle->CastToRigidParticle())
+	{
+		WrappedParticleData.ParticleDynamics.CopyFrom(*RigidParticle);
+		WrappedParticleData.ParticleDynamicsMisc.CopyFrom(*RigidParticle);
+		WrappedParticleData.ParticleMassProps.CopyFrom(*RigidParticle);
+	}
+
+	WrappedParticleData.DirtyFlagsBits = GTParticle->DirtyFlags().GetBits();
+
+	// TODO: Add support for Cluster unions and Geometry Collections
+
+	WrappedParticleData.MarkAsValid();
+
+	return MoveTemp(WrappedParticleData);
+}
+
 FChaosVDConstraint FChaosVDDataWrapperUtils::BuildConstraintDataWrapperFromConstraint(const Chaos::FPBDCollisionConstraint& InConstraint)
 {
 	FChaosVDConstraint WrappedConstraintData;
@@ -240,6 +287,93 @@ FChaosVDParticlePairMidPhase FChaosVDDataWrapperUtils::BuildMidPhaseDataWrapperF
 	Target.Field = FChaosVDDataWrapperUtils::ConvertToFVector(Source.Field);
 #endif
 
+void FChaosVDDataWrapperUtils::CopyJointConstraintSettingsToWrapper(const Chaos::FPBDJointSettings& InSettings, FChaosVDJointSettingsDataWrapper& InOutWrappedJointSettingsData)
+{
+	CVD_COPY_FIELD_TO_WRAPPER(InSettings, InOutWrappedJointSettingsData, Stiffness);
+	CVD_COPY_FIELD_TO_WRAPPER(InSettings, InOutWrappedJointSettingsData, LinearProjection);
+	CVD_COPY_FIELD_TO_WRAPPER(InSettings, InOutWrappedJointSettingsData, AngularProjection);
+	CVD_COPY_FIELD_TO_WRAPPER(InSettings, InOutWrappedJointSettingsData, AngularProjection);
+	CVD_COPY_FIELD_TO_WRAPPER(InSettings, InOutWrappedJointSettingsData, TeleportDistance);
+	CVD_COPY_FIELD_TO_WRAPPER(InSettings, InOutWrappedJointSettingsData, TeleportAngle);
+	CVD_COPY_FIELD_TO_WRAPPER(InSettings, InOutWrappedJointSettingsData, ParentInvMassScale);
+	CVD_COPY_FIELD_TO_WRAPPER(InSettings, InOutWrappedJointSettingsData, bCollisionEnabled);
+	CVD_COPY_FIELD_TO_WRAPPER(InSettings, InOutWrappedJointSettingsData, bMassConditioningEnabled);
+	CVD_COPY_FIELD_TO_WRAPPER(InSettings, InOutWrappedJointSettingsData, bUseLinearSolver);
+	CVD_COPY_FIELD_TO_WRAPPER(InSettings, InOutWrappedJointSettingsData, bSoftLinearLimitsEnabled);
+	CVD_COPY_FIELD_TO_WRAPPER(InSettings, InOutWrappedJointSettingsData, bSoftTwistLimitsEnabled);
+	CVD_COPY_FIELD_TO_WRAPPER(InSettings, InOutWrappedJointSettingsData, bAngularSLerpPositionDriveEnabled);
+	CVD_COPY_FIELD_TO_WRAPPER(InSettings, InOutWrappedJointSettingsData, bAngularSLerpVelocityDriveEnabled);
+	CVD_COPY_FIELD_TO_WRAPPER(InSettings, InOutWrappedJointSettingsData, bAngularTwistPositionDriveEnabled);
+	CVD_COPY_FIELD_TO_WRAPPER(InSettings, InOutWrappedJointSettingsData, bAngularTwistVelocityDriveEnabled);
+	CVD_COPY_FIELD_TO_WRAPPER(InSettings, InOutWrappedJointSettingsData, bAngularSwingPositionDriveEnabled);
+	CVD_COPY_FIELD_TO_WRAPPER(InSettings, InOutWrappedJointSettingsData, bAngularSwingVelocityDriveEnabled);
+
+	InOutWrappedJointSettingsData.LinearMotionTypes[0] = static_cast<EChaosVDJointMotionType>(InSettings.LinearMotionTypes[0]);
+	InOutWrappedJointSettingsData.LinearMotionTypes[1] = static_cast<EChaosVDJointMotionType>(InSettings.LinearMotionTypes[1]);
+	InOutWrappedJointSettingsData.LinearMotionTypes[2] = static_cast<EChaosVDJointMotionType>(InSettings.LinearMotionTypes[2]);
+
+	CVD_COPY_FIELD_TO_WRAPPER(InSettings, InOutWrappedJointSettingsData, LinearLimit);
+
+	InOutWrappedJointSettingsData.AngularMotionTypes[0] = static_cast<EChaosVDJointMotionType>(InSettings.AngularMotionTypes[0]);
+	InOutWrappedJointSettingsData.AngularMotionTypes[1] = static_cast<EChaosVDJointMotionType>(InSettings.AngularMotionTypes[1]);
+	InOutWrappedJointSettingsData.AngularMotionTypes[2] = static_cast<EChaosVDJointMotionType>(InSettings.AngularMotionTypes[2]);
+
+	CVD_COPY_FIELD_TO_WRAPPER(InSettings, InOutWrappedJointSettingsData, AngularLimits);
+	CVD_COPY_FIELD_TO_WRAPPER_WITH_CAST(InSettings, InOutWrappedJointSettingsData, LinearSoftForceMode, EChaosVDJointForceMode);
+	CVD_COPY_FIELD_TO_WRAPPER_WITH_CAST(InSettings, InOutWrappedJointSettingsData, AngularSoftForceMode, EChaosVDJointForceMode);
+	CVD_COPY_FIELD_TO_WRAPPER(InSettings, InOutWrappedJointSettingsData, SoftLinearStiffness);
+	CVD_COPY_FIELD_TO_WRAPPER(InSettings, InOutWrappedJointSettingsData, SoftLinearDamping);
+	CVD_COPY_FIELD_TO_WRAPPER(InSettings, InOutWrappedJointSettingsData, SoftTwistStiffness);
+	CVD_COPY_FIELD_TO_WRAPPER(InSettings, InOutWrappedJointSettingsData, SoftTwistDamping);
+	CVD_COPY_FIELD_TO_WRAPPER(InSettings, InOutWrappedJointSettingsData, SoftSwingStiffness);
+	CVD_COPY_FIELD_TO_WRAPPER(InSettings, InOutWrappedJointSettingsData, SoftSwingDamping);
+	CVD_COPY_FIELD_TO_WRAPPER(InSettings, InOutWrappedJointSettingsData, LinearRestitution);
+	CVD_COPY_FIELD_TO_WRAPPER(InSettings, InOutWrappedJointSettingsData, TwistRestitution);
+	CVD_COPY_FIELD_TO_WRAPPER(InSettings, InOutWrappedJointSettingsData, TwistRestitution);
+	CVD_COPY_FIELD_TO_WRAPPER(InSettings, InOutWrappedJointSettingsData, SwingRestitution);
+	CVD_COPY_FIELD_TO_WRAPPER(InSettings, InOutWrappedJointSettingsData, LinearContactDistance);
+	CVD_COPY_FIELD_TO_WRAPPER(InSettings, InOutWrappedJointSettingsData, TwistContactDistance);
+	CVD_COPY_FIELD_TO_WRAPPER(InSettings, InOutWrappedJointSettingsData, SwingContactDistance);
+
+	CVD_COPY_VECTOR_FIELD_TO_WRAPPER(InSettings, InOutWrappedJointSettingsData, LinearDrivePositionTarget);
+	CVD_COPY_VECTOR_FIELD_TO_WRAPPER(InSettings, InOutWrappedJointSettingsData, LinearDriveVelocityTarget);
+		
+	InOutWrappedJointSettingsData.bLinearPositionDriveEnabled0 = InSettings.bLinearPositionDriveEnabled[0];
+	InOutWrappedJointSettingsData.bLinearPositionDriveEnabled1 = InSettings.bLinearPositionDriveEnabled[1];
+	InOutWrappedJointSettingsData.bLinearPositionDriveEnabled2 = InSettings.bLinearPositionDriveEnabled[2];
+
+	InOutWrappedJointSettingsData.bLinearVelocityDriveEnabled0 = InSettings.bLinearVelocityDriveEnabled[0];
+	InOutWrappedJointSettingsData.bLinearVelocityDriveEnabled1 = InSettings.bLinearVelocityDriveEnabled[1];
+	InOutWrappedJointSettingsData.bLinearVelocityDriveEnabled2 = InSettings.bLinearVelocityDriveEnabled[2];
+
+	CVD_COPY_FIELD_TO_WRAPPER_WITH_CAST(InSettings, InOutWrappedJointSettingsData, LinearDriveForceMode, EChaosVDJointForceMode);
+
+	CVD_COPY_VECTOR_FIELD_TO_WRAPPER(InSettings, InOutWrappedJointSettingsData, LinearDriveStiffness);
+	CVD_COPY_VECTOR_FIELD_TO_WRAPPER(InSettings, InOutWrappedJointSettingsData, LinearDriveDamping);
+	CVD_COPY_VECTOR_FIELD_TO_WRAPPER(InSettings, InOutWrappedJointSettingsData, LinearDriveMaxForce);
+	CVD_COPY_FIELD_TO_WRAPPER(InSettings, InOutWrappedJointSettingsData, AngularDrivePositionTarget);
+	CVD_COPY_VECTOR_FIELD_TO_WRAPPER(InSettings, InOutWrappedJointSettingsData, AngularDriveVelocityTarget);
+	CVD_COPY_FIELD_TO_WRAPPER_WITH_CAST(InSettings, InOutWrappedJointSettingsData, AngularDriveForceMode, EChaosVDJointForceMode);
+	CVD_COPY_VECTOR_FIELD_TO_WRAPPER(InSettings, InOutWrappedJointSettingsData, AngularDriveStiffness);
+	CVD_COPY_VECTOR_FIELD_TO_WRAPPER(InSettings, InOutWrappedJointSettingsData, AngularDriveDamping);
+	CVD_COPY_VECTOR_FIELD_TO_WRAPPER(InSettings, InOutWrappedJointSettingsData, AngularDriveMaxTorque);
+
+	CVD_COPY_FIELD_TO_WRAPPER(InSettings, InOutWrappedJointSettingsData, LinearBreakForce);
+	CVD_COPY_FIELD_TO_WRAPPER(InSettings, InOutWrappedJointSettingsData, LinearPlasticityLimit);
+
+	CVD_COPY_FIELD_TO_WRAPPER_WITH_CAST(InSettings, InOutWrappedJointSettingsData, LinearPlasticityType, EChaosVDPlasticityType);
+		
+	CVD_COPY_FIELD_TO_WRAPPER(InSettings, InOutWrappedJointSettingsData, LinearPlasticityInitialDistanceSquared);
+	CVD_COPY_FIELD_TO_WRAPPER(InSettings, InOutWrappedJointSettingsData, AngularBreakTorque);
+	CVD_COPY_FIELD_TO_WRAPPER(InSettings, InOutWrappedJointSettingsData, AngularPlasticityLimit);
+	CVD_COPY_FIELD_TO_WRAPPER(InSettings, InOutWrappedJointSettingsData, ContactTransferScale);
+		
+	InOutWrappedJointSettingsData.ConnectorTransforms[0] = InSettings.ConnectorTransforms[0];
+	InOutWrappedJointSettingsData.ConnectorTransforms[1] = InSettings.ConnectorTransforms[1];
+	
+	InOutWrappedJointSettingsData.MarkAsValid();
+}
+
 FChaosVDJointConstraint FChaosVDDataWrapperUtils::BuildJointDataWrapper(const Chaos::FPBDJointConstraintHandle* ConstaintHanlde)
 {
 	FChaosVDJointConstraint WrappedJointData;
@@ -251,126 +385,43 @@ FChaosVDJointConstraint FChaosVDDataWrapperUtils::BuildJointDataWrapper(const Ch
 		WrappedJointData.ParticleParIndexes[0] = ConstaintHanlde->GetConstrainedParticles()[0]->UniqueIdx().Idx;
 		WrappedJointData.ParticleParIndexes[1] = ConstaintHanlde->GetConstrainedParticles()[1]->UniqueIdx().Idx;
 
-		WrappedJointData.JointState.bBroken = ConstaintHanlde->IsConstraintBroken();
-		WrappedJointData.JointState.bBreaking = ConstaintHanlde->IsConstraintBreaking();
-		WrappedJointData.JointState.bDisabled = !ConstaintHanlde->IsConstraintEnabled();
-		WrappedJointData.JointState.bDriveTargetChanged = ConstaintHanlde->IsDriveTargetChanged();
+		WrappedJointData.PhysicsThreadJointState.bBroken = ConstaintHanlde->IsConstraintBroken();
+		WrappedJointData.PhysicsThreadJointState.bBreaking = ConstaintHanlde->IsConstraintBreaking();
+		WrappedJointData.PhysicsThreadJointState.bDisabled = !ConstaintHanlde->IsConstraintEnabled();
+		WrappedJointData.PhysicsThreadJointState.bDriveTargetChanged = ConstaintHanlde->IsDriveTargetChanged();
 
 		//TODO: Island related data getters are deprecated. We need to see where is bet to get that data now and if this should be recorded as part of the CVD Constraint wrapper
 		//WrappedJointData.JointState.Color = ConstaintHanlde->GetConstraintColor();
 		//WrappedJointData.JointState.Island = ConstaintHanlde->GetConstraintIsland();
 		//WrappedJointData.JointState.IslandSize = ConstaintHanlde->GetConstraintIsland();
 		
-		WrappedJointData.JointState.bEnabledDuringResim = ConstaintHanlde->IsEnabledDuringResim();
-		WrappedJointData.JointState.AngularImpulse = ConstaintHanlde->GetAngularImpulse();
-		WrappedJointData.JointState.LinearImpulse = ConstaintHanlde->GetLinearImpulse();
+		WrappedJointData.PhysicsThreadJointState.bEnabledDuringResim = ConstaintHanlde->IsEnabledDuringResim();
+		WrappedJointData.PhysicsThreadJointState.AngularImpulse = ConstaintHanlde->GetAngularImpulse();
+		WrappedJointData.PhysicsThreadJointState.LinearImpulse = ConstaintHanlde->GetLinearImpulse();
 
 		switch(ConstaintHanlde->ResimType())
 		{
 			case Chaos::EResimType::FullResim:
-				WrappedJointData.JointState.ResimType = EChaosVDJointReSimType::FullResim;
+				WrappedJointData.PhysicsThreadJointState.ResimType = EChaosVDJointReSimType::FullResim;
 				break;
 			case Chaos::EResimType::ResimAsFollower:
-				WrappedJointData.JointState.ResimType = EChaosVDJointReSimType::ResimAsFollower;
+				WrappedJointData.PhysicsThreadJointState.ResimType = EChaosVDJointReSimType::ResimAsFollower;
 				break;
 		}
 
 		switch(ConstaintHanlde->SyncState())
 		{
 			case Chaos::ESyncState::InSync:
-				WrappedJointData.JointState.SyncState = EChaosVDJointSyncType::InSync;
+				WrappedJointData.PhysicsThreadJointState.SyncState = EChaosVDJointSyncType::InSync;
 				break;
 			case Chaos::ESyncState::HardDesync:
-				WrappedJointData.JointState.SyncState = EChaosVDJointSyncType::HardDesync;
+				WrappedJointData.PhysicsThreadJointState.SyncState = EChaosVDJointSyncType::HardDesync;
 				break;
 		}
 
-		WrappedJointData.JointState.MarkAsValid();
-
-		CVD_COPY_FIELD_TO_WRAPPER(ConstaintHanlde->GetJointSettings(), WrappedJointData.JointSettings, Stiffness);
-		CVD_COPY_FIELD_TO_WRAPPER(ConstaintHanlde->GetJointSettings(), WrappedJointData.JointSettings, LinearProjection);
-		CVD_COPY_FIELD_TO_WRAPPER(ConstaintHanlde->GetJointSettings(), WrappedJointData.JointSettings, AngularProjection);
-		CVD_COPY_FIELD_TO_WRAPPER(ConstaintHanlde->GetJointSettings(), WrappedJointData.JointSettings, AngularProjection);
-		CVD_COPY_FIELD_TO_WRAPPER(ConstaintHanlde->GetJointSettings(), WrappedJointData.JointSettings, TeleportDistance);
-		CVD_COPY_FIELD_TO_WRAPPER(ConstaintHanlde->GetJointSettings(), WrappedJointData.JointSettings, TeleportAngle);
-		CVD_COPY_FIELD_TO_WRAPPER(ConstaintHanlde->GetJointSettings(), WrappedJointData.JointSettings, ParentInvMassScale);
-		CVD_COPY_FIELD_TO_WRAPPER(ConstaintHanlde->GetJointSettings(), WrappedJointData.JointSettings, bCollisionEnabled);
-		CVD_COPY_FIELD_TO_WRAPPER(ConstaintHanlde->GetJointSettings(), WrappedJointData.JointSettings, bMassConditioningEnabled);
-		CVD_COPY_FIELD_TO_WRAPPER(ConstaintHanlde->GetJointSettings(), WrappedJointData.JointSettings, bUseLinearSolver);
-		CVD_COPY_FIELD_TO_WRAPPER(ConstaintHanlde->GetJointSettings(), WrappedJointData.JointSettings, bSoftLinearLimitsEnabled);
-		CVD_COPY_FIELD_TO_WRAPPER(ConstaintHanlde->GetJointSettings(), WrappedJointData.JointSettings, bSoftTwistLimitsEnabled);
-		CVD_COPY_FIELD_TO_WRAPPER(ConstaintHanlde->GetJointSettings(), WrappedJointData.JointSettings, bAngularSLerpPositionDriveEnabled);
-		CVD_COPY_FIELD_TO_WRAPPER(ConstaintHanlde->GetJointSettings(), WrappedJointData.JointSettings, bAngularSLerpVelocityDriveEnabled);
-		CVD_COPY_FIELD_TO_WRAPPER(ConstaintHanlde->GetJointSettings(), WrappedJointData.JointSettings, bAngularTwistPositionDriveEnabled);
-		CVD_COPY_FIELD_TO_WRAPPER(ConstaintHanlde->GetJointSettings(), WrappedJointData.JointSettings, bAngularTwistVelocityDriveEnabled);
-		CVD_COPY_FIELD_TO_WRAPPER(ConstaintHanlde->GetJointSettings(), WrappedJointData.JointSettings, bAngularSwingPositionDriveEnabled);
-		CVD_COPY_FIELD_TO_WRAPPER(ConstaintHanlde->GetJointSettings(), WrappedJointData.JointSettings, bAngularSwingVelocityDriveEnabled);
-
-		WrappedJointData.JointSettings.LinearMotionTypes[0] = static_cast<EChaosVDJointMotionType>(ConstaintHanlde->GetJointSettings().LinearMotionTypes[0]);
-		WrappedJointData.JointSettings.LinearMotionTypes[1] = static_cast<EChaosVDJointMotionType>(ConstaintHanlde->GetJointSettings().LinearMotionTypes[1]);
-		WrappedJointData.JointSettings.LinearMotionTypes[2] = static_cast<EChaosVDJointMotionType>(ConstaintHanlde->GetJointSettings().LinearMotionTypes[2]);
-
-		CVD_COPY_FIELD_TO_WRAPPER(ConstaintHanlde->GetJointSettings(), WrappedJointData.JointSettings, LinearLimit);
-
-		WrappedJointData.JointSettings.AngularMotionTypes[0] = static_cast<EChaosVDJointMotionType>(ConstaintHanlde->GetJointSettings().AngularMotionTypes[0]);
-		WrappedJointData.JointSettings.AngularMotionTypes[1] = static_cast<EChaosVDJointMotionType>(ConstaintHanlde->GetJointSettings().AngularMotionTypes[1]);
-		WrappedJointData.JointSettings.AngularMotionTypes[2] = static_cast<EChaosVDJointMotionType>(ConstaintHanlde->GetJointSettings().AngularMotionTypes[2]);
-
-		CVD_COPY_FIELD_TO_WRAPPER(ConstaintHanlde->GetJointSettings(), WrappedJointData.JointSettings, AngularLimits);
-		CVD_COPY_FIELD_TO_WRAPPER_WITH_CAST(ConstaintHanlde->GetJointSettings(), WrappedJointData.JointSettings, LinearSoftForceMode, EChaosVDJointForceMode);
-		CVD_COPY_FIELD_TO_WRAPPER_WITH_CAST(ConstaintHanlde->GetJointSettings(), WrappedJointData.JointSettings, AngularSoftForceMode, EChaosVDJointForceMode);
-		CVD_COPY_FIELD_TO_WRAPPER(ConstaintHanlde->GetJointSettings(), WrappedJointData.JointSettings, SoftLinearStiffness);
-		CVD_COPY_FIELD_TO_WRAPPER(ConstaintHanlde->GetJointSettings(), WrappedJointData.JointSettings, SoftLinearDamping);
-		CVD_COPY_FIELD_TO_WRAPPER(ConstaintHanlde->GetJointSettings(), WrappedJointData.JointSettings, SoftTwistStiffness);
-		CVD_COPY_FIELD_TO_WRAPPER(ConstaintHanlde->GetJointSettings(), WrappedJointData.JointSettings, SoftTwistDamping);
-		CVD_COPY_FIELD_TO_WRAPPER(ConstaintHanlde->GetJointSettings(), WrappedJointData.JointSettings, SoftSwingStiffness);
-		CVD_COPY_FIELD_TO_WRAPPER(ConstaintHanlde->GetJointSettings(), WrappedJointData.JointSettings, SoftSwingDamping);
-		CVD_COPY_FIELD_TO_WRAPPER(ConstaintHanlde->GetJointSettings(), WrappedJointData.JointSettings, LinearRestitution);
-		CVD_COPY_FIELD_TO_WRAPPER(ConstaintHanlde->GetJointSettings(), WrappedJointData.JointSettings, TwistRestitution);
-		CVD_COPY_FIELD_TO_WRAPPER(ConstaintHanlde->GetJointSettings(), WrappedJointData.JointSettings, TwistRestitution);
-		CVD_COPY_FIELD_TO_WRAPPER(ConstaintHanlde->GetJointSettings(), WrappedJointData.JointSettings, SwingRestitution);
-		CVD_COPY_FIELD_TO_WRAPPER(ConstaintHanlde->GetJointSettings(), WrappedJointData.JointSettings, LinearContactDistance);
-		CVD_COPY_FIELD_TO_WRAPPER(ConstaintHanlde->GetJointSettings(), WrappedJointData.JointSettings, TwistContactDistance);
-		CVD_COPY_FIELD_TO_WRAPPER(ConstaintHanlde->GetJointSettings(), WrappedJointData.JointSettings, SwingContactDistance);
-
-		CVD_COPY_VECTOR_FIELD_TO_WRAPPER(ConstaintHanlde->GetJointSettings(), WrappedJointData.JointSettings, LinearDrivePositionTarget);
-		CVD_COPY_VECTOR_FIELD_TO_WRAPPER(ConstaintHanlde->GetJointSettings(), WrappedJointData.JointSettings, LinearDriveVelocityTarget);
-		
-		WrappedJointData.JointSettings.bLinearPositionDriveEnabled0 = ConstaintHanlde->GetJointSettings().bLinearPositionDriveEnabled[0];
-		WrappedJointData.JointSettings.bLinearPositionDriveEnabled1 = ConstaintHanlde->GetJointSettings().bLinearPositionDriveEnabled[1];
-		WrappedJointData.JointSettings.bLinearPositionDriveEnabled2 = ConstaintHanlde->GetJointSettings().bLinearPositionDriveEnabled[2];
-
-		WrappedJointData.JointSettings.bLinearVelocityDriveEnabled0 = ConstaintHanlde->GetJointSettings().bLinearVelocityDriveEnabled[0];
-		WrappedJointData.JointSettings.bLinearVelocityDriveEnabled1 = ConstaintHanlde->GetJointSettings().bLinearVelocityDriveEnabled[1];
-		WrappedJointData.JointSettings.bLinearVelocityDriveEnabled2 = ConstaintHanlde->GetJointSettings().bLinearVelocityDriveEnabled[2];
-
-		CVD_COPY_FIELD_TO_WRAPPER_WITH_CAST(ConstaintHanlde->GetJointSettings(), WrappedJointData.JointSettings, LinearDriveForceMode, EChaosVDJointForceMode);
-
-		CVD_COPY_VECTOR_FIELD_TO_WRAPPER(ConstaintHanlde->GetJointSettings(), WrappedJointData.JointSettings, LinearDriveStiffness);
-		CVD_COPY_VECTOR_FIELD_TO_WRAPPER(ConstaintHanlde->GetJointSettings(), WrappedJointData.JointSettings, LinearDriveDamping);
-		CVD_COPY_VECTOR_FIELD_TO_WRAPPER(ConstaintHanlde->GetJointSettings(), WrappedJointData.JointSettings, LinearDriveMaxForce);
-		CVD_COPY_FIELD_TO_WRAPPER(ConstaintHanlde->GetJointSettings(), WrappedJointData.JointSettings, AngularDrivePositionTarget);
-		CVD_COPY_VECTOR_FIELD_TO_WRAPPER(ConstaintHanlde->GetJointSettings(), WrappedJointData.JointSettings, AngularDriveVelocityTarget);
-		CVD_COPY_FIELD_TO_WRAPPER_WITH_CAST(ConstaintHanlde->GetJointSettings(), WrappedJointData.JointSettings, AngularDriveForceMode, EChaosVDJointForceMode);
-		CVD_COPY_VECTOR_FIELD_TO_WRAPPER(ConstaintHanlde->GetJointSettings(), WrappedJointData.JointSettings, AngularDriveStiffness);
-		CVD_COPY_VECTOR_FIELD_TO_WRAPPER(ConstaintHanlde->GetJointSettings(), WrappedJointData.JointSettings, AngularDriveDamping);
-		CVD_COPY_VECTOR_FIELD_TO_WRAPPER(ConstaintHanlde->GetJointSettings(), WrappedJointData.JointSettings, AngularDriveMaxTorque);
-
-		CVD_COPY_FIELD_TO_WRAPPER(ConstaintHanlde->GetJointSettings(), WrappedJointData.JointSettings, LinearBreakForce);
-		CVD_COPY_FIELD_TO_WRAPPER(ConstaintHanlde->GetJointSettings(), WrappedJointData.JointSettings, LinearPlasticityLimit);
-
-		CVD_COPY_FIELD_TO_WRAPPER_WITH_CAST(ConstaintHanlde->GetJointSettings(), WrappedJointData.JointSettings, LinearPlasticityType, EChaosVDPlasticityType);
-		
-		CVD_COPY_FIELD_TO_WRAPPER(ConstaintHanlde->GetJointSettings(), WrappedJointData.JointSettings, LinearPlasticityInitialDistanceSquared);
-		CVD_COPY_FIELD_TO_WRAPPER(ConstaintHanlde->GetJointSettings(), WrappedJointData.JointSettings, AngularBreakTorque);
-		CVD_COPY_FIELD_TO_WRAPPER(ConstaintHanlde->GetJointSettings(), WrappedJointData.JointSettings, AngularPlasticityLimit);
-		CVD_COPY_FIELD_TO_WRAPPER(ConstaintHanlde->GetJointSettings(), WrappedJointData.JointSettings, ContactTransferScale);
-		
-		WrappedJointData.JointSettings.ConnectorTransforms[0] = ConstaintHanlde->GetJointSettings().ConnectorTransforms[0];
-		WrappedJointData.JointSettings.ConnectorTransforms[1] = ConstaintHanlde->GetJointSettings().ConnectorTransforms[1];
-	
-		WrappedJointData.JointSettings.MarkAsValid();
-		 
+		WrappedJointData.PhysicsThreadJointState.MarkAsValid();
+
+		CopyJointConstraintSettingsToWrapper(ConstaintHanlde->GetJointSettings(), WrappedJointData.JointSettings);	 
 	}
 	
 	WrappedJointData.MarkAsValid();
@@ -378,6 +429,120 @@ FChaosVDJointConstraint FChaosVDDataWrapperUtils::BuildJointDataWrapper(const Ch
 	return MoveTemp(WrappedJointData);
 }
 
+FChaosVDJointConstraint FChaosVDDataWrapperUtils::BuildGTJointDataWrapper(Chaos::FJointConstraint* GTConstraintPtr)
+{
+	FChaosVDJointConstraint WrappedJointData;
+
+	if (!GTConstraintPtr)
+	{
+		return MoveTemp(WrappedJointData);
+	}
+
+	//TODO: Extract constrained particle indexes from the proxy pair. This will needed if we use this in CVD
+
+	const Chaos::FJointConstraint::FOutputData& CurrentState = GTConstraintPtr->GetOutputData();
+
+	WrappedJointData.GameThreadJointState.AngularViolation = CurrentState.AngularViolation;
+	WrappedJointData.GameThreadJointState.bDriveTargetChanged = CurrentState.bDriveTargetChanged;
+	WrappedJointData.GameThreadJointState.bIsBreaking = CurrentState.bIsBreaking;
+	WrappedJointData.GameThreadJointState.bIsBroken = CurrentState.bIsBroken;
+	WrappedJointData.GameThreadJointState.bIsViolating = CurrentState.bIsViolating;
+	WrappedJointData.GameThreadJointState.LinearViolation = CurrentState.LinearViolation;
+	WrappedJointData.GameThreadJointState.Force = CurrentState.Force;
+	WrappedJointData.GameThreadJointState.Torque = CurrentState.Torque;
+
+	CopyJointConstraintSettingsToWrapper(GTConstraintPtr->GetJointSettings(), WrappedJointData.JointSettings);
+
+	return MoveTemp(WrappedJointData);
+}
+
+
+void FChaosVDDataWrapperUtils::CopyJointSettingsFromWrapper(const FChaosVDJointSettingsDataWrapper& InWrappedJointSettingsData, Chaos::FPBDJointSettings& NewSettings)
+{
+	CVD_COPY_FIELD_TO_WRAPPER(InWrappedJointSettingsData, NewSettings, Stiffness);
+	CVD_COPY_FIELD_TO_WRAPPER(InWrappedJointSettingsData, NewSettings, LinearProjection);
+	CVD_COPY_FIELD_TO_WRAPPER(InWrappedJointSettingsData, NewSettings, AngularProjection);
+	CVD_COPY_FIELD_TO_WRAPPER(InWrappedJointSettingsData, NewSettings, AngularProjection);
+	CVD_COPY_FIELD_TO_WRAPPER(InWrappedJointSettingsData, NewSettings, TeleportDistance);
+	CVD_COPY_FIELD_TO_WRAPPER(InWrappedJointSettingsData, NewSettings, TeleportAngle);
+	CVD_COPY_FIELD_TO_WRAPPER(InWrappedJointSettingsData, NewSettings, ParentInvMassScale);
+	CVD_COPY_FIELD_TO_WRAPPER(InWrappedJointSettingsData, NewSettings, bCollisionEnabled);
+	CVD_COPY_FIELD_TO_WRAPPER(InWrappedJointSettingsData, NewSettings, bMassConditioningEnabled);
+	CVD_COPY_FIELD_TO_WRAPPER(InWrappedJointSettingsData, NewSettings, bUseLinearSolver);
+	CVD_COPY_FIELD_TO_WRAPPER(InWrappedJointSettingsData, NewSettings, bSoftLinearLimitsEnabled);
+	CVD_COPY_FIELD_TO_WRAPPER(InWrappedJointSettingsData, NewSettings, bSoftTwistLimitsEnabled);
+	CVD_COPY_FIELD_TO_WRAPPER(InWrappedJointSettingsData, NewSettings, bAngularSLerpPositionDriveEnabled);
+	CVD_COPY_FIELD_TO_WRAPPER(InWrappedJointSettingsData, NewSettings, bAngularSLerpVelocityDriveEnabled);
+	CVD_COPY_FIELD_TO_WRAPPER(InWrappedJointSettingsData, NewSettings, bAngularTwistPositionDriveEnabled);
+	CVD_COPY_FIELD_TO_WRAPPER(InWrappedJointSettingsData, NewSettings, bAngularTwistVelocityDriveEnabled);
+	CVD_COPY_FIELD_TO_WRAPPER(InWrappedJointSettingsData, NewSettings, bAngularSwingPositionDriveEnabled);
+	CVD_COPY_FIELD_TO_WRAPPER(InWrappedJointSettingsData, NewSettings, bAngularSwingVelocityDriveEnabled);
+
+	NewSettings.LinearMotionTypes[0] = static_cast<Chaos::EJointMotionType>(InWrappedJointSettingsData.LinearMotionTypes[0]);
+	NewSettings.LinearMotionTypes[1] = static_cast<Chaos::EJointMotionType>(InWrappedJointSettingsData.LinearMotionTypes[1]);
+	NewSettings.LinearMotionTypes[2] = static_cast<Chaos::EJointMotionType>(InWrappedJointSettingsData.LinearMotionTypes[2]);
+
+	CVD_COPY_FIELD_TO_WRAPPER(InWrappedJointSettingsData, NewSettings, LinearLimit);
+
+	NewSettings.AngularMotionTypes[0] = static_cast<Chaos::EJointMotionType>(InWrappedJointSettingsData.AngularMotionTypes[0]);
+	NewSettings.AngularMotionTypes[1] = static_cast<Chaos::EJointMotionType>(InWrappedJointSettingsData.AngularMotionTypes[1]);
+	NewSettings.AngularMotionTypes[2] = static_cast<Chaos::EJointMotionType>(InWrappedJointSettingsData.AngularMotionTypes[2]);
+
+	CVD_COPY_FIELD_TO_WRAPPER(InWrappedJointSettingsData, NewSettings, AngularLimits);
+	CVD_COPY_FIELD_TO_WRAPPER_WITH_CAST(InWrappedJointSettingsData, NewSettings, LinearSoftForceMode, Chaos::EJointForceMode);
+	CVD_COPY_FIELD_TO_WRAPPER_WITH_CAST(InWrappedJointSettingsData, NewSettings, AngularSoftForceMode, Chaos::EJointForceMode);
+	CVD_COPY_FIELD_TO_WRAPPER(InWrappedJointSettingsData, NewSettings, SoftLinearStiffness);
+	CVD_COPY_FIELD_TO_WRAPPER(InWrappedJointSettingsData, NewSettings, SoftLinearDamping);
+	CVD_COPY_FIELD_TO_WRAPPER(InWrappedJointSettingsData, NewSettings, SoftTwistStiffness);
+	CVD_COPY_FIELD_TO_WRAPPER(InWrappedJointSettingsData, NewSettings, SoftTwistDamping);
+	CVD_COPY_FIELD_TO_WRAPPER(InWrappedJointSettingsData, NewSettings, SoftSwingStiffness);
+	CVD_COPY_FIELD_TO_WRAPPER(InWrappedJointSettingsData, NewSettings, SoftSwingDamping);
+	CVD_COPY_FIELD_TO_WRAPPER(InWrappedJointSettingsData, NewSettings, LinearRestitution);
+	CVD_COPY_FIELD_TO_WRAPPER(InWrappedJointSettingsData, NewSettings, TwistRestitution);
+	CVD_COPY_FIELD_TO_WRAPPER(InWrappedJointSettingsData, NewSettings, TwistRestitution);
+	CVD_COPY_FIELD_TO_WRAPPER(InWrappedJointSettingsData, NewSettings, SwingRestitution);
+	CVD_COPY_FIELD_TO_WRAPPER(InWrappedJointSettingsData, NewSettings, LinearContactDistance);
+	CVD_COPY_FIELD_TO_WRAPPER(InWrappedJointSettingsData, NewSettings, TwistContactDistance);
+	CVD_COPY_FIELD_TO_WRAPPER(InWrappedJointSettingsData, NewSettings, SwingContactDistance);
+
+	CVD_COPY_VECTOR_FIELD_TO_WRAPPER(InWrappedJointSettingsData, NewSettings, LinearDrivePositionTarget);
+	CVD_COPY_VECTOR_FIELD_TO_WRAPPER(InWrappedJointSettingsData, NewSettings, LinearDriveVelocityTarget);
+	
+	NewSettings.bLinearPositionDriveEnabled[0] = InWrappedJointSettingsData.bLinearPositionDriveEnabled0;
+	NewSettings.bLinearPositionDriveEnabled[1] = InWrappedJointSettingsData.bLinearPositionDriveEnabled1;
+	NewSettings.bLinearPositionDriveEnabled[2] = InWrappedJointSettingsData.bLinearPositionDriveEnabled2;
+
+	NewSettings.bLinearVelocityDriveEnabled[0] = InWrappedJointSettingsData.bLinearVelocityDriveEnabled0;
+	NewSettings.bLinearVelocityDriveEnabled[1] = InWrappedJointSettingsData.bLinearVelocityDriveEnabled1;
+	NewSettings.bLinearVelocityDriveEnabled[2] = InWrappedJointSettingsData.bLinearVelocityDriveEnabled2;
+
+	CVD_COPY_FIELD_TO_WRAPPER_WITH_CAST(InWrappedJointSettingsData, NewSettings, LinearDriveForceMode, Chaos::EJointForceMode);
+
+	CVD_COPY_VECTOR_FIELD_TO_WRAPPER(InWrappedJointSettingsData, NewSettings, LinearDriveStiffness);
+	CVD_COPY_VECTOR_FIELD_TO_WRAPPER(InWrappedJointSettingsData, NewSettings, LinearDriveDamping);
+	CVD_COPY_VECTOR_FIELD_TO_WRAPPER(InWrappedJointSettingsData, NewSettings, LinearDriveMaxForce);
+	CVD_COPY_FIELD_TO_WRAPPER(InWrappedJointSettingsData, NewSettings, AngularDrivePositionTarget);
+	CVD_COPY_VECTOR_FIELD_TO_WRAPPER(InWrappedJointSettingsData, NewSettings, AngularDriveVelocityTarget);
+	CVD_COPY_FIELD_TO_WRAPPER_WITH_CAST(InWrappedJointSettingsData, NewSettings, AngularDriveForceMode, Chaos::EJointForceMode);
+	CVD_COPY_VECTOR_FIELD_TO_WRAPPER(InWrappedJointSettingsData, NewSettings, AngularDriveStiffness);
+	CVD_COPY_VECTOR_FIELD_TO_WRAPPER(InWrappedJointSettingsData, NewSettings, AngularDriveDamping);
+	CVD_COPY_VECTOR_FIELD_TO_WRAPPER(InWrappedJointSettingsData, NewSettings, AngularDriveMaxTorque);
+
+	CVD_COPY_FIELD_TO_WRAPPER(InWrappedJointSettingsData, NewSettings, LinearBreakForce);
+	CVD_COPY_FIELD_TO_WRAPPER(InWrappedJointSettingsData, NewSettings, LinearPlasticityLimit);
+
+	CVD_COPY_FIELD_TO_WRAPPER_WITH_CAST(InWrappedJointSettingsData, NewSettings, LinearPlasticityType, Chaos::EPlasticityType);
+	
+	CVD_COPY_FIELD_TO_WRAPPER(InWrappedJointSettingsData, NewSettings, LinearPlasticityInitialDistanceSquared);
+	CVD_COPY_FIELD_TO_WRAPPER(InWrappedJointSettingsData, NewSettings, AngularBreakTorque);
+	CVD_COPY_FIELD_TO_WRAPPER(InWrappedJointSettingsData, NewSettings, AngularPlasticityLimit);
+	CVD_COPY_FIELD_TO_WRAPPER(InWrappedJointSettingsData, NewSettings, ContactTransferScale);
+	
+	NewSettings.ConnectorTransforms[0] = InWrappedJointSettingsData.ConnectorTransforms[0];
+	NewSettings.ConnectorTransforms[1] = InWrappedJointSettingsData.ConnectorTransforms[1];
+}
+
+
 void FChaosVDDataWrapperUtils::ApplyJointDataWrapperToHandle(Chaos::FPBDJointConstraintHandle* ConstaintHanlde, const FChaosVDJointConstraint& WrappedJointData)
 {
 	if (!ensure(ConstaintHanlde))
@@ -390,106 +555,48 @@ void FChaosVDDataWrapperUtils::ApplyJointDataWrapperToHandle(Chaos::FPBDJointCon
 	FPBDJointState NewState;
 	FPBDJointSettings NewSettings;
 
-	NewState.bBroken = WrappedJointData.JointState.bBroken;
-	NewState.bBreaking = WrappedJointData.JointState.bBreaking;
-	NewState.bDisabled = WrappedJointData.JointState.bDisabled;
-	NewState.bDriveTargetChanged = WrappedJointData.JointState.bDriveTargetChanged;
-
-	NewState.bEnabledDuringResim = WrappedJointData.JointState.bEnabledDuringResim;
-	NewState.AngularImpulse = WrappedJointData.JointState.AngularImpulse;
-	NewState.LinearImpulse = WrappedJointData.JointState.LinearImpulse;
-
-	NewState.LinearImpulse = WrappedJointData.JointState.LinearImpulse;
-
-	CVD_COPY_FIELD_TO_WRAPPER_WITH_CAST(WrappedJointData.JointState, NewState, ResimType, Chaos::EResimType);
-	CVD_COPY_FIELD_TO_WRAPPER_WITH_CAST(WrappedJointData.JointState, NewState, SyncState, Chaos::ESyncState);
-
-	CVD_COPY_FIELD_TO_WRAPPER(WrappedJointData.JointSettings, NewSettings, Stiffness);
-	CVD_COPY_FIELD_TO_WRAPPER(WrappedJointData.JointSettings, NewSettings, LinearProjection);
-	CVD_COPY_FIELD_TO_WRAPPER(WrappedJointData.JointSettings, NewSettings, AngularProjection);
-	CVD_COPY_FIELD_TO_WRAPPER(WrappedJointData.JointSettings, NewSettings, AngularProjection);
-	CVD_COPY_FIELD_TO_WRAPPER(WrappedJointData.JointSettings, NewSettings, TeleportDistance);
-	CVD_COPY_FIELD_TO_WRAPPER(WrappedJointData.JointSettings, NewSettings, TeleportAngle);
-	CVD_COPY_FIELD_TO_WRAPPER(WrappedJointData.JointSettings, NewSettings, ParentInvMassScale);
-	CVD_COPY_FIELD_TO_WRAPPER(WrappedJointData.JointSettings, NewSettings, bCollisionEnabled);
-	CVD_COPY_FIELD_TO_WRAPPER(WrappedJointData.JointSettings, NewSettings, bMassConditioningEnabled);
-	CVD_COPY_FIELD_TO_WRAPPER(WrappedJointData.JointSettings, NewSettings, bUseLinearSolver);
-	CVD_COPY_FIELD_TO_WRAPPER(WrappedJointData.JointSettings, NewSettings, bSoftLinearLimitsEnabled);
-	CVD_COPY_FIELD_TO_WRAPPER(WrappedJointData.JointSettings, NewSettings, bSoftTwistLimitsEnabled);
-	CVD_COPY_FIELD_TO_WRAPPER(WrappedJointData.JointSettings, NewSettings, bAngularSLerpPositionDriveEnabled);
-	CVD_COPY_FIELD_TO_WRAPPER(WrappedJointData.JointSettings, NewSettings, bAngularSLerpVelocityDriveEnabled);
-	CVD_COPY_FIELD_TO_WRAPPER(WrappedJointData.JointSettings, NewSettings, bAngularTwistPositionDriveEnabled);
-	CVD_COPY_FIELD_TO_WRAPPER(WrappedJointData.JointSettings, NewSettings, bAngularTwistVelocityDriveEnabled);
-	CVD_COPY_FIELD_TO_WRAPPER(WrappedJointData.JointSettings, NewSettings, bAngularSwingPositionDriveEnabled);
-	CVD_COPY_FIELD_TO_WRAPPER(WrappedJointData.JointSettings, NewSettings, bAngularSwingVelocityDriveEnabled);
-
-	NewSettings.LinearMotionTypes[0] = static_cast<EJointMotionType>(WrappedJointData.JointSettings.LinearMotionTypes[0]);
-	NewSettings.LinearMotionTypes[1] = static_cast<EJointMotionType>(WrappedJointData.JointSettings.LinearMotionTypes[1]);
-	NewSettings.LinearMotionTypes[2] = static_cast<EJointMotionType>(WrappedJointData.JointSettings.LinearMotionTypes[2]);
-
-	CVD_COPY_FIELD_TO_WRAPPER(WrappedJointData.JointSettings, NewSettings, LinearLimit);
-
-	NewSettings.AngularMotionTypes[0] = static_cast<EJointMotionType>(WrappedJointData.JointSettings.AngularMotionTypes[0]);
-	NewSettings.AngularMotionTypes[1] = static_cast<EJointMotionType>(WrappedJointData.JointSettings.AngularMotionTypes[1]);
-	NewSettings.AngularMotionTypes[2] = static_cast<EJointMotionType>(WrappedJointData.JointSettings.AngularMotionTypes[2]);
-
-	CVD_COPY_FIELD_TO_WRAPPER(WrappedJointData.JointSettings, NewSettings, AngularLimits);
-	CVD_COPY_FIELD_TO_WRAPPER_WITH_CAST(WrappedJointData.JointSettings, NewSettings, LinearSoftForceMode, EJointForceMode);
-	CVD_COPY_FIELD_TO_WRAPPER_WITH_CAST(WrappedJointData.JointSettings, NewSettings, AngularSoftForceMode, EJointForceMode);
-	CVD_COPY_FIELD_TO_WRAPPER(WrappedJointData.JointSettings, NewSettings, SoftLinearStiffness);
-	CVD_COPY_FIELD_TO_WRAPPER(WrappedJointData.JointSettings, NewSettings, SoftLinearDamping);
-	CVD_COPY_FIELD_TO_WRAPPER(WrappedJointData.JointSettings, NewSettings, SoftTwistStiffness);
-	CVD_COPY_FIELD_TO_WRAPPER(WrappedJointData.JointSettings, NewSettings, SoftTwistDamping);
-	CVD_COPY_FIELD_TO_WRAPPER(WrappedJointData.JointSettings, NewSettings, SoftSwingStiffness);
-	CVD_COPY_FIELD_TO_WRAPPER(WrappedJointData.JointSettings, NewSettings, SoftSwingDamping);
-	CVD_COPY_FIELD_TO_WRAPPER(WrappedJointData.JointSettings, NewSettings, LinearRestitution);
-	CVD_COPY_FIELD_TO_WRAPPER(WrappedJointData.JointSettings, NewSettings, TwistRestitution);
-	CVD_COPY_FIELD_TO_WRAPPER(WrappedJointData.JointSettings, NewSettings, TwistRestitution);
-	CVD_COPY_FIELD_TO_WRAPPER(WrappedJointData.JointSettings, NewSettings, SwingRestitution);
-	CVD_COPY_FIELD_TO_WRAPPER(WrappedJointData.JointSettings, NewSettings, LinearContactDistance);
-	CVD_COPY_FIELD_TO_WRAPPER(WrappedJointData.JointSettings, NewSettings, TwistContactDistance);
-	CVD_COPY_FIELD_TO_WRAPPER(WrappedJointData.JointSettings, NewSettings, SwingContactDistance);
-
-	CVD_COPY_VECTOR_FIELD_TO_WRAPPER(WrappedJointData.JointSettings, NewSettings, LinearDrivePositionTarget);
-	CVD_COPY_VECTOR_FIELD_TO_WRAPPER(WrappedJointData.JointSettings, NewSettings, LinearDriveVelocityTarget);
-	
-	NewSettings.bLinearPositionDriveEnabled[0] = WrappedJointData.JointSettings.bLinearPositionDriveEnabled0;
-	NewSettings.bLinearPositionDriveEnabled[1] = WrappedJointData.JointSettings.bLinearPositionDriveEnabled1;
-	NewSettings.bLinearPositionDriveEnabled[2] = WrappedJointData.JointSettings.bLinearPositionDriveEnabled2;
-
-	NewSettings.bLinearVelocityDriveEnabled[0] = WrappedJointData.JointSettings.bLinearVelocityDriveEnabled0;
-	NewSettings.bLinearVelocityDriveEnabled[1] = WrappedJointData.JointSettings.bLinearVelocityDriveEnabled1;
-	NewSettings.bLinearVelocityDriveEnabled[2] = WrappedJointData.JointSettings.bLinearVelocityDriveEnabled2;
-
-	CVD_COPY_FIELD_TO_WRAPPER_WITH_CAST(WrappedJointData.JointSettings, NewSettings, LinearDriveForceMode, EJointForceMode);
-
-	CVD_COPY_VECTOR_FIELD_TO_WRAPPER(WrappedJointData.JointSettings, NewSettings, LinearDriveStiffness);
-	CVD_COPY_VECTOR_FIELD_TO_WRAPPER(WrappedJointData.JointSettings, NewSettings, LinearDriveDamping);
-	CVD_COPY_VECTOR_FIELD_TO_WRAPPER(WrappedJointData.JointSettings, NewSettings, LinearDriveMaxForce);
-	CVD_COPY_FIELD_TO_WRAPPER(WrappedJointData.JointSettings, NewSettings, AngularDrivePositionTarget);
-	CVD_COPY_VECTOR_FIELD_TO_WRAPPER(WrappedJointData.JointSettings, NewSettings, AngularDriveVelocityTarget);
-	CVD_COPY_FIELD_TO_WRAPPER_WITH_CAST(WrappedJointData.JointSettings, NewSettings, AngularDriveForceMode, EJointForceMode);
-	CVD_COPY_VECTOR_FIELD_TO_WRAPPER(WrappedJointData.JointSettings, NewSettings, AngularDriveStiffness);
-	CVD_COPY_VECTOR_FIELD_TO_WRAPPER(WrappedJointData.JointSettings, NewSettings, AngularDriveDamping);
-	CVD_COPY_VECTOR_FIELD_TO_WRAPPER(WrappedJointData.JointSettings, NewSettings, AngularDriveMaxTorque);
-
-	CVD_COPY_FIELD_TO_WRAPPER(WrappedJointData.JointSettings, NewSettings, LinearBreakForce);
-	CVD_COPY_FIELD_TO_WRAPPER(WrappedJointData.JointSettings, NewSettings, LinearPlasticityLimit);
-
-	CVD_COPY_FIELD_TO_WRAPPER_WITH_CAST(WrappedJointData.JointSettings, NewSettings, LinearPlasticityType, EPlasticityType);
-	
-	CVD_COPY_FIELD_TO_WRAPPER(WrappedJointData.JointSettings, NewSettings, LinearPlasticityInitialDistanceSquared);
-	CVD_COPY_FIELD_TO_WRAPPER(WrappedJointData.JointSettings, NewSettings, AngularBreakTorque);
-	CVD_COPY_FIELD_TO_WRAPPER(WrappedJointData.JointSettings, NewSettings, AngularPlasticityLimit);
-	CVD_COPY_FIELD_TO_WRAPPER(WrappedJointData.JointSettings, NewSettings, ContactTransferScale);
-	
-	NewSettings.ConnectorTransforms[0] = WrappedJointData.JointSettings.ConnectorTransforms[0];
-	NewSettings.ConnectorTransforms[1] = WrappedJointData.JointSettings.ConnectorTransforms[1];
+	NewState.bBroken = WrappedJointData.PhysicsThreadJointState.bBroken;
+	NewState.bBreaking = WrappedJointData.PhysicsThreadJointState.bBreaking;
+	NewState.bDisabled = WrappedJointData.PhysicsThreadJointState.bDisabled;
+	NewState.bDriveTargetChanged = WrappedJointData.PhysicsThreadJointState.bDriveTargetChanged;
+
+	NewState.bEnabledDuringResim = WrappedJointData.PhysicsThreadJointState.bEnabledDuringResim;
+	NewState.AngularImpulse = WrappedJointData.PhysicsThreadJointState.AngularImpulse;
+	NewState.LinearImpulse = WrappedJointData.PhysicsThreadJointState.LinearImpulse;
+
+	CVD_COPY_FIELD_TO_WRAPPER_WITH_CAST(WrappedJointData.PhysicsThreadJointState, NewState, ResimType, Chaos::EResimType);
+	CVD_COPY_FIELD_TO_WRAPPER_WITH_CAST(WrappedJointData.PhysicsThreadJointState, NewState, SyncState, Chaos::ESyncState);
+
+	CopyJointSettingsFromWrapper(WrappedJointData.JointSettings, NewSettings);
 
 	ConstaintHanlde->SetSettings(MoveTemp(NewSettings));
 	ConstaintHanlde->SetJointState(MoveTemp(NewState));
 }
 
+void FChaosVDDataWrapperUtils::ApplyJointDataWrapperGTConstraint(Chaos::FJointConstraint* ConstraintPtr, const FChaosVDJointConstraint& InWrappedJointData)
+{
+	if (!ensure(ConstraintPtr))
+	{
+		return;
+	}
+
+	using namespace Chaos;
+
+	FJointConstraint::FOutputData& OutputStateDataRef = ConstraintPtr->GetOutputData();
+	OutputStateDataRef.Torque = InWrappedJointData.GameThreadJointState.Torque;
+	OutputStateDataRef.AngularViolation = InWrappedJointData.GameThreadJointState.AngularViolation;
+	OutputStateDataRef.bDriveTargetChanged = InWrappedJointData.GameThreadJointState.bDriveTargetChanged;
+	OutputStateDataRef.bIsBreaking = InWrappedJointData.GameThreadJointState.bIsBreaking;
+	OutputStateDataRef.bIsBroken = InWrappedJointData.GameThreadJointState.bIsBroken;
+	OutputStateDataRef.bIsViolating = InWrappedJointData.GameThreadJointState.bIsViolating;
+	OutputStateDataRef.Force = InWrappedJointData.GameThreadJointState.Force;
+	OutputStateDataRef.LinearViolation = InWrappedJointData.GameThreadJointState.LinearViolation;
+
+	FPBDJointSettings NewSettings;
+	CopyJointSettingsFromWrapper(InWrappedJointData.JointSettings, NewSettings);
+	ConstraintPtr->SetJointSettings(MoveTemp(NewSettings));
+}
+
 FChaosVDCharacterGroundConstraint FChaosVDDataWrapperUtils::BuildCharacterGroundConstraintDataWrapper(const Chaos::FCharacterGroundConstraintHandle* ConstraintHandle)
 {
 	FChaosVDCharacterGroundConstraint WrappedConstraintData;
diff --git a/Engine/Source/Runtime/Experimental/Chaos/Private/PBDRigidsSolver.cpp b/Engine/Source/Runtime/Experimental/Chaos/Private/PBDRigidsSolver.cpp
--- a/Engine/Source/Runtime/Experimental/Chaos/Private/PBDRigidsSolver.cpp
+++ b/Engine/Source/Runtime/Experimental/Chaos/Private/PBDRigidsSolver.cpp
@@ -1708,7 +1708,7 @@ namespace Chaos
 						if (bIsNew)
 						{
 							
-							if (FSerializedDataBufferPtr StateBuffer = Serializer.PopPendingSerializedStateForProxy(JointProxy))
+							if (FSerializedDataBufferPtr StateBuffer = Serializer.PopPendingInternalSerializedStateForProxy(JointProxy))
 							{
 								Serializer.ApplySerializedStateToJointConstraint(JointProxy->GetHandle(), *StateBuffer);
 							}
diff --git a/Engine/Source/Runtime/Experimental/Chaos/Private/PhysicsProxy/SingleParticlePhysicsProxy.cpp b/Engine/Source/Runtime/Experimental/Chaos/Private/PhysicsProxy/SingleParticlePhysicsProxy.cpp
--- a/Engine/Source/Runtime/Experimental/Chaos/Private/PhysicsProxy/SingleParticlePhysicsProxy.cpp
+++ b/Engine/Source/Runtime/Experimental/Chaos/Private/PhysicsProxy/SingleParticlePhysicsProxy.cpp
@@ -72,7 +72,7 @@ void PushToPhysicsStateImp(const Chaos::FDirtyPropertiesManager& Manager, Chaos:
 	const FParticleNonFrequentData* NewNonFrequentData = ParticleData.FindNonFrequentData(Manager, DataIdx);
 
 #if UE_WITH_REMOTE_OBJECT_HANDLE
-	if (FSerializedDataBufferPtr PendingStateData = Dirty.Proxy && !Dirty.Proxy->IsInitialized() ? Solver.GetSerializer().PopPendingSerializedStateForProxy(Dirty.Proxy) : nullptr)
+	if (FSerializedDataBufferPtr PendingStateData = Dirty.Proxy && !Dirty.Proxy->IsInitialized() ? Solver.GetSerializer().PopPendingInternalSerializedStateForProxy(Dirty.Proxy) : nullptr)
 	{
 		// Currently we don't serialize collision geometry and shape data, so we need to apply it first
 		// This might change in the future
diff --git a/Engine/Source/Runtime/Experimental/Chaos/Public/Chaos/ConstraintHandle.h b/Engine/Source/Runtime/Experimental/Chaos/Public/Chaos/ConstraintHandle.h
--- a/Engine/Source/Runtime/Experimental/Chaos/Public/Chaos/ConstraintHandle.h
+++ b/Engine/Source/Runtime/Experimental/Chaos/Public/Chaos/ConstraintHandle.h
@@ -78,6 +78,11 @@ namespace Chaos
 			return false;
 		}
 
+		FString ToString() const
+		{
+			return TypeName.ToString();
+		}
+
 	private:
 		FName TypeName;
 		const FConstraintHandleTypeID* BaseType;
diff --git a/Engine/Source/Runtime/Experimental/Chaos/Public/Chaos/Framework/PhysicsSolverBase.h b/Engine/Source/Runtime/Experimental/Chaos/Public/Chaos/Framework/PhysicsSolverBase.h
--- a/Engine/Source/Runtime/Experimental/Chaos/Public/Chaos/Framework/PhysicsSolverBase.h
+++ b/Engine/Source/Runtime/Experimental/Chaos/Public/Chaos/Framework/PhysicsSolverBase.h
@@ -961,6 +961,11 @@ namespace Chaos
 		/** Scene lock object for external threads (non-physics) */
 		TUniquePtr<FPhysSceneLock> ExternalDataLock_External;
 
+#if UE_WITH_REMOTE_OBJECT_HANDLE
+		/** Scene lock object for internal threads */
+		TUniquePtr<FPhysSceneLock> InternalDataLock;
+#endif
+
 		friend FChaosSolversModule;
 		friend FPhysicsSolverProcessPushDataTask;
 		friend FPhysicsSolverAdvanceTask;
@@ -986,6 +991,13 @@ namespace Chaos
 		//NOTE: this is a const operation so that you can acquire a read lock on a const solver. The assumption is that non-const write operations are already marked non-const
 		FPhysSceneLock& GetExternalDataLock_External() const { return *ExternalDataLock_External; }
 
+#if UE_WITH_REMOTE_OBJECT_HANDLE
+		UE_INTERNAL FPhysSceneLock& GetInternalDataLock() const
+		{
+			return *InternalDataLock;
+		}
+#endif
+
 	protected:
 
 		CHAOS_API void TrackGTParticle_External(FGeometryParticle& Particle);
diff --git a/Engine/Source/Runtime/Experimental/Chaos/Public/Chaos/PBDJointConstraintData.h b/Engine/Source/Runtime/Experimental/Chaos/Public/Chaos/PBDJointConstraintData.h
--- a/Engine/Source/Runtime/Experimental/Chaos/Public/Chaos/PBDJointConstraintData.h
+++ b/Engine/Source/Runtime/Experimental/Chaos/Public/Chaos/PBDJointConstraintData.h
@@ -45,6 +45,8 @@ namespace Chaos
 		virtual ~FJointConstraint() override {}
 
 		const FPBDJointSettings& GetJointSettings()const { return JointSettings.Read(); }
+	
+		void SetJointSettings(const FPBDJointSettings& NewSettings);
 
 		// If we created particle to serve as kinematic endpoint, track so we can release later. This will add particle to solver.
 		CHAOS_API void SetKinematicEndPoint(FSingleParticlePhysicsProxy* InParticle, FPBDRigidsSolver* Solver);
diff --git a/Engine/Source/Runtime/Experimental/Chaos/Public/Chaos/ParticleDirtyFlags.h b/Engine/Source/Runtime/Experimental/Chaos/Public/Chaos/ParticleDirtyFlags.h
--- a/Engine/Source/Runtime/Experimental/Chaos/Public/Chaos/ParticleDirtyFlags.h
+++ b/Engine/Source/Runtime/Experimental/Chaos/Public/Chaos/ParticleDirtyFlags.h
@@ -920,6 +920,11 @@ inline FChaosArchive& operator<<(FChaosArchive& Ar,FMaterialData& Data)
 			return Bits == 0;
 		}
 
+		int32 GetBits() const
+		{
+			return Bits;
+		}
+
 	private:
 		int32 Bits;
 	};
diff --git a/Engine/Source/Runtime/Experimental/Chaos/Public/Chaos/ParticleHandle.h b/Engine/Source/Runtime/Experimental/Chaos/Public/Chaos/ParticleHandle.h
--- a/Engine/Source/Runtime/Experimental/Chaos/Public/Chaos/ParticleHandle.h
+++ b/Engine/Source/Runtime/Experimental/Chaos/Public/Chaos/ParticleHandle.h
@@ -3706,6 +3706,16 @@ public:
 		return Particle ? Particle->CastToRigidParticle() : nullptr;
 	}
 
+	FRigidParticleControlFlags ControlFlags() const
+	{
+		return MMiscData.Read().ControlFlags();
+	}
+
+	void SetControlFlags(const FRigidParticleControlFlags& Flags)
+	{
+		MMiscData.Modify(true,MDirtyFlags,Proxy,[Flags](auto& Data){ Data.SetControlFlags(Flags);});
+	}
+
 private:
 	TChaosProperty<FParticleMassProps,EChaosProperty::MassProps> MMassProps;
 	TChaosProperty<FParticleDynamics, EChaosProperty::Dynamics> MDynamics;
diff --git a/Engine/Source/Runtime/Experimental/Chaos/Public/Chaos/Serialization/SerializedMultiPhysicsState.h b/Engine/Source/Runtime/Experimental/Chaos/Public/Chaos/Serialization/SerializedMultiPhysicsState.h
new file mode 100644
--- /dev/null
+++ b/Engine/Source/Runtime/Experimental/Chaos/Public/Chaos/Serialization/SerializedMultiPhysicsState.h
@@ -0,0 +1,110 @@
+// Copyright Epic Games, Inc. All Rights Reserved.
+
+#pragma once
+#include "Containers/Array.h"
+#include "SolverSerializer.h"
+#include "Serialization/Archive.h"
+
+namespace Chaos
+{
+	struct UE_INTERNAL FSerializedMultiPhysicsState
+	{
+		FSerializedMultiPhysicsState()
+		{
+		}
+
+		FSerializedMultiPhysicsState(ESerializedDataContext DataContext)
+		{
+			Header.DataContext = DataContext;
+		}
+
+		void ReadElementDataIntoBuffer(FSerializedDataBuffer& TargetBuffer);
+
+		template<typename ElementType>
+		void AppendDataFromElement(ElementType& InElement)
+		{
+			FDataEntryTag DataTag;
+			DataTag.DataOffset = GetSize();
+				
+			InElement.SerializeState(MigratedStateAsBytes, Header.DataContext);
+
+			DataTag.DataSize = GetSize() - DataTag.DataOffset;
+
+			Header.DataTagPerElementIndex.Emplace(DataTag);
+		}
+
+		void Serialize(FArchive& Ar);
+
+		int32 GetSize() const
+		{
+			return MigratedStateAsBytes.GetSize();
+		}
+
+	private:
+
+		FSerializedDataBuffer MigratedStateAsBytes;
+
+		struct FDataEntryTag
+		{
+			int32 DataOffset = INDEX_NONE;
+			int32 DataSize = INDEX_NONE;
+
+			friend FArchive& operator<<(FArchive& Ar, FDataEntryTag& Data)
+			{
+				Ar << Data.DataOffset;
+				Ar << Data.DataSize;
+				return Ar;
+			}
+
+			bool IsValid() const
+			{
+				return DataOffset != INDEX_NONE && DataSize != INDEX_NONE;
+			}
+		};
+	public:
+		struct FHeader
+		{
+			ESerializedDataContext DataContext = Chaos::ESerializedDataContext::Invalid;
+			int32 NumBodies = 0;
+			int32 NumConstraints = 0;
+
+			TArray<FDataEntryTag> DataTagPerElementIndex;
+				
+			friend FArchive& operator<<(FArchive& Ar, FHeader& Data)
+			{
+				Ar << Data.DataContext;
+				Ar << Data.NumBodies;
+				Ar << Data.NumConstraints;
+				Ar << Data.DataTagPerElementIndex;
+
+				return Ar;
+			}
+
+			bool IsValid() const
+			{
+				return DataContext != ESerializedDataContext::Invalid && NumBodies + NumConstraints == DataTagPerElementIndex.Num();
+			}
+		};
+
+		void SetNumOfBodies(int32 NumBodies)
+		{
+			Header.NumBodies = NumBodies;
+		}
+
+		void SetNumOfConstraints(int32 NumConstraints)
+		{
+			Header.NumConstraints = NumConstraints;
+		}
+
+		const FHeader& GetHeader() const
+		{
+			return Header;
+		}
+
+	private:
+		FHeader Header;
+		int32 CurrentReadElementIndex = 0;
+	};
+}
+
+
diff --git a/Engine/Source/Runtime/Experimental/Chaos/Public/Chaos/Serialization/SolverSerializer.h b/Engine/Source/Runtime/Experimental/Chaos/Public/Chaos/Serialization/SolverSerializer.h
--- a/Engine/Source/Runtime/Experimental/Chaos/Public/Chaos/Serialization/SolverSerializer.h
+++ b/Engine/Source/Runtime/Experimental/Chaos/Public/Chaos/Serialization/SolverSerializer.h
@@ -1,85 +1,78 @@
 // Copyright Epic Games, Inc. All Rights Reserved.
 
 #pragma once
-#include "Chaos/ImplicitFwd.h"
+
 #include "Chaos/ParticleHandleFwd.h"
-#include "Containers/Array.h"
-#include "Chaos/Framework/PhysicsProxyBase.h"
 #include "Chaos/Serialization/SerializedDataBuffer.h"
-#include "DataWrappers/ChaosVDImplicitObjectDataWrapper.h"
-#include "SolverSerializer.generated.h"
-
-namespace Chaos
-{
-struct FSerializedDataBufferHandle;
-}
-
-namespace Chaos
-{
-class FChaosArchive;
-}
+#include "Containers/StripedMap.h"
+#include "UObject/ObjectMacros.h"
 
-struct FChaosVDParticlePairMidPhase;
-struct FChaosVDCharacterGroundConstraint;
-struct FChaosVDCharacterGroundConstraintDataDataWrapper;
 struct FChaosVDJointConstraint;
 struct FChaosVDParticleDataWrapper;
-
-using FSerializedImplicitObject = FChaosVDImplicitObjectDataWrapper<Chaos::FImplicitObjectPtr, Chaos::FChaosArchive>;
-
-/**
- * Structure holding the serialized state of a rigid solver
- * Used to save it to disk or to re-hydrate a new solver instance
- */
-USTRUCT()
-struct UE_INTERNAL FSerializedSolverScene
-{
-	GENERATED_BODY()
-
-	UPROPERTY()
-	TArray<FChaosVDParticleDataWrapper> ParticleData;
-	UPROPERTY()
-	TArray<FChaosVDJointConstraint> JointConstraintData;
-	UPROPERTY()
-	TArray<FChaosVDCharacterGroundConstraint> CharacterGroundConstraintData;
-	UPROPERTY()
-	TArray<FChaosVDParticlePairMidPhase> CollisionMidPhaseData;
-
-	TArray<FSerializedImplicitObject> ImplicitObjectData;
-};
+class IPhysicsProxyBase;
 
 namespace Chaos
 {
-	class FPBDJointConstraintHandle;
+	class FConstraintBase;
 	class FConstraintHandle;
+	class FJointConstraint;
+	class FPBDJointConstraintHandle;
 	class FPBDRigidsSolver;
+	struct FSerializedDataBufferHandle;
 
 	typedef TUniquePtr<FSerializedDataBuffer> FSerializedDataBufferPtr;
 
+	/** Enum used to indicate where we should read/write the data in serialize request */
+	enum class ESerializedDataContext : uint8
+	{
+		Invalid,
+		/** The data is owned by the solver */
+		Internal,
+		/** The data is owned by the game thread */
+		External,
+		/** Serialize the data from both, the GT and PT*/
+		Both
+	};
+	
+	CHAOS_API const TCHAR* LexToString(ESerializedDataContext Value);
+
 	/* Object capable of serializing total or partially a rigid solver instance **/
 	class FSolverSerializer
 	{
 	public:
+	
 		FSolverSerializer(FPBDRigidsSolver* InSolver);
 
 		void SerializeToBuffer(FSerializedDataBuffer& OutSerializedData);
 		void PopulateFromSerializedBuffer(const FSerializedDataBuffer& InSerializedData);
 
 		CHAOS_API void SerializeParticleStateToBuffer(FGeometryParticleHandle* InParticleHandle, FSerializedDataBuffer& OutSerializedData);
+		CHAOS_API void SerializeParticleStateToBuffer(FGeometryParticle* GTParticle, FSerializedDataBuffer& OutSerializedData);
+
 		CHAOS_API void SerializeConstraintStateToBuffer(FConstraintHandle* ConstraintHandlePtr, FSerializedDataBuffer& OutSerializedData);
+		CHAOS_API void SerializeConstraintStateToBuffer(FConstraintBase* GTConstraintPtr, FSerializedDataBuffer& OutSerializedData);
+
+		CHAOS_API void ApplySerializedStateToParticle(FGeometryParticleHandle* InParticleHandle, FSerializedDataBuffer& InSerializedData);
+		CHAOS_API void ApplySerializedStateToParticle(FGeometryParticleHandle* InParticleHandle, const FChaosVDParticleDataWrapper& InParticleState);
 
-		void ApplySerializedStateToParticle(FGeometryParticleHandle* InParticleHandle, FSerializedDataBuffer& InSerializedData);
-		void ApplySerializedStateToParticle(FGeometryParticleHandle* InParticleHandle, const FChaosVDParticleDataWrapper& InParticleState);
+		CHAOS_API void ApplySerializedStateToParticle(FGeometryParticle* GTParticle, FSerializedDataBuffer& InSerializedData);
+		CHAOS_API void ApplySerializedStateToParticle(FGeometryParticle* GTParticle, const FChaosVDParticleDataWrapper& InParticleState);
 
-		void ApplySerializedStateToJointConstraint(FPBDJointConstraintHandle* ConstraintHandlePtr, FSerializedDataBuffer& InSerializedData);
+		CHAOS_API void ApplySerializedStateToJointConstraint(FPBDJointConstraintHandle* ConstraintHandlePtr, FSerializedDataBuffer& InSerializedData);
+		CHAOS_API void ApplySerializedStateToJointConstraint(FJointConstraint* ConstraintPtr, FSerializedDataBuffer& InSerializedData);
 
-		CHAOS_API void PushPendingSerializedStateForProxy(IPhysicsProxyBase* Proxy, FSerializedDataBufferPtr&& InState);
-		CHAOS_API FSerializedDataBufferPtr PopPendingSerializedStateForProxy(IPhysicsProxyBase* Proxy);
+		CHAOS_API void ApplySerializedStateToConstraint(FConstraintHandle* ConstraintHandlePtr, FSerializedDataBuffer& InSerializedData);
+		CHAOS_API void ApplySerializedStateToConstraint(FConstraintBase* ConstraintPtr, FSerializedDataBuffer& InSerializedData);
+
+		CHAOS_API void PushPendingInternalSerializedStateForProxy(IPhysicsProxyBase* Proxy, FSerializedDataBufferPtr&& InState);
+		CHAOS_API FSerializedDataBufferPtr PopPendingInternalSerializedStateForProxy(IPhysicsProxyBase* Proxy);
 
 
 	private:
 		FPBDRigidsSolver* SolverInstance;
-		TMap<IPhysicsProxyBase*, FSerializedDataBufferPtr> PendingMigratedPhysicsStateByProxy;
+
+		static constexpr int32 StripeCount = 32;
+		TStripedMap<StripeCount, IPhysicsProxyBase*, FSerializedDataBufferPtr> PendingMigratedPhysicsStateByProxy;
 	};
 }
 
diff --git a/Engine/Source/Runtime/Experimental/Chaos/Public/ChaosVisualDebugger/ChaosVDDataWrapperUtils.h b/Engine/Source/Runtime/Experimental/Chaos/Public/ChaosVisualDebugger/ChaosVDDataWrapperUtils.h
--- a/Engine/Source/Runtime/Experimental/Chaos/Public/ChaosVisualDebugger/ChaosVDDataWrapperUtils.h
+++ b/Engine/Source/Runtime/Experimental/Chaos/Public/ChaosVisualDebugger/ChaosVDDataWrapperUtils.h
@@ -10,8 +10,11 @@
 #include "DataWrappers/ChaosVDAccelerationStructureDataWrappers.h"
 #include "Chaos/ISpatialAccelerationCollection.h"
 
+struct FChaosVDJointSettingsDataWrapper;
+
 namespace Chaos
 {
+class FJointConstraint;
 class FSolverSerializer;
 }
 
@@ -82,26 +85,32 @@ private:
 
 	/** Takes a FPBDCollisionConstraintMaterial and copies the relevant data to the CVD counterpart */
 	static void CopyCollisionMaterialToDataWrapper(const Chaos::FPBDCollisionConstraintMaterial& InCopyFrom, FChaosVDCollisionMaterial& OutCopyTo);
-
+public:
 	/** Creates and populates a FChaosVDParticleDataWrapper with the data of the provided FGeometryParticleHandle */
-	static FChaosVDParticleDataWrapper BuildParticleDataWrapperFromParticle(const Chaos::FGeometryParticleHandle* ParticleHandlePtr);
-	static FChaosVDParticleDataWrapper BuildParticleDataWrapperFromGTParticle(Chaos::FGeometryParticle& GTParticle);
+	UE_INTERNAL static FChaosVDParticleDataWrapper BuildParticleDataWrapperFromParticle(const Chaos::FGeometryParticleHandle* ParticleHandlePtr);
+	UE_INTERNAL static FChaosVDParticleDataWrapper BuildParticleDataWrapperFromGTParticle(const Chaos::FGeometryParticle* GTParticle);
 
+private:
 	/** Creates and populates a FChaosVDConstraint with the data of the provided FPBDCollisionConstraint */
 	static FChaosVDConstraint BuildConstraintDataWrapperFromConstraint(const Chaos::FPBDCollisionConstraint& InConstraint);
 
 	/** Creates and populates a FChaosVDParticlePairMidPhase with the data of the provided FParticlePairMidPhase */
 	static FChaosVDParticlePairMidPhase BuildMidPhaseDataWrapperFromMidPhase(const Chaos::FParticlePairMidPhase& InMidPhase);
+	static void CopyJointConstraintSettingsToWrapper(const Chaos::FPBDJointSettings& InSettings, FChaosVDJointSettingsDataWrapper& InOutWrappedJointSettingsData);
+
 public:
 	/** Creates and populates a FChaosVDJointConstraint with the data of the provided FPBDJointConstraintHandle */
-	static FChaosVDJointConstraint BuildJointDataWrapper(const Chaos::FPBDJointConstraintHandle* ConstaintHanlde);
+	UE_INTERNAL static FChaosVDJointConstraint BuildJointDataWrapper(const Chaos::FPBDJointConstraintHandle* ConstaintHanlde);
+	UE_INTERNAL static FChaosVDJointConstraint BuildGTJointDataWrapper(Chaos::FJointConstraint* GTConstraintPtr);
 
-	static void ApplyJointDataWrapperToHandle(Chaos::FPBDJointConstraintHandle* ConstaintHanlde, const FChaosVDJointConstraint& WrappedJointData);
+	UE_INTERNAL static void ApplyJointDataWrapperToHandle(Chaos::FPBDJointConstraintHandle* ConstaintHanlde, const FChaosVDJointConstraint& WrappedJointData);
+	UE_INTERNAL static void ApplyJointDataWrapperGTConstraint(Chaos::FJointConstraint* ConstraintPtr, const FChaosVDJointConstraint& InWrappedJointData);
 
+	/** Creates and populates a FChaosVDChartacterGroundConstraint with the data of the provided FCharacterGroundConstraintHandle */
+	UE_INTERNAL static FChaosVDCharacterGroundConstraint BuildCharacterGroundConstraintDataWrapper(const Chaos::FCharacterGroundConstraintHandle* ConstaintHanlde);
 private:
 
-	/** Creates and populates a FChaosVDChartacterGroundConstraint with the data of the provided FCharacterGroundConstraintHandle */
-	static FChaosVDCharacterGroundConstraint BuildCharacterGroundConstraintDataWrapper(const Chaos::FCharacterGroundConstraintHandle* ConstaintHanlde);
+	static void CopyJointSettingsFromWrapper(const FChaosVDJointSettingsDataWrapper& InWrappedJointSettingsData, Chaos::FPBDJointSettings& NewSettings);
 
 	/** Converts a Chaos::FVec3 to a FVector. It is worth notice that FVector is double precision and FVec3 is single */
 	static FVector ConvertToFVector(const Chaos::FVec3f& VectorRef) { return FVector(VectorRef); }
@@ -124,7 +133,6 @@ private:
 	static void AddTreeLeaves(const TConstArrayView<Chaos::TBoundingVolume<Chaos::FAccelerationStructureHandle>>& LeavesContainer, FChaosVDAABBTreeDataWrapper& InOutAABBTreeWrapper);
 
 	friend FChaosVisualDebuggerTrace;
-	friend Chaos::FSolverSerializer;
 };
 
 template<typename LeafType>
diff --git a/Engine/Source/Runtime/Experimental/ChaosVisualDebugger/Private/DataWrappers/ChaosVDJointDataWrappers.cpp b/Engine/Source/Runtime/Experimental/ChaosVisualDebugger/Private/DataWrappers/ChaosVDJointDataWrappers.cpp
--- a/Engine/Source/Runtime/Experimental/ChaosVisualDebugger/Private/DataWrappers/ChaosVDJointDataWrappers.cpp
+++ b/Engine/Source/Runtime/Experimental/ChaosVisualDebugger/Private/DataWrappers/ChaosVDJointDataWrappers.cpp
@@ -47,6 +47,42 @@ bool FChaosVDJointStateDataWrapper::Serialize(FArchive& Ar)
 	return !Ar.IsError();
 }
 
+bool FChaosVDGTJointStateDataWrapper::Serialize(FArchive& Ar)
+{
+	Ar << bHasValidData;
+	if (!bHasValidData)
+	{
+		return !Ar.IsError();
+	}
+
+	EChaosVDGTJointStateFlags StateFlags = EChaosVDGTJointStateFlags::None;
+	
+	if (Ar.IsLoading())
+	{
+		Ar << StateFlags;
+		CVD_UNPACK_BITFIELD_DATA(bIsBreaking, StateFlags, EChaosVDGTJointStateFlags::IsBreaking);
+		CVD_UNPACK_BITFIELD_DATA(bIsBroken, StateFlags, EChaosVDGTJointStateFlags::IsBroken);
+		CVD_UNPACK_BITFIELD_DATA(bDriveTargetChanged, StateFlags, EChaosVDGTJointStateFlags::DriveTargetChanged);
+		CVD_UNPACK_BITFIELD_DATA(bIsViolating, StateFlags, EChaosVDGTJointStateFlags::IsViolating);
+	}
+	else
+	{
+		CVD_UNPACK_BITFIELD_DATA(bIsBreaking, StateFlags, EChaosVDGTJointStateFlags::IsBreaking);
+		CVD_UNPACK_BITFIELD_DATA(bIsBroken, StateFlags, EChaosVDGTJointStateFlags::IsBroken);
+		CVD_UNPACK_BITFIELD_DATA(bDriveTargetChanged, StateFlags, EChaosVDGTJointStateFlags::DriveTargetChanged);
+		CVD_UNPACK_BITFIELD_DATA(bIsViolating, StateFlags, EChaosVDGTJointStateFlags::IsViolating);
+
+		Ar << StateFlags;
+	}
+
+	Ar <<  Force;
+	Ar <<  Torque;
+	Ar <<  LinearViolation;
+	Ar <<  AngularViolation;
+
+	return !Ar.IsError();
+}
+
 bool FChaosVDJointSolverSettingsDataWrapper::Serialize(FArchive& Ar)
 {
 	Ar << bHasValidData;
@@ -233,7 +269,7 @@ bool FChaosVDJointConstraint::Serialize(FArchive& Ar)
 	CVD_SERIALIZE_STATIC_ARRAY(Ar, ParticleParIndexes);
 
 	Ar << JointSettings;
-	Ar << JointState;
+	Ar << PhysicsThreadJointState;
 
 	return !Ar.IsError();
 }
diff --git a/Engine/Source/Runtime/Experimental/ChaosVisualDebugger/Private/DataWrappers/ChaosVDParticleDataWrapper.cpp b/Engine/Source/Runtime/Experimental/ChaosVisualDebugger/Private/DataWrappers/ChaosVDParticleDataWrapper.cpp
--- a/Engine/Source/Runtime/Experimental/ChaosVisualDebugger/Private/DataWrappers/ChaosVDParticleDataWrapper.cpp
+++ b/Engine/Source/Runtime/Experimental/ChaosVisualDebugger/Private/DataWrappers/ChaosVDParticleDataWrapper.cpp
@@ -155,6 +155,35 @@ bool FChaosVDParticleCluster::Serialize(FArchive& Ar)
 	return !Ar.IsError();
 }
 
+bool FChaosVDKinematicTarget::Serialize(FArchive& Ar)
+{
+	Ar << bHasValidData;
+	if (!bHasValidData)
+	{
+		return !Ar.IsError();
+	}
+
+	Ar << Rotation;
+	Ar << Position;
+	Ar << Mode;
+
+	return !Ar.IsError();
+}
+
+bool FChaosVDVSmooth::Serialize(FArchive& Ar)
+{
+	Ar << bHasValidData;
+	if (!bHasValidData)
+	{
+		return !Ar.IsError();
+	}
+
+	Ar << MV;
+	Ar << MW;
+
+	return !Ar.IsError();
+}
+
 bool FChaosVDParticleDataWrapper::Serialize(FArchive& Ar)
 {
 	Ar << bHasValidData;
diff --git a/Engine/Source/Runtime/Experimental/ChaosVisualDebugger/Public/DataWrappers/ChaosVDJointDataWrappers.h b/Engine/Source/Runtime/Experimental/ChaosVisualDebugger/Public/DataWrappers/ChaosVDJointDataWrappers.h
--- a/Engine/Source/Runtime/Experimental/ChaosVisualDebugger/Public/DataWrappers/ChaosVDJointDataWrappers.h
+++ b/Engine/Source/Runtime/Experimental/ChaosVisualDebugger/Public/DataWrappers/ChaosVDJointDataWrappers.h
@@ -36,9 +36,19 @@ enum class EChaosVDJointStateFlags : uint8
 	DriveTargetChanged = 1 << 3,
 	EnabledDuringResim = 1 << 4,
 };
-
 ENUM_CLASS_FLAGS(EChaosVDJointStateFlags)
 
+UENUM()
+enum class EChaosVDGTJointStateFlags : uint8
+{
+	None = 0,
+	IsBreaking = 1 << 0,
+	IsBroken = 1 << 1,
+	DriveTargetChanged = 1 << 2,
+	IsViolating = 1 << 3,
+};
+ENUM_CLASS_FLAGS(EChaosVDGTJointStateFlags)
+
 USTRUCT()
 struct FChaosVDJointStateDataWrapper : public FChaosVDWrapperDataBase
 {
@@ -77,6 +87,33 @@ public:
 };
 CVD_IMPLEMENT_SERIALIZER(FChaosVDJointStateDataWrapper)
 
+USTRUCT()
+struct FChaosVDGTJointStateDataWrapper : public FChaosVDWrapperDataBase
+{
+	GENERATED_BODY()
+public:
+
+	CHAOSVDRUNTIME_API bool Serialize(FArchive& Ar);
+
+	UPROPERTY(VisibleAnywhere, Category=JointState)
+	uint8 bIsBreaking : 1 = false;
+	UPROPERTY(VisibleAnywhere, Category=JointState)
+	uint8 bIsBroken : 1 = false;
+	UPROPERTY(VisibleAnywhere, Category=JointState)
+	uint8 bDriveTargetChanged : 1 = false;
+	UPROPERTY(VisibleAnywhere, Category=JointState)
+	uint8 bIsViolating : 1 = false;
+	UPROPERTY(VisibleAnywhere, Category=JointState)
+	FVector Force = FVector::ZeroVector;
+	UPROPERTY(VisibleAnywhere, Category=JointState)
+	FVector Torque = FVector::ZeroVector;
+	UPROPERTY(VisibleAnywhere, Category=JointState)
+	float LinearViolation = 0.f;
+	UPROPERTY(VisibleAnywhere, Category=JointState)
+	float AngularViolation = 0.f;
+};
+CVD_IMPLEMENT_SERIALIZER(FChaosVDGTJointStateDataWrapper)
+
 UENUM()
 enum class EChaosVDJointSolverSettingsFlags : uint16
 {
@@ -380,7 +417,10 @@ struct FChaosVDJointConstraint : public FChaosVDConstraintDataWrapperBase
 	int32 ParticleParIndexes[2] = { INDEX_NONE, INDEX_NONE };
 
 	UPROPERTY(VisibleAnywhere, Category=JointState)
-	FChaosVDJointStateDataWrapper JointState;
+	FChaosVDJointStateDataWrapper PhysicsThreadJointState;
+
+	UPROPERTY()
+	FChaosVDGTJointStateDataWrapper GameThreadJointState;
 	
 	UPROPERTY(VisibleAnywhere, Category=JointSettings)
 	FChaosVDJointSettingsDataWrapper JointSettings;
diff --git a/Engine/Source/Runtime/Experimental/ChaosVisualDebugger/Public/DataWrappers/ChaosVDParticleDataWrapper.h b/Engine/Source/Runtime/Experimental/ChaosVisualDebugger/Public/DataWrappers/ChaosVDParticleDataWrapper.h
--- a/Engine/Source/Runtime/Experimental/ChaosVisualDebugger/Public/DataWrappers/ChaosVDParticleDataWrapper.h
+++ b/Engine/Source/Runtime/Experimental/ChaosVisualDebugger/Public/DataWrappers/ChaosVDParticleDataWrapper.h
@@ -166,15 +166,45 @@ struct FChaosVDParticlePositionRotation : public FChaosVDWrapperDataBase
 	{
 	}
 
-	template <typename TOther>
-	void CopyFrom(const TOther& Other)
+	enum class EAccessorType
+	{
+		PQ,
+		XR
+	};
+
+	template <typename OtherType, EAccessorType AccessorType>
+	void CopyFrom(const OtherType& Other)
 	{
-		MX = Other->GetP();
-		MR = Other->GetQ();
+		if constexpr (AccessorType == EAccessorType::PQ)
+		{
+			MX = Other.GetP();
+			MR = Other.GetQ();
+
+		}
+		else if (AccessorType == EAccessorType::XR)
+		{
+			MX = Other.GetX();
+			MR = Other.GetR();
+		}
 
 		bHasValidData = true;
 	}
 
+	template <typename OtherType, EAccessorType AccessorType>
+	void CopyTo(OtherType& Other) const
+	{
+		if constexpr (AccessorType == EAccessorType::PQ)
+		{
+			Other.SetP(MX);
+			Other.SetQ(MR);
+		}
+		else if (AccessorType == EAccessorType::XR)
+		{
+			Other.SetX(MX);
+			Other.SetR(MR);
+		}
+	}
+
 	CHAOSVDRUNTIME_API bool Serialize(FArchive& Ar);
 
 	UPROPERTY(VisibleAnywhere, Category= "Particle Velocities")
@@ -252,10 +282,10 @@ struct FChaosVDParticleDynamics : public FChaosVDWrapperDataBase
 	template <typename TOther>
 	void CopyTo(TOther& Other) const
 	{
-		Other.Acceleration() = MAcceleration;
-        Other.AngularAcceleration() = MAngularAcceleration;
-		Other.LinearImpulseVelocity() =  MLinearImpulseVelocity;
-        Other.AngularImpulseVelocity() = MAngularImpulseVelocity;
+		Other.SetAcceleration(MAcceleration);
+        Other.SetAngularAcceleration(MAngularAcceleration);
+		Other.SetLinearImpulseVelocity(MLinearImpulseVelocity);
+        Other.SetAngularImpulseVelocity(MAngularImpulseVelocity);
 	}
 
 	UPROPERTY(VisibleAnywhere, Category= "Particle Dynamics")
@@ -305,10 +335,10 @@ struct FChaosVDParticleMassProps : public FChaosVDWrapperDataBase
 	{
 		Other.SetCenterOfMass(MCenterOfMass);
 		Other.SetRotationOfMass(MRotationOfMass);
-		Other.I() = MI;
-		Other.InvI() = MInvI;
-		Other.M() = MM;
-		Other.InvM() = MInvM;
+		Other.SetI(MI);
+		Other.SetInvI(MInvI);
+		Other.SetM(MM);
+		Other.SetInvM(MInvM);
 	}
 	
 	UPROPERTY(VisibleAnywhere, Category= "Particle Mass Props")
@@ -368,13 +398,13 @@ struct FChaosVDParticleDynamicMisc : public FChaosVDWrapperDataBase
 		bHasValidData = true;
 	}
 	
-	template <typename TOther, typename ControlFlagsType, typename ObjectStateType, typename SleepStateType>
-	void CopyTo(TOther& Other) const
+	template <typename OtherType, typename ControlFlagsType, typename SleepStateType>
+	void CopyWithoutStateTo(OtherType& Other) const
 	{
-		Other.LinearEtherDrag() = MLinearEtherDrag;
-		Other.AngularEtherDrag() = MAngularEtherDrag;
-		Other.MaxLinearSpeedSq() = MMaxLinearSpeedSq;
-		Other.MaxAngularSpeedSq() = MMaxAngularSpeedSq;
+		Other.SetLinearEtherDrag(MLinearEtherDrag);
+		Other.SetAngularEtherDrag(MAngularEtherDrag);
+		Other.SetMaxLinearSpeedSq(MMaxLinearSpeedSq);
+		Other.SetMaxAngularSpeedSq(MMaxAngularSpeedSq);
 		Other.SetInitialOverlapDepenetrationVelocity(MInitialOverlapDepenetrationVelocity);
 		Other.SetSleepThresholdMultiplier(MSleepThresholdMultiplier);
 		Other.SetCollisionGroup(MCollisionGroup);
@@ -385,8 +415,6 @@ struct FChaosVDParticleDynamicMisc : public FChaosVDWrapperDataBase
 		MControlFlags.CopyTo(ControlFlags);
 
 		Other.SetControlFlags(ControlFlags);
-
-		Other.SetObjectStateLowLevel(static_cast<ObjectStateType>(MObjectState));
 		
 		Other.SetDisabled(bDisabled);
 	}
@@ -545,6 +573,94 @@ struct FChaosVDParticleCluster : public FChaosVDWrapperDataBase
 
 CVD_IMPLEMENT_SERIALIZER(FChaosVDParticleCluster)
 
+UENUM()
+enum class EChaosVDParticleContext
+{
+	Invalid,
+	GameThread,
+	PhysicsThread,
+};
+
+UENUM()
+enum class EChaosVDKinematicTargetMode
+{
+	None,			/** Particle does not move and no data is changed */
+	Reset,			/** Particle does not move, velocity and angular velocity are zeroed, then mode is set to "None". */
+	Position,		/** Particle is moved to Kinematic Target transform, velocity and angular velocity updated to reflect the change, then mode is set to "Reset". */
+	Velocity,		/** Particle is moved based on velocity and angular velocity, mode remains as "Velocity" until changed. */
+};
+
+USTRUCT()
+struct FChaosVDKinematicTarget : public FChaosVDWrapperDataBase
+{
+	GENERATED_BODY()
+
+	CHAOSVDRUNTIME_API bool Serialize(FArchive& Ar);
+
+	template <typename OtherType>
+	void CopyFrom(const OtherType& Other)
+	{
+		Mode = static_cast<EChaosVDKinematicTargetMode>(Other.GetMode());
+
+		if (Mode == EChaosVDKinematicTargetMode::Position)
+		{
+			Position = Other.GetPosition();
+			Rotation = FQuat(Other.GetRotation());
+		}
+
+		bHasValidData = true;
+	}
+
+	template <typename OtherType, class ModeType>
+	void CopyTo(OtherType& Other) const
+	{
+		Other.SetTargetMode(Position, Rotation);
+		Other.SetMode(static_cast<ModeType>(Mode));
+	}
+
+	UPROPERTY(VisibleAnywhere, Category= "Kinematic Target")
+	FQuat Rotation = FQuat(ForceInit);
+
+	UPROPERTY(VisibleAnywhere, Category= "Kinematic Target")
+	FVector Position = FVector(ForceInit);;
+
+	UPROPERTY(VisibleAnywhere, Category= "Kinematic Target")
+	EChaosVDKinematicTargetMode Mode = EChaosVDKinematicTargetMode::None;
+};
+
+CVD_IMPLEMENT_SERIALIZER(FChaosVDKinematicTarget)
+
+USTRUCT()
+struct FChaosVDVSmooth : public FChaosVDWrapperDataBase
+{
+	GENERATED_BODY()
+
+	CHAOSVDRUNTIME_API bool Serialize(FArchive& Ar);
+
+	template <typename OtherType>
+	void CopyFrom(const OtherType& Other)
+	{
+		MV = Other.VSmooth();
+		MW = Other.WSmooth();
+
+		bHasValidData = true;
+	}
+
+	template <typename OtherType>
+	void CopyTo(OtherType& Other) const
+	{
+		Other.SetVSmooth(MV);
+		Other.SetWSmooth(MW);
+	}
+
+	UPROPERTY(VisibleAnywhere, Category= "VSmooth")
+	FVector MV;
+	UPROPERTY(VisibleAnywhere, Category= "VSmooth")
+	FVector MW;
+};
+
+CVD_IMPLEMENT_SERIALIZER(FChaosVDVSmooth)
+
 /** Simplified UStruct version of FChaosVDParticleDataWrapper.
  * Used to be able to show the values in the editor and allow changes via the Property Editor.
  */
@@ -561,6 +677,12 @@ struct FChaosVDParticleDataWrapper : public FChaosVDWrapperDataBase
 	{
 	}
 
+	UPROPERTY()
+	int32 DirtyFlagsBits = 0;
+
+	UPROPERTY()
+	EChaosVDParticleContext ParticleContext = EChaosVDParticleContext::Invalid;
+
 	UPROPERTY(VisibleAnywhere, Category= "General")
 	uint32 GeometryHash = 0;
 
@@ -585,6 +707,12 @@ struct FChaosVDParticleDataWrapper : public FChaosVDWrapperDataBase
 	UPROPERTY(VisibleAnywhere, Category= "Particle Velocities")
 	FChaosVDParticleVelocities ParticleVelocities;
 
+	UPROPERTY(VisibleAnywhere, Category= "Particle Kinematic Target")
+	FChaosVDKinematicTarget ParticleKinematicTarget;
+
+	UPROPERTY(VisibleAnywhere, Category= "Particle V/W Smooth")
+	FChaosVDVSmooth ParticleVWSmooth;
+
 	UPROPERTY(VisibleAnywhere, Category= "Particle Dynamics")
 	FChaosVDParticleDynamics ParticleDynamics;
 
diff --git a/Engine/Source/Runtime/PhysicsCore/Private/ChaosEngineInterface.cpp b/Engine/Source/Runtime/PhysicsCore/Private/ChaosEngineInterface.cpp
--- a/Engine/Source/Runtime/PhysicsCore/Private/ChaosEngineInterface.cpp
+++ b/Engine/Source/Runtime/PhysicsCore/Private/ChaosEngineInterface.cpp
@@ -2097,7 +2097,7 @@ void FChaosEngineInterface::SetLocalTransform(const FPhysicsShapeHandle& InShape
 	}
 }
 
-void FChaosEngineInterface::SerializeParticleStateToBuffer(const FPhysicsActorHandle& InActorHandle, Chaos::FSerializedDataBuffer& OutSerializedData)
+void FChaosEngineInterface::SerializeParticleStateToBuffer(const FPhysicsActorHandle& InActorHandle, Chaos::FSerializedDataBuffer& OutSerializedData, Chaos::ESerializedDataContext DataContext)
 {
 	if (!IsValid(InActorHandle))
 	{
@@ -2107,40 +2107,100 @@ void FChaosEngineInterface::SerializeParticleStateToBuffer(const FPhysicsActorHa
 
 	using namespace Chaos;
 
-	if (Chaos::FPBDRigidsSolver* Solver = InActorHandle->GetSolver<Chaos::FPBDRigidsSolver>())
+	auto SerializeInternalStateToBuffer = [InActorHandle](FPBDRigidsSolver& Solver, FSerializedDataBuffer& OutData)
 	{
 		if (InActorHandle->GetHandle_LowLevel())
 		{
-			Solver->GetSerializer().SerializeParticleStateToBuffer(InActorHandle->GetHandle_LowLevel(), OutSerializedData);
+			Solver.GetSerializer().SerializeParticleStateToBuffer(InActorHandle->GetHandle_LowLevel(), OutData);
 		}
-		else 
+		else
 		{
-			if (FSerializedDataBufferPtr PendingDataHandle = Solver->GetSerializer().PopPendingSerializedStateForProxy(InActorHandle))
+			if (FSerializedDataBufferPtr PendingDataHandle = Solver.GetSerializer().PopPendingInternalSerializedStateForProxy(InActorHandle))
 			{
-				OutSerializedData.GetDataAsByteArrayRef().Append(PendingDataHandle->GetDataAsByteArrayRef());
+				OutData.GetDataAsByteArrayRef().Append(PendingDataHandle->GetDataAsByteArrayRef());
 			}
 			else
 			{
-				UE_LOG(LogChaos, Warning, TEXT("[%hs]: Failed to serialize particle | Owner [%s] | Invalid PT Handle and no pending state available."), __func__, *GetNameSafe(InActorHandle->GetOwner()));
+				SerializeParticleStateToBuffer(InActorHandle, OutData, ESerializedDataContext::External);
+				UE_LOG(LogChaos, Warning, TEXT("[%hs]: Failed to get a valid Internal physics state. Falling back to the external state | Owner [%s]."), __func__, *GetNameSafe(InActorHandle->GetOwner()));
 			}
 		}
+	};
+
+	if (FPBDRigidsSolver* Solver = InActorHandle->GetSolver<FPBDRigidsSolver>())
+	{
+		switch (DataContext)
+		{
+			case ESerializedDataContext::External:
+				Solver->GetSerializer().SerializeParticleStateToBuffer(InActorHandle->GetParticle_LowLevel(), OutSerializedData);
+				break;
+			case ESerializedDataContext::Internal:
+				SerializeInternalStateToBuffer(*Solver, OutSerializedData);
+				break;
+			case ESerializedDataContext::Invalid:
+			case ESerializedDataContext::Both:
+			default:
+				{
+					// Note: We don't support serialize data from both states into a single buffer yet. But we do support apply the same state to External and Internal during deserialization if ESerializedDataContext::Both is used.
+					UE_LOG(LogChaos, Warning, TEXT("[%hs]: Failed to serialize particle | Unsupported Data Context [%s]."), __func__, Chaos::LexToString(DataContext));
+					break;
+				}
+		}
 	}
 }
 
-void FChaosEngineInterface::ApplySerializedStateToParticle(const FPhysicsActorHandle& InActorHandle, Chaos::FSerializedDataBufferPtr&& InSerializedData)
+void FChaosEngineInterface::ApplySerializedStateToParticle(const FPhysicsActorHandle& InActorHandle, Chaos::FSerializedDataBufferPtr&& InSerializedData, Chaos::ESerializedDataContext DataContext)
 {
-	if (!IsValid(InActorHandle))
+	if (!IsValid(InActorHandle) || !InSerializedData)
 	{
 		return;
 	}
 
-	if (Chaos::FPBDRigidsSolver* Solver = InActorHandle->GetSolver<Chaos::FPBDRigidsSolver>())
+	using namespace Chaos;
+
+	auto ApplySerializedPTSerializedState = [&InActorHandle](FPBDRigidsSolver& Solver, Chaos::FSerializedDataBufferPtr&& InSerializedData)
+	{
+		if (FGeometryParticleHandle* ParticleHandle = InActorHandle->GetHandle_LowLevel())
+		{
+			Solver.GetSerializer().ApplySerializedStateToParticle(ParticleHandle, *InSerializedData);
+		}
+		else
+		{
+			Solver.GetSerializer().PushPendingInternalSerializedStateForProxy(InActorHandle, MoveTemp(InSerializedData));
+		}
+	};
+
+	if (FPBDRigidsSolver* Solver = InActorHandle->GetSolver<FPBDRigidsSolver>())
 	{
-		Solver->GetSerializer().PushPendingSerializedStateForProxy(InActorHandle, MoveTemp(InSerializedData));
+		switch (DataContext)
+		{
+			case ESerializedDataContext::External:
+				Solver->GetSerializer().ApplySerializedStateToParticle(InActorHandle->GetParticle_LowLevel(), *InSerializedData);
+				break;
+			case ESerializedDataContext::Internal:
+				{
+					if (InActorHandle->IsDirty())
+					{
+						UE_LOG(LogChaos, Warning, TEXT("Attempting to apply an internal serialized state but the external state is dirty. Part of the internal state might get stomp in the next push. Consider using ESerializedDataContext::Both."))
+					}
+					ApplySerializedPTSerializedState(*Solver, MoveTemp(InSerializedData));
+					break;
+				}
+			case ESerializedDataContext::Both:
+				{
+					Solver->GetSerializer().ApplySerializedStateToParticle(InActorHandle->GetParticle_LowLevel(), *InSerializedData);
+					ApplySerializedPTSerializedState(*Solver, MoveTemp(InSerializedData));
+					break;
+				}
+		case ESerializedDataContext::Invalid:
+		default:
+			UE_LOG(LogChaos, Warning, TEXT("[%hs]: Failed to apply serialized state to particle | Unsupported Data Context [%s]."), __func__, Chaos::LexToString(DataContext));
+			break;
+		}
 	}
 }
 
-void FChaosEngineInterface::SerializeConstraintStateToBuffer(const FPhysicsConstraintHandle& ConstraintHandle, Chaos::FSerializedDataBuffer& OutSerializedData)
+void FChaosEngineInterface::SerializeConstraintStateToBuffer(const FPhysicsConstraintHandle& ConstraintHandle, Chaos::FSerializedDataBuffer& OutSerializedData, Chaos::ESerializedDataContext DataContext)
 {
 	if (!ConstraintHandle.IsValid())
 	{
@@ -2150,55 +2210,108 @@ void FChaosEngineInterface::SerializeConstraintStateToBuffer(const FPhysicsConst
 
 	using namespace Chaos;
 
-	switch (EConstraintType Type = ConstraintHandle.Constraint->GetType())
+	IPhysicsProxyBase* Proxy = ConstraintHandle.Constraint->GetProxy();
+
+	if (FPBDRigidsSolver* Solver = Proxy ? Proxy->GetSolver<FPBDRigidsSolver>() : nullptr)
 	{
-	case EConstraintType::NoneType:
-		break;
-	case EConstraintType::JointConstraintType:
+		const bool bIsConstraintProxy = Proxy->GetType() == EPhysicsProxyType::CharacterGroundConstraintType ||
+										Proxy->GetType() == EPhysicsProxyType::JointConstraintType ||
+										Proxy->GetType() == EPhysicsProxyType::SuspensionConstraintType;
+		if (!ensure(bIsConstraintProxy))
 		{
-			if (FJointConstraintPhysicsProxy* JointProxy = ConstraintHandle.Constraint->GetProxy<FJointConstraintPhysicsProxy>())
+			return;
+		}
+
+		auto SerializeInternalStateToBuffer = [&ConstraintHandle](IPhysicsProxyBase* Proxy, FPBDRigidsSolver& Solver, FSerializedDataBuffer& OutData)
+		{
+			if (FConstraintHandle* PTConstraintHandle = static_cast<FConstraintHandle*>(Proxy->GetHandleUnsafe()))
+			{
+				Solver.GetSerializer().SerializeConstraintStateToBuffer(PTConstraintHandle, OutData);
+			}
+			else
 			{
-				if (FPBDRigidsSolver* Solver = JointProxy->GetSolver<FPBDRigidsSolver>())
+				if (FSerializedDataBufferPtr PendingDataHandle = Solver.GetSerializer().PopPendingInternalSerializedStateForProxy(Proxy))
 				{
-					if (JointProxy->GetHandle())
-					{
-						Solver->GetSerializer().SerializeConstraintStateToBuffer(JointProxy->GetHandle(), OutSerializedData);
-					}
-					else
-					{
-						if (FSerializedDataBufferPtr PendingData = Solver->GetSerializer().PopPendingSerializedStateForProxy(JointProxy))
-						{
-							OutSerializedData.GetDataAsByteArrayRef().Append(PendingData->GetDataAsByteArrayRef());
-						}
-						else
-						{
-							UE_LOG(LogChaos, Warning, TEXT("[%hs]: Failed to serialize constraint | Invalid PT Handle and no pending state available."),__func__);
-						}
-					}
+					OutData.GetDataAsByteArrayRef().Append(PendingDataHandle->GetDataAsByteArrayRef());
+				}
+				else
+				{
+					SerializeConstraintStateToBuffer(ConstraintHandle, OutData, ESerializedDataContext::External);
+					UE_LOG(LogChaos, Warning, TEXT("[%hs]: Failed to get a valid Internal physics state. Falling back to the external state | Owner [%s]."), __func__, *GetNameSafe(Proxy->GetOwner()));
 				}
 			}
+		};
+
+		switch (DataContext)
+		{
+		case ESerializedDataContext::External:
+			Solver->GetSerializer().SerializeConstraintStateToBuffer(ConstraintHandle.Constraint, OutSerializedData);
+			break;
+		case ESerializedDataContext::Internal:
+			SerializeInternalStateToBuffer(Proxy, *Solver, OutSerializedData);
+			break;
+		case ESerializedDataContext::Invalid:
+		case ESerializedDataContext::Both:
+		default:
+			{
+				// Note: We don't support serialize data from both states into a single buffer yet. But we do support apply the same state to External and Internal during deserialization if ESerializedDataContext::Both is used.
+				UE_LOG(LogChaos, Warning, TEXT("[%hs]: Failed to serialize constraint | Unsupported Data Context [%s]."), __func__, Chaos::LexToString(DataContext));
+				break;
+			}
 		}
-		break;
-	case EConstraintType::SpringConstraintType:
-		break;
-	case EConstraintType::SuspensionConstraintType:
-		break;
-	case EConstraintType::CharacterGroundConstraintType:
-		break;
 	}
 }
 
-void FChaosEngineInterface::ApplySerializedStateToConstraint(const FPhysicsConstraintHandle& ConstraintHandle, Chaos::FSerializedDataBufferPtr&& InSerializedData)
+void FChaosEngineInterface::ApplySerializedStateToConstraint(const FPhysicsConstraintHandle& ConstraintHandle, Chaos::FSerializedDataBufferPtr&& InSerializedData, Chaos::ESerializedDataContext DataContext)
 {
-	if (!ConstraintHandle.IsValid())
+	if (!ConstraintHandle.IsValid() || !InSerializedData)
 	{
 		return;
 	}
 
+	using namespace Chaos;
+
 	IPhysicsProxyBase* Proxy = ConstraintHandle.Constraint->GetProxy();
-	if (Chaos::FPBDRigidsSolver* Solver = Proxy ? Proxy->GetSolver<Chaos::FPBDRigidsSolver>() : nullptr)
+
+	auto ApplySerializedPTSerializedState = [](IPhysicsProxyBase* Proxy, FPBDRigidsSolver& Solver, Chaos::FSerializedDataBufferPtr&& InSerializedData)
+	{
+		if (FConstraintHandle* PTConstraintHandle = static_cast<FConstraintHandle*>(Proxy->GetHandleUnsafe()))
+		{
+			Solver.GetSerializer().ApplySerializedStateToConstraint(PTConstraintHandle, *InSerializedData);
+		}
+		else
+		{
+			Solver.GetSerializer().PushPendingInternalSerializedStateForProxy(Proxy, MoveTemp(InSerializedData));
+		}
+	};
+
+	if (FPBDRigidsSolver* Solver = Proxy ? Proxy->GetSolver<FPBDRigidsSolver>() : nullptr)
 	{
-		Solver->GetSerializer().PushPendingSerializedStateForProxy(Proxy, MoveTemp(InSerializedData));
+		switch (DataContext)
+		{
+			case ESerializedDataContext::External:
+				Solver->GetSerializer().ApplySerializedStateToConstraint(ConstraintHandle.Constraint, *InSerializedData);
+				break;
+			case ESerializedDataContext::Internal:
+				{
+					if (ConstraintHandle->IsDirty())
+					{
+						UE_LOG(LogChaos, Warning, TEXT("Attempting to apply an internal serialized state but the external state is dirty. Part of the internal state might get stomp in the next push. Consider using ESerializedDataContext::Both."))
+					}
+					ApplySerializedPTSerializedState(Proxy, *Solver, MoveTemp(InSerializedData));
+					break;
+				}
+			case ESerializedDataContext::Both:
+				{
+					Solver->GetSerializer().ApplySerializedStateToConstraint(ConstraintHandle.Constraint, *InSerializedData);
+					ApplySerializedPTSerializedState(Proxy, *Solver, MoveTemp(InSerializedData));
+					break;
+				}
+			case ESerializedDataContext::Invalid:
+			default:
+				UE_LOG(LogChaos, Warning, TEXT("[%hs]: Failed to apply serialized state to particle | Unsupported Data Context [%s]."), __func__, Chaos::LexToString(DataContext));
+				break;	
+		}
 	}
 }
 
diff --git a/Engine/Source/Runtime/PhysicsCore/Public/Chaos/ChaosEngineInterface.h b/Engine/Source/Runtime/PhysicsCore/Public/Chaos/ChaosEngineInterface.h
--- a/Engine/Source/Runtime/PhysicsCore/Public/Chaos/ChaosEngineInterface.h
+++ b/Engine/Source/Runtime/PhysicsCore/Public/Chaos/ChaosEngineInterface.h
@@ -201,6 +201,7 @@ typedef uint8 FMaskFilter;
 
 namespace Chaos
 {
+enum class ESerializedDataContext : uint8;
 struct FSerializedDataBufferHandle;
 class FBVHParticles;
 
@@ -574,9 +575,9 @@ public:
 	static PHYSICSCORE_API void SetIsQueryShape(const FPhysicsShapeHandle& InShape,bool bIsQueryShape);
 	static PHYSICSCORE_API void SetLocalTransform(const FPhysicsShapeHandle& InShape,const FTransform& NewLocalTransform);
 
-	static PHYSICSCORE_API void SerializeParticleStateToBuffer(const FPhysicsActorHandle& InActorHandle, Chaos::FSerializedDataBuffer& OutSerializedData);
-	static PHYSICSCORE_API void ApplySerializedStateToParticle(const FPhysicsActorHandle& InActorHandle, TUniquePtr<Chaos::FSerializedDataBuffer>&& InSerializedData);
+	UE_INTERNAL static PHYSICSCORE_API void SerializeParticleStateToBuffer(const FPhysicsActorHandle& InActorHandle, Chaos::FSerializedDataBuffer& OutSerializedData, Chaos::ESerializedDataContext DataContext);
+	UE_INTERNAL static PHYSICSCORE_API void ApplySerializedStateToParticle(const FPhysicsActorHandle& InActorHandle, TUniquePtr<Chaos::FSerializedDataBuffer>&& InSerializedData, Chaos::ESerializedDataContext DataContext);
 
-	static PHYSICSCORE_API void SerializeConstraintStateToBuffer(const FPhysicsConstraintHandle& ConstraintHandle, Chaos::FSerializedDataBuffer& OutSerializedData);
-	static PHYSICSCORE_API void ApplySerializedStateToConstraint(const FPhysicsConstraintHandle& ConstraintHandle, Chaos::FSerializedDataBufferPtr&& InSerializedData);
+	UE_INTERNAL static PHYSICSCORE_API void SerializeConstraintStateToBuffer(const FPhysicsConstraintHandle& ConstraintHandle, Chaos::FSerializedDataBuffer& OutSerializedData, Chaos::ESerializedDataContext DataContext);
+	UE_INTERNAL static PHYSICSCORE_API void ApplySerializedStateToConstraint(const FPhysicsConstraintHandle& ConstraintHandle, Chaos::FSerializedDataBufferPtr&& InSerializedData, Chaos::ESerializedDataContext DataContext);
 };
