Add some trace scopes to IAS.

#rnx

diff --git a/Engine/Source/Runtime/Experimental/IoStore/OnDemand/Private/IasHostGroup.cpp b/Engine/Source/Runtime/Experimental/IoStore/OnDemand/Private/IasHostGroup.cpp
--- a/Engine/Source/Runtime/Experimental/IoStore/OnDemand/Private/IasHostGroup.cpp
+++ b/Engine/Source/Runtime/Experimental/IoStore/OnDemand/Private/IasHostGroup.cpp
@@ -358,6 +358,8 @@ void FHostGroupManager::ForEachHostGroup(TFunctionRef<void(const FIASHostGroup&)
 
 void FHostGroupManager::Tick(uint32 TimeoutMs, std::atomic_bool& CancellationToken)
 {
+	TRACE_CPUPROFILER_EVENT_SCOPE(FHostGroupManager::Tick);
+
 	UE::TUniqueLock _(Mutex);
 
 	for (FIASHostGroup& HostGroup : HostGroups)
diff --git a/Engine/Source/Runtime/Experimental/IoStore/OnDemand/Private/OnDemandHttpClient.cpp b/Engine/Source/Runtime/Experimental/IoStore/OnDemand/Private/OnDemandHttpClient.cpp
--- a/Engine/Source/Runtime/Experimental/IoStore/OnDemand/Private/OnDemandHttpClient.cpp
+++ b/Engine/Source/Runtime/Experimental/IoStore/OnDemand/Private/OnDemandHttpClient.cpp
@@ -239,6 +239,8 @@ bool FMultiEndpointHttpClient::Tick(int32 WaitTimeMs, uint32 MaxKiBPerSecond)
 
 void FMultiEndpointHttpClient::UpdateConnections()
 {
+	TRACE_CPUPROFILER_EVENT_SCOPE(FMultiEndpointHttpClient::UpdateConnections);
+
 	UE_CHECK_THREADSAFETY(FMultiEndpointHttpClient::UpdateConnections);
 
 	for (TUniquePtr<FConnection>& Connection : Connections)
