[NetPhysics] Delete ISimCallbackObject instead of having it linger for the length of the rewind history which can produce issues, for example if the ISimCallbackObject is controlling a particle which gets deleted before the ISimCallbackObject.

diff --git a/Engine/Source/Runtime/Engine/Private/PhysicsEngine/NetworkPhysicsComponent.cpp b/Engine/Source/Runtime/Engine/Private/PhysicsEngine/NetworkPhysicsComponent.cpp
--- a/Engine/Source/Runtime/Engine/Private/PhysicsEngine/NetworkPhysicsComponent.cpp
+++ b/Engine/Source/Runtime/Engine/Private/PhysicsEngine/NetworkPhysicsComponent.cpp
@@ -570,6 +570,7 @@ void UNetworkPhysicsComponent::UninitializeComponent()
 			}
 		}
 	}
+	NetworkPhysicsComponent_Internal = nullptr;
 }
 
 void UNetworkPhysicsComponent::GetLifetimeReplicatedProps(TArray< FLifetimeProperty >& OutLifetimeProps) const
diff --git a/Engine/Source/Runtime/Experimental/Chaos/Private/Chaos/ChaosMarshallingManager.cpp b/Engine/Source/Runtime/Experimental/Chaos/Private/Chaos/ChaosMarshallingManager.cpp
--- a/Engine/Source/Runtime/Experimental/Chaos/Private/Chaos/ChaosMarshallingManager.cpp
+++ b/Engine/Source/Runtime/Experimental/Chaos/Private/Chaos/ChaosMarshallingManager.cpp
@@ -147,27 +147,37 @@ void FChaosMarshallingManager::FreePullData_External(FPullPhysicsData* PullData)
 
 void FPushPhysicsData::Reset()
 {
-	for(FSimCallbackInputAndObject& CallbackInputAndObject : SimCallbackInputs)
-	{
-		CallbackInputAndObject.Input->Release_Internal(*CallbackInputAndObject.CallbackObject);
-	}
-
+	ClearAsyncInputs();
+	
 	for(ISimCallbackObject* CallbackToRemove : SimCallbackObjectsToRemove)
 	{
 		ensure(CallbackToRemove->bPendingDelete);	//should already be marked pending delete
 		delete CallbackToRemove;
 	}
 
+	SimCommands.Reset();
 	DirtyProxiesDataBuffer.Reset();
-	SimCallbackInputs.Reset();
-	ResetForHistory();
+	SimCallbackObjectsToAdd.Reset();
+	SimCallbackObjectsToRemove.Reset();
 }
 
 void FPushPhysicsData::ResetForHistory()
 {
+	// Clear dirty proxies data if we should not them cache in history
+	if (bCachePushDataDirtyProxies == false)
+	{
+		ResetDirtyProxiesBuffer();
+		DirtyProxiesDataBuffer.Reset();
+	}
+
+	// Clear async inputs if we should not them cache in history
+	if (bCachePushDataAsyncInputs == false)
+	{
+		ClearAsyncInputs();
+	}
+
 	SimCommands.Reset();
 	SimCallbackObjectsToAdd.Reset();
-	SimCallbackObjectsToRemove.Reset();
 }
 
 void FPushPhysicsData::ResetDirtyProxiesBuffer()
@@ -183,24 +193,53 @@ void FPushPhysicsData::ResetDirtyProxiesBuffer()
 		});
 }
 
+void FPushPhysicsData::ClearAsyncInputs()
+{
+	for (FSimCallbackInputAndObject& CallbackInputAndObject : SimCallbackInputs)
+	{
+		if (CallbackInputAndObject.CallbackObject)
+		{
+			CallbackInputAndObject.Input->Release_Internal(*CallbackInputAndObject.CallbackObject);
+		}
+	}
+	SimCallbackInputs.Reset();
+}
+
+void FPushPhysicsData::ClearAsyncInputForCallback(ISimCallbackObject* Callback)
+{
+	for (TArray<FSimCallbackInputAndObject>::TIterator Itr = SimCallbackInputs.CreateIterator(); Itr; ++Itr)
+	{
+		FSimCallbackInputAndObject& CallbackInputAndObject = *Itr;
+		if (CallbackInputAndObject.CallbackObject == Callback)
+		{
+			CallbackInputAndObject.Input->Release_Internal(*CallbackInputAndObject.CallbackObject);
+			Itr.RemoveCurrent();
+		}
+	}
+}
+
 void FChaosMarshallingManager::FreeDataToHistory_Internal(FPushPhysicsData* PushData)
 {
-	if(HistoryLength == 0)
+	if (HistoryLength == 0)
 	{
 		FreeData_Internal(PushData);
 	}
 	else
 	{
-		if (bCachePushDataDirtyProxies == false)
+		/** Process SimCallbackObjects that are pending delete and clear their cached AsyncInputs from history
+		* At this point the SimCallbackObject is already unregistered from all callbacks and keeping it alive longer (until the end of the rewind history) is only a risk */
+		for (ISimCallbackObject* CallbackToRemove : PushData->SimCallbackObjectsToRemove)
 		{
-			PushData->ResetDirtyProxiesBuffer();
-			PushData->DirtyProxiesDataBuffer.Reset();
-		}
+			ensure(CallbackToRemove->bPendingDelete);	//should already be marked pending delete
 
-		if (bCachePushDataAsyncInputs == false)
-		{
-			PushData->SimCallbackInputs.Reset();
+			// Iterate over cached history and remove any async input that is linked to the SimCallbackObject that is getting deleted
+			for (FPushPhysicsData* PushDataHistoryEntry : HistoryQueue_Internal)
+			{
+				PushDataHistoryEntry->ClearAsyncInputForCallback(CallbackToRemove);
+			}
+			delete CallbackToRemove;
 		}
+		PushData->SimCallbackObjectsToRemove.Reset();
 
 		PushData->ResetForHistory();
 		HistoryQueue_Internal.Insert(PushData, 0);
diff --git a/Engine/Source/Runtime/Experimental/Chaos/Public/Chaos/ChaosMarshallingManager.h b/Engine/Source/Runtime/Experimental/Chaos/Public/Chaos/ChaosMarshallingManager.h
--- a/Engine/Source/Runtime/Experimental/Chaos/Public/Chaos/ChaosMarshallingManager.h
+++ b/Engine/Source/Runtime/Experimental/Chaos/Public/Chaos/ChaosMarshallingManager.h
@@ -260,6 +260,8 @@ struct FPushPhysicsData
 	void Reset();
 	void ResetForHistory();
 	void ResetDirtyProxiesBuffer();
+	void ClearAsyncInputs();
+	void ClearAsyncInputForCallback(ISimCallbackObject* Callback);
 	void CopySubstepData(const FPushPhysicsData& FirstStepData);
 };
 
