Fix that FMultiEndpointHttpClient::Get was ignoring FMultiEndpointHttpClientConfig::Redirects

#rnx

diff --git a/Engine/Source/Runtime/Experimental/IoStore/OnDemand/Private/OnDemandHttpClient.cpp b/Engine/Source/Runtime/Experimental/IoStore/OnDemand/Private/OnDemandHttpClient.cpp
--- a/Engine/Source/Runtime/Experimental/IoStore/OnDemand/Private/OnDemandHttpClient.cpp
+++ b/Engine/Source/Runtime/Experimental/IoStore/OnDemand/Private/OnDemandHttpClient.cpp
@@ -327,7 +327,12 @@ void FMultiEndpointHttpClient::IssueRequest(FRequest&& Request)
 		}
 	};
 
-	FEventLoop::FRequestParams RequestParams = { .bAllowChunked = Config.bAllowChunkedTransfer };
+	FEventLoop::FRequestParams RequestParams = 
+		{
+			.bAutoRedirect = Config.Redirects == EHttpRedirects::Follow,
+			.bAllowChunked = Config.bAllowChunkedTransfer
+		};
+
 	UE::IoStore::HTTP::FRequest HttpRequest = EventLoop.Get(Url, ConnectionPool, &RequestParams);
 
 #if !UE_BUILD_TEST
