Fix a recent unintended change in behavior causing FOnDemandIoBackend to enable http redirects.

#rnx

- During the recent client refactor I forgot that the new client enables redirects by default where as the old client disabled them. However due to the bug fixed in CL 40820484, this went un-noticed because the FMultiEndpointHttpClient value for redirects was being ignored and we were just using the event loop default of off.

diff --git a/Engine/Source/Runtime/Experimental/IoStore/OnDemand/Private/OnDemandIoDispatcherBackend.cpp b/Engine/Source/Runtime/Experimental/IoStore/OnDemand/Private/OnDemandIoDispatcherBackend.cpp
--- a/Engine/Source/Runtime/Experimental/IoStore/OnDemand/Private/OnDemandIoDispatcherBackend.cpp
+++ b/Engine/Source/Runtime/Experimental/IoStore/OnDemand/Private/OnDemandIoDispatcherBackend.cpp
@@ -1664,6 +1664,7 @@ uint32 FOnDemandIoBackend::Run()
 			.ReceiveBufferSize = GIasHttpRecvBufKiB >= 0 ? GIasHttpRecvBufKiB << 10 : -1,
 			.MaxRetryCount = GIasHttpRetryCount,
 			.TimeoutMs = GIasHttpFailTimeOutMs,
+			.Redirects = EHttpRedirects::Disabled,
 			.bEnableThreadSafetyChecks = true,
 			.bAllowChunkedTransfer = GIasHttpAllowChunkedXfer,
 			.LogCategory = &LogIas,
