Iris
- Added warning and detection if object is prevented from being sent
- Added warning if queued batches has been pending for a long time
- QueuedBathes now propagates NetRefHandle error to server

diff --git a/Engine/Source/Runtime/Experimental/Iris/Core/Private/Iris/IrisConfigInternal.h b/Engine/Source/Runtime/Experimental/Iris/Core/Private/Iris/IrisConfigInternal.h
--- a/Engine/Source/Runtime/Experimental/Iris/Core/Private/Iris/IrisConfigInternal.h
+++ b/Engine/Source/Runtime/Experimental/Iris/Core/Private/Iris/IrisConfigInternal.h
@@ -38,3 +38,12 @@
 #	define UE_NET_VALIDATE_DC_BASELINES 1
 #endif 
 #endif
+
+#ifndef UE_NET_ENABLE_REPLICATIONWRITER_CANNOT_SEND_WARNING
+#if (UE_BUILD_SHIPPING)
+#	define UE_NET_ENABLE_REPLICATIONWRITER_CANNOT_SEND_WARNING 1
+#else
+#	define UE_NET_ENABLE_REPLICATIONWRITER_CANNOT_SEND_WARNING 1
+#endif 
+#endif
+
diff --git a/Engine/Source/Runtime/Experimental/Iris/Core/Private/Iris/ReplicationSystem/PendingBatchData.h b/Engine/Source/Runtime/Experimental/Iris/Core/Private/Iris/ReplicationSystem/PendingBatchData.h
--- a/Engine/Source/Runtime/Experimental/Iris/Core/Private/Iris/ReplicationSystem/PendingBatchData.h
+++ b/Engine/Source/Runtime/Experimental/Iris/Core/Private/Iris/ReplicationSystem/PendingBatchData.h
@@ -47,6 +47,12 @@ struct FPendingBatchData
 
 	// Batch owner with queued data chunks
 	FNetRefHandle Handle;
+
+	// When we started to accumulate pending batch data for this object
+	uint64 PendingBatchStartCycles = 0;
+	
+	// Incremented every time we try to process queued batches, reset each time we output warning.
+	int32 PendingBatchTryProcessCount = 0;
 };
 
 struct FPendingBatches
diff --git a/Engine/Source/Runtime/Experimental/Iris/Core/Private/Iris/ReplicationSystem/ReplicationReader.cpp b/Engine/Source/Runtime/Experimental/Iris/Core/Private/Iris/ReplicationSystem/ReplicationReader.cpp
--- a/Engine/Source/Runtime/Experimental/Iris/Core/Private/Iris/ReplicationSystem/ReplicationReader.cpp
+++ b/Engine/Source/Runtime/Experimental/Iris/Core/Private/Iris/ReplicationSystem/ReplicationReader.cpp
@@ -26,6 +26,7 @@
 #include "UObject/Class.h"
 #include "Algo/RemoveIf.h"
 #include "Misc/ScopeExit.h"
+#include "Misc/CoreMisc.h"
 #include "ProfilingDebugging/CsvProfiler.h"
 
 #if UE_NET_ENABLE_REPLICATIONREADER_LOG
@@ -100,6 +101,18 @@ static FAutoConsoleVariableRef CvarResolvedObjectsDispatchDebugging(
 	bResolvedObjectsDispatchDebugging,
 	TEXT("Debug logging of resolved object state dispatching. Default is false."));
 
+static int32 QueuedBatchTimeoutWarningInterval = 30;
+static FAutoConsoleVariableRef CVarQueuedBatchTimeoutWarningInterval(
+	TEXT("net.QueuedBatchTimeoutWarningInterval"),
+	QueuedBatchTimeoutWarningInterval,
+	TEXT("Number of ticks between checks for queuedbatches timeout."));
+
+static float QueuedBatchTimeoutSeconds = 30.0f;
+static FAutoConsoleVariableRef CVarQueuedBatchTimeoutSeconds(
+	TEXT("net.QueuedBatchTimeoutSeconds"),
+	QueuedBatchTimeoutSeconds,
+	TEXT("Time in seconds to wait for queued batches to flush before logging a warning."));
+
 static const FName NetError_FailedToFindAttachmentQueue("Failed to find attachment queue");
 
 //------------------------------------------------------------------------
@@ -735,6 +748,7 @@ FPendingBatchData* FReplicationReader::UpdateUnresolvedMustBeMappedReferences(FN
 		{
 			Batch = &PendingBatches.PendingBatches.AddDefaulted_GetRef();
 			Batch->Handle = InHandle;
+			Batch->PendingBatchStartCycles = FPlatformTime::Cycles64();
 		}
 
 		// Update
@@ -2671,15 +2685,14 @@ void FReplicationReader::ProcessQueuedBatches()
 
 				if (Context.HasErrorOrOverflow())
 				{
-						// Log error  if get more incoming data for an object in the broken state we will skip it.
-					UE_LOG(LogIris, Error, TEXT("FReplicationReader::ProcessQueuedBatches Failed to process object batch handle: %s skipping batch data. bIsEndReplication: %d"), ToCStr(PendingBatchData.Handle.ToString()), int(CurrentChunk.bIsEndReplicationChunk));
-
 					BrokenObjects.AddUnique(PendingBatchData.Handle);
 
+					// Log error  if get more incoming data for an object in the broken state we will skip it.
+					UE_LOG(LogIris, Error, TEXT("FReplicationReader::ProcessQueuedBatches Failed to process object batch handle: %s skipping batch data. bIsEndReplication: %d"), ToCStr(PendingBatchData.Handle.ToString()), int(CurrentChunk.bIsEndReplicationChunk));
+
 					if (Context.GetError() == GNetError_BrokenNetHandle)
 					{
-						// $TODO: Report this to the server so it knows that the state of data in the batch is unknown
-
+						ReplicationBridge->SendErrorWithNetRefHandle(UE::Net::ENetRefHandleError::ReplicationDisabled, PendingBatchData.Handle, Parameters.ConnectionId);
 
 						Context.ResetErrorContext();
 					}
@@ -2709,6 +2722,36 @@ void FReplicationReader::ProcessQueuedBatches()
 		}
 		else
 		{
+			// Warn if we have been waiting for too long
+			if ((QueuedBatchTimeoutWarningInterval > 0) && (PendingBatchData.PendingBatchTryProcessCount > QueuedBatchTimeoutWarningInterval) && FPlatformProperties::RequiresCookedData())
+			{
+				const double QueuedBatchPendingTime = FPlatformTime::ToSeconds64(FPlatformTime::Cycles64() - PendingBatchData.PendingBatchStartCycles);
+				if (QueuedBatchPendingTime > QueuedBatchTimeoutSeconds)
+				{
+					UE_LOG(LogIris, Warning, TEXT("FReplicationWriter::ProcessQueuedBatches: Queued batches for longer than normal. %s, Queued: %i, PendingNetRefHandleResolves: %i, BatchPendingTime: %f")
+						, *NetRefHandleManager->PrintObjectFromNetRefHandle(PendingBatchData.Handle)
+						, PendingBatchData.QueuedDataChunks.Num()
+						, PendingBatchData.PendingMustBeMappedReferences.Num()
+						, QueuedBatchPendingTime);
+
+					if (UE_LOG_ACTIVE(LogIris, Log))
+					{
+						for (const FNetRefHandle& NetRefHandle : PendingBatchData.PendingMustBeMappedReferences)
+						{
+							UE_LOG(LogIris, Log, TEXT("  NetRefHandle Pending resolve [%s]"), *ObjectReferenceCache->DescribeObjectReference(FObjectReferenceCache::MakeNetObjectReference(NetRefHandle), ResolveContext));
+						}
+					}
+
+					// Reset counter
+					PendingBatchData.PendingBatchTryProcessCount = 0;
+				}
+			}
+			else
+			{
+				// Increment counter
+				++PendingBatchData.PendingBatchTryProcessCount;
+			}
+
 			++BatchIt;
 		}
 	}
diff --git a/Engine/Source/Runtime/Experimental/Iris/Core/Private/Iris/ReplicationSystem/ReplicationWriter.cpp b/Engine/Source/Runtime/Experimental/Iris/Core/Private/Iris/ReplicationSystem/ReplicationWriter.cpp
--- a/Engine/Source/Runtime/Experimental/Iris/Core/Private/Iris/ReplicationSystem/ReplicationWriter.cpp
+++ b/Engine/Source/Runtime/Experimental/Iris/Core/Private/Iris/ReplicationSystem/ReplicationWriter.cpp
@@ -75,7 +75,7 @@ static FAutoConsoleVariableRef CVarReplicationWriterMaxDestroyObjectsPerFrame(TE
  * When the number of ReplicationRecords left subceeds this number the ReplicationWriter will limit what is replicated in a packet to a minimum, effectively limiting replication to OOB attachments and huge objects.
  */
 static int32 GReplicationWriterReplicationRecordStarvationThreshold = 1000;
-static FAutoConsoleVariableRef CVarGReplicationWriterReplicationRecordStarvationThreshold(TEXT("net.Iris.ReplicationWriterReplicationRecordStarvationThreshold"), GReplicationWriterReplicationRecordStarvationThreshold,
+static FAutoConsoleVariableRef CVarReplicationWriterReplicationRecordStarvationThreshold(TEXT("net.Iris.ReplicationWriterReplicationRecordStarvationThreshold"), GReplicationWriterReplicationRecordStarvationThreshold,
 	TEXT("How many ReplicationRecords need to be left in order to proceed with replication as normal. Below this threshold replication will be limited to OOB attachments and huge objects until there are more ReplicationRecords available."));
 
 static bool bValidateObjectsWithDirtyChanges = true;
@@ -89,6 +89,10 @@ static FAutoConsoleVariableRef CvarWriteBatchSizePerObject(TEXT("net.Iris.Replic
 static bool bDebugSentinelsEnabled = false;
 static FAutoConsoleVariableRef CvarWriteSentinels(TEXT("net.Iris.ReplicationWriter.WriteSentinels"), bDebugSentinelsEnabled, TEXT("Write sentinels at carefully chosen points in the stream. Helps tracking down bitstream errors. Requires code to be compiled with UE_NET_REPLICATIONDATASTREAM_DEBUG to be enabled."));
 
+// Allow warning if object has been prevented from sending for a long time.
+static int32 GReplicationWriterCannotSendWarningInterval = 256;
+static FAutoConsoleVariableRef CVarReplicationWriterCannotSendWarningInterval(TEXT("net.Iris.ReplicationWriterCannotSendWarningInterval"), GReplicationWriterCannotSendWarningInterval,
+	TEXT("Warn if we are prevented from Sending for more than a certain number of updates, setting it to 0 will disable warning"));
 
 static const FName NetError_ObjectStateTooLarge("Object state is too large to be split.");
 
@@ -597,6 +601,14 @@ void FReplicationWriter::StopReplication(uint32 InternalIndex)
 	}
 
 	Info.PendingBaselineIndex = Info.LastAckedBaselineIndex = FDeltaCompressionBaselineManager::InvalidBaselineIndex;
+
+#if UE_NET_ENABLE_REPLICATIONWRITER_CANNOT_SEND_WARNING
+	if (Info.HasCannotSendInfo)
+	{
+		CannotSendInfos.Remove(InternalIndex);
+		Info.HasCannotSendInfo = 0U;
+	}
+#endif
 }
 
 FReplicationWriter::FReplicationInfo& FReplicationWriter::GetReplicationInfo(uint32 InternalIndex)
@@ -2093,6 +2105,43 @@ uint32 FReplicationWriter::WriteRootObjectsPendingDestroy(FNetSerializationConte
 	return WrittenCount;
 }
 
+FReplicationWriter::FCannotSendInfo* FReplicationWriter::ShouldWarnIfCannotSend(const FReplicationInfo& Info, FInternalNetRefIndex InternalIndex) const
+{
+#if UE_NET_ENABLE_REPLICATIONWRITER_CANNOT_SEND_WARNING
+	if (GReplicationWriterCannotSendWarningInterval > 0)
+	{
+		if (Info.HasCannotSendInfo)
+		{
+			FCannotSendInfo* CannotSendInfo = &CannotSendInfos.FindChecked(InternalIndex);
+			if (--CannotSendInfo->SuppressWarningCounter == 0U)
+			{
+				return CannotSendInfo;
+			}
+		}
+		else
+		{
+			// Store info
+			FCannotSendInfo& CannotSendInfo = CannotSendInfos.Add(InternalIndex);
+			CannotSendInfo.SuppressWarningCounter = GReplicationWriterCannotSendWarningInterval;
+			CannotSendInfo.StartCycles = FPlatformTime::Cycles64();
+			Info.HasCannotSendInfo = 1U;
+		}
+	}
+#endif
+	return static_cast<FCannotSendInfo*>(nullptr);
+};
+
+// local macro to save some typing
+#if UE_NET_ENABLE_REPLICATIONWRITER_CANNOT_SEND_WARNING
+#	define UE_NET_REPLICATIONWRITER_WARN_IF_CANNOT_SEND(Format, ...) \
+		if (FCannotSendInfo* CannotSendInfo = ShouldWarnIfCannotSend(Info, InternalIndex)) \
+		{ \
+			UE_LOG(LogIris, Warning, TEXT("Conn: %u Object %s Blocked from sending for %f seconds: ") Format, Parameters.ConnectionId, *NetRefHandleManager->PrintObjectFromIndex(InternalIndex), FPlatformTime::ToSeconds64(FPlatformTime::Cycles64() - CannotSendInfo->StartCycles), ##__VA_ARGS__); \
+		}
+#else
+#	define UE_NET_REPLICATIONWRITER_WARN_IF_CANNOT_SEND(...)
+#endif
+
 bool FReplicationWriter::CanSendObject(uint32 InternalIndex) const
 {
 	const FReplicationInfo& Info = GetReplicationInfo(InternalIndex);
@@ -2103,6 +2152,7 @@ bool FReplicationWriter::CanSendObject(uint32 InternalIndex) const
 	// We also prevent objects from being transmitted if they are waiting on destroy/tear-off confirmation or cancelling destroy.
 	if (State == EReplicatedObjectState::WaitOnCreateConfirmation || State == EReplicatedObjectState::WaitOnDestroyConfirmation || State == EReplicatedObjectState::CancelPendingDestroy)
 	{
+		UE_NET_REPLICATIONWRITER_WARN_IF_CANNOT_SEND(TEXT("Due to State (%s)"), LexToString(State));
 		return false;
 	}
 
@@ -2111,6 +2161,7 @@ bool FReplicationWriter::CanSendObject(uint32 InternalIndex) const
 	{
 		if (!Info.IsSubObject)
 		{
+			UE_NET_REPLICATIONWRITER_WARN_IF_CANNOT_SEND(TEXT("IsActiveHugeObject"));
 			return false;
 		}
 	}
@@ -2121,6 +2172,7 @@ bool FReplicationWriter::CanSendObject(uint32 InternalIndex) const
 		{
 			if (!CanSendObject(SubObjectInternalIndex))
 			{
+				UE_NET_REPLICATIONWRITER_WARN_IF_CANNOT_SEND(TEXT("SubObject %s cannot be sent"), *NetRefHandleManager->PrintObjectFromIndex(SubObjectInternalIndex));
 				return false;
 			}
 		}
@@ -2147,6 +2199,8 @@ bool FReplicationWriter::CanSendObject(uint32 InternalIndex) const
 				if (!CanSendObject(DependentInternalIndex))
 				{
 					UE_LOG(LogIris, Verbose, TEXT("ReplicationWriter: Cannot send internal index (%u) due to waiting on init dependency internal index (%d)"), InternalIndex, DependentInternalIndex);
+					UE_NET_REPLICATIONWRITER_WARN_IF_CANNOT_SEND(TEXT("Waiting on init dependency %s "), *NetRefHandleManager->PrintObjectFromIndex(DependentInternalIndex));
+
 					return false;
 				}
 
@@ -2154,6 +2208,8 @@ bool FReplicationWriter::CanSendObject(uint32 InternalIndex) const
 				if ((DependentObjectInfo.SchedulingHint == EDependentObjectSchedulingHint::ScheduleBeforeParent) && ObjectsWithDirtyChanges.GetBit(DependentInternalIndex))
 				{
 					UE_LOG(LogIris, Verbose, TEXT("ReplicationWriter: Cannot send internal index (%u) due to waiting on ScheduleBefore dependency internal index (%d)"), InternalIndex, DependentInternalIndex);
+					UE_NET_REPLICATIONWRITER_WARN_IF_CANNOT_SEND(TEXT("Waiting on ScheduleBefore dependency %s "), *NetRefHandleManager->PrintObjectFromIndex(DependentInternalIndex));
+
 					return false;
 				}
 			}
@@ -2163,6 +2219,8 @@ bool FReplicationWriter::CanSendObject(uint32 InternalIndex) const
 	return true;
 }
 
+#undef UE_NET_REPLICATIONWRITER_WARN_IF_CANNOT_SEND
+
 void FReplicationWriter::SerializeObjectStateDelta(FNetSerializationContext& Context, uint32 InternalIndex, const FReplicationInfo& Info, const FNetRefHandleManager::FReplicatedObjectData& ObjectData, const uint8* ReplicatedObjectStateBuffer, FDeltaCompressionBaseline& CurrentBaseline, uint32 CreatedBaselineIndex)
 {
 	FNetBitStreamWriter& Writer = *Context.GetBitStreamWriter();
@@ -3149,6 +3207,14 @@ FReplicationWriter::EWriteStatus FReplicationWriter::WriteDestructionInfo(FNetSe
 
 		ObjectsWithDirtyChanges.ClearBit(InternalIndex);
 
+#if UE_NET_ENABLE_REPLICATIONWRITER_CANNOT_SEND_WARNING
+		if (Info.HasCannotSendInfo)
+		{
+			CannotSendInfos.Remove(InternalIndex);
+			Info.HasCannotSendInfo = 0U;
+		}
+#endif
+
 		// Reset scheduling priority
 		SchedulingPriorities[InternalIndex] = 0.0f;
 
@@ -3438,6 +3504,19 @@ uint32 FReplicationWriter::HandleObjectBatchSuccess(const FBatchInfo& BatchInfo,
 		{
 			SchedulingPriorities[BatchObjectInfo.InternalIndex] = 0.0f;
 		}
+
+#if UE_NET_ENABLE_REPLICATIONWRITER_CANNOT_SEND_WARNING
+		if (Info.HasCannotSendInfo)
+		{
+			const FCannotSendInfo CannotSendInfo = CannotSendInfos.FindAndRemoveChecked(BatchObjectInfo.InternalIndex);
+			Info.HasCannotSendInfo = 0U;
+
+			if (UE_LOG_ACTIVE(LogIris, Verbose))
+			{
+				UE_LOG_REPLICATIONWRITER_CONN(TEXT("Blocked Object %s was sent after waiting for %f s"), *NetRefHandleManager->PrintObjectFromIndex(BatchObjectInfo.InternalIndex), FPlatformTime::ToSeconds64(FPlatformTime::Cycles64() - CannotSendInfo.StartCycles));
+			}
+		}
+#endif
 	}
 
 #if UE_NET_IRIS_CSV_STATS
diff --git a/Engine/Source/Runtime/Experimental/Iris/Core/Private/Iris/ReplicationSystem/ReplicationWriter.h b/Engine/Source/Runtime/Experimental/Iris/Core/Private/Iris/ReplicationSystem/ReplicationWriter.h
--- a/Engine/Source/Runtime/Experimental/Iris/Core/Private/Iris/ReplicationSystem/ReplicationWriter.h
+++ b/Engine/Source/Runtime/Experimental/Iris/Core/Private/Iris/ReplicationSystem/ReplicationWriter.h
@@ -16,6 +16,7 @@
 #include "Iris/ReplicationSystem/ReplicationDataStreamDebug.h"
 #include "Iris/Stats/NetStats.h"
 #include "Containers/Array.h"
+#include "Containers/Map.h"
 #include "Containers/List.h"
 #include "Containers/Set.h"
 #include "Misc/EnumClassFlags.h"
@@ -132,6 +133,7 @@ public:
 				uint64 PendingBaselineIndex : 2;						// Baseline index pending acknowledgment from client
 				uint64 FlushFlags : 3;									// Flags indicating what we are waiting for when flushing
 				uint64 HasDirtyConditionals : 1;						// If this flag is set, we must update conditionals.
+				mutable uint64 HasCannotSendInfo : 1;					// If this flag is set, the object has been prevented from being sent at least once.
 			};
 		};
 
@@ -145,6 +147,12 @@ public:
 
 	static_assert(sizeof(FReplicationInfo) == 16, "Expected sizeof FReplicationInfo to be 16 bytes");
 
+	struct FCannotSendInfo
+	{
+		uint64 StartCycles = 0;
+		uint32 SuppressWarningCounter = 0U;
+	};
+
 public:
 	~FReplicationWriter();
 
@@ -584,6 +592,9 @@ private:
 	// Writes a sentinel naking it easier for the receiver to detect bitstream errors.
 	void WriteSentinel(FNetBitStreamWriter* Writer, const TCHAR* DebugName);
 
+	// Returns a valid FCannotSendInfo if we should start tracking how long we are blocked from sending.
+	FCannotSendInfo* ShouldWarnIfCannotSend(const FReplicationInfo& Info, FInternalNetRefIndex InternalIndex) const;
+
 private:
 	// Replication parameters
 	FReplicationParameters Parameters;
@@ -636,6 +647,8 @@ private:
 	
 	// Should we use high prio create?
 	const bool bHighPrioCreate = false;
+
+	mutable TMap<FInternalNetRefIndex, FCannotSendInfo> CannotSendInfos;
 };
 
 inline FReplicationWriter::FReplicationInfo::FReplicationInfo()
