Merge sorting merge dynamically first sorted data

diff --git a/Engine/Source/Runtime/Experimental/Chaos/Public/Chaos/Collision/SpatialAccelerationBroadPhase.h b/Engine/Source/Runtime/Experimental/Chaos/Public/Chaos/Collision/SpatialAccelerationBroadPhase.h
--- a/Engine/Source/Runtime/Experimental/Chaos/Public/Chaos/Collision/SpatialAccelerationBroadPhase.h
+++ b/Engine/Source/Runtime/Experimental/Chaos/Public/Chaos/Collision/SpatialAccelerationBroadPhase.h
@@ -439,85 +439,135 @@ namespace Chaos
 
 		void GatherConstraints(bool bIsDeterministic)
 		{
+			// This lambda merge constraints of 2 broad phase context. 
+			// If it is deterministic the lambda assume they are already sorted and merge them in the first given context 
+			// maintaining the sorted data.
+			auto MergeLambda = [this, bIsDeterministic](int32 SmallIndexToMerge, int32 BigIndexToMerge)
+			{
+				QUICK_SCOPE_CYCLE_COUNTER(STAT_MergeConstraints);
+				Private::FBroadPhaseContext& Context = BroadphaseContexts[SmallIndexToMerge];
+				Private::FBroadPhaseContext& OtherContext = BroadphaseContexts[BigIndexToMerge];
+
+				TArray<FPBDCollisionConstraint*>& ThisArray = Context.CollisionContext.Allocator->NewActiveConstraints;
+				TArray<FPBDCollisionConstraint*>& OtherArray = OtherContext.CollisionContext.Allocator->NewActiveConstraints;
+
+				if (bIsDeterministic)
+				{
+					TArray<FPBDCollisionConstraint*> ResultArray;
+					const int32 NumThis = ThisArray.Num();
+					const int32 NumOther = OtherArray.Num();
+					ResultArray.Reserve(NumThis + NumOther);
+
+					int32 ThisIndex = 0;
+					int32 OtherIndex = 0;
+
+					while (ThisIndex < NumThis || OtherIndex < NumOther)
+					{
+						if (ThisIndex < NumThis && (OtherIndex >= NumOther || 
+							ThisArray[ThisIndex]->GetContact().GetCollisionSortKey() < OtherArray[OtherIndex]->GetContact().GetCollisionSortKey()))
+						{
+							ResultArray.Add(ThisArray[ThisIndex++]);
+						}
+						else
+						{
+							check(OtherIndex < NumOther);
+							ResultArray.Add(OtherArray[OtherIndex++]);
+						}
+					}
+					ThisArray = MoveTemp(ResultArray);
+				}
+				else
+				{
+					ThisArray.Append(OtherArray);
+				}
+				OtherArray.Reset();
+			};
+			
 			if (bIsDeterministic)
 			{
+				// This atomic is use to merge the sorted data in a safe multithreaded program. 
+				// When reaching the end of a process, 
+				// - either this atomic is INDEX_NONE, meaning there is no other batch to be merged with. So this atomic take the value of the batch
+				// - either this atomic is the index to be merged with, so the merging is called with this current index, and the one stored in the atomic.
+				std::atomic<int32> NextBatchToMerge = INDEX_NONE;
+
 				TArray<UE::Tasks::FTask> ProcessSorting;
 				for (int32 ContextIndex = 0; ContextIndex < PendingTasks.Num(); ContextIndex++)
 				{
-					Private::FBroadPhaseContext& BroadphaseContext = BroadphaseContexts[ContextsWithConstraints[ContextIndex]];
-					UE::Tasks::FTask SortConstraintTask = UE::Tasks::Launch(UE_SOURCE_LOCATION, [this, &BroadphaseContext]() {
+					int32 CurrentIndex = ContextsWithConstraints[ContextIndex];
+					Private::FBroadPhaseContext& BroadphaseContext = BroadphaseContexts[CurrentIndex];
+					UE::Tasks::FTask SortConstraintTask = UE::Tasks::Launch(UE_SOURCE_LOCATION, [this, &BroadphaseContext, CurrentIndex, &NextBatchToMerge, MergeLambda]()
+					{
 						QUICK_SCOPE_CYCLE_COUNTER(STAT_SortConstraintBatch);
 						BroadphaseContext.CollisionContext.Allocator->NewActiveConstraints.Sort(
 							[](const FPBDCollisionConstraintHandle& L, const FPBDCollisionConstraintHandle& R)
 							{
 								return L.GetContact().GetCollisionSortKey() < R.GetContact().GetCollisionSortKey();
 							});
-						}, PendingTasks[ContextIndex]);
-					ProcessSorting.Add(SortConstraintTask);
-				}
-				PendingTasks = MoveTemp(ProcessSorting);
-			}
-			// Create a binary tree of task to merge all sorted constraints
-			int32 NeighborStep = 1;
-			while (PendingTasks.Num() >= 2)
-			{
-				int32 NumProcessSorting = PendingTasks.Num();
-				TArray<UE::Tasks::FTask> ProcessMerging;
-				// We merge data in the context that have a power of 2 index, so for first layer 0, 2, 4..., for the second layer 0, 4, 8... and so on.
-				for (int32 ProcessIndex = 0, ContextIndex = 0; ProcessIndex < NumProcessSorting - 1; ProcessIndex += 2, ContextIndex += 2*NeighborStep)
-				{
-					TArray<UE::Tasks::FTask> Preriq{ PendingTasks[ProcessIndex], PendingTasks[ProcessIndex + 1] };
-					Private::FBroadPhaseContext& Context = BroadphaseContexts[ContextsWithConstraints[ContextIndex]];
-					Private::FBroadPhaseContext& OtherContext = BroadphaseContexts[ContextsWithConstraints[ContextIndex + NeighborStep]];
-					UE::Tasks::FTask MergeConstraintTask = UE::Tasks::Launch(UE_SOURCE_LOCATION, [this, &Context, &OtherContext, bIsDeterministic]() 
-					{
-						QUICK_SCOPE_CYCLE_COUNTER(STAT_MergeConstraints);
-
-						TArray<FPBDCollisionConstraint*>& ThisArray = Context.CollisionContext.Allocator->NewActiveConstraints;
-						TArray<FPBDCollisionConstraint*>& OtherArray = OtherContext.CollisionContext.Allocator->NewActiveConstraints;
-
-						if (bIsDeterministic)
+						int32 NextIndexToMerge = CurrentIndex;
+						while (true) 
 						{
-							TArray<FPBDCollisionConstraint*> ResultArray;
-							const int32 NumThis = ThisArray.Num();
-							const int32 NumOther = OtherArray.Num();
-							ResultArray.Reserve(NumThis + NumOther);
-
-							int32 ThisIndex = 0;
-							int32 OtherIndex = 0;
-
-							while (ThisIndex < NumThis || OtherIndex < NumOther)
+							int32 NextExpected = NextBatchToMerge.load();
+							// If there is a sorted batch ready to be merged, lets merge it
+							if (NextExpected != INDEX_NONE)
 							{
-								if (ThisIndex < NumThis && (OtherIndex >= NumOther || ThisArray[ThisIndex]->GetContact().GetCollisionSortKey() < OtherArray[OtherIndex]->GetContact().GetCollisionSortKey()))
+								// Make sure the value hasn't been overwritten by another thread otherwise, retry
+								if (NextBatchToMerge.compare_exchange_weak(NextExpected, INDEX_NONE, std::memory_order_relaxed))
 								{
-									ResultArray.Add(ThisArray[ThisIndex++]);
+									int32 SmallIndexToMerge = FMath::Min(NextIndexToMerge, NextExpected);
+									int32 BigIndexToMerge = FMath::Max(NextIndexToMerge, NextExpected);
+									check(SmallIndexToMerge != BigIndexToMerge);
+									MergeLambda(SmallIndexToMerge, BigIndexToMerge);
+									NextIndexToMerge = SmallIndexToMerge;
 								}
-								else
+							}
+							else // Otherwise if no batch are available set the current one as the next available for merging
+							{
+								int32 IndexNone = INDEX_NONE;
+								// Make sure there is no available one, otherwise retry
+								if (NextBatchToMerge.compare_exchange_weak(IndexNone, NextIndexToMerge, std::memory_order_relaxed))
 								{
-									check(OtherIndex < NumOther);
-									ResultArray.Add(OtherArray[OtherIndex++]);
+									break;
 								}
 							}
-							ThisArray = MoveTemp(ResultArray);
-						}
-						else
-						{
-							ThisArray.Append(OtherArray);
 						}
-						OtherArray.Reset();
-						}, Preriq);
-					ProcessMerging.Add(MergeConstraintTask);
+					}, PendingTasks[ContextIndex]);
+					ProcessSorting.Add(SortConstraintTask);
 				}
-				// Add the last one for odd process number
-				if (NumProcessSorting % 2 == 1)
+				PendingTasks = MoveTemp(ProcessSorting);
+			}
+			else // For no deterministic program merge all context in new tasks, without sorting them
+			{
+				// Create a binary tree of task to merge all sorted constraints
+				int32 NeighborStep = 1;
+				while (PendingTasks.Num() >= 2)
 				{
-					ProcessMerging.Add(PendingTasks[NumProcessSorting - 1]);
+					int32 NumProcessSorting = PendingTasks.Num();
+					TArray<UE::Tasks::FTask> ProcessMerging;
+					// We merge data in the context that have a power of 2 index, so for first layer 0, 2, 4..., for the second layer 0, 4, 8... and so on.
+					for (int32 ProcessIndex = 0, ContextIndex = 0; ProcessIndex < NumProcessSorting - 1; ProcessIndex += 2, ContextIndex += 2 * NeighborStep)
+					{
+						TArray<UE::Tasks::FTask> Preriq{ PendingTasks[ProcessIndex], PendingTasks[ProcessIndex + 1] };
+						int32 SmallIndexToMerge = ContextsWithConstraints[ContextIndex];
+						int32 BigIndexToMerge = ContextsWithConstraints[ContextIndex + NeighborStep];
+						UE::Tasks::FTask MergeConstraintTask = UE::Tasks::Launch(UE_SOURCE_LOCATION, [SmallIndexToMerge, BigIndexToMerge, MergeLambda]()
+						{
+							QUICK_SCOPE_CYCLE_COUNTER(STAT_MergeConstraints);
+							MergeLambda(SmallIndexToMerge, BigIndexToMerge);
+								
+						}, Preriq);
+						ProcessMerging.Add(MergeConstraintTask);
+					}
+					// Add the last one for odd process number
+					if (NumProcessSorting % 2 == 1)
+					{
+						ProcessMerging.Add(PendingTasks[NumProcessSorting - 1]);
+					}
+					PendingTasks = MoveTemp(ProcessMerging);
+					NeighborStep *= 2;
 				}
-				PendingTasks = MoveTemp(ProcessMerging);
-				NeighborStep *= 2;
+				ensure(PendingTasks.Num() <= 1);
 			}
-			ensure(PendingTasks.Num() <= 1);
-
 			UE::Tasks::Wait(PendingTasks);
 
 			// Lets move all constraints in the first allocator to be picked up by the CollisionConstraintAllocator
