Fix incorrect index used in THashMappedArray remove function

diff --git a/Engine/Plugins/Experimental/SpatialReadiness/Tests/Private/PhysicsSpatialReadinessTests.cpp b/Engine/Plugins/Experimental/SpatialReadiness/Tests/Private/PhysicsSpatialReadinessTests.cpp
--- a/Engine/Plugins/Experimental/SpatialReadiness/Tests/Private/PhysicsSpatialReadinessTests.cpp
+++ b/Engine/Plugins/Experimental/SpatialReadiness/Tests/Private/PhysicsSpatialReadinessTests.cpp
@@ -231,7 +231,7 @@ TEST_CASE("Sim Callback", "[physics spatial readiness]")
 
 	SECTION("Add and remove many random volumes")
 	{
-		const int32 NumVolumeActions = 100;
+		const int32 NumVolumeActions = 200;
 		FRandomStream RandStream(42);
 
 		TArray<FSpatialReadinessVolume> Volumes;
diff --git a/Engine/Source/Runtime/Experimental/Chaos/Public/Chaos/Framework/HashMappedArray.h b/Engine/Source/Runtime/Experimental/Chaos/Public/Chaos/Framework/HashMappedArray.h
--- a/Engine/Source/Runtime/Experimental/Chaos/Public/Chaos/Framework/HashMappedArray.h
+++ b/Engine/Source/Runtime/Experimental/Chaos/Public/Chaos/Framework/HashMappedArray.h
@@ -146,7 +146,7 @@ namespace Chaos::Private
 						const FHashType LastElemKey = FTraits::GetIDHash(LastElemID);
 
 						// Remove the last element's key/index from the hash table
-						HashTable.Remove(LastElemKey, LastElemID);
+						HashTable.Remove(LastElemKey, LastElemIndex);
 
 						// Remove the element we want to remove, swapping it with the last element
 						Elements.RemoveAtSwap(Index);
