Extract IoStoreOnDemand core module

- I/O store on-demand functionality is now disabled by default unless the engine is compiled with -CompileIoStoreOnDemand
- moved the public header and relevant public types to a new module IoStoreOnDemandCore
- added a new module feature that can create a concrete implementation of the on-demand I/O store interface
- calling TryGetOnDemandIoStore will return a nullptr untill the IoStoreOnDemand module has been loaded

#rnx

diff --git a/Engine/Plugins/Experimental/DefaultInstallBundleManager/Source/DefaultInstallBundleManager.Build.cs b/Engine/Plugins/Experimental/DefaultInstallBundleManager/Source/DefaultInstallBundleManager.Build.cs
--- a/Engine/Plugins/Experimental/DefaultInstallBundleManager/Source/DefaultInstallBundleManager.Build.cs
+++ b/Engine/Plugins/Experimental/DefaultInstallBundleManager/Source/DefaultInstallBundleManager.Build.cs
@@ -23,7 +23,7 @@ public class DefaultInstallBundleManager : ModuleRules
 				"Json",
 				"PakFile",
 				"RenderCore",
-				"IoStoreOnDemand",
+				"IoStoreOnDemandCore",
 			}
 		);
 	}
diff --git a/Engine/Plugins/Experimental/RuntimeTelemetry/Source/Private/RuntimeTelemetry.cpp b/Engine/Plugins/Experimental/RuntimeTelemetry/Source/Private/RuntimeTelemetry.cpp
--- a/Engine/Plugins/Experimental/RuntimeTelemetry/Source/Private/RuntimeTelemetry.cpp
+++ b/Engine/Plugins/Experimental/RuntimeTelemetry/Source/Private/RuntimeTelemetry.cpp
@@ -5,10 +5,7 @@
 #include "Misc/CoreDelegates.h"
 #include "Containers/Map.h"
 #include "HAL/LowLevelMemTracker.h"
-
-#ifdef USE_IOSTOREONDEMAND
 #include "IO/IoStoreOnDemand.h"
-#endif
 
 DEFINE_LOG_CATEGORY(LogRuntimeTelemetry);
 
@@ -20,9 +17,6 @@ FRuntimeTelemetry& FRuntimeTelemetry::Get()
 
 void FRuntimeTelemetry::RecordEvent_IoStoreOnDemand(const FString& Context, TArray<FAnalyticsEventAttribute> Attributes)
 {
-#ifdef USE_IOSTOREONDEMAND	
-	using namespace UE::IoStore;
-
 	if (FStudioTelemetry::Get().IsSessionRunning())
 	{
 		if (UE::IoStore::IOnDemandIoStore* IoStore = UE::IoStore::TryGetOnDemandIoStore())
@@ -44,7 +38,6 @@ void FRuntimeTelemetry::RecordEvent_IoStoreOnDemand(const FString& Context, TArr
 			}
 		}
 	}
-#endif
 }
 
 
diff --git a/Engine/Plugins/Experimental/RuntimeTelemetry/Source/RuntimeTelemetry.Build.cs b/Engine/Plugins/Experimental/RuntimeTelemetry/Source/RuntimeTelemetry.Build.cs
--- a/Engine/Plugins/Experimental/RuntimeTelemetry/Source/RuntimeTelemetry.Build.cs
+++ b/Engine/Plugins/Experimental/RuntimeTelemetry/Source/RuntimeTelemetry.Build.cs
@@ -9,7 +9,8 @@ namespace UnrealBuildTool.Rules
 			PublicDependencyModuleNames.AddRange(
 				new string[] {
 					"Core",
-					"Engine"
+					"Engine",
+					"IoStoreOnDemandCore"
 				}
 			);
 
@@ -20,12 +21,6 @@ namespace UnrealBuildTool.Rules
 					"StudioTelemetry"
 				}
 			);
-
-			if (Target.bCompileIoStoreOnDemand && !Target.bBuildEditor)
-			{
-				PrivateDependencyModuleNames.Add("IoStoreOnDemand");
-				PublicDefinitions.Add("USE_IOSTOREONDEMAND");
-			}
 		}
 	}
 }
diff --git a/Engine/Plugins/Runtime/GameFeatures/Source/GameFeatures/GameFeatures.Build.cs b/Engine/Plugins/Runtime/GameFeatures/Source/GameFeatures/GameFeatures.Build.cs
--- a/Engine/Plugins/Runtime/GameFeatures/Source/GameFeatures/GameFeatures.Build.cs
+++ b/Engine/Plugins/Runtime/GameFeatures/Source/GameFeatures/GameFeatures.Build.cs
@@ -24,7 +24,7 @@ namespace UnrealBuildTool.Rules
 					"AssetRegistry",
 					"GameplayTags",
 					"InstallBundleManager",
-					"IoStoreOnDemand",
+					"IoStoreOnDemandCore",
 					"Json",
 					"JsonUtilities",
 					"PakFile",
@@ -37,7 +37,7 @@ namespace UnrealBuildTool.Rules
 			PublicIncludePathModuleNames.AddRange(
 				new string[]
 				{
-					"IoStoreOnDemand",
+					"IoStoreOnDemandCore",
 				}
 			);
 
diff --git a/Engine/Source/Runtime/Engine/Engine.Build.cs b/Engine/Source/Runtime/Engine/Engine.Build.cs
--- a/Engine/Source/Runtime/Engine/Engine.Build.cs
+++ b/Engine/Source/Runtime/Engine/Engine.Build.cs
@@ -115,6 +115,7 @@ public class Engine : ModuleRules
 				"DeveloperSettings",
 				"AudioLinkCore",
 				"CookOnTheFly",
+				"IoStoreOnDemandCore"
 			}
 		);
 
@@ -126,6 +127,7 @@ public class Engine : ModuleRules
 		if (Target.bCompileIoStoreOnDemand && !Target.bBuildEditor)
 		{
 			PrivateDependencyModuleNames.Add("IoStoreOnDemand");
+			PublicDefinitions.Add("UE_WITH_IOSTOREONDEMAND=1");
 		}
 
 		PublicIncludePathModuleNames.AddRange(
diff --git a/Engine/Source/Runtime/Experimental/IoStore/OnDemand/IoStoreOnDemand.Build.cs b/Engine/Source/Runtime/Experimental/IoStore/OnDemand/IoStoreOnDemand.Build.cs
--- a/Engine/Source/Runtime/Experimental/IoStore/OnDemand/IoStoreOnDemand.Build.cs
+++ b/Engine/Source/Runtime/Experimental/IoStore/OnDemand/IoStoreOnDemand.Build.cs
@@ -7,6 +7,7 @@ public class IoStoreOnDemand : ModuleRules
 	public IoStoreOnDemand(ReadOnlyTargetRules Target) : base(Target)
 	{
 		PublicDependencyModuleNames.Add("Core");
+		PublicDependencyModuleNames.Add("IoStoreOnDemandCore");
 		PublicDependencyModuleNames.Add("TraceLog");
 		PrivateDependencyModuleNames.AddRange(
 			new string[] {
diff --git a/Engine/Source/Runtime/Experimental/IoStore/OnDemand/Private/IoStoreOnDemand.cpp b/Engine/Source/Runtime/Experimental/IoStore/OnDemand/Private/IoStoreOnDemand.cpp
deleted file mode 100644
diff --git a/Engine/Source/Runtime/Experimental/IoStore/OnDemand/Private/OnDemandHostGroup.cpp b/Engine/Source/Runtime/Experimental/IoStore/OnDemand/Private/OnDemandHostGroup.cpp
deleted file mode 100644
diff --git a/Engine/Source/Runtime/Experimental/IoStore/OnDemand/Private/OnDemandIoStore.cpp b/Engine/Source/Runtime/Experimental/IoStore/OnDemand/Private/OnDemandIoStore.cpp
--- a/Engine/Source/Runtime/Experimental/IoStore/OnDemand/Private/OnDemandIoStore.cpp
+++ b/Engine/Source/Runtime/Experimental/IoStore/OnDemand/Private/OnDemandIoStore.cpp
@@ -518,7 +518,8 @@ FIoStatus FOnDemandIoStore::Initialize()
 	}
 #endif
 
-	const FIoStatus InstallCacheStatus = InitializeInstallCache();
+
+	FIoStatus InstallCacheStatus = InitializeInstallCache();
 	if (InstallCacheStatus.GetErrorCode() == EIoErrorCode::PendingFork)
 	{
 		UE_LOG(LogIoStoreOnDemand, Log, TEXT("Deferring initialization of install cache until post server fork"))
@@ -540,18 +541,149 @@ FIoStatus FOnDemandIoStore::Initialize()
 					}
 				}
 			});
-
-		return FIoStatus::Ok;
+		InstallCacheStatus = FIoStatus::Ok;
 	}
 	else if (InstallCacheStatus.GetErrorCode() == EIoErrorCode::Disabled)
 	{
 		UE_LOG(LogIoStoreOnDemand, Log, TEXT("Install cache disabled"));
-		return FIoStatus::Ok;
 	}
 
 	return InstallCacheStatus;
 }
 
+FIoStatus FOnDemandIoStore::InitializePostHotfix()
+{
+	const TCHAR* CommandLine = FCommandLine::Get();
+
+	TIoStatusOr<FOnDemandEndpointConfig> EndpointConfigStatus = Config::TryParseEndpointConfig(CommandLine);
+	if (EndpointConfigStatus.IsOk() == false)
+	{
+		// Temporary solution: Delete the ias cache directory if the game hasn't been configured to use IAS
+		if (FPaths::HasProjectPersistentDownloadDir())
+		{
+			TStringBuilder<256> CachePath;
+			CachePath << FPaths::ProjectPersistentDownloadDir();
+			FPathViews::Append(CachePath, TEXTVIEW("ias"));
+
+			IFileManager& Ifm = IFileManager::Get();
+			if (Ifm.DirectoryExists(CachePath.ToString()))
+			{
+				Ifm.DeleteDirectory(CachePath.ToString(), false, true);
+			}
+		}
+
+		return EndpointConfigStatus.Status();
+	}
+
+	FOnDemandEndpointConfig EndpointConfig = EndpointConfigStatus.ConsumeValueOrDie();
+
+	{
+		bool bFirstAttempt = false;
+		if (FIoStatus CertStatus = LoadDefaultHttpCertificates(bFirstAttempt); !CertStatus.IsOk())
+		{
+			if (EndpointConfig.RequiresTls())
+			{
+				FIoStatus Status = FIoStatusBuilder(EIoErrorCode::ReadError)
+					<< TEXT("Endpoint requires TLS but certificates failed to load due to '")
+					<< CertStatus.ToString()
+					<< TEXT("'");
+				return Status;
+			}
+		}
+	}
+
+	if (FIoStatus Status = InitializeStreamingBackend(EndpointConfig); Status.IsOk() == false)
+	{
+		return Status;
+	}
+
+	bool bUsePerContainerTocsConfigValue = false;
+	if (GConfig)
+	{
+		GConfig->GetBool(TEXT("Ias"), TEXT("UsePerContainerTocs"), bUsePerContainerTocsConfigValue, GEngineIni);
+	}
+	bool bUsePerContainerTocsParam = false;
+#if !UE_BUILD_SHIPPING
+	bUsePerContainerTocsParam = FParse::Param(CommandLine, TEXT("Ias.UsePerContainerTocs"));
+#endif
+
+	const bool bUsePerContainerTocs = bUsePerContainerTocsConfigValue || bUsePerContainerTocsParam;
+#if !UE_BUILD_SHIPPING
+	UE_LOG(LogIas, Log, TEXT("Using per container TOCs=%s"), bUsePerContainerTocs ? TEXT("True") : TEXT("False"));
+#endif
+
+	TOptional<FOnDemandMountArgs> MountArgs;
+	if (EndpointConfig.TocFilePath.IsEmpty() == false)
+	{
+		if (bUsePerContainerTocs == false)
+		{
+			MountArgs.Emplace(FOnDemandMountArgs
+			{
+				.MountId = EndpointConfig.TocFilePath,
+				.FilePath = EndpointConfig.TocFilePath,
+				.Options = EOnDemandMountOptions::StreamOnDemand
+			});
+		}
+	}
+	else if (!EndpointConfig.ServiceUrls.IsEmpty() && !EndpointConfig.TocPath.IsEmpty())
+	{
+		TIoStatusOr<FOnDemandHostGroup> HostGroup = FOnDemandHostGroup::Create(EndpointConfig.ServiceUrls);
+		if (HostGroup.IsOk() == false)
+		{
+			return HostGroup.Status();
+		}
+		MountArgs.Emplace(FOnDemandMountArgs
+		{
+			.MountId = EndpointConfig.ServiceUrls[0] / EndpointConfig.TocPath,
+			.TocRelativeUrl = EndpointConfig.TocPath,
+			.HostGroup = HostGroup.ConsumeValueOrDie(),
+			.Options = EOnDemandMountOptions::StreamOnDemand
+		});
+	}
+
+#if !UE_BUILD_SHIPPING
+	TOptional<FOnDemandInstallArgs> InstallArgs;
+	if (FParse::Param(FCommandLine::Get(), TEXT("Iad")))
+	{
+		TIoStatusOr<FOnDemandHostGroup> HostGroup = FOnDemandHostGroup::Create(EndpointConfig.ServiceUrls);
+		check(HostGroup.IsOk());
+
+		if (MountArgs)
+		{
+			MountArgs.GetValue().HostGroup = HostGroup.ConsumeValueOrDie(); 
+			MountArgs.GetValue().Options = EOnDemandMountOptions::InstallOnDemand;
+			MountArgs.GetValue().TocRelativeUrl = EndpointConfig.TocPath;
+
+			static FOnDemandContentHandle ContentHandle = FOnDemandContentHandle::Create(TEXT("AllContent"));
+			InstallArgs.Emplace();
+			InstallArgs->MountId = MountArgs.GetValue().MountId;
+			InstallArgs->ContentHandle = ContentHandle;
+		}
+	}
+#endif
+
+	if (MountArgs)
+	{
+		Mount(MoveTemp(MountArgs.GetValue()), [](FOnDemandMountResult MountResult)
+		{
+			UE_CLOG(!MountResult.Status.IsOk(), LogIas, Error,
+				TEXT("Failed to mount TOC for '%s', reason '%s'"), *MountResult.MountId, *MountResult.Status.ToString());
+		});
+#if !UE_BUILD_SHIPPING
+		if (InstallArgs)
+		{
+			Install(MoveTemp(InstallArgs.GetValue()), [](FOnDemandInstallResult InstallResult)
+			{
+				UE_CLOG(!InstallResult.Status.IsOk(), LogIoStoreOnDemand, Error,
+					TEXT("Failed to install content, reason '%s'"), *InstallResult.Status.ToString());
+			});
+		}
+#endif
+	}
+
+	return FIoStatus::Ok;
+}
+
 FIoStatus FOnDemandIoStore::InitializeStreamingBackend(const FOnDemandEndpointConfig& EndpointConfig)
 {
 	TUniquePtr<IIasCache> StreamingCache;
@@ -1133,16 +1265,29 @@ FIoStatus FOnDemandIoStore::TickMountRequest(FMountRequest& MountRequest)
 
 		UE_LOG(LogIoStoreOnDemand, Log, TEXT("Loading TOC from URL '%s'"), Url.ToString());
 
-		const uint32 RetryCount = 2;
-		const bool bFollowRedirects = true;
-		TIoStatusOr<FOnDemandToc> TocStatus = FOnDemandToc::LoadFromUrl(Url, 2, bFollowRedirects);
+		auto AnsiUrl = StringCast<ANSICHAR>(Url.GetData(), Url.Len());
+		TIoStatusOr<FMultiEndpointHttpClientResponse> Response = FMultiEndpointHttpClient::Get(AnsiUrl, FMultiEndpointHttpClientConfig
+		{ 
+			.MaxRetryCount	= 2,
+			.Redirects		= EHttpRedirects::Follow
+		});
 
-		if (!TocStatus.IsOk())
+		if (!Response.IsOk())
 		{
-			return TocStatus.Status();
+			FIoStatus Status = FIoStatusBuilder(EIoErrorCode::ReadError) << TEXT("Failed to fetch TOC from URL");
+			return Status;
 		}
 
-		Args.Toc = MakeUnique<FOnDemandToc>(TocStatus.ConsumeValueOrDie());
+		FMemoryReaderView Ar(Response.ValueOrDie().Body.GetView());
+		Args.Toc = MakeUnique<FOnDemandToc>();
+		Ar << *Args.Toc;
+
+		if (Ar.IsError() || Ar.IsCriticalError())
+		{
+			Args.Toc.Reset();
+			FIoStatus Status = FIoStatusBuilder(EIoErrorCode::ReadError) << TEXT("Failed to serialize TOC from HTTP response");
+			return Status;
+		}
 
 		CreateContainersFromToc(Args.MountId, TocPath, *Args.Toc, Args.HostGroup, ExistingContainers, RequestedContainers);
 	}
@@ -2069,119 +2214,4 @@ TUniquePtr<IAnalyticsRecording> FOnDemandIoStore::StartAnalyticsRecording() cons
 	return TUniquePtr<IAnalyticsRecording>();
 }
 
-///////////////////////////////////////////////////////////////////////////////
-FOnDemandInternalContentHandle::~FOnDemandInternalContentHandle()
-{
-	if (TSharedPtr<FOnDemandIoStore, ESPMode::ThreadSafe> Pinned = IoStore.Pin(); Pinned.IsValid())
-	{
-		Pinned->ReleaseContent(*this);
-	}
-}
-
-FString LexToString(const FOnDemandInternalContentHandle& Handle)
-{
-	return FString::Printf(TEXT("Id=0x%llX DebugName='%s'"), Handle.HandleId(), *Handle.DebugName);
-}
-
-///////////////////////////////////////////////////////////////////////////////
-FOnDemandContentHandle::FOnDemandContentHandle()
-{
-}
-
-FOnDemandContentHandle::~FOnDemandContentHandle()
-{
-}
-
-FOnDemandContentHandle FOnDemandContentHandle::Create()
-{
-	FOnDemandContentHandle NewHandle;
-	NewHandle.Handle = MakeShared<FOnDemandInternalContentHandle, ESPMode::ThreadSafe>();
-	return NewHandle;
-}
-
-FOnDemandContentHandle FOnDemandContentHandle::Create(FSharedString DebugName)
-{
-	FOnDemandContentHandle NewHandle;
-	NewHandle.Handle = MakeShared<FOnDemandInternalContentHandle, ESPMode::ThreadSafe>(DebugName);
-	return NewHandle;
-}
-
-FOnDemandContentHandle FOnDemandContentHandle::Create(FStringView DebugName)
-{
-	return Create(FSharedString(DebugName));
-}
-
-FString LexToString(const FOnDemandContentHandle& Handle)
-{
-	return Handle.IsValid() ? LexToString(*Handle.Handle) : TEXT("Invalid");
-}
-
-////////////////////////////////////////////////////////////////////////////////
-FString LexToString(const FOnDemandInternalInstallRequest& Request)
-{
-	return FString::Printf(TEXT("Id=0x%llX, DebugName='%s', ContentHandle='%s'"),
-		UPTRINT(&Request),
-		Request.Args.DebugName.IsEmpty()
-			? TEXT("NoName")
-			: *Request.Args.DebugName,
-		*LexToString(Request.Args.ContentHandle));
-}
-
-uint32 FOnDemandInternalInstallRequest::NextSeqNo = 0;
-
-///////////////////////////////////////////////////////////////////////////////
-FOnDemandInstallRequest::FOnDemandInstallRequest()
-{
-}
-
-FOnDemandInstallRequest::FOnDemandInstallRequest(FWeakOnDemandIoStore InIoStore, FSharedInternalInstallRequest InternalRequest)
-	: IoStore(InIoStore)
-	, Request(InternalRequest)
-{
-}
-
-FOnDemandInstallRequest::~FOnDemandInstallRequest()
-{
-}
-
-FOnDemandInstallRequest::FOnDemandInstallRequest(FOnDemandInstallRequest&& Other)
-	: IoStore(MoveTemp(Other.IoStore))
-	, Request(MoveTemp(Other.Request))
-{
-}
-
-FOnDemandInstallRequest& FOnDemandInstallRequest::operator=(FOnDemandInstallRequest&& Other)
-{
-	Request = MoveTemp(Other.Request);
-	IoStore = MoveTemp(Other.IoStore);
-	return *this;
-}
-
-FOnDemandInstallRequest::EStatus FOnDemandInstallRequest::GetStatus() const
-{
-	return Request.IsValid() ? Request->Status.load(std::memory_order_relaxed) : EStatus::None; 
-}
-
-void FOnDemandInstallRequest::Cancel()
-{
-	if (IsPending())
-	{
-		if (TSharedPtr<FOnDemandIoStore, ESPMode::ThreadSafe> Pinned = IoStore.Pin(); Pinned.IsValid())
-		{
-			Pinned->CancelInstallRequest(Request);
-		}
-	}
-}
-
-void FOnDemandInstallRequest::UpdatePriority(int32 NewPriority)
-{
-	if (IsPending())
-	{
-		if (TSharedPtr<FOnDemandIoStore, ESPMode::ThreadSafe> Pinned = IoStore.Pin(); Pinned.IsValid())
-		{
-			Pinned->UpdateInstallRequestPriority(Request, NewPriority);
-		}
-	}
-}
-
 } // namespace UE::IoStore
diff --git a/Engine/Source/Runtime/Experimental/IoStore/OnDemand/Private/OnDemandIoStore.h b/Engine/Source/Runtime/Experimental/IoStore/OnDemand/Private/OnDemandIoStore.h
--- a/Engine/Source/Runtime/Experimental/IoStore/OnDemand/Private/OnDemandIoStore.h
+++ b/Engine/Source/Runtime/Experimental/IoStore/OnDemand/Private/OnDemandIoStore.h
@@ -11,7 +11,7 @@
 #include "IO/IoChunkEncoding.h"
 #include "IO/IoChunkId.h"
 #include "IO/IoHash.h"
-#include "IO/IoStoreOnDemand.h"
+#include "IO/IoStoreOnDemandInternals.h"
 #include "IasHostGroup.h"
 #include "Misc/AES.h"
 #include "Misc/EnumClassFlags.h"
@@ -33,65 +33,6 @@ using FSharedStreamingBackend		= TSharedPtr<class IOnDemandIoDispatcherBackend>;
 using FSharedPackageStoreBackend	= TSharedPtr<class IOnDemandPackageStoreBackend>;
 using FSharedInstallCache			= TSharedPtr<class IOnDemandInstallCache>;
 
-///////////////////////////////////////////////////////////////////////////////
-class FOnDemandInternalContentHandle
-{
-public:
-	FOnDemandInternalContentHandle()
-		: DebugName(TEXT("NoName"))
-	{ }
-	FOnDemandInternalContentHandle(FSharedString InDebugName)
-		: DebugName(InDebugName)
-	{ }
-	~FOnDemandInternalContentHandle();
-
-	UPTRINT HandleId() const { return UPTRINT(this); }
-
-	FSharedString			DebugName;
-	FWeakOnDemandIoStore	IoStore;
-};
-
-FString LexToString(const FOnDemandInternalContentHandle& Handle);
-
-////////////////////////////////////////////////////////////////////////////////
-class FOnDemandInternalInstallRequest
-{
-public:
-	using EStatus = FOnDemandRequest::EStatus;
-
-	static uint32 NextSeqNo;
-
-	static bool SortByPriority(FSharedInternalInstallRequest LHS, FSharedInternalInstallRequest RHS)
-	{
-		if (LHS->Args.Priority == RHS->Args.Priority)
-		{
-			return LHS->SeqNo > RHS->SeqNo;
-		}
-		return static_cast<uint8>(LHS->Args.Priority) < static_cast<uint8>(RHS->Args.Priority);
-	};
-
-	FOnDemandInternalInstallRequest(
-		FOnDemandInstallArgs&& InArgs,
-		FOnDemandInstallCompleted&& InOnCompleted,
-		FOnDemandInstallProgressed&& InOnProgress)
-			: Args(MoveTemp(InArgs))
-			, OnCompleted(MoveTemp(InOnCompleted))
-			, OnProgress(MoveTemp(InOnProgress))
-	{ }
-
-	bool IsCancelRequested() const { return EIoErrorCode::Cancelled == ErrorCode.load(std::memory_order_relaxed); }
-
-	FOnDemandInstallArgs		Args;
-	FOnDemandInstallCompleted	OnCompleted;
-	FOnDemandInstallProgressed	OnProgress;
-	double						StartTime = FPlatformTime::Seconds();
-	uint32						SeqNo = NextSeqNo++;
-	std::atomic<EStatus>		Status{EStatus::Pending};
-	std::atomic<EIoErrorCode>	ErrorCode{EIoErrorCode::Unknown};
-};
-
-FString LexToString(const FOnDemandInternalInstallRequest& Request);
-
 ///////////////////////////////////////////////////////////////////////////////
 enum class EOnDemandContainerFlags : uint8
 {
@@ -325,7 +266,6 @@ public:
 										FOnDemandIoStore();
 										~FOnDemandIoStore();
 
-	FIoStatus							Initialize();
 	FIoStatus							InitializeStreamingBackend(const FOnDemandEndpointConfig& EndpointConfig);
 
 	FOnDemandChunkInfo					GetStreamingChunkInfo(const FIoChunkId& ChunkId);
@@ -335,16 +275,13 @@ public:
 	/** Used to access a number of FIoChunkId that exist in the system when running tests and is not intended for production use */
 	TArray<FIoChunkId>					DebugFindStreamingChunkIds(int32 NumToFind);
 #endif //!UE_BUILD_SHIPPING
-
-	void								ReleaseContent(FOnDemandInternalContentHandle& ContentHandle);
 	void								GetReferencedContent(TArray<FSharedOnDemandContainer>& OutContainers, TArray<TBitArray<>>& OutChunkEntryIndices);
-	void								CancelInstallRequest(FSharedInternalInstallRequest InstallRequest);
-	void								UpdateInstallRequestPriority(FSharedInternalInstallRequest InstallRequest, int32 NewPriority);
 
 	// IOnDemandIoStore interface
+	virtual FIoStatus					Initialize() override;
+	virtual FIoStatus					InitializePostHotfix() override;
 	virtual void						Mount(FOnDemandMountArgs&& Args, FOnDemandMountCompleted&& OnCompleted) override;
-	virtual FOnDemandInstallRequest 	Install(
-											FOnDemandInstallArgs&& Args,
+	virtual FOnDemandInstallRequest 	Install(FOnDemandInstallArgs&& Args,
 											FOnDemandInstallCompleted&& OnCompleted,
 											FOnDemandInstallProgressed&& OnProgress = nullptr) override;
 	virtual void						Purge(FOnDemandPurgeArgs&& Args, FOnDemandPurgeCompleted&& OnCompleted) override;
@@ -356,33 +293,36 @@ public:
 	virtual void						SetStreamingOptions(EOnDemandStreamingOptions Options) override;
 	virtual void						ReportAnalytics(TArray<FAnalyticsEventAttribute>& OutAnalyticsArray) const override;
 	virtual TUniquePtr<IAnalyticsRecording> StartAnalyticsRecording() const override;
+	virtual void						CancelInstallRequest(FSharedInternalInstallRequest InstallRequest) override;
+	virtual void						UpdateInstallRequestPriority(FSharedInternalInstallRequest InstallRequest, int32 NewPriority) override;
+	virtual void						ReleaseContent(FOnDemandInternalContentHandle& ContentHandle) override;
 
 private:
-	FIoStatus				GetContainersAndPackagesForInstall(
-								FStringView MountId,
-								const TArray<FString>& TagSets,
-								const TArray<FPackageId>& PackageIds,
-								TSet<FSharedOnDemandContainer>& OutContainersForInstallation,
-								TSet<FPackageId>& OutPackageIdsToInstall) const;
-	FIoStatus				InitializeInstallCache();
-	FOnDemandChunkInfo		GetChunkInfo(const FIoChunkId& ChunkId, EOnDemandContainerFlags ContainerFlags);
-	void					TryEnterTickLoop();
-	void					TickLoop();
-	bool					Tick();
-	FIoStatus				TickMountRequest(FMountRequest& MountRequest);
-	void					CompleteMountRequest(FMountRequest& MountRequest, FOnDemandMountResult&& MountResult);
-	FOnDemandInstallResult	TickInstallRequest(FOnDemandInternalInstallRequest& InstallRequest);
-	void					CompleteInstallRequest(FSharedInternalInstallRequest InstallRequest, FOnDemandInstallResult&& InstallResult);
-	void					ProgressInstallRequest(const FOnDemandInternalInstallRequest& InstallRequest, const FOnDemandInstallProgress& Progress);
-	void					CompletePurgeRequest(FPurgeRequest& PurgeRequest, FOnDemandPurgeResult&& Result);
-	void					OnEncryptionKeyAdded(const FGuid& Id, const FAES::FAESKey& Key);
-	static void				CreateContainersFromToc(
-								FStringView MountId,
-								FStringView TocPath,
-								FOnDemandToc& Toc,
-								FOnDemandHostGroup HostGroup,
-								const TArray<FString>& ExistingContainers,
-								TArray<FSharedOnDemandContainer>& Out);
+	FIoStatus							GetContainersAndPackagesForInstall(
+											FStringView MountId,
+											const TArray<FString>& TagSets,
+											const TArray<FPackageId>& PackageIds,
+											TSet<FSharedOnDemandContainer>& OutContainersForInstallation,
+											TSet<FPackageId>& OutPackageIdsToInstall) const;
+	FIoStatus							InitializeInstallCache();
+	FOnDemandChunkInfo					GetChunkInfo(const FIoChunkId& ChunkId, EOnDemandContainerFlags ContainerFlags);
+	void								TryEnterTickLoop();
+	void								TickLoop();
+	bool								Tick();
+	FIoStatus							TickMountRequest(FMountRequest& MountRequest);
+	void								CompleteMountRequest(FMountRequest& MountRequest, FOnDemandMountResult&& MountResult);
+	FOnDemandInstallResult				TickInstallRequest(FOnDemandInternalInstallRequest& InstallRequest);
+	void								CompleteInstallRequest(FSharedInternalInstallRequest InstallRequest, FOnDemandInstallResult&& InstallResult);
+	void								ProgressInstallRequest(const FOnDemandInternalInstallRequest& InstallRequest, const FOnDemandInstallProgress& Progress);
+	void								CompletePurgeRequest(FPurgeRequest& PurgeRequest, FOnDemandPurgeResult&& Result);
+	void								OnEncryptionKeyAdded(const FGuid& Id, const FAES::FAESKey& Key);
+	static void							CreateContainersFromToc(
+											FStringView MountId,
+											FStringView TocPath,
+											FOnDemandToc& Toc,
+											FOnDemandHostGroup HostGroup,
+											const TArray<FString>& ExistingContainers,
+											TArray<FSharedOnDemandContainer>& Out);
 
 	FSharedInstallCache						InstallCache;
 	FSharedPackageStoreBackend				PackageStoreBackend;
diff --git a/Engine/Source/Runtime/Experimental/IoStore/OnDemand/Private/OnDemandToc.cpp b/Engine/Source/Runtime/Experimental/IoStore/OnDemand/Private/OnDemandToc.cpp
deleted file mode 100644
diff --git a/Engine/Source/Runtime/Experimental/IoStore/OnDemand/Public/IO/IoStoreOnDemand.h b/Engine/Source/Runtime/Experimental/IoStore/OnDemand/Public/IO/IoStoreOnDemand.h
deleted file mode 100644
diff --git a/Engine/Source/Runtime/Experimental/IoStore/OnDemand/Public/IO/OnDemandHostGroup.h b/Engine/Source/Runtime/Experimental/IoStore/OnDemand/Public/IO/OnDemandHostGroup.h
deleted file mode 100644
diff --git a/Engine/Source/Runtime/Experimental/IoStore/OnDemand/Public/IO/OnDemandToc.h b/Engine/Source/Runtime/Experimental/IoStore/OnDemand/Public/IO/OnDemandToc.h
deleted file mode 100644
diff --git a/Engine/Source/Runtime/Experimental/IoStore/OnDemandCore/Internal/IO/IoStoreOnDemandInternals.h b/Engine/Source/Runtime/Experimental/IoStore/OnDemandCore/Internal/IO/IoStoreOnDemandInternals.h
new file mode 100644
--- /dev/null
+++ b/Engine/Source/Runtime/Experimental/IoStore/OnDemandCore/Internal/IO/IoStoreOnDemandInternals.h
@@ -0,0 +1,76 @@
+// Copyright Epic Games, Inc. All Rights Reserved.
+
+#pragma once
+
+#include "IO/IoStoreOnDemand.h"
+#include "HAL/PlatformTime.h"
+
+#define UE_API IOSTOREONDEMANDCORE_API
+
+namespace UE::IoStore
+{
+
+///////////////////////////////////////////////////////////////////////////////
+class FOnDemandInternalContentHandle
+{
+public:
+	FOnDemandInternalContentHandle()
+		: DebugName(TEXT("NoName"))
+	{ }
+
+	FOnDemandInternalContentHandle(FSharedString InDebugName)
+		: DebugName(InDebugName)
+	{ }
+
+	UE_API ~FOnDemandInternalContentHandle();
+
+	UPTRINT HandleId() const { return UPTRINT(this); }
+
+	FSharedString			DebugName;
+	FWeakOnDemandIoStore	IoStore;
+};
+
+UE_API FString LexToString(const FOnDemandInternalContentHandle& Handle);
+
+////////////////////////////////////////////////////////////////////////////////
+class FOnDemandInternalInstallRequest
+{
+public:
+	using EStatus = FOnDemandRequest::EStatus;
+
+	UE_API static uint32 NextSeqNo;
+
+	static bool SortByPriority(FSharedInternalInstallRequest LHS, FSharedInternalInstallRequest RHS)
+	{
+		if (LHS->Args.Priority == RHS->Args.Priority)
+		{
+			return LHS->SeqNo > RHS->SeqNo;
+		}
+		return static_cast<uint8>(LHS->Args.Priority) < static_cast<uint8>(RHS->Args.Priority);
+	};
+
+	FOnDemandInternalInstallRequest(
+		FOnDemandInstallArgs&& InArgs,
+		FOnDemandInstallCompleted&& InOnCompleted,
+		FOnDemandInstallProgressed&& InOnProgress)
+			: Args(MoveTemp(InArgs))
+			, OnCompleted(MoveTemp(InOnCompleted))
+			, OnProgress(MoveTemp(InOnProgress))
+	{ }
+
+	bool IsCancelRequested() const { return EIoErrorCode::Cancelled == ErrorCode.load(std::memory_order_relaxed); }
+
+	FOnDemandInstallArgs		Args;
+	FOnDemandInstallCompleted	OnCompleted;
+	FOnDemandInstallProgressed	OnProgress;
+	double						StartTime = FPlatformTime::Seconds();
+	uint32						SeqNo = NextSeqNo++;
+	std::atomic<EStatus>		Status{EStatus::Pending};
+	std::atomic<EIoErrorCode>	ErrorCode{EIoErrorCode::Unknown};
+};
+
+UE_API FString LexToString(const FOnDemandInternalInstallRequest& Request);
+
+} // namespace UE::IoStore
+
+#undef UE_API
diff --git a/Engine/Source/Runtime/Experimental/IoStore/OnDemandCore/IoStoreOnDemandCore.Build.cs b/Engine/Source/Runtime/Experimental/IoStore/OnDemandCore/IoStoreOnDemandCore.Build.cs
new file mode 100644
--- /dev/null
+++ b/Engine/Source/Runtime/Experimental/IoStore/OnDemandCore/IoStoreOnDemandCore.Build.cs
@@ -0,0 +1,20 @@
+// Copyright Epic Games, Inc. All Rights Reserved.
+
+using UnrealBuildTool;
+
+public class IoStoreOnDemandCore : ModuleRules
+{
+	public IoStoreOnDemandCore(ReadOnlyTargetRules Target) : base(Target)
+	{
+		PublicDependencyModuleNames.Add("Core");
+		PrivateDependencyModuleNames.AddRange(
+			new string[] {
+				"CoreUObject",
+				"Json"
+			}
+		);
+
+		bAllowConfidentialPlatformDefines = true;
+		CppCompileWarningSettings.UnsafeTypeCastWarningLevel = WarningLevel.Error; 
+	}
+}
diff --git a/Engine/Source/Runtime/Experimental/IoStore/OnDemandCore/Private/IoStoreOnDemand.cpp b/Engine/Source/Runtime/Experimental/IoStore/OnDemandCore/Private/IoStoreOnDemand.cpp
new file mode 100644
--- /dev/null
+++ b/Engine/Source/Runtime/Experimental/IoStore/OnDemandCore/Private/IoStoreOnDemand.cpp
@@ -0,0 +1,436 @@
+// Copyright Epic Games, Inc. All Rights Reserved.
+
+#include "IO/IoStoreOnDemand.h"
+
+#include "Containers/StringView.h"
+#include "Features/IModularFeatures.h"
+#include "HAL/FileManager.h"
+#include "HAL/IConsoleManager.h"
+#include "IO/IoStatus.h"
+#include "IO/IoStoreOnDemand.h"
+#include "Misc/Paths.h"
+#include "Misc/PathViews.h"
+#include "Misc/Parse.h"
+#include "Modules/ModuleManager.h"
+#include "IO/IoStoreOnDemand.h"
+#if !UE_BUILD_SHIPPING
+#include "Misc/PackageName.h"
+#endif //!UE_BUILD_SHIPPING
+
+DEFINE_LOG_CATEGORY(LogIoStoreOnDemand);
+DEFINE_LOG_CATEGORY(LogIas);
+
+namespace UE::IoStore
+{
+
+FSharedOnDemandIoStore GOnDemandIoStore;
+
+///////////////////////////////////////////////////////////////////////////////
+#if !UE_BUILD_SHIPPING
+namespace Commands
+{
+
+////////////////////////////////////////////////////////////////////////////////
+static TIoStatusOr<FOnDemandHostGroup> SplitHostAndTocUrl(FStringView Url, FStringView& OutTocRelativUrl)
+{
+	if (Url.StartsWith(TEXTVIEW("http")) == false)
+	{
+		return FIoStatus(EIoErrorCode::InvalidParameter, TEXT("Invalid URL protocol"));
+	}
+
+	int32 Delim = INDEX_NONE;
+	if (Url.FindChar(':', Delim) == false)
+	{
+		return FIoStatus(EIoErrorCode::InvalidParameter, TEXT("Invalid URL"));
+	}
+
+	const int32 ProtocolDelim = Delim + 3;
+	if (Url.RightChop(ProtocolDelim).FindChar('/', Delim) == false)
+	{
+		return FIoStatus(EIoErrorCode::InvalidParameter, TEXT("Failed to find host and TOC path delimiter"));
+	}
+
+	const FStringView Host	= Url.Left(ProtocolDelim + Delim);
+	OutTocRelativUrl		= Url.RightChop(Host.Len());
+
+	return FOnDemandHostGroup::Create(Host);
+}
+
+////////////////////////////////////////////////////////////////////////////////
+static void PurgeInstallCache()
+{
+	if (UE::IoStore::IOnDemandIoStore* IoStore = UE::IoStore::TryGetOnDemandIoStore())
+	{
+		UE_LOG(LogIoStoreOnDemand, Display, TEXT("Purging on demand install cache"));
+		IoStore->Purge(FOnDemandPurgeArgs(), [](const FOnDemandPurgeResult& Result)
+		{
+			if (Result.Status.IsOk())
+			{
+				UE_LOG(LogIoStoreOnDemand, Display, TEXT("Purged on demand install cache"));
+			}
+			else
+			{
+				UE_LOG(LogIoStoreOnDemand, Error, TEXT("Failed Purged on demand install cache: %s"), *Result.Status.ToString());
+			}
+		});
+	}
+	else
+	{
+		UE_LOG(LogIoStoreOnDemand, Error, TEXT("Purge install cache failed, reason 'I/O store on-demand module not initialized'"));
+	}
+}
+
+////////////////////////////////////////////////////////////////////////////////
+static void PrintCacheUsage()
+{
+	if (UE::IoStore::IOnDemandIoStore* IoStore = UE::IoStore::TryGetOnDemandIoStore())
+	{
+		TIoStatusOr<FOnDemandCacheUsage> MaybeUsage = IoStore->GetCacheUsage();
+		if (!MaybeUsage.IsOk())
+		{
+			UE_LOG(LogIoStoreOnDemand, Error, TEXT("iostore.CacheUsage failed: %s"), *MaybeUsage.Status().ToString());
+			return;
+		}
+
+		const FOnDemandCacheUsage CacheUsage = MaybeUsage.ConsumeValueOrDie();
+		UE_LOG(LogIoStoreOnDemand, Display, TEXT("iostore.CacheUsage"));
+		TStringBuilder<512> Sb;
+		Sb << TEXTVIEW("InstallCache: ") << CacheUsage.InstallCache << TEXTVIEW(", StreamingCache: ") << CacheUsage.StreamingCache;
+		UE_LOG(LogIoStoreOnDemand, Display, TEXT("%s"), Sb.ToString());
+	}
+	else
+	{
+		UE_LOG(LogIoStoreOnDemand, Error, TEXT("Print cache usage failed, reason 'I/O store on-demand module not initialized'"));
+	}
+}
+
+////////////////////////////////////////////////////////////////////////////////
+static void MountUrl(const TArray<FString>& Args)
+{
+	UE::IoStore::IOnDemandIoStore* IoStore = UE::IoStore::TryGetOnDemandIoStore();
+	if (IoStore == nullptr)
+	{
+		UE_LOG(LogIoStoreOnDemand, Error, TEXT("Failed to load I/O store on-demand module."));
+		return;
+	}
+
+	if (Args.Num() < 2)
+	{
+		UE_LOG(LogIoStoreOnDemand, Error, TEXT("Not enough arguments."));
+		return;
+	}
+
+	FString Url = Args[0];
+	Url.TrimQuotesInline();
+
+	FStringView TocRelativeUrl;
+	TIoStatusOr<FOnDemandHostGroup> HostGroup = SplitHostAndTocUrl(Url, TocRelativeUrl);
+	if (HostGroup.IsOk() == false)
+	{
+		UE_LOG(LogIoStoreOnDemand, Error, TEXT("%s"), *HostGroup.Status().ToString());
+		return;
+	}
+
+	EOnDemandMountOptions Options = EOnDemandMountOptions::StreamOnDemand;
+	const FString& Arg = Args[1];
+	if (Arg.ToLower().Contains(TEXT("install")))
+	{
+		Options = EOnDemandMountOptions::InstallOnDemand; 
+	}
+	const FString& MountId = Args.Num() > 2 ? Args[2] : Url; 
+
+	IoStore->Mount(
+		FOnDemandMountArgs
+		{
+			.MountId = MountId,
+			.TocRelativeUrl = FString(TocRelativeUrl),
+			.HostGroup = HostGroup.ConsumeValueOrDie(),
+			.Options = Options
+		},
+		[](FOnDemandMountResult MountResult)
+		{
+			UE_CLOG(!MountResult.Status.IsOk(), LogIas, Error,
+				TEXT("Failed to mount TOC for '%s', reason '%s'"), *MountResult.MountId, *MountResult.Status.ToString());
+		});
+}
+
+////////////////////////////////////////////////////////////////////////////////
+static void MountFile(const TArray<FString>& Args)
+{
+	UE::IoStore::IOnDemandIoStore* IoStore = UE::IoStore::TryGetOnDemandIoStore();
+	if (IoStore == nullptr)
+	{
+		UE_LOG(LogIoStoreOnDemand, Error, TEXT("Failed to load I/O store on-demand module."));
+		return;
+	}
+
+	if (Args.Num() < 3)
+	{
+		UE_LOG(LogIoStoreOnDemand, Error, TEXT("Not enough arguments."));
+		return;
+	}
+
+	TArray<FOnDemandMountArgs> AllMountArgs;
+	FString FilenameOrWildcard = Args[0];
+	FilenameOrWildcard.TrimQuotesInline();
+	TArray<FString> FilePaths;
+
+	const FString ContentDir = FString::Printf(TEXT("%sPaks/"), *FPaths::ProjectContentDir());
+	IFileManager& Ifm = IFileManager::Get();
+	Ifm.IterateDirectoryRecursively(
+		*ContentDir,
+		[&FilenameOrWildcard, &FilePaths, &ContentDir](const TCHAR* FilenameOrDirectory, bool bIsDirectory)
+		{
+			if (bIsDirectory == false)
+			{
+				FString FilePath(FilenameOrDirectory);
+				if (FPathViews::GetExtension(FilePath) == TEXT("uondemandtoc"))
+				{
+					FStringView Filename = FPathViews::GetBaseFilename(FilePath);
+					if (Filename == FilenameOrWildcard || FString(Filename).MatchesWildcard(FilenameOrWildcard))
+					{
+						FilePaths.Add(FilenameOrDirectory);
+					}
+				}
+			}
+			return true;
+		});
+
+	if (FilePaths.IsEmpty())
+	{
+		UE_LOG(LogIoStoreOnDemand, Error, TEXT("Failed to find any on-demand TOC file(s) matching '%s'"), *FilenameOrWildcard);
+		return;
+	}
+
+	EOnDemandMountOptions Options = EOnDemandMountOptions::StreamOnDemand;
+	if (Args[1].ToLower().Contains(TEXT("install")))
+	{
+		Options = EOnDemandMountOptions::InstallOnDemand; 
+	}
+
+	FString HostPath = Args[2];
+	HostPath.TrimQuotesInline();
+	FStringView TocRelativeUrl;
+	TIoStatusOr<FOnDemandHostGroup> HostGroup = SplitHostAndTocUrl(HostPath, TocRelativeUrl);
+	if (HostGroup.IsOk() == false)
+	{
+		UE_LOG(LogIoStoreOnDemand, Error, TEXT("%s"), *HostGroup.Status().ToString());
+		return;
+	}
+
+	const FString& MountId = Args.Num() > 3 ? Args[3] : FilePaths[0];
+	for (const FString& FilePath : FilePaths)
+	{
+		FOnDemandMountArgs& MountArgs	= AllMountArgs.AddDefaulted_GetRef();
+		MountArgs.FilePath				= FilePath;
+		MountArgs.HostGroup				= HostGroup.ConsumeValueOrDie();
+		MountArgs.TocRelativeUrl		= TocRelativeUrl;
+		MountArgs.MountId				= MountId;
+	}
+
+	if (AllMountArgs.IsEmpty())
+	{
+		UE_LOG(LogIoStoreOnDemand, Error, TEXT("Failed to parse mount arguments"));
+		return;
+	}
+
+	for (FOnDemandMountArgs& MountArgs : AllMountArgs)
+	{
+		MountArgs.Options = Options;
+		IoStore->Mount(
+			MoveTemp(MountArgs),
+			[](FOnDemandMountResult MountResult)
+			{
+				UE_CLOG(!MountResult.Status.IsOk(), LogIas, Error,
+					TEXT("Failed to mount TOC for '%s', reason '%s'"), *MountResult.MountId, *MountResult.Status.ToString());
+			});
+	}
+}
+
+////////////////////////////////////////////////////////////////////////////////
+static void InstallPackage(const TArray<FString>& Args)
+{
+	UE::IoStore::IOnDemandIoStore* IoStore = UE::IoStore::TryGetOnDemandIoStore();
+	if (IoStore == nullptr)
+	{
+		UE_LOG(LogIoStoreOnDemand, Error, TEXT("Failed to load I/O store on-demand module."));
+		return;
+	}
+
+	if (Args.Num() < 1)
+	{
+		UE_LOG(LogIoStoreOnDemand, Error, TEXT("Not enough arguments."));
+		return;
+	}
+
+	FString PackageName = Args[0];
+	PackageName.TrimQuotesInline();
+	FPackageId PackageId;
+	if (FPackageName::IsValidLongPackageName(PackageName))
+	{
+		PackageId = FPackageId::FromName(FName(*PackageName));
+	}
+
+	if (PackageId.IsValid() == false)
+	{
+		UE_LOG(LogIoStoreOnDemand, Error, TEXT("Invalid package name '%s'"), *PackageName);
+		return;
+	}
+
+	static FOnDemandContentHandle DefaultContentHandle = FOnDemandContentHandle::Create(TEXT("ConsoleCommand"));
+
+	FOnDemandInstallArgs InstallArgs;
+	InstallArgs.PackageIds.Add(PackageId);
+	InstallArgs.ContentHandle = DefaultContentHandle;
+
+	IoStore->Install(
+		MoveTemp(InstallArgs),
+		[PackageName = MoveTemp(PackageName)](FOnDemandInstallResult InstallResult)
+		{
+			if (InstallResult.Status.IsOk())
+			{
+				UE_LOG(LogIoStoreOnDemand, Log, TEXT("Successfully installed package '%s'"), *PackageName);
+			}
+			else
+			{
+				UE_LOG(LogIoStoreOnDemand, Error, TEXT("Failed to install package '%s', reason '%s'"),
+					*PackageName, *InstallResult.Status.ToString());
+			}
+		});
+}
+
+////////////////////////////////////////////////////////////////////////////////
+static FAutoConsoleCommand PurgeCacheCommand(
+	TEXT("iostore.PurgeInstallCache"),
+	TEXT("Purge On Demand Install Cache"),
+	FConsoleCommandDelegate::CreateStatic(PurgeInstallCache),
+	ECVF_Cheat);
+
+static FAutoConsoleCommand PrintCacheUsageCommand(
+	TEXT("iostore.CacheUsage"),
+	TEXT("print cache usage"),
+	FConsoleCommandDelegate::CreateStatic(PrintCacheUsage),
+	ECVF_Cheat);
+
+static FAutoConsoleCommand MountUrlCommand(
+	TEXT("iostore.MountUrl"),
+	TEXT("<URL> <Install|Stream> <MountId>"),
+	FConsoleCommandWithArgsDelegate::CreateStatic(MountUrl));
+
+static FAutoConsoleCommand MountFileCommand(
+	TEXT("iostore.MountFile"),
+	TEXT("<Filename|Wildcard> <Install|Stream> <HostPath> <MountId>"),
+	FConsoleCommandWithArgsDelegate::CreateStatic(MountFile));
+
+static FAutoConsoleCommand InstallPackageCommand(
+	TEXT("iostore.InstallPackage"),
+	TEXT("<PackageName>"),
+	FConsoleCommandWithArgsDelegate::CreateStatic(InstallPackage));
+
+} // namespace Commands
+#endif // !UE_BUILD_SHIPPING
+
+////////////////////////////////////////////////////////////////////////////////
+FString GIasOnDemandTocExt = TEXT(".uondemandtoc");
+
+static const TCHAR* NotInitializedError = TEXT("I/O store on-demand not initialized");
+
+////////////////////////////////////////////////////////////////////////////////
+FStringBuilderBase& operator<<(FStringBuilderBase& Sb, const FOnDemandInstallCacheUsage& CacheUsage)
+{
+	Sb.Appendf(TEXT("MaxSize=%.2lf MiB, TotalSize=%.2lf MiB, ReferencedBlockSize=%.2lf MiB"),
+		double(CacheUsage.MaxSize) / 1024.0 / 1024.0,
+		double(CacheUsage.TotalSize) / 1024.0 / 1024.0,
+		double(CacheUsage.ReferencedBlockSize) / 1024.0 / 1024.0);
+
+	return Sb;
+}
+
+////////////////////////////////////////////////////////////////////////////////
+FStringBuilderBase& operator<<(FStringBuilderBase& Sb, const FOnDemandStreamingCacheUsage& CacheUsage)
+{
+	Sb.Appendf(TEXT("MaxSize=%.2lf MiB, TotalSize=%.2lf MiB"),
+		double(CacheUsage.MaxSize) / 1024.0 / 1024.0, double(CacheUsage.TotalSize) / 1024.0 / 1024.0);
+	return Sb;
+}
+
+////////////////////////////////////////////////////////////////////////////////
+class FIoStoreOnDemandCoreModule final
+	: public IModuleInterface
+{
+public:
+	virtual void StartupModule() override;
+	virtual void ShutdownModule() override;
+
+private:
+	void HandleModularFeatureRegistered(const FName& Type, IModularFeature* ModularFeature);
+};
+IMPLEMENT_MODULE(UE::IoStore::FIoStoreOnDemandCoreModule, IoStoreOnDemandCore);
+
+////////////////////////////////////////////////////////////////////////////////
+void FIoStoreOnDemandCoreModule::HandleModularFeatureRegistered(const FName& Type, IModularFeature* ModularFeature)
+{
+	if (Type != IOnDemandIoStoreFactory::FeatureName || GOnDemandIoStore.IsValid())
+	{
+		return;
+	}
+
+	IOnDemandIoStoreFactory* Factory = static_cast<IOnDemandIoStoreFactory*>(ModularFeature);
+	GOnDemandIoStore = Factory->CreateInstance();
+	bool bInitializePostHostfix = false;
+	#if UE_IAS_CUSTOM_INITIALIZATION
+	bInitializePostHostfix = true;
+	#endif
+	if (FIoStatus Status = GOnDemandIoStore->Initialize(); !Status.IsOk())
+	{
+		UE_LOG(LogIoStoreOnDemand, Error, TEXT("Failed to initialize I/O store, reason ''"), *Status.ToString());
+		GOnDemandIoStore.Reset();
+	}
+	if (bInitializePostHostfix == false)
+	{
+		if (FIoStatus Status = GOnDemandIoStore->InitializePostHotfix(); !Status.IsOk())
+		{
+			UE_LOG(LogIoStoreOnDemand, Error, TEXT("Failed to initialize I/O store, reason ''"), *Status.ToString());
+			GOnDemandIoStore.Reset();
+		}
+	}
+	IModularFeatures::Get().OnModularFeatureRegistered().RemoveAll(this);
+}
+
+void FIoStoreOnDemandCoreModule::StartupModule()
+{
+	IModularFeatures& Features = IModularFeatures::Get();
+	if (Features.GetModularFeatureImplementationCount(IOnDemandIoStoreFactory::FeatureName) > 0)
+	{
+		IModularFeature* Feature = Features.GetModularFeatureImplementation(IOnDemandIoStoreFactory::FeatureName, 0);
+		HandleModularFeatureRegistered(IOnDemandIoStoreFactory::FeatureName, Feature);
+	}
+	else
+	{
+		Features.OnModularFeatureRegistered().AddRaw(this, &FIoStoreOnDemandCoreModule::HandleModularFeatureRegistered);
+	}
+}
+
+void FIoStoreOnDemandCoreModule::ShutdownModule()
+{
+	IModularFeatures::Get().OnModularFeatureRegistered().RemoveAll(this);
+}
+
+////////////////////////////////////////////////////////////////////////////////
+IOnDemandIoStore* TryGetOnDemandIoStore()
+{
+	return GOnDemandIoStore.Get();
+}
+
+////////////////////////////////////////////////////////////////////////////////
+IOnDemandIoStore& GetOnDemandIoStore()
+{
+	check(GOnDemandIoStore.IsValid());
+	return *GOnDemandIoStore.Get();
+}
+
+FName IOnDemandIoStoreFactory::FeatureName = FName(TEXT("OnDemandIoStoreFactory"));
+
+} // namespace UE::IoStore
+
+////////////////////////////////////////////////////////////////////////////////
diff --git a/Engine/Source/Runtime/Experimental/IoStore/OnDemandCore/Private/IoStoreOnDemandInternals.cpp b/Engine/Source/Runtime/Experimental/IoStore/OnDemandCore/Private/IoStoreOnDemandInternals.cpp
new file mode 100644
--- /dev/null
+++ b/Engine/Source/Runtime/Experimental/IoStore/OnDemandCore/Private/IoStoreOnDemandInternals.cpp
@@ -0,0 +1,123 @@
+// Copyright Epic Games, Inc. All Rights Reserved.
+
+#include "IO/IoStoreOnDemandInternals.h"
+
+namespace UE::IoStore
+{
+
+///////////////////////////////////////////////////////////////////////////////
+FOnDemandInternalContentHandle::~FOnDemandInternalContentHandle()
+{
+	if (TSharedPtr<IOnDemandIoStore, ESPMode::ThreadSafe> Pinned = IoStore.Pin(); Pinned.IsValid())
+	{
+		Pinned->ReleaseContent(*this);
+	}
+}
+
+FString LexToString(const FOnDemandInternalContentHandle& Handle)
+{
+	return FString::Printf(TEXT("Id=0x%llX DebugName='%s'"), Handle.HandleId(), *Handle.DebugName);
+}
+
+///////////////////////////////////////////////////////////////////////////////
+FOnDemandContentHandle::FOnDemandContentHandle()
+{
+}
+
+FOnDemandContentHandle::~FOnDemandContentHandle()
+{
+}
+
+FOnDemandContentHandle FOnDemandContentHandle::Create()
+{
+	FOnDemandContentHandle NewHandle;
+	NewHandle.Handle = MakeShared<FOnDemandInternalContentHandle, ESPMode::ThreadSafe>();
+	return NewHandle;
+}
+
+FOnDemandContentHandle FOnDemandContentHandle::Create(FSharedString DebugName)
+{
+	FOnDemandContentHandle NewHandle;
+	NewHandle.Handle = MakeShared<FOnDemandInternalContentHandle, ESPMode::ThreadSafe>(DebugName);
+	return NewHandle;
+}
+
+FOnDemandContentHandle FOnDemandContentHandle::Create(FStringView DebugName)
+{
+	return Create(FSharedString(DebugName));
+}
+
+FString LexToString(const FOnDemandContentHandle& Handle)
+{
+	return Handle.IsValid() ? LexToString(*Handle.Handle) : TEXT("Invalid");
+}
+
+////////////////////////////////////////////////////////////////////////////////
+FString LexToString(const FOnDemandInternalInstallRequest& Request)
+{
+	return FString::Printf(TEXT("Id=0x%llX, DebugName='%s', ContentHandle='%s'"),
+		UPTRINT(&Request),
+		Request.Args.DebugName.IsEmpty()
+			? TEXT("NoName")
+			: *Request.Args.DebugName,
+		*LexToString(Request.Args.ContentHandle));
+}
+
+uint32 FOnDemandInternalInstallRequest::NextSeqNo = 0;
+
+///////////////////////////////////////////////////////////////////////////////
+FOnDemandInstallRequest::FOnDemandInstallRequest()
+{
+}
+
+FOnDemandInstallRequest::FOnDemandInstallRequest(FWeakOnDemandIoStore InIoStore, FSharedInternalInstallRequest InternalRequest)
+	: IoStore(InIoStore)
+	, Request(InternalRequest)
+{
+}
+
+FOnDemandInstallRequest::~FOnDemandInstallRequest()
+{
+}
+
+FOnDemandInstallRequest::FOnDemandInstallRequest(FOnDemandInstallRequest&& Other)
+	: IoStore(MoveTemp(Other.IoStore))
+	, Request(MoveTemp(Other.Request))
+{
+}
+
+FOnDemandInstallRequest& FOnDemandInstallRequest::operator=(FOnDemandInstallRequest&& Other)
+{
+	Request = MoveTemp(Other.Request);
+	IoStore = MoveTemp(Other.IoStore);
+	return *this;
+}
+
+FOnDemandInstallRequest::EStatus FOnDemandInstallRequest::GetStatus() const
+{
+	return Request.IsValid() ? Request->Status.load(std::memory_order_relaxed) : EStatus::None; 
+}
+
+void FOnDemandInstallRequest::Cancel()
+{
+	if (IsPending())
+	{
+		if (TSharedPtr<IOnDemandIoStore, ESPMode::ThreadSafe> Pinned = IoStore.Pin(); Pinned.IsValid())
+		{
+			Pinned->CancelInstallRequest(Request);
+		}
+	}
+}
+
+void FOnDemandInstallRequest::UpdatePriority(int32 NewPriority)
+{
+	if (IsPending())
+	{
+		if (TSharedPtr<IOnDemandIoStore, ESPMode::ThreadSafe> Pinned = IoStore.Pin(); Pinned.IsValid())
+		{
+			Pinned->UpdateInstallRequestPriority(Request, NewPriority);
+		}
+	}
+}
+
+} // namespace UE::IoStore
diff --git a/Engine/Source/Runtime/Experimental/IoStore/OnDemandCore/Private/OnDemandHostGroup.cpp b/Engine/Source/Runtime/Experimental/IoStore/OnDemandCore/Private/OnDemandHostGroup.cpp
new file mode 100644
--- /dev/null
+++ b/Engine/Source/Runtime/Experimental/IoStore/OnDemandCore/Private/OnDemandHostGroup.cpp
@@ -0,0 +1,173 @@
+// Copyright Epic Games, Inc. All Rights Reserved.
+
+#include "IO/OnDemandHostGroup.h"
+
+#include "Containers/AnsiString.h"
+#include "Containers/Array.h"
+#include "Containers/StringView.h"
+
+namespace UE::IoStore
+{
+
+namespace Private
+{
+
+static bool ValidateUrl(FAnsiStringView Url, FString& Reason)
+{
+	//TODO: Add better validation
+	return Url.StartsWith("http") || Url.StartsWith("https");
+}
+
+} // namespace Private
+
+struct FOnDemandHostGroup::FImpl
+{
+	TArray<FAnsiString> HostUrls;
+	int32				PrimaryIndex = INDEX_NONE;
+};
+
+FOnDemandHostGroup::FOnDemandHostGroup()
+	: Impl(MakeShared<FImpl>())
+{
+}
+
+FOnDemandHostGroup::FOnDemandHostGroup(FOnDemandHostGroup::FSharedImpl&& InImpl)
+	: Impl(MoveTemp(InImpl))
+{
+}
+
+FOnDemandHostGroup::~FOnDemandHostGroup()
+{
+}
+
+TConstArrayView<FAnsiString> FOnDemandHostGroup::Hosts() const
+{
+	return Impl->HostUrls;
+}
+
+FAnsiStringView FOnDemandHostGroup::Host(int32 Index) const
+{
+	return Impl->HostUrls.IsEmpty() ? FAnsiStringView() : Impl->HostUrls[Index];
+}
+
+FAnsiStringView FOnDemandHostGroup::CycleHost(int32& InOutIndex) const
+{
+	InOutIndex = (InOutIndex + 1) % Impl->HostUrls.Num();
+	return Host(InOutIndex);
+}
+
+void FOnDemandHostGroup::SetPrimaryHost(int32 Index)
+{
+	check(Index >= 0 && Index < Impl->HostUrls.Num() || Index == INDEX_NONE);
+	Impl->PrimaryIndex = Index;
+}
+
+FAnsiStringView FOnDemandHostGroup::PrimaryHost() const
+{
+	if (Impl->PrimaryIndex != INDEX_NONE)
+	{
+		check(Impl->PrimaryIndex >= 0 && Impl->PrimaryIndex < Impl->HostUrls.Num());
+		return Impl->HostUrls[Impl->PrimaryIndex];
+	}
+
+	return FAnsiStringView();
+}
+
+int32 FOnDemandHostGroup::PrimaryHostIndex() const
+{
+	return Impl->PrimaryIndex;
+}
+
+bool FOnDemandHostGroup::IsEmpty() const
+{
+	return Impl->HostUrls.IsEmpty();
+}
+
+TIoStatusOr<FOnDemandHostGroup> FOnDemandHostGroup::Create(FAnsiStringView Url)
+{
+	if (Url.EndsWith('/'))
+	{
+		Url.RemoveSuffix(1);
+	}
+
+	FString Reason;
+	if (Private::ValidateUrl(Url, Reason) == false)
+	{
+		FIoStatus(EIoErrorCode::InvalidParameter, Reason);
+	}
+
+	FSharedImpl Impl = MakeShared<FImpl>();
+	Impl->HostUrls.Add(FAnsiString(Url));
+	Impl->PrimaryIndex = 0;
+
+	return FOnDemandHostGroup(MoveTemp(Impl)); 
+}
+
+TIoStatusOr<FOnDemandHostGroup> FOnDemandHostGroup::Create(FStringView Url)
+{
+	return Create(FAnsiString(StringCast<ANSICHAR>(Url.GetData(), Url.Len()))); 
+}
+
+TIoStatusOr<FOnDemandHostGroup> FOnDemandHostGroup::Create(TConstArrayView<FAnsiString> Urls)
+{
+	if (Urls.IsEmpty())
+	{
+		return FIoStatus(EIoErrorCode::InvalidParameter, TEXT("No URLs specified"));
+	}
+
+	FSharedImpl Impl = MakeShared<FImpl>();
+	FString		Reason;
+
+	for (const FAnsiString& Url : Urls)
+	{
+		FAnsiStringView UrlView = Url;
+
+		if (UrlView.EndsWith('/'))
+		{
+			UrlView.RemoveSuffix(1);
+		}
+
+		if (Private::ValidateUrl(UrlView, Reason) == false)
+		{
+			return FIoStatus(EIoErrorCode::InvalidParameter, Reason);
+		}
+		Impl->HostUrls.Emplace(UrlView);
+	}
+	Impl->PrimaryIndex = 0; 
+
+	return FOnDemandHostGroup(MoveTemp(Impl)); 
+}
+
+TIoStatusOr<FOnDemandHostGroup> FOnDemandHostGroup::Create(TConstArrayView<FString> Urls)
+{
+	if (Urls.IsEmpty())
+	{
+		return FIoStatus(EIoErrorCode::InvalidParameter, TEXT("No URLs specified"));
+	}
+
+	FSharedImpl Impl = MakeShared<FImpl>();
+	FString		Reason;
+
+	for (const FString& Url : Urls)
+	{
+		FStringView UrlView = Url;
+
+		if (UrlView.EndsWith(TEXT('/')))
+		{
+			UrlView.RemoveSuffix(1);
+		}
+
+		FAnsiString AnsiUrl(StringCast<ANSICHAR>(UrlView.GetData(), UrlView.Len()));
+
+		if (Private::ValidateUrl(AnsiUrl, Reason) == false)
+		{
+			return FIoStatus(EIoErrorCode::InvalidParameter, Reason);
+		}
+		Impl->HostUrls.Emplace(MoveTemp(AnsiUrl));
+	}
+	Impl->PrimaryIndex = 0;
+
+	return FOnDemandHostGroup(MoveTemp(Impl));
+}
+
+} //namespace UE::IoStore
diff --git a/Engine/Source/Runtime/Experimental/IoStore/OnDemandCore/Private/OnDemandToc.cpp b/Engine/Source/Runtime/Experimental/IoStore/OnDemandCore/Private/OnDemandToc.cpp
new file mode 100644
--- /dev/null
+++ b/Engine/Source/Runtime/Experimental/IoStore/OnDemandCore/Private/OnDemandToc.cpp
@@ -0,0 +1,658 @@
+// Copyright Epic Games, Inc. All Rights Reserved.
+
+#include "IO/OnDemandToc.h"
+
+#include "HAL/FileManager.h"
+#include "HAL/FileManagerGeneric.h"
+#include "Misc/Paths.h"
+#include "Misc/PathViews.h"
+#include "Serialization/Archive.h"
+#include "Serialization/CompactBinarySerialization.h"
+#include "Serialization/CompactBinaryWriter.h"
+#include "Serialization/CustomVersion.h"
+#include "Serialization/MemoryReader.h"
+
+namespace UE::IoStore
+{
+
+////////////////////////////////////////////////////////////////////////////////
+// 
+// Utility to create a FArchive capable of reading from disk using the exact same pathing
+// rules as FPlatformMisc::LoadTextFileFromPlatformPackage but without forcing the entire
+// file to be loaded at once.
+// 
+static TUniquePtr<FArchive> CreateReaderFromPlatformPackage(const FString& RelPath)
+{
+#if PLATFORM_IOS
+    // IOS OpenRead assumes it is in cookeddata, using ~ for the base path tells it to use the package base path instead
+    const FString AbsPath = FPaths::Combine(TEXT("~"), RelPath);
+#else
+    const FString AbsPath = FPaths::Combine(FGenericPlatformMisc::RootDir(), RelPath);
+#endif
+	if (TUniquePtr<IFileHandle> File(IPlatformFile::GetPlatformPhysical().OpenRead(*AbsPath)); File.IsValid())
+	{
+#if PLATFORM_ANDROID
+		// This is a handle to an asset so we need to call Seek(0) to move the internal
+		// offset to the start of the asset file.
+		File->Seek(0);
+#endif //PLATFORM_ANDROID
+		const uint32 ReadBufferSize = 256 * 1024;
+		const int64 FileSize = File->Size();
+		return MakeUnique<FArchiveFileReaderGeneric>(File.Release(), *AbsPath, FileSize, ReadBufferSize);
+	}
+
+	return TUniquePtr<FArchive>();
+}
+
+////////////////////////////////////////////////////////////////////////////////
+FArchive& operator<<(FArchive& Ar, FTocMeta& Meta)
+{
+	Ar << Meta.EpochTimestamp;
+	Ar << Meta.BuildVersion;
+	Ar << Meta.TargetPlatform;
+	return Ar;
+}
+
+FCbWriter& operator<<(FCbWriter& Writer, const FTocMeta& Meta)
+{
+	Writer.BeginObject();
+	Writer.AddInteger(UTF8TEXTVIEW("EpochTimestamp"), Meta.EpochTimestamp);
+	Writer.AddString(UTF8TEXTVIEW("BuildVersion"), Meta.BuildVersion);
+	Writer.AddString(UTF8TEXTVIEW("TargetPlatform"), Meta.TargetPlatform);
+	Writer.EndObject();
+
+	return Writer;
+}
+
+bool LoadFromCompactBinary(FCbFieldView Field, FTocMeta& OutMeta)
+{
+	if (FCbObjectView Obj = Field.AsObjectView())
+	{
+		OutMeta.EpochTimestamp = Obj["EpochTimestamp"].AsInt64();
+		OutMeta.BuildVersion = FString(Obj["BuildVersion"].AsString());
+		OutMeta.TargetPlatform = FString(Obj["TargetPlatform"].AsString());
+		return true;
+	}
+	
+	return false;
+}
+
+////////////////////////////////////////////////////////////////////////////////
+FArchive& operator<<(FArchive& Ar, FOnDemandTocHeader& Header)
+{
+	if (Ar.IsLoading() && Ar.TotalSize() < sizeof(FOnDemandTocHeader))
+	{
+		Ar.SetError();
+		return Ar;
+	}
+
+	Ar << Header.Magic;
+	if (Header.Magic != FOnDemandTocHeader::ExpectedMagic)
+	{
+		Ar.SetError();
+		return Ar;
+	}
+
+	Ar << Header.Version;
+	if (static_cast<EOnDemandTocVersion>(Header.Version) == EOnDemandTocVersion::Invalid)
+	{
+		Ar.SetError();
+		return Ar;
+	}
+
+	if (uint32(Header.Version) > uint32(EOnDemandTocVersion::Latest))
+	{
+		Ar.SetError();
+		return Ar;
+	}
+
+	Ar << Header.ChunkVersion;
+	Ar << Header.BlockSize;
+	Ar << Header.CompressionFormat;
+	Ar << Header.ChunksDirectory;
+
+	return Ar;
+}
+
+FCbWriter& operator<<(FCbWriter& Writer, const FOnDemandTocHeader& Header)
+{
+	Writer.BeginObject();
+	Writer.AddInteger(UTF8TEXTVIEW("Magic"), Header.Magic);
+	Writer.AddInteger(UTF8TEXTVIEW("Version"), Header.Version);
+	Writer.AddInteger(UTF8TEXTVIEW("ChunkVersion"), Header.ChunkVersion);
+	Writer.AddInteger(UTF8TEXTVIEW("BlockSize"), Header.BlockSize);
+	Writer.AddString(UTF8TEXTVIEW("CompressionFormat"), Header.CompressionFormat);
+	Writer.AddString(UTF8TEXTVIEW("ChunksDirectory"), Header.ChunksDirectory);
+	Writer.EndObject();
+
+	return Writer;
+}
+
+bool LoadFromCompactBinary(FCbFieldView Field, FOnDemandTocHeader& OutTocHeader)
+{
+	if (FCbObjectView Obj = Field.AsObjectView())
+	{
+		OutTocHeader.Magic = Obj["Magic"].AsUInt64();
+		OutTocHeader.Version = Obj["Version"].AsUInt32();
+		OutTocHeader.ChunkVersion = Obj["ChunkVersion"].AsUInt32();
+		OutTocHeader.BlockSize = Obj["BlockSize"].AsUInt32();
+		OutTocHeader.CompressionFormat = FString(Obj["CompressionFormat"].AsString());
+		OutTocHeader.ChunksDirectory = FString(Obj["ChunksDirectory"].AsString());
+
+		return OutTocHeader.Magic == FOnDemandTocHeader::ExpectedMagic &&
+			static_cast<EOnDemandTocVersion>(OutTocHeader.Version) != EOnDemandTocVersion::Invalid;
+	}
+
+	return false;
+}
+
+////////////////////////////////////////////////////////////////////////////////
+FArchive& operator<<(FArchive& Ar, FOnDemandTocEntry& Entry)
+{
+	Ar << Entry.Hash;
+	Ar << Entry.ChunkId;
+	Ar << Entry.RawSize;
+	Ar << Entry.EncodedSize;
+	Ar << Entry.BlockOffset;
+	Ar << Entry.BlockCount;
+
+	return Ar;
+}
+
+FCbWriter& operator<<(FCbWriter& Writer, const FOnDemandTocEntry& Entry)
+{
+	Writer.BeginObject();
+	Writer.AddHash(UTF8TEXTVIEW("Hash"), Entry.Hash);
+	Writer << UTF8TEXTVIEW("ChunkId") << Entry.ChunkId;
+	Writer.AddInteger(UTF8TEXTVIEW("RawSize"), Entry.RawSize);
+	Writer.AddInteger(UTF8TEXTVIEW("EncodedSize"), Entry.EncodedSize);
+	Writer.AddInteger(UTF8TEXTVIEW("BlockOffset"), Entry.BlockOffset);
+	Writer.AddInteger(UTF8TEXTVIEW("BlockCount"), Entry.BlockCount);
+	Writer.EndObject();
+
+	return Writer;
+}
+
+bool LoadFromCompactBinary(FCbFieldView Field, FOnDemandTocEntry& OutTocEntry)
+{
+	if (FCbObjectView Obj = Field.AsObjectView())
+	{
+		if (!LoadFromCompactBinary(Obj["ChunkId"], OutTocEntry.ChunkId))
+		{
+			return false;
+		}
+
+		OutTocEntry.Hash = Obj["Hash"].AsHash();
+		OutTocEntry.RawSize = Obj["RawSize"].AsUInt64(~uint64(0));
+		OutTocEntry.EncodedSize = Obj["EncodedSize"].AsUInt64(~uint64(0));
+		OutTocEntry.BlockOffset = Obj["BlockOffset"].AsUInt32(~uint32(0));
+		OutTocEntry.BlockCount = Obj["BlockCount"].AsUInt32();
+
+		return OutTocEntry.Hash != FIoHash::Zero &&
+			OutTocEntry.RawSize != ~uint64(0) &&
+			OutTocEntry.EncodedSize != ~uint64(0) &&
+			OutTocEntry.BlockOffset != ~uint32(0);
+	}
+
+	return false;
+}
+
+////////////////////////////////////////////////////////////////////////////////
+FArchive& operator<<(FArchive& Ar, FOnDemandTocContainerEntry& ContainerEntry)
+{
+	EOnDemandTocVersion TocVersion = EOnDemandTocVersion::Latest;
+
+	if (Ar.IsLoading())
+	{
+		const FCustomVersion* CustomVersion = Ar.GetCustomVersions().GetVersion(FOnDemandToc::VersionGuid);
+		check(CustomVersion);
+		TocVersion = static_cast<EOnDemandTocVersion>(CustomVersion->Version);
+
+		if (TocVersion >= EOnDemandTocVersion::ContainerId)
+		{
+			Ar << ContainerEntry.ContainerId;
+		}
+	}
+	else
+	{
+		Ar << ContainerEntry.ContainerId;
+	}
+
+	Ar << ContainerEntry.ContainerName;
+	Ar << ContainerEntry.EncryptionKeyGuid;
+	Ar << ContainerEntry.Entries;
+	Ar << ContainerEntry.BlockSizes;
+	Ar << ContainerEntry.BlockHashes;
+	Ar << ContainerEntry.UTocHash;
+
+	if (!Ar.IsLoading() || (TocVersion >= EOnDemandTocVersion::ContainerFlags))
+	{
+		Ar << ContainerEntry.ContainerFlags;
+	}
+
+	return Ar;
+}
+
+FCbWriter& operator<<(FCbWriter& Writer, const FOnDemandTocContainerEntry& ContainerEntry)
+{
+	Writer.BeginObject();
+	Writer << UTF8TEXTVIEW("Id") << ContainerEntry.ContainerId;
+	Writer.AddString(UTF8TEXTVIEW("Name"), ContainerEntry.ContainerName);
+	Writer.AddString(UTF8TEXTVIEW("EncryptionKeyGuid"), ContainerEntry.EncryptionKeyGuid);
+
+	Writer.BeginArray(UTF8TEXTVIEW("Entries"));
+	for (const FOnDemandTocEntry& Entry : ContainerEntry.Entries)
+	{
+		Writer << Entry;
+	}
+	Writer.EndArray();
+	
+	Writer.BeginArray(UTF8TEXTVIEW("BlockSizes"));
+	for (uint32 BlockSize : ContainerEntry.BlockSizes)
+	{
+		Writer << BlockSize;
+	}
+	Writer.EndArray();
+
+	Writer.BeginArray(UTF8TEXTVIEW("BlockHashes"));
+	for (uint32 BlockHash : ContainerEntry.BlockHashes)
+	{
+		Writer << BlockHash;
+	}
+	Writer.EndArray();
+
+	Writer.AddHash(UTF8TEXTVIEW("UTocHash"), ContainerEntry.UTocHash);
+
+	Writer.EndObject();
+
+	return Writer;
+}
+
+bool LoadFromCompactBinary(FCbFieldView Field, FOnDemandTocContainerEntry& OutContainer)
+{
+	if (FCbObjectView Obj = Field.AsObjectView())
+	{
+		OutContainer.ContainerName = FString(Obj["Name"].AsString());
+		OutContainer.EncryptionKeyGuid = FString(Obj["EncryptionKeyGuid"].AsString());
+
+		FCbArrayView Entries = Obj["Entries"].AsArrayView();
+		OutContainer.Entries.Reserve(int32(Entries.Num()));
+		for (FCbFieldView ArrayField : Entries)
+		{
+			if (!LoadFromCompactBinary(ArrayField, OutContainer.Entries.AddDefaulted_GetRef()))
+			{
+				return false;
+			}
+		}
+
+		FCbArrayView BlockSizes = Obj["BlockSizes"].AsArrayView();
+		OutContainer.BlockSizes.Reserve(int32(BlockSizes.Num()));
+		for (FCbFieldView ArrayField : BlockSizes)
+		{
+			OutContainer.BlockSizes.Add(ArrayField.AsUInt32());
+		}
+
+		FCbArrayView BlockHashes = Obj["BlockHashes"].AsArrayView();
+		OutContainer.BlockHashes.Reserve(int32(BlockHashes.Num()));
+		for (FCbFieldView ArrayField : BlockHashes)
+		{
+			if (ArrayField.IsHash())
+			{
+				const FIoHash BlockHash = ArrayField.AsHash();
+				OutContainer.BlockHashes.Add(*reinterpret_cast<const uint32*>(&BlockHash));
+			}
+			else
+			{
+				OutContainer.BlockHashes.Add(ArrayField.AsUInt32());
+			}
+		}
+
+		OutContainer.UTocHash = Obj["UTocHash"].AsHash();
+
+		return true;
+	}
+
+	return false;
+}
+
+////////////////////////////////////////////////////////////////////////////////
+bool FOnDemandTocSentinel::IsValid()
+{
+	return FMemory::Memcmp(&Data, FOnDemandTocSentinel::SentinelImg, FOnDemandTocSentinel::SentinelSize) == 0;
+}
+
+FArchive& operator<<(FArchive& Ar, FOnDemandTocSentinel& Sentinel)
+{
+	if (Ar.IsSaving())
+	{	
+		// We could just cast FOnDemandTocSentinel::SentinelImg to a non-const pointer but we can't be 
+		// 100% sure that the FArchive won't change the data, even if it is in Saving mode. Since this 
+		// isn't performance critical we will play it safe.
+		uint8 Output[FOnDemandTocSentinel::SentinelSize];
+		FMemory::Memcpy(Output, FOnDemandTocSentinel::SentinelImg, FOnDemandTocSentinel::SentinelSize);
+
+		Ar.Serialize(&Output, FOnDemandTocSentinel::SentinelSize);
+	}
+	else
+	{
+		Ar.Serialize(&Sentinel.Data, FOnDemandTocSentinel::SentinelSize);
+	}
+
+	return Ar;
+}
+
+////////////////////////////////////////////////////////////////////////////////
+FArchive& operator<<(FArchive& Ar, FOnDemandTocAdditionalFile& AdditionalFile)
+{
+	Ar << AdditionalFile.Hash;
+	Ar << AdditionalFile.Filename;
+	Ar << AdditionalFile.FileSize;
+	return Ar;
+}
+
+FCbWriter& operator<<(FCbWriter& Writer, const FOnDemandTocAdditionalFile& AdditionalFile)
+{
+	Writer.BeginObject();
+	Writer.AddHash(UTF8TEXTVIEW("Hash"), AdditionalFile.Hash);
+	Writer.AddString(UTF8TEXTVIEW("Filename"), AdditionalFile.Filename);
+	Writer.AddInteger(UTF8TEXTVIEW("Filename"), AdditionalFile.FileSize);
+	Writer.EndObject();
+
+	return Writer;
+}
+
+bool LoadFromCompactBinary(FCbFieldView Field, FOnDemandTocAdditionalFile& AdditionalFile)
+{
+	if (FCbObjectView Obj = Field.AsObjectView())
+	{
+		AdditionalFile.Hash = Obj["Hash"].AsHash();
+		AdditionalFile.Filename = FString(Obj["Filename"].AsString());
+		AdditionalFile.FileSize = Obj["FileSize"].AsUInt64();
+		return true;
+	}
+
+	return false;
+}
+
+////////////////////////////////////////////////////////////////////////////////
+FArchive& operator<<(FArchive& Ar, FOnDemandTocTagSetPackageList& TagSet)
+{
+	Ar << TagSet.ContainerIndex;
+	Ar << TagSet.PackageIndicies;
+	return Ar;
+}
+
+FCbWriter& operator<<(FCbWriter& Writer, const FOnDemandTocTagSetPackageList& TagSet)
+{
+	Writer.BeginObject();
+	Writer.AddInteger(UTF8TEXTVIEW("ContainerIndex"), TagSet.ContainerIndex);
+	Writer.BeginArray(UTF8TEXTVIEW("PackageIndicies"));
+	for (const uint32 Index : TagSet.PackageIndicies)
+	{
+		Writer << Index;
+	}
+	Writer.EndArray();
+	Writer.EndObject();
+	return Writer;
+}
+
+bool LoadFromCompactBinary(FCbFieldView Field, FOnDemandTocTagSetPackageList& TagSet)
+{
+	if (FCbObjectView Obj = Field.AsObjectView())
+	{
+		FCbFieldView ContainerIndex = Obj["ContainerIndex"];
+		TagSet.ContainerIndex = ContainerIndex.AsUInt32();
+		if (ContainerIndex.HasError())
+		{
+			return false;
+		}
+
+		FCbFieldView PackageIndicies = Obj["PackageIndicies"];
+		FCbArrayView PackageIndiciesArray = PackageIndicies.AsArrayView();
+		if(PackageIndicies.HasError())
+		{
+			return false;
+		}
+
+		TagSet.PackageIndicies.Reserve(int32(PackageIndiciesArray.Num()));
+		for (FCbFieldView ArrayField : PackageIndiciesArray)
+		{
+			uint32 Index = ArrayField.AsUInt32();
+			if (ArrayField.HasError())
+			{
+				return false;
+			}
+			TagSet.PackageIndicies.Emplace(Index);
+		}
+
+		return true;
+	}
+
+	return false;
+}
+
+////////////////////////////////////////////////////////////////////////////////
+FArchive& operator<<(FArchive& Ar, FOnDemandTocTagSet& TagSet)
+{
+	Ar << TagSet.Tag;
+	Ar << TagSet.Packages;
+	return Ar;
+}
+
+FCbWriter& operator<<(FCbWriter& Writer, const FOnDemandTocTagSet& TagSet)
+{
+	Writer.BeginObject();
+	Writer.AddString(UTF8TEXTVIEW("Tag"), TagSet.Tag);
+	Writer.BeginArray(UTF8TEXTVIEW("Packages"));
+	for (const FOnDemandTocTagSetPackageList& PackageList : TagSet.Packages)
+	{
+		Writer << PackageList;
+	}
+	Writer.EndArray();
+	Writer.EndObject();
+	return Writer;
+}
+
+bool LoadFromCompactBinary(FCbFieldView Field, FOnDemandTocTagSet& TagSet)
+{
+	if (FCbObjectView Obj = Field.AsObjectView())
+	{
+		TagSet.Tag = FString(Obj["Tag"].AsString());
+		FCbArrayView Packages = Obj["Packages"].AsArrayView();
+		TagSet.Packages.Reserve(int32(Packages.Num()));
+		for (FCbFieldView ArrayField : Packages)
+		{
+			if (!LoadFromCompactBinary(ArrayField, TagSet.Packages.Emplace_GetRef()))
+			{
+				return false;
+			}
+		}
+		return true;
+	}
+
+	return false;
+}
+
+////////////////////////////////////////////////////////////////////////////////
+FArchive& operator<<(FArchive& Ar, FOnDemandToc& Toc)
+{
+	Ar << Toc.Header;
+	if (Ar.IsError())
+	{
+		return Ar;
+	}
+
+	Ar.SetCustomVersion(Toc.VersionGuid, int32(Toc.Header.Version), TEXT("OnDemandToc"));
+
+	if (uint32(Toc.Header.Version) >= uint32(EOnDemandTocVersion::Meta))
+	{
+		Ar << Toc.Meta;
+	}
+	Ar << Toc.Containers;
+
+	if (uint32(Toc.Header.Version) >= uint32(EOnDemandTocVersion::AdditionalFiles))
+	{
+		Ar << Toc.AdditionalFiles;
+	}
+
+	if (uint32(Toc.Header.Version) >= uint32(EOnDemandTocVersion::TagSets))
+	{
+		Ar << Toc.TagSets;
+	}
+
+	return Ar;
+}
+
+FCbWriter& operator<<(FCbWriter& Writer, const FOnDemandToc& Toc)
+{
+	Writer.BeginObject();
+	Writer << UTF8TEXTVIEW("Header") << Toc.Header;
+
+	Writer.BeginArray(UTF8TEXTVIEW("Containers"));
+	for (const FOnDemandTocContainerEntry& Container : Toc.Containers)
+	{
+		Writer << Container;
+	}
+	Writer.EndArray();
+
+	if (Toc.AdditionalFiles.Num() > 0)
+	{
+		Writer.BeginArray(UTF8TEXTVIEW("Files"));
+		for (const FOnDemandTocAdditionalFile& File : Toc.AdditionalFiles)
+		{
+			Writer << File;
+		}
+		Writer.EndArray();
+	}
+
+	if (Toc.TagSets.Num() > 0)
+	{
+		Writer.BeginArray(UTF8TEXTVIEW("TagSets"));
+		for (const FOnDemandTocTagSet& TagSet : Toc.TagSets)
+		{
+			Writer << TagSet;
+		}
+		Writer.EndArray();
+	}
+
+	Writer.EndObject();
+	
+	return Writer;
+}
+
+FGuid FOnDemandToc::VersionGuid = FGuid("C43DD98F353F499D9A0767F6EA0155EB");
+
+bool LoadFromCompactBinary(FCbFieldView Field, FOnDemandToc& OutToc)
+{
+	if (FCbObjectView Obj = Field.AsObjectView())
+	{
+		if (!LoadFromCompactBinary(Obj["Header"], OutToc.Header))
+		{
+			return false;
+		}
+
+		if (uint32(OutToc.Header.Version) >= uint32(EOnDemandTocVersion::Meta))
+		{
+			if (!LoadFromCompactBinary(Obj["Meta"], OutToc.Meta))
+			{
+				return false;
+			}
+		}
+
+		FCbArrayView Containers = Obj["Containers"].AsArrayView();
+		OutToc.Containers.Reserve(int32(Containers.Num()));
+		for (FCbFieldView ArrayField : Containers)
+		{
+			if (!LoadFromCompactBinary(ArrayField, OutToc.Containers.AddDefaulted_GetRef()))
+			{
+				return false;
+			}
+		}
+
+		if (uint32(OutToc.Header.Version) >= uint32(EOnDemandTocVersion::AdditionalFiles))
+		{
+			FCbArrayView Files = Obj["Files"].AsArrayView();
+			OutToc.AdditionalFiles.Reserve(int32(Files.Num()));
+			for (FCbFieldView ArrayField : Files)
+			{
+				if (!LoadFromCompactBinary(ArrayField, OutToc.AdditionalFiles.AddDefaulted_GetRef()))
+				{
+					return false;
+				}
+			}
+		}
+
+		if (uint32(OutToc.Header.Version) >= uint32(EOnDemandTocVersion::TagSets))
+		{
+			FCbArrayView TagSets = Obj["TagSets"].AsArrayView();
+			OutToc.TagSets.Reserve(int32(TagSets.Num()));
+			for (FCbFieldView ArrayField : TagSets)
+			{
+				if (!LoadFromCompactBinary(ArrayField, OutToc.TagSets.AddDefaulted_GetRef()))
+				{
+					return false;
+				}
+			}
+		}
+
+		return true;
+	}
+
+	return false;
+}
+
+TIoStatusOr<FOnDemandToc> FOnDemandToc::LoadFromFile(const FString& FilePath, bool bValidate)
+{
+	TUniquePtr<FArchive> Ar;
+	if (FPathViews::IsRelativePath(FilePath) && FPlatformMisc::FileExistsInPlatformPackage(FilePath))
+	{
+		Ar = CreateReaderFromPlatformPackage(FilePath);
+	}
+	
+	if (Ar.IsValid() == false)
+	{
+		Ar.Reset(IFileManager::Get().CreateFileReader(*FilePath));
+	}
+
+	if (Ar.IsValid() == false)
+	{
+		FIoStatus Status = FIoStatusBuilder(EIoErrorCode::FileNotOpen) << TEXT("Failed to open '") << FilePath << TEXT("'");
+		return Status;
+	}
+
+	if (bValidate)
+	{
+		const int64 SentinelPos = Ar->TotalSize() - FOnDemandTocSentinel::SentinelSize;
+
+		if (SentinelPos < 0)
+		{
+			FIoStatus Status = FIoStatusBuilder(EIoErrorCode::CorruptToc) << TEXT("Unexpected file size");
+			return Status;
+		}
+
+		Ar->Seek(SentinelPos);
+
+		FOnDemandTocSentinel Sentinel;
+		*Ar << Sentinel;
+
+		if (!Sentinel.IsValid())
+		{
+			return FIoStatus(EIoErrorCode::CorruptToc);
+		}
+
+		Ar->Seek(0);
+	}
+
+	FOnDemandToc Toc;
+	*Ar << Toc;
+
+	if (Ar->IsError() || Ar->IsCriticalError())
+	{
+		FIoStatus Status = FIoStatusBuilder(EIoErrorCode::FileNotOpen) << TEXT("Failed to serialize TOC file");
+		return Status;
+	}
+
+	return Toc; 
+}
+
+} // namespace UE::IoStore
diff --git a/Engine/Source/Runtime/Experimental/IoStore/OnDemandCore/Public/IO/IoStoreOnDemand.h b/Engine/Source/Runtime/Experimental/IoStore/OnDemandCore/Public/IO/IoStoreOnDemand.h
new file mode 100644
--- /dev/null
+++ b/Engine/Source/Runtime/Experimental/IoStore/OnDemandCore/Public/IO/IoStoreOnDemand.h
@@ -0,0 +1,491 @@
+// Copyright Epic Games, Inc. All Rights Reserved.
+
+#pragma once
+
+#include "Containers/Map.h"
+#include "Containers/UnrealString.h"
+#include "Containers/StringFwd.h"
+#include "Containers/SharedString.h"
+#include "IO/OnDemandHostGroup.h"
+#include "IO/OnDemandToc.h"
+#include "IO/PackageId.h"
+#include "Misc/EnumClassFlags.h"
+#include "Features/IModularFeature.h"
+#include "Modules/ModuleInterface.h"
+#include "Templates/SharedPointer.h"
+#include "UObject/NameTypes.h"
+
+#define UE_API IOSTOREONDEMANDCORE_API
+
+struct FAnalyticsEventAttribute;
+
+// Custom initialization allows users to control when
+// the system should be initialized.
+#if !defined(UE_IAS_CUSTOM_INITIALIZATION)
+	#define UE_IAS_CUSTOM_INITIALIZATION 0
+#endif
+
+UE_API DECLARE_LOG_CATEGORY_EXTERN(LogIoStoreOnDemand, Log, All);
+UE_API DECLARE_LOG_CATEGORY_EXTERN(LogIas, Log, All);
+
+namespace UE::IoStore
+{
+
+using FSharedOnDemandIoStore		= TSharedPtr<class IOnDemandIoStore, ESPMode::ThreadSafe>;
+using FWeakOnDemandIoStore			= TWeakPtr<class IOnDemandIoStore, ESPMode::ThreadSafe>;
+using FSharedInternalContentHandle	= TSharedPtr<class FOnDemandInternalContentHandle, ESPMode::ThreadSafe>;
+using FSharedInternalInstallRequest	= TSharedPtr<class FOnDemandInternalInstallRequest, ESPMode::ThreadSafe>;
+
+/** Interface for on-demand request types. */
+class FOnDemandRequest
+{
+public:
+	/** Request status. */
+	enum EStatus : uint8
+	{
+		/** The request has not been issued. */
+		None,
+		/** The request is pending. */
+		Pending,
+		/** The request completed successfully. */
+		Ok,
+		/** The request was cancelled. */
+		Cancelled,
+		/** The request failed. */
+		Error
+	};
+	/** Destructor. */
+	virtual ~FOnDemandRequest() = default;
+	/** Returns the current status. */
+	virtual EStatus GetStatus() const = 0;
+	/** Returns true if the request is invalid. */
+	bool IsNone() const { return EStatus::None == GetStatus(); }
+	/** Returns true if the request is pending. */
+	bool IsPending() const { return EStatus::Pending == GetStatus(); }
+	/** Returns true if the request was successful. */
+	bool IsOk() const { return EStatus::Ok == GetStatus(); }
+	/** Returns true if the request was cancelled. */
+	bool IsCancelled() const { return EStatus::Cancelled == GetStatus(); }
+	/** Returns true if the request was unsuccessful. */
+	bool IsError() const { return EStatus::Error == GetStatus(); }
+	/** Returns true if the request is completed. */
+	bool IsCompleted() const { return static_cast<uint8>(EStatus::Pending) < static_cast<uint8>(GetStatus()); }
+
+protected:
+	FOnDemandRequest() = default;
+};
+
+/**
+ * Keeps referenced data pinned in the cache until released.
+ */
+class FOnDemandContentHandle
+{
+public:
+	/** Creates a new invalid content handle. */ 
+	UE_API FOnDemandContentHandle();
+	/** Destroy the handle and release any referenced content. */ 
+	UE_API ~FOnDemandContentHandle();
+	/** Destroy the handle and release any referenced content. */ 
+	void Reset() { Handle.Reset(); }
+	/** Returns whether the handle is valid. */
+	bool IsValid() const { return Handle.IsValid(); }
+	/** Create a new content handle .*/
+	UE_API static FOnDemandContentHandle Create();
+	/** Create a new content handle with a debug name. */
+	UE_API static FOnDemandContentHandle Create(FSharedString DebugName);
+	/** Create a new content handle with a debug name. */
+	UE_API static FOnDemandContentHandle Create(FStringView DebugName);
+	/** Returns a string representing the content handle. */
+	UE_API friend FString LexToString(const FOnDemandContentHandle& Handle);
+
+	bool operator==(FOnDemandContentHandle& Other) const
+	{
+		return Handle == Other.Handle;
+	}
+
+private:
+	friend class FOnDemandIoStore;
+	FSharedInternalContentHandle Handle;
+};
+
+/** Options for controlling the behavior of mounted container(s). */
+enum class EOnDemandMountOptions
+{
+	/** Mount containers with the purpose of streaming the content on-demand. */
+	StreamOnDemand			= 1 << 0,
+	/** Mount containers with the purpose of installing/downloading the content on-demand. */
+	InstallOnDemand			= 1 << 1,
+	/** Trigger callback on game thread. */
+	CallbackOnGameThread	= 1 << 2,
+};
+ENUM_CLASS_FLAGS(EOnDemandMountOptions);
+
+/** Arguments for mounting on-demand container TOC(s).
+ *
+ * On-demand content can be mounted by providing:
+ * 1. a serialized TOC
+ * 2. a filepath to a TOC on disk
+ * 3. a URL from where to fetch the TOC using HTTP
+ *
+ * The chunk URLs are derived from the provided TOC URL or from the
+ * serialized chunks directory property in the TOC, i.e. if the
+ * TocRelativeUrl is not specified the TOC ChunksDirectory property
+ * needs to form a qualified path from the host.
+ * Example:
+ * http(s)://<Host>/<TocRelativePath>/chunks/<1-Byte Hex>/<hash>.iochunk
+ * http(s)://<Host>/<TOC.ChunksDirectory>/chunks/<1-Byte Hex>/<hash>.iochunk
+ */
+struct FOnDemandMountArgs
+{
+	/** Mount an already serialized TOC. */
+	TUniquePtr<FOnDemandToc> Toc;
+	/** Mandatory ID to be used for unmounting all container file(s) included in the TOC. */
+	FString MountId;
+	/** Relative URL from the primary endpoint from where to download the TOC. */
+	FString TocRelativeUrl;
+	/** Serialize the TOC from the specified file path. */
+	FString FilePath;
+	/** List of URLs from where to download the chunks. */
+	FOnDemandHostGroup HostGroup;
+	/** Mount options. */
+	EOnDemandMountOptions Options = EOnDemandMountOptions::StreamOnDemand;
+};
+
+/** Holds information about a mount request. */
+struct FOnDemandMountResult
+{
+	/** The mount ID used for mounting the container(s). */
+	FString MountId;
+	/** The status of the mount request. */
+	FIoStatus Status;
+	/** Duration in seconds. */
+	double DurationInSeconds = 0.0;
+};
+
+/** Mount completion callback. */
+using FOnDemandMountCompleted = TUniqueFunction<void(FOnDemandMountResult)>;
+
+/** Options for controlling the behavior of the install request. */
+enum class EOnDemandInstallOptions
+{
+	/** No additional options. */
+	None				 = 0,
+	/** Trigger callback on game thread. */
+	CallbackOnGameThread = 1 << 0,
+	/** Follow soft references when gathering packages to install */
+	InstallSoftReferences = 1 << 1,
+};
+ENUM_CLASS_FLAGS(EOnDemandInstallOptions);
+
+/** Arguments for installing/downloading on-demand content. */
+struct FOnDemandInstallArgs 
+{
+	/** Install all content from containers matching this mount ID. */
+	FString MountId;
+	/** Install content matching a set of tag(s) and optionally the mount ID. */
+	TArray<FString> TagSets;
+	/** Package ID's to install. */
+	TArray<FPackageId> PackageIds;
+	/** Content handle. */
+	FOnDemandContentHandle ContentHandle;
+	/** Install options. */
+	EOnDemandInstallOptions Options = EOnDemandInstallOptions::None;
+	/** Priority. */
+	int32 Priority = 0;
+	/** Optional debug name . */
+	FSharedString DebugName;
+};
+
+/** Holds information about progress for an install request. */
+struct FOnDemandInstallProgress
+{
+	/** The total size of the requested content. */
+	uint64 TotalContentSize = 0;
+	/** The total size to be installed/downloaded (<= TotalContentSize). */
+	uint64 TotalInstallSize = 0;
+	/** The size currently installed/downloaded (<= TotalInstallSize). */
+	uint64 CurrentInstallSize = 0;
+
+	FOnDemandInstallProgress& Combine(const FOnDemandInstallProgress& Other)
+	{
+		TotalContentSize += Other.TotalContentSize;
+		TotalInstallSize += Other.TotalInstallSize;
+		CurrentInstallSize += Other.CurrentInstallSize;
+		return *this;
+	}
+
+	uint64 GetTotalDownloadSize() const
+	{
+		return TotalInstallSize;
+	}
+
+	uint64 GetAlreadyDownloadedSize() const
+	{
+		return CurrentInstallSize;
+	}
+
+	float GetRelativeProgress() const 
+	{
+		const double Progress = (TotalInstallSize > 0) ?
+			FMath::Clamp(double(CurrentInstallSize) / double(TotalInstallSize), 0.0, 1.0) :
+			0;
+
+		return float(Progress);
+	}
+
+	float GetAbsoluteProgress() const
+	{
+		const double Progress = (TotalContentSize > 0) ?
+			FMath::Clamp(double(GetCachedSize()) / double(TotalContentSize), 0.0, 1.0) :
+			0;
+
+		return float(Progress);
+	}
+
+	uint64 GetCachedSize() const 
+	{
+		return TotalContentSize - TotalInstallSize + CurrentInstallSize;
+	}
+
+	uint64 GetTotalSize() const
+	{
+		return TotalContentSize;
+	}
+};
+
+/** Install Progress callback. */
+using FOnDemandInstallProgressed = TFunction<void(FOnDemandInstallProgress)>;
+
+/** Holds information about an install request. */
+struct FOnDemandInstallResult
+{
+	/** The status of the install request. */
+	FIoStatus Status;
+	/** Duration in seconds. */
+	double DurationInSeconds = 0.0;
+	/** Final progress for the install request. */
+	FOnDemandInstallProgress Progress;
+};
+
+/** Install completion callback. */
+using FOnDemandInstallCompleted = TUniqueFunction<void(FOnDemandInstallResult)>;
+
+/** A single-ownership handle to an install request. */
+class FOnDemandInstallRequest final
+	: public FOnDemandRequest
+{
+	friend class FOnDemandIoStore;
+public:
+	/** Creates an invalid install request. */
+	UE_API FOnDemandInstallRequest();
+	/** Move constructable. */;
+	UE_API FOnDemandInstallRequest(FOnDemandInstallRequest&&);
+	/** Move assignable. */
+	UE_API FOnDemandInstallRequest& operator=(FOnDemandInstallRequest&&);
+	/** Destructor. */
+	UE_API ~FOnDemandInstallRequest();
+	/** Returns the current request status. */
+	UE_API virtual EStatus GetStatus() const override;
+	/** Cancel the install request. */
+	UE_API void Cancel();
+	/** Update priority of the instlal request. */
+	UE_API void UpdatePriority(int32 NewPriority);
+
+private:
+	UE_API FOnDemandInstallRequest(FWeakOnDemandIoStore IoStore, FSharedInternalInstallRequest InternalRequest);
+	FOnDemandInstallRequest(const FOnDemandInstallRequest&) = delete;
+	FOnDemandInstallRequest& operator=(const FOnDemandInstallRequest&) = delete;
+
+	FWeakOnDemandIoStore IoStore;	
+	FSharedInternalInstallRequest Request;
+};
+
+/** Arguments for purging on-demand content. */
+struct FOnDemandPurgeArgs
+{
+	/** Purge options. */
+	EOnDemandInstallOptions Options = EOnDemandInstallOptions::None;
+};
+
+/** Holds information about a purge request */
+struct FOnDemandPurgeResult
+{
+	/** The status of the purge request. */
+	FIoStatus Status;
+	/** Duration in seconds. */
+	double DurationInSeconds = 0.0;
+};
+
+/** Purge completion callback */
+using FOnDemandPurgeCompleted = TUniqueFunction<void(FOnDemandPurgeResult)>;
+
+/** Options for controlling the behavior of the install size request. */
+enum class EOnDemandGetInstallSizeOptions
+{
+	/** No additional options. */
+	None = 0,
+	/** Follow soft references when gathering packages to install */
+	IncludeSoftReferences = 1 << 0,
+};
+ENUM_CLASS_FLAGS(EOnDemandGetInstallSizeOptions);
+
+/** Arguments for getting the size of on-demand content. */
+struct FOnDemandGetInstallSizeArgs 
+{
+	/** Get total install size for containers matching this mount ID. */
+	FString MountId;
+	/** Get total install size for the specified tag(s) and optionally matching the mount ID. */
+	TArray<FString> TagSets;
+	/** Get total intall size for the specified package IDs. */
+	TArray<FPackageId> PackageIds;
+	/** Options */
+	EOnDemandGetInstallSizeOptions Options = EOnDemandGetInstallSizeOptions::None;
+};
+
+/** Holds information about the install cache usage */
+struct FOnDemandInstallCacheUsage
+{
+	uint64 MaxSize = 0;
+	uint64 TotalSize = 0;
+	uint64 ReferencedBlockSize = 0;
+
+	UE_API friend FStringBuilderBase& operator<<(FStringBuilderBase& Sb, const FOnDemandInstallCacheUsage& CacheUsage);
+};
+
+/** Holds information about the stremaing cache usage */
+struct FOnDemandStreamingCacheUsage
+{
+	uint64 MaxSize = 0;
+	uint64 TotalSize = 0;
+
+	UE_API friend FStringBuilderBase& operator<<(FStringBuilderBase& Sb, const FOnDemandStreamingCacheUsage& CacheUsage);
+};
+
+/** Holds information about install and streaming cache usage */
+struct FOnDemandCacheUsage
+{
+	FOnDemandInstallCacheUsage InstallCache;
+	FOnDemandStreamingCacheUsage StreamingCache;
+};
+
+/** Interface for recording analytics over a given time period */
+class IAnalyticsRecording
+{
+public:
+	IAnalyticsRecording() = default;
+	virtual ~IAnalyticsRecording() = default;
+
+	/** Writes the current value of the analytics to the output array */
+	virtual void Report(TArray<FAnalyticsEventAttribute>& OutAnalyticsArray) const = 0;
+	
+	/** Calling this will stop recording. Any calls to ::Report after this point will return the same data */
+	virtual void StopRecording() = 0;
+};
+
+/** Streaming options (IAS). */
+enum class EOnDemandStreamingOptions
+{
+	/** No special options enabled. */
+	Default						= 0,
+	/** Optional bulk data disabled. */
+	OptionalBulkDataDisabled	= (1 << 0)
+};
+ENUM_CLASS_FLAGS(EOnDemandStreamingOptions);
+
+/** Interface for installing and streaming content on-demand. */
+class IOnDemandIoStore
+{
+public:
+	/** Destroy the I/O store. */
+	virtual ~IOnDemandIoStore() = default;
+	/** Initialize the I/O store. Called after the module feature has been registered. */
+	virtual FIoStatus Initialize() = 0;
+	/** Initialize the I/O store. Called after the module feature has been registered. */
+	virtual FIoStatus InitializePostHotfix() = 0;
+	/**
+	 * Mount an on-demand container.
+	 * @param	Args		Mount arguments.
+	 * @param	OnCompleted	Completion callback.
+	 **/
+	virtual void Mount(FOnDemandMountArgs&& Args, FOnDemandMountCompleted&& OnCompleted) = 0;
+	/**
+	 * Unmount all container(s) associated with the specified mount ID.
+	 * @param	MountId	The Mount ID passed to the mount call.
+	 **/
+	virtual FIoStatus Unmount(FStringView MountId) = 0;
+	/**
+	 * Install content.
+	 * @param	Args				The install arguments.
+	 * @param	OnCompleted			Completion callback.
+	 * @param	OnProgress			Optional progress callback.
+	 * @param	CancellationToken	Optional cancellation token.
+	 **/
+	virtual FOnDemandInstallRequest Install(
+			FOnDemandInstallArgs&& Args,
+			FOnDemandInstallCompleted&& OnCompleted,
+			FOnDemandInstallProgressed&& OnProgress = nullptr) = 0;
+	/**
+	 * Purge the cache.
+	 * @param	Args		The purge arguments.
+	 * @param	OnComplete	Completion callback.
+	 **/
+	virtual void Purge(FOnDemandPurgeArgs&& Args, FOnDemandPurgeCompleted&& OnCompleted) = 0;
+	/**
+	 * Get the total space in bytes needed to install the specified content.
+	 * @param Args	Arguments for computing required install size.
+	 **/
+	virtual TIoStatusOr<uint64> GetInstallSize(const FOnDemandGetInstallSizeArgs& Args) const = 0;
+	/**
+	 * Get the total space in bytes needed to install the specified content.
+	 * @param Args	Arguments for computing required install size.
+	 * @return		The total size by mount ID.
+	 **/
+	virtual FIoStatus GetInstallSizesByMountId(const FOnDemandGetInstallSizeArgs& Args, TMap<FString, uint64>& OutSizesByMountId) const = 0;
+	/** Returns the total cache size in bytes. */
+	virtual FOnDemandCacheUsage GetCacheUsage() const = 0;
+	/** Return if the OnDemand streaming system is enabled. */
+	virtual bool IsOnDemandStreamingEnabled() const = 0;
+	/** Set streaming options.*/
+	virtual void SetStreamingOptions(EOnDemandStreamingOptions Options) = 0;
+	/**
+	 * Reports the statistics for the current OnDemandBackend. It will record any events that have occurred from the last time that this was called.
+	 * This is a legacy method, prefer using StartAnalyticsRecording instead.
+	 */
+	virtual void ReportAnalytics(TArray<FAnalyticsEventAttribute>& OutAnalyticsArray) const = 0;
+	/** 
+	 * Create a new analytics interface for the current OnDemandBackend, can be used to record statistics for a set period of time.
+	 * This can return a nullptr if the analytics system is not enabled.
+	 */
+	virtual TUniquePtr<IAnalyticsRecording> StartAnalyticsRecording() const = 0;
+
+private:
+	/** Cancel the install request. */
+	virtual void CancelInstallRequest(FSharedInternalInstallRequest InstallRequest) = 0;
+	/** Update the priority for the install request. */
+	virtual void UpdateInstallRequestPriority(FSharedInternalInstallRequest InstallRequest, int32 NewPriority) = 0;
+	/** Release all content referenced by the content handle. */
+	virtual void ReleaseContent(class FOnDemandInternalContentHandle& ContentHandle) = 0;
+
+	friend class FOnDemandInternalContentHandle;
+	friend class FOnDemandInstallRequest;
+};
+
+/** Module feature for creating a concrete implementation of the on-demand I/O store. */
+class IOnDemandIoStoreFactory
+	: public IModularFeature
+{
+public:
+	/** Feature name. */
+	UE_API static FName FeatureName;
+	/** Create a new instance of the I/O store.  Called once the feature has been registered. */
+	virtual FSharedOnDemandIoStore CreateInstance() = 0;
+};
+
+/** Returns the on-demand I/O store if available. */
+UE_API IOnDemandIoStore* TryGetOnDemandIoStore();
+/** Returns the on-demand I/O store. */
+UE_API IOnDemandIoStore& GetOnDemandIoStore();
+
+} // namespace UE::IoStore
+
+#undef UE_API
diff --git a/Engine/Source/Runtime/Experimental/IoStore/OnDemandCore/Public/IO/OnDemandHostGroup.h b/Engine/Source/Runtime/Experimental/IoStore/OnDemandCore/Public/IO/OnDemandHostGroup.h
new file mode 100644
--- /dev/null
+++ b/Engine/Source/Runtime/Experimental/IoStore/OnDemandCore/Public/IO/OnDemandHostGroup.h
@@ -0,0 +1,69 @@
+// Copyright Epic Games, Inc. All Rights Reserved.
+
+#pragma once
+
+#include "CoreTypes.h"
+#include "Containers/AnsiString.h"
+#include "Containers/StringFwd.h"
+#include "IO/IoStatus.h"
+#include "Templates/SharedPointer.h"
+
+#define UE_API IOSTOREONDEMANDCORE_API
+
+namespace UE::IoStore
+{
+
+/** Holds a set of host URLs. */
+class FOnDemandHostGroup
+{
+	struct FImpl;
+	using FSharedImpl = TSharedPtr<FImpl, ESPMode::ThreadSafe>;
+
+public:
+	/** Creates a new empty host group. */
+	UE_API FOnDemandHostGroup();
+	/** Destructor. */ 
+	UE_API ~FOnDemandHostGroup();
+	/** Returns the list of available host URLs. */ 
+	UE_API TConstArrayView<FAnsiString> Hosts() const;
+	/** Get the URL at the specified index. */
+	UE_API FAnsiStringView Host(int32 Index) const;
+	/** Get the next available host starting from the specified index. */
+	UE_API FAnsiStringView CycleHost(int32& InOutIndex) const;
+	/** Set the primary host URL. */
+	UE_API void SetPrimaryHost(int32 Index);
+	/** Get the primary host URL. */
+	UE_API FAnsiStringView PrimaryHost() const;
+	/** Get the primary host index. */
+	UE_API int32 PrimaryHostIndex() const;
+	/** Returns whether the group is empty or not. */
+	UE_API bool IsEmpty() const;
+	/** Create a new host group with the specified URL. */
+	UE_API static TIoStatusOr<FOnDemandHostGroup> Create(FAnsiStringView Url);
+	/** Create a new host group with the specified URL. */
+	UE_API static TIoStatusOr<FOnDemandHostGroup> Create(FStringView Url);
+	/** Create a new host group with the specified URLs. */
+	UE_API static TIoStatusOr<FOnDemandHostGroup> Create(TConstArrayView<FAnsiString> Urls);
+	/** Create a new host group with the specified URLs. */
+	UE_API static TIoStatusOr<FOnDemandHostGroup> Create(TConstArrayView<FString> Urls);
+
+	/** Returns whether two endpoint instance are the same or not. */
+	friend inline bool operator==(const FOnDemandHostGroup& Lhs, const FOnDemandHostGroup& Rhs)
+	{
+		return Lhs.Impl.Get() == Rhs.Impl.Get();
+	}
+	/** Returns the type hash of the endpoint. */ 
+	friend inline uint32 GetTypeHash(const FOnDemandHostGroup& Endpoint)
+	{
+		return *reinterpret_cast<const uint32*>(Endpoint.Impl.Get());
+	}
+
+private:
+	FOnDemandHostGroup(FSharedImpl&& Impl);
+
+	FSharedImpl Impl;
+};
+
+} // namespace UE::IoStore
+
+#undef UE_API
diff --git a/Engine/Source/Runtime/Experimental/IoStore/OnDemandCore/Public/IO/OnDemandToc.h b/Engine/Source/Runtime/Experimental/IoStore/OnDemandCore/Public/IO/OnDemandToc.h
new file mode 100644
--- /dev/null
+++ b/Engine/Source/Runtime/Experimental/IoStore/OnDemandCore/Public/IO/OnDemandToc.h
@@ -0,0 +1,196 @@
+// Copyright Epic Games, Inc. All Rights Reserved.
+
+#pragma once
+
+#include "IO/IoChunkId.h"
+#include "IO/IoContainerId.h"
+#include "IO/IoHash.h"
+#include "IO/IoStatus.h"
+
+#define UE_API IOSTOREONDEMANDCORE_API
+
+using FIoBlockHash = uint32;
+
+class FArchive;
+class FCbWriter;
+
+namespace UE::IoStore
+{
+
+///////////////////////////////////////////////////////////////////////////////
+enum class EOnDemandTocVersion : uint32
+{
+	Invalid			= 0,
+	Initial			= 1,
+	UTocHash		= 2,
+	BlockHash32		= 3,
+	NoRawHash		= 4,
+	Meta			= 5,
+	ContainerId		= 6,
+	AdditionalFiles	= 7,
+	TagSets			= 8,
+	ContainerFlags	= 9,
+
+	LatestPlusOne,
+	Latest			= (LatestPlusOne - 1)
+};
+
+///////////////////////////////////////////////////////////////////////////////
+enum class EOnDemandChunkVersion : uint32
+{
+	Invalid			= 0,
+	Initial			= 1,
+
+	LatestPlusOne,
+	Latest			= (LatestPlusOne - 1)
+};
+
+///////////////////////////////////////////////////////////////////////////////
+struct FTocMeta
+{
+	int64	EpochTimestamp = 0;
+	FString BuildVersion;
+	FString TargetPlatform;
+
+	UE_API friend FArchive& operator<<(FArchive& Ar, FTocMeta& Meta);
+	UE_API friend FCbWriter& operator<<(FCbWriter& Writer, const FTocMeta& Meta);
+};
+
+UE_API bool LoadFromCompactBinary(FCbFieldView Field, FTocMeta& OutMeta);
+
+///////////////////////////////////////////////////////////////////////////////
+struct FOnDemandTocHeader
+{
+	static constexpr uint64 ExpectedMagic = 0x6f6e64656d616e64; // ondemand
+
+	uint64	Magic = ExpectedMagic;
+	uint32	Version = uint32(EOnDemandTocVersion::Latest);
+	uint32	ChunkVersion = uint32(EOnDemandChunkVersion::Latest);
+	uint32	BlockSize = 0;
+	FString CompressionFormat;
+	FString ChunksDirectory;
+	
+	UE_API friend FArchive& operator<<(FArchive& Ar, FOnDemandTocHeader& Header);
+	UE_API friend FCbWriter& operator<<(FCbWriter& Writer, const FOnDemandTocHeader& Header);
+};
+
+UE_API bool LoadFromCompactBinary(FCbFieldView Field, FOnDemandTocHeader& OutTocHeader);
+
+///////////////////////////////////////////////////////////////////////////////
+struct FOnDemandTocEntry
+{
+	FIoHash		Hash = FIoHash::Zero;
+	FIoChunkId	ChunkId = FIoChunkId::InvalidChunkId;
+	uint64		RawSize = 0;
+	uint64 		EncodedSize = 0;
+	uint32 		BlockOffset = ~uint32(0);
+	uint32 		BlockCount = 0; 
+	
+	UE_API friend FArchive& operator<<(FArchive& Ar, FOnDemandTocEntry& Entry);
+	UE_API friend FCbWriter& operator<<(FCbWriter& Writer, const FOnDemandTocEntry& Entry);
+};
+
+UE_API bool LoadFromCompactBinary(FCbFieldView Field, FOnDemandTocEntry& OutTocEntry);
+
+///////////////////////////////////////////////////////////////////////////////
+struct FOnDemandTocContainerEntry
+{
+	FIoContainerId				ContainerId;
+	FString						ContainerName;
+	FString						EncryptionKeyGuid;
+	TArray<FOnDemandTocEntry>	Entries;
+	TArray<uint32>				BlockSizes;
+	TArray<FIoBlockHash>		BlockHashes;
+
+	/** Hash of the .utoc file (on disk) used to generate this data */
+	FIoHash						UTocHash;
+	uint8						ContainerFlags = 0;
+
+	UE_API friend FArchive& operator<<(FArchive& Ar, FOnDemandTocContainerEntry& ContainerEntry);
+	UE_API friend FCbWriter& operator<<(FCbWriter& Writer, const FOnDemandTocContainerEntry& ContainerEntry);
+};
+
+UE_API bool LoadFromCompactBinary(FCbFieldView Field, FOnDemandTocContainerEntry& OutContainer);
+
+///////////////////////////////////////////////////////////////////////////////
+struct FOnDemandTocSentinel
+{
+public:
+	static constexpr inline char SentinelImg[] = "-[]--[]--[]--[]-";
+	static constexpr uint32 SentinelSize = 16;
+
+	bool IsValid();
+
+	UE_API friend FArchive& operator<<(FArchive& Ar, FOnDemandTocSentinel& Sentinel);
+
+private:
+	uint8 Data[SentinelSize] = { 0 };
+};
+
+struct FOnDemandTocAdditionalFile
+{
+	FIoHash	Hash;
+	FString	Filename;
+	uint64	FileSize = 0;
+
+	UE_API friend FArchive& operator<<(FArchive& Ar, FOnDemandTocAdditionalFile& AdditionalFile);
+	UE_API friend FCbWriter& operator<<(FCbWriter& Writer, const FOnDemandTocAdditionalFile& AdditionalFile);
+};
+
+UE_API bool LoadFromCompactBinary(FCbFieldView Field, FOnDemandTocAdditionalFile& AdditionalFile);
+
+///////////////////////////////////////////////////////////////////////////////
+struct FOnDemandTocTagSetPackageList
+{
+	uint32			ContainerIndex = 0;
+	TArray<uint32>	PackageIndicies;
+
+	UE_API friend FArchive& operator<<(FArchive& Ar, FOnDemandTocTagSetPackageList& TagSet);
+	UE_API friend FCbWriter& operator<<(FCbWriter& Writer, const FOnDemandTocTagSetPackageList& TagSet);
+};
+
+UE_API bool LoadFromCompactBinary(FCbFieldView Field, FOnDemandTocTagSetPackageList& TagSet);
+
+///////////////////////////////////////////////////////////////////////////////
+struct FOnDemandTocTagSet
+{
+	using FOnDemandTocTagSetPackageLists = TArray<FOnDemandTocTagSetPackageList>;
+
+	FString							Tag;
+	FOnDemandTocTagSetPackageLists	Packages;
+
+	UE_API friend FArchive& operator<<(FArchive& Ar, FOnDemandTocTagSet& TagSet);
+	UE_API friend FCbWriter& operator<<(FCbWriter& Writer, const FOnDemandTocTagSet& TagSet);
+};
+
+UE_API bool LoadFromCompactBinary(FCbFieldView Field, FOnDemandTocTagSet& TagSet);
+
+///////////////////////////////////////////////////////////////////////////////
+struct FOnDemandToc
+{
+	FOnDemandToc() = default;
+	~FOnDemandToc() = default;
+	FOnDemandToc(FOnDemandToc&&) = default;
+	FOnDemandToc& operator= (FOnDemandToc&&) = default;
+	FOnDemandToc(const FOnDemandToc&) = delete;
+	FOnDemandToc&  operator= (const FOnDemandToc&) = delete;
+
+	FOnDemandTocHeader					Header;
+	FTocMeta							Meta;
+	TArray<FOnDemandTocContainerEntry>	Containers;
+	TArray<FOnDemandTocAdditionalFile>	AdditionalFiles;
+	TArray<FOnDemandTocTagSet>			TagSets;
+
+	UE_API friend FArchive& operator<<(FArchive& Ar, FOnDemandToc& Toc);
+	UE_API friend FCbWriter& operator<<(FCbWriter& Writer, const FOnDemandToc& Toc);
+
+	static FGuid VersionGuid;
+
+	UE_API static TIoStatusOr<FOnDemandToc> LoadFromFile(const FString& FilePath, bool bValidate);
+};
+
+UE_API bool LoadFromCompactBinary(FCbFieldView Field, FOnDemandToc& OutToc);
+
+} // namespace UE::IoStore
+
+#undef UE_API
diff --git a/Engine/Source/Runtime/InstallBundleManager/InstallBundleManager.Build.cs b/Engine/Source/Runtime/InstallBundleManager/InstallBundleManager.Build.cs
--- a/Engine/Source/Runtime/InstallBundleManager/InstallBundleManager.Build.cs
+++ b/Engine/Source/Runtime/InstallBundleManager/InstallBundleManager.Build.cs
@@ -10,7 +10,7 @@ public class InstallBundleManager : ModuleRules
 			new string[] {
 				"ApplicationCore",
 				"Core",
-				"IoStoreOnDemand",
+				"IoStoreOnDemandCore",
 				"Json"
 			}
 		);
diff --git a/Engine/Source/Runtime/Launch/Private/LaunchEngineLoop.cpp b/Engine/Source/Runtime/Launch/Private/LaunchEngineLoop.cpp
--- a/Engine/Source/Runtime/Launch/Private/LaunchEngineLoop.cpp
+++ b/Engine/Source/Runtime/Launch/Private/LaunchEngineLoop.cpp
@@ -96,6 +96,10 @@
 #include "IO/IoDispatcher.h"
 #endif
 
+#if !defined(UE_WITH_IOSTOREONDEMAND)
+#define UE_WITH_IOSTOREONDEMAND 0
+#endif
+
 #if WITH_COREUOBJECT
 	#include "Internationalization/PackageLocalizationManager.h"
 	#include "Misc/PackageName.h"
@@ -5209,6 +5213,10 @@ void FEngineLoop::Exit()
 	UE::DerivedData::IoStore::TearDownIoDispatcher();
 #endif
 #if USE_IO_DISPATCHER
+#if UE_WITH_IOSTOREONDEMAND
+	FModuleManager::Get().UnloadModule("IoStoreOnDemand");
+#endif
+	FModuleManager::Get().UnloadModule("IoStoreOnDemandCore");
 	FIoDispatcher::Shutdown();
 #endif
 
@@ -6522,8 +6530,10 @@ bool FEngineLoop::AppInit( )
 	}
 
 #if USE_IO_DISPATCHER
-	// Initialize on demand I/O dispatcher backend
+	FModuleManager::Get().LoadModule("IoStoreOnDemandCore");
+#if UE_WITH_IOSTOREONDEMAND
 	FModuleManager::Get().LoadModule("IoStoreOnDemand");
+#endif
 #endif
 
 	// Apply config driven presets
