Add missing I/O store on-demand module file.

#rnx

diff --git a/Engine/Source/Runtime/Experimental/IoStore/OnDemand/Private/OnDemandModule.cpp b/Engine/Source/Runtime/Experimental/IoStore/OnDemand/Private/OnDemandModule.cpp
new file mode 100644
--- /dev/null
+++ b/Engine/Source/Runtime/Experimental/IoStore/OnDemand/Private/OnDemandModule.cpp
@@ -0,0 +1,97 @@
+// Copyright Epic Games, Inc. All Rights Reserved.
+
+#include "Features/IModularFeature.h"
+#include "Features/IModularFeatures.h"
+#include "IO/IoStoreOnDemand.h"
+#include "Modules/ModuleInterface.h"
+#include "Modules/ModuleManager.h"
+#include "OnDemandIoStore.h"
+
+#if PLATFORM_WINDOWS
+#	include <Windows/AllowWindowsPlatformTypes.h>
+#		include <winsock2.h>
+#	include <Windows/HideWindowsPlatformTypes.h>
+#endif //PLATFORM_WINDOWS
+
+namespace UE::IoStore
+{
+namespace Private
+{
+////////////////////////////////////////////////////////////////////////////////
+struct FPlatformSocketSystem
+{
+	#if PLATFORM_WINDOWS
+	void Startup()
+	{
+		WSADATA WsaData;
+		Result = WSAStartup(MAKEWORD(2, 2), &WsaData);
+		if (Result != 0)
+		{
+			TCHAR SystemErrorMsg[MAX_SPRINTF] = { 0 };
+			FPlatformMisc::GetSystemErrorMessage(SystemErrorMsg, UE_ARRAY_COUNT(SystemErrorMsg), Result);
+			UE_LOG(LogIoStoreOnDemand, Error, TEXT("WSAStartup failed due to: %s (%d)"), SystemErrorMsg, Result);
+		}
+	}
+
+	void Cleanup()
+	{
+		if (Result != 0)
+		{
+			return;
+		}
+		Result = WSACleanup();
+		if (Result != 0)
+		{
+			const uint32 SystemError = FPlatformMisc::GetLastError();
+			TCHAR SystemErrorMsg[MAX_SPRINTF] = { 0 };
+			FPlatformMisc::GetSystemErrorMessage(SystemErrorMsg, UE_ARRAY_COUNT(SystemErrorMsg), SystemError);
+			UE_LOG(LogIoStoreOnDemand, Error, TEXT("WSACleanup failed due to: %s (%u)"), SystemErrorMsg, SystemError);
+		}
+	}
+	int Result = -1; 
+	#else
+	void Startup()
+	{ }
+	void Cleanup()
+	{ }
+	#endif
+};
+} // UE::IoStore::Private
+////////////////////////////////////////////////////////////////////////////////
+class FOnDemandIoStoreFactory final
+	: public IOnDemandIoStoreFactory
+{
+	virtual FSharedOnDemandIoStore CreateInstance() override
+	{
+		return MakeShared<FOnDemandIoStore>();
+	}
+};
+
+////////////////////////////////////////////////////////////////////////////////
+class FIoStoreOnDemandModule final
+	: public IModuleInterface
+{
+public:
+	virtual void StartupModule() override;
+	virtual void ShutdownModule() override;
+
+private:
+	Private::FPlatformSocketSystem PlatformSocketSystem;
+	TUniquePtr<FOnDemandIoStoreFactory> Factory = MakeUnique<FOnDemandIoStoreFactory>();
+};
+IMPLEMENT_MODULE(UE::IoStore::FIoStoreOnDemandModule, IoStoreOnDemand);
+
+////////////////////////////////////////////////////////////////////////////////
+void FIoStoreOnDemandModule::StartupModule()
+{
+	PlatformSocketSystem.Startup();
+	IModularFeatures::Get().RegisterModularFeature(IOnDemandIoStoreFactory::FeatureName, Factory.Get());
+}
+
+void FIoStoreOnDemandModule::ShutdownModule()
+{
+	IModularFeatures::Get().UnregisterModularFeature(IOnDemandIoStoreFactory::FeatureName, Factory.Get());
+	PlatformSocketSystem.Cleanup();
+}
+
+} // namespace UE::IoStore
