Destroy I/O store on-demand instance when unloading IoStoreOnDemandCore module.

#rnx

diff --git a/Engine/Source/Runtime/Experimental/IoStore/OnDemandCore/Private/IoStoreOnDemand.cpp b/Engine/Source/Runtime/Experimental/IoStore/OnDemandCore/Private/IoStoreOnDemand.cpp
--- a/Engine/Source/Runtime/Experimental/IoStore/OnDemandCore/Private/IoStoreOnDemand.cpp
+++ b/Engine/Source/Runtime/Experimental/IoStore/OnDemandCore/Private/IoStoreOnDemand.cpp
@@ -413,6 +413,7 @@ void FIoStoreOnDemandCoreModule::StartupModule()
 
 void FIoStoreOnDemandCoreModule::ShutdownModule()
 {
+	GOnDemandIoStore.Reset();
 	IModularFeatures::Get().OnModularFeatureRegistered().RemoveAll(this);
 }
 
