Move the IAS http thread code so that it can be used by both IAS and IAD (currently disabled)

#rnx

- Currently the old thread code has been left alone, but it will be disabled in favour of the new thread if UE_ENABLE_REQUEST_THREAD=1. No existing behavior should be changed.
- For now the FChunkRequest code has been moved to a new header so that both the old and new threading systems can use it.
- It is assumed that the new thread will be created before IAS or IAD and destroyed after they have shutdown so it can always be accessed by these systems.
- For now the new thread uses LogIas as it's logging category.
- IAD is untouched and uses it's own code for http transfers.
- A number of cvars had their static keyword removed so both old and new thread could use them.
- All code for creating the default host group is now in a single method call, CreateDefaultHostGroup. If the new thread is used then CreateDefaultHostGroup is called on the task system.
- Code that ran when the old thread was idle has been moved to FOnDemandIoBackend::OnThreadTickIdle which will be called once per loop of the old thread. The backend can register for a callback from the new thread so that this idle work can also be run there too.

diff --git a/Engine/Source/Runtime/Experimental/IoStore/OnDemand/IoStoreOnDemand.Build.cs b/Engine/Source/Runtime/Experimental/IoStore/OnDemand/IoStoreOnDemand.Build.cs
--- a/Engine/Source/Runtime/Experimental/IoStore/OnDemand/IoStoreOnDemand.Build.cs
+++ b/Engine/Source/Runtime/Experimental/IoStore/OnDemand/IoStoreOnDemand.Build.cs
@@ -35,5 +35,7 @@ public class IoStoreOnDemand : ModuleRules
 		{
 			PrivateDefinitions.Add("UE_HTTPCLIENT_THREADSAFE_CHECKS=0");
 		}
+
+		PrivateDefinitions.Add("UE_ENABLE_REQUEST_THREAD=0");
 	}
 }
diff --git a/Engine/Source/Runtime/Experimental/IoStore/OnDemand/Private/ChunkRequestShared.h b/Engine/Source/Runtime/Experimental/IoStore/OnDemand/Private/ChunkRequestShared.h
new file mode 100644
--- /dev/null
+++ b/Engine/Source/Runtime/Experimental/IoStore/OnDemand/Private/ChunkRequestShared.h
@@ -0,0 +1,351 @@
+// Copyright Epic Games, Inc. All Rights Reserved.
+
+#pragma once
+
+#include "IO/IoDispatcherBackend.h"
+#include "OnDemandIoStore.h"
+
+namespace UE::IoStore
+{
+
+extern int32 GIasHttpRangeRequestMinSizeKiB;
+
+template<typename T>
+class TThreadSafeIntrusiveQueue
+{
+public:
+	void Enqueue(T* Request)
+	{
+		check(Request->NextRequest == nullptr);
+		FScopeLock _(&CriticalSection);
+
+		if (Tail)
+		{
+			Tail->NextRequest = Request;
+		}
+		else
+		{
+			check(Head == nullptr);
+			Head = Request;
+		}
+
+		Tail = Request;
+	}
+
+	void EnqueueByPriority(T* Request)
+	{
+		FScopeLock _(&CriticalSection);
+		EnqueueByPriorityInternal(Request);
+	}
+
+	T* Dequeue()
+	{
+		FScopeLock _(&CriticalSection);
+
+		T* Requests = Head;
+		Head = Tail = nullptr;
+
+		return Requests;
+	}
+
+	void Reprioritize(T* Request)
+	{
+		// Switch to double linked list/array if this gets too expensive
+		FScopeLock _(&CriticalSection);
+		if (RemoveInternal(Request))
+		{
+			EnqueueByPriorityInternal(Request);
+		}
+	}
+
+private:
+	void EnqueueByPriorityInternal(T* Request)
+	{
+		check(Request->NextRequest == nullptr);
+
+		if (Head == nullptr || Request->Priority > Head->Priority)
+		{
+			if (Head == nullptr)
+			{
+				check(Tail == nullptr);
+				Tail = Request;
+			}
+
+			Request->NextRequest = Head;
+			Head = Request;
+		}
+		else if (Request->Priority <= Tail->Priority)
+		{
+			check(Tail != nullptr);
+			Tail->NextRequest = Request;
+			Tail = Request;
+		}
+		else
+		{
+			// NOTE: This can get expensive if the queue gets too long, might be better to have x number of bucket(s)
+			TRACE_CPUPROFILER_EVENT_SCOPE(IasBackend::EnqueueByPriority);
+			T* It = Head;
+			while (It->NextRequest != nullptr && Request->Priority <= It->NextRequest->Priority)
+			{
+				It = It->NextRequest;
+			}
+
+			Request->NextRequest = It->NextRequest;
+			It->NextRequest = Request;
+		}
+	}
+
+	bool RemoveInternal(T* Request)
+	{
+		check(Request != nullptr);
+		if (Head == nullptr)
+		{
+			check(Tail == nullptr);
+			return false;
+		}
+
+		if (Head == Request)
+		{
+			Head = Request->NextRequest;
+			if (Tail == Request)
+			{
+				check(Head == nullptr);
+				Tail = nullptr;
+			}
+
+			Request->NextRequest = nullptr;
+			return true;
+		}
+		else
+		{
+			T* It = Head;
+			while (It->NextRequest && It->NextRequest != Request)
+			{
+				It = It->NextRequest;
+			}
+
+			if (It->NextRequest == Request)
+			{
+				It->NextRequest = It->NextRequest->NextRequest;
+				Request->NextRequest = nullptr;
+				if (Tail == Request)
+				{
+					Tail = It;
+				}
+				return true;
+			}
+		}
+
+		return false;
+	}
+
+	FCriticalSection CriticalSection;
+	T* Head = nullptr;
+	T* Tail = nullptr;
+};
+
+///////////////////////////////////////////////////////////////////////////////
+static FIoHash GetChunkKey(const FIoHash& ChunkHash, const FIoOffsetAndLength& Range)
+{
+	FIoHashBuilder HashBuilder;
+	HashBuilder.Update(ChunkHash.GetBytes(), sizeof(FIoHash::ByteArray));
+	HashBuilder.Update(&Range, sizeof(FIoOffsetAndLength));
+
+	return HashBuilder.Finalize();
+}
+
+///////////////////////////////////////////////////////////////////////////////
+struct FChunkRequestParams
+{
+	static FChunkRequestParams Create(const FIoOffsetAndLength& OffsetLength, const FOnDemandChunkInfo& ChunkInfo)
+	{
+		FIoOffsetAndLength ChunkRange;
+		if (ChunkInfo.EncodedSize() <= (uint64(GIasHttpRangeRequestMinSizeKiB) << 10))
+		{
+			ChunkRange = FIoOffsetAndLength(0, ChunkInfo.EncodedSize());
+		}
+		else
+		{
+			const uint64 RawSize = FMath::Min<uint64>(OffsetLength.GetLength(), ChunkInfo.RawSize() - OffsetLength.GetOffset());
+
+			ChunkRange = FIoChunkEncoding::GetChunkRange(
+				ChunkInfo.RawSize(),
+				ChunkInfo.BlockSize(),
+				ChunkInfo.Blocks(),
+				OffsetLength.GetOffset(),
+				RawSize).ConsumeValueOrDie();
+		}
+
+		return FChunkRequestParams{ GetChunkKey(ChunkInfo.Hash(), ChunkRange), ChunkRange, ChunkInfo };
+	}
+
+	static FChunkRequestParams Create(FIoRequestImpl* Request, const FOnDemandChunkInfo& ChunkInfo)
+	{
+		check(Request);
+		check(Request->NextRequest == nullptr);
+		return Create(FIoOffsetAndLength(Request->Options.GetOffset(), Request->Options.GetSize()), ChunkInfo);
+	}
+
+	const FIoHash& GetUrlHash() const
+	{
+		return ChunkInfo.Hash();
+	}
+
+	void GetUrl(FAnsiStringBuilderBase& Url) const
+	{
+		const FString HashString = LexToString(ChunkInfo.Hash());
+		Url << "/" << ChunkInfo.ChunksDirectory()
+			<< "/" << HashString.Left(2)
+			<< "/" << HashString << ANSITEXTVIEW(".iochunk");
+	}
+
+	FIoChunkDecodingParams GetDecodingParams() const
+	{
+		FIoChunkDecodingParams Params;
+		Params.EncryptionKey = ChunkInfo.EncryptionKey();
+		Params.CompressionFormat = ChunkInfo.CompressionFormat();
+		Params.BlockSize = ChunkInfo.BlockSize();
+		Params.TotalRawSize = ChunkInfo.RawSize();
+		Params.EncodedBlockSize = ChunkInfo.Blocks();
+		Params.BlockHash = ChunkInfo.BlockHashes();
+		Params.EncodedOffset = ChunkRange.GetOffset();
+
+		return Params;
+	}
+
+	FIoHash ChunkKey;
+	FIoOffsetAndLength ChunkRange;
+	FOnDemandChunkInfo ChunkInfo;
+};
+
+///////////////////////////////////////////////////////////////////////////////
+struct FChunkRequest
+{
+#if UE_ENABLE_REQUEST_THREAD
+	DECLARE_DELEGATE_OneParam(FCompletionCallback, FChunkRequest* /*CompletedRequest*/);
+#endif // UE_ENABLE_REQUEST_THREAD
+
+	explicit FChunkRequest(FIoRequestImpl* Request, const FChunkRequestParams& RequestParams)
+		: NextRequest()
+		, Params(RequestParams)
+		, RequestHead(Request)
+		, RequestTail(Request)
+		, StartTime(FPlatformTime::Cycles64())
+		, Priority(Request->Priority)
+		, RequestCount(1)
+		, bCached(false)
+	{
+		check(Request && NextRequest == nullptr);
+	}
+
+	bool AddDispatcherRequest(FIoRequestImpl* Request)
+	{
+		check(RequestHead && RequestTail);
+		check(Request && !Request->NextRequest);
+
+		const bool bPriorityChanged = Request->Priority > RequestHead->Priority;
+		if (bPriorityChanged)
+		{
+			Priority = Request->Priority;
+			Request->NextRequest = RequestHead;
+			RequestHead = Request;
+		}
+		else
+		{
+			FIoRequestImpl* It = RequestHead;
+			while (It->NextRequest != nullptr && Request->Priority <= It->NextRequest->Priority)
+			{
+				It = It->NextRequest;
+			}
+
+			if (RequestTail == It)
+			{
+				check(It->NextRequest == nullptr);
+				RequestTail = Request;
+			}
+
+			Request->NextRequest = It->NextRequest;
+			It->NextRequest = Request;
+		}
+
+		RequestCount++;
+		return bPriorityChanged;
+	}
+
+	int32 RemoveDispatcherRequest(FIoRequestImpl* Request)
+	{
+		check(Request != nullptr);
+		check(RequestCount > 0);
+
+		if (RequestHead == Request)
+		{
+			RequestHead = Request->NextRequest;
+			if (RequestTail == Request)
+			{
+				check(RequestHead == nullptr);
+				RequestTail = nullptr;
+			}
+		}
+		else
+		{
+			FIoRequestImpl* It = RequestHead;
+			while (It->NextRequest != Request)
+			{
+				It = It->NextRequest;
+				if (It == nullptr)
+				{
+					return INDEX_NONE; // Not found
+				}
+			}
+			check(It->NextRequest == Request);
+			It->NextRequest = It->NextRequest->NextRequest;
+
+			if (RequestTail == Request)
+			{
+				check(It->NextRequest == nullptr);
+				RequestTail = It;
+			}
+		}
+
+		Request->NextRequest = nullptr;
+		RequestCount--;
+
+		return RequestCount;
+	}
+
+	FIoRequestImpl* DeqeueDispatcherRequests()
+	{
+		FIoRequestImpl* Head = RequestHead;
+		RequestHead = RequestTail = nullptr;
+		RequestCount = 0;
+
+		return Head;
+	}
+
+#if UE_ENABLE_REQUEST_THREAD
+	void OnRequestCompleted()
+	{
+		CompletionCallback.Execute(this);
+	}
+#endif // UE_ENABLE_REQUEST_THREAD
+
+	FChunkRequest* NextRequest;
+	FChunkRequestParams Params;
+	FIoRequestImpl* RequestHead;
+	FIoRequestImpl* RequestTail;
+#if UE_ENABLE_REQUEST_THREAD
+	FCompletionCallback CompletionCallback;
+#endif //UE_ENABLE_REQUEST_THREAD
+	FIASHostGroup HostGroup;
+	FIoBuffer Chunk;
+	uint64 StartTime;
+	int32 Priority;
+	uint16 RequestCount;
+	bool bCached;
+	bool bCancelled = false;
+	EIoErrorCode CacheGetStatus;
+};
+
+using FChunkRequestQueue = TThreadSafeIntrusiveQueue<FChunkRequest>;
+
+} //namespace UE::IoStore
diff --git a/Engine/Source/Runtime/Experimental/IoStore/OnDemand/Private/OnDemandHttpThread.cpp b/Engine/Source/Runtime/Experimental/IoStore/OnDemand/Private/OnDemandHttpThread.cpp
new file mode 100644
--- /dev/null
+++ b/Engine/Source/Runtime/Experimental/IoStore/OnDemand/Private/OnDemandHttpThread.cpp
@@ -0,0 +1,346 @@
+// Copyright Epic Games, Inc. All Rights Reserved.
+
+#include "OnDemandHttpThread.h"
+
+#if UE_ENABLE_REQUEST_THREAD
+
+#include "HAL/IConsoleManager.h"
+#include "HAL/LowLevelMemTracker.h"
+#include "HAL/RunnableThread.h"
+#include "IO/IoStoreOnDemand.h"
+#include "IasHostGroup.h"
+#include "Logging/StructuredLog.h"
+#include "OnDemandHttpClient.h"
+#include "Statistics.h"
+
+namespace UE::IoStore
+{
+
+////////////////////////////////////////
+// START EXTERN CVARS
+////////////////////////////////////////
+
+extern int32 GIasHttpHealthCheckWaitTime;
+extern int32 GOnDemandBackendThreadPriorityIndex;
+extern int32 GIasHttpTimeOutMs;
+extern int32 GIasHttpConnectionCount;
+extern int32 GIasHttpRecvBufKiB;
+extern int32 GIasHttpRetryCount;
+extern int32 GIasHttpFailTimeOutMs;
+extern bool GIasHttpAllowChunkedXfer;
+extern int32 GIasHttpConcurrentRequests;
+extern int32 GIasHttpRateLimitKiBPerSecond;
+extern int32 GIasHttpPollTimeoutMs;
+
+extern UE::Tasks::ETaskPriority GetRequestCompletionTaskPriority();
+
+// TODO - Maybe do LexToString/LexFromString (dupe code in UnrealEngine.cpp ThreadPriorityToString)
+const int32 GREPLACEOnDemandBackendThreadPriorities[5] =
+{
+	EThreadPriority::TPri_Lowest,
+	EThreadPriority::TPri_BelowNormal,
+	EThreadPriority::TPri_SlightlyBelowNormal,
+	EThreadPriority::TPri_Normal,
+	EThreadPriority::TPri_AboveNormal
+};
+
+const TCHAR* LexToString(EThreadPriority Priority)
+{
+	switch (Priority)
+	{
+		case EThreadPriority::TPri_Normal:
+			return TEXT("TPri_Normal");
+		case EThreadPriority::TPri_AboveNormal:
+			return TEXT("TPri_AboveNormal");
+		case EThreadPriority::TPri_BelowNormal:
+			return TEXT("TPri_BelowNormal");
+		case EThreadPriority::TPri_Highest:
+			return TEXT("TPri_Highest");
+		case EThreadPriority::TPri_Lowest:
+			return TEXT("TPri_Lowest");
+		case EThreadPriority::TPri_SlightlyBelowNormal:
+			return TEXT("TPri_SlightlyBelowNormal");
+		case EThreadPriority::TPri_TimeCritical:
+			return TEXT("TPri_TimeCritical");
+		case EThreadPriority::TPri_Num:
+		default:
+			return TEXT("TPri_Undefined");
+			break;
+	};
+}
+
+////////////////////////////////////////
+// END EXTERN CVARS
+////////////////////////////////////////
+
+FOnDemandThread* FOnDemandThread::SingletonInstance = nullptr;
+
+FOnDemandThread::FOnDemandThread()
+{
+	LLM_SCOPE_BYTAG(Ias);
+
+	checkf(SingletonInstance == nullptr, TEXT("Attempting to create multiple instances of FOnDemandThread"));
+	SingletonInstance = this;
+
+	const int32 ThreadPriorityIndex = FMath::Clamp(GOnDemandBackendThreadPriorityIndex, 0, (int32)UE_ARRAY_COUNT(GREPLACEOnDemandBackendThreadPriorities) - 1);
+	EThreadPriority DesiredThreadPriority = (EThreadPriority)GREPLACEOnDemandBackendThreadPriorities[ThreadPriorityIndex];
+	ThreadPriority = DesiredThreadPriority;
+
+	const uint32 StackSize = 0; // Use default stack size
+
+	Thread.Reset(FRunnableThread::Create(this, TEXT("IoStoreOnDemand.Http"), StackSize, ThreadPriority));
+}
+
+FOnDemandThread::~FOnDemandThread()
+{
+	checkf(SingletonInstance == this, TEXT("SingletonInstance was reassigned"));
+	SingletonInstance = nullptr;
+
+	Thread.Reset();
+
+	// TODO: Not 100% sure this is still needed?
+	DrainHttpRequests();
+}
+
+FOnDemandThread& FOnDemandThread::Get()
+{
+	check(SingletonInstance != nullptr);
+	return *SingletonInstance;
+}
+
+void FOnDemandThread::QueueRequest(FChunkRequest* Request)
+{
+	TRACE_CPUPROFILER_EVENT_SCOPE(FOnDemandThread::QueueRequest);
+
+	FOnDemandIoBackendStats::Get()->OnHttpEnqueue();
+
+	HttpRequests.EnqueueByPriority(Request);
+	TickThreadEvent->Trigger();
+}
+
+void FOnDemandThread::ReprioritizeRequest(FChunkRequest* Request)
+{
+	TRACE_CPUPROFILER_EVENT_SCOPE(FOnDemandThread::ReprioritizeRequest);
+	HttpRequests.Reprioritize(Request);
+}
+
+bool FOnDemandThread::Init()
+{
+	LLM_SCOPE_BYTAG(Ias);
+	if (!TryCreateHttpClient())
+	{
+		return false;
+	}
+
+	return true;
+}
+
+uint32 FOnDemandThread::Run()
+{
+	LLM_SCOPE_BYTAG(Ias);
+
+	check(HttpClient.IsValid());
+
+	while (!bStopRequested)
+	{
+		UpdateThreadPriorityIfNeeded();
+
+		if(!bStopRequested)
+		{
+			FHostGroupManager::Get().Tick(GIasHttpTimeOutMs, bStopRequested);
+			
+			// TODO: It would be better to only update connections as they need it, consider doing this on
+			// hostgroup connect/disconnect events.
+			HttpClient->UpdateConnections();
+
+			TickRequests();
+
+			OnTickIdleDelegate.Broadcast();
+
+			TickThreadEvent->Wait(FHostGroupManager::Get().GetNumDisconnctedHosts() > 0 ? GIasHttpHealthCheckWaitTime : 1000/*MAX_uint32*/);
+		}
+	}
+
+	return 0;
+}
+
+void FOnDemandThread::Stop()
+{
+	bStopRequested = true;
+	TickThreadEvent->Trigger();
+}
+
+void FOnDemandThread::Exit()
+{
+	HttpClient.Reset();
+}
+
+void FOnDemandThread::TickRequests()
+{
+	// We can't limit cvars so add a hard coded cap to prevent concurrent requests from getting out of hand.
+	const int32 MaxConcurrentRequests = FMath::Min(GIasHttpConcurrentRequests, 32);
+
+	int32 NumConcurrentRequests = 0;
+	FChunkRequest* NextChunkRequest = HttpRequests.Dequeue();
+
+	while (NextChunkRequest)
+	{
+		while (NextChunkRequest)
+		{
+			{
+				TRACE_CPUPROFILER_EVENT_SCOPE(FOnDemandThread::IssueHttpGet);
+				FChunkRequest* ChunkRequest = NextChunkRequest;
+				NextChunkRequest = ChunkRequest->NextRequest;
+				ChunkRequest->NextRequest = nullptr;
+
+				FOnDemandIoBackendStats::Get()->OnHttpDequeue();
+
+				if (ChunkRequest->bCancelled)
+				{
+					ChunkRequest->OnRequestCompleted();
+					FOnDemandIoBackendStats::Get()->OnHttpCancel();
+				}
+				else if (!ChunkRequest->HostGroup.IsConnected() || !bHttpEnabled)
+				{
+					ChunkRequest->OnRequestCompleted();
+
+					// Technically this request is being skipped because of a pre-existing error. It is not
+					// an error itself and it is not being canceled by higher level code. However we do not
+					// currently have a statistic for that and we have to call one of the existing types in
+					// order to correctly reduce the pending count.
+					FOnDemandIoBackendStats::Get()->OnHttpCancel();
+				}
+				else
+				{
+					TAnsiStringBuilder<256> Url;
+					ChunkRequest->Params.GetUrl(Url);
+
+					NumConcurrentRequests++;
+
+					auto AnsiUrl = StringCast<ANSICHAR>(Url.GetData(), Url.Len());
+
+					HttpClient->Get(ChunkRequest->HostGroup.GetUnderlyingHostGroup(), AnsiUrl, ChunkRequest->Params.ChunkRange,
+						[this, &NextChunkRequest, ChunkRequest, &NumConcurrentRequests]
+						(FMultiEndpointHttpClientResponse&& HttpResponse)
+						{
+							TRACE_CPUPROFILER_EVENT_SCOPE(FOnDemandThread::RequestCallback);
+
+							NumConcurrentRequests--;
+
+							if (HttpResponse.RetryCount > 0)
+							{
+								FOnDemandIoBackendStats::Get()->OnHttpRetry();
+							}
+
+							FOnDemandIoBackendStats::Get()->OnHttpCdnCacheReply(HttpResponse.CDNCacheStatus);
+
+							if (HttpResponse.IsOk())
+							{
+								ChunkRequest->HostGroup.OnSuccessfulResponse();
+
+								ChunkRequest->Chunk = MoveTemp(HttpResponse.Body);
+								FOnDemandIoBackendStats::Get()->OnHttpGet(ChunkRequest->Chunk.DataSize(), HttpResponse.DurationMilliseconds);
+
+								UE::Tasks::Launch(UE_SOURCE_LOCATION, [ ChunkRequest]()
+									{
+										ChunkRequest->OnRequestCompleted();
+									}, GetRequestCompletionTaskPriority());
+							}
+							else
+							{
+								FOnDemandIoBackendStats::Get()->OnHttpError();
+
+								if (ChunkRequest->HostGroup.OnFailedResponse())
+								{
+									// A disconnect was triggered
+									FOnDemandIoBackendStats::Get()->OnHttpDisconnected();
+								}
+
+								ChunkRequest->OnRequestCompleted();
+							}
+						});
+				}
+			}
+
+			if (NumConcurrentRequests >= MaxConcurrentRequests)
+			{
+				TRACE_CPUPROFILER_EVENT_SCOPE(FOnDemandThread::TickHttpSaturated);
+				while (NumConcurrentRequests >= MaxConcurrentRequests) //-V654
+				{
+					HttpClient->Tick(MAX_uint32, GIasHttpRateLimitKiBPerSecond);
+				}
+			}
+
+			if (!NextChunkRequest)
+			{
+				NextChunkRequest = HttpRequests.Dequeue();
+			}
+		}
+
+		{
+			// Keep processing pending connections until all requests are completed or a new one is issued
+			TRACE_CPUPROFILER_EVENT_SCOPE(FOnDemandThread::TickHttp);
+			while (HttpClient->Tick(GIasHttpPollTimeoutMs, GIasHttpRateLimitKiBPerSecond))
+			{
+				if (!NextChunkRequest)
+				{
+					NextChunkRequest = HttpRequests.Dequeue();
+				}
+				if (NextChunkRequest)
+				{
+					break;
+				}
+			}
+		}
+	}
+}
+
+void FOnDemandThread::DrainHttpRequests()
+{
+	FChunkRequest* Iterator = HttpRequests.Dequeue();
+	while (Iterator != nullptr)
+	{
+		FChunkRequest* Request = Iterator;
+		Iterator = Iterator->NextRequest;
+
+		FOnDemandIoBackendStats::Get()->OnHttpDequeue();
+		Request->OnRequestCompleted();
+		FOnDemandIoBackendStats::Get()->OnHttpCancel();
+	}
+}
+
+bool FOnDemandThread::TryCreateHttpClient()
+{
+	HttpClient = FMultiEndpointHttpClient::Create(FMultiEndpointHttpClientConfig
+		{
+			.MaxConnectionCount = GIasHttpConnectionCount,
+			.ReceiveBufferSize = GIasHttpRecvBufKiB >= 0 ? GIasHttpRecvBufKiB << 10 : -1,
+			.MaxRetryCount = GIasHttpRetryCount,
+			.TimeoutMs = GIasHttpFailTimeOutMs,
+			.Redirects = EHttpRedirects::Disabled,
+			.bEnableThreadSafetyChecks = true,
+			.bAllowChunkedTransfer = GIasHttpAllowChunkedXfer,
+			.LogCategory = &LogIas,
+			.LogVerbosity = ELogVerbosity::VeryVerbose
+		});
+
+	return HttpClient.IsValid();
+}
+
+void FOnDemandThread::UpdateThreadPriorityIfNeeded()
+{
+	// Read the thread priority from the cvar
+	int32 ThreadPriorityIndex = FMath::Clamp(GOnDemandBackendThreadPriorityIndex, 0, (int32)UE_ARRAY_COUNT(GREPLACEOnDemandBackendThreadPriorities) - 1);
+	EThreadPriority DesiredThreadPriority = (EThreadPriority)GREPLACEOnDemandBackendThreadPriorities[ThreadPriorityIndex];
+	if (DesiredThreadPriority != ThreadPriority)
+	{
+		UE_LOGFMT(LogIas, Log, "Updated IoStoreOnDemand.Http thread priority to '{Priority}'", LexToString(DesiredThreadPriority));
+		
+		FPlatformProcess::SetThreadPriority(DesiredThreadPriority);
+		ThreadPriority = DesiredThreadPriority;
+	}
+}
+
+} //namespace UE::IoStore
+
+#endif // UE_ENABLE_REQUEST_THREAD
diff --git a/Engine/Source/Runtime/Experimental/IoStore/OnDemand/Private/OnDemandHttpThread.h b/Engine/Source/Runtime/Experimental/IoStore/OnDemand/Private/OnDemandHttpThread.h
new file mode 100644
--- /dev/null
+++ b/Engine/Source/Runtime/Experimental/IoStore/OnDemand/Private/OnDemandHttpThread.h
@@ -0,0 +1,72 @@
+// Copyright Epic Games, Inc. All Rights Reserved.
+
+#pragma once
+
+#if UE_ENABLE_REQUEST_THREAD
+
+#include "ChunkRequestShared.h" // Very hefty
+#include "Delegates/Delegate.h"
+#include "HAL/Event.h"
+#include "HAL/PlatformAffinity.h"
+#include "HAL/Runnable.h"
+#include "Templates/UniquePtr.h"
+
+class FRunnableThread;
+
+namespace UE::IoStore
+{
+
+class FMultiEndpointHttpClient;
+
+class FOnDemandThread final
+	: public FRunnable
+{
+public:
+
+	FOnDemandThread();
+	virtual ~FOnDemandThread();
+
+	static FOnDemandThread& Get();
+
+	void QueueRequest(FChunkRequest* Request);
+	void ReprioritizeRequest(FChunkRequest* Request);
+
+	DECLARE_MULTICAST_DELEGATE(FOnTickIdle);
+
+	FOnTickIdle& OnTickIdle()
+	{
+		return OnTickIdleDelegate;
+	}
+
+private:
+
+	virtual bool Init() override;
+	virtual uint32 Run() override;
+	virtual void Stop() override;
+	virtual void Exit() override;
+
+	void TickRequests();
+	void DrainHttpRequests();
+
+	bool TryCreateHttpClient();
+	void UpdateThreadPriorityIfNeeded();
+
+private:
+	static FOnDemandThread* SingletonInstance;
+	TUniquePtr<FRunnableThread> Thread;
+
+	FEventRef TickThreadEvent;
+
+	FOnTickIdle OnTickIdleDelegate;
+	TUniquePtr<FMultiEndpointHttpClient> HttpClient;
+
+	FChunkRequestQueue HttpRequests;
+
+	EThreadPriority ThreadPriority = EThreadPriority::TPri_Normal;
+	std::atomic_bool bStopRequested = false;
+	std::atomic_bool bHttpEnabled = true;
+};
+
+} //namespace UE::IoStore
+
+#endif // UE_ENABLE_REQUEST_THREAD
diff --git a/Engine/Source/Runtime/Experimental/IoStore/OnDemand/Private/OnDemandIoDispatcherBackend.cpp b/Engine/Source/Runtime/Experimental/IoStore/OnDemand/Private/OnDemandIoDispatcherBackend.cpp
--- a/Engine/Source/Runtime/Experimental/IoStore/OnDemand/Private/OnDemandIoDispatcherBackend.cpp
+++ b/Engine/Source/Runtime/Experimental/IoStore/OnDemand/Private/OnDemandIoDispatcherBackend.cpp
@@ -3,6 +3,7 @@
 #include "OnDemandIoDispatcherBackend.h"
 
 #include "AnalyticsEventAttribute.h"
+#include "ChunkRequestShared.h"
 #include "Containers/BitArray.h"
 #include "Containers/StringView.h"
 #include "DistributionEndpoints.h"
@@ -33,6 +34,7 @@
 #include "Misc/ScopeRWLock.h"
 #include "OnDemandBackendStatus.h"
 #include "OnDemandHttpClient.h"
+#include "OnDemandHttpThread.h"
 #include "OnDemandIoStore.h"
 #include "ProfilingDebugging/IoStoreTrace.h"
 #include "Statistics.h"
@@ -64,6 +66,7 @@ namespace UE::IoStore
 extern FString GIasOnDemandTocExt;
 
 ///////////////////////////////////////////////////////////////////////////////
+
 /** Note that GIasHttpPrimaryEndpoint has no effect after initial start up */
 int32 GIasHttpPrimaryEndpoint = 0;
 static FAutoConsoleVariableRef CVar_IasHttpPrimaryEndpoint(
@@ -86,35 +89,35 @@ static FAutoConsoleVariableRef CVar_GIasHttpRateLimitKiBPerSecond(
 	TEXT("Http throttle limit in KiBPerSecond")
 );
 
-static int32 GIasHttpRecvBufKiB = -1;
+int32 GIasHttpRecvBufKiB = -1;
 static FAutoConsoleVariableRef CVar_GIasHttpRecvBufKiB(
 	TEXT("ias.HttpRecvBufKiB"),
 	GIasHttpRecvBufKiB,
 	TEXT("Recv buffer size")
 );
 
-static int32 GIasHttpConcurrentRequests = 8;
+int32 GIasHttpConcurrentRequests = 8;
 static FAutoConsoleVariableRef CVar_IasHttpConcurrentRequests(
 	TEXT("ias.HttpConcurrentRequests"),
 	GIasHttpConcurrentRequests,
 	TEXT("Number of concurrent requests in the http client.")
 );
 
-static int32 GIasHttpConnectionCount = 4;
+int32 GIasHttpConnectionCount = 4;
 static FAutoConsoleVariableRef CVar_IasHttpConnectionCount(
 	TEXT("ias.HttpConnectionCount"),
 	GIasHttpConnectionCount,
 	TEXT("Number of open HTTP connections to the on demand endpoint(s).")
 );
 
-static int32 GIasHttpFailTimeOutMs = 4 * 1000;
+int32 GIasHttpFailTimeOutMs = 4 * 1000;
 static FAutoConsoleVariableRef CVar_IasHttpFailTimeOutMs(
 	TEXT("ias.HttpFailTimeOutMs"),
 	GIasHttpFailTimeOutMs,
 	TEXT("Fail infinite network waits that take longer than this (in ms, a value of zero will use the default)")
 );
 
-static int32 GIasHttpRetryCount = -1;
+int32 GIasHttpRetryCount = -1;
 static FAutoConsoleVariableRef CVar_IasHttpRetryCount(
 	TEXT("ias.HttpRetryCount"),
 	GIasHttpRetryCount,
@@ -156,7 +159,7 @@ static FAutoConsoleVariableRef CVar_IoReportAnalytics(
 	TEXT("Enables reporting statics to the analytics system")
 );
 
-static int32 GIasHttpRangeRequestMinSizeKiB = 128;
+int32 GIasHttpRangeRequestMinSizeKiB = 128;
 static FAutoConsoleVariableRef CVar_IasHttpRangeRequestMinSizeKiB(
 	TEXT("ias.HttpRangeRequestMinSizeKiB"),
 	GIasHttpRangeRequestMinSizeKiB,
@@ -177,7 +180,7 @@ static FAutoConsoleVariableRef CVar_DistributedEndpointAttemptCount(
 	TEXT("Number of times we should try to resolve a distributed endpoint befor eusing the fallback url (if there is one)")
 );
 
-static bool GIasHttpAllowChunkedXfer = false;
+bool GIasHttpAllowChunkedXfer = false;
 static FAutoConsoleVariableRef CVar_IasHttpAllowChunkedXfer(
 	TEXT("ias.HttpAllowChunkedXfer"),
 	GIasHttpAllowChunkedXfer,
@@ -207,6 +210,7 @@ static FAutoConsoleVariableRef CVar_IasPoisonCache(
 	GIasPoisonCache,
 	TEXT("Fills all data materialized from the cache with 0x4d")
 );
+
 #endif // !UE_BUILD_SHIPPING
 
 // These priorities are indexed using the cvar below
@@ -247,7 +251,7 @@ const TCHAR* GOnDemandBackendThreadPriorityNames[] =
 	TEXT("TPri_AboveNormal")
 };
 
-static int32 GOnDemandBackendThreadPriorityIndex = 4; // EThreadPriority::TPri_AboveNormal
+int32 GOnDemandBackendThreadPriorityIndex = 4; // EThreadPriority::TPri_AboveNormal
 FAutoConsoleVariableRef CVarOnDemandBackendThreadPriority(
 	TEXT("ias.onDemandBackendThreadPriority"),
 	GOnDemandBackendThreadPriorityIndex,
@@ -256,334 +260,11 @@ FAutoConsoleVariableRef CVarOnDemandBackendThreadPriority(
 	ECVF_Default);
 
 ///////////////////////////////////////////////////////////////////////////////
-[[nodiscard]] static UE::Tasks::ETaskPriority GetRequestCompletionTaskPriority()
+[[nodiscard]] UE::Tasks::ETaskPriority GetRequestCompletionTaskPriority()
 {
 	return GCompleteMaterializeTaskPriorities[FMath::Clamp(GCompleteMaterializeTaskPriority, 0, UE_ARRAY_COUNT(GCompleteMaterializeTaskPriorities) - 1)];
 }
 
-///////////////////////////////////////////////////////////////////////////////
-FIoHash GetChunkKey(const FIoHash& ChunkHash, const FIoOffsetAndLength& Range)
-{
-	FIoHashBuilder HashBuilder;
-	HashBuilder.Update(ChunkHash.GetBytes(), sizeof(FIoHash::ByteArray));
-	HashBuilder.Update(&Range, sizeof(FIoOffsetAndLength));
-
-	return HashBuilder.Finalize();
-}
-
-///////////////////////////////////////////////////////////////////////////////
-template<typename T>
-class TThreadSafeIntrusiveQueue
-{
-public:
-	void Enqueue(T* Request)
-	{
-		check(Request->NextRequest == nullptr);
-		FScopeLock _(&CriticalSection);
-
-		if (Tail)
-		{
-			Tail->NextRequest = Request;
-		}
-		else
-		{
-			check(Head == nullptr);
-			Head = Request;	
-		}
-
-		Tail = Request;
-	}
-
-	void EnqueueByPriority(T* Request)
-	{
-		FScopeLock _(&CriticalSection);
-		EnqueueByPriorityInternal(Request);
-	}
-
-	T* Dequeue()
-	{
-		FScopeLock _(&CriticalSection);
-
-		T* Requests = Head;
-		Head = Tail = nullptr;
-
-		return Requests;
-	}
-
-	void Reprioritize(T* Request)
-	{
-		// Switch to double linked list/array if this gets too expensive
-		FScopeLock _(&CriticalSection);
-		if (RemoveInternal(Request))
-		{
-			EnqueueByPriorityInternal(Request);
-		}
-	}
-
-private:
-	void EnqueueByPriorityInternal(T* Request)
-	{
-		check(Request->NextRequest == nullptr);
-
-		if (Head == nullptr || Request->Priority > Head->Priority)
-		{
-			if (Head == nullptr)
-			{
-				check(Tail == nullptr);
-				Tail = Request;
-			}
-
-			Request->NextRequest = Head;
-			Head = Request;
-		}
-		else if (Request->Priority <= Tail->Priority)
-		{
-			check(Tail != nullptr);
-			Tail->NextRequest = Request;
-			Tail = Request;
-		}
-		else
-		{
-			// NOTE: This can get expensive if the queue gets too long, might be better to have x number of bucket(s)
-			TRACE_CPUPROFILER_EVENT_SCOPE(IasBackend::EnqueueByPriority);
-			T* It = Head;
-			while (It->NextRequest != nullptr && Request->Priority <= It->NextRequest->Priority)
-			{
-				It = It->NextRequest;
-			}
-
-			Request->NextRequest = It->NextRequest;
-			It->NextRequest = Request;
-		}
-	}
-
-	bool RemoveInternal(T* Request)
-	{
-		check(Request != nullptr);
-		if (Head == nullptr)
-		{
-			check(Tail == nullptr);
-			return false;
-		}
-
-		if (Head == Request)
-		{
-			Head = Request->NextRequest; 
-			if (Tail == Request)
-			{
-				check(Head == nullptr);
-				Tail = nullptr;
-			}
-
-			Request->NextRequest = nullptr;
-			return true;
-		}
-		else
-		{
-			T* It = Head;
-			while (It->NextRequest && It->NextRequest != Request)
-			{
-				It = It->NextRequest;
-			}
-
-			if (It->NextRequest == Request)
-			{
-				It->NextRequest = It->NextRequest->NextRequest;
-				Request->NextRequest = nullptr;
-				if (Tail == Request)
-				{
-					Tail = It;
-				}
-				return true;
-			}
-		}
-
-		return false;
-	}
-
-	FCriticalSection CriticalSection;
-	T* Head = nullptr;
-	T* Tail = nullptr;
-};
-
-///////////////////////////////////////////////////////////////////////////////
-struct FChunkRequestParams
-{
-	static FChunkRequestParams Create(const FIoOffsetAndLength& OffsetLength, FOnDemandChunkInfo ChunkInfo)
-	{
-		FIoOffsetAndLength ChunkRange;
-		if (ChunkInfo.EncodedSize() <= (uint64(GIasHttpRangeRequestMinSizeKiB) << 10))
-		{
-			ChunkRange = FIoOffsetAndLength(0, ChunkInfo.EncodedSize());
-		}
-		else
-		{
-			const uint64 RawSize = FMath::Min<uint64>(OffsetLength.GetLength(), ChunkInfo.RawSize() - OffsetLength.GetOffset());
-
-			ChunkRange = FIoChunkEncoding::GetChunkRange(
-				ChunkInfo.RawSize(),
-				ChunkInfo.BlockSize(),
-				ChunkInfo.Blocks(),
-				OffsetLength.GetOffset(),
-				RawSize).ConsumeValueOrDie();
-		}
-
-		return FChunkRequestParams{GetChunkKey(ChunkInfo.Hash(), ChunkRange), ChunkRange, ChunkInfo};
-	}
-
-	static FChunkRequestParams Create(FIoRequestImpl* Request, FOnDemandChunkInfo ChunkInfo)
-	{
-		check(Request);
-		check(Request->NextRequest == nullptr);
-		return Create(FIoOffsetAndLength(Request->Options.GetOffset(), Request->Options.GetSize()), ChunkInfo);
-	}
-
-	const FIoHash& GetUrlHash() const
-	{
-		return ChunkInfo.Hash();
-	}
-
-	void GetUrl(FAnsiStringBuilderBase& Url) const
-	{
-		const FString HashString = LexToString(ChunkInfo.Hash());
-		Url << "/" << ChunkInfo.ChunksDirectory()
-			<< "/" << HashString.Left(2)
-			<< "/" << HashString << ANSITEXTVIEW(".iochunk");
-	}
-
-	FIoChunkDecodingParams GetDecodingParams() const
-	{
-		FIoChunkDecodingParams Params;
-		Params.EncryptionKey = ChunkInfo.EncryptionKey(); 
-		Params.CompressionFormat = ChunkInfo.CompressionFormat();
-		Params.BlockSize = ChunkInfo.BlockSize();
-		Params.TotalRawSize = ChunkInfo.RawSize();
-		Params.EncodedBlockSize = ChunkInfo.Blocks(); 
-		Params.BlockHash = ChunkInfo.BlockHashes(); 
-		Params.EncodedOffset = ChunkRange.GetOffset();
-
-		return Params;
-	}
-
-	FIoHash ChunkKey;
-	FIoOffsetAndLength ChunkRange;
-	FOnDemandChunkInfo ChunkInfo;
-};
-
-///////////////////////////////////////////////////////////////////////////////
-struct FChunkRequest
-{
-	explicit FChunkRequest(FIoRequestImpl* Request, const FChunkRequestParams& RequestParams)
-		: NextRequest()
-		, Params(RequestParams)
-		, RequestHead(Request)
-		, RequestTail(Request)
-		, StartTime(FPlatformTime::Cycles64())
-		, Priority(Request->Priority)
-		, RequestCount(1)
-		, bCached(false)
-	{
-		check(Request && NextRequest == nullptr);
-	}
-
-	bool AddDispatcherRequest(FIoRequestImpl* Request)
-	{
-		check(RequestHead && RequestTail);
-		check(Request && !Request->NextRequest);
-
-		const bool bPriorityChanged = Request->Priority > RequestHead->Priority;
-		if (bPriorityChanged)
-		{
-			Priority = Request->Priority;
-			Request->NextRequest = RequestHead;
-			RequestHead = Request;
-		}
-		else
-		{
-			FIoRequestImpl* It = RequestHead;
-			while (It->NextRequest != nullptr && Request->Priority <= It->NextRequest->Priority)
-			{
-				It = It->NextRequest;
-			}
-
-			if (RequestTail == It)
-			{
-				check(It->NextRequest == nullptr);
-				RequestTail = Request;
-			}
-
-			Request->NextRequest = It->NextRequest;
-			It->NextRequest = Request;
-		}
-
-		RequestCount++;
-		return bPriorityChanged;
-	}
-
-	int32 RemoveDispatcherRequest(FIoRequestImpl* Request)
-	{
-		check(Request != nullptr);
-		check(RequestCount > 0);
-
-		if (RequestHead == Request)
-		{
-			RequestHead = Request->NextRequest; 
-			if (RequestTail == Request)
-			{
-				check(RequestHead == nullptr);
-				RequestTail = nullptr;
-			}
-		}
-		else
-		{
-			FIoRequestImpl* It = RequestHead;
-			while (It->NextRequest != Request)
-			{
-				It = It->NextRequest;
-				if (It == nullptr)
-				{
-					return INDEX_NONE; // Not found
-				}
-			}
-			check(It->NextRequest == Request);
-			It->NextRequest = It->NextRequest->NextRequest;
-
-			if (RequestTail == Request)
-			{
-				check(It->NextRequest == nullptr);
-				RequestTail = It;
-			}
-		}
-
-		Request->NextRequest = nullptr;
-		RequestCount--;
-
-		return RequestCount;
-	}
-
-	FIoRequestImpl* DeqeueDispatcherRequests()
-	{
-		FIoRequestImpl* Head = RequestHead;
-		RequestHead = RequestTail = nullptr;
-		RequestCount = 0;
-
-		return Head;
-	}
-
-	FChunkRequest* NextRequest;
-	FChunkRequestParams Params;
-	FIoRequestImpl* RequestHead;
-	FIoRequestImpl* RequestTail;
-	FIASHostGroup HostGroup;
-	FIoBuffer Chunk;
-	uint64 StartTime;
-	int32 Priority;
-	uint16 RequestCount;
-	bool bCached;
-	bool bCancelled = false;
-	EIoErrorCode CacheGetStatus;
-};
-
 ///////////////////////////////////////////////////////////////////////////////
 
 static void LogIoResult(
@@ -646,7 +327,6 @@ class FOnDemandIoBackend final
 	, public IOnDemandIoDispatcherBackend
 {
 	using FIoRequestQueue = TThreadSafeIntrusiveQueue<FIoRequestImpl>;
-	using FChunkRequestQueue = TThreadSafeIntrusiveQueue<FChunkRequest>;
 
 	struct FBackendData
 	{
@@ -853,12 +533,14 @@ private:
 	void ConditionallyStartBackendThread();
 	void CompleteRequest(FChunkRequest* ChunkRequest);
 	void CompleteCacheRead(FChunkRequest* ChunkRequest);
+	bool ResolveDefaultHostGroup();
 	bool CreateDefaultHostGroup();
 	bool ResolveDistributedEndpoint(const FDistributedEndpointUrl& Url, TArray<FString>& OutUrls);
 	void ProcessHttpRequests(FMultiEndpointHttpClient& HttpClient, int32 MaxConcurrentRequests);
 	int32 WaitForCompleteRequestTasks(float WaitTimeSeconds, float PollTimeSeconds);
 	void DrainHttpRequests();
 	void UpdateThreadPriorityIfNeeded();
+	void OnThreadTickIdle();
 
 #if UE_IAS_DEBUG_CONSOLE_CMDS
 	void RunRequestTest(const TArray<FString>& Args) const;
@@ -883,6 +565,10 @@ private:
 
 	EThreadPriority CurrentThreadPriority;
 
+#if UE_ENABLE_REQUEST_THREAD
+	FDelegateHandle OnThreadTickIdleHandle;
+#endif // UE_ENABLE_REQUEST_THREAD
+
 	mutable FRWLock Lock;
 	std::atomic_uint32_t InflightCacheRequestCount{0};
 	std::atomic_bool bStopRequested{false};
@@ -931,7 +617,18 @@ FOnDemandIoBackend::FOnDemandIoBackend(const FOnDemandEndpointConfig& Config, FO
 		ECVF_Default)
 	);
 
-#if !UE_BUILD_SHIPPING
+	DynamicConsoleCommands.Emplace(
+		IConsoleManager::Get().RegisterConsoleCommand(
+			TEXT("ias.AbandonCache"),
+			TEXT("Simulate code calling FOnDemandIoBackend::AbandonCache()"),
+			FConsoleCommandDelegate::CreateLambda([this]()
+				{
+					UE_LOG(LogIas, Display, TEXT("User invoked a cache abandon via 'ias.AbandonCache' (disable the IAS cache)"));
+					AbandonCache();
+				}),
+			ECVF_Default)
+	);
+
 	DynamicConsoleCommands.Emplace(
 		IConsoleManager::Get().RegisterConsoleCommand(
 			TEXT("ias.RunRequestTest"),
@@ -939,7 +636,6 @@ FOnDemandIoBackend::FOnDemandIoBackend(const FOnDemandEndpointConfig& Config, FO
 			FConsoleCommandWithArgsDelegate::CreateRaw(this, &FOnDemandIoBackend::RunRequestTest),
 			ECVF_Default)
 	);
-#endif // !UE_BUILD_SHIPPING
 
 #endif // UE_IAS_DEBUG_CONSOLE_CMDS
 }
@@ -977,11 +673,19 @@ void FOnDemandIoBackend::Shutdown()
 
 	UE_LOG(LogIas, Log, TEXT("Shutting down on demand I/O dispatcher backend"));
 
+#if UE_ENABLE_REQUEST_THREAD
+	if (OnThreadTickIdleHandle.IsValid())
+	{
+		FOnDemandThread::Get().OnTickIdle().Remove(OnThreadTickIdleHandle);
+		OnThreadTickIdleHandle.Reset();
+	}
+#endif // UE_ENABLE_REQUEST_THREAD
+
 	// Stop and wait for our backend thread to finish.
 	// Note, the IoDispatcher typically waits for all its pending io requests before shutting down its backends.
 	BackendThread.Reset();
 
-	// Drain any reamaining (cancelled) http requests that already been completed from the IoDispatcher point of view.
+	// Drain any remaining (cancelled) http requests that already been completed from the IoDispatcher point of view.
 	DrainHttpRequests();
 
 	// The CompleteRequest tasks may still be executing a while after the IoDispatcher has been notified about the completed io requests.
@@ -1004,9 +708,23 @@ void FOnDemandIoBackend::UpdateThreadPriorityIfNeeded()
 	}
 }
 
+void FOnDemandIoBackend::OnThreadTickIdle()
+{
+	if (BackendStatus.ShouldAbandonCache())
+	{
+		BackendStatus.SetAbandonCache(false);
+		check(BackendStatus.IsCacheEnabled() == false);
+		if (Cache.IsValid())
+		{
+			UE_LOG(LogIas, Log, TEXT("Abandoning cache, local file cache is no longer available"));
+			Cache.Release()->Abandon(); // Will delete its self
+		}
+	}
+}
 
 void FOnDemandIoBackend::ConditionallyStartBackendThread()
 {
+#if !UE_ENABLE_REQUEST_THREAD
 	FWriteScopeLock _(Lock);
 	if (BackendThread.IsValid() == false)
 	{
@@ -1017,8 +735,21 @@ void FOnDemandIoBackend::ConditionallyStartBackendThread()
 
 		BackendThread.Reset(FRunnableThread::Create(this, TEXT("Ias.Http"), 0, DesiredThreadPriority));
 	}
+#else
+
+	UE::Tasks::Launch(UE_SOURCE_LOCATION, [this]()
+		{
+			if (CreateDefaultHostGroup())
+			{
+				BackendStatus.SetHttpEnabled(true);
+			}
+		});
+
+	OnThreadTickIdleHandle = FOnDemandThread::Get().OnTickIdle().AddRaw(this, &FOnDemandIoBackend::OnThreadTickIdle);
+#endif // UE_ENABLE_REQUEST_THREAD
 }
 
+
 void FOnDemandIoBackend::CompleteRequest(FChunkRequest* ChunkRequest)
 {
 	LLM_SCOPE_BYTAG(Ias);
@@ -1159,9 +890,14 @@ void FOnDemandIoBackend::CompleteCacheRead(FChunkRequest* ChunkRequest)
 		return;
 	}
 
+#if UE_ENABLE_REQUEST_THREAD
+	ChunkRequest->CompletionCallback.BindRaw(this, &FOnDemandIoBackend::CompleteRequest);
+	FOnDemandThread::Get().QueueRequest(ChunkRequest);
+#else
 	Stats.OnHttpEnqueue();
 	HttpRequests.EnqueueByPriority(ChunkRequest);
 	TickBackendEvent->Trigger();
+#endif //UE_ENABLE_REQUEST_THREAD
 }
 
 bool FOnDemandIoBackend::Resolve(FIoRequestImpl* Request)
@@ -1196,8 +932,12 @@ bool FOnDemandIoBackend::Resolve(FIoRequestImpl* Request)
 	{
 		if (bUpdatePriority)
 		{
+#if UE_ENABLE_REQUEST_THREAD
+			FOnDemandThread::Get().ReprioritizeRequest(ChunkRequest);
+#else
 			TRACE_CPUPROFILER_EVENT_SCOPE(IasBackend::UpdatePriorityForIoRequest);
 			HttpRequests.Reprioritize(ChunkRequest);
+#endif
 		}
 		// The chunk for the request is already inflight 
 		return true;
@@ -1243,9 +983,15 @@ bool FOnDemandIoBackend::Resolve(FIoRequestImpl* Request)
 
 	ChunkRequest->HostGroup = ChunkInfo.HostGroup();
 
+#if UE_ENABLE_REQUEST_THREAD
+	ChunkRequest->CompletionCallback.BindRaw(this, &FOnDemandIoBackend::CompleteRequest);
+	FOnDemandThread::Get().QueueRequest(ChunkRequest);
+#else
 	Stats.OnHttpEnqueue();
 	HttpRequests.EnqueueByPriority(ChunkRequest);
 	TickBackendEvent->Trigger();
+#endif // UE_ENABLE_REQUEST_THREAD
+
 	return true;
 }
 
@@ -1274,7 +1020,11 @@ void FOnDemandIoBackend::UpdatePriorityForIoRequest(FIoRequestImpl* Request)
 	TRACE_CPUPROFILER_EVENT_SCOPE(IasBackend::UpdatePriorityForIoRequest);
 	if (FChunkRequest* ChunkRequest = ChunkRequests.TryUpdatePriority(Request))
 	{
+#if UE_ENABLE_REQUEST_THREAD
+		FOnDemandThread::Get().ReprioritizeRequest(ChunkRequest);
+#else
 		HttpRequests.Reprioritize(ChunkRequest);
+#endif // UE_ENABLE_REQUEST_THREAD
 	}
 }
 
@@ -1357,6 +1107,32 @@ const TCHAR* FOnDemandIoBackend::GetName() const
 }
 
 bool FOnDemandIoBackend::CreateDefaultHostGroup()
+{
+	TRACE_CPUPROFILER_EVENT_SCOPE(IasBackend::ResolveDefaultHostGroup);
+
+	const double InitStart = FPlatformTime::Seconds();
+
+	if (!ResolveDefaultHostGroup())
+	{
+		return false;
+	}
+
+	if (ConnectionTest(DefaultHostGroup.GetPrimaryHostUrl(), DefaultHostGroup.GetTestPath(), GIasHttpTimeOutMs))
+	{
+		Stats.OnHttpConnected();
+	}
+	else
+	{
+		DefaultHostGroup.Disconnect();
+	}
+
+	const double InitTime = FPlatformTime::Seconds() - InitStart;
+	UE_LOG(LogIas, Display, TEXT("HostGroup init took %.3f seconds"), InitTime);
+
+	return true;
+}
+
+bool FOnDemandIoBackend::ResolveDefaultHostGroup()
 {
 	if (DistributionUrl.IsValid())
 	{
@@ -1634,28 +1410,9 @@ uint32 FOnDemandIoBackend::Run()
 {
 	LLM_SCOPE_BYTAG(Ias);
 
+	if (!CreateDefaultHostGroup())
 	{
-		TRACE_CPUPROFILER_EVENT_SCOPE(IasBackend::ResolveDefaultHostGroup);
-
-		const double InitStart = FPlatformTime::Seconds();
-
-		if (!CreateDefaultHostGroup())
-		{
-			return 0;
-		}
-
-		if (ConnectionTest(DefaultHostGroup.GetPrimaryHostUrl(), DefaultHostGroup.GetTestPath(), GIasHttpTimeOutMs))
-		{
-			Stats.OnHttpConnected();
-
-		}
-		else
-		{
-			DefaultHostGroup.Disconnect();
-		}
-
-		const double InitTime = FPlatformTime::Seconds() - InitStart;
-		UE_LOG(LogIas, Display, TEXT("HostGroup init took %.3f seconds"), InitTime);
+		return 0;
 	}
 
 	TUniquePtr<FMultiEndpointHttpClient> HttpClient = FMultiEndpointHttpClient::Create(FMultiEndpointHttpClientConfig
@@ -1684,16 +1441,7 @@ uint32 FOnDemandIoBackend::Run()
 		{
 			FHostGroupManager::Get().Tick(GIasHttpTimeOutMs, bStopRequested);
 
-			if (BackendStatus.ShouldAbandonCache())
-			{
-				BackendStatus.SetAbandonCache(false);
-				check(BackendStatus.IsCacheEnabled() == false);
-				if (Cache.IsValid())
-				{
-					UE_LOG(LogIas, Log, TEXT("Abandoning cache, local file cache is no longer available"));
-					Cache.Release()->Abandon(); // Will delete its self
-				}
-			}
+			OnThreadTickIdle();
 
 			TickBackendEvent->Wait(FHostGroupManager::Get().GetNumDisconnctedHosts() > 0 ? GIasHttpHealthCheckWaitTime : MAX_uint32);
 		}
diff --git a/Engine/Source/Runtime/Experimental/IoStore/OnDemand/Private/OnDemandIoStore.cpp b/Engine/Source/Runtime/Experimental/IoStore/OnDemand/Private/OnDemandIoStore.cpp
--- a/Engine/Source/Runtime/Experimental/IoStore/OnDemand/Private/OnDemandIoStore.cpp
+++ b/Engine/Source/Runtime/Experimental/IoStore/OnDemand/Private/OnDemandIoStore.cpp
@@ -1,13 +1,6 @@
 // Copyright Epic Games, Inc. All Rights Reserved.
 
 #include "OnDemandIoStore.h"
-#include "IO/IoStatus.h"
-#include "OnDemandConfig.h"
-#include "OnDemandHttpClient.h"
-#include "OnDemandInstallCache.h"
-#include "OnDemandPackageStoreBackend.h"
-#include "OnDemandIoDispatcherBackend.h"
-#include "IasCache.h"
 
 #include "Algo/Accumulate.h"
 #include "Algo/AnyOf.h"
@@ -19,25 +12,34 @@
 #include "Containers/Ticker.h"
 #include "HAL/FileManager.h"
 #include "HAL/IConsoleManager.h"
+#include "HAL/LowLevelMemTracker.h"
 #include "HAL/PlatformFile.h"
 #include "HAL/PlatformFileManager.h"
-#include "HAL/PlatformTime.h"
 #include "HAL/PlatformProcess.h"
-#include "HAL/LowLevelMemTracker.h"
+#include "HAL/PlatformTime.h"
 #include "IO/IoChunkEncoding.h"
 #include "IO/IoContainerHeader.h"
+#include "IO/IoStatus.h"
 #include "IO/IoStoreOnDemand.h"
 #include "IO/PackageStore.h"
+#include "IasCache.h"
 #include "Misc/CommandLine.h"
 #include "Misc/ConfigCacheIni.h"
 #include "Misc/CoreDelegates.h"
 #include "Misc/CoreDelegatesInternal.h"
 #include "Misc/EncryptionKeyManager.h"
 #include "Misc/Guid.h"
-#include "Misc/Paths.h"
 #include "Misc/PathViews.h"
+#include "Misc/Paths.h"
+#include "OnDemandConfig.h"
+#include "OnDemandHttpClient.h"
+#include "OnDemandHttpThread.h"
+#include "OnDemandInstallCache.h"
+#include "OnDemandIoDispatcherBackend.h"
+#include "OnDemandPackageStoreBackend.h"
 #include "ProfilingDebugging/AssetMetadataTrace.h"
 #include "Serialization/MemoryReader.h"
+
 #if !(UE_BUILD_SHIPPING|UE_BUILD_TEST)
 #include "String/LexFromString.h"
 #endif
@@ -414,6 +416,10 @@ FString FOnDemandContainer::UniqueName() const
 ///////////////////////////////////////////////////////////////////////////////
 FOnDemandIoStore::FOnDemandIoStore()
 {
+#if UE_ENABLE_REQUEST_THREAD
+	OnDemandHttpThread = MakeUnique<FOnDemandThread>();
+#endif // UE_ENABLE_REQUEST_THREAD
+
 	FEncryptionKeyManager::Get().OnKeyAdded().AddRaw(this, &FOnDemandIoStore::OnEncryptionKeyAdded);
 }
 
@@ -435,6 +441,10 @@ FOnDemandIoStore::~FOnDemandIoStore()
 	{
 		TickFuture.Wait();
 	}
+
+#if UE_ENABLE_REQUEST_THREAD
+	OnDemandHttpThread.Reset();
+#endif // UE_ENABLE_REQUEST_THREAD
 }
 
 FIoStatus FOnDemandIoStore::Initialize()
diff --git a/Engine/Source/Runtime/Experimental/IoStore/OnDemand/Private/OnDemandIoStore.h b/Engine/Source/Runtime/Experimental/IoStore/OnDemand/Private/OnDemandIoStore.h
--- a/Engine/Source/Runtime/Experimental/IoStore/OnDemand/Private/OnDemandIoStore.h
+++ b/Engine/Source/Runtime/Experimental/IoStore/OnDemand/Private/OnDemandIoStore.h
@@ -26,6 +26,7 @@ using FSharedContainerHeader		= TSharedPtr<FIoContainerHeader>;
 namespace UE::IoStore
 {
 
+class FOnDemandThread;
 struct FOnDemandEndpointConfig;
 enum class ELogOnDemandCacheUsage : uint8;
 
@@ -324,6 +325,10 @@ private:
 											const TArray<FString>& ExistingContainers,
 											TArray<FSharedOnDemandContainer>& Out);
 
+#if UE_ENABLE_REQUEST_THREAD
+	TUniquePtr<FOnDemandThread>				OnDemandHttpThread;
+#endif // UE_ENABLE_REQUEST_THREAD
+
 	FSharedInstallCache						InstallCache;
 	FSharedPackageStoreBackend				PackageStoreBackend;
 	FSharedStreamingBackend					StreamingBackend;
