Remove a variable that was being unnecessarily captured by our http callback lambda.

#rnx

diff --git a/Engine/Source/Runtime/Experimental/IoStore/OnDemand/Private/OnDemandHttpThread.cpp b/Engine/Source/Runtime/Experimental/IoStore/OnDemand/Private/OnDemandHttpThread.cpp
--- a/Engine/Source/Runtime/Experimental/IoStore/OnDemand/Private/OnDemandHttpThread.cpp
+++ b/Engine/Source/Runtime/Experimental/IoStore/OnDemand/Private/OnDemandHttpThread.cpp
@@ -220,7 +220,7 @@ void FOnDemandThread::TickRequests()
 					auto AnsiUrl = StringCast<ANSICHAR>(Url.GetData(), Url.Len());
 
 					HttpClient->Get(ChunkRequest->HostGroup.GetUnderlyingHostGroup(), AnsiUrl, ChunkRequest->Params.ChunkRange,
-						[this, &NextChunkRequest, ChunkRequest, &NumConcurrentRequests]
+						[this, ChunkRequest, &NumConcurrentRequests]
 						(FMultiEndpointHttpClientResponse&& HttpResponse)
 						{
 							TRACE_CPUPROFILER_EVENT_SCOPE(FOnDemandThread::RequestCallback);
diff --git a/Engine/Source/Runtime/Experimental/IoStore/OnDemand/Private/OnDemandIoDispatcherBackend.cpp b/Engine/Source/Runtime/Experimental/IoStore/OnDemand/Private/OnDemandIoDispatcherBackend.cpp
--- a/Engine/Source/Runtime/Experimental/IoStore/OnDemand/Private/OnDemandIoDispatcherBackend.cpp
+++ b/Engine/Source/Runtime/Experimental/IoStore/OnDemand/Private/OnDemandIoDispatcherBackend.cpp
@@ -1318,7 +1318,7 @@ void FOnDemandIoBackend::ProcessHttpRequests(FMultiEndpointHttpClient& HttpClien
 					auto AnsiUrl = StringCast<ANSICHAR>(Url.GetData(), Url.Len());
 
 					HttpClient.Get(ChunkRequest->HostGroup.GetUnderlyingHostGroup(), AnsiUrl, ChunkRequest->Params.ChunkRange,
-						[this, &NextChunkRequest, ChunkRequest, &NumConcurrentRequests]
+						[this, ChunkRequest, &NumConcurrentRequests]
 						(FMultiEndpointHttpClientResponse&& HttpResponse)
 						{
 							TRACE_CPUPROFILER_EVENT_SCOPE(IasBackend::RequestCallback);
