Fix PVS warning in IoStoreOnDemandCore module.

#rnx

diff --git a/Engine/Source/Runtime/Experimental/IoStore/OnDemandCore/Private/IoStoreOnDemand.cpp b/Engine/Source/Runtime/Experimental/IoStore/OnDemandCore/Private/IoStoreOnDemand.cpp
--- a/Engine/Source/Runtime/Experimental/IoStore/OnDemandCore/Private/IoStoreOnDemand.cpp
+++ b/Engine/Source/Runtime/Experimental/IoStore/OnDemandCore/Private/IoStoreOnDemand.cpp
@@ -377,23 +377,18 @@ void FIoStoreOnDemandCoreModule::HandleModularFeatureRegistered(const FName& Typ
 
 	IOnDemandIoStoreFactory* Factory = static_cast<IOnDemandIoStoreFactory*>(ModularFeature);
 	GOnDemandIoStore = Factory->CreateInstance();
-	bool bInitializePostHostfix = false;
-	#if UE_IAS_CUSTOM_INITIALIZATION
-	bInitializePostHostfix = true;
-	#endif
 	if (FIoStatus Status = GOnDemandIoStore->Initialize(); !Status.IsOk())
 	{
 		UE_LOG(LogIoStoreOnDemand, Error, TEXT("Failed to initialize I/O store, reason ''"), *Status.ToString());
 		GOnDemandIoStore.Reset();
 	}
-	if (bInitializePostHostfix == false)
+	#if !UE_IAS_CUSTOM_INITIALIZATION
+	if (FIoStatus Status = GOnDemandIoStore->InitializePostHotfix(); !Status.IsOk())
 	{
-		if (FIoStatus Status = GOnDemandIoStore->InitializePostHotfix(); !Status.IsOk())
-		{
-			UE_LOG(LogIoStoreOnDemand, Error, TEXT("Failed to initialize I/O store, reason ''"), *Status.ToString());
-			GOnDemandIoStore.Reset();
-		}
+		UE_LOG(LogIoStoreOnDemand, Error, TEXT("Failed to initialize I/O store, reason ''"), *Status.ToString());
+		GOnDemandIoStore.Reset();
 	}
+	#endif
 	IModularFeatures::Get().OnModularFeatureRegistered().RemoveAll(this);
 }
 
