Dataflow:
- Added overlay to Construction view
- now we can display overlay text for every node from DebugDraw

diff --git a/Engine/Plugins/Experimental/Dataflow/Source/DataflowEditor/Private/Dataflow/DataflowConstructionViewport.cpp b/Engine/Plugins/Experimental/Dataflow/Source/DataflowEditor/Private/Dataflow/DataflowConstructionViewport.cpp
--- a/Engine/Plugins/Experimental/Dataflow/Source/DataflowEditor/Private/Dataflow/DataflowConstructionViewport.cpp
+++ b/Engine/Plugins/Experimental/Dataflow/Source/DataflowEditor/Private/Dataflow/DataflowConstructionViewport.cpp
@@ -18,6 +18,7 @@
 #include "ToolMenus.h"
 #include "Dataflow/DataflowConstructionVisualization.h"
 #include "ViewportToolbar/UnrealEdViewportToolbar.h"
+#include "Widgets/Text/SRichTextBlock.h"
 
 #define LOCTEXT_NAMESPACE "SDataflowConstructionViewport"
 
@@ -259,5 +260,50 @@ void SDataflowConstructionViewport::OnFloatingButtonClicked()
 {
 }
 
+void SDataflowConstructionViewport::PopulateViewportOverlays(TSharedRef<SOverlay> Overlay)
+{
+	SEditorViewport::PopulateViewportOverlays(Overlay);
+
+	Overlay->AddSlot()
+		.VAlign(VAlign_Top)
+		.HAlign(HAlign_Left)
+		.Padding(TAttribute<FMargin>(this, &SDataflowConstructionViewport::GetOverlayMargin))
+		[
+			SNew(SBorder)
+				.BorderImage(FAppStyle::Get().GetBrush("FloatingBorder"))
+				.Padding(4.f)
+				[
+					SNew(SRichTextBlock)
+					.Text(this, &SDataflowConstructionViewport::GetOverlayText)
+				]
+		];
+
+	// this widget will display the current viewed feature level
+	Overlay->AddSlot()
+		.VAlign(VAlign_Bottom)
+		.HAlign(HAlign_Right)
+		.Padding(5.0f)
+		[
+			BuildFeatureLevelWidget()
+		];
+}
+
+FMargin SDataflowConstructionViewport::GetOverlayMargin() const
+{
+	return UE::UnrealEd::ShowOldViewportToolbars()
+		? FMargin(6.0f, 36.0f, 6.0f, 6.0f) 
+		: FMargin(6.0f, 6.0f, 6.0f, 6.0f);
+}
+
+FText SDataflowConstructionViewport::GetOverlayText() const
+{
+	const TSharedPtr<FDataflowConstructionViewportClient> DataflowClient = StaticCastSharedPtr<FDataflowConstructionViewportClient>(Client);
+	if (DataflowClient)
+	{
+		return FText::FromString(DataflowClient->GetOverlayString());
+	}
+
+	return FText();
+}
 
 #undef LOCTEXT_NAMESPACE
diff --git a/Engine/Plugins/Experimental/Dataflow/Source/DataflowEditor/Private/Dataflow/DataflowConstructionViewportClient.cpp b/Engine/Plugins/Experimental/Dataflow/Source/DataflowEditor/Private/Dataflow/DataflowConstructionViewportClient.cpp
--- a/Engine/Plugins/Experimental/Dataflow/Source/DataflowEditor/Private/Dataflow/DataflowConstructionViewportClient.cpp
+++ b/Engine/Plugins/Experimental/Dataflow/Source/DataflowEditor/Private/Dataflow/DataflowConstructionViewportClient.cpp
@@ -358,4 +358,14 @@ void FDataflowConstructionViewportClient::AddReferencedObjects(FReferenceCollect
 	Collector.AddReferencedObjects(BehaviorsFor2DMode);
 }
 
+FString FDataflowConstructionViewportClient::GetOverlayString() const
+{
+	if (TSharedPtr<FDataflowEditorToolkit> DataflowEditorToolkit = DataflowEditorToolkitPtr.Pin())
+	{
+		return DataflowEditorToolkit->GetDebugDrawOverlayString();
+	}
+
+	return {};
+}
+
 #undef LOCTEXT_NAMESPACE 
diff --git a/Engine/Plugins/Experimental/Dataflow/Source/DataflowEditor/Private/Dataflow/DataflowEditorToolkit.cpp b/Engine/Plugins/Experimental/Dataflow/Source/DataflowEditor/Private/Dataflow/DataflowEditorToolkit.cpp
--- a/Engine/Plugins/Experimental/Dataflow/Source/DataflowEditor/Private/Dataflow/DataflowEditorToolkit.cpp
+++ b/Engine/Plugins/Experimental/Dataflow/Source/DataflowEditor/Private/Dataflow/DataflowEditorToolkit.cpp
@@ -911,8 +911,10 @@ void FDataflowEditorToolkit::UpdateDebugDraw()
 
 	DebugDrawComponent->UpdateBounds();
 	DebugDrawComponent->MarkRenderTransformDirty();
-}
 
+	// Build DebugDraw overlay string from the DebugDraw interface
+	DebugDrawOverlayString = DebugDrawObject->GetOverlayText();
+}
 
 namespace UE::Dataflow::Private
 {
@@ -2372,5 +2374,9 @@ void FDataflowEditorToolkit::AddReferencedObjects(FReferenceCollector& Collector
 	Collector.AddReferencedObject(DebugDrawComponent);
 }
 
+const FString& FDataflowEditorToolkit::GetDebugDrawOverlayString() const
+{
+	return DebugDrawOverlayString;
+}
 
 #undef LOCTEXT_NAMESPACE
diff --git a/Engine/Plugins/Experimental/Dataflow/Source/DataflowEditor/Public/Dataflow/DataflowConstructionViewport.h b/Engine/Plugins/Experimental/Dataflow/Source/DataflowEditor/Public/Dataflow/DataflowConstructionViewport.h
--- a/Engine/Plugins/Experimental/Dataflow/Source/DataflowEditor/Public/Dataflow/DataflowConstructionViewport.h
+++ b/Engine/Plugins/Experimental/Dataflow/Source/DataflowEditor/Public/Dataflow/DataflowConstructionViewport.h
@@ -10,6 +10,7 @@ class FAdvancedPreviewScene;
 class UDataflowEditorMode;
 class FDataflowConstructionViewportClient;
 class FDataflowConstructionScene;
+class SRichTextBlock;
 
 // ----------------------------------------------------------------------------------
 
@@ -18,7 +19,7 @@ class SDataflowConstructionViewport : public SBaseCharacterFXEditorViewport, pub
 public:
 	SLATE_BEGIN_ARGS(SDataflowConstructionViewport) {}
 	SLATE_ARGUMENT(TSharedPtr<FEditorViewportClient>, ViewportClient)
-		SLATE_END_ARGS()
+	SLATE_END_ARGS()
 
 	SDataflowConstructionViewport();
 
@@ -38,8 +39,16 @@ public:
 	virtual TSharedRef<class SEditorViewport> GetViewportWidget() override;
 	virtual TSharedPtr<FExtender> GetExtenders() const override;
 	virtual void OnFloatingButtonClicked() override;
+	virtual void PopulateViewportOverlays(TSharedRef<SOverlay> Overlay) override;
+
+	FText GetOverlayText() const;
 
 private:
+	FMargin GetOverlayMargin() const;
+
+	/** Pointer to the box into which the overlay text items are added */
+	TSharedPtr<SRichTextBlock> OverlayText;
+
 	UDataflowEditorMode* GetEdMode() const;
 };
 
diff --git a/Engine/Plugins/Experimental/Dataflow/Source/DataflowEditor/Public/Dataflow/DataflowConstructionViewportClient.h b/Engine/Plugins/Experimental/Dataflow/Source/DataflowEditor/Public/Dataflow/DataflowConstructionViewportClient.h
--- a/Engine/Plugins/Experimental/Dataflow/Source/DataflowEditor/Public/Dataflow/DataflowConstructionViewportClient.h
+++ b/Engine/Plugins/Experimental/Dataflow/Source/DataflowEditor/Public/Dataflow/DataflowConstructionViewportClient.h
@@ -56,6 +56,8 @@ public:
 	DECLARE_MULTICAST_DELEGATE_OneParam(FOnSelectionChangedMulticast, const TArray<UPrimitiveComponent*>&)
 	FOnSelectionChangedMulticast OnSelectionChangedMulticast;
 
+	FString GetOverlayString() const;
+
 protected:
 
 	// FGCObject
diff --git a/Engine/Plugins/Experimental/Dataflow/Source/DataflowEditor/Public/Dataflow/DataflowEditorToolkit.h b/Engine/Plugins/Experimental/Dataflow/Source/DataflowEditor/Public/Dataflow/DataflowEditorToolkit.h
--- a/Engine/Plugins/Experimental/Dataflow/Source/DataflowEditor/Public/Dataflow/DataflowEditorToolkit.h
+++ b/Engine/Plugins/Experimental/Dataflow/Source/DataflowEditor/Public/Dataflow/DataflowEditorToolkit.h
@@ -100,6 +100,8 @@ public:
 	UDataflowSubGraph* GetSubGraph(const FGuid& SubGraphGuid) const;
 	UDataflowSubGraph* GetSubGraph(FName SubGraphName) const;
 
+	const FString& GetDebugDrawOverlayString() const;
+
 protected:
 
 	UDataflowEditor* DataflowEditor = nullptr;
@@ -314,4 +316,6 @@ private:
 
 	/** Persistent component used for debug drawing */
 	TObjectPtr<class UDataflowDebugDrawComponent> DebugDrawComponent;
+
+	FString DebugDrawOverlayString;
 };
diff --git a/Engine/Source/Runtime/Experimental/Dataflow/Core/Private/Dataflow/DataflowSelection.cpp b/Engine/Source/Runtime/Experimental/Dataflow/Core/Private/Dataflow/DataflowSelection.cpp
--- a/Engine/Source/Runtime/Experimental/Dataflow/Core/Private/Dataflow/DataflowSelection.cpp
+++ b/Engine/Source/Runtime/Experimental/Dataflow/Core/Private/Dataflow/DataflowSelection.cpp
@@ -155,3 +155,14 @@ void FDataflowSelection::InitFromArray(const FManagedArrayCollection& InCollecti
 	Initialize(InCollection.NumElements(GroupName), false);
 	SetFromArray(InSelectionArr);
 }
+
+FString FDataflowSelection::ToString()
+{
+	const int32 NumSelectedElems = NumSelected();
+	const int32 NumElems = SelectionArray.Num();
+	FString OutString;
+
+	OutString = FString::Printf(TEXT("Selected %s: %s of %s"), *GroupName.ToString(), *FString::FromInt(NumSelected()), *FString::FromInt(NumElems));
+
+	return OutString;
+}
diff --git a/Engine/Source/Runtime/Experimental/Dataflow/Core/Public/Dataflow/DataflowDebugDrawInterface.h b/Engine/Source/Runtime/Experimental/Dataflow/Core/Public/Dataflow/DataflowDebugDrawInterface.h
--- a/Engine/Source/Runtime/Experimental/Dataflow/Core/Public/Dataflow/DataflowDebugDrawInterface.h
+++ b/Engine/Source/Runtime/Experimental/Dataflow/Core/Public/Dataflow/DataflowDebugDrawInterface.h
@@ -46,6 +46,8 @@ public:
 	DATAFLOWCORE_API virtual void DrawSphere(const FVector& Center, double Radius) const = 0;
 	DATAFLOWCORE_API virtual void DrawCapsule(const FVector& Center, const double& Radius, const double& HalfHeight, const FVector& XAxis, const FVector& YAxis, const FVector &ZAxis) const = 0;
 
+	DATAFLOWCORE_API virtual void DrawOverlayText(const FString& InString) = 0;
+	DATAFLOWCORE_API virtual FString GetOverlayText() const = 0;
 
 	// TODO: More draw functions
 
diff --git a/Engine/Source/Runtime/Experimental/Dataflow/Core/Public/Dataflow/DataflowSelection.h b/Engine/Source/Runtime/Experimental/Dataflow/Core/Public/Dataflow/DataflowSelection.h
--- a/Engine/Source/Runtime/Experimental/Dataflow/Core/Public/Dataflow/DataflowSelection.h
+++ b/Engine/Source/Runtime/Experimental/Dataflow/Core/Public/Dataflow/DataflowSelection.h
@@ -57,6 +57,12 @@ struct FDataflowSelection
 	*/
 	DATAFLOWCORE_API void InitFromArray(const FManagedArrayCollection& InCollection, const TArray<int32>& InSelectionArr);
 
+	/**
+	* Print selection in 
+	* "Selected Transforms: 23 of 34" format;
+	*/
+	DATAFLOWCORE_API FString ToString();
+
 private:
 	FName GroupName;
 	TBitArray<> SelectionArray;
@@ -123,3 +129,13 @@ struct FDataflowMaterialSelection : public FDataflowSelection
 		: FDataflowSelection(MaterialGroupName)
 	{}
 };
+
+UENUM(BlueprintType)
+enum class EDataflowSelectionType : uint8
+{
+	Transform UMETA(DisplayName = "Transform"),
+	Vertices UMETA(DisplayName = "Vertices"),
+	Faces UMETA(DisplayName = "Faces"),
+	Geometry UMETA(DisplayName = "Geometry"),
+};
+
diff --git a/Engine/Source/Runtime/Experimental/Dataflow/Engine/Private/Dataflow/DataflowDebugDraw.cpp b/Engine/Source/Runtime/Experimental/Dataflow/Engine/Private/Dataflow/DataflowDebugDraw.cpp
--- a/Engine/Source/Runtime/Experimental/Dataflow/Engine/Private/Dataflow/DataflowDebugDraw.cpp
+++ b/Engine/Source/Runtime/Experimental/Dataflow/Engine/Private/Dataflow/DataflowDebugDraw.cpp
@@ -191,7 +191,7 @@ void FDataflowDebugDraw::DrawMesh(const IDebugDrawMesh& Mesh) const
 
 void FDataflowDebugDraw::DrawBox(const FVector& Extents, const FQuat& Rotation, const FVector& Center, double UniformScale) const
 {
-	DebugRenderSceneProxy->Boxes.Add(FDebugRenderSceneProxy::FDebugBox(FBox(Center - Extents, Center + Extents), ColorWithTranslucency.ToFColor(true), DrawType, LineWidth));
+	DebugRenderSceneProxy->Boxes.Add(FDebugRenderSceneProxy::FDebugBox(FBox(Center - UniformScale * Extents, Center + UniformScale * Extents), ColorWithTranslucency.ToFColor(true), DrawType, LineWidth));
 }
 
 void FDataflowDebugDraw::DrawSphere(const FVector& Center, double Radius) const
@@ -204,6 +204,21 @@ void FDataflowDebugDraw::DrawCapsule(const FVector& Center, const double& Radius
 	DebugRenderSceneProxy->Capsules.Add(FDebugRenderSceneProxy::FCapsule(Center, Radius, XAxis, YAxis, ZAxis, HalfHeight, ColorWithTranslucency, DrawType));
 }
 
+void FDataflowDebugDraw::DrawOverlayText(const FString& InString)
+{
+	OverlayStrings.Add(InString);
+}
+
+FString FDataflowDebugDraw::GetOverlayText() const
+{
+	if (OverlayStrings.Num() > 0)
+	{
+		return OverlayStrings[0];
+	}
+
+	return {};
+}
+
 /* ----------------------------------------------------------------------------------------------------------------------- */
 
 void FDataflowNodeDebugDrawSettings::SetDebugDrawSettings(IDataflowDebugDrawInterface& DataflowRenderingInterface) const
diff --git a/Engine/Source/Runtime/Experimental/Dataflow/Engine/Private/Dataflow/DataflowOverlay.cpp b/Engine/Source/Runtime/Experimental/Dataflow/Engine/Private/Dataflow/DataflowOverlay.cpp
new file mode 100644
--- /dev/null
+++ b/Engine/Source/Runtime/Experimental/Dataflow/Engine/Private/Dataflow/DataflowOverlay.cpp
@@ -0,0 +1,103 @@
+// Copyright Epic Games, Inc. All Rights Reserved.
+
+#include "Dataflow/DataflowOverlay.h"
+#include "GeometryCollection/Facades/CollectionBoundsFacade.h"
+#include "GeometryCollection/GeometryCollection.h"
+#include "Dataflow/DataflowNode.h"
+#include "Dataflow/DataflowSelection.h"
+
+namespace UE::Dataflow::Overlay
+{
+	static const FString COverlaySeparator = TEXT("─────────────────────────\n");
+
+	FString BuildOverlayNodeInfoString(const FDataflowNode* InNode)
+	{
+		FString Path = "/Construction/"; // TODO: Get this from Toolkit
+		FString NodeName = InNode->GetName().ToString();
+		FString NodeType = InNode->GetType().ToString();
+
+		FString OutStr;
+		OutStr += FString::Printf(TEXT("%s\n%s\n%s\n"), *Path, *NodeName, *NodeType);
+		return OutStr;
+	}
+
+	FString BuildOverlaySelectionEvaluateResultString(FDataflowSelection& InSelection)
+	{
+		FString OutStr;
+		OutStr += FString::Printf(TEXT("%s\n"), *InSelection.ToString());
+		return OutStr;
+	}
+
+	FString BuildOverlayCollectionInfoString(const FManagedArrayCollection& InCollection)
+	{
+		FString NumVertsStr = FString::FormatAsNumber(InCollection.NumElements(FGeometryCollection::VerticesGroup));
+		FString NumFacesStr = FString::FormatAsNumber(InCollection.NumElements(FGeometryCollection::FacesGroup));
+		FString NumTransformsStr = FString::FormatAsNumber(InCollection.NumElements(FGeometryCollection::TransformGroup));
+		FString NumGeometryStr = FString::FormatAsNumber(InCollection.NumElements(FGeometryCollection::GeometryGroup));
+
+		FString OutStr;
+		OutStr += FString::Printf(TEXT("Vertices: %s\nFaces: %s\nTransforms: %s\nGeometry: %s\n"), *NumVertsStr, *NumFacesStr, *NumTransformsStr, *NumTransformsStr);
+		return OutStr;
+	}
+
+	FString BuildOverlayBoundsInfoString(const FManagedArrayCollection& InCollection)
+	{
+		GeometryCollection::Facades::FBoundsFacade BoundsFacade(InCollection);
+		const FBox BBox = BoundsFacade.GetBoundingBoxInCollectionSpace();
+
+		FString BBoxCenterStr = FString::Printf(TEXT("%3.3f, %3.3f, %3.3f"), BBox.GetCenter().X, BBox.GetCenter().Y, BBox.GetCenter().Z);
+		FString BBoxMinStr = FString::Printf(TEXT("%3.3f, %3.3f, %3.3f"), BBox.Min.X, BBox.Min.Y, BBox.Min.Z);
+		FString BBoxMaxStr = FString::Printf(TEXT("%3.3f, %3.3f, %3.3f"), BBox.Max.X, BBox.Max.Y, BBox.Max.Z);
+		FString BBoxSizeStr = FString::Printf(TEXT("%3.3f, %3.3f, %3.3f"), BBox.GetSize().X, BBox.GetSize().Y, BBox.GetSize().Z);
+
+		FString OutStr;
+		OutStr += FString::Printf(TEXT("Center: %s\nMin: %s\nMax: %s\nSize: %s\n"), *BBoxCenterStr, *BBoxMinStr, *BBoxMaxStr, *BBoxSizeStr);
+		return OutStr;
+	}
+
+	FString BuildOverlayMemInfoString(const FDataflowNode* InNode)
+	{
+		FString OutStr;
+
+		int32 SumMem = 0;
+		for (FPropertyValueIterator PropertyIt(FProperty::StaticClass(), InNode->TypedScriptStruct(), InNode); PropertyIt; ++PropertyIt)
+		{
+			const FProperty* Property = PropertyIt.Key();
+
+			if (const FEnumProperty* PropAsEnumProp = CastField<const FEnumProperty>(Property))
+			{
+				SumMem += PropAsEnumProp->GetUnderlyingProperty()->GetElementSize();
+			}
+			else if (Property)
+			{
+				SumMem += Property->GetElementSize();
+			}
+		}
+
+		float MemResult = (float)SumMem / 1024.f;
+		OutStr += FString::Printf(TEXT("Memory: %03fKB\n"), MemResult);
+
+		return OutStr;
+	}
+
+	FString BuildOverlayFinalString(const TArray<FString>& InStringArr)
+	{
+		FString OutStr;
+
+		int32 Idx = 0;
+		for (const FString& Str : InStringArr)
+		{
+			OutStr += Str;
+
+			if (Idx < InStringArr.Num() - 1)
+			{
+				OutStr += COverlaySeparator;
+			}
+
+			Idx++;
+		}
+
+		return OutStr;
+	}
+}
+
diff --git a/Engine/Source/Runtime/Experimental/Dataflow/Engine/Public/Dataflow/DataflowDebugDraw.h b/Engine/Source/Runtime/Experimental/Dataflow/Engine/Public/Dataflow/DataflowDebugDraw.h
--- a/Engine/Source/Runtime/Experimental/Dataflow/Engine/Public/Dataflow/DataflowDebugDraw.h
+++ b/Engine/Source/Runtime/Experimental/Dataflow/Engine/Public/Dataflow/DataflowDebugDraw.h
@@ -6,6 +6,10 @@
 #include "DebugRenderSceneProxy.h"
 #include "Engine/EngineTypes.h"
 #include "Math/Color.h"
+#include "GeometryCollection/Facades/CollectionBoundsFacade.h"
+#include "GeometryCollection/GeometryCollection.h"
+#include "Dataflow/DataflowNode.h"
+#include "UObject/EnumProperty.h"
 
 #include "DataflowDebugDraw.generated.h"
 
@@ -37,6 +41,9 @@ private:
 	virtual void DrawSphere(const FVector& Center, double Radius) const override;
 	virtual void DrawCapsule(const FVector& Center, const double& Radius, const double& HalfHeight, const FVector& XAxis, const FVector& YAxis, const FVector &ZAxis) const override;
 
+	virtual void DrawOverlayText(const FString& InString) override;
+	virtual FString GetOverlayText() const override;
+
 	class FDataflowDebugRenderSceneProxy* DebugRenderSceneProxy = nullptr;
 
 	FLinearColor Color = FLinearColor::White;
@@ -49,6 +56,8 @@ private:
 
 	FLinearColor ColorWithTranslucency = FLinearColor::White;
 	FDebugRenderSceneProxy::EDrawType DrawType = FDebugRenderSceneProxy::EDrawType::Invalid;
+
+	TArray<FString> OverlayStrings;
 };
 
 UENUM(BlueprintType)
@@ -81,4 +90,16 @@ struct FDataflowNodeDebugDrawSettings
 	DATAFLOWENGINE_API void SetDebugDrawSettings(IDataflowDebugDrawInterface& DataflowRenderingInterface) const;
 };
 
+static void SetWireframeRender(IDataflowDebugDrawInterface& DataflowRenderingInterface)
+{
+	DataflowRenderingInterface.SetShaded(false);
+	DataflowRenderingInterface.SetWireframe(true);
+}
+
+static void SetShadedRender(IDataflowDebugDrawInterface& DataflowRenderingInterface)
+{
+	DataflowRenderingInterface.SetShaded(true);
+	DataflowRenderingInterface.SetWireframe(true);
+}
+
 
diff --git a/Engine/Source/Runtime/Experimental/Dataflow/Engine/Public/Dataflow/DataflowOverlay.h b/Engine/Source/Runtime/Experimental/Dataflow/Engine/Public/Dataflow/DataflowOverlay.h
new file mode 100644
--- /dev/null
+++ b/Engine/Source/Runtime/Experimental/Dataflow/Engine/Public/Dataflow/DataflowOverlay.h
@@ -0,0 +1,20 @@
+// Copyright Epic Games, Inc. All Rights Reserved.
+
+#pragma once 
+
+#include "Containers/UnrealString.h"
+
+struct FDataflowNode;
+struct FManagedArrayCollection;
+struct FDataflowSelection;
+
+namespace UE::Dataflow::Overlay
+{
+	DATAFLOWENGINE_API FString BuildOverlayNodeInfoString(const FDataflowNode* InNode);
+	DATAFLOWENGINE_API FString BuildOverlaySelectionEvaluateResultString(FDataflowSelection& InSelection);
+	DATAFLOWENGINE_API FString BuildOverlayCollectionInfoString(const FManagedArrayCollection& InCollection);
+	DATAFLOWENGINE_API FString BuildOverlayBoundsInfoString(const FManagedArrayCollection& InCollection);
+	DATAFLOWENGINE_API FString BuildOverlayMemInfoString(const FDataflowNode* InNode);
+	DATAFLOWENGINE_API FString BuildOverlayFinalString(const TArray<FString>& InStringArr);
+}
+
