[ChaosVD] Fixing unitialized variables in FChaosVDVSmooth

#rnx

diff --git a/Engine/Source/Runtime/Experimental/ChaosVisualDebugger/Public/DataWrappers/ChaosVDParticleDataWrapper.h b/Engine/Source/Runtime/Experimental/ChaosVisualDebugger/Public/DataWrappers/ChaosVDParticleDataWrapper.h
--- a/Engine/Source/Runtime/Experimental/ChaosVisualDebugger/Public/DataWrappers/ChaosVDParticleDataWrapper.h
+++ b/Engine/Source/Runtime/Experimental/ChaosVisualDebugger/Public/DataWrappers/ChaosVDParticleDataWrapper.h
@@ -654,9 +654,9 @@ struct FChaosVDVSmooth : public FChaosVDWrapperDataBase
 	}
 
 	UPROPERTY(VisibleAnywhere, Category= "VSmooth")
-	FVector MV;
+	FVector MV = FVector(ForceInit);
 	UPROPERTY(VisibleAnywhere, Category= "VSmooth")
-	FVector MW;
+	FVector MW = FVector(ForceInit);
 };
 
 CVD_IMPLEMENT_SERIALIZER(FChaosVDVSmooth)
