Chaos Cloth Physics Assets: Add the option to create a "OneSided" tapered capsule. Tapered capsules are currently only used by cloth, so adding this additional parameter should not affect non-cloth PAs.

When calculating the cloth collisions response, OneSided tapered capsules will reflect the resolved collision location across the YZ-plane if necessary to ensure the collision is always pushed to the +X side of the capsule. This prevents colliders from getting trapped inside cloth.

diff --git a/Engine/Plugins/ChaosCloth/Source/ChaosCloth/Private/ChaosCloth/ChaosClothVisualization.cpp b/Engine/Plugins/ChaosCloth/Source/ChaosCloth/Private/ChaosCloth/ChaosClothVisualization.cpp
--- a/Engine/Plugins/ChaosCloth/Source/ChaosCloth/Private/ChaosCloth/ChaosClothVisualization.cpp
+++ b/Engine/Plugins/ChaosCloth/Source/ChaosCloth/Private/ChaosCloth/ChaosClothVisualization.cpp
@@ -1561,6 +1561,12 @@ namespace Chaos
 								DrawSphere(PDI, FSphere(X1, Radius1), Rotation, Position, Solver->GetLocalSpaceScale(), Color);
 								DrawSphere(PDI, FSphere(X2, Radius2), Rotation, Position, Solver->GetLocalSpaceScale(), Color);
 								DrawTaperedCylinder(PDI, FTaperedCylinder(X1, X2, Radius1, Radius2), Rotation, Position, Solver->GetLocalSpaceScale(), Color);
+								if (TaperedCapsule.IsOneSided())
+								{
+									const FVec3 Center = Position + Rotation.RotateVector(Solver->GetLocalSpaceScale() * .5 * (X1 + X2));
+									const FVec3 NormalEnd = Center + Rotation.RotateVector(Solver->GetLocalSpaceScale() * .5 * (Radius1 + Radius2) * TaperedCapsule.GetOneSidedPlaneNormal());
+									DrawLine(PDI, Center, NormalEnd, Color);
+								}
 							}
 							break;
 
diff --git a/Engine/Plugins/ChaosCloth/Source/ChaosCloth/Private/ChaosCloth/ChaosClothingSimulationCollider.cpp b/Engine/Plugins/ChaosCloth/Source/ChaosCloth/Private/ChaosCloth/ChaosClothingSimulationCollider.cpp
--- a/Engine/Plugins/ChaosCloth/Source/ChaosCloth/Private/ChaosCloth/ChaosClothingSimulationCollider.cpp
+++ b/Engine/Plugins/ChaosCloth/Source/ChaosCloth/Private/ChaosCloth/ChaosClothingSimulationCollider.cpp
@@ -31,7 +31,7 @@ namespace Chaos
 
 namespace ClothingSimulationColliderConsoleVariables
 {
-	TAutoConsoleVariable<bool> CVarUseOptimizedTaperedCapsule(TEXT("p.ChaosCloth.UseOptimizedTaperedCapsule"), true, TEXT("Use the optimized TaperedCapsule code instead of using a tapered cylinder and two spheres"));
+	TAutoConsoleVariable<bool> CVarUseOptimizedTaperedCapsule(TEXT("p.ChaosCloth.UseOptimizedTaperedCapsule"), true, TEXT("Use the optimized TaperedCapsule code instead of using a tapered cylinder and two spheres. Note that only optimized TaperedCapsule supports split tapered capsules"));
 }
 
 struct FClothingSimulationCollider::FLODData
@@ -172,13 +172,13 @@ void FClothingSimulationCollider::FLODData::Add(
 
 			BaseTransforms[Index] = Softs::FSolverRigidTransform3::Identity;
 
-			if (Axis.SizeSquared() < SMALL_NUMBER)
+			if (Axis.SizeSquared() < SMALL_NUMBER && Connection.OneSidedPlaneNormal.IsNearlyZero())
 			{
 				// Sphere
 				Solver->SetCollisionGeometry(CollisionRangeId, Index,
 					MakeImplicitObjectPtr<FSphere>(Center, MaxRadius));
 			}
-			else if (MaxRadius - MinRadius < KINDA_SMALL_NUMBER)
+			else if (MaxRadius - MinRadius < KINDA_SMALL_NUMBER && Connection.OneSidedPlaneNormal.IsNearlyZero())
 			{
 				// Capsule
 				Solver->SetCollisionGeometry(CollisionRangeId, Index,
@@ -189,7 +189,7 @@ void FClothingSimulationCollider::FLODData::Add(
 				if (ClothingSimulationColliderConsoleVariables::CVarUseOptimizedTaperedCapsule.GetValueOnAnyThread())
 				{
 					Solver->SetCollisionGeometry(CollisionRangeId, Index,
-						MakeImplicitObjectPtr<FTaperedCapsule>(P0, P1, Radius0, Radius1));
+						MakeImplicitObjectPtr<FTaperedCapsule>(P0, P1, Radius0, Radius1, Connection.OneSidedPlaneNormal));
 				}
 				else
 				{
@@ -633,7 +633,7 @@ void FClothingSimulationCollider::ExtractPhysicsAssetCollision(
 			// Add tapered capsules
 			for (const FKTaperedCapsuleElem& TaperedCapsuleElem : AggGeom.TaperedCapsuleElems)
 			{
-				if (TaperedCapsuleElem.Length == 0)
+				if (TaperedCapsuleElem.Length == 0 && !TaperedCapsuleElem.bOneSidedCollision)
 				{
 					// Add extracted sphere collision data
 					FClothCollisionPrim_Sphere Sphere;
@@ -660,6 +660,10 @@ void FClothingSimulationCollider::ExtractPhysicsAssetCollision(
 					FClothCollisionPrim_SphereConnection SphereConnection;
 					SphereConnection.SphereIndices[0] = ClothCollisionData.Spheres.Add(Sphere0);
 					SphereConnection.SphereIndices[1] = ClothCollisionData.Spheres.Add(Sphere1);
+					if (TaperedCapsuleElem.bOneSidedCollision)
+					{
+						SphereConnection.OneSidedPlaneNormal = TaperedCapsuleElem.Rotation.RotateVector(FVector::ForwardVector);
+					}
 					ClothCollisionData.SphereConnections.Add(SphereConnection);
 				}
 			}
diff --git a/Engine/Source/Runtime/ClothingSystemRuntimeInterface/Public/ClothCollisionPrim.h b/Engine/Source/Runtime/ClothingSystemRuntimeInterface/Public/ClothCollisionPrim.h
--- a/Engine/Source/Runtime/ClothingSystemRuntimeInterface/Public/ClothCollisionPrim.h
+++ b/Engine/Source/Runtime/ClothingSystemRuntimeInterface/Public/ClothCollisionPrim.h
@@ -38,14 +38,19 @@ struct FClothCollisionPrim_SphereConnection
 {
 	GENERATED_BODY()
 
-	FClothCollisionPrim_SphereConnection(int32 SphereIndex0 = INDEX_NONE, int32 SphereIndex1 = INDEX_NONE)
+	FClothCollisionPrim_SphereConnection(int32 SphereIndex0 = INDEX_NONE, int32 SphereIndex1 = INDEX_NONE, const FVector& InOneSidedPlaneNormal = FVector::ZeroVector)
 	{
 		SphereIndices[0] = SphereIndex0;
 		SphereIndices[1] = SphereIndex1;
+		OneSidedPlaneNormal = InOneSidedPlaneNormal;
 	}
 
 	UPROPERTY()
 	int32 SphereIndices[2];
+
+	/* Zero vector means this is not a one-sided capsule.*/
+	UPROPERTY()
+	FVector OneSidedPlaneNormal = FVector::ZeroVector;
 };
 
 /** Data for a convex face. */
diff --git a/Engine/Source/Runtime/Engine/Classes/PhysicsEngine/TaperedCapsuleElem.h b/Engine/Source/Runtime/Engine/Classes/PhysicsEngine/TaperedCapsuleElem.h
--- a/Engine/Source/Runtime/Engine/Classes/PhysicsEngine/TaperedCapsuleElem.h
+++ b/Engine/Source/Runtime/Engine/Classes/PhysicsEngine/TaperedCapsuleElem.h
@@ -36,6 +36,10 @@ struct FKTaperedCapsuleElem : public FKShapeElem
 	UPROPERTY(Category= Capsule, EditAnywhere)
 	float Length;
 
+	/** (Cloth-only) Treat as one-sided collider, where all collisions are pushed to the +x side.*/
+	UPROPERTY(Category = Capsule, EditAnywhere)
+	bool bOneSidedCollision = false;
+
 	ENGINE_API FKTaperedCapsuleElem();
 	ENGINE_API FKTaperedCapsuleElem( float InRadius0, float InRadius1, float InLength );
 	ENGINE_API FKTaperedCapsuleElem(const FKTaperedCapsuleElem&);
@@ -47,7 +51,8 @@ struct FKTaperedCapsuleElem : public FKShapeElem
 			LHS.Rotation == RHS.Rotation &&
 			LHS.Radius0 == RHS.Radius0 &&
 			LHS.Radius1 == RHS.Radius1 &&
-			LHS.Length == RHS.Length );
+			LHS.Length == RHS.Length &&
+			LHS.bOneSidedCollision == RHS.bOneSidedCollision);
 	};
 
 	// Utility function that builds an FTransform from the current data
@@ -88,7 +93,7 @@ struct FKTaperedCapsuleElem : public FKShapeElem
 	 * Extent geometry endpoints not necessarily coplanar with sphere origins (uses hull horizon)
 	 * Otherwise uses the great-circle cap assumption.
 	 */
-	ENGINE_API static void DrawTaperedCapsuleSides(FPrimitiveDrawInterface* PDI, const FTransform& ElemTM, const FVector& InCenter0, const FVector& InCenter1, float InRadius0, float InRadius1, const FColor& Color);
+	ENGINE_API static void DrawTaperedCapsuleSides(FPrimitiveDrawInterface* PDI, const FTransform& ElemTM, const FVector& InCenter0, const FVector& InCenter1, float InRadius0, float InRadius1, const FColor& Color, bool bInSplitTaperedCylinder = false);
 
 	ENGINE_API static EAggCollisionShape::Type StaticShapeType;
 };
\ No newline at end of file
diff --git a/Engine/Source/Runtime/Engine/Private/PhysicsEngine/PhysDrawing.cpp b/Engine/Source/Runtime/Engine/Private/PhysicsEngine/PhysDrawing.cpp
--- a/Engine/Source/Runtime/Engine/Private/PhysicsEngine/PhysDrawing.cpp
+++ b/Engine/Source/Runtime/Engine/Private/PhysicsEngine/PhysDrawing.cpp
@@ -426,8 +426,10 @@ void FKSphylElem::DrawElemSolid(FPrimitiveDrawInterface* PDI, const FTransform&
 // FKTaperedCapsuleElem
 /////////////////////////////////////////////////////////////////////////////////////
 
-void FKTaperedCapsuleElem::DrawTaperedCapsuleSides(FPrimitiveDrawInterface* PDI, const FTransform& ElemTM, const FVector& InCenter0, const FVector& InCenter1, float InRadius0, float InRadius1, const FColor& Color)
+void FKTaperedCapsuleElem::DrawTaperedCapsuleSides(FPrimitiveDrawInterface* PDI, const FTransform& ElemTM, const FVector& InCenter0, const FVector& InCenter1, float InRadius0, float InRadius1, const FColor& Color, bool bInSplitTaperedCylinder)
 {
+	const FColor HalfTaperedColor = bInSplitTaperedCylinder ? (FLinearColor(Color) * .2f).ToFColor(true) : Color;
+
 	// Draws just the sides of a tapered capsule specified by provided Spheres that can have different radii.  Does not draw the spheres, just the sleeve.
 	// Extent geometry endpoints not necessarily coplanar with sphere origins (uses hull horizon)
 	// Otherwise uses the great-circle cap assumption.
@@ -442,23 +444,35 @@ void FKTaperedCapsuleElem::DrawTaperedCapsuleSides(FPrimitiveDrawInterface* PDI,
 	FQuat CapsuleOrientation = ElemTM.GetRotation();
 	float OffsetZ = (InRadius1 - InRadius0) / Distance;
 	float ScaleXY = FMath::Sqrt(1.0f - FMath::Square(OffsetZ));
-	FVector StartVertex = CapsuleOrientation.RotateVector(FVector(ScaleXY, 0, OffsetZ));
+
+	FVector StartVertex = CapsuleOrientation.RotateVector(FVector(0, ScaleXY, OffsetZ));
 	FVector VertexPrevious = StartVertex;
-	for(float Angle = AngleIncrementCircle; Angle <= 360.0f; Angle += AngleIncrementCircle)  // iterate over unit circle about capsule's major axis (which is orientation.AxisZ)
+	for(float Angle = AngleIncrementCircle; Angle <= 180.0f; Angle += AngleIncrementCircle)  // iterate over unit circle about capsule's major axis (which is orientation.AxisZ)
 	{
-		FVector VertexCurrent = CapsuleOrientation.RotateVector(FVector(FMath::Cos(FMath::DegreesToRadians(Angle))*ScaleXY, FMath::Sin(FMath::DegreesToRadians(Angle))*ScaleXY, OffsetZ));
-		PDI->DrawLine(InCenter0 + VertexPrevious * InRadius0, InCenter0 + VertexCurrent * InRadius0, Color, SDPG_World);  // cap-circle segment on sphere S0
-		PDI->DrawLine(InCenter1 + VertexPrevious * InRadius1, InCenter1 + VertexCurrent * InRadius1, Color, SDPG_World);  // cap-circle segment on sphere S1
+		FVector VertexCurrent = CapsuleOrientation.RotateVector(FVector(FMath::Sin(FMath::DegreesToRadians(Angle)) * ScaleXY, FMath::Cos(FMath::DegreesToRadians(Angle))*ScaleXY, OffsetZ));
+		PDI->DrawTranslucentLine(InCenter0 + VertexPrevious * InRadius0, InCenter0 + VertexCurrent * InRadius0, Color, SDPG_World);  // cap-circle segment on sphere S0
+		PDI->DrawTranslucentLine(InCenter1 + VertexPrevious * InRadius1, InCenter1 + VertexCurrent * InRadius1, Color, SDPG_World);  // cap-circle segment on sphere S1
 		VertexPrevious = VertexCurrent;
 	}
 
-	VertexPrevious = StartVertex;
-	for(float Angle = AngleIncrementSides; Angle <= 360.0f; Angle += AngleIncrementSides)  // iterate over unit circle about capsule's major axis (which is orientation.AxisZ)
+	for (float Angle = 180.f + AngleIncrementCircle; Angle <= 360.0f; Angle += AngleIncrementCircle)  // iterate over unit circle about capsule's major axis (which is orientation.AxisZ)
 	{
-		FVector VertexCurrent = CapsuleOrientation.RotateVector(FVector(FMath::Cos(FMath::DegreesToRadians(Angle))*ScaleXY, FMath::Sin(FMath::DegreesToRadians(Angle))*ScaleXY, OffsetZ));
-		PDI->DrawLine(InCenter0 + VertexCurrent  * InRadius0, InCenter1 + VertexCurrent * InRadius1, Color, SDPG_World);  // capsule side segment between spheres
+		FVector VertexCurrent = CapsuleOrientation.RotateVector(FVector(FMath::Sin(FMath::DegreesToRadians(Angle)) * ScaleXY, FMath::Cos(FMath::DegreesToRadians(Angle)) * ScaleXY, OffsetZ));
+		PDI->DrawTranslucentLine(InCenter0 + VertexPrevious * InRadius0, InCenter0 + VertexCurrent * InRadius0, HalfTaperedColor, SDPG_World);  // cap-circle segment on sphere S0
+		PDI->DrawTranslucentLine(InCenter1 + VertexPrevious * InRadius1, InCenter1 + VertexCurrent * InRadius1, HalfTaperedColor, SDPG_World);  // cap-circle segment on sphere S1
 		VertexPrevious = VertexCurrent;
 	}
+
+	for (float Angle = 0; Angle <= 180.f; Angle += AngleIncrementSides)  // iterate over unit circle about capsule's major axis (which is orientation.AxisZ)
+	{
+		FVector VertexCurrent = CapsuleOrientation.RotateVector(FVector(FMath::Sin(FMath::DegreesToRadians(Angle)) * ScaleXY, FMath::Cos(FMath::DegreesToRadians(Angle)) * ScaleXY, OffsetZ));
+		PDI->DrawTranslucentLine(InCenter0 + VertexCurrent * InRadius0, InCenter1 + VertexCurrent * InRadius1, Color, SDPG_World);  // capsule side segment between spheres
+	}
+	for(float Angle = 180.f + AngleIncrementSides; Angle < 360.0f; Angle += AngleIncrementSides)  // iterate over unit circle about capsule's major axis (which is orientation.AxisZ)
+	{
+		FVector VertexCurrent = CapsuleOrientation.RotateVector(FVector(FMath::Sin(FMath::DegreesToRadians(Angle)) * ScaleXY, FMath::Cos(FMath::DegreesToRadians(Angle)) * ScaleXY, OffsetZ));
+		PDI->DrawTranslucentLine(InCenter0 + VertexCurrent  * InRadius0, InCenter1 + VertexCurrent * InRadius1, HalfTaperedColor, SDPG_World);  // capsule side segment between spheres
+	}
 }
 
 void FKTaperedCapsuleElem::DrawElemWire(class FPrimitiveDrawInterface* PDI, const FTransform& ElemTM, float Scale, const FColor Color) const
@@ -473,6 +487,7 @@ void FKTaperedCapsuleElem::DrawElemSolid(class FPrimitiveDrawInterface* PDI, con
 
 void FKTaperedCapsuleElem::DrawElemWire(class FPrimitiveDrawInterface* PDI, const FTransform& ElemTM, const FVector& Scale3D, const FColor Color) const
 {
+	const FColor HalfTaperedColor = bOneSidedCollision ? ( FLinearColor(Color) * .2f).ToFColor(true) : Color;
 	const FVector Origin = ElemTM.GetLocation();
 	const FVector XAxis = ElemTM.GetScaledAxis(EAxis::X);
 	const FVector YAxis = ElemTM.GetScaledAxis(EAxis::Y);
@@ -494,14 +509,20 @@ void FKTaperedCapsuleElem::DrawElemWire(class FPrimitiveDrawInterface* PDI, cons
 		// Degenerate or one end encompasses the other - draw a sphere
 		if(ScaledRadius0 > ScaledRadius1)
 		{
-			DrawCircle(PDI, TopEnd, XAxis, YAxis, Color, ScaledRadius0, DrawCollisionSides, SDPG_World);
-			DrawCircle(PDI, TopEnd, XAxis, ZAxis, Color, ScaledRadius0, DrawCollisionSides, SDPG_World);
+			DrawArc(PDI, TopEnd, XAxis, YAxis, -90.f, 90.f, ScaledRadius0, DrawCollisionSides, Color, SDPG_World);
+			DrawArc(PDI, TopEnd, XAxis, ZAxis, -90.f, 90.f, ScaledRadius0, DrawCollisionSides, Color, SDPG_World);
+
+			DrawArc(PDI, TopEnd, XAxis, YAxis, 90.f, 270.f, ScaledRadius0, DrawCollisionSides, HalfTaperedColor, SDPG_World);
+			DrawArc(PDI, TopEnd, XAxis, ZAxis, 90.f, 270.f, ScaledRadius0, DrawCollisionSides, HalfTaperedColor, SDPG_World);
 			DrawCircle(PDI, TopEnd, YAxis, ZAxis, Color, ScaledRadius0, DrawCollisionSides, SDPG_World);
 		}
 		else
 		{
-			DrawCircle(PDI, BottomEnd, XAxis, YAxis, Color, ScaledRadius1, DrawCollisionSides, SDPG_World);
-			DrawCircle(PDI, BottomEnd, XAxis, ZAxis, Color, ScaledRadius1, DrawCollisionSides, SDPG_World);
+			DrawArc(PDI, BottomEnd, XAxis, YAxis, -90.f, 90.f, ScaledRadius1, DrawCollisionSides, Color, SDPG_World);
+			DrawArc(PDI, BottomEnd, XAxis, ZAxis, -90.f, 90.f, ScaledRadius1, DrawCollisionSides, Color, SDPG_World);
+
+			DrawArc(PDI, BottomEnd, XAxis, YAxis, 90.f, 270.f, ScaledRadius1, DrawCollisionSides, HalfTaperedColor, SDPG_World);
+			DrawArc(PDI, BottomEnd, XAxis, ZAxis, 90.f, 270.f, ScaledRadius1, DrawCollisionSides, HalfTaperedColor, SDPG_World);
 			DrawCircle(PDI, BottomEnd, YAxis, ZAxis, Color, ScaledRadius1, DrawCollisionSides, SDPG_World);
 		}
 	}
@@ -509,11 +530,13 @@ void FKTaperedCapsuleElem::DrawElemWire(class FPrimitiveDrawInterface* PDI, cons
 	{
 		if(ScaledRadius0 > ScaledRadius1)
 		{
-			DrawCircle(PDI, TopEnd, XAxis, YAxis, Color, ScaledRadius0, DrawCollisionSides, SDPG_World);
+			DrawArc(PDI, TopEnd, XAxis, YAxis, -90.f, 90.f, ScaledRadius0, DrawCollisionSides, Color, SDPG_World);
+			DrawArc(PDI, TopEnd, XAxis, YAxis, 90.f, 270.f, ScaledRadius0, DrawCollisionSides, HalfTaperedColor, SDPG_World);
 		}
 		else
 		{
-			DrawCircle(PDI, BottomEnd, XAxis, YAxis, Color, ScaledRadius1, DrawCollisionSides, SDPG_World);
+			DrawArc(PDI, BottomEnd, XAxis, YAxis, -90.f, 90.f, ScaledRadius1, DrawCollisionSides, Color, SDPG_World);
+			DrawArc(PDI, BottomEnd, XAxis, YAxis, 90.f, 270.f, ScaledRadius1, DrawCollisionSides, HalfTaperedColor, SDPG_World);
 		}	
 
 		const FVector NegZAxis = -ZAxis;
@@ -521,15 +544,17 @@ void FKTaperedCapsuleElem::DrawElemWire(class FPrimitiveDrawInterface* PDI, cons
 		const float OffsetZ0 = ((ScaledRadius0 - ScaledRadius1) * ScaledRadius0) / Distance;
 		float TopLimit = 90.0f - FMath::RadiansToDegrees(FMath::Asin(OffsetZ0 / ScaledRadius0));
 		DrawArc(PDI, TopEnd, NegZAxis, YAxis, TopLimit, 360.0f-TopLimit, ScaledRadius0, DrawCollisionSides, Color, SDPG_World);
-		DrawArc(PDI, TopEnd, NegZAxis, XAxis, TopLimit, 360.0f-TopLimit, ScaledRadius0, DrawCollisionSides, Color, SDPG_World);
+		DrawArc(PDI, TopEnd, NegZAxis, XAxis, TopLimit, 180.f, ScaledRadius0, DrawCollisionSides, Color, SDPG_World);
+		DrawArc(PDI, TopEnd, NegZAxis, XAxis, 180.f,360.0f - TopLimit , ScaledRadius0, DrawCollisionSides, HalfTaperedColor, SDPG_World);
 
 		const float OffsetZ1 = ((ScaledRadius0 - ScaledRadius1) * ScaledRadius1) / Distance;
 		float BottomLimit =  180.0f - TopLimit;
  		DrawArc(PDI, BottomEnd, ZAxis, YAxis, BottomLimit, 360.0f-BottomLimit, ScaledRadius1, DrawCollisionSides, Color, SDPG_World);
- 		DrawArc(PDI, BottomEnd, ZAxis, XAxis, BottomLimit, 360.0f-BottomLimit, ScaledRadius1, DrawCollisionSides, Color, SDPG_World);
+		DrawArc(PDI, BottomEnd, ZAxis, XAxis, BottomLimit, 180.f, ScaledRadius1, DrawCollisionSides, Color, SDPG_World);
+		DrawArc(PDI, BottomEnd, ZAxis, XAxis, 180.f, 360.0f - BottomLimit, ScaledRadius1, DrawCollisionSides, HalfTaperedColor, SDPG_World);
 
 		// Draw connecty lines
-		DrawTaperedCapsuleSides(PDI, ElemTM, TopEnd, BottomEnd, ScaledRadius0, ScaledRadius1, Color);
+		DrawTaperedCapsuleSides(PDI, ElemTM, TopEnd, BottomEnd, ScaledRadius0, ScaledRadius1, Color, bOneSidedCollision);
 	}
 }
 
diff --git a/Engine/Source/Runtime/Engine/Private/SkeletalMeshComponentPhysics.cpp b/Engine/Source/Runtime/Engine/Private/SkeletalMeshComponentPhysics.cpp
--- a/Engine/Source/Runtime/Engine/Private/SkeletalMeshComponentPhysics.cpp
+++ b/Engine/Source/Runtime/Engine/Private/SkeletalMeshComponentPhysics.cpp
@@ -2990,6 +2990,10 @@ void USkeletalMeshComponent::ExtractCollisionsForCloth(
 						FClothCollisionPrim_SphereConnection Connection;
 						Connection.SphereIndices[0] = ClothCollisionSource.CachedSpheres.Num() - 2;
 						Connection.SphereIndices[1] = ClothCollisionSource.CachedSpheres.Num() - 1;
+						if (TaperedCapsuleElems.bOneSidedCollision)
+						{
+							Connection.OneSidedPlaneNormal = TaperedCapsuleElems.Rotation.RotateVector(FVector::ForwardVector);
+						}
 
 						ClothCollisionSource.CachedSphereConnections.Add(Connection);
 					}
diff --git a/Engine/Source/Runtime/Experimental/Chaos/Private/Chaos/PBDSoftBodyCollisionConstraint.cpp b/Engine/Source/Runtime/Experimental/Chaos/Private/Chaos/PBDSoftBodyCollisionConstraint.cpp
--- a/Engine/Source/Runtime/Experimental/Chaos/Private/Chaos/PBDSoftBodyCollisionConstraint.cpp
+++ b/Engine/Source/Runtime/Experimental/Chaos/Private/Chaos/PBDSoftBodyCollisionConstraint.cpp
@@ -71,6 +71,24 @@ namespace Chaos::Softs {
 			FSolverVec3 ColliderVelocityAtPoint;
 			ApplyFriction(P, X, NormalWorld, MaxFrictionCorrection, Dt, CollisionX, CollisionV, CollisionW, ColliderVelocityAtPoint);
 		}
+
+		static void ReflectOneSidedCollision(const FSolverVec3& P, const FSolverVec3& OneSidedPlaneNormal, const FSolverVec3& SplitOrigin, FSolverReal& Penetration, FSolverVec3& ImplicitNormal)
+		{
+			check(Penetration > (FSolverReal)0.f);
+
+			FSolverVec3 PNew = P + Penetration * ImplicitNormal;
+			const FSolverReal SplitAxisProj = (PNew - SplitOrigin).Dot(OneSidedPlaneNormal);
+			if (SplitAxisProj >= 0.)
+			{
+				return;
+			}
+
+			PNew -= 2. * SplitAxisProj * OneSidedPlaneNormal;
+			ImplicitNormal = PNew - P;
+			Penetration = ImplicitNormal.Length();
+			ImplicitNormal = Penetration > UE_SMALL_NUMBER ? ImplicitNormal / Penetration : FSolverVec3::ZeroVector;
+			return;
+		}
 	}
 
 // Helper function to call PhiWithNormal and return data to ISPC
@@ -431,13 +449,24 @@ void FPBDSoftBodyCollisionConstraintBase::ApplySimpleInternal(FSolverParticlesRa
 					}
 
 					const FSolverRigidTransform3 Frame(CollisionParticlesRange.GetX(CollisionIndex), CollisionParticlesRange.R(CollisionIndex));
-					const FVec3 RigidSpacePosition(Frame.InverseTransformPosition(X[Index]));  // PhiWithNormal requires FReal based arguments
-					FVec3 ImplicitNormal;
-					const FSolverReal Phi = (FSolverReal)Geometry->PhiWithNormal(RigidSpacePosition, ImplicitNormal);
-					const FSolverReal Penetration = SoftBodyCollisionThickness + CollisionThickness - Phi;
+					const FVec3 RigidSpacePosition(Frame.InverseTransformPosition(PAndInvM[Index].P));  // PhiWithNormal requires FReal based arguments
+					FVec3 ImplicitNormalReal;
+					const FSolverReal Phi = (FSolverReal)Geometry->PhiWithNormal(RigidSpacePosition, ImplicitNormalReal);
+					FSolverVec3 ImplicitNormal(ImplicitNormalReal);
+					FSolverReal Penetration = SoftBodyCollisionThickness + CollisionThickness - Phi;
 					if (Penetration > (FSolverReal)0.)
 					{
-						const FSolverVec3 NormalWorld = Frame.TransformVector(FSolverVec3(ImplicitNormal));
+						// Split capsules always push out in the OneSidedPlaneNormal direction.
+						if (CollisionType == Chaos::ImplicitObjectType::TaperedCapsule)
+						{
+							const FTaperedCapsule& Capsule = Geometry->GetObjectChecked<FTaperedCapsule>();
+							if (Capsule.IsOneSided())
+							{
+								Private::ReflectOneSidedCollision(FSolverVec3(RigidSpacePosition), Capsule.GetOneSidedPlaneNormalf(), Capsule.GetOriginf(), Penetration, ImplicitNormal);
+							}
+						}
+
+						const FSolverVec3 NormalWorld = Frame.TransformVector(ImplicitNormal);
 						if constexpr (bLockAndWriteContacts)
 						{
 							FScopeLock Lock(&DebugMutex);
diff --git a/Engine/Source/Runtime/Experimental/Chaos/Private/Chaos/PerParticlePBDCollisionConstraint.cpp b/Engine/Source/Runtime/Experimental/Chaos/Private/Chaos/PerParticlePBDCollisionConstraint.cpp
--- a/Engine/Source/Runtime/Experimental/Chaos/Private/Chaos/PerParticlePBDCollisionConstraint.cpp
+++ b/Engine/Source/Runtime/Experimental/Chaos/Private/Chaos/PerParticlePBDCollisionConstraint.cpp
@@ -18,7 +18,27 @@ static int32 Chaos_PerParticleCollision_ISPC_ParallelBatchSize = 128;
 FAutoConsoleVariableRef CVarChaosPerParticleCollisionISPCParallelBatchSize(TEXT("p.Chaos.PerParticleCollision.ISPC.ParallelBatchSize"), Chaos_PerParticleCollision_ISPC_ParallelBatchSize, TEXT("Parallel batch size for ISPC"));
 #endif
 
-namespace Chaos::Softs {
+namespace Chaos::Softs 
+{
+namespace Private
+{
+	static void ReflectOneSidedCollision(const FSolverVec3& P, const FSolverVec3& OneSidedPlaneNormal, const FSolverVec3& SplitOrigin, FSolverReal& Penetration, FSolverVec3& ImplicitNormal)
+	{
+		check(Penetration > (FSolverReal)0.f);
+
+		FSolverVec3 PNew = P + Penetration * ImplicitNormal;
+		const FSolverReal SplitAxisProj = (PNew - SplitOrigin).Dot(OneSidedPlaneNormal);
+		if (SplitAxisProj >= 0.)
+		{
+			return;
+		}
+
+		PNew -= 2. * SplitAxisProj * OneSidedPlaneNormal;
+		ImplicitNormal = PNew - P;
+		Penetration = ImplicitNormal.Length();
+		ImplicitNormal = Penetration > UE_SMALL_NUMBER ? ImplicitNormal / Penetration : FSolverVec3::ZeroVector;
+	}
+}
 
 // Helper function to call PhiWithNormal and return data to ISPC
 extern "C" void GetPhiWithNormal(const uint8* CollisionParticles, const FSolverReal* InV, FSolverReal* Normal, FSolverReal* Phi, const int32 i, const int32 ProgramCount, const int32 Mask)
@@ -133,4 +153,169 @@ void FPerParticlePBDCollisionConstraint::ApplyHelperISPC(FSolverParticles& InPar
 	}
 }
 
+template<bool bLockAndWriteContacts>
+void FPerParticlePBDCollisionConstraint::ApplyHelper(FSolverParticles& Particles, const FSolverReal Dt, const int32 Offset, const int32 Range) const
+{
+	const uint32 DynamicGroupId = MDynamicGroupIds[Offset];  // Particle group Id, must be the same across the entire range
+	const FSolverReal PerGroupFriction = MPerGroupFriction[DynamicGroupId];
+	const FSolverReal PerGroupThickness = MPerGroupThickness[DynamicGroupId];
+
+	if (PerGroupFriction > (FSolverReal)UE_KINDA_SMALL_NUMBER)
+	{
+		PhysicsParallelFor(Range - Offset, [this, &Particles, Dt, Offset, DynamicGroupId, PerGroupFriction, PerGroupThickness](int32 i)
+			{
+				const int32 Index = Offset + i;
+
+				if (Particles.InvM(Index) == (FSolverReal)0.)
+				{
+					return;  // Continue
+				}
+
+				MCollisionParticlesActiveView.SequentialFor([this, &Particles, &Dt, &Index, DynamicGroupId, PerGroupFriction, PerGroupThickness](FSolverCollisionParticles& CollisionParticles, int32 i)
+					{
+						const uint32 KinematicGroupId = MKinematicGroupIds[i];  // Collision group Id
+
+						if (KinematicGroupId != (uint32)INDEX_NONE && DynamicGroupId != KinematicGroupId)
+						{
+							return; // Bail out if the collision groups doesn't match the particle group id, or use INDEX_NONE (= global collision that affects all particle)
+						}
+						const FImplicitObjectPtr& Geometry = CollisionParticles.GetGeometry(i);
+						const FSolverRigidTransform3 Frame(CollisionParticles.GetX(i), CollisionParticles.GetR(i));
+						const FVec3 RigidSpacePosition(Frame.InverseTransformPosition(Particles.P(Index)));  // PhiWithNormal requires FReal based arguments
+						FVec3 ImplicitNormal;                                                                // since implicits don't use FSolverReal
+						const FSolverReal Phi = (FSolverReal)Geometry->PhiWithNormal(RigidSpacePosition, ImplicitNormal);
+						FSolverVec3 Normal(ImplicitNormal);
+						FSolverReal Penetration = PerGroupThickness - Phi; // This is related to the Normal impulse
+
+						if (Penetration > (FSolverReal)0.)
+						{
+							// Split capsules always push out in the OneSidedPlaneNormal direction.
+							if (const FTaperedCapsule* const Capsule = Geometry->GetObject<FTaperedCapsule>())
+							{
+								if (Capsule->IsOneSided())
+								{
+									Private::ReflectOneSidedCollision(FSolverVec3(RigidSpacePosition), Capsule->GetOneSidedPlaneNormalf(), Capsule->GetOriginf(), Penetration, Normal);
+								}
+							}
+							const FSolverVec3 NormalWorld = Frame.TransformVector(Normal);
+
+							if (bLockAndWriteContacts)
+							{
+								check(Mutex);
+								checkSlow(Contacts);
+								checkSlow(Normals);
+								checkSlow(Phis);
+								FScopeLock Lock(Mutex);
+								Contacts->Emplace(Particles.P(Index));
+								Normals->Emplace(NormalWorld);
+								Phis->Emplace(Phi);
+							}
+
+							Particles.P(Index) += Penetration * NormalWorld;
+
+							if (bFastPositionBasedFriction)
+							{
+								FSolverVec3 VectorToPoint = Particles.GetP(Index) - CollisionParticles.GetX(i);
+								const FSolverVec3 RelativeDisplacement = (Particles.GetP(Index) - Particles.GetX(Index)) - (CollisionParticles.V(i) + FSolverVec3::CrossProduct(CollisionParticles.W(i), VectorToPoint)) * Dt; // This corresponds to the tangential velocity multiplied by dt (friction will drive this to zero if it is high enough)
+								const FSolverVec3 RelativeDisplacementTangent = RelativeDisplacement - FSolverVec3::DotProduct(RelativeDisplacement, NormalWorld) * NormalWorld; // Project displacement into the tangential plane
+								const FSolverReal RelativeDisplacementTangentLength = RelativeDisplacementTangent.Size();
+								if (RelativeDisplacementTangentLength >= UE_SMALL_NUMBER)
+								{
+									const FSolverReal PositionCorrection = FMath::Min<FSolverReal>(Penetration * PerGroupFriction, RelativeDisplacementTangentLength);
+									const FSolverReal CorrectionRatio = PositionCorrection / RelativeDisplacementTangentLength;
+									Particles.P(Index) -= CorrectionRatio * RelativeDisplacementTangent;
+								}
+							}
+							else
+							{
+								// Note, to fix: Only use fast position based friction for now, since adding to TMaps here is not thread safe when calling Apply on multiple threads (will cause crash)
+								FVelocityConstraint Constraint;
+								FSolverVec3 VectorToPoint = Particles.GetP(Index) - CollisionParticles.GetX(i);
+								Constraint.Velocity = CollisionParticles.V(i) + FSolverVec3::CrossProduct(CollisionParticles.W(i), VectorToPoint);
+								Constraint.Normal = Frame.TransformVector(Normal);
+
+								MVelocityConstraints.Add(Index, Constraint);
+							}
+						}
+					});
+			});
+	}
+	else
+	{
+		PhysicsParallelFor(Range - Offset, [this, &Particles, Dt, Offset, DynamicGroupId, PerGroupFriction, PerGroupThickness](int32 i)
+			{
+				const int32 Index = Offset + i;
+
+				if (Particles.InvM(Index) == 0)
+				{
+					return;  // Continue
+				}
+
+				MCollisionParticlesActiveView.SequentialFor([this, &Particles, &Dt, &Index, DynamicGroupId, PerGroupFriction, PerGroupThickness](FSolverCollisionParticles& CollisionParticles, int32 i)
+					{
+						const uint32 KinematicGroupId = MKinematicGroupIds[i];  // Collision group Id
+
+						if (KinematicGroupId != (uint32)INDEX_NONE && DynamicGroupId != KinematicGroupId)
+						{
+							return; // Bail out if the collision groups doesn't match the particle group id, or use INDEX_NONE (= global collision that affects all particle)
+						}
+						const FImplicitObjectPtr& Geometry = CollisionParticles.GetGeometry(i);
+						const FSolverRigidTransform3 Frame(CollisionParticles.GetX(i), CollisionParticles.GetR(i));
+						const FVec3 RigidSpacePosition(Frame.InverseTransformPosition(Particles.P(Index)));  // PhiWithNormal requires FReal based arguments
+						FVec3 ImplicitNormal;                                                                // since implicits don't use FSolverReal
+						const FSolverReal Phi = (FSolverReal)Geometry->PhiWithNormal(RigidSpacePosition, ImplicitNormal);
+						FSolverVec3 Normal(ImplicitNormal);
+						FSolverReal Penetration = PerGroupThickness - Phi; // This is related to the Normal impulse
+
+						if (Penetration > (FSolverReal)0.)
+						{
+							// Split capsules always push out in the OneSidedPlaneNormal direction.
+							if (const FTaperedCapsule* const Capsule = Geometry->GetObject<FTaperedCapsule>())
+							{
+								if (Capsule->IsOneSided())
+								{
+									Private::ReflectOneSidedCollision(FSolverVec3(RigidSpacePosition), Capsule->GetOneSidedPlaneNormalf(), Capsule->GetOriginf(), Penetration, Normal);
+								}
+							}
+
+							const FSolverVec3 NormalWorld = Frame.TransformVector(Normal);
+							if (bLockAndWriteContacts)
+							{
+								check(Mutex);
+								checkSlow(Contacts);
+								checkSlow(Normals);
+								checkSlow(Phis);
+								FScopeLock Lock(Mutex);
+								Contacts->Emplace(Particles.P(Index));
+								Normals->Emplace(NormalWorld);
+								Phis->Emplace(Phi);
+							}
+
+							Particles.P(Index) += Penetration * NormalWorld;
+						}
+					});
+			});
+	}
+}
+
+
+void FPerParticlePBDCollisionConstraint::ApplyRange(FSolverParticles& Particles, const FSolverReal Dt, const int32 Offset, const int32 Range) const
+{
+	// NOTE: currently using ISPC with TWeightedLatticeImplicitObject<FLevelSet> is significantly slower than not using ISPC (largely because it has not been fully implemented)
+	if (bRealTypeCompatibleWithISPC && bChaos_PerParticleCollision_ISPC_Enabled && bFastPositionBasedFriction)
+	{
+		ApplyHelperISPC(Particles, Dt, Offset, Range);
+	}
+	else
+	{
+		if (Mutex)
+		{
+			ApplyHelper<true>(Particles, Dt, Offset, Range);
+		}
+		else
+		{
+			ApplyHelper<false>(Particles, Dt, Offset, Range);
+		}
+	}
+}
 }  // End namespace Chaos::Softs
diff --git a/Engine/Source/Runtime/Experimental/Chaos/Private/Chaos/PerParticlePBDCollisionConstraint.ispc b/Engine/Source/Runtime/Experimental/Chaos/Private/Chaos/PerParticlePBDCollisionConstraint.ispc
--- a/Engine/Source/Runtime/Experimental/Chaos/Private/Chaos/PerParticlePBDCollisionConstraint.ispc
+++ b/Engine/Source/Runtime/Experimental/Chaos/Private/Chaos/PerParticlePBDCollisionConstraint.ispc
@@ -159,9 +159,11 @@ struct FTaperedCapsule
 {
 	FVector3f Origin;
 	FVector3f Axis;
+	FVector3f OneSidedPlaneNormal;
 	float Height;
 	float Radius1;
 	float Radius2;
+	bool bIsOneSided;
 	// FAABB3 LocalBoundingBox; // this code never uses this, and we never assume to know the sizeof(FTaperedCapsule)
 };
 
@@ -217,14 +219,19 @@ static inline uniform bool IsSimpleCollider(const uniform uint8 CollisionType)
 	return true;
 }
 
-static inline float SafeNormalize(FVector3f &Direction)
+static inline float SafeNormalize(FVector3f &Direction, const uniform FVector3f& DefaultIfZero)
 {
 	const float SizeSqr = VectorSizeSquared(Direction);
 	const float Size = sqrt(SizeSqr);
-	Direction = VectorSelect((SizeSqr < KINDA_SMALL_NUMBER), FloatForwardVector, Direction / Size);
+	Direction = VectorSelect((SizeSqr < KINDA_SMALL_NUMBER), DefaultIfZero, Direction / Size);
 	return (Size < KINDA_SMALL_NUMBER) ? ZERO : Size;
 }
 
+static inline float SafeNormalize(FVector3f &Direction)
+{
+	return SafeNormalize(Direction, FloatForwardVector);
+}
+
 static inline bool VectorNormalize(FVector3f &V)
 {
 	const float SquareSum = VectorSizeSquared(V);
@@ -237,11 +244,25 @@ static inline bool VectorNormalize(FVector3f &V)
 	return false;
 }
 
-static inline void PhiWithNormalSphere(const uniform uint8 *uniform TypedObjectPtr, const FVector3f &PhiWithNormalInput, FVector3f &Normal, float &Phi, const uniform float Radius)
+static inline void ReflectOneSidedCollision(const FVector3f& P, const uniform FVector3f& OneSidedPlaneNormal, const uniform FVector3f& SplitOrigin, float& Penetration, FVector3f& ImplicitNormal)
+{
+	FVector3f PNew = P + Penetration * ImplicitNormal;
+	const float SplitAxisProj = VectorDot(PNew - SplitOrigin, OneSidedPlaneNormal);
+	if (SplitAxisProj >= 0.)
+	{
+		return;
+	}
+
+	PNew = PNew - OneSidedPlaneNormal * (2. * SplitAxisProj);
+	ImplicitNormal = PNew - P;
+	Penetration = SafeNormalize(ImplicitNormal, FloatZeroVector);
+}
+
+static inline void PenetrationWithNormalSphere(const uniform uint8 *uniform TypedObjectPtr, const FVector3f &PhiWithNormalInput, const uniform float Thickness,  FVector3f &Normal, float &Penetration, const uniform float Radius)
 {
 	const uniform FVector3f Center = *((const uniform FVector3f *uniform)&TypedObjectPtr[0]);
 	Normal = PhiWithNormalInput - Center;
-	Phi = SafeNormalize(Normal) - Radius;
+	Penetration = Thickness - (SafeNormalize(Normal) - Radius);
 }
 
 static inline void PhiWithNormalPlane(const uniform FPlaneConcrete3f *varying InPlane, const FVector3f &PhiWithNormalInput, FVector3f &Normal, float &Phi)
@@ -259,29 +280,34 @@ static inline float SignedDistancePlane(const uniform FPlaneConcrete3f *uniform
 	return VectorDot(x - P.MX, P.MNormal);
 }
 
-static inline void PhiWithNormalCapsule(const uniform uint8 *uniform TypedObjectPtr, const FVector3f &PhiWithNormalInput, FVector3f &Normal, float &Phi, const uniform float Radius)
+static inline void PenetrationWithNormalCapsule(const uniform uint8 *uniform TypedObjectPtr, const FVector3f &PhiWithNormalInput, const uniform float Thickness, FVector3f &Normal, float &Penetration, const uniform float Radius)
 {
 	const uniform Segment MSegment = *((const uniform Segment *uniform)&TypedObjectPtr[0]);
 
 	const float Dot = clamp(VectorDot(PhiWithNormalInput - MSegment.Point, MSegment.Axis), (uniform float)(0), MSegment.Length);
 	const FVector3f ProjectedPoint = MSegment.Axis * Dot + MSegment.Point;
 	Normal = PhiWithNormalInput - ProjectedPoint;
-	Phi = SafeNormalize(Normal) - Radius;
+	Penetration = Thickness - (SafeNormalize(Normal) - Radius);
 }
 
-static inline void PhiWithNormalTaperedCapsule(const uniform uint8* uniform TypedObjectPtr, const FVector3f& PhiWithNormalInput, FVector3f& Normal, float& Phi)
+static inline void PenetrationWithNormalTaperedCapsule(const uniform uint8* uniform TypedObjectPtr, const FVector3f& PhiWithNormalInput, const uniform float Thickness, FVector3f& Normal, float& Penetration)
 {
 	const uniform FTaperedCapsule CapsuleObj = *((const uniform FTaperedCapsule *uniform)&TypedObjectPtr[0]);
 
 	const float DistanceAlongAxis = clamp(VectorDot(PhiWithNormalInput - CapsuleObj.Origin, CapsuleObj.Axis), (uniform float)(0), CapsuleObj.Height);
 	const FVector3f ProjectedPoint = CapsuleObj.Axis * DistanceAlongAxis + CapsuleObj.Origin;
 	const float Alpha = DistanceAlongAxis / CapsuleObj.Height;
-	const float Radius = CapsuleObj.Radius1 * ((uniform float)(1.0) - Alpha) + CapsuleObj.Radius2 * Alpha;
+	const float Radius = select(CapsuleObj.Height > FLOAT_SMALL_NUMBER, CapsuleObj.Radius1 * ((uniform float)(1.0) - Alpha) + CapsuleObj.Radius2 * Alpha, max(CapsuleObj.Radius1, CapsuleObj.Radius2));
 	Normal = PhiWithNormalInput - ProjectedPoint;
-	Phi = SafeNormalize(Normal) - Radius;
+	Penetration = Thickness - (SafeNormalize(Normal) - Radius);
+
+	if(Penetration >= 0.f && CapsuleObj.bIsOneSided)
+	{
+		ReflectOneSidedCollision(PhiWithNormalInput, CapsuleObj.OneSidedPlaneNormal, CapsuleObj.Origin, Penetration, Normal);
+	}
 }
 
-static inline void PhiWithNormalTaperedCylinder(const uniform uint8 *uniform TypedObjectPtr, const FVector3f &PhiWithNormalInput, FVector3f &Normal, float &Phi, const uniform int SizeofFImplicitObject)
+static inline void PenetrationWithNormalTaperedCylinder(const uniform uint8 *uniform TypedObjectPtr, const FVector3f &PhiWithNormalInput, const uniform float Thickness, FVector3f &Normal, float &Penetration, const uniform int SizeofFImplicitObject)
 {
 	const uniform FTaperedCylinder CylinderObj = *((const uniform FTaperedCylinder *uniform)&TypedObjectPtr[0]);
 
@@ -294,6 +320,8 @@ static inline void PhiWithNormalTaperedCylinder(const uniform uint8 *uniform Typ
 	const float SideDistance = VectorSize(SideVector) - Radius;
 	const float TopDistance = Distance1 < Distance2 ? Distance1 : Distance2;
 
+	float Phi = FLT_MAX;
+
 	if (Distance1 < SMALL_NUMBER)
 	{
 		const FVector3f v = PhiWithNormalInput - (CylinderObj.MPlane1.MNormal * Distance1 + CylinderObj.MPlane1.MX);
@@ -336,33 +364,35 @@ static inline void PhiWithNormalTaperedCylinder(const uniform uint8 *uniform Typ
 		Normal = VectorGetSafeNormal(SideVector);
 		Phi = SideDistance;
 	}
+
+	Penetration = Thickness - Phi;
 }
 
-static inline void PhiWithNormalUnion(const uniform uint8 *uniform TypedPtr, const FVector3f &PhiWithNormalInput, FVector3f &Normal, float &Phi, const uniform int SizeofFImplicitObject, const uniform int OffsetofGeometryType, const uniform int OffsetOfMargin)
+static inline void PenetrationWithNormalUnion(const uniform uint8 *uniform TypedPtr, const FVector3f &PhiWithNormalInput, const uniform float Thickness, FVector3f &Normal, float &Penetration, const uniform int SizeofFImplicitObject, const uniform int OffsetofGeometryType, const uniform int OffsetOfMargin)
 {
 	// This is expecting a MObjects = [TaperedCylinder, Sphere, Sphere]
 	assert( ((const uniform TArray *uniform)TypedPtr)->ArrayNum == 3 );
 
 	const uniform UniquePtr *uniform MObjects = (const uniform UniquePtr *uniform)TypedPtr;
-	Phi = FLT_MAX;
+	Penetration = -FLT_MAX;
 	bool NeedsNormalize = false;
 	
 	uniform uint8 *uniform Object = MObjects->Data[0].Ptr;
 	uniform uint8 *uniform TypedObjectPtr = Object + SizeofFImplicitObject;
 
 	FVector3f NextNormal;
-	float NextPhi;
+	float NextPenetration;
 
 	assert( (Object[OffsetofGeometryType] & 0x3F) == TaperedCylinder );
-	PhiWithNormalTaperedCylinder(TypedObjectPtr, PhiWithNormalInput, NextNormal, NextPhi, SizeofFImplicitObject);
+	PenetrationWithNormalTaperedCylinder(TypedObjectPtr, PhiWithNormalInput, Thickness, NextNormal, NextPenetration, SizeofFImplicitObject);
 
-	if (NextPhi < Phi)
+	if (NextPenetration > Penetration)
 	{
-		Phi = NextPhi;
+		Penetration = NextPenetration;
 		Normal = NextNormal;
 		NeedsNormalize = false;
 	}
-	else if (NextPhi == Phi)
+	else if (Penetration == NextPenetration)
 	{
 		Normal = Normal + NextNormal;
 		NeedsNormalize = true;
@@ -373,15 +403,15 @@ static inline void PhiWithNormalUnion(const uniform uint8 *uniform TypedPtr, con
 	TypedObjectPtr = Object + SizeofFImplicitObject;
 
 	assert( (Object[OffsetofGeometryType] & 0x3F) == Sphere );
-	PhiWithNormalSphere(TypedObjectPtr, PhiWithNormalInput, NextNormal, NextPhi, Margin);
-
-	if (NextPhi < Phi)
+	PenetrationWithNormalSphere(TypedObjectPtr, PhiWithNormalInput, Thickness, NextNormal, NextPenetration, Margin);
+	
+	if (NextPenetration > Penetration)
 	{
-		Phi = NextPhi;
+		Penetration = NextPenetration;
 		Normal = NextNormal;
 		NeedsNormalize = false;
 	}
-	else if (NextPhi == Phi)
+	else if (Penetration == NextPenetration)
 	{
 		Normal = Normal + NextNormal;
 		NeedsNormalize = true;
@@ -392,15 +422,15 @@ static inline void PhiWithNormalUnion(const uniform uint8 *uniform TypedPtr, con
 	TypedObjectPtr = Object + SizeofFImplicitObject;
 
 	assert( (Object[OffsetofGeometryType] & 0x3F) == Sphere );
-	PhiWithNormalSphere(TypedObjectPtr, PhiWithNormalInput, NextNormal, NextPhi, Margin);
-
-	if (NextPhi < Phi)
+	PenetrationWithNormalSphere(TypedObjectPtr, PhiWithNormalInput, Thickness, NextNormal, NextPenetration, Margin);
+	
+	if (NextPenetration > Penetration)
 	{
-		Phi = NextPhi;
+		Penetration = NextPenetration;
 		Normal = NextNormal;
 		NeedsNormalize = false;
 	}
-	else if (NextPhi == Phi)
+	else if (Penetration == NextPenetration)
 	{
 		Normal = Normal + NextNormal;
 		NeedsNormalize = true;
@@ -525,14 +555,15 @@ static inline int32 GetPlaneVertex(const uniform FConvexStructureData *uniform S
 	}
 }
 
-static inline void PhiWithNormalConvex(const uniform uint8 *uniform TypedObjectPtr, const FVector3f &PhiWithNormalInput, FVector3f &Normal, float &Phi)
+static inline void PenetrationWithNormalConvex(const uniform uint8 *uniform TypedObjectPtr, const FVector3f &PhiWithNormalInput, const uniform float Thickness, FVector3f &Normal, float &Penetration)
 {
 	const uniform FConvex *uniform Object = (const uniform FConvex *uniform)TypedObjectPtr;
 	const uniform FPlaneConcrete3f *uniform Planes = (const uniform FPlaneConcrete3f *uniform)&Object->Planes.Data[0];
 	const uniform int32 NumPlanes = Object->Planes.ArrayNum;
+
 	if (NumPlanes == 0)
 	{
-		Phi = FLT_MAX;
+		Penetration = -FLT_MAX;
 		return;
 	}
 
@@ -548,10 +579,12 @@ static inline void PhiWithNormalConvex(const uniform uint8 *uniform TypedObjectP
 			MaxPlane = Idx;
 		}
 	}
-
+	
+	float Phi = FLT_MAX;
 	PhiWithNormalPlane(&Planes[MaxPlane], PhiWithNormalInput, Normal, Phi);
 	if (Phi <= 0)
 	{
+		Penetration = Thickness - Phi;
 		return;
 	}
 
@@ -561,6 +594,7 @@ static inline void PhiWithNormalConvex(const uniform uint8 *uniform TypedObjectP
 
 	if(PlaneVerticesNum > 2 && IsVectorEqual(XOnPlane, PhiWithNormalInput))
 	{
+		Penetration = Thickness - Phi;
 		return;
 	}
 
@@ -593,37 +627,38 @@ static inline void PhiWithNormalConvex(const uniform uint8 *uniform TypedObjectP
 	{
 		Normal = (Difference) / Phi;
 	}
+	Penetration = Thickness - Phi;
 }
 
-static inline bool GetSimplePhiWithNormal(const FVector3f& PhiWithNormalInputFloat, const uniform uint8 Type, const uniform uint8 *uniform Ptr, const uniform uint8 *uniform TypedPtr, const uniform int SizeofFImplicitObject, const uniform int OffsetofGeometryType, const uniform int OffsetOfMargin, float& Phi, FVector3f& Normal)
+static inline bool GetSimplePenetrationWithNormal(const FVector3f& PhiWithNormalInputFloat, const uniform uint8 Type, const uniform uint8 *uniform Ptr, const uniform uint8 *uniform TypedPtr, const uniform int SizeofFImplicitObject, const uniform int OffsetofGeometryType, const uniform int OffsetOfMargin, const uniform float Thickness, float& Penetration, FVector3f& Normal)
 {
 	switch(Type)
 	{
 	case Sphere:
 		{
 			const uniform float Margin = *((const uniform float *uniform)(Ptr + OffsetOfMargin));
-			PhiWithNormalSphere(TypedPtr, PhiWithNormalInputFloat, Normal, Phi, Margin);
+			PenetrationWithNormalSphere(TypedPtr, PhiWithNormalInputFloat, Thickness, Normal, Penetration, Margin);
 			return true;
 		}
 	case Capsule:
 		{
 			const uniform float Margin = *((const uniform float *uniform)(Ptr + OffsetOfMargin));
-			PhiWithNormalCapsule(TypedPtr, PhiWithNormalInputFloat, Normal, Phi, Margin);
+			PenetrationWithNormalCapsule(TypedPtr, PhiWithNormalInputFloat, Thickness, Normal, Penetration, Margin);
 			return true;
 		}
 	case Union:
 		{
-			PhiWithNormalUnion(TypedPtr, PhiWithNormalInputFloat, Normal, Phi, SizeofFImplicitObject, OffsetofGeometryType, OffsetOfMargin);
+			PenetrationWithNormalUnion(TypedPtr, PhiWithNormalInputFloat, Thickness, Normal, Penetration, SizeofFImplicitObject, OffsetofGeometryType, OffsetOfMargin);
 			return true;
 		}
 	case TaperedCapsule:
 		{
-			PhiWithNormalTaperedCapsule(TypedPtr, PhiWithNormalInputFloat, Normal, Phi);
+			PenetrationWithNormalTaperedCapsule(TypedPtr, PhiWithNormalInputFloat, Thickness, Normal, Penetration);
 			return true;
 		}
 	case Convex:
 		{
-			PhiWithNormalConvex(TypedPtr, PhiWithNormalInputFloat, Normal, Phi);
+			PenetrationWithNormalConvex(TypedPtr, PhiWithNormalInputFloat, Thickness, Normal, Penetration);
 			return true;
 		}
 	default:
@@ -722,24 +757,25 @@ export void ApplyPerParticleCollisionFastFriction(uniform FVector4f ParticlesPan
 			const FVector4f InverseR = QuatInverse(CollisionR[i]);
 			const FVector3f PhiWithNormalInputFloat = VectorQuaternionRotateVector(InverseR, TranslatedVec);
 			FVector3f Normal;
-			float Phi;
+			float Penetration; 
 
 			const uniform SerializedPtr *uniform GeometryPtr = (const uniform SerializedPtr *uniform)Geometry;
 			const uniform uint8 *uniform Ptr = GeometryPtr[i].Ptr;
 			const uniform uint8 *uniform TypedPtr = Ptr + SizeofFImplicitObject;
 			const uniform uint8 Type = Ptr[OffsetofGeometryType] & 0x3F;
 
-			if(!GetSimplePhiWithNormal(PhiWithNormalInputFloat, Type, Ptr, TypedPtr, SizeofFImplicitObject, OffsetofGeometryType, OffsetOfMargin, Phi, Normal))
+			if(!GetSimplePenetrationWithNormal(PhiWithNormalInputFloat, Type, Ptr, TypedPtr, SizeofFImplicitObject, OffsetofGeometryType, OffsetOfMargin, PerGroupThickness, Penetration, Normal))
 			{
 				uniform int IMMask;
 				if(IM != 0)
 				{
 					IMMask = lanemask();
 				}
+				float Phi;
 				GetPhiWithNormal(CollisionParticles, (const uniform float *uniform)&PhiWithNormalInputFloat, (uniform float *uniform)&Normal, (uniform float *uniform)&Phi, i, programCount, IMMask);
+				Penetration = PerGroupThickness - Phi;
 			}
 
-			const float Penetration = PerGroupThickness - Phi; // This is related to the Normal impulse
 			if(Penetration > 0.0f)
 			{
 				// TransformVector. Scale is 1 so that part can be removed.
@@ -822,25 +858,26 @@ export void ApplyPerParticleCollisionNoFriction(uniform FVector4f ParticlesPandI
 			const FVector4f InverseR = QuatInverse(CollisionR[i]);
 			const FVector3f PhiWithNormalInputFloat = VectorQuaternionRotateVector(InverseR, TranslatedVec);
 			FVector3f Normal;
-			float Phi;
+			float Penetration;
 
 			const uniform SerializedPtr *uniform GeometryPtr = (const uniform SerializedPtr *uniform)Geometry;
 			const uniform uint8 *uniform Ptr = GeometryPtr[i].Ptr;
 			const uniform uint8 *uniform TypedPtr = Ptr + SizeofFImplicitObject;
 			const uniform uint8 Type = Ptr[OffsetofGeometryType] & 0x3F;
 			
-			if(!GetSimplePhiWithNormal(PhiWithNormalInputFloat, Type, Ptr, TypedPtr, SizeofFImplicitObject, OffsetofGeometryType, OffsetOfMargin, Phi, Normal))
+			if(!GetSimplePenetrationWithNormal(PhiWithNormalInputFloat, Type, Ptr, TypedPtr, SizeofFImplicitObject, OffsetofGeometryType, OffsetOfMargin, PerGroupThickness, Penetration, Normal))
 			{
 				uniform int IMMask;
 				if(IM != 0)
 				{
 					IMMask = lanemask();
 				}
-
+				
+				float Phi;
 				GetPhiWithNormal(CollisionParticles, (const uniform float *uniform)&PhiWithNormalInputFloat, (uniform float *uniform)&Normal, (uniform float *uniform)&Phi, i, programCount, IMMask);
+				Penetration = PerGroupThickness - Phi;
 			}
 
-			const float Penetration = PerGroupThickness - Phi; // This is related to the Normal impulse
 			if(Penetration > 0.0f)
 			{
 				// TransformVector. Scale is 1 so that part can be removed.
@@ -903,20 +940,21 @@ export void ApplyPerParticleSimpleCollisionFastFriction(uniform FVector4f Partic
 			const uniform FVector4f InverseR = QuatInverse(CollisionR[i]);
 			const FVector3f PhiWithNormalInputFloat = VectorQuaternionRotateVector(InverseR, TranslatedVec);
 			FVector3f Normal;
-			float Phi;
+			float Penetration;
 			
-			if(!GetSimplePhiWithNormal(PhiWithNormalInputFloat, Type, Ptr, TypedPtr, SizeofFImplicitObject, OffsetofGeometryType, OffsetOfMargin, Phi, Normal))
+			if(!GetSimplePenetrationWithNormal(PhiWithNormalInputFloat, Type, Ptr, TypedPtr, SizeofFImplicitObject, OffsetofGeometryType, OffsetOfMargin, PerGroupThickness, Penetration, Normal))
 			{
 				uniform int IMMask;
 				if(IM != 0)
 				{
 					IMMask = lanemask();
 				}
-
+				
+				float Phi;
 				GetPhiWithNormalCollisionParticleRange(CollisionParticlesRange, (const uniform float *uniform)&PhiWithNormalInputFloat, (uniform float *uniform)&Normal, (uniform float *uniform)&Phi, i, programCount, IMMask);
+				Penetration = PerGroupThickness - Phi;
 			}
 
-			const float Penetration = PerGroupThickness - Phi; // This is related to the Normal impulse
 			if(Penetration > 0.0f)
 			{
 				// TransformVector. Scale is 1 so that part can be removed.
@@ -991,20 +1029,21 @@ export void ApplyPerParticleSimpleCollisionFastFrictionAndGeneratePlanarConstrai
 			const uniform FVector4f InverseR = QuatInverse(CollisionR[i]);
 			const FVector3f PhiWithNormalInputFloat = VectorQuaternionRotateVector(InverseR, TranslatedVec);
 			FVector3f Normal;
-			float Phi;
+			float Penetration;
 			
-			if(!GetSimplePhiWithNormal(PhiWithNormalInputFloat, Type, Ptr, TypedPtr, SizeofFImplicitObject, OffsetofGeometryType, OffsetOfMargin, Phi, Normal))
+			if(!GetSimplePenetrationWithNormal(PhiWithNormalInputFloat, Type, Ptr, TypedPtr, SizeofFImplicitObject, OffsetofGeometryType, OffsetOfMargin, PerGroupThickness, Penetration, Normal))
 			{
 				uniform int IMMask;
 				if(IM != 0)
 				{
 					IMMask = lanemask();
 				}
-
+				
+				float Phi;
 				GetPhiWithNormalCollisionParticleRange(CollisionParticlesRange, (const uniform float *uniform)&PhiWithNormalInputFloat, (uniform float *uniform)&Normal, (uniform float *uniform)&Phi, i, programCount, IMMask);
+				Penetration = PerGroupThickness - Phi;
 			}
 
-			const float Penetration = PerGroupThickness - Phi; // This is related to the Normal impulse
 			if(Penetration > 0.0f)
 			{
 				// TransformVector. Scale is 1 so that part can be removed.
@@ -1320,20 +1359,21 @@ export void ApplyPerParticleSimpleCollisionNoFriction(uniform FVector4f Particle
 			const uniform FVector4f InverseR = QuatInverse(CollisionR[i]);
 			const FVector3f PhiWithNormalInputFloat = VectorQuaternionRotateVector(InverseR, TranslatedVec);
 			FVector3f Normal;
-			float Phi;
+			float Penetration;
 			
-			if(!GetSimplePhiWithNormal(PhiWithNormalInputFloat, Type, Ptr, TypedPtr, SizeofFImplicitObject, OffsetofGeometryType, OffsetOfMargin, Phi, Normal))
+			if(!GetSimplePenetrationWithNormal(PhiWithNormalInputFloat, Type, Ptr, TypedPtr, SizeofFImplicitObject, OffsetofGeometryType, OffsetOfMargin, PerGroupThickness, Penetration, Normal))
 			{
 				uniform int IMMask;
 				if(IM != 0)
 				{
 					IMMask = lanemask();
 				}
-
+				
+				float Phi;
 				GetPhiWithNormalCollisionParticleRange(CollisionParticlesRange, (const uniform float *uniform)&PhiWithNormalInputFloat, (uniform float *uniform)&Normal, (uniform float *uniform)&Phi, i, programCount, IMMask);
+				Penetration = PerGroupThickness - Phi;
 			}
 
-			const float Penetration = PerGroupThickness - Phi; // This is related to the Normal impulse
 			if(Penetration > 0.0f)
 			{
 				// TransformVector. Scale is 1 so that part can be removed.
@@ -1400,20 +1440,21 @@ export void ApplyPerParticleSimpleCollisionNoFrictionAndGeneratePlanarConstraint
 			const uniform FVector4f InverseR = QuatInverse(CollisionR[i]);
 			const FVector3f PhiWithNormalInputFloat = VectorQuaternionRotateVector(InverseR, TranslatedVec);
 			FVector3f Normal;
-			float Phi;
+			float Penetration;
 			
-			if(!GetSimplePhiWithNormal(PhiWithNormalInputFloat, Type, Ptr, TypedPtr, SizeofFImplicitObject, OffsetofGeometryType, OffsetOfMargin, Phi, Normal))
+			if(!GetSimplePenetrationWithNormal(PhiWithNormalInputFloat, Type, Ptr, TypedPtr, SizeofFImplicitObject, OffsetofGeometryType, OffsetOfMargin, PerGroupThickness, Penetration, Normal))
 			{
 				uniform int IMMask;
 				if(IM != 0)
 				{
 					IMMask = lanemask();
 				}
-
+				
+				float Phi;
 				GetPhiWithNormalCollisionParticleRange(CollisionParticlesRange, (const uniform float *uniform)&PhiWithNormalInputFloat, (uniform float *uniform)&Normal, (uniform float *uniform)&Phi, i, programCount, IMMask);
+				Penetration = PerGroupThickness - Phi;
 			}
 
-			const float Penetration = PerGroupThickness - Phi; // This is related to the Normal impulse
 			if(Penetration > 0.0f)
 			{
 				// TransformVector. Scale is 1 so that part can be removed.
diff --git a/Engine/Source/Runtime/Experimental/Chaos/Private/Chaos/PerParticlePBDCollisionConstraintISPCDataVerification.h b/Engine/Source/Runtime/Experimental/Chaos/Private/Chaos/PerParticlePBDCollisionConstraintISPCDataVerification.h
--- a/Engine/Source/Runtime/Experimental/Chaos/Private/Chaos/PerParticlePBDCollisionConstraintISPCDataVerification.h
+++ b/Engine/Source/Runtime/Experimental/Chaos/Private/Chaos/PerParticlePBDCollisionConstraintISPCDataVerification.h
@@ -47,12 +47,16 @@ static_assert(sizeof(Chaos::FImplicitObject) + offsetof(ispc::FTaperedCapsule, O
 static_assert(sizeof(ispc::FTaperedCapsule::Origin) == Chaos::FTaperedCapsule::FISPCDataVerifier::SizeOfOrigin());
 static_assert(sizeof(Chaos::FImplicitObject) + offsetof(ispc::FTaperedCapsule, Axis) == Chaos::FTaperedCapsule::FISPCDataVerifier::OffsetOfAxis());
 static_assert(sizeof(ispc::FTaperedCapsule::Axis) == Chaos::FTaperedCapsule::FISPCDataVerifier::SizeOfAxis());
+static_assert(sizeof(Chaos::FImplicitObject) + offsetof(ispc::FTaperedCapsule, OneSidedPlaneNormal) == Chaos::FTaperedCapsule::FISPCDataVerifier::OffsetOfOneSidedPlaneNormal());
+static_assert(sizeof(ispc::FTaperedCapsule::OneSidedPlaneNormal) == Chaos::FTaperedCapsule::FISPCDataVerifier::SizeOfOneSidedPlaneNormal());
 static_assert(sizeof(Chaos::FImplicitObject) + offsetof(ispc::FTaperedCapsule, Height) == Chaos::FTaperedCapsule::FISPCDataVerifier::OffsetOfHeight());
 static_assert(sizeof(ispc::FTaperedCapsule::Height) == Chaos::FTaperedCapsule::FISPCDataVerifier::SizeOfHeight());
 static_assert(sizeof(Chaos::FImplicitObject) + offsetof(ispc::FTaperedCapsule, Radius1) == Chaos::FTaperedCapsule::FISPCDataVerifier::OffsetOfRadius1());
 static_assert(sizeof(ispc::FTaperedCapsule::Radius1) == Chaos::FTaperedCapsule::FISPCDataVerifier::SizeOfRadius1());
 static_assert(sizeof(Chaos::FImplicitObject) + offsetof(ispc::FTaperedCapsule, Radius2) == Chaos::FTaperedCapsule::FISPCDataVerifier::OffsetOfRadius2());
 static_assert(sizeof(ispc::FTaperedCapsule::Radius2) == Chaos::FTaperedCapsule::FISPCDataVerifier::SizeOfRadius2());
+static_assert(sizeof(Chaos::FImplicitObject) + offsetof(ispc::FTaperedCapsule, bIsOneSided) == Chaos::FTaperedCapsule::FISPCDataVerifier::OffsetOfIsOneSided());
+static_assert(sizeof(ispc::FTaperedCapsule::bIsOneSided) == Chaos::FTaperedCapsule::FISPCDataVerifier::SizeOfIsOneSided());
 // Convex
 static_assert(sizeof(Chaos::FImplicitObject) + offsetof(ispc::FConvex, Planes) == Chaos::FConvex::FISPCDataVerifier::OffsetOfPlanes());
 static_assert(sizeof(ispc::FConvex::Planes) == Chaos::FConvex::FISPCDataVerifier::SizeOfPlanes());
diff --git a/Engine/Source/Runtime/Experimental/Chaos/Public/Chaos/PerParticlePBDCollisionConstraint.h b/Engine/Source/Runtime/Experimental/Chaos/Public/Chaos/PerParticlePBDCollisionConstraint.h
--- a/Engine/Source/Runtime/Experimental/Chaos/Public/Chaos/PerParticlePBDCollisionConstraint.h
+++ b/Engine/Source/Runtime/Experimental/Chaos/Public/Chaos/PerParticlePBDCollisionConstraint.h
@@ -65,25 +65,7 @@ public:
 		delete Mutex;
 	}
 
-	inline void ApplyRange(FSolverParticles& Particles, const FSolverReal Dt, const int32 Offset, const int32 Range) const
-	{
-		// NOTE: currently using ISPC with TWeightedLatticeImplicitObject<FLevelSet> is significantly slower than not using ISPC (largely because it has not been fully implemented)
-		if (bRealTypeCompatibleWithISPC && bChaos_PerParticleCollision_ISPC_Enabled && bFastPositionBasedFriction )
-		{
-			ApplyHelperISPC(Particles, Dt, Offset, Range);
-		}
-		else
-		{
-			if (Mutex)
-			{
-				ApplyHelper<true>(Particles, Dt, Offset, Range);
-			}
-			else
-			{
-				ApplyHelper<false>(Particles, Dt, Offset, Range);
-			}
-		}
-	}
+	CHAOS_API void ApplyRange(FSolverParticles& Particles, const FSolverReal Dt, const int32 Offset, const int32 Range) const;
 
 	void ApplyFriction(FSolverParticles& Particles, const FSolverReal Dt, const int32 Index) const
 	{
@@ -108,132 +90,8 @@ public:
 
 private:
 	template<bool bLockAndWriteContacts>
-	inline void ApplyHelper(FSolverParticles& Particles, const FSolverReal Dt, const int32 Offset, const int32 Range) const
-	{
-		const uint32 DynamicGroupId = MDynamicGroupIds[Offset];  // Particle group Id, must be the same across the entire range
-		const FSolverReal PerGroupFriction = MPerGroupFriction[DynamicGroupId];
-		const FSolverReal PerGroupThickness = MPerGroupThickness[DynamicGroupId];
-
-		if (PerGroupFriction > (FSolverReal)UE_KINDA_SMALL_NUMBER)
-		{
-			PhysicsParallelFor(Range - Offset, [this, &Particles, Dt, Offset, DynamicGroupId, PerGroupFriction, PerGroupThickness](int32 i)
-			{
-				const int32 Index = Offset + i;
-
-				if (Particles.InvM(Index) == (FSolverReal)0.)
-				{
-					return;  // Continue
-				}
-
-				MCollisionParticlesActiveView.SequentialFor([this, &Particles, &Dt, &Index, DynamicGroupId, PerGroupFriction, PerGroupThickness](FSolverCollisionParticles& CollisionParticles, int32 i)
-				{
-					const uint32 KinematicGroupId = MKinematicGroupIds[i];  // Collision group Id
-
-					if (KinematicGroupId != (uint32)INDEX_NONE && DynamicGroupId != KinematicGroupId )
-					{
-						return; // Bail out if the collision groups doesn't match the particle group id, or use INDEX_NONE (= global collision that affects all particle)
-					}
-					const FSolverRigidTransform3 Frame(CollisionParticles.GetX(i), CollisionParticles.GetR(i));
-					const FVec3 RigidSpacePosition(Frame.InverseTransformPosition(Particles.P(Index)));  // PhiWithNormal requires FReal based arguments
-					FVec3 ImplicitNormal;                                                                // since implicits don't use FSolverReal
-					const FSolverReal Phi = (FSolverReal)CollisionParticles.GetGeometry(i)->PhiWithNormal(RigidSpacePosition, ImplicitNormal);
-					const FSolverVec3 Normal(ImplicitNormal);
-					const FSolverReal Penetration = PerGroupThickness - Phi; // This is related to the Normal impulse
-
-					if (Penetration > (FSolverReal)0.)
-					{
-						const FSolverVec3 NormalWorld = Frame.TransformVector(Normal);
-
-						if (bLockAndWriteContacts)
-						{
-							check(Mutex);
-							checkSlow(Contacts);
-							checkSlow(Normals);
-							checkSlow(Phis);
-							FScopeLock Lock(Mutex);
-							Contacts->Emplace(Particles.P(Index));
-							Normals->Emplace(NormalWorld);
-							Phis->Emplace(Phi);
-						}
-
-						Particles.P(Index) += Penetration * NormalWorld;
-
-						if (bFastPositionBasedFriction)
-						{
-							FSolverVec3 VectorToPoint = Particles.GetP(Index) - CollisionParticles.GetX(i);
-							const FSolverVec3 RelativeDisplacement = (Particles.GetP(Index) - Particles.GetX(Index)) - (CollisionParticles.V(i) + FSolverVec3::CrossProduct(CollisionParticles.W(i), VectorToPoint)) * Dt; // This corresponds to the tangential velocity multiplied by dt (friction will drive this to zero if it is high enough)
-							const FSolverVec3 RelativeDisplacementTangent = RelativeDisplacement - FSolverVec3::DotProduct(RelativeDisplacement, NormalWorld) * NormalWorld; // Project displacement into the tangential plane
-							const FSolverReal RelativeDisplacementTangentLength = RelativeDisplacementTangent.Size();
-							if (RelativeDisplacementTangentLength >= UE_SMALL_NUMBER)
-							{
-								const FSolverReal PositionCorrection = FMath::Min<FSolverReal>(Penetration * PerGroupFriction, RelativeDisplacementTangentLength);
-								const FSolverReal CorrectionRatio = PositionCorrection / RelativeDisplacementTangentLength;
-								Particles.P(Index) -= CorrectionRatio * RelativeDisplacementTangent;
-							}
-						}
-						else
-						{
-							// Note, to fix: Only use fast position based friction for now, since adding to TMaps here is not thread safe when calling Apply on multiple threads (will cause crash)
-							FVelocityConstraint Constraint;
-							FSolverVec3 VectorToPoint = Particles.GetP(Index) - CollisionParticles.GetX(i);
-							Constraint.Velocity = CollisionParticles.V(i) + FSolverVec3::CrossProduct(CollisionParticles.W(i), VectorToPoint);
-							Constraint.Normal = Frame.TransformVector(Normal);
-						
-							MVelocityConstraints.Add(Index, Constraint);
-						}
-					}
-				});
-			});
-		}
-		else
-		{
-			PhysicsParallelFor(Range - Offset, [this, &Particles, Dt, Offset, DynamicGroupId, PerGroupFriction, PerGroupThickness](int32 i)
-			{
-				const int32 Index = Offset + i;
-
-				if (Particles.InvM(Index) == 0)
-				{
-					return;  // Continue
-				}
-
-				MCollisionParticlesActiveView.SequentialFor([this, &Particles, &Dt, &Index, DynamicGroupId, PerGroupFriction, PerGroupThickness](FSolverCollisionParticles& CollisionParticles, int32 i)
-				{
-					const uint32 KinematicGroupId = MKinematicGroupIds[i];  // Collision group Id
-
-					if (KinematicGroupId != (uint32)INDEX_NONE && DynamicGroupId != KinematicGroupId)
-					{
-						return; // Bail out if the collision groups doesn't match the particle group id, or use INDEX_NONE (= global collision that affects all particle)
-					}
-					const FSolverRigidTransform3 Frame(CollisionParticles.GetX(i), CollisionParticles.GetR(i));
-					const FVec3 RigidSpacePosition(Frame.InverseTransformPosition(Particles.P(Index)));  // PhiWithNormal requires FReal based arguments
-					FVec3 ImplicitNormal;                                                                // since implicits don't use FSolverReal
-					const FSolverReal Phi = (FSolverReal)CollisionParticles.GetGeometry(i)->PhiWithNormal(RigidSpacePosition, ImplicitNormal);
-					const FSolverVec3 Normal(ImplicitNormal);
-
-					const FSolverReal Penetration = PerGroupThickness - Phi; // This is related to the Normal impulse
-					if (Penetration > (FSolverReal)0.)
-					{
-						const FSolverVec3 NormalWorld = Frame.TransformVector(Normal);
-						if (bLockAndWriteContacts)
-						{
-							check(Mutex);
-							checkSlow(Contacts);
-							checkSlow(Normals);
-							checkSlow(Phis);
-							FScopeLock Lock(Mutex);
-							Contacts->Emplace(Particles.P(Index));
-							Normals->Emplace(NormalWorld);
-							Phis->Emplace(Phi);
-						}
-
-						Particles.P(Index) += Penetration * NormalWorld;
-					}
-				});
-			});
-		}
-	}
-
-	CHAOS_API void ApplyHelperISPC(FSolverParticles& Particles, const FSolverReal Dt, const int32 Offset, const int32 Range) const;
+	void ApplyHelper(FSolverParticles& Particles, const FSolverReal Dt, const int32 Offset, const int32 Range) const;
+	void ApplyHelperISPC(FSolverParticles& Particles, const FSolverReal Dt, const int32 Offset, const int32 Range) const;
 
 private:
 	bool bFastPositionBasedFriction;
diff --git a/Engine/Source/Runtime/Experimental/Chaos/Public/Chaos/TaperedCapsule.h b/Engine/Source/Runtime/Experimental/Chaos/Public/Chaos/TaperedCapsule.h
--- a/Engine/Source/Runtime/Experimental/Chaos/Public/Chaos/TaperedCapsule.h
+++ b/Engine/Source/Runtime/Experimental/Chaos/Public/Chaos/TaperedCapsule.h
@@ -17,13 +17,15 @@ namespace Chaos
 		{
 			this->bIsConvex = true;
 		}
-		FTaperedCapsule(const FVec3& X1, const FVec3& X2, const FReal InRadius1, const FReal InRadius2)
+		FTaperedCapsule(const FVec3& X1, const FVec3& X2, const FReal InRadius1, const FReal InRadius2, const FVec3& SplitXAxis = FVec3::ZeroVector)
 		    : FImplicitObject(EImplicitObject::FiniteConvex, ImplicitObjectType::TaperedCapsule)
 		    , Origin(X1)
 		    , Axis((X2 - X1).GetSafeNormal())
+			, OneSidedPlaneNormal((FVec3f(SplitXAxis) - FVec3f(SplitXAxis).Dot(Axis)*Axis).GetSafeNormal())
 		    , Height((FRealSingle)((X2 - X1).Size()))
 		    , Radius1((FRealSingle)InRadius1)
 		    , Radius2((FRealSingle)InRadius2)
+			, bIsOneSided(!OneSidedPlaneNormal.IsNearlyZero())
 		    , LocalBoundingBox(X1, X1)
 		{
 			this->bIsConvex = true;
@@ -35,9 +37,11 @@ namespace Chaos
 		    : FImplicitObject(EImplicitObject::FiniteConvex, ImplicitObjectType::TaperedCapsule)
 		    , Origin(Other.Origin)
 		    , Axis(Other.Axis)
+			, OneSidedPlaneNormal(Other.OneSidedPlaneNormal)
 		    , Height(Other.Height)
 		    , Radius1(Other.Radius1)
 		    , Radius2(Other.Radius2)
+			, bIsOneSided(Other.bIsOneSided)
 		    , LocalBoundingBox(Other.LocalBoundingBox)
 		{
 			this->bIsConvex = true;
@@ -46,9 +50,11 @@ namespace Chaos
 		    : FImplicitObject(EImplicitObject::FiniteConvex, ImplicitObjectType::TaperedCapsule)
 		    , Origin(MoveTemp(Other.Origin))
 		    , Axis(MoveTemp(Other.Axis))
+			, OneSidedPlaneNormal(MoveTemp(Other.OneSidedPlaneNormal))
 		    , Height(Other.Height)
 		    , Radius1(Other.Radius1)
 		    , Radius2(Other.Radius2)
+			, bIsOneSided(Other.bIsOneSided)
 		    , LocalBoundingBox(MoveTemp(Other.LocalBoundingBox))
 		{
 			this->bIsConvex = true;
@@ -122,9 +128,13 @@ namespace Chaos
 		FVec3 GetX2() const { return Origin + Axis * Height; }
 		/** Returns the bottommost hemisphere center of the capsule. */
 		FVec3 GetOrigin() const { return GetX1(); }
+		const FVec3f& GetOriginf() const { return Origin; }
 		/** Returns the topmost hemisphere center of capsule . */
 		FVec3 GetInsertion() const { return GetX2(); }
 		FVec3 GetCenter() const { return Origin + Axis * (Height * (FRealSingle)0.5); }
+		bool IsOneSided() const { return bIsOneSided; }
+		FVec3 GetOneSidedPlaneNormal() const { return OneSidedPlaneNormal; }
+		const FVec3f& GetOneSidedPlaneNormalf() const { return OneSidedPlaneNormal; }
 		/** Returns the centroid (center of mass). */
 		FVec3 GetCenterOfMass() const // centroid
 		{
@@ -227,12 +237,16 @@ namespace Chaos
 			static constexpr int32 SizeOfOrigin() { return sizeof(FTaperedCapsule::Origin); }
 			static constexpr int32 OffsetOfAxis() { return offsetof(FTaperedCapsule, Axis); }
 			static constexpr int32 SizeOfAxis() { return sizeof(FTaperedCapsule::Axis); }
+			static constexpr int32 OffsetOfOneSidedPlaneNormal() { return offsetof(FTaperedCapsule, OneSidedPlaneNormal); }
+			static constexpr int32 SizeOfOneSidedPlaneNormal() { return sizeof(FTaperedCapsule::OneSidedPlaneNormal); }
 			static constexpr int32 OffsetOfHeight() { return offsetof(FTaperedCapsule, Height); }
 			static constexpr int32 SizeOfHeight() { return sizeof(FTaperedCapsule::Height); }
 			static constexpr int32 OffsetOfRadius1() { return offsetof(FTaperedCapsule, Radius1); }
 			static constexpr int32 SizeOfRadius1() { return sizeof(FTaperedCapsule::Radius1); }
 			static constexpr int32 OffsetOfRadius2() { return offsetof(FTaperedCapsule, Radius2); }
 			static constexpr int32 SizeOfRadius2() { return sizeof(FTaperedCapsule::Radius2); }
+			static constexpr int32 OffsetOfIsOneSided() { return offsetof(FTaperedCapsule, bIsOneSided); }
+			static constexpr int32 SizeOfIsOneSided() { return sizeof(FTaperedCapsule::bIsOneSided); }
 		};
 		friend FISPCDataVerifier;
 #endif // #if INTEL_ISPC
@@ -257,7 +271,9 @@ namespace Chaos
 		}
 
 		FVec3f Origin, Axis;
+		FVec3f OneSidedPlaneNormal = FVec3f::ZeroVector;
 		FRealSingle Height, Radius1, Radius2;
+		bool bIsOneSided = false;
 		FAABB3 LocalBoundingBox;
 	};
 
