Fix CIS issue with missing includes when adding the one-sided tapered capsule.

diff --git a/Engine/Source/Runtime/Experimental/Chaos/Private/Chaos/PBDSoftBodyCollisionConstraint.cpp b/Engine/Source/Runtime/Experimental/Chaos/Private/Chaos/PBDSoftBodyCollisionConstraint.cpp
--- a/Engine/Source/Runtime/Experimental/Chaos/Private/Chaos/PBDSoftBodyCollisionConstraint.cpp
+++ b/Engine/Source/Runtime/Experimental/Chaos/Private/Chaos/PBDSoftBodyCollisionConstraint.cpp
@@ -4,6 +4,7 @@
 #include "Chaos/Levelset.h"
 #include "Chaos/MLLevelset.h"
 #include "Chaos/SoftsEvolutionLinearSystem.h"
+#include "Chaos/TaperedCapsule.h"
 #include "Chaos/WeightedLatticeImplicitObject.h"
 #include "ChaosStats.h"
 #include "HAL/IConsoleManager.h"
diff --git a/Engine/Source/Runtime/Experimental/Chaos/Private/Chaos/PerParticlePBDCollisionConstraint.cpp b/Engine/Source/Runtime/Experimental/Chaos/Private/Chaos/PerParticlePBDCollisionConstraint.cpp
--- a/Engine/Source/Runtime/Experimental/Chaos/Private/Chaos/PerParticlePBDCollisionConstraint.cpp
+++ b/Engine/Source/Runtime/Experimental/Chaos/Private/Chaos/PerParticlePBDCollisionConstraint.cpp
@@ -1,6 +1,7 @@
 // Copyright Epic Games, Inc. All Rights Reserved.
 
 #include "Chaos/PerParticlePBDCollisionConstraint.h"
+#include "Chaos/TaperedCapsule.h"
 #include "ChaosStats.h"
 #include "HAL/IConsoleManager.h"
 #if INTEL_ISPC
