Another CIS bugfix for recent OneSided Tapered Capsule change. This time fixing unity builds.

diff --git a/Engine/Source/Runtime/Experimental/Chaos/Private/Chaos/PBDSoftBodyCollisionConstraint.cpp b/Engine/Source/Runtime/Experimental/Chaos/Private/Chaos/PBDSoftBodyCollisionConstraint.cpp
--- a/Engine/Source/Runtime/Experimental/Chaos/Private/Chaos/PBDSoftBodyCollisionConstraint.cpp
+++ b/Engine/Source/Runtime/Experimental/Chaos/Private/Chaos/PBDSoftBodyCollisionConstraint.cpp
@@ -73,23 +73,8 @@ namespace Chaos::Softs {
 			ApplyFriction(P, X, NormalWorld, MaxFrictionCorrection, Dt, CollisionX, CollisionV, CollisionW, ColliderVelocityAtPoint);
 		}
 
-		static void ReflectOneSidedCollision(const FSolverVec3& P, const FSolverVec3& OneSidedPlaneNormal, const FSolverVec3& SplitOrigin, FSolverReal& Penetration, FSolverVec3& ImplicitNormal)
-		{
-			check(Penetration > (FSolverReal)0.f);
-
-			FSolverVec3 PNew = P + Penetration * ImplicitNormal;
-			const FSolverReal SplitAxisProj = (PNew - SplitOrigin).Dot(OneSidedPlaneNormal);
-			if (SplitAxisProj >= 0.)
-			{
-				return;
-			}
-
-			PNew -= 2. * SplitAxisProj * OneSidedPlaneNormal;
-			ImplicitNormal = PNew - P;
-			Penetration = ImplicitNormal.Length();
-			ImplicitNormal = Penetration > UE_SMALL_NUMBER ? ImplicitNormal / Penetration : FSolverVec3::ZeroVector;
-			return;
-		}
+		// Defined in PerParticlePBDCollisionConstraint.cpp
+		extern void ReflectOneSidedCollision(const FSolverVec3& P, const FSolverVec3& OneSidedPlaneNormal, const FSolverVec3& SplitOrigin, FSolverReal& Penetration, FSolverVec3& ImplicitNormal);
 	}
 
 // Helper function to call PhiWithNormal and return data to ISPC
diff --git a/Engine/Source/Runtime/Experimental/Chaos/Private/Chaos/PerParticlePBDCollisionConstraint.cpp b/Engine/Source/Runtime/Experimental/Chaos/Private/Chaos/PerParticlePBDCollisionConstraint.cpp
--- a/Engine/Source/Runtime/Experimental/Chaos/Private/Chaos/PerParticlePBDCollisionConstraint.cpp
+++ b/Engine/Source/Runtime/Experimental/Chaos/Private/Chaos/PerParticlePBDCollisionConstraint.cpp
@@ -23,7 +23,8 @@ namespace Chaos::Softs
 {
 namespace Private
 {
-	static void ReflectOneSidedCollision(const FSolverVec3& P, const FSolverVec3& OneSidedPlaneNormal, const FSolverVec3& SplitOrigin, FSolverReal& Penetration, FSolverVec3& ImplicitNormal)
+	// Also called in PBDSoftBodyCollisionConstraint
+	void ReflectOneSidedCollision(const FSolverVec3& P, const FSolverVec3& OneSidedPlaneNormal, const FSolverVec3& SplitOrigin, FSolverReal& Penetration, FSolverVec3& ImplicitNormal)
 	{
 		check(Penetration > (FSolverReal)0.f);
 
