Fixed bug when initializing I/O store on-demand and post hotfix initialization fails.

diff --git a/Engine/Source/Runtime/Experimental/IoStore/OnDemandCore/Private/IoStoreOnDemand.cpp b/Engine/Source/Runtime/Experimental/IoStore/OnDemandCore/Private/IoStoreOnDemand.cpp
--- a/Engine/Source/Runtime/Experimental/IoStore/OnDemandCore/Private/IoStoreOnDemand.cpp
+++ b/Engine/Source/Runtime/Experimental/IoStore/OnDemandCore/Private/IoStoreOnDemand.cpp
@@ -376,19 +376,33 @@ void FIoStoreOnDemandCoreModule::HandleModularFeatureRegistered(const FName& Typ
 	}
 
 	IOnDemandIoStoreFactory* Factory = static_cast<IOnDemandIoStoreFactory*>(ModularFeature);
-	GOnDemandIoStore = Factory->CreateInstance();
-	if (FIoStatus Status = GOnDemandIoStore->Initialize(); !Status.IsOk())
+	FSharedOnDemandIoStore IoStore = Factory->CreateInstance();
+	if (FIoStatus Status = IoStore->Initialize(); !Status.IsOk())
 	{
-		UE_LOG(LogIoStoreOnDemand, Error, TEXT("Failed to initialize I/O store, reason ''"), *Status.ToString());
-		GOnDemandIoStore.Reset();
+		if (Status.GetErrorCode() == EIoErrorCode::Disabled || Status.GetErrorCode() == EIoErrorCode::NotFound)
+		{
+			UE_LOG(LogIoStoreOnDemand, Log, TEXT("I/O store on-demand disabled, reason '%s'"), *Status.ToString());
+		}
+		else
+		{
+			UE_LOG(LogIoStoreOnDemand, Error, TEXT("I/O store on-demand disabled, reason '%s'"), *Status.ToString());
+		}
+		return;
 	}
 	#if !UE_IAS_CUSTOM_INITIALIZATION
-	if (FIoStatus Status = GOnDemandIoStore->InitializePostHotfix(); !Status.IsOk())
+	if (FIoStatus Status = IoStore->InitializePostHotfix(); !Status.IsOk())
 	{
-		UE_LOG(LogIoStoreOnDemand, Error, TEXT("Failed to initialize I/O store, reason ''"), *Status.ToString());
-		GOnDemandIoStore.Reset();
+		if (Status.GetErrorCode() == EIoErrorCode::Disabled || Status.GetErrorCode() == EIoErrorCode::NotFound)
+		{
+			UE_LOG(LogIoStoreOnDemand, Log, TEXT("I/O store post hotfix initialization failed, reason '%s'"), *Status.ToString());
+		}
+		else
+		{
+			UE_LOG(LogIoStoreOnDemand, Error, TEXT("I/O store post hotfix initialization failed, reason '%s'"), *Status.ToString());
+		}
 	}
 	#endif
+	GOnDemandIoStore = IoStore;
 	IModularFeatures::Get().OnModularFeatureRegistered().RemoveAll(this);
 }
 
