Make sure initialize I/O store on-demand only can be called from core module.

diff --git a/Engine/Source/Runtime/Experimental/IoStore/OnDemandCore/Public/IO/IoStoreOnDemand.h b/Engine/Source/Runtime/Experimental/IoStore/OnDemandCore/Public/IO/IoStoreOnDemand.h
--- a/Engine/Source/Runtime/Experimental/IoStore/OnDemandCore/Public/IO/IoStoreOnDemand.h
+++ b/Engine/Source/Runtime/Experimental/IoStore/OnDemandCore/Public/IO/IoStoreOnDemand.h
@@ -399,8 +399,6 @@ public:
 	/** Destroy the I/O store. */
 	virtual ~IOnDemandIoStore() = default;
 	/** Initialize the I/O store. Called after the module feature has been registered. */
-	virtual FIoStatus Initialize() = 0;
-	/** Initialize the I/O store. Called after the module feature has been registered. */
 	virtual FIoStatus InitializePostHotfix() = 0;
 	/**
 	 * Mount an on-demand container.
@@ -459,6 +457,8 @@ public:
 	virtual TUniquePtr<IAnalyticsRecording> StartAnalyticsRecording() const = 0;
 
 private:
+	/** Initialize the I/O store. Called after the I/O store factory module feature has been registered. */
+	virtual FIoStatus Initialize() = 0;
 	/** Cancel the install request. */
 	virtual void CancelInstallRequest(FSharedInternalInstallRequest InstallRequest) = 0;
 	/** Update the priority for the install request. */
@@ -466,6 +466,7 @@ private:
 	/** Release all content referenced by the content handle. */
 	virtual void ReleaseContent(class FOnDemandInternalContentHandle& ContentHandle) = 0;
 
+	friend class FIoStoreOnDemandCoreModule;
 	friend class FOnDemandInternalContentHandle;
 	friend class FOnDemandInstallRequest;
 };
