Changed command line argument for installing all IAD content from -Iad to -Iad.InstallAllContent

- the command line argument -Iad.InstallAllContent is a non shipping test command for install all mounted content

#rnx

diff --git a/Engine/Source/Runtime/Experimental/IoStore/OnDemand/Private/OnDemandIoStore.cpp b/Engine/Source/Runtime/Experimental/IoStore/OnDemand/Private/OnDemandIoStore.cpp
--- a/Engine/Source/Runtime/Experimental/IoStore/OnDemand/Private/OnDemandIoStore.cpp
+++ b/Engine/Source/Runtime/Experimental/IoStore/OnDemand/Private/OnDemandIoStore.cpp
@@ -653,7 +653,7 @@ FIoStatus FOnDemandIoStore::InitializePostHotfix()
 
 #if !UE_BUILD_SHIPPING
 	TOptional<FOnDemandInstallArgs> InstallArgs;
-	if (FParse::Param(FCommandLine::Get(), TEXT("Iad")))
+	if (FParse::Param(FCommandLine::Get(), TEXT("Iad.InstallAllContent")))
 	{
 		TIoStatusOr<FOnDemandHostGroup> HostGroup = FOnDemandHostGroup::Create(EndpointConfig.ServiceUrls);
 		check(HostGroup.IsOk());
