Fix bug static analysis caught. (warning V610: Undefined behavior. Check the shift operator '<<'. The right operand ('Idx' = [0..32]) is greater than or equal to the length in bits of the promoted left operand.)
- Fix LexToString(EOnDemandContainerFlags) iterating too many bits
- Fix missing missing string from bit flag that I added
- Add a static assert to make sure that strings are added for any new flags in the future

diff --git a/Engine/Source/Runtime/Experimental/IoStore/OnDemand/Private/OnDemandIoStore.cpp b/Engine/Source/Runtime/Experimental/IoStore/OnDemand/Private/OnDemandIoStore.cpp
--- a/Engine/Source/Runtime/Experimental/IoStore/OnDemand/Private/OnDemandIoStore.cpp
+++ b/Engine/Source/Runtime/Experimental/IoStore/OnDemand/Private/OnDemandIoStore.cpp
@@ -369,7 +369,8 @@ void LexToString(EOnDemandContainerFlags Flags, FStringBuilderBase& Out)
 		TEXT("Mounted"),
 		TEXT("StreamOnDemand"),
 		TEXT("InstallOnDemand"),
-		TEXT("Encrypted")
+		TEXT("Encrypted"),
+		TEXT("WithSoftReferences")
 	};
 
 	if (Flags == EOnDemandContainerFlags::None)
@@ -378,7 +379,11 @@ void LexToString(EOnDemandContainerFlags Flags, FStringBuilderBase& Out)
 		return;
 	}
 
-	for (int32 Idx = 0, Count = int32(EOnDemandContainerFlags::Count); Idx < Count; ++Idx)
+	constexpr uint32 BitCount = 1 + FMath::CountTrailingZerosConstExpr(
+		static_cast<std::underlying_type_t<EOnDemandContainerFlags>>(EOnDemandContainerFlags::Last));
+	static_assert(UE_ARRAY_COUNT(Names) == BitCount + 1, "Please update names list");
+
+	for (int32 Idx = 0; Idx < BitCount; ++Idx)
 	{
 		const EOnDemandContainerFlags FlagToTest = static_cast<EOnDemandContainerFlags>(1 << Idx);
 		if (EnumHasAnyFlags(Flags, FlagToTest))
diff --git a/Engine/Source/Runtime/Experimental/IoStore/OnDemand/Private/OnDemandIoStore.h b/Engine/Source/Runtime/Experimental/IoStore/OnDemand/Private/OnDemandIoStore.h
--- a/Engine/Source/Runtime/Experimental/IoStore/OnDemand/Private/OnDemandIoStore.h
+++ b/Engine/Source/Runtime/Experimental/IoStore/OnDemand/Private/OnDemandIoStore.h
@@ -44,7 +44,7 @@ enum class EOnDemandContainerFlags : uint8
 	InstallOnDemand			= (1 << 3),
 	Encrypted				= (1 << 4),
 	WithSoftReferences		= (1 << 5),
-	Count
+	Last = WithSoftReferences
 };
 ENUM_CLASS_FLAGS(EOnDemandContainerFlags);
 
