Iris
- Subobjects replicating destroy is now dispatched immediately instead of deffered to end of dispatch
- Added cvar to enable/disable net.Iris.ImmediateDispatchEndReplicationForSubObjects, default is true

Missing file

diff --git a/Engine/Source/Runtime/Experimental/Iris/Core/Private/Iris/ReplicationSystem/ReplicationReader.cpp b/Engine/Source/Runtime/Experimental/Iris/Core/Private/Iris/ReplicationSystem/ReplicationReader.cpp
--- a/Engine/Source/Runtime/Experimental/Iris/Core/Private/Iris/ReplicationSystem/ReplicationReader.cpp
+++ b/Engine/Source/Runtime/Experimental/Iris/Core/Private/Iris/ReplicationSystem/ReplicationReader.cpp
@@ -81,7 +81,13 @@ static bool bDeferEndReplication = true;
 static FAutoConsoleVariableRef CVarDeferEndReplication(
 	TEXT("net.Iris.DeferEndReplication"),
 	bDeferEndReplication,
-	TEXT("bDeferEndReplication if true calls to EndReplication will be defered until after we have applied statedata. Default is true."));
+	TEXT("if true calls to EndReplication will be defered until after we have applied statedata. Default is true."));
+
+static bool bImmediateDispatchEndReplicationForSubObjects = true;
+static FAutoConsoleVariableRef CVarImmediateDispatchEndReplicationForSubObjects(
+	TEXT("net.Iris.ImmediateDispatchEndReplicationForSubObjects"),
+	bImmediateDispatchEndReplicationForSubObjects,
+	TEXT("if true replicated calls to EndReplication for SubObjects will dispatched immediately. Default is true."));
 
 static bool bDispatchUnresolvedPreviouslyReceivedChanges = false;
 static FAutoConsoleVariableRef CvarDispatchUnresolvedPreviouslyReceivedChanges(
@@ -180,6 +186,7 @@ struct FReplicationReader::FDispatchObjectInfo
 	uint32 bDestroy : 1 = false;
 	uint32 bTearOff : 1 = false;
 	uint32 bDeferredEndReplication : 1 = false;
+	uint32 bReplicatedDestroyForSubObject : 1 = false;
 	uint32 bProcessEndReplicationForHierarchy : 1 = false;
 	uint32 bShouldCallSubObjectCreatedFromReplication : 1 = false;
 	uint32 bDynamicObjectCreated : 1 = false;
@@ -1267,7 +1274,17 @@ void FReplicationReader::ReadObjectInBatch(FNetSerializationContext& Context, FN
 		// Update info based on ReplicatedDestroyHeader
 		Info.bDestroy = !!(ReplicatedDestroyHeaderFlags & (ReplicatedDestroyHeaderFlags_TearOff | ReplicatedDestroyHeaderFlags_DestroyInstance));
 		Info.bTearOff = !!(ReplicatedDestroyHeaderFlags & ReplicatedDestroyHeaderFlags_TearOff);
-		Info.bDeferredEndReplication = !!(ReplicatedDestroyHeaderFlags & (ReplicatedDestroyHeaderFlags_TearOff | ReplicatedDestroyHeaderFlags_EndReplication));
+		// We set one of these flags
+		if (bImmediateDispatchEndReplicationForSubObjects && Info.bDestroy && !Info.bTearOff)
+		{
+			Info.bReplicatedDestroyForSubObject = 1U;
+			Info.bDeferredEndReplication = 0U;
+		}
+		else
+		{
+			Info.bReplicatedDestroyForSubObject = 0U;
+			Info.bDeferredEndReplication = !!(ReplicatedDestroyHeaderFlags & (ReplicatedDestroyHeaderFlags_TearOff | ReplicatedDestroyHeaderFlags_EndReplication));
+		}
 		Info.bShouldCallSubObjectCreatedFromReplication = bShouldCallSubObjectCreatedFromReplication;
 
 		if (bHasState)
@@ -2037,6 +2054,20 @@ void FReplicationReader::DispatchStateData(FNetSerializationContext& Context)
 	// Function to flush all objects pending post dispatch
 	auto FlushPostDispatchForBatch = [PostDispatchObjectInfos, &NumObjectsPendingPostDistpatch, this, &Context]()
 	{
+		// Dispatch replicated subobject destroy here to behave as close as possible to subobject replication in ActorChannel.
+		// Added to solve issues with assumptions made by blueprint logic when replacing a component with another of the same type
+		if (bImmediateDispatchEndReplicationForSubObjects)
+		{
+			for (FPostDispatchObjectInfo& PostDispatchObjectInfo : MakeArrayView(PostDispatchObjectInfos, NumObjectsPendingPostDistpatch))
+			{
+				FDispatchObjectInfo& Info = *PostDispatchObjectInfo.Info;
+				if (Info.bReplicatedDestroyForSubObject)
+				{
+					EndReplication(Info.InternalIndex, Info.bTearOff, Info.bDestroy);
+				}
+			}
+		}
+
 		// When all received states have been applied we invoke PostReplicate and RepNotifies
 		for (FPostDispatchObjectInfo& PostDispatchObjectInfo : MakeArrayView(PostDispatchObjectInfos, NumObjectsPendingPostDistpatch))
 		{
