Iris
* Fix rare crash in GetObjectReferenceHandleFromObject when verbose logs are enabled and trying to print name of a deleted pointer.

#rnx

diff --git a/Engine/Source/Runtime/Experimental/Iris/Core/Private/Iris/ReplicationSystem/ObjectReferenceCache.cpp b/Engine/Source/Runtime/Experimental/Iris/Core/Private/Iris/ReplicationSystem/ObjectReferenceCache.cpp
--- a/Engine/Source/Runtime/Experimental/Iris/Core/Private/Iris/ReplicationSystem/ObjectReferenceCache.cpp
+++ b/Engine/Source/Runtime/Experimental/Iris/Core/Private/Iris/ReplicationSystem/ObjectReferenceCache.cpp
@@ -614,14 +614,14 @@ void FObjectReferenceCache::RemoveReference(FNetRefHandle RefHandle, const UObje
 			// Remove if RelativePath is invalid
 			if (!CachedItemToRemove.RelativePath.IsValid())
 			{
-				UE_LOG_REFERENCECACHE(Verbose, TEXT("ObjectReferenceCache::RemoveReference: %s, Object: %s (0x%p). Destroying cache since Path is invalid (%s)"),
-					*RefHandle.ToString(), *GetNameSafe(Object), Object, *CachedItemToRemove.RelativePath.ToString());
+				UE_LOG_REFERENCECACHE(Verbose, TEXT("ObjectReferenceCache::RemoveReference: %s, Object: (0x%p). Destroying cache since Path is invalid (%s)"),
+					*RefHandle.ToString(), Object, *CachedItemToRemove.RelativePath.ToString());
 				It.RemoveCurrent();
 			}
 			else
 			{
-				// Clear out object reference since it will be invalid
-				UE_LOG_REFERENCECACHE(Verbose, TEXT("ObjectReferenceCache::RemoveReference: %s, Object: %s (0x%p). Clearing cache reference to (0x%p)"), *RefHandle.ToString(), *GetNameSafe(Object), Object, CachedItemToRemove.ObjectKey);
+				UE_LOG_REFERENCECACHE(Verbose, TEXT("ObjectReferenceCache::RemoveReference: %s, Object: (0x%p). Clearing cache reference to (0x%p)"),
+					*RefHandle.ToString(), Object, CachedItemToRemove.ObjectKey);
 				CachedItemToRemove.ObjectKey = nullptr;
 				CachedItemToRemove.Object.Reset();
 			}
@@ -636,8 +636,8 @@ void FObjectReferenceCache::RemoveReference(FNetRefHandle RefHandle, const UObje
 				const FNetRefHandle ExistingRefHandle = ObjectIt.Value();
 				if (ExistingRefHandle == RefHandle)
 				{
-					UE_LOG_REFERENCECACHE(Verbose, TEXT("ObjectReferenceCache::RemoveReference: %s removed from ObjectToNetReferenceHandle Object: %s (0x%p)."), 
-						*RefHandle.ToString(), *GetNameSafe(Object), Object);
+					UE_LOG_REFERENCECACHE(Verbose, TEXT("ObjectReferenceCache::RemoveReference: %s removed from ObjectToNetReferenceHandle Object: (0x%p)."), 
+						*RefHandle.ToString(), Object);
 					ObjectIt.RemoveCurrent();
 				}
 				else
