Iris
* Fix protocols without fragments not storing the fragment traits set in StartReplication
* Deleted NeedsWorldLocationUpdate trait. This is a dynamic root object property and not tied to any fragment
* Improved FReplicationSystemUtil functions to be more actor agnostic

#rnx

diff --git a/Engine/Source/Runtime/Engine/Private/Net/Iris/ReplicationSystem/EngineReplicationBridge.cpp b/Engine/Source/Runtime/Engine/Private/Net/Iris/ReplicationSystem/EngineReplicationBridge.cpp
--- a/Engine/Source/Runtime/Engine/Private/Net/Iris/ReplicationSystem/EngineReplicationBridge.cpp
+++ b/Engine/Source/Runtime/Engine/Private/Net/Iris/ReplicationSystem/EngineReplicationBridge.cpp
@@ -101,11 +101,11 @@ void ActorReplicationBridgePreUpdateFunction(TArrayView<UObject*> Instances, con
 	}
 }
 
-bool ShouldIncludeActorInLevelGroups(const AActor* Actor)
+bool ShouldIncludeRootObjectInLevelGroups(const UObject* RootObject)
 {
 	// Never filter out PlayerControllers based on level as they are required for travel.
 	// Preserves the special case for PlayerControllers from UNetDriver::IsLevelInitializedForActor.
-	return !Actor->IsA<APlayerController>();
+	return !RootObject->IsA<APlayerController>();
 }
 
 } // end namespace UE::Net::Private
@@ -357,7 +357,7 @@ UE::Net::FNetRefHandle UEngineReplicationBridge::StartReplicatingActor(AActor* A
 	}
 
 	// Setup Level filtering
-	AddActorToLevelGroup(Actor);
+	AddRootObjectToLevelGroup(Actor, Actor->GetLevel());
 
 	// If we have registered sub objects we replicate them as well
 	const FSubObjectRegistry& ActorSubObjects = FSubObjectRegistryGetter::GetSubObjects(Actor);
@@ -812,7 +812,7 @@ void UEngineReplicationBridge::ActorChangedLevel(const AActor* Actor, const ULev
 
 	UE_LOG_ACTORREPLICATIONBRIDGE(Verbose, TEXT("ActorChangedLevel: Actor %s from PreviousLevel %s"), *GetFullNameSafe(Actor), *GetFullNameSafe(PreviousLevel));
 
-	if (!UE::Net::Private::ShouldIncludeActorInLevelGroups(Actor))
+	if (!UE::Net::Private::ShouldIncludeRootObjectInLevelGroups(Actor))
 	{
 		return;
 	}
@@ -838,7 +838,7 @@ void UEngineReplicationBridge::ActorChangedLevel(const AActor* Actor, const ULev
 		}
 	}
 
-	AddActorToLevelGroup(Actor);
+	AddRootObjectToLevelGroup(Actor, Actor->GetLevel());
 }
 
 void UEngineReplicationBridge::OnNetUpdateFrequencyChanged(const AActor* Actor)
@@ -854,15 +854,13 @@ void UEngineReplicationBridge::OnNetUpdateFrequencyChanged(const AActor* Actor)
 	}
 }
 
-void UEngineReplicationBridge::AddActorToLevelGroup(const AActor* Actor)
+void UEngineReplicationBridge::AddRootObjectToLevelGroup(const UObject* RootObject, const ULevel* Level)
 {
-	if (!UE::Net::Private::ShouldIncludeActorInLevelGroups(Actor))
+	if (!UE::Net::Private::ShouldIncludeRootObjectInLevelGroups(RootObject))
 	{
 		return;
 	}
 	
-	const ULevel* Level = Actor->GetLevel();
-	
 	// Don't filter out actors in the persistent level
 	if (Level && (!Level->IsPersistentLevel() || (Level != NetDriver->GetWorld()->PersistentLevel)))
 	{
@@ -880,11 +878,11 @@ void UEngineReplicationBridge::AddActorToLevelGroup(const AActor* Actor)
 			NetDriver->UpdateGroupFilterStatusForLevel(Level, LevelGroup);
 		}
 
-		const UE::Net::FNetRefHandle ActorRefHandle = GetReplicatedRefHandle(Actor);
+		const UE::Net::FNetRefHandle RootObjectRefHandle = GetReplicatedRefHandle(RootObject);
 
 		// Add object to group
-		UE_LOG_ACTORREPLICATIONBRIDGE(Verbose, TEXT("Added %s to GroupIndex: %u Level: %s"), *ActorRefHandle.ToString(), LevelGroup.GetGroupIndex(), ToCStr(PackageName.ToString()));
-		GetReplicationSystem()->AddToGroup(LevelGroup, ActorRefHandle);
+		UE_LOG_ACTORREPLICATIONBRIDGE(Verbose, TEXT("Added %s to GroupIndex: %u Level: %s"), *RootObjectRefHandle.ToString(), LevelGroup.GetGroupIndex(), ToCStr(PackageName.ToString()));
+		GetReplicationSystem()->AddToGroup(LevelGroup, RootObjectRefHandle);
 	}
 }
 
diff --git a/Engine/Source/Runtime/Engine/Private/Net/Iris/ReplicationSystem/ReplicationSystemUtil.cpp b/Engine/Source/Runtime/Engine/Private/Net/Iris/ReplicationSystem/ReplicationSystemUtil.cpp
--- a/Engine/Source/Runtime/Engine/Private/Net/Iris/ReplicationSystem/ReplicationSystemUtil.cpp
+++ b/Engine/Source/Runtime/Engine/Private/Net/Iris/ReplicationSystem/ReplicationSystemUtil.cpp
@@ -72,6 +72,12 @@ UReplicationSystem* FReplicationSystemUtil::GetReplicationSystem(const UNetDrive
 	return NetDriver ? NetDriver->GetReplicationSystem() : nullptr;
 }
 
+UReplicationSystem* FReplicationSystemUtil::GetReplicationSystem(const UWorld* World)
+{
+	UNetDriver* NetDriver = World ? World->GetNetDriver() : nullptr;
+	return NetDriver ? NetDriver->GetReplicationSystem() : nullptr;
+}
+
 UEngineReplicationBridge* FReplicationSystemUtil::GetActorReplicationBridge(const AActor* Actor)
 {
 	if (UReplicationSystem* ReplicationSystem = GetReplicationSystem(Actor))
@@ -97,22 +103,15 @@ UEngineReplicationBridge* FReplicationSystemUtil::GetActorReplicationBridge(cons
 	}
 }
 
-FNetHandle FReplicationSystemUtil::GetNetHandle(const AActor* Actor)
-{
-	FNetHandle NetHandle = FNetHandleManager::GetNetHandle(Actor);
-	return NetHandle;
-}
-
-FNetHandle FReplicationSystemUtil::GetNetHandle(const UActorComponent* SubObject)
+UEngineReplicationBridge* FReplicationSystemUtil::GetEngineReplicationBridge(const UWorld* World)
 {
-	FNetHandle NetHandle = FNetHandleManager::GetNetHandle(SubObject);
-	return NetHandle;
+	UReplicationSystem* ReplicationSystem = GetReplicationSystem(World);	
+	return Cast<UEngineReplicationBridge>(ReplicationSystem ? ReplicationSystem->GetReplicationBridge() : nullptr);
 }
 
-FNetHandle FReplicationSystemUtil::GetNetHandle(const UObject* Object)
+FNetHandle FReplicationSystemUtil::GetNetHandle(const UObject* ReplicatedObject)
 {
-	FNetHandle NetHandle = FNetHandleManager::GetNetHandle(Object);
-	return NetHandle;
+	return FNetHandleManager::GetNetHandle(ReplicatedObject);
 }
 
 void FReplicationSystemUtil::BeginReplication(AActor* Actor, const FActorReplicationParams& Params)
diff --git a/Engine/Source/Runtime/Engine/Public/Net/Iris/ReplicationSystem/EngineReplicationBridge.h b/Engine/Source/Runtime/Engine/Public/Net/Iris/ReplicationSystem/EngineReplicationBridge.h
--- a/Engine/Source/Runtime/Engine/Public/Net/Iris/ReplicationSystem/EngineReplicationBridge.h
+++ b/Engine/Source/Runtime/Engine/Public/Net/Iris/ReplicationSystem/EngineReplicationBridge.h
@@ -98,6 +98,9 @@ public:
 
 	/** Tell the remote connection that we detected a reading error with a specific replicated object */
 	ENGINE_API virtual void SendErrorWithNetRefHandle(UE::Net::ENetRefHandleError ErrorType, FNetRefHandle RefHandle, uint32 ConnectionId) override;
+
+	/** Add the rootobject to the level's filter group so it will only be relevant if the connection has that level streamed in. */
+	ENGINE_API void AddRootObjectToLevelGroup(const UObject* RootObject, const ULevel* Level);
 	
 	/** Updates the level group for an actor that changed levels */
 	void ActorChangedLevel(const AActor* Actor, const ULevel* PreviousLevel);
@@ -143,7 +146,6 @@ private:
 	
 	void OnMaxTickRateChanged(UNetDriver* InNetDriver, int32 NewMaxTickRate, int32 OldMaxTickRate);
 
-	void AddActorToLevelGroup(const AActor* Actor);
 
 private:
 
diff --git a/Engine/Source/Runtime/Engine/Public/Net/Iris/ReplicationSystem/ReplicationSystemUtil.h b/Engine/Source/Runtime/Engine/Public/Net/Iris/ReplicationSystem/ReplicationSystemUtil.h
--- a/Engine/Source/Runtime/Engine/Public/Net/Iris/ReplicationSystem/ReplicationSystemUtil.h
+++ b/Engine/Source/Runtime/Engine/Public/Net/Iris/ReplicationSystem/ReplicationSystemUtil.h
@@ -46,20 +46,20 @@ struct FReplicationSystemUtil
 	/** Returns the UReplicationSystem for a specific NetDriver. May return null. */
 	ENGINE_API static UReplicationSystem* GetReplicationSystem(const UNetDriver* NetDriver);
 
-	/** Returns the UEngineReplicationBridge of the UReplicationSystem belogning to the main NetDriver assigned to the Actor. May return null. */
+	/** Returns the UReplicationSystem for the main NetDriver of this World. May return null. */
+	ENGINE_API static UReplicationSystem* GetReplicationSystem(const UWorld* World);
+
+	/** Returns the UEngineReplicationBridge of the UReplicationSystem belonging to the main NetDriver assigned to the Actor. May return null. */
 	ENGINE_API static UEngineReplicationBridge* GetActorReplicationBridge(const AActor* Actor);
 
 	/** Returns the UReplicationSystem of the UNetDriver the UNetConnection belongs to. */
 	ENGINE_API static UEngineReplicationBridge* GetActorReplicationBridge(const UNetConnection* NetConnection);
 
-	/** Returns the NetHandle for the actor. The returned handle may be invalid. */
-	ENGINE_API static FNetHandle GetNetHandle(const AActor* Actor);
-
-	/** Returns the NetHandle for an actor component. The returned handle may be invalid. */
-	ENGINE_API static FNetHandle GetNetHandle(const UActorComponent* SubObject);
+	/** Returns the UEngineReplicationBridge of the UReplicationSystem belonging to the main NetDriver of this World. May return null. */
+	ENGINE_API static UEngineReplicationBridge* GetEngineReplicationBridge(const UWorld* World);
 
-	/** Returns the NetHandle for an object. The returned handle may be invalid. */
-	ENGINE_API static FNetHandle GetNetHandle(const UObject* Object);
+		/** Returns the NetHandle for any replicated object. The returned handle may be invalid. */
+	ENGINE_API static FNetHandle GetNetHandle(const UObject* ReplicatedObject);
 
 	/** Begins replication of an actor and all of its registered subobjects. If any ReplicationSystem wants to replicate the actor a NetHandle will be created. */
 	ENGINE_API static void BeginReplication(AActor* Actor, const FActorReplicationParams& Params);
diff --git a/Engine/Source/Runtime/Experimental/Iris/Core/Private/Iris/ReplicationSystem/ObjectReplicationBridge.cpp b/Engine/Source/Runtime/Experimental/Iris/Core/Private/Iris/ReplicationSystem/ObjectReplicationBridge.cpp
--- a/Engine/Source/Runtime/Experimental/Iris/Core/Private/Iris/ReplicationSystem/ObjectReplicationBridge.cpp
+++ b/Engine/Source/Runtime/Experimental/Iris/Core/Private/Iris/ReplicationSystem/ObjectReplicationBridge.cpp
@@ -57,6 +57,7 @@ namespace UE::Net::Private::ObjectBridgeDebugging
 DEFINE_LOG_CATEGORY(LogIrisFilterConfig)
 
 #define UE_LOG_OBJECTREPLICATIONBRIDGE(Category, Format, ...)  UE_LOG(LogIrisBridge, Category, TEXT("ObjectReplicationBridge(%u)::") Format, GetReplicationSystem()->GetId(), ##__VA_ARGS__)
+#define UE_LOG_BRIDGEID(Category, Verbosity, Format, ...)  UE_LOG(Category, Verbosity, TEXT("ObjectReplicationBridge(%u)::") Format, GetReplicationSystem()->GetId(), ##__VA_ARGS__)
 
 static bool bUseFrequencyBasedPolling = true;
 static FAutoConsoleVariableRef CVarUseFrequencyBasedPolling(
@@ -368,6 +369,10 @@ UE::Net::FNetRefHandle UObjectReplicationBridge::StartReplicatingNetObject(UObje
 				*GetNameSafe(Instance), *GetNameSafe(Instance->GetClass()), ReplicatedProps.Num());
 		}
 	}
+	else if (FragmentRegistrationContext.IsFragmentlessNetObject() && !RegisteredFragments.IsEmpty())
+	{
+		ensureMsgf(false, TEXT("NetObject %s (class %s) says that it is Fragment-less but still created %d Fragments."), *GetNameSafe(Instance), *GetNameSafe(Instance->GetClass()), RegisteredFragments.Num());
+	}
 #endif
 
 	// We currently identify protocols by local archetype or CDO pointer and verified the protocol id received from server and the hash of the default state
@@ -376,7 +381,7 @@ UE::Net::FNetRefHandle UObjectReplicationBridge::StartReplicatingNetObject(UObje
 	// Create Protocols
 	FReplicationProtocolManager* ProtocolManager = GetReplicationProtocolManager();
 
-	FReplicationInstanceProtocolPtr InstanceProtocol(ProtocolManager->CreateInstanceProtocol(RegisteredFragments));
+	FReplicationInstanceProtocolPtr InstanceProtocol(ProtocolManager->CreateInstanceProtocol(RegisteredFragments, FragmentRegistrationContext.GetFragmentTraits()));
 		
 	FReplicationProtocolIdentifier ProtocolIdentifier = FReplicationProtocolManager::CalculateProtocolIdentifier(RegisteredFragments);
 	const FReplicationProtocol* ReplicationProtocol = ProtocolManager->GetReplicationProtocol(ProtocolIdentifier, ArchetypeOrCDOUsedAsKey);
@@ -415,7 +420,13 @@ UE::Net::FNetRefHandle UObjectReplicationBridge::StartReplicatingNetObject(UObje
 		return FNetRefHandle::GetInvalid();
 	}
 
-	UE_LOG_OBJECTREPLICATIONBRIDGE(Verbose, TEXT("StartReplicatingNetObject Created %s with ProtocolId:0x%x for Object named %s"), *RefHandle.ToString(), ReplicationProtocol->ProtocolIdentifier, *Instance->GetName());
+	UE_LOG_OBJECTREPLICATIONBRIDGE(Verbose, TEXT("StartReplicatingNetObject Created %s [handle:%s|stable:%s] with ProtocolId:0x%x for Object named %s"), 
+		*RefHandle.ToString(), 
+		RefHandle.IsStatic()?TEXT("Static"):TEXT("Dynamic"),
+		Instance->IsNameStableForNetworking()?TEXT("yes"):TEXT("no"),
+		ReplicationProtocol->ProtocolIdentifier, 
+		*Instance->GetName()
+	);
 
 	// Attach the instance and bind the instance protocol to dirty tracking
 	constexpr bool bBindInstanceProtocol = true;
@@ -462,7 +473,6 @@ UE::Net::FNetRefHandle UObjectReplicationBridge::StartReplicatingRootObject(UObj
 	
 	EReplicationFragmentTraits Traits = EReplicationFragmentTraits::CanReplicate;
 	Traits |= Params.bNeedsPreUpdate ? EReplicationFragmentTraits::NeedsPreSendUpdate : EReplicationFragmentTraits::None;
-	Traits |= Params.bNeedsWorldLocationUpdate ? EReplicationFragmentTraits::NeedsWorldLocationUpdate : EReplicationFragmentTraits::None;
 
 	// Register the object
 	FNetRefHandle RefHandle = StartReplicatingNetObject(Instance, Traits, NetFactoryId);
@@ -710,7 +720,7 @@ bool UObjectReplicationBridge::WriteNetRefHandleCreationInfo(FReplicationBridgeS
 	{
 		const FNetObjectCreationHeader* Header = CachedHeader->Get();
 		UNetObjectFactory* Factory = GetNetFactory(Header->GetNetFactoryId());
-		UE_LOG(LogIrisBridge, VeryVerbose, TEXT("WriteNetRefHandleCreationInfo (cached) for %s : Factory: %s (id:%u) Header: %s"), *Handle.ToString(), *GetNameSafe(Factory), Header->GetNetFactoryId(), *Header->ToString());
+		UE_LOG_BRIDGEID(LogIrisBridge, VeryVerbose, TEXT("WriteNetRefHandleCreationInfo (cached) for %s : Factory: %s (id:%u) Header: %s"), *Handle.ToString(), *GetNameSafe(Factory), Header->GetNetFactoryId(), *Header->ToString());
 		return Factory->WriteHeader(Handle, Context.SerializationContext, Header);
 	}
 	
@@ -726,7 +736,7 @@ bool UObjectReplicationBridge::WriteNetRefHandleCreationInfo(FReplicationBridgeS
 	const FInternalNetRefIndex ObjectIndex = NetRefHandleManager->GetInternalIndex(Handle);
 	if (ObjectIndex == FNetRefHandleManager::InvalidInternalIndex)
 	{
-		UE_LOG(LogIrisBridge, Error, TEXT("WriteNetRefHandleCreationInfo: Cannot write creationinfo for %s: No internal index assigned to it"), *Handle.ToString());
+		UE_LOG_BRIDGEID(LogIrisBridge, Error, TEXT("WriteNetRefHandleCreationInfo: Cannot write creationinfo for %s: No internal index assigned to it"), *Handle.ToString());
 		return false;
 	}
 
@@ -739,11 +749,11 @@ bool UObjectReplicationBridge::WriteNetRefHandleCreationInfo(FReplicationBridgeS
 
 	if (!Header.IsValid())
 	{
-		UE_LOG(LogIrisBridge, Error, TEXT("WriteNetRefHandleCreationInfo: CreateHeader via %s failed for %s"), *GetNameSafe(Factory), *Handle.ToString());
+		UE_LOG_BRIDGEID(LogIrisBridge, Error, TEXT("WriteNetRefHandleCreationInfo: CreateHeader via %s failed for %s"), *GetNameSafe(Factory), *Handle.ToString());
 		return false;
 	}
 
-	UE_LOG(LogIrisBridge, VeryVerbose, TEXT("WriteNetRefHandleCreationInfo (new) for %s : Factory: %s (id:%u) Header: %s"), *Handle.ToString(), *GetNameSafe(Factory), FactoryId, *Header->ToString());
+	UE_LOG_BRIDGEID(LogIrisBridge, VeryVerbose, TEXT("WriteNetRefHandleCreationInfo (new) for %s : Factory: %s (id:%u) Header: %s"), *Handle.ToString(), *GetNameSafe(Factory), FactoryId, *Header->ToString());
 
 	return Factory->WriteHeader(Handle, Context.SerializationContext, Header.Get());
 }
@@ -762,7 +772,7 @@ bool UObjectReplicationBridge::CacheNetRefHandleCreationInfo(FNetRefHandle Handl
 	const FInternalNetRefIndex ObjectIndex = NetRefHandleManager->GetInternalIndex(Handle);
 	if (ObjectIndex == FNetRefHandleManager::InvalidInternalIndex)
 	{
-		UE_LOG(LogIrisBridge, Error, TEXT("CacheNetRefHandleCreationInfo cannot cache creation info for %s: no internal index assigned to it"), *Handle.ToString());
+		UE_LOG_BRIDGEID(LogIrisBridge, Error, TEXT("CacheNetRefHandleCreationInfo cannot cache creation info for %s: no internal index assigned to it"), *Handle.ToString());
 		ensureMsgf(false, TEXT("CacheNetRefHandleCreationInfo cannot cache creation info for %s: no internal index assigned to it"), *Handle.ToString());
 		return false;
 	}
@@ -908,11 +918,11 @@ FReplicationBridgeCreateNetRefHandleResult UObjectReplicationBridge::CreateNetRe
 	{
 		if (RootObjectOfSubObject.IsValid())
 		{
-			UE_LOG(LogIrisBridge, Verbose, TEXT("CreateNetRefHandleFromRemote: SubObject: %s of RootObject: %s using header: %s"), *WantedNetHandle.ToString(), *RootObjectOfSubObject.ToString(), ToCStr(Header->ToString()));
+			UE_LOG_BRIDGEID(LogIrisBridge, Verbose, TEXT("CreateNetRefHandleFromRemote: SubObject: %s of RootObject: %s using header: %s"), *WantedNetHandle.ToString(), *RootObjectOfSubObject.ToString(), ToCStr(Header->ToString()));
 		}
 		else
 		{
-			UE_LOG(LogIrisBridge, Verbose, TEXT("CreateNetRefHandleFromRemote: RootObject: %s using header: %s"), *WantedNetHandle.ToString(), ToCStr(Header->ToString()));
+			UE_LOG_BRIDGEID(LogIrisBridge, Verbose, TEXT("CreateNetRefHandleFromRemote: RootObject: %s using header: %s"), *WantedNetHandle.ToString(), ToCStr(Header->ToString()));
 		}
 	}
 	
@@ -934,11 +944,11 @@ FReplicationBridgeCreateNetRefHandleResult UObjectReplicationBridge::CreateNetRe
 		{
 			if (RootObjectOfSubObject.IsValid())
 			{
-				UE_LOG(LogIrisBridge, Warning, TEXT("CreateNetRefHandleFromRemote: Failed to instantiate SubObject NetHandle: %s of %s using header: %s"), *WantedNetHandle.ToString(), *PrintObjectFromNetRefHandle(RootObjectOfSubObject), ToCStr(Header->ToString()));
+				UE_LOG_BRIDGEID(LogIrisBridge, Warning, TEXT("CreateNetRefHandleFromRemote: Failed to instantiate SubObject NetHandle: %s of %s using header: %s"), *WantedNetHandle.ToString(), *PrintObjectFromNetRefHandle(RootObjectOfSubObject), ToCStr(Header->ToString()));
 			}
 			else
 			{
-				UE_LOG(LogIrisBridge, Warning, TEXT("CreateNetRefHandleFromRemote: Failed to instantiate RootObject NetHandle: %s using header: %s"), *WantedNetHandle.ToString(), ToCStr(Header->ToString()));
+				UE_LOG_BRIDGEID(LogIrisBridge, Warning, TEXT("CreateNetRefHandleFromRemote: Failed to instantiate RootObject NetHandle: %s using header: %s"), *WantedNetHandle.ToString(), ToCStr(Header->ToString()));
 			}
 		}
 		
@@ -959,7 +969,7 @@ FReplicationBridgeCreateNetRefHandleResult UObjectReplicationBridge::CreateNetRe
 	const UObject* ArchetypeOrCDOUsedAsKey = InstancePtr->GetArchetype();
 
 	// Create Protocols
-	FReplicationInstanceProtocolPtr InstanceProtocol(ProtocolManager->CreateInstanceProtocol(RegisteredFragments));
+	FReplicationInstanceProtocolPtr InstanceProtocol(ProtocolManager->CreateInstanceProtocol(RegisteredFragments, FragmentRegistrationContext.GetFragmentTraits()));
 
 	// See if the protocol already is known
 	const FReplicationProtocol* ReplicationProtocol = ProtocolManager->GetReplicationProtocol(ReceivedProtocolId, ArchetypeOrCDOUsedAsKey);
@@ -980,7 +990,7 @@ FReplicationBridgeCreateNetRefHandleResult UObjectReplicationBridge::CreateNetRe
 
 	if (!ReplicationProtocol)
 	{
-		UE_LOG(LogIris, Error, TEXT("Protocol mismatch prevents binding %s to instanced object %s (CDO: %s)."), *WantedNetHandle.ToString(), *GetNameSafe(InstancePtr), *GetPathNameSafe(ArchetypeOrCDOUsedAsKey));
+		UE_LOG_BRIDGEID(LogIris, Error, TEXT("Protocol mismatch prevents binding %s to instanced object %s (CDO: %s)."), *WantedNetHandle.ToString(), *GetNameSafe(InstancePtr), *GetPathNameSafe(ArchetypeOrCDOUsedAsKey));
 
 		if (UE_LOG_ACTIVE(LogIris, Error))
 		{
@@ -1507,9 +1517,10 @@ void UObjectReplicationBridge::UpdateInstancesWorldLocation()
 	// Retrieve the world location for instances that supports it. Only dirty objects are considered.
 	auto UpdateInstanceWorldLocation = [this, &ReplicatedInstances, &LocalNetRefHandleManager, &WorldLocations](uint32 InternalObjectIndex)
 	{
-		const FNetRefHandleManager::FReplicatedObjectData& ObjectData = LocalNetRefHandleManager.GetReplicatedObjectDataNoCheck(InternalObjectIndex);
-		if (ObjectData.InstanceProtocol && EnumHasAnyFlags(ObjectData.InstanceProtocol->InstanceTraits, EReplicationInstanceProtocolTraits::NeedsWorldLocationUpdate))
+		if (WorldLocations.HasInfoForObject(InternalObjectIndex))
 		{
+			const FNetRefHandleManager::FReplicatedObjectData& ObjectData = LocalNetRefHandleManager.GetReplicatedObjectDataNoCheck(InternalObjectIndex);
+
             const UNetObjectFactory::FWorldInfoContext Context
             {
                 .Instance = ReplicatedInstances[InternalObjectIndex],
@@ -1939,17 +1950,17 @@ void UObjectReplicationBridge::SetClassDynamicFilterConfig(FName ClassPathName,
 		{
 			if (OldFilterInfo->FilterHandle != FilterHandle)
 			{
-				UE_LOG(LogIrisFilterConfig, Log, TEXT("SetClassDynamicFilterConfig assigned %s to use filter %s (Profile %s). Previously using filter %s."),
+				UE_LOG_BRIDGEID(LogIrisFilterConfig, Log, TEXT("SetClassDynamicFilterConfig assigned %s to use filter %s (Profile %s). Previously using filter %s."),
 					*ClassPathName.ToString(), *(ReplicationSystem->GetFilterName(FilterHandle).ToString()), *FilterProfile.ToString(), *(ReplicationSystem->GetFilterName(OldFilterInfo->FilterHandle).ToString()));
 			}
 			else
 			{
-				UE_LOG(LogIrisFilterConfig, Log, TEXT("SetClassDynamicFilterConfig assigned %s to use filter %s but the class was already assigned to this filter."), *ClassPathName.ToString(), *(ReplicationSystem->GetFilterName(FilterHandle).ToString()));
+				UE_LOG_BRIDGEID(LogIrisFilterConfig, Log, TEXT("SetClassDynamicFilterConfig assigned %s to use filter %s but the class was already assigned to this filter."), *ClassPathName.ToString(), *(ReplicationSystem->GetFilterName(FilterHandle).ToString()));
 			}
 		}
 		else
 		{
-			UE_LOG(LogIrisFilterConfig, Log, TEXT("SetClassDynamicFilterConfig assigned %s to use filter %s (Profile %s)."), *ClassPathName.ToString(), *(ReplicationSystem->GetFilterName(FilterHandle).ToString()), *FilterProfile.ToString());
+			UE_LOG_BRIDGEID(LogIrisFilterConfig, Log, TEXT("SetClassDynamicFilterConfig assigned %s to use filter %s (Profile %s)."), *ClassPathName.ToString(), *(ReplicationSystem->GetFilterName(FilterHandle).ToString()), *FilterProfile.ToString());
 		}
 	}
 
@@ -2383,14 +2394,14 @@ void UObjectReplicationBridge::OnProtocolMismatchReported(FNetRefHandle RefHandl
 	const FInternalNetRefIndex ObjectInternalIndex = NetRefHandleManager->GetInternalIndex(RefHandle);
 	if (ObjectInternalIndex == FNetRefHandleManager::InvalidInternalIndex)
 	{
-		UE_LOG(LogIris, Warning, TEXT("OnProtocolMismatchReported from Connection:%s for %s. But object has no InternalIndex."), *PrintConnectionInfo(ConnectionId), *RefHandle.ToString());
+		UE_LOG_BRIDGEID(LogIris, Warning, TEXT("OnProtocolMismatchReported from Connection:%s for %s. But object has no InternalIndex."), *PrintConnectionInfo(ConnectionId), *RefHandle.ToString());
 		return;
 	}
 
 	UObject* ObjInstance = NetRefHandleManager->GetReplicatedObjectInstance(ObjectInternalIndex);
 	UObject* ObjArchetype = ObjInstance ? ObjInstance->GetArchetype() : nullptr;
 
-	UE_LOG(LogIris, Error, TEXT("OnProtocolMismatchReported from client:%s when instancing %s. CDO:%s ReplicatedObject:%s NetObject:%s"), *PrintConnectionInfo(ConnectionId), *RefHandle.ToString(), *GetPathNameSafe(ObjArchetype), *GetNameSafe(ObjInstance), *NetRefHandleManager->PrintObjectFromIndex(ObjectInternalIndex));
+	UE_LOG_BRIDGEID(LogIris, Error, TEXT("OnProtocolMismatchReported from client:%s when instancing %s. CDO:%s ReplicatedObject:%s NetObject:%s"), *PrintConnectionInfo(ConnectionId), *RefHandle.ToString(), *GetPathNameSafe(ObjArchetype), *GetNameSafe(ObjInstance), *NetRefHandleManager->PrintObjectFromIndex(ObjectInternalIndex));
 
 	if (UE_LOG_ACTIVE(LogIris, Error))
 	{
@@ -2398,14 +2409,14 @@ void UObjectReplicationBridge::OnProtocolMismatchReported(FNetRefHandle RefHandl
 		const FReplicationInstanceProtocol* InstanceProtocol = ObjectData.InstanceProtocol;
 		if (!InstanceProtocol)
 		{
-			UE_LOG(LogIris, Warning, TEXT("OnProtocolMismatchReported from Connection:%s for %s. But object %s has no InstanceProtocol."), *PrintConnectionInfo(ConnectionId), *RefHandle.ToString(), *GetNameSafe(ObjInstance));
+			UE_LOG_BRIDGEID(LogIris, Warning, TEXT("OnProtocolMismatchReported from Connection:%s for %s. But object %s has no InstanceProtocol."), *PrintConnectionInfo(ConnectionId), *RefHandle.ToString(), *GetNameSafe(ObjInstance));
 			return;
 		}
 
 		const FReplicationProtocol* Protocol = ObjectData.Protocol;
 		if (!Protocol)
 		{
-			UE_LOG(LogIris, Warning, TEXT("OnProtocolMismatchReported from Connection:%s for %s. But object %s has no Protocol."), *PrintConnectionInfo(ConnectionId), *RefHandle.ToString(), *GetNameSafe(ObjInstance));
+			UE_LOG_BRIDGEID(LogIris, Warning, TEXT("OnProtocolMismatchReported from Connection:%s for %s. But object %s has no Protocol."), *PrintConnectionInfo(ConnectionId), *RefHandle.ToString(), *GetNameSafe(ObjInstance));
 			return;
 		}
 
@@ -2458,11 +2469,11 @@ void UObjectReplicationBridge::OnErrorWithNetRefHandleReported(UE::Net::ENetRefH
 	const FInternalNetRefIndex ObjectInternalIndex = NetRefHandleManager->GetInternalIndex(RefHandle);
 	if (ObjectInternalIndex == FNetRefHandleManager::InvalidInternalIndex)
 	{
-		UE_LOG(LogIris, Warning, TEXT("OnErrorWithNetRefHandleReported: %s from client %s for %s but object has no InternalIndex."), LexToString(ErrorType), *PrintConnectionInfo(ConnectionId), *RefHandle.ToString());
+		UE_LOG_BRIDGEID(LogIris, Warning, TEXT("OnErrorWithNetRefHandleReported: %s from client %s for %s but object has no InternalIndex."), LexToString(ErrorType), *PrintConnectionInfo(ConnectionId), *RefHandle.ToString());
 		return;
 	}
 
-	UE_LOG(LogIris, Error, TEXT("OnErrorWithNetRefHandleReported: %s from client %s. Problematic object was %s"),  LexToString(ErrorType), *PrintConnectionInfo(ConnectionId), *NetRefHandleManager->PrintObjectFromIndex(ObjectInternalIndex));
+	UE_LOG_BRIDGEID(LogIris, Error, TEXT("OnErrorWithNetRefHandleReported: %s from client %s. Problematic object was %s"),  LexToString(ErrorType), *PrintConnectionInfo(ConnectionId), *NetRefHandleManager->PrintObjectFromIndex(ObjectInternalIndex));
 }
 
 TArray<uint32> UObjectReplicationBridge::FindConnectionsFromArgs(const TArray<FString>& Args) const
diff --git a/Engine/Source/Runtime/Experimental/Iris/Core/Private/Iris/ReplicationSystem/ReplicationProtocol.h b/Engine/Source/Runtime/Experimental/Iris/Core/Private/Iris/ReplicationSystem/ReplicationProtocol.h
--- a/Engine/Source/Runtime/Experimental/Iris/Core/Private/Iris/ReplicationSystem/ReplicationProtocol.h
+++ b/Engine/Source/Runtime/Experimental/Iris/Core/Private/Iris/ReplicationSystem/ReplicationProtocol.h
@@ -23,8 +23,7 @@ enum class EReplicationInstanceProtocolTraits : uint16
 	NeedsLegacyCallbacks		= NeedsPoll << 1,
 	IsBound						= NeedsLegacyCallbacks << 1,
 	NeedsPreSendUpdate			= IsBound << 1,
-	NeedsWorldLocationUpdate	= NeedsPreSendUpdate << 1,
-	HasPartialPushBasedDirtiness = NeedsWorldLocationUpdate << 1,
+	HasPartialPushBasedDirtiness = NeedsPreSendUpdate << 1,
 	HasFullPushBasedDirtiness	 = HasPartialPushBasedDirtiness << 1,
 	// Whether there's a state that has object references
 	HasObjectReference = HasFullPushBasedDirtiness << 1,
diff --git a/Engine/Source/Runtime/Experimental/Iris/Core/Private/Iris/ReplicationSystem/ReplicationProtocolManager.cpp b/Engine/Source/Runtime/Experimental/Iris/Core/Private/Iris/ReplicationSystem/ReplicationProtocolManager.cpp
--- a/Engine/Source/Runtime/Experimental/Iris/Core/Private/Iris/ReplicationSystem/ReplicationProtocolManager.cpp
+++ b/Engine/Source/Runtime/Experimental/Iris/Core/Private/Iris/ReplicationSystem/ReplicationProtocolManager.cpp
@@ -38,7 +38,7 @@ static FAutoConsoleVariableRef CVarIrisLogReplicationProtocols(
 	TEXT("If true, log all created replication protocols."
 	));
 
-FReplicationInstanceProtocol* FReplicationProtocolManager::CreateInstanceProtocol(const FReplicationFragments& Fragments)
+FReplicationInstanceProtocol* FReplicationProtocolManager::CreateInstanceProtocol(const FReplicationFragments& Fragments, UE::Net::EReplicationFragmentTraits ObjectTraits)
 {
 	const uint32 FragmentCount = Fragments.Num();
 	if (!ensure(FragmentCount < 65536))
@@ -71,7 +71,7 @@ FReplicationInstanceProtocol* FReplicationProtocolManager::CreateInstanceProtoco
 
 	// Fill in fragment data and fragment pointer
 	uint32 FragmentIt = 0u;
-	EReplicationFragmentTraits AccumulatedTraits = EReplicationFragmentTraits::None;
+	EReplicationFragmentTraits AccumulatedTraits = ObjectTraits;
 	EReplicationFragmentTraits SharedTraits = FragmentCount > 0 ? ~EReplicationFragmentTraits::None : EReplicationFragmentTraits::None;
 	for (const FReplicationFragmentInfo& Info : Fragments)
 	{
@@ -97,10 +97,6 @@ FReplicationInstanceProtocol* FReplicationProtocolManager::CreateInstanceProtoco
 	{
 		InstanceTraits |= EReplicationInstanceProtocolTraits::NeedsPreSendUpdate;
 	}
-	if (EnumHasAnyFlags(AccumulatedTraits, EReplicationFragmentTraits::NeedsWorldLocationUpdate))
-	{
-		InstanceTraits |= EReplicationInstanceProtocolTraits::NeedsWorldLocationUpdate;
-	}
 	if (EnumHasAnyFlags(AccumulatedTraits, EReplicationFragmentTraits::HasObjectReference))
 	{
 		InstanceTraits |= EReplicationInstanceProtocolTraits::HasObjectReference;
diff --git a/Engine/Source/Runtime/Experimental/Iris/Core/Private/Iris/ReplicationSystem/ReplicationProtocolManager.h b/Engine/Source/Runtime/Experimental/Iris/Core/Private/Iris/ReplicationSystem/ReplicationProtocolManager.h
--- a/Engine/Source/Runtime/Experimental/Iris/Core/Private/Iris/ReplicationSystem/ReplicationProtocolManager.h
+++ b/Engine/Source/Runtime/Experimental/Iris/Core/Private/Iris/ReplicationSystem/ReplicationProtocolManager.h
@@ -40,11 +40,12 @@ public:
 	/* Destroy existing replication protocol */
 	IRISCORE_API void DestroyReplicationProtocol(const FReplicationProtocol* ReplicationProtocol);
 
-	/** Create instance protocol from registered fragment data
-		Lifetime either the same as the lifetime of a remotely created instance (we are the client) or controlled by
-		calls to the bridge api to replicate an object or not
-	*/
-	IRISCORE_API static FReplicationInstanceProtocol* CreateInstanceProtocol(const FReplicationFragments& Fragments);
+	/** 
+	 * Create instance protocol from registered fragment data
+	 * Lifetime either the same as the lifetime of a remotely created instance (we are the client) or controlled by	calls to the bridge api to replicate an object or not
+	 * @param ObjectTraits Traits that are passed by the object instance itself
+	 */
+	IRISCORE_API static FReplicationInstanceProtocol* CreateInstanceProtocol(const FReplicationFragments& Fragments, UE::Net::EReplicationFragmentTraits ObjectTraits);
 	IRISCORE_API static void DestroyInstanceProtocol(FReplicationInstanceProtocol*);
 
 	IRISCORE_API void InvalidateDescriptor(const FReplicationStateDescriptor* InvalidatedReplicationStateDescriptor);
diff --git a/Engine/Source/Runtime/Experimental/Iris/Core/Public/Iris/ReplicationSystem/ReplicationFragment.h b/Engine/Source/Runtime/Experimental/Iris/Core/Public/Iris/ReplicationSystem/ReplicationFragment.h
--- a/Engine/Source/Runtime/Experimental/Iris/Core/Public/Iris/ReplicationSystem/ReplicationFragment.h
+++ b/Engine/Source/Runtime/Experimental/Iris/Core/Public/Iris/ReplicationSystem/ReplicationFragment.h
@@ -129,11 +129,8 @@ enum class EReplicationFragmentTraits : uint32
 	// Fragment requires PreSendUpdate to be called 
 	NeedsPreSendUpdate					= NeedsLegacyCallbacks << 1,
 
-	// Fragment requires world location to be updated before filtering and prioritization
-	NeedsWorldLocationUpdate			= NeedsPreSendUpdate << 1,
-
 	// Fragment supports push based dirtiness
-	HasPushBasedDirtiness				= NeedsWorldLocationUpdate << 1,
+	HasPushBasedDirtiness				= NeedsPreSendUpdate << 1,
 
 	// Fragment is a PropertyReplication, or is using PropertyReplicationState
 	HasPropertyReplicationState			= HasPushBasedDirtiness << 1,
