Dataflow image : fix crash when write channels from an zero sized image

diff --git a/Engine/Source/Runtime/Experimental/Dataflow/Core/Private/Dataflow/DataflowImage.cpp b/Engine/Source/Runtime/Experimental/Dataflow/Core/Private/Dataflow/DataflowImage.cpp
--- a/Engine/Source/Runtime/Experimental/Dataflow/Core/Private/Dataflow/DataflowImage.cpp
+++ b/Engine/Source/Runtime/Experimental/Dataflow/Core/Private/Dataflow/DataflowImage.cpp
@@ -101,7 +101,14 @@ void FDataflowImage::ReadChannel(EDataflowImageChannel Channel, FDataflowImage&
 
 void FDataflowImage::WriteChannel(EDataflowImageChannel Channel, const FDataflowImage& SrcImage)
 {
+	FDataflowImage DefaultMinImage;
+	
 	FImageView SrcImageView = SrcImage.Image;
+	if (SrcImageView.GetWidth() == 0 || SrcImageView.GetHeight() == 0)
+	{
+		DefaultMinImage.CreateFromColor(4, 4, FLinearColor::Black);
+		SrcImageView = DefaultMinImage.Image;
+	}
 
 	FImage ResizedSrcImage;
 	if (SrcImageView.GetWidth() != GetWidth() && SrcImageView.GetHeight() != GetHeight())
