Calling FIoDispatcher::DoesChunkExist for a chunk that is not installed should return false.

#rnx

diff --git a/Engine/Source/Runtime/Experimental/IoStore/OnDemand/Private/OnDemandInstallCache.cpp b/Engine/Source/Runtime/Experimental/IoStore/OnDemand/Private/OnDemandInstallCache.cpp
--- a/Engine/Source/Runtime/Experimental/IoStore/OnDemand/Private/OnDemandInstallCache.cpp
+++ b/Engine/Source/Runtime/Experimental/IoStore/OnDemand/Private/OnDemandInstallCache.cpp
@@ -1515,8 +1515,13 @@ void FOnDemandInstallCache::UpdatePriorityForIoRequest(FIoRequestImpl* Request)
 
 bool FOnDemandInstallCache::DoesChunkExist(const FIoChunkId& ChunkId) const
 {
-	const TIoStatusOr<uint64> Status = GetSizeForChunk(ChunkId);
-	return Status.IsOk();
+	if (FOnDemandChunkInfo ChunkInfo = IoStore.GetInstalledChunkInfo(ChunkId))
+	{
+		const FCasLocation CasLoc = Cas.FindChunk(ChunkInfo.Hash());
+		return CasLoc.IsValid();
+	}
+
+	return false;
 }
 
 TIoStatusOr<uint64> FOnDemandInstallCache::GetSizeForChunk(const FIoChunkId& ChunkId) const
