Chaos: Fix event manager registration to properly handle re-entrant register and unregister calls.

The previous method of handling re-entrancy relied on strict control of a write lock to detect when a register happens during the dispatch of events. However when transactions are enabled the write lock leaks due to the requirement of holding a write lock until the end of a transaction. This could lead to deferred handlers building up outside of a dispatch and maybe sitting in the pending list for many frames. If an object that contained a handler was destroyed and garbage collected while its handler was still pending the next dispatch would crash when attempting to process the deferred request.

As these calls all come from the game thread we can track re-entrancy with a guard flag instead and handle the situation without relying on the lock. All the locks are still present and the locking semantics of the event manager haven't changed, although it should be re-examined in the future as it's likely we don't need the handler and the deffered handler locks as they're never accessed concurrently

diff --git a/Engine/Source/Runtime/Experimental/Chaos/Public/EventManager.h b/Engine/Source/Runtime/Experimental/Chaos/Public/EventManager.h
--- a/Engine/Source/Runtime/Experimental/Chaos/Public/EventManager.h
+++ b/Engine/Source/Runtime/Experimental/Chaos/Public/EventManager.h
@@ -9,6 +9,7 @@
 #include "Misc/ScopeLock.h"
 #include "Chaos/Defines.h"
 #include "Misc/TransactionallySafeRWLock.h"
+#include "Misc/ScopeRWLock.h"
 
 namespace Chaos
 {
@@ -167,6 +168,7 @@ namespace Chaos
 		{
 			InjectedFunction = Other.InjectedFunction;
 			EventBuffer = MoveTemp(Other.EventBuffer);
+			bDispatchInProgress = Other.bDispatchInProgress;
 		}
 
 		/**
@@ -174,13 +176,13 @@ namespace Chaos
 		 */
 		~TEventContainer()
 		{
-			HandlerLock.WriteLock();
+			UE::TWriteScopeLock ScopeLock(HandlerLock);
+
 			for (TPair<void*, FEventHandlerPtr>& Pair : HandlerMap)
 			{
 				delete Pair.Value;
 				Pair.Value = nullptr;
 			}
-			HandlerLock.WriteUnlock();
 		}
 #
 		/**
@@ -188,35 +190,35 @@ namespace Chaos
 		 */
 		virtual void RegisterHandler(const FEventHandlerPtr& Handler)
 		{
-			// If Register Handler is called by the user inside a dispatch event we ll end up with a dead lock. 
-			// so if we cannot lock it we store the event handler to try to register later
-			if (HandlerLock.TryWriteLock())
+			check(IsInGameThread());
+
+			if(bDispatchInProgress)
 			{
-				HandlerMap.Add(Handler->GetHandler(), Handler);
-				TArray<UObject*> ProxyOwners;
-				bool bValidProxyFilter = Handler->GetInterestedProxyOwners(ProxyOwners);
+				UE::TWriteScopeLock ScopeLock(DeferredHandlerLock);
+				DeferredHandlers.AddUnique(Handler);
 
-				if (bValidProxyFilter)
-				{
-					for (UObject* ProxyOwner : ProxyOwners)
-					{
-						ProxyOwnerToHandlerMap.Add(ProxyOwner, Handler);
-					}
-				}
-				else
+				return;
+			}
+
+			UE::TRWScopeLock ScopeLock(HandlerLock, SLT_Write);
+
+			HandlerMap.Add(Handler->GetHandler(), Handler);
+			TArray<UObject*> ProxyOwners;
+			bool bValidProxyFilter = Handler->GetInterestedProxyOwners(ProxyOwners);
+
+			if(bValidProxyFilter)
+			{
+				for(UObject* ProxyOwner : ProxyOwners)
 				{
-					if (GetProxyToIndexMap(EventBuffer.Get()->GetConsumerBuffer()) != nullptr) // Only if our type supports getting the ProxyToIndexMap do we bother adding this
-					{
-						HandlersNotInProxyOwnerMap.Add(Handler->GetHandler(), Handler);
-					}
+					ProxyOwnerToHandlerMap.Add(ProxyOwner, Handler);
 				}
-				HandlerLock.WriteUnlock();
 			}
 			else
 			{
-				DeferredHandlerLock.WriteLock();
-				DeferredHandlers.AddUnique(Handler);
-				DeferredHandlerLock.WriteUnlock();
+				if(GetProxyToIndexMap(EventBuffer.Get()->GetConsumerBuffer()) != nullptr) // Only if our type supports getting the ProxyToIndexMap do we bother adding this
+				{
+					HandlersNotInProxyOwnerMap.Add(Handler->GetHandler(), Handler);
+				}
 			}
 		}
 
@@ -225,39 +227,39 @@ namespace Chaos
 		 */
 		virtual void UnregisterHandler(const void* InHandler)
 		{
-			// If Unregister Handler is called by the user inside a dispatch event we ll end up with a dead lock. 
-			// so if we cannot lock it we store the event handler to try to unregister it later
-			if (HandlerLock.TryWriteLock())
+			check(IsInGameThread());
+
+			if(bDispatchInProgress)
 			{
-				TArray<TPair<UObject*, FEventHandlerPtr>> KeysAndValuesToRemove;
-				for (TPair<UObject*, FEventHandlerPtr>& KeyValue : ProxyOwnerToHandlerMap)
-				{
-					const FEventHandlerPtr Value = KeyValue.Value;
-					if (Value && Value->GetHandler() == InHandler)
-					{
-						KeysAndValuesToRemove.Add(KeyValue);
-					}
-				}
+				UE::TWriteScopeLock ScopeLock(DeferredHandlerLock);
+				DeferredUnregisterHandlers.AddUnique(InHandler);
 
-				for (TPair<UObject*, FEventHandlerPtr>& KeyAndValue : KeysAndValuesToRemove)
-				{
-					ProxyOwnerToHandlerMap.Remove(KeyAndValue.Get<0>(), KeyAndValue.Get<1>());
-				}
+				return;
+			}
 
-				HandlersNotInProxyOwnerMap.Remove(InHandler);
+			UE::TWriteScopeLock ScopeLock(HandlerLock);
 
-				if (FEventHandlerPtr* HandlerPtr = HandlerMap.Find(InHandler))
+			TArray<TPair<UObject*, FEventHandlerPtr>> KeysAndValuesToRemove;
+			for(const TPair<UObject*, FEventHandlerPtr>& KeyValue : ProxyOwnerToHandlerMap)
+			{
+				const FEventHandlerPtr Value = KeyValue.Value;
+				if(Value && Value->GetHandler() == InHandler)
 				{
-					DeleteHandler(*HandlerPtr);
-					HandlerMap.Remove(InHandler);
+					KeysAndValuesToRemove.Add(KeyValue);
 				}
-				HandlerLock.WriteUnlock();
 			}
-			else
+
+			for(const TPair<UObject*, FEventHandlerPtr>& KeyAndValue : KeysAndValuesToRemove)
 			{
-				DeferredHandlerLock.WriteLock();
-				DeferredUnregisterHandlers.AddUnique(InHandler);
-				DeferredHandlerLock.WriteUnlock();
+				ProxyOwnerToHandlerMap.Remove(KeyAndValue.Get<0>(), KeyAndValue.Get<1>());
+			}
+
+			HandlersNotInProxyOwnerMap.Remove(InHandler);
+
+			if(FEventHandlerPtr* HandlerPtr = HandlerMap.Find(InHandler))
+			{
+				DeleteHandler(*HandlerPtr);
+				HandlerMap.Remove(InHandler);
 			}
 		}
 
@@ -301,6 +303,8 @@ namespace Chaos
 		 */
 		virtual void DispatchConsumerData()
 		{
+			check(IsInGameThread());
+
 			const PayloadType* Buffer = EventBuffer.Get()->GetConsumerBuffer();
 
 			// We dispatch events even when the event buffers are empty because the absence
@@ -312,46 +316,57 @@ namespace Chaos
 				return;
 			}
 
+			// Unregister any handlers that should no longer be invoked
 			UnregisterDeferedHandler();
-			HandlerLock.ReadLock();
-			const TMap<IPhysicsProxyBase*, TArray<int32>>* Map = GetProxyToIndexMap(Buffer); // Use t his map to get all proxies used in the event buffer
-			// Only take this path if we have fewer Events than Handlers
-			if (Map && Map->Num() + HandlersNotInProxyOwnerMap.Num() < HandlerMap.Num())
+			
 			{
-				if (!ProxyOwnerToHandlerMap.IsEmpty())
+				// We're going to iterate the handlers and invoke them, read lock the list to block any changes
+				UE::TReadScopeLock ScopeLock(HandlerLock);
+
+				// Dispatch process beginning, we can no longer modify the handlers until the dispatch
+				// completes, register and unregister will defer while this guard is set
+				TGuardValue<bool> ReentrancyGuard(bDispatchInProgress, true);
+
+				const TMap<IPhysicsProxyBase*, TArray<int32>>* Map = GetProxyToIndexMap(Buffer); // Use t his map to get all proxies used in the event buffer
+				// Only take this path if we have fewer Events than Handlers
+				if(Map && Map->Num() + HandlersNotInProxyOwnerMap.Num() < HandlerMap.Num())
 				{
-					TSet<FEventHandlerPtr> UniqueHandlers;
-					UniqueHandlers.Reserve(Map->Num());
-					for (auto& KeyValue : *Map) // Only iterating over objects that are associated with events here
+					if(!ProxyOwnerToHandlerMap.IsEmpty())
 					{
-						const IPhysicsProxyBase* Proxy = KeyValue.Get<0>();
-						const UObject* Owner = Proxy->GetOwner();
-						for (TMultiMap<UObject*, FEventHandlerPtr>::TConstKeyIterator It = ProxyOwnerToHandlerMap.CreateConstKeyIterator(Owner); It; ++It)
-						{ 
-							UniqueHandlers.Add(It.Value());
+						TSet<FEventHandlerPtr> UniqueHandlers;
+						UniqueHandlers.Reserve(Map->Num());
+						for(auto& KeyValue : *Map) // Only iterating over objects that are associated with events here
+						{
+							const IPhysicsProxyBase* Proxy = KeyValue.Get<0>();
+							const UObject* Owner = Proxy->GetOwner();
+							for(TMultiMap<UObject*, FEventHandlerPtr>::TConstKeyIterator It = ProxyOwnerToHandlerMap.CreateConstKeyIterator(Owner); It; ++It)
+							{
+								UniqueHandlers.Add(It.Value());
+							}
+						}
+
+						for(const FEventHandlerPtr Handler : UniqueHandlers)
+						{
+							Handler->HandleEvent(Buffer);
 						}
 					}
 
-					for (const FEventHandlerPtr Handler : UniqueHandlers)
+					for(const TPair<void*, FEventHandlerPtr>& Pair : HandlersNotInProxyOwnerMap)
 					{
-						Handler->HandleEvent(Buffer);
+						Pair.Value->HandleEvent(Buffer);
 					}
 				}
-
-				for (const TPair<void*, FEventHandlerPtr>& Pair : HandlersNotInProxyOwnerMap)
-				{
-					Pair.Value->HandleEvent(Buffer);
-				}
-			}
-			else
-			// This path is taken if there are fewer Handlers than events or the handler does not support GetProxyToIndexMap
-			{
-				for (const TPair<void*, FEventHandlerPtr>& Pair : HandlerMap)
+				else
+					// This path is taken if there are fewer Handlers than events or the handler does not support GetProxyToIndexMap
 				{
-					Pair.Value->HandleEvent(Buffer);
+					for(const TPair<void*, FEventHandlerPtr>& Pair : HandlerMap)
+					{
+						Pair.Value->HandleEvent(Buffer);
+					}
 				}
 			}
-			HandlerLock.ReadUnlock();
+
+			// Register any new handlers enqueued during the dispatch process
 			RegisterDeferedHandler();
 		}
 
@@ -359,12 +374,16 @@ private:
 
 	void UnregisterDeferedHandler()
 	{
-		DeferredHandlerLock.WriteLock();
-		// Move array
-		TArray<const void*> DeferredUnregisterHandlersCopy(MoveTemp(DeferredUnregisterHandlers));
-		check(DeferredUnregisterHandlers.Num() == 0);
-		DeferredHandlerLock.WriteUnlock();
-		for (const void* HandlerPtr : DeferredUnregisterHandlersCopy)
+		TArray<const void*> DeferredRequests;
+
+		{
+			UE::TWriteScopeLock ScopeLock(DeferredHandlerLock);
+			
+			DeferredRequests = MoveTemp(DeferredUnregisterHandlers);
+			check(DeferredUnregisterHandlers.Num() == 0);
+		}
+
+		for (const void* HandlerPtr : DeferredRequests)
 		{
 			UnregisterHandler(HandlerPtr);
 		}
@@ -372,12 +391,16 @@ private:
 
 	void RegisterDeferedHandler()
 	{
-		DeferredHandlerLock.WriteLock();
-		// Move array
-		TArray<FEventHandlerPtr> DeferredHandlersCopy(MoveTemp(DeferredHandlers));
-		check(DeferredHandlers.Num() == 0);
-		DeferredHandlerLock.WriteUnlock();
-		for (const FEventHandlerPtr& HandlerPtr : DeferredHandlersCopy)
+		TArray<FEventHandlerPtr> DeferredRequests;
+
+		{
+			UE::TWriteScopeLock ScopeLock(DeferredHandlerLock);
+			
+			DeferredRequests = MoveTemp(DeferredHandlers);
+			check(DeferredHandlers.Num() == 0);
+		}
+
+		for (const FEventHandlerPtr& HandlerPtr : DeferredRequests)
 		{
 			RegisterHandler(HandlerPtr);
 		}
@@ -411,6 +434,9 @@ private:
 		TArray<FEventHandlerPtr> DeferredHandlers; // Store handler to register, they couldn't registered to avoid reentrant lock
 		TArray<const void*> DeferredUnregisterHandlers; // Store handler to unregister, they couldn't unregistered to avoid reentrant lock
 		FTransactionallySafeRWLock DeferredHandlerLock; // protect access to DeferredHandlers and DeferredUnregisterHandlers
+
+		/** Guard flag for the dispatch process - used to defer requests to register or unregister new handlers during dispatch */
+		bool bDispatchInProgress = false;
 	};
 
 	/**
