[ChaosVD] Adding CVD support to serialize new types of game thread data added a few days ago for DSTM scene serialization

diff --git a/Engine/Source/Runtime/Core/Public/UObject/FortniteMainBranchObjectVersions.inl b/Engine/Source/Runtime/Core/Public/UObject/FortniteMainBranchObjectVersions.inl
--- a/Engine/Source/Runtime/Core/Public/UObject/FortniteMainBranchObjectVersions.inl
+++ b/Engine/Source/Runtime/Core/Public/UObject/FortniteMainBranchObjectVersions.inl
@@ -610,5 +610,8 @@ FORTNITE_MAIN_VERSION(ClothAssetSkinweightsValidation, 198)
 // Switching verse from right handed to left handed
 FORTNITE_MAIN_VERSION(VerseRightToLeftHandedness, 199)
 
+// Added additional data required to record and represent particle data from the game thread (Kinematic targets, and SQ rejection reasons)
+FORTNITE_MAIN_VERSION(AdditionalGameThreadDataSupportInChaosVisualDebugger, 200)
+
 // -----<new versions can be added above this line>-------------------------------------------------
 #undef FORTNITE_MAIN_VERSION
\ No newline at end of file
diff --git a/Engine/Source/Runtime/Experimental/ChaosVisualDebugger/Private/DataWrappers/ChaosVDJointDataWrappers.cpp b/Engine/Source/Runtime/Experimental/ChaosVisualDebugger/Private/DataWrappers/ChaosVDJointDataWrappers.cpp
--- a/Engine/Source/Runtime/Experimental/ChaosVisualDebugger/Private/DataWrappers/ChaosVDJointDataWrappers.cpp
+++ b/Engine/Source/Runtime/Experimental/ChaosVisualDebugger/Private/DataWrappers/ChaosVDJointDataWrappers.cpp
@@ -3,6 +3,8 @@
 #include "DataWrappers/ChaosVDJointDataWrappers.h"
 
 #include "DataWrappers/ChaosVDDataSerializationMacros.h"
+#include "UObject/FortniteMainBranchObjectVersion.h"
+
 
 bool FChaosVDJointStateDataWrapper::Serialize(FArchive& Ar)
 {
@@ -271,5 +273,11 @@ bool FChaosVDJointConstraint::Serialize(FArchive& Ar)
 	Ar << JointSettings;
 	Ar << PhysicsThreadJointState;
 
+	Ar.UsingCustomVersion(FFortniteMainBranchObjectVersion::GUID);
+	if (Ar.CustomVer(FFortniteMainBranchObjectVersion::GUID) >= FFortniteMainBranchObjectVersion::AdditionalGameThreadDataSupportInChaosVisualDebugger)
+	{
+		Ar << GameThreadJointState;
+	}
+
 	return !Ar.IsError();
 }
diff --git a/Engine/Source/Runtime/Experimental/ChaosVisualDebugger/Private/DataWrappers/ChaosVDParticleDataWrapper.cpp b/Engine/Source/Runtime/Experimental/ChaosVisualDebugger/Private/DataWrappers/ChaosVDParticleDataWrapper.cpp
--- a/Engine/Source/Runtime/Experimental/ChaosVisualDebugger/Private/DataWrappers/ChaosVDParticleDataWrapper.cpp
+++ b/Engine/Source/Runtime/Experimental/ChaosVisualDebugger/Private/DataWrappers/ChaosVDParticleDataWrapper.cpp
@@ -219,6 +219,7 @@ bool FChaosVDParticleDataWrapper::Serialize(FArchive& Ar)
 
 	Ar << ParticlePositionRotation;
 	Ar << ParticleVelocities;
+
 	Ar << ParticleDynamics;
 	Ar << ParticleDynamicsMisc;
 	Ar << ParticleMassProps;
@@ -226,6 +227,14 @@ bool FChaosVDParticleDataWrapper::Serialize(FArchive& Ar)
 	Ar << CollisionDataPerShape;
 
 	Ar << ParticleCluster;
+	
+	Ar.UsingCustomVersion(FFortniteMainBranchObjectVersion::GUID);
+	if (Ar.CustomVer(FFortniteMainBranchObjectVersion::GUID) >= FFortniteMainBranchObjectVersion::AdditionalGameThreadDataSupportInChaosVisualDebugger)
+	{
+		Ar << ParticleContext;
+		Ar << ParticleKinematicTarget;
+		Ar << ParticleVWSmooth;
+	}
 
 	return !Ar.IsError();
 }
diff --git a/Engine/Source/Runtime/Experimental/ChaosVisualDebugger/Private/DataWrappers/ChaosVDQueryDataWrappers.cpp b/Engine/Source/Runtime/Experimental/ChaosVisualDebugger/Private/DataWrappers/ChaosVDQueryDataWrappers.cpp
--- a/Engine/Source/Runtime/Experimental/ChaosVisualDebugger/Private/DataWrappers/ChaosVDQueryDataWrappers.cpp
+++ b/Engine/Source/Runtime/Experimental/ChaosVisualDebugger/Private/DataWrappers/ChaosVDQueryDataWrappers.cpp
@@ -3,6 +3,8 @@
 #include "DataWrappers/ChaosVDQueryDataWrappers.h"
 
 #include "DataWrappers/ChaosVDDataSerializationMacros.h"
+#include "UObject/FortniteMainBranchObjectVersion.h"
+
 
 bool FChaosVDCollisionResponseParams::Serialize(FArchive& Ar)
 {
@@ -155,6 +157,12 @@ bool FChaosVDQueryVisitStep::Serialize(FArchive& Ar)
 
 	Ar << QueryFastData;
 
+	Ar.UsingCustomVersion(FFortniteMainBranchObjectVersion::GUID);
+	if (Ar.CustomVer(FFortniteMainBranchObjectVersion::GUID) >= FFortniteMainBranchObjectVersion::AdditionalGameThreadDataSupportInChaosVisualDebugger)
+	{
+		Ar << RejectReason;
+	}
+
 	return !Ar.IsError();
 }
 
