ChaosFlesh:
-Added per-muscle InflationVolumeScale parameter in SetMuscleActivationParameter node to control muscle rest volume
-Added solver support to modify element rest state based on InflationVolumeScale
-Added p.Chaos.Deformable.LogExtremeParticles to print most deformed particle indices for debugging

diff --git a/Engine/Plugins/Experimental/ChaosFlesh/Source/ChaosFleshNodes/Private/Dataflow/ChaosFleshComputeMuscleActivationNode.cpp b/Engine/Plugins/Experimental/ChaosFlesh/Source/ChaosFleshNodes/Private/Dataflow/ChaosFleshComputeMuscleActivationNode.cpp
--- a/Engine/Plugins/Experimental/ChaosFlesh/Source/ChaosFleshNodes/Private/Dataflow/ChaosFleshComputeMuscleActivationNode.cpp
+++ b/Engine/Plugins/Experimental/ChaosFlesh/Source/ChaosFleshNodes/Private/Dataflow/ChaosFleshComputeMuscleActivationNode.cpp
@@ -62,6 +62,7 @@ FSetMuscleActivationParameterNode::FSetMuscleActivationParameterNode(const UE::D
 					ParameterArray[MuscleIdx].ContractionVolumeScale = ContractionVolumeScale;
 					ParameterArray[MuscleIdx].FiberLengthRatioAtMaxActivation = GlobalFiberLengthRatioAtMaxActivation;
 					ParameterArray[MuscleIdx].MuscleLengthRatioThresholdForMaxActivation = GlobalMuscleLengthRatioThresholdForMaxActivation;
+					ParameterArray[MuscleIdx].InflationVolumeScale = GlobalInflationVolumeScale;
 				}
 			}))
 	, ResetToGlobalParameters(
@@ -72,6 +73,7 @@ FSetMuscleActivationParameterNode::FSetMuscleActivationParameterNode(const UE::D
 					ParameterArray[Idx].ContractionVolumeScale = ContractionVolumeScale;
 					ParameterArray[Idx].FiberLengthRatioAtMaxActivation = GlobalFiberLengthRatioAtMaxActivation;
 					ParameterArray[Idx].MuscleLengthRatioThresholdForMaxActivation = GlobalMuscleLengthRatioThresholdForMaxActivation;
+					ParameterArray[Idx].InflationVolumeScale = GlobalInflationVolumeScale;
 				}
 			}))
 	, ApplyCustomParameters(
@@ -96,7 +98,8 @@ void FSetMuscleActivationParameterNode::Evaluate(UE::Dataflow::FContext& Context
 			FMuscleActivation.UpdateGlobalMuscleActivationParameters(
 					ContractionVolumeScale,
 					GlobalFiberLengthRatioAtMaxActivation,
-					GlobalMuscleLengthRatioThresholdForMaxActivation);
+					GlobalMuscleLengthRatioThresholdForMaxActivation,
+					GlobalInflationVolumeScale);
 			if (ParameterMethod == EParameterMethod::Custom)
 			{
 				// match muscle names and override parameters
@@ -110,7 +113,8 @@ void FSetMuscleActivationParameterNode::Evaluate(UE::Dataflow::FContext& Context
 							MuscleIndex,
 							Params.ContractionVolumeScale,
 							Params.FiberLengthRatioAtMaxActivation,
-							Params.MuscleLengthRatioThresholdForMaxActivation);
+							Params.MuscleLengthRatioThresholdForMaxActivation,
+							Params.InflationVolumeScale);
 					}
 				}
 			}
diff --git a/Engine/Plugins/Experimental/ChaosFlesh/Source/ChaosFleshNodes/Public/Dataflow/ChaosFleshComputeMuscleActivationNode.h b/Engine/Plugins/Experimental/ChaosFlesh/Source/ChaosFleshNodes/Public/Dataflow/ChaosFleshComputeMuscleActivationNode.h
--- a/Engine/Plugins/Experimental/ChaosFlesh/Source/ChaosFleshNodes/Public/Dataflow/ChaosFleshComputeMuscleActivationNode.h
+++ b/Engine/Plugins/Experimental/ChaosFlesh/Source/ChaosFleshNodes/Public/Dataflow/ChaosFleshComputeMuscleActivationNode.h
@@ -89,12 +89,14 @@ struct FPerMuscleParameter
 	bool bCanEditMuscleName;
 	UPROPERTY(EditAnywhere, meta = (EditCondition = "bCanEditMuscleName"), Category = "Muscle Activation Parameters")
 	FString MuscleName;
-	UPROPERTY(EditAnywhere, Category = "Muscle Activation Parameters")
-	float ContractionVolumeScale;
-	UPROPERTY(EditAnywhere, Category = "Muscle Activation Parameters")
-	float FiberLengthRatioAtMaxActivation;
-	UPROPERTY(EditAnywhere, Category = "Muscle Activation Parameters")
-	float MuscleLengthRatioThresholdForMaxActivation;
+	UPROPERTY(EditAnywhere, Category = "Muscle Activation Parameters", meta = (ClampMin = "0"))
+	float ContractionVolumeScale = 1.f;
+	UPROPERTY(EditAnywhere, Category = "Muscle Activation Parameters", meta = (ClampMin = "0", ClampMax = "1"))
+	float FiberLengthRatioAtMaxActivation = 0.5f;
+	UPROPERTY(EditAnywhere, Category = "Muscle Activation Parameters", meta = (ClampMin = "0", ClampMax = "1"))
+	float MuscleLengthRatioThresholdForMaxActivation = 0.75f;
+	UPROPERTY(EditAnywhere, Category = "Muscle Activation Parameters", meta = (ClampMin = "0"))
+	float InflationVolumeScale = 1.f;
 };
 
 UENUM()
@@ -126,17 +128,21 @@ public:
 	FDataflowFunctionProperty ApplyGlobalParameters;
 
 	/* Muscles gain volume during contraction if > 1. Volume-preserving if 1. Default: 1*/
-	UPROPERTY(EditAnywhere, Category = "Global Parameters", DisplayName = "Global Contraction Volume Scale")
+	UPROPERTY(EditAnywhere, Category = "Global Parameters", DisplayName = "Global Contraction Volume Scale", meta = (ClampMin = "0"))
 	float ContractionVolumeScale = 1.f;
 
 	/* How much muscle fibers shorten at max activation 1. A smaller value means more contraction in the fiber direction. Default: 0.5 */
-	UPROPERTY(EditAnywhere, Category = "Global Parameters")
+	UPROPERTY(EditAnywhere, Category = "Global Parameters", meta = (ClampMin = "0", ClampMax = "1"))
 	float GlobalFiberLengthRatioAtMaxActivation = 0.5f;
 
 	/* Muscle length ratio (defined by origin-insertion distance) below this threshold is considered to reach max activation 1. Default: 0.75 */
-	UPROPERTY(EditAnywhere, Category = "Global Parameters")
+	UPROPERTY(EditAnywhere, Category = "Global Parameters", meta = (ClampMin = "0", ClampMax = "1"))
 	float GlobalMuscleLengthRatioThresholdForMaxActivation = 0.75f;
 
+	/* Inflates muscle rest volume if > 1 and deflates muscle rest volume if < 1. Default: 1*/
+	UPROPERTY(EditAnywhere, Category = "Global Parameters", meta = (ClampMin = "0"))
+	float GlobalInflationVolumeScale = 1.f;
+	
 	/** Click on this button to import muscle bone names (in the Transform group) from the collection. */
 	UPROPERTY(EditAnywhere, Category = "Custom Parameters", meta = (ButtonImage = "Icons.Import"))
 	FDataflowFunctionProperty ImportAllMuscleNames;
diff --git a/Engine/Source/Runtime/Experimental/Chaos/Private/Chaos/Deformable/ChaosDeformableSolver.cpp b/Engine/Source/Runtime/Experimental/Chaos/Private/Chaos/Deformable/ChaosDeformableSolver.cpp
--- a/Engine/Source/Runtime/Experimental/Chaos/Private/Chaos/Deformable/ChaosDeformableSolver.cpp
+++ b/Engine/Source/Runtime/Experimental/Chaos/Private/Chaos/Deformable/ChaosDeformableSolver.cpp
@@ -96,6 +96,7 @@ namespace Chaos::Softs
 	FDeformableXPBDCorotatedParams GDeformableXPBDCorotatedParams;
 	FAutoConsoleVariableRef CVarDeformableXPBDCorotatedBatchSize(TEXT("p.Chaos.Deformable.XPBDBatchSize"), GDeformableXPBDCorotatedParams.XPBDCorotatedBatchSize, TEXT("Batch size for physics parallel for. [def: 5]"));
 	FAutoConsoleVariableRef CVarDeformableXPBDCorotatedBatchThreshold(TEXT("p.Chaos.Deformable.XPBDBatchThreshold"), GDeformableXPBDCorotatedParams.XPBDCorotatedBatchThreshold, TEXT("Batch threshold for physics parallel for. [def: 5]"));
+	FAutoConsoleVariableRef CVarDeformableXPBDCorotatedBatchNumLogExtremeParticle (TEXT("p.Chaos.Deformable.NumLogExtremeParticle "), GDeformableXPBDCorotatedParams.NumLogExtremeParticle, TEXT("Number of most deformed particles logged. [def: 0]"));
 
 	FDeformableXPBDWeakConstraintParams GDeformableXPBDWeakConstraintParams;
 	FAutoConsoleVariableRef CVarDeformableXPBDWeakConstraintLineWidth(TEXT("p.Chaos.Deformable.XPBDWeakConstraintLineWidth"), GDeformableXPBDWeakConstraintParams.DebugLineWidth, TEXT("Line width for visualizing the double bindings in XPBD weak constraints. [def: 5]"));
@@ -424,10 +425,7 @@ namespace Chaos::Softs
 					{
 						InitializeGaussSeidelConstraintVariables();
 					}
-					if (Property.bDoLengthBasedMuscleActivation || Property.bOverrideMuscleActivationWithAnimatedCurves)
-					{
-						InitializeMuscleActivationVariables();
-					}
+					InitializeMuscleActivationVariables();
 				}
 				UninitializedProxys_Internal.SetNum(0, EAllowShrinking::Yes);
 			}
@@ -1913,6 +1911,7 @@ namespace Chaos::Softs
 
 	void FDeformableSolver::InitializeMuscleActivationVariables()
 	{
+		// Apply muscle activation lambda
 		int32 InitIndex = Evolution->AddConstraintInitRange(1, true);
 		Evolution->ConstraintInits()[InitIndex] = [this](FSolverParticles& InParticles, const FSolverReal Dt)
 		{
@@ -1987,24 +1986,42 @@ namespace Chaos::Softs
 				UE_LOG(LogChaosDeformableSolver, Warning, TEXT("Muscle activation is not supported for XPBD now, please switch to Gauss-Seidel."));
 			}
 		};
+
+		//Adjust muscle rest volume
+		if (MuscleActivationConstraints->NumMuscles())
+		{
+			if (Property.bUseGaussSeidelConstraints)
+			{
+				if (Property.bUseGSNeohookean)
+				{
+					MuscleActivationConstraints->ApplyInflationVolumeScale(*GSNeohookeanConstraints);
+				}
+				else
+				{
+					MuscleActivationConstraints->ApplyInflationVolumeScale(*GSCorotatedConstraints);
+				}
+			}
+			else
+			{
+				// TODO(flesh): Support muscle activation with XPBD, use an array for GSConstraints
+				UE_LOG(LogChaosDeformableSolver, Warning, TEXT("Inflation volume scale is only supported for the Gauss-Seidel solver."));
+			}
+		}
 	}
 
 	void FDeformableSolver::InitializeMuscleActivation(FFleshThreadingProxy& Proxy)
 	{
 		PERF_SCOPE(STAT_ChaosDeformableSolver_InitializeMuscleActivation);
-		if (Property.bUseGaussSeidelConstraints && (Property.bDoLengthBasedMuscleActivation || Property.bOverrideMuscleActivationWithAnimatedCurves))
+		//Grab MuscleActivationElements, OriginInsertion, FiberDirectionMatrix
+		const FManagedArrayCollection& Rest = Proxy.GetRestCollection();
+		GeometryCollection::Facades::FMuscleActivationFacade MuscleActivationFacade(Rest);
+		if (MuscleActivationFacade.IsValid())
 		{
-			//Grab MuscleActivationElements, OriginInsertion, FiberDirectionMatrix
-			const FManagedArrayCollection& Rest = Proxy.GetRestCollection();
-			GeometryCollection::Facades::FMuscleActivationFacade MuscleActivationFacade(Rest);
-			if (MuscleActivationFacade.IsValid())
-			{
-				int32 VertexOffset = Proxy.GetSolverParticleRange().Start;
-				// Only supports Gauss Seidel for now because we are using AllElements
-				int32 ElementOffset = AllElements->Num(); //InitializeMuscleActivation has to go before adding tetrahedrons
-				MuscleIndexOffset.Add(Proxy.GetOwner(), MuscleActivationConstraints->NumMuscles()); 
-				MuscleActivationConstraints->AddMuscles(Evolution->Particles(), MuscleActivationFacade, VertexOffset, ElementOffset);
-			}
+			int32 VertexOffset = Proxy.GetSolverParticleRange().Start;
+			// Only supports Gauss Seidel for now because we are using AllElements
+			int32 ElementOffset = AllElements->Num(); //InitializeMuscleActivation has to go before adding tetrahedrons
+			MuscleIndexOffset.Add(Proxy.GetOwner(), MuscleActivationConstraints->NumMuscles()); 
+			MuscleActivationConstraints->AddMuscles(Evolution->Particles(), MuscleActivationFacade, VertexOffset, ElementOffset);
 		}
 	}
 
@@ -2424,6 +2441,81 @@ namespace Chaos::Softs
 					}
 				});
 		}
+
+		// p.Chaos.Deformable.NumLogExtremeParticle 
+		if (GDeformableXPBDCorotatedParams.NumLogExtremeParticle > 0)
+		{
+			auto RestPosition = [this](int32 VertIdx, FVector3f& OutPosition)
+				{
+					if (const UObject* Owner = this->MObjects[VertIdx])
+					{
+						if (const FFleshThreadingProxy* Proxy = Proxies[Owner]->As<FFleshThreadingProxy>())
+						{
+							const FManagedArrayCollection& Rest = Proxy->GetRestCollection();
+							const Chaos::FRange& Range = Proxy->GetSolverParticleRange();
+							if (const TManagedArray<FVector3f>* Vertex = Rest.FindAttributeTyped<FVector3f>("Vertex", FGeometryCollection::VerticesGroup))
+							{
+								OutPosition = (*Vertex)[VertIdx - Range.Start];
+								return true;
+							}
+						}
+					}
+					return false;
+				};
+			TArray<float> DistRatio;
+			DistRatio.Init(0, AllIncidentElements->Num());
+			TArray<int32> DistIndices;
+			DistIndices.SetNum(AllIncidentElements->Num());
+			for (int32 VertIdx = 0; VertIdx < AllIncidentElements->Num(); ++VertIdx)
+			{
+				DistIndices[VertIdx] = VertIdx;
+				FVector3f VertPosition;
+				if (!RestPosition(VertIdx, VertPosition))
+				{
+					continue;
+				}
+
+				TSet<int32> Neighbors;
+				for (int32 IncidentIdx = 0; IncidentIdx < (*AllIncidentElements)[VertIdx].Num(); ++IncidentIdx)
+				{
+					for (int32 LocalTetIdx = 0; LocalTetIdx < 4; ++LocalTetIdx)
+					{
+						int32 Neighbor = (*AllElements)[(*AllIncidentElements)[VertIdx][IncidentIdx]][LocalTetIdx];
+						if (Neighbor != VertIdx)
+						{
+							Neighbors.Add(Neighbor);
+						}
+					}
+				}
+
+				float TotalDist = 0;
+				float TotalRest = 0;
+				for (int32 NeighborIdx : Neighbors)
+				{
+					FVector3f NeighborPosition;
+					if (!RestPosition(NeighborIdx, NeighborPosition))
+					{
+						continue;
+					}
+					float Dist = (Evolution->Particles().GetX(VertIdx) - Evolution->Particles().GetX(NeighborIdx)).Size();
+					TotalDist += Dist;
+					float Rest = (VertPosition - NeighborPosition).Size();
+					TotalRest += Rest;
+				}
+				if (TotalRest > UE_SMALL_NUMBER)
+				{
+					DistRatio[VertIdx] = TotalDist / TotalRest;
+				}
+			}
+			DistIndices.Sort([&DistRatio](int32 A, int32 B)
+				{
+					return DistRatio[A] > DistRatio[B]; // Descending order
+				});
+			for (int32 i = 0; i < FMath::Min(GDeformableXPBDCorotatedParams.NumLogExtremeParticle, DistIndices.Num()); ++i)
+			{
+				UE_LOG(LogChaosDeformableSolver, Warning, TEXT("Particle index %d has average deformation ratio %f"), DistIndices[i], DistRatio[DistIndices[i]]);
+			}
+		}
 #endif
 	}
 
diff --git a/Engine/Source/Runtime/Experimental/Chaos/Private/GeometryCollection/Facades/CollectionMuscleActivationFacade.cpp b/Engine/Source/Runtime/Experimental/Chaos/Private/GeometryCollection/Facades/CollectionMuscleActivationFacade.cpp
--- a/Engine/Source/Runtime/Experimental/Chaos/Private/GeometryCollection/Facades/CollectionMuscleActivationFacade.cpp
+++ b/Engine/Source/Runtime/Experimental/Chaos/Private/GeometryCollection/Facades/CollectionMuscleActivationFacade.cpp
@@ -25,6 +25,7 @@ namespace GeometryCollection::Facades
 	const FName FMuscleActivationFacade::ContractionVolumeScale("ContractionVolumeScale");
 	const FName FMuscleActivationFacade::FiberLengthRatioAtMaxActivation("FiberLengthRatioAtMaxActivation");
 	const FName FMuscleActivationFacade::MuscleLengthRatioThresholdForMaxActivation("MuscleLengthRatioThresholdForMaxActivation");
+	const FName FMuscleActivationFacade::InflationVolumeScale("InflationVolumeScale");
 	const FName FMuscleActivationFacade::FiberStreamline("FiberStreamline");
 	const FName FMuscleActivationFacade::FiberStreamlineRestLength("FiberStreamlineRestLength");
 	const FName FMuscleActivationFacade::MuscleActivationCurveName("MuscleActivationCurveName");
@@ -40,6 +41,7 @@ namespace GeometryCollection::Facades
 		, ContractionVolumeScaleAttribute(InCollection, ContractionVolumeScale, GroupName)
 		, FiberLengthRatioAtMaxActivationAttribute(InCollection, FiberLengthRatioAtMaxActivation, GroupName)
 		, MuscleLengthRatioThresholdForMaxActivationAttribute(InCollection, MuscleLengthRatioThresholdForMaxActivation, GroupName)
+		, InflationVolumeScaleAttribute(InCollection, InflationVolumeScale, GroupName)
 		, FiberStreamlineAttribute(InCollection, FiberStreamline, GroupName)
 		, FiberStreamlineRestLengthAttribute(InCollection, FiberStreamlineRestLength, GroupName)
 		, MuscleActivationCurveNameAttribute(InCollection, MuscleActivationCurveName, GroupName)
@@ -58,6 +60,7 @@ namespace GeometryCollection::Facades
 		, ContractionVolumeScaleAttribute(InCollection, ContractionVolumeScale, GroupName)
 		, FiberLengthRatioAtMaxActivationAttribute(InCollection, FiberLengthRatioAtMaxActivation, GroupName)
 		, MuscleLengthRatioThresholdForMaxActivationAttribute(InCollection, MuscleLengthRatioThresholdForMaxActivation, GroupName)
+		, InflationVolumeScaleAttribute(InCollection, InflationVolumeScale, GroupName)
 		, FiberStreamlineAttribute(InCollection, FiberStreamline, GroupName)
 		, FiberStreamlineRestLengthAttribute(InCollection, FiberStreamlineRestLength, GroupName)
 		, MuscleActivationCurveNameAttribute(InCollection, MuscleActivationCurveName, GroupName)
@@ -69,7 +72,8 @@ namespace GeometryCollection::Facades
 	{
 		return GeometryGroupIndexAttribute.IsValid() && MuscleActivationElementAttribute.IsValid() && OriginInsertionPairAttribute.IsValid() &&
 			OriginInsertionRestLengthAttribute.IsValid() && FiberDirectionMatrixAttribute.IsValid() && ContractionVolumeScaleAttribute.IsValid() &&
-			FiberLengthRatioAtMaxActivationAttribute.IsValid() && MuscleLengthRatioThresholdForMaxActivationAttribute.IsValid();
+			FiberLengthRatioAtMaxActivationAttribute.IsValid() && MuscleLengthRatioThresholdForMaxActivationAttribute.IsValid() &&
+			InflationVolumeScaleAttribute.IsValid();
 	}
 
 	void FMuscleActivationFacade::DefineSchema()
@@ -83,6 +87,7 @@ namespace GeometryCollection::Facades
 		ContractionVolumeScaleAttribute.Add();
 		FiberLengthRatioAtMaxActivationAttribute.Add(); 
 		MuscleLengthRatioThresholdForMaxActivationAttribute.Add();
+		InflationVolumeScaleAttribute.Add();
 		FiberStreamlineAttribute.Add();
 		FiberStreamlineRestLengthAttribute.Add();
 		MuscleActivationCurveNameAttribute.Add();
@@ -113,6 +118,7 @@ namespace GeometryCollection::Facades
 			ContractionVolumeScaleAttribute.Modify()[DataIndex] = InputData.ContractionVolumeScale;
 			FiberLengthRatioAtMaxActivationAttribute.Modify()[DataIndex] = InputData.FiberLengthRatioAtMaxActivation;
 			MuscleLengthRatioThresholdForMaxActivationAttribute.Modify()[DataIndex] = InputData.MuscleLengthRatioThresholdForMaxActivation;
+			InflationVolumeScaleAttribute.Modify()[DataIndex] = InputData.InflationVolumeScale;
 			FiberStreamlineAttribute.Modify()[DataIndex] = InputData.FiberStreamline;
 			FiberStreamlineRestLengthAttribute.Modify()[DataIndex] = InputData.FiberStreamlineRestLength;
 			return true;
@@ -133,6 +139,7 @@ namespace GeometryCollection::Facades
 			ReturnData.ContractionVolumeScale = ContractionVolumeScaleAttribute[DataIndex];
 			ReturnData.FiberLengthRatioAtMaxActivation = FiberLengthRatioAtMaxActivationAttribute[DataIndex];
 			ReturnData.MuscleLengthRatioThresholdForMaxActivation = MuscleLengthRatioThresholdForMaxActivationAttribute[DataIndex];
+			ReturnData.InflationVolumeScale = InflationVolumeScaleAttribute[DataIndex];
 			ReturnData.FiberStreamline = FiberStreamlineAttribute[DataIndex];
 			ReturnData.FiberStreamlineRestLength = FiberStreamlineRestLengthAttribute[DataIndex];
 		}
@@ -223,7 +230,7 @@ namespace GeometryCollection::Facades
 		return InvalidMuscleIndices.Num();
 	}
 
-	bool FMuscleActivationFacade::SetUpMuscleActivation(const TArray<int32>& InOrigin, const TArray<int32>& Insertion)
+	bool FMuscleActivationFacade::SetUpMuscleActivation(const TArray<int32>& InOrigin, const TArray<int32>& Insertion, float InContractionVolumeScale)
 	{
 		// Vertices and fiber field
 		if (InOrigin.Num() == 0)
@@ -384,115 +391,7 @@ namespace GeometryCollection::Facades
 					W = (U ^ V).GetSafeNormal();
 					MuscleActivationData.FiberDirectionMatrix[LocalElemIdx] = Chaos::PMatrix33d(V, W, U);
 				}
-				MuscleActivationData.ContractionVolumeScale.Init(1.f, MuscleActivationElements[MuscleComponentIdx].Num());
-				AddMuscleActivationData(MuscleActivationData);
-			}
-		}
-		return true;
-	}
-
-	bool FMuscleActivationFacade::SetUpMuscleActivation(const TArray<int32>& InOrigin, const TArray<int32>& Insertion,
-		float InContractionVolumeScale)
-	{
-		// Vertices and fiber field
-		if (!(InOrigin.Num() > 0
-			&& Insertion.Num() > 0
-			&& ConstCollection.FindAttribute<FVector3f>("Vertex", FGeometryCollection::VerticesGroup)
-			&& ConstCollection.FindAttribute<FIntVector4>("Tetrahedron", "Tetrahedral")
-			&& ConstCollection.FindAttribute<FVector3f>("FiberDirection", "Tetrahedral")))
-		{
-			return false;
-		}
-		TArray<int32> Origin = InOrigin;
-		const TArray<FVector3f>& Vertices = ConstCollection.FindAttribute<FVector3f>("Vertex", FGeometryCollection::VerticesGroup)->GetConstArray();
-		const TArray<FIntVector4>& Elements = ConstCollection.FindAttribute<FIntVector4>("Tetrahedron", "Tetrahedral")->GetConstArray();
-		const TArray<FVector3f>& FiberDirections = ConstCollection.FindAttribute<FVector3f>("FiberDirection", "Tetrahedral")->GetConstArray();
-		TArray<TArray<int32>> MuscleActivationElements;
-		TArray<TArray<int32>> ComponentOrigins; //One origin node per muscle component
-		TArray<TArray<int32>> ComponentInsertions; //One insertion node per muscle component
-		GeometryCollection::Facades::FCollectionMeshFacade MeshFacade(ConstCollection);
-		TArray<int32> ComponentIndex = MeshFacade.GetGeometryGroupIndexArray(); //Vertex index to geometry index
-		TMap<int32, int32> ComponentToMuscleIndex; //Component index to muscle index
-		TMap<int32, int32> MuscleToComponentIndex; //Muscle index to component index
-		Origin.Sort(); //For some order in muscle groups
-		for (int32 i = 0; i < Origin.Num(); i++)
-		{
-			if (!ComponentToMuscleIndex.Contains(ComponentIndex[Origin[i]]))
-			{
-				ComponentToMuscleIndex.Add(ComponentIndex[Origin[i]], ComponentOrigins.Num());
-				MuscleToComponentIndex.Add(ComponentOrigins.Num(), ComponentIndex[Origin[i]]);
-				ComponentOrigins.SetNum(ComponentOrigins.Num() + 1);
-				ComponentOrigins[ComponentOrigins.Num() - 1].Add(Origin[i]);
-			}
-			else
-			{
-				ComponentOrigins[ComponentToMuscleIndex[ComponentIndex[Origin[i]]]].Add(Origin[i]);
-			}
-		}
-		ComponentInsertions.SetNum(ComponentOrigins.Num());
-		for (int32 i = 0; i < Insertion.Num(); i++)
-		{
-			if (!ComponentToMuscleIndex.Contains(ComponentIndex[Insertion[i]]))
-			{
-				ensureMsgf(false, TEXT("No origin in this component"));
-			}
-			else
-			{
-				ComponentInsertions[ComponentToMuscleIndex[ComponentIndex[Insertion[i]]]].Add(Insertion[i]);
-			}
-		}
-		MuscleActivationElements.SetNum(ComponentOrigins.Num());
-		for (int32 ElemIdx = 0; ElemIdx < Elements.Num(); ElemIdx++)
-		{
-			if (ComponentToMuscleIndex.Contains(ComponentIndex[Elements[ElemIdx][0]]))
-			{
-				MuscleActivationElements[ComponentToMuscleIndex[ComponentIndex[Elements[ElemIdx][0]]]].Add(ElemIdx);
-			}
-		}
-		//Choose one origin-insertion pair per muscle that has largest distance apart within each muscle
-		// TODO: painted attribute directed origin-insertion pair
-		//use origin-insertion line segment length to estimate activation
-		for (int32 MuscleComponentIdx = 0; MuscleComponentIdx < ComponentOrigins.Num(); MuscleComponentIdx++)
-		{
-			if (ensureMsgf(ComponentOrigins.Num() > 0 && ComponentInsertions.Num() > 0, TEXT("Origin or Insertion missing in the muscle %d"), MuscleComponentIdx))
-			{
-				FMuscleActivationData MuscleActivationData;
-				MuscleActivationData.GeometryGroupIndex = MuscleToComponentIndex[MuscleComponentIdx];
-				MuscleActivationData.OriginInsertionRestLength = 0;
-				for (int32 OriginIdx : ComponentOrigins[MuscleComponentIdx])
-				{
-					for (int32 InsertionIdx : ComponentInsertions[MuscleComponentIdx])
-					{
-						float Dist = (Vertices[OriginIdx] - Vertices[InsertionIdx]).Size();
-						if (Dist > MuscleActivationData.OriginInsertionRestLength)
-						{
-							MuscleActivationData.OriginInsertionPair = FIntVector2(OriginIdx, InsertionIdx);
-							MuscleActivationData.OriginInsertionRestLength = Dist;
-						}
-					}
-				}
-				MuscleActivationData.MuscleActivationElement = MuscleActivationElements[MuscleComponentIdx];
-				MuscleActivationData.FiberDirectionMatrix.SetNum(MuscleActivationElements[MuscleComponentIdx].Num());
-				MuscleActivationData.ContractionVolumeScale.SetNum(MuscleActivationElements[MuscleComponentIdx].Num());
-				for (int32 LocalElemIdx = 0; LocalElemIdx < MuscleActivationElements[MuscleComponentIdx].Num(); LocalElemIdx++)
-				{
-					FVector3f V = FiberDirections[MuscleActivationElements[MuscleComponentIdx][LocalElemIdx]];
-					// QR decomposition on vvT for orthogonal directions
-					FVector3f W = V;
-					if (V.X < V.Y)
-					{
-						W.X += 1.f;
-					}
-					else
-					{
-						W.Y += 1.f;
-					}
-					FVector3f U = (V ^ W).GetSafeNormal();
-					W = (U ^ V).GetSafeNormal();
-					MuscleActivationData.FiberDirectionMatrix[LocalElemIdx] = Chaos::PMatrix33d(V, W, U);
-					//Muscle contraction volume scale
-					MuscleActivationData.ContractionVolumeScale[LocalElemIdx] = InContractionVolumeScale;
-				}
+				MuscleActivationData.ContractionVolumeScale.Init(InContractionVolumeScale, MuscleActivationElements[MuscleComponentIdx].Num());
 				AddMuscleActivationData(MuscleActivationData);
 			}
 		}
@@ -500,16 +399,18 @@ namespace GeometryCollection::Facades
 	}
 
 	void FMuscleActivationFacade::UpdateGlobalMuscleActivationParameters(
-		float InContractionVolumeScale,
-		float InFiberLengthRatioAtMaxActivation,
-		float InMuscleLengthRatioThresholdForMaxActivation)
+		float InGlobalContractionVolumeScale,
+		float InGlobalFiberLengthRatioAtMaxActivation,
+		float InGlobalMuscleLengthRatioThresholdForMaxActivation,
+		float InGlobalInflationVolumeScale)
 	{
 		for (int32 MuscleIndex = 0; MuscleIndex < NumMuscles(); ++MuscleIndex)
 		{
 			FMuscleActivationData MuscleActivationData = GetMuscleActivationData(MuscleIndex);
-			MuscleActivationData.ContractionVolumeScale.Init(InContractionVolumeScale, MuscleActivationData.MuscleActivationElement.Num());
-			MuscleActivationData.FiberLengthRatioAtMaxActivation = InFiberLengthRatioAtMaxActivation;
-			MuscleActivationData.MuscleLengthRatioThresholdForMaxActivation = InMuscleLengthRatioThresholdForMaxActivation;
+			MuscleActivationData.ContractionVolumeScale.Init(InGlobalContractionVolumeScale, MuscleActivationData.MuscleActivationElement.Num());
+			MuscleActivationData.FiberLengthRatioAtMaxActivation = InGlobalFiberLengthRatioAtMaxActivation;
+			MuscleActivationData.MuscleLengthRatioThresholdForMaxActivation = InGlobalMuscleLengthRatioThresholdForMaxActivation;
+			MuscleActivationData.InflationVolumeScale = InGlobalInflationVolumeScale;
 			UpdateMuscleActivationData(MuscleIndex, MuscleActivationData);
 		}
 	}
@@ -518,7 +419,8 @@ namespace GeometryCollection::Facades
 		int32 MuscleIndex,
 		float InContractionVolumeScale,
 		float InFiberLengthRatioAtMaxActivation,
-		float InMuscleLengthRatioThresholdForMaxActivation)
+		float InMuscleLengthRatioThresholdForMaxActivation,
+		float InInflationVolumeScale)
 	{
 		if (IsValidMuscleIndex(MuscleIndex))
 		{
@@ -526,6 +428,7 @@ namespace GeometryCollection::Facades
 			MuscleActivationData.ContractionVolumeScale.Init(InContractionVolumeScale, MuscleActivationData.MuscleActivationElement.Num());
 			MuscleActivationData.FiberLengthRatioAtMaxActivation = InFiberLengthRatioAtMaxActivation;
 			MuscleActivationData.MuscleLengthRatioThresholdForMaxActivation = InMuscleLengthRatioThresholdForMaxActivation;
+			MuscleActivationData.InflationVolumeScale = InInflationVolumeScale;
 			UpdateMuscleActivationData(MuscleIndex, MuscleActivationData);
 			return true;
 		}
diff --git a/Engine/Source/Runtime/Experimental/Chaos/Public/Chaos/Deformable/ChaosDeformableSolverTypes.h b/Engine/Source/Runtime/Experimental/Chaos/Public/Chaos/Deformable/ChaosDeformableSolverTypes.h
--- a/Engine/Source/Runtime/Experimental/Chaos/Public/Chaos/Deformable/ChaosDeformableSolverTypes.h
+++ b/Engine/Source/Runtime/Experimental/Chaos/Public/Chaos/Deformable/ChaosDeformableSolverTypes.h
@@ -219,7 +219,7 @@ namespace Chaos::Softs
 	{
 		int32 XPBDCorotatedBatchSize = 5;
 		int32 XPBDCorotatedBatchThreshold = 5;
-
+		int32 NumLogExtremeParticle = 0;
 	};
 
 
diff --git a/Engine/Source/Runtime/Experimental/Chaos/Public/Chaos/Deformable/MuscleActivationConstraints.h b/Engine/Source/Runtime/Experimental/Chaos/Public/Chaos/Deformable/MuscleActivationConstraints.h
--- a/Engine/Source/Runtime/Experimental/Chaos/Public/Chaos/Deformable/MuscleActivationConstraints.h
+++ b/Engine/Source/Runtime/Experimental/Chaos/Public/Chaos/Deformable/MuscleActivationConstraints.h
@@ -39,7 +39,9 @@ namespace Chaos::Softs
 						ensureMsgf(MuscleActivationData.FiberLengthRatioAtMaxActivation > 0 &&
 							MuscleActivationData.FiberLengthRatioAtMaxActivation < 1,
 							TEXT("FiberLengthRatioAtMaxActivation %f of muscle indexed %d is out of range (0,1), please check your setup."), MuscleActivationData.FiberLengthRatioAtMaxActivation, MuscleIdx)
-						)
+						&&
+						ensureMsgf(MuscleActivationData.InflationVolumeScale > 0,
+							TEXT("FiberLengthRatioAtMaxActivation %f of muscle indexed %d is <= 0, please check your setup."), MuscleActivationData.InflationVolumeScale, MuscleIdx))
 					{
 						const int32 OldSize = MuscleActivationElement.AddDefaulted(1);
 						FiberDirectionMatrix.AddDefaulted(1);
@@ -55,6 +57,7 @@ namespace Chaos::Softs
 						}
 						FiberLengthRatioAtMaxActivation.Add(MuscleActivationData.FiberLengthRatioAtMaxActivation);
 						MuscleLengthRatioThresholdForMaxActivation.Add(MuscleActivationData.MuscleLengthRatioThresholdForMaxActivation);
+						InflationVolumeScale.Add(MuscleActivationData.InflationVolumeScale);
 						MuscleActivationData.OriginInsertionPair[0] += VertexOffset;
 						MuscleActivationData.OriginInsertionPair[1] += VertexOffset;
 						OriginInsertionPair.Add(MuscleActivationData.OriginInsertionPair);
@@ -110,6 +113,23 @@ namespace Chaos::Softs
 			}
 		}
 
+		void ApplyInflationVolumeScale(FXPBDCorotatedConstraints<T, ParticleType>& Constraints) const
+		{
+			for (int32 MuscleIdx = 0; MuscleIdx < MuscleActivationElement.Num(); MuscleIdx++)
+			{
+				if (ensureMsgf(MuscleActivationElement[MuscleIdx].Num() == FiberDirectionMatrix[MuscleIdx].Num(),
+					TEXT("MuscleActivationElement[%d].Num() = %d, not equal to FiberDirectionMatrix[%d].Num() = %d"), MuscleIdx, MuscleActivationElement[MuscleIdx].Num(), MuscleIdx, FiberDirectionMatrix[MuscleIdx].Num())
+					&& ensureMsgf(InflationVolumeScale[MuscleIdx] > 0,
+						TEXT("InflationVolumeScale %f of muscle indexed %d is <= 0, please check your setup."), InflationVolumeScale[MuscleIdx], MuscleIdx))
+				{
+					for (int32 ElemIdx = 0; ElemIdx < MuscleActivationElement[MuscleIdx].Num(); ElemIdx++)
+					{
+						Constraints.ModifyDmInverseSaveFromInflationVolumeScale(MuscleActivationElement[MuscleIdx][ElemIdx], InflationVolumeScale[MuscleIdx], FiberDirectionMatrix[MuscleIdx][ElemIdx]);
+					}
+				}
+			}
+		}
+
 		int32 NumMuscles() { return MuscleActivationElement.Num(); };
 		int32 GetMuscleVertexOffset(int32 MuscleIndex) { return MuscleVertexOffset[MuscleIndex]; };
 		int32 GetMuscleVertexCount(int32 MuscleIndex) { return MuscleVertexCount[MuscleIndex]; };
@@ -143,6 +163,7 @@ namespace Chaos::Softs
 		TArray<TArray<float>> ContractionVolumeScale;
 		TArray<float> FiberLengthRatioAtMaxActivation;
 		TArray<float> MuscleLengthRatioThresholdForMaxActivation;
+		TArray<float> InflationVolumeScale;
 		TArray<int32> MuscleVertexOffset;
 		TArray<int32> MuscleVertexCount;
 	};
diff --git a/Engine/Source/Runtime/Experimental/Chaos/Public/Chaos/XPBDCorotatedConstraints.h b/Engine/Source/Runtime/Experimental/Chaos/Public/Chaos/XPBDCorotatedConstraints.h
--- a/Engine/Source/Runtime/Experimental/Chaos/Public/Chaos/XPBDCorotatedConstraints.h
+++ b/Engine/Source/Runtime/Experimental/Chaos/Public/Chaos/XPBDCorotatedConstraints.h
@@ -406,25 +406,44 @@ namespace Chaos::Softs
 			return dC2;
 		}
 
-		void ModifyDmInverseFromMuscleLength(const int32 e, const T FiberLengthRatio, const PMatrix<T, 3, 3>& MFiberDir, const T ContractionVolumeScale) const
+		void ModifyDmInverseFromMuscleLength(const int32 ElemIdx, const T FiberLengthRatio, const PMatrix<T, 3, 3>& FiberDir, const T ContractionVolumeScale) const
 		{
 			TRACE_CPUPROFILER_EVENT_SCOPE(STAT_ChaosXPBDCorotatedModifyDmInverseFromMuscleLength);
-			PMatrix<T, 3, 3> DmInv = ElementDmInvSave(e);
+			PMatrix<T, 3, 3> DmInv = ElementDmInvSave(ElemIdx);
 			if (FiberLengthRatio < 1)
 			{
 				PMatrix<T, 3, 3> D(1 / FiberLengthRatio, FMath::Pow(FiberLengthRatio, ContractionVolumeScale/T(2)), FMath::Pow(FiberLengthRatio, ContractionVolumeScale/T(2)));
-				PMatrix<T, 3, 3> Factor = MFiberDir * D * MFiberDir.GetTransposed();
+				PMatrix<T, 3, 3> Factor = FiberDir * D * FiberDir.GetTransposed();
 				DmInv = Factor * DmInv; //Matrix multiplication convention: AB = B*A
 			}
-			for (int r = 0; r < 3; r++) 
+			for (int Row = 0; Row < 3; Row++)
 			{
-				for (int c = 0; c < 3; c++) 
+				for (int Column = 0; Column < 3; Column++)
 				{
-					DmInverse[(3 * 3) * e + 3 * r + c] = DmInv.GetAt(r, c);
+					DmInverse[(3 * 3) * ElemIdx + 3 * Row + Column] = DmInv.GetAt(Row, Column);
 				}
 			}
 		}
 
+		void ModifyDmInverseSaveFromInflationVolumeScale(const int32 ElemIdx, const T InflationVolumeScale, const PMatrix<T, 3, 3>& FiberDir)
+		{
+			PMatrix<T, 3, 3> DmInv = ElementDmInvSave(ElemIdx);
+			if (FMath::Abs(InflationVolumeScale - 1) > UE_SMALL_NUMBER)
+			{
+				T InvSqrt = FMath::InvSqrt(InflationVolumeScale);
+				PMatrix<T, 3, 3> D(1, InvSqrt, InvSqrt);
+				PMatrix<T, 3, 3> Factor = FiberDir * D * FiberDir.GetTransposed();
+				DmInv = Factor * DmInv; //Matrix multiplication convention: AB = B*A
+			}
+			for (int Row = 0; Row < 3; Row++)
+			{
+				for (int Column = 0; Column < 3; Column++)
+				{
+					DmInverseSave[(3 * 3) * ElemIdx + 3 * Row + Column] = DmInv.GetAt(Row, Column);
+					DmInverse[(3 * 3) * ElemIdx + 3 * Row + Column] = DmInv.GetAt(Row, Column);
+				}
+			}
+		}
 	protected:
 
 		void InitColor(const ParticleType& Particles)
diff --git a/Engine/Source/Runtime/Experimental/Chaos/Public/GeometryCollection/Facades/CollectionMuscleActivationFacade.h b/Engine/Source/Runtime/Experimental/Chaos/Public/GeometryCollection/Facades/CollectionMuscleActivationFacade.h
--- a/Engine/Source/Runtime/Experimental/Chaos/Public/GeometryCollection/Facades/CollectionMuscleActivationFacade.h
+++ b/Engine/Source/Runtime/Experimental/Chaos/Public/GeometryCollection/Facades/CollectionMuscleActivationFacade.h
@@ -19,8 +19,9 @@ namespace GeometryCollection::Facades
 		float OriginInsertionRestLength; //Muscle origin-insertion rest length
 		TArray<Chaos::PMatrix33d> FiberDirectionMatrix; //Per-element fiber direction orthogonal matrix: [v, w1, w2], v is the fiber direction
 		TArray<float> ContractionVolumeScale; // Per-element volume scale for muscle contraction. Muscles gain volume during contraction if > 1. Volume-preserving if 1.
-		float FiberLengthRatioAtMaxActivation; // How much muscle fibers shorten at max activation 1. A smaller value means more contraction in the fiber direction.
-		float MuscleLengthRatioThresholdForMaxActivation; // Muscle length ratio (defined by origin-insertion distance) below this threshold is considered to reach max activation 1.
+		float FiberLengthRatioAtMaxActivation = 0.5f; // How much muscle fibers shorten at max activation 1. A smaller value means more contraction in the fiber direction.
+		float MuscleLengthRatioThresholdForMaxActivation = 0.75f; // Muscle length ratio (defined by origin-insertion distance) below this threshold is considered to reach max activation 1.
+		float InflationVolumeScale = 1.f; // Increases muscle rest volume if > 1 and decreases muscle rest volume if < 1.
 		TArray<TArray<FVector3f>> FiberStreamline; //Fiber streamline(s) for inverse dynamics
 		TArray<float> FiberStreamlineRestLength; //Fiber streamline rest length(s)
 	};
@@ -38,6 +39,7 @@ namespace GeometryCollection::Facades
 		static CHAOS_API const FName ContractionVolumeScale;
 		static CHAOS_API const FName FiberLengthRatioAtMaxActivation;
 		static CHAOS_API const FName MuscleLengthRatioThresholdForMaxActivation;
+		static CHAOS_API const FName InflationVolumeScale;
 		static CHAOS_API const FName FiberStreamline;
 		static CHAOS_API const FName FiberStreamlineRestLength;
 		static CHAOS_API const FName MuscleActivationCurveName;
@@ -70,17 +72,18 @@ namespace GeometryCollection::Facades
 		CHAOS_API int32 FindMuscleIndexByName(const FString MuscleName) const;
 		CHAOS_API int32 FindMuscleGeometryIndex(const int32 MuscleIndex) const;
 		CHAOS_API int32 RemoveInvalidMuscles();
-		CHAOS_API bool SetUpMuscleActivation(const TArray<int32>& Origin, const TArray<int32>& Insertion);
-		CHAOS_API bool SetUpMuscleActivation(const TArray<int32>& Origin, const TArray<int32>& Insertion, float ContractionVolumeScale);
+		CHAOS_API bool SetUpMuscleActivation(const TArray<int32>& Origin, const TArray<int32>& Insertion, float ContractionVolumeScale = 1.f);
 		CHAOS_API void UpdateGlobalMuscleActivationParameters(
-			float InContractionVolumeScale,
+			float InGlobalContractionVolumeScale,
 			float InGlobalFiberLengthRatioAtMaxActivation,
-			float InGlobalMuscleLengthRatioThresholdForMaxActivation);
+			float InGlobalMuscleLengthRatioThresholdForMaxActivation,
+			float InGlobalInflationVolumeScale);
 		CHAOS_API bool UpdateMuscleActivationParameters(
 			int32 MuscleIndex,
 			float InContractionVolumeScale,
-			float InGlobalFiberLengthRatioAtMaxActivation,
-			float InGlobalMuscleLengthRatioThresholdForMaxActivation);
+			float InFiberLengthRatioAtMaxActivation,
+			float InMuscleLengthRatioThresholdForMaxActivation,
+			float InInflationVolumeScale);
 		CHAOS_API TArray<TArray<TArray<FVector3f>>> BuildStreamlines(const TArray<int32>& Origin, const TArray<int32>& Insertion,
 			 int32 NumLinesMultiplier, int32 MaxStreamlineIterations, int32 MaxPointsPerLine);
 		CHAOS_API int32 AssignCurveName(const FString& CurveName, const FString& MuscleName);
@@ -96,6 +99,7 @@ namespace GeometryCollection::Facades
 		TManagedArrayAccessor<TArray<float>> ContractionVolumeScaleAttribute;
 		TManagedArrayAccessor<float> FiberLengthRatioAtMaxActivationAttribute;
 		TManagedArrayAccessor<float> MuscleLengthRatioThresholdForMaxActivationAttribute;
+		TManagedArrayAccessor<float> InflationVolumeScaleAttribute;
 		TManagedArrayAccessor<TArray<TArray<FVector3f>>> FiberStreamlineAttribute;
 		TManagedArrayAccessor<TArray<float>> FiberStreamlineRestLengthAttribute;
 		TManagedArrayAccessor<FString> MuscleActivationCurveNameAttribute;
