[NetPhysics] Fix crashes that comes from doing static_cast of an invalid IPhysicsProxyBase pointer in FPBDRigidsSolver::PushPhysicsState during resimulation.
- PushData (which is cached in rewind history) has a pointer to IPhysicsProxyBase, when a physics proxy is deleted we could still rewind back and use the PushData which does a static_cast with the cached pointer which result in an uninitialized casted type that passes a null check but crashes when used.
-- Solution: When IPhysicsProxyBase destructor is called, clear the rewind history of all references to that physics proxy.
--- Implemented: FDirtySet::RemoveWithoutDirtyIdx to remove data from the dirty set based on the physics proxy instead of a dirty index, which we don't have at this point.
--- Implemented: FDirtyChaosProperties.MoveData to move data from one index to another. Required since when we remove an entry in the FDirtySet we swap the last index to the removed index and the FDirtySet index is linked to the FDirtyChaosProperties index.
--- Implemented: FChaosMarshallingManager.ConsumerData to hold a pointer to the currently used internal PushData. There was previously no reference to this data and we need to also clear that of deleted PhysicsProxy references.

- Also fixed some threading crashes in MoverNetworkPhysicsLiason which accessed game thread properties on the physics thread.

diff --git a/Engine/Plugins/Experimental/Mover/Source/Mover/Private/Backends/MoverNetworkPhysicsLiaisonBase.cpp b/Engine/Plugins/Experimental/Mover/Source/Mover/Private/Backends/MoverNetworkPhysicsLiaisonBase.cpp
--- a/Engine/Plugins/Experimental/Mover/Source/Mover/Private/Backends/MoverNetworkPhysicsLiaisonBase.cpp
+++ b/Engine/Plugins/Experimental/Mover/Source/Mover/Private/Backends/MoverNetworkPhysicsLiaisonBase.cpp
@@ -647,9 +647,9 @@ void UMoverNetworkPhysicsLiaisonComponentBase::ProcessInputs_Internal(int32 Phys
 
 	if (CanProcessInputs_Internal(Input))
 	{
-		if (NetworkPhysicsComponent)
+		if (NetworkPhysicsComponent && NetworkPhysicsComponent->GetNetworkPhysicsComponent_Internal())
 		{
-			const bool bIsInputGenerator = NetworkPhysicsComponent->IsLocallyControlled();
+			const bool bIsInputGenerator = NetworkPhysicsComponent->GetNetworkPhysicsComponent_Internal()->IsLocallyControlled();
 			if (!bIsInputGenerator)
 			{
 				// If this instance is not generating input for the owning actor, the NPC will have written our LatestInputCmd (via FNetworkPhysicsMoverInputs::ApplyData)
@@ -660,7 +660,7 @@ void UMoverNetworkPhysicsLiaisonComponentBase::ProcessInputs_Internal(int32 Phys
 			// Unlike inputs, there's no case where the server isn't the authority on the sync state. By default, the LatestSyncState is
 			// only set by NPC replication when resimulating (via FNetworkPhysicsMoverState::ApplyData), but that setting can be adjusted.
 			// So to account for all possible configs, always establish the initial sync state of the frame based on LatestSyncState on clients (it'd be redundant on servers).
-			if (!NetworkPhysicsComponent->HasServerWorld())
+			if (!NetworkPhysicsComponent->GetNetworkPhysicsComponent_Internal()->IsServer())
 			{
 				Chaos::FPBDRigidsEvolution* RigidsEvolution = GetPhysicsSolver()->GetEvolution();
 				const bool bIsSolverResim = RigidsEvolution->IsResimming();
@@ -871,9 +871,10 @@ void UMoverNetworkPhysicsLiaisonComponentBase::WakeParticleIfSleeping(Chaos::FGe
 
 int32 UMoverNetworkPhysicsLiaisonComponentBase::GetNetworkPhysicsTickOffset_Internal() const
 {
-	if (NetworkPhysicsComponent && !NetworkPhysicsComponent->HasServerWorld())
+	if (NetworkPhysicsComponent)
 	{
-		if (const FAsyncNetworkPhysicsComponent* NPCPhysicsThreadAPI = NetworkPhysicsComponent->GetNetworkPhysicsComponent_Internal())
+		const FAsyncNetworkPhysicsComponent* NPCPhysicsThreadAPI = NetworkPhysicsComponent->GetNetworkPhysicsComponent_Internal();
+		if (NPCPhysicsThreadAPI && !NPCPhysicsThreadAPI->IsServer())
 		{
 			return NPCPhysicsThreadAPI->GetNetworkPhysicsTickOffset();
 		}
diff --git a/Engine/Source/Runtime/Engine/Private/PhysicsEngine/NetworkPhysicsComponent.cpp b/Engine/Source/Runtime/Engine/Private/PhysicsEngine/NetworkPhysicsComponent.cpp
--- a/Engine/Source/Runtime/Engine/Private/PhysicsEngine/NetworkPhysicsComponent.cpp
+++ b/Engine/Source/Runtime/Engine/Private/PhysicsEngine/NetworkPhysicsComponent.cpp
@@ -1266,11 +1266,13 @@ const float UNetworkPhysicsComponent::GetCurrentInputDecay(FNetworkPhysicsData*
 
 bool UNetworkPhysicsComponent::HasServerWorld() const
 {
+	Chaos::EnsureIsInGameThreadContext();
 	return GetWorld()->IsNetMode(NM_DedicatedServer) || GetWorld()->IsNetMode(NM_ListenServer);
 }
 
 bool UNetworkPhysicsComponent::IsLocallyControlled() const
 {
+	Chaos::EnsureIsInGameThreadContext();
 	if (bIsRelayingLocalInputs)
 	{
 		return true;
@@ -1285,6 +1287,7 @@ bool UNetworkPhysicsComponent::IsLocallyControlled() const
 
 bool UNetworkPhysicsComponent::IsNetworkPhysicsTickOffsetAssigned() const
 {
+	Chaos::EnsureIsInGameThreadContext();
 	if (APlayerController* PlayerController = GetPlayerController())
 	{
 		return PlayerController->GetNetworkPhysicsTickOffsetAssigned();
@@ -1320,6 +1323,7 @@ void UNetworkPhysicsComponent::SetCompareInputToTriggerRewind(const bool bInComp
 
 APlayerController* UNetworkPhysicsComponent::GetPlayerController() const
 {
+	Chaos::EnsureIsInGameThreadContext();
 	if (APlayerController* PC = Cast<APlayerController>(GetOwner()))
 	{
 		return PC;
diff --git a/Engine/Source/Runtime/Experimental/Chaos/Private/Chaos/ChaosMarshallingManager.cpp b/Engine/Source/Runtime/Experimental/Chaos/Private/Chaos/ChaosMarshallingManager.cpp
--- a/Engine/Source/Runtime/Experimental/Chaos/Private/Chaos/ChaosMarshallingManager.cpp
+++ b/Engine/Source/Runtime/Experimental/Chaos/Private/Chaos/ChaosMarshallingManager.cpp
@@ -134,6 +134,10 @@ FPushPhysicsData* FChaosMarshallingManager::StepInternalTime_External()
 void FChaosMarshallingManager::FreeData_Internal(FPushPhysicsData* PushData)
 {
 	UE_CHAOS_ASYNC_INITBODY_WRITESCOPELOCK(MarshallingManagerLock);
+	if (PushData == ConsumerData)
+	{
+		ConsumerData = nullptr;
+	}
 	PushData->ResetDirtyProxiesBuffer();
 	PushData->Reset();
 	PushDataPool.Enqueue(PushData);
@@ -220,6 +224,11 @@ void FPushPhysicsData::ClearAsyncInputForCallback(ISimCallbackObject* Callback)
 
 void FChaosMarshallingManager::FreeDataToHistory_Internal(FPushPhysicsData* PushData)
 {
+	if (PushData == ConsumerData)
+	{
+		ConsumerData = nullptr;
+	}
+
 	if (HistoryLength == 0)
 	{
 		FreeData_Internal(PushData);
diff --git a/Engine/Source/Runtime/Experimental/Chaos/Private/Framework/PhysicsProxyBase.cpp b/Engine/Source/Runtime/Experimental/Chaos/Private/Framework/PhysicsProxyBase.cpp
--- a/Engine/Source/Runtime/Experimental/Chaos/Private/Framework/PhysicsProxyBase.cpp
+++ b/Engine/Source/Runtime/Experimental/Chaos/Private/Framework/PhysicsProxyBase.cpp
@@ -32,8 +32,9 @@ namespace RenderInterpolationCVars
 
 IPhysicsProxyBase::~IPhysicsProxyBase()
 {
-	if(GetSolver<Chaos::FPhysicsSolverBase>())
+	if (GetSolver<Chaos::FPhysicsSolverBase>())
 	{
+		GetSolver<Chaos::FPhysicsSolverBase>()->RemoveDirtyProxyFromHistory_Internal(this);
 		GetSolver<Chaos::FPhysicsSolverBase>()->RemoveDirtyProxy(this);
 	}
 }
diff --git a/Engine/Source/Runtime/Experimental/Chaos/Private/PBDRigidsSolver.cpp b/Engine/Source/Runtime/Experimental/Chaos/Private/PBDRigidsSolver.cpp
--- a/Engine/Source/Runtime/Experimental/Chaos/Private/PBDRigidsSolver.cpp
+++ b/Engine/Source/Runtime/Experimental/Chaos/Private/PBDRigidsSolver.cpp
@@ -830,6 +830,9 @@ namespace Chaos
 		return Serializer;
 	}
 
+	bool bResimDelayDestroyObject = false;
+	FAutoConsoleVariableRef CVarResimDelayDestroyObject(TEXT("p.Resim.DelayDestroyObject"), bResimDelayDestroyObject, TEXT("Keep physics objects and physics proxy available until after they go out of rewind history. Disabled by default due to not fully implemented."));
+
 	void FPBDRigidsSolver::UnregisterObject(FSingleParticlePhysicsProxy* Proxy)
 	{
 		UE_CHAOS_ASYNC_INITBODY_WRITESCOPELOCK(GetExternalDataLock_External());
@@ -889,7 +892,7 @@ namespace Chaos
 			}
 			
 			Proxy->SetHandle(nullptr);
-			const int32 OffsetForRewind = MRewindData ? MRewindData->Capacity() : 0;
+			const int32 OffsetForRewind = (bResimDelayDestroyObject && MRewindData) ? MRewindData->Capacity() : 0;
 			PendingDestroyPhysicsProxy.Add(FPendingDestroyInfo{Proxy, GetCurrentFrame() + OffsetForRewind, Handle, UniqueIdx});
 			
 			//If particle was created and destroyed before commands were enqueued just skip. I suspect we can skip entire lambda, but too much code to verify right now
@@ -1792,6 +1795,9 @@ namespace Chaos
 		QUICK_SCOPE_CYCLE_COUNTER(ChaosPushData);
 		ensure(PushData.InternalStep == CurrentFrame);	//push data was generated for this specific frame
 
+		// Set the current PushData being used on the internal thread
+		MarshallingManager.SetConsumerData_Internal(&PushData);
+
 		//update callbacks
 		SimCallbackObjects.Reserve(SimCallbackObjects.Num() + PushData.SimCallbackObjectsToAdd.Num());
 		for(ISimCallbackObject* SimCallbackObject : PushData.SimCallbackObjectsToAdd)
@@ -2019,6 +2025,7 @@ namespace Chaos
 					SetIsResimming(true);
 					CurrentFrame = ResimStep;
 
+					FPushPhysicsData* ConsumerData = MarshallingManager.GetConsumerData_Internal();
 					TArray<FPushPhysicsData*> RecordedPushData = MarshallingManager.StealHistory_Internal(NumResimSteps);
 					bool bFirst = true;
 
@@ -2084,6 +2091,8 @@ namespace Chaos
 							bFirst = false;
 						}
 					}
+
+					MarshallingManager.SetConsumerData_Internal(ConsumerData);
 					GetEvolution()->GetIslandManager().ResetParticleResimFrame();
 
 					SetIsResimming(false);
diff --git a/Engine/Source/Runtime/Experimental/Chaos/Public/Chaos/ChaosMarshallingManager.h b/Engine/Source/Runtime/Experimental/Chaos/Public/Chaos/ChaosMarshallingManager.h
--- a/Engine/Source/Runtime/Experimental/Chaos/Public/Chaos/ChaosMarshallingManager.h
+++ b/Engine/Source/Runtime/Experimental/Chaos/Public/Chaos/ChaosMarshallingManager.h
@@ -46,6 +46,11 @@ struct FDirtyProxy
 			ShapesData[ShapeDataIdx].Clear(Manager,ShapeDataIdx);
 		}
 	}
+
+	void MoveData(FDirtyPropertiesManager& Manager, int32 FromIdx, int32 ToIdx)
+	{
+		PropertyData.MoveData(Manager, FromIdx, ToIdx);
+	}
 };
 
 struct FDirtyProxiesBucket
@@ -120,6 +125,43 @@ public:
 		}
 	}
 
+	/** Forcefully removes the proxy from being dirty without relying on DirtyIdx
+	* NOTE: Does not reset IPhysicsProxyBase::DirtyIdx and also doesn't update the DirtyIdx for a moved entries IPhysicsProxyBase, both of which FDirtySet::Remove does */
+	void RemoveWithoutDirtyIdx(IPhysicsProxyBase* Base, FDirtyPropertiesManager& Manager)
+	{
+		FDirtyProxiesBucket& Bucket = DirtyProxyBuckets[(uint32)Base->GetType()];
+		for (int32 Idx = Bucket.ProxiesData.Num() - 1; Idx >= 0; Idx--)
+		{
+			if (Bucket.ProxiesData[Idx].Proxy == Base)
+			{
+				const int32 LastIdx = Bucket.ProxiesData.Num() - 1;
+
+				// Clear data from manager
+				Bucket.ProxiesData[Idx].Clear(Manager, Idx, GetShapesDirtyData());
+
+				if (Idx == LastIdx)
+				{
+					// Last element, pop to remove it
+					Bucket.ProxiesData.Pop(EAllowShrinking::No);
+				}
+				else
+				{
+					// Remove and swap with last element, update moved proxy's idx
+					Bucket.ProxiesData.RemoveAtSwap(Idx);
+
+					// Move data in manager for the swapped element
+					Bucket.ProxiesData[Idx].MoveData(Manager, /*FromIdx*/LastIdx, /*ToIdx*/Idx);
+				}
+
+				--DirtyProxyBucketInfo.Num[(uint32)Base->GetType()];
+				--DirtyProxyBucketInfo.TotalNum;
+
+				// There is only one entry per proxy
+				break;
+			}
+		}
+	}
+
 	// Only does the removal if no shapes are dirty.
 	void RemoveIfNoShapesAreDirty(IPhysicsProxyBase* Base)
 	{
@@ -279,6 +321,18 @@ public:
 		return ProducerData;
 	}
 
+	/** Grabs the consumer data to read from. Should only be called by internal thread */
+	FPushPhysicsData* GetConsumerData_Internal()
+	{
+		return ConsumerData;
+	}
+
+	/** Sets a new internal consumer data to read from. Should only be called by internal thread */
+	void SetConsumerData_Internal(FPushPhysicsData* InPushData)
+	{
+		ConsumerData = InPushData;
+	}
+
 	void AddDirtyProxy(IPhysicsProxyBase* ProxyBaseIn)
 	{
 		UE_CHAOS_ASYNC_INITBODY_WRITESCOPELOCK(MarshallingManagerLock);
@@ -291,6 +345,25 @@ public:
 		ProducerData->DirtyProxiesDataBuffer.Remove(ProxyBaseIn);
 	}
 
+	void RemoveDirtyProxyFromHistory_Internal(IPhysicsProxyBase* ProxyBaseIn)
+	{
+		// Clear any dirty data from push data in history since we can't recreate a deleted physics proxy or particle during resimulation
+		if (HistoryQueue_Internal.Num() > 0)
+		{
+			if (ConsumerData)
+			{
+				// Remove from current PushData not yet put into history
+				ConsumerData->DirtyProxiesDataBuffer.RemoveWithoutDirtyIdx(ProxyBaseIn, ConsumerData->DirtyPropertiesManager);
+			}
+
+			// Remove from history PusData
+			for (FPushPhysicsData* PushDataHistoryEntry : HistoryQueue_Internal)
+			{
+				PushDataHistoryEntry->DirtyProxiesDataBuffer.RemoveWithoutDirtyIdx(ProxyBaseIn, PushDataHistoryEntry->DirtyPropertiesManager);
+			}
+		}
+	}
+
 	void RemoveDirtyProxyIfNoShapesAreDirty(IPhysicsProxyBase* ProxyBaseIn)
 	{
 		UE_CHAOS_ASYNC_INITBODY_WRITESCOPELOCK(MarshallingManagerLock);
@@ -411,7 +484,8 @@ private:
 	std::atomic<int32> InternalStep_External; //the current internal step we are pushing work into. This is not synced with external timestamp as we may sub-step. This should match the solver step
 	
 	//push
-	FPushPhysicsData* ProducerData;
+	FPushPhysicsData* ConsumerData; // The current internal PushData 
+	FPushPhysicsData* ProducerData; // The external PushData before it's handed over to ExternalQueue
 	TArray<FPushPhysicsData*> ExternalQueue;	//the data pushed from external thread with a time stamp
 	TSpscQueue<FPushPhysicsData*> PushDataPool;	//pool to grab more push data from to avoid expensive reallocs
 	TArray<TUniquePtr<FPushPhysicsData>> BackingBuffer;	//all push data is cleaned up by this
diff --git a/Engine/Source/Runtime/Experimental/Chaos/Public/Chaos/Framework/PhysicsSolverBase.h b/Engine/Source/Runtime/Experimental/Chaos/Public/Chaos/Framework/PhysicsSolverBase.h
--- a/Engine/Source/Runtime/Experimental/Chaos/Public/Chaos/Framework/PhysicsSolverBase.h
+++ b/Engine/Source/Runtime/Experimental/Chaos/Public/Chaos/Framework/PhysicsSolverBase.h
@@ -335,6 +335,11 @@ namespace Chaos
 			MarshallingManager.RemoveDirtyProxy(ProxyBaseIn);
 		}
 
+		void RemoveDirtyProxyFromHistory_Internal(IPhysicsProxyBase* ProxyBaseIn)
+		{
+			MarshallingManager.RemoveDirtyProxyFromHistory_Internal(ProxyBaseIn);
+		}
+
 		void RemoveDirtyProxyIfNoShapesAreDirty(IPhysicsProxyBase* ProxyBaseIn)
 		{
 			MarshallingManager.RemoveDirtyProxyIfNoShapesAreDirty(ProxyBaseIn);
diff --git a/Engine/Source/Runtime/Experimental/Chaos/Public/Chaos/ParticleDirtyFlags.h b/Engine/Source/Runtime/Experimental/Chaos/Public/Chaos/ParticleDirtyFlags.h
--- a/Engine/Source/Runtime/Experimental/Chaos/Public/Chaos/ParticleDirtyFlags.h
+++ b/Engine/Source/Runtime/Experimental/Chaos/Public/Chaos/ParticleDirtyFlags.h
@@ -1144,6 +1144,13 @@ public:
 		return Flags.IsDirty(InBits);
 	}
 
+	void MoveData(FDirtyPropertiesManager& Manager, int32 FromIdx, int32 ToIdx)
+	{
+#define CHAOS_PROPERTY(PropName, Type, ProxyType) MoveDataHelper<Type, EChaosProperty::PropName>(Manager, FromIdx, ToIdx);
+#include "ParticleProperties.inl"
+#undef CHAOS_PROPERTY
+	}
+
 #define CHAOS_PROPERTY(PropName, Type, ProxyType)\
 Type const & Get##PropName(const FDirtyPropertiesManager& Manager, int32 Idx) const { return ReadImp<Type, EChaosProperty::PropName>(Manager, Idx); }\
 bool Has##PropName() const { return Flags.IsDirty(ChaosPropertyToFlag(EChaosProperty::PropName)); }\
@@ -1171,6 +1178,16 @@ private:
 			Manager.GetChaosPropertyPool<T, PropName>().Reset(Idx);
 		}
 	}
+
+	template <typename T, EChaosProperty PropName>
+	void MoveDataHelper(FDirtyPropertiesManager& Manager, int32 FromIdx, int32 ToIdx)
+	{
+		if (Flags.IsDirty(ChaosPropertyToFlag(PropName)))
+		{
+			Manager.GetChaosPropertyPool<T, PropName>().GetElement(ToIdx) = Manager.GetChaosPropertyPool<T, PropName>().GetElement(FromIdx);
+			Manager.GetChaosPropertyPool<T, PropName>().Reset(FromIdx);
+		}
+	}
 };
 
 class FShapeDirtyData
