Serialize on-demand TOC flags to indicate whether to stream or install the on-demand content.

- changed unused ChunkVersion to be used as TOC flags
- updated IasTool Upload and ChunkPlugin command to use two new arguments -StreamOnDemand, -InstallOnDemand that sets the corresponding TOC flags

diff --git a/Engine/Source/Runtime/Experimental/IoStore/OnDemand/Private/OnDemandIoStore.cpp b/Engine/Source/Runtime/Experimental/IoStore/OnDemand/Private/OnDemandIoStore.cpp
--- a/Engine/Source/Runtime/Experimental/IoStore/OnDemand/Private/OnDemandIoStore.cpp
+++ b/Engine/Source/Runtime/Experimental/IoStore/OnDemand/Private/OnDemandIoStore.cpp
@@ -1314,13 +1314,21 @@ FIoStatus FOnDemandIoStore::TickMountRequest(FMountRequest& MountRequest)
 	{
 		if (EnumHasAnyFlags(Args.Options, EOnDemandMountOptions::StreamOnDemand))
 		{
-			EnumAddFlags(Container->Flags, EOnDemandContainerFlags::StreamOnDemand);
+			if (!EnumHasAnyFlags(Container->Flags, EOnDemandContainerFlags::StreamOnDemand))
+			{
+				UE_LOG(LogIoStoreOnDemand, Log, TEXT("Overriding container flag 'StreamOnDemand' from mount option"));
+				EnumAddFlags(Container->Flags, EOnDemandContainerFlags::StreamOnDemand);
+			}
 
 			Container->IasHostGroup = FHostGroupManager::Get().Find(FName("DefaultOnDemand"));
 		}
 		else if (EnumHasAnyFlags(Args.Options, EOnDemandMountOptions::InstallOnDemand))
 		{
-			EnumAddFlags(Container->Flags, EOnDemandContainerFlags::InstallOnDemand);
+			if (!EnumHasAnyFlags(Container->Flags, EOnDemandContainerFlags::InstallOnDemand))
+			{
+				UE_LOG(LogIoStoreOnDemand, Log, TEXT("Overriding container flag 'InstallOnDemand' from mount option"));
+				EnumAddFlags(Container->Flags, EOnDemandContainerFlags::InstallOnDemand);
+			}
 
 			if (EnumHasAnyFlags(Args.Options, EOnDemandMountOptions::WithSoftReferences))
 			{
@@ -1938,6 +1946,7 @@ void FOnDemandIoStore::CreateContainersFromToc(
 		}
 	}
 
+	const EOnDemandTocFlags TocFlags = static_cast<EOnDemandTocFlags>(Toc.Header.Flags);
 	for (FOnDemandTocContainerEntry& ContainerEntry : Toc.Containers)
 	{
 		if (ExistingContainers.Contains(ContainerEntry.ContainerName))
@@ -1964,6 +1973,14 @@ void FOnDemandIoStore::CreateContainersFromToc(
 		{
 			EnumAddFlags(Container->Flags, EOnDemandContainerFlags::Encrypted);
 		}
+		if (EnumHasAnyFlags(TocFlags, EOnDemandTocFlags::InstallOnDemand))
+		{
+			EnumAddFlags(Container->Flags, EOnDemandContainerFlags::InstallOnDemand);
+		}
+		else if (EnumHasAnyFlags(TocFlags, EOnDemandTocFlags::StreamOnDemand))
+		{
+			EnumAddFlags(Container->Flags, EOnDemandContainerFlags::StreamOnDemand);
+		}
 
 		//TODO: Sort before uploading
 		ContainerEntry.Entries.Sort([](const FOnDemandTocEntry& Lhs, const FOnDemandTocEntry& Rhs)
diff --git a/Engine/Source/Runtime/Experimental/IoStore/OnDemand/Private/Tool/ChunkPlugin.cpp b/Engine/Source/Runtime/Experimental/IoStore/OnDemand/Private/Tool/ChunkPlugin.cpp
--- a/Engine/Source/Runtime/Experimental/IoStore/OnDemand/Private/Tool/ChunkPlugin.cpp
+++ b/Engine/Source/Runtime/Experimental/IoStore/OnDemand/Private/Tool/ChunkPlugin.cpp
@@ -247,7 +247,26 @@ static int32 ChunkPluginCommandEntry(const FContext& Context)
 	const FString IntermediateFolder	= FString(Context.Get<FStringView>(TEXT("-IntermediateFolder"), FString()));
 	FString SettingsFile				= FString(Context.Get<FStringView>(TEXT("-SettingsFile"), FString()));
 	FString OutputStatsJson				= FString(Context.Get<FStringView>(TEXT("-OutputStatsJson"), FString()));
-	
+
+	const bool bStreamOnDemand			= Context.Get<bool>(TEXT("-StreamOnDemand"), false);
+	const bool bInstallOnDemand			= Context.Get<bool>(TEXT("-InstallOnDemand"), false);
+
+	EOnDemandTocFlags TocFlags = EOnDemandTocFlags::None;
+	if (bStreamOnDemand)
+	{
+		TocFlags |= EOnDemandTocFlags::StreamOnDemand;
+	}
+	else if (bInstallOnDemand)
+	{
+		TocFlags |= EOnDemandTocFlags::InstallOnDemand;
+	}
+	else
+	{
+		const EOnDemandTocFlags FallbackFlags = EOnDemandTocFlags::InstallOnDemand;
+		UE_LOG(LogIoStore, Warning, TEXT("No TOC flags was specified, falling back to '%s'"), *LexToString(FallbackFlags));
+		TocFlags = FallbackFlags;
+	}
+
 	FS3Params S3Params(Context);
 
 	const bool bIncludeSigPak			= Context.Get<bool>(TEXT("-IncludeSigPak"), false);
@@ -274,6 +293,7 @@ static int32 ChunkPluginCommandEntry(const FContext& Context)
 	UE_LOG(LogIoStore, Display, TEXT("\tOutputStatsJson: %s"), *OutputStatsJson);
 	UE_LOG(LogIoStore, Display, TEXT("\tIncludeSigPak: %s"), bIncludeSigPak ? TEXT("true") : TEXT("false"));
 	UE_LOG(LogIoStore, Display, TEXT("\tDeleteContainerFiles: %s"), bDeleteContainerFiles ? TEXT("true") : TEXT("false"));
+	UE_LOG(LogIoStore, Display, TEXT("\tTOC flags: '%s'"), *LexToString(TocFlags));
 
 	IFileManager& FileMgr = IFileManager::Get();
 
@@ -358,6 +378,7 @@ static int32 ChunkPluginCommandEntry(const FContext& Context)
 	UE_LOG(LogIoStore, Display, TEXT("Found %d container files(s)"), ContainerFilenames.Num());
 
 	FOnDemandToc OnDemandToc;
+	OnDemandToc.Header.Flags = static_cast<uint32>(TocFlags);
 	//OnDemandToc.Header.ChunksDirectory = TODO 
 	OnDemandToc.Containers.Reserve(ContainerFilenames.Num());
 
@@ -737,6 +758,8 @@ static FCommand ChunkPluginCommand(
 		TArgument<FStringView>(TEXT("-OutputStatsJson"),	TEXT("Path to write a json file with statistics.")),
 		TArgument<bool>(TEXT("-IncludeSigPak"),				TEXT("Include .sig and .pak file in the uondemandtoc")),
 		TArgument<bool>(TEXT("-KeepContainerFiles"),		TEXT("Should we keep the container files after processing them.")),
+		TArgument<bool>(TEXT("-StreamOnDemand"),			TEXT("Set the content to be streamed on-demand")),
+		TArgument<bool>(TEXT("-InstallOnDemand"),			TEXT("Set the content to be installed on-demand")),
 		TArgument<FStringView>(TEXT("-BucketPrefix"),		TEXT("Path to prefix to bucket objects")),
 		TArgument<int32>(TEXT("-MaxConcurrentUploads"),		TEXT("Number of simultaneous uploads")),
 		S3Arguments,
diff --git a/Engine/Source/Runtime/Experimental/IoStore/OnDemand/Private/Tool/Upload.cpp b/Engine/Source/Runtime/Experimental/IoStore/OnDemand/Private/Tool/Upload.cpp
--- a/Engine/Source/Runtime/Experimental/IoStore/OnDemand/Private/Tool/Upload.cpp
+++ b/Engine/Source/Runtime/Experimental/IoStore/OnDemand/Private/Tool/Upload.cpp
@@ -53,6 +53,8 @@ struct FUploadParams
 	int32 MaxConcurrentUploads = 16;
 	int32 MaxTocListCount = 10000;
 	int32 MaxTocDownloadCount = 100;
+	EOnDemandTocFlags TocFlags = EOnDemandTocFlags::None;
+
 	bool bDeleteContainerFiles = true;
 	bool bDeletePakFiles = true;
 	bool bPerContainerTocs = false;
@@ -134,6 +136,24 @@ static FUploadParams BuildUploadParams(const FContext& Context)
 	Ret.MaxTocListCount			= Context.Get<int32>(TEXT("-MaxTocListCount"),				Ret.MaxTocListCount);
 	Ret.MaxTocDownloadCount		= Context.Get<int32>(TEXT("-MaxTocDownloadCount"),			Ret.MaxTocDownloadCount);
 
+	const bool bStreamOnDemand	= Context.Get<bool>(TEXT("-StreamOnDemand"),				false);
+	const bool bInstallOnDemand	= Context.Get<bool>(TEXT("-InstallOnDemand"),				false);
+
+	if (bStreamOnDemand)
+	{
+		Ret.TocFlags |= EOnDemandTocFlags::StreamOnDemand;
+	}
+	else if (bInstallOnDemand)
+	{
+		Ret.TocFlags |= EOnDemandTocFlags::InstallOnDemand;
+	}
+	else
+	{
+		const EOnDemandTocFlags FallbackFlags = EOnDemandTocFlags::StreamOnDemand;
+		UE_LOG(LogIas, Warning, TEXT("No TOC flags was specified, falling back to '%s'"), *LexToString(FallbackFlags));
+		Ret.TocFlags = FallbackFlags;
+	}
+
 	Ret.bDeleteContainerFiles	= !Context.Get<bool>(TEXT("-KeepContainerFiles"),			!Ret.bDeleteContainerFiles);
 	Ret.bDeletePakFiles			= !Context.Get<bool>(TEXT("-KeepPakFiles"),					!Ret.bDeletePakFiles);
 	Ret.bIgnoreContainerFlags	= Context.Get<bool>(TEXT("-IgnoreContainerFlags"),			Ret.bIgnoreContainerFlags);
@@ -546,6 +566,7 @@ static TIoStatusOr<FUploadResult> UploadContainerFiles(
 	uint64 TotalUploadedBytes = 0;
 
 	FOnDemandToc OnDemandToc;
+	OnDemandToc.Header.Flags = static_cast<uint32>(UploadParams.TocFlags);
 	OnDemandToc.Containers.Reserve(ContainerFiles.Num());
 
 	// Map of the .utoc paths that we have created ondemand containers for, indexed by their hash so that the paths can
@@ -822,6 +843,7 @@ static TIoStatusOr<FUploadResult> UploadContainerFiles(
 		UE_LOG(LogIas, Display, TEXT("%-15s: %s"), TEXT("TargetPlatform"), *UploadParams.TargetPlatform);
 		UE_LOG(LogIas, Display, TEXT("%-15s: %s"), TEXT("BuildVersion"), *UploadParams.BuildVersion);
 		UE_LOG(LogIas, Display, TEXT("%-15s: %s"), TEXT("TOC path"), *UploadResult.TocPath);
+		UE_LOG(LogIas, Display, TEXT("%-15s: %s"), TEXT("TOC flags"), *LexToString(UploadParams.TocFlags));
 		UE_LOG(LogIas, Display, TEXT("%-15s: %.2lf KiB"), TEXT("TOC size"), double(UploadResult.TocSize) / 1024.0);
 		UE_LOG(LogIas, Display, TEXT("%-15s: %.2lf second(s)"), TEXT("Duration"), Duration);
 		UE_LOG(LogIas, Display, TEXT(""));
@@ -938,6 +960,8 @@ static FCommand UploadCommand(
 		TArgument<bool>(TEXT("-WriteTocToDisk"),		TEXT("Output the TOC to disk as well as uploading")),
 		TArgument<bool>(TEXT("-PerContainerTocs"),		TEXT("Whether to generate TOC's for each container file(s)")),
 		TArgument<bool>(TEXT("-IgnoreContainerFlags"),	TEXT("Whether to ignore the OnDemand container flag")),
+		TArgument<bool>(TEXT("-StreamOnDemand"),		TEXT("Set the content to be streamed on-demand")),
+		TArgument<bool>(TEXT("-InstallOnDemand"),		TEXT("Set the content to be installed on-demand")),
 		TArgument<int32>(TEXT("-MaxConcurrentUploads"),	TEXT("Number of simultaneous uploads")),
 		TArgument<int32>(TEXT("-MaxTocListCount"),		TEXT("Maximum number of TOC file(s) to list from the bucket")),
 		TArgument<int32>(TEXT("-MaxTocDownloadCount"),	TEXT("Maximum number of TOC file(s) to download")),
diff --git a/Engine/Source/Runtime/Experimental/IoStore/OnDemandCore/Private/OnDemandToc.cpp b/Engine/Source/Runtime/Experimental/IoStore/OnDemandCore/Private/OnDemandToc.cpp
--- a/Engine/Source/Runtime/Experimental/IoStore/OnDemandCore/Private/OnDemandToc.cpp
+++ b/Engine/Source/Runtime/Experimental/IoStore/OnDemandCore/Private/OnDemandToc.cpp
@@ -4,13 +4,14 @@
 
 #include "HAL/FileManager.h"
 #include "HAL/FileManagerGeneric.h"
+#include "Misc/EnumClassFlags.h"
 #include "Misc/Paths.h"
 #include "Misc/PathViews.h"
+#include "Misc/StringBuilder.h"
 #include "Serialization/Archive.h"
 #include "Serialization/CompactBinarySerialization.h"
 #include "Serialization/CompactBinaryWriter.h"
 #include "Serialization/CustomVersion.h"
-#include "Serialization/MemoryReader.h"
 
 namespace UE::IoStore
 {
@@ -106,11 +107,16 @@ FArchive& operator<<(FArchive& Ar, FOnDemandTocHeader& Header)
 		return Ar;
 	}
 
-	Ar << Header.ChunkVersion;
+	Ar << Header.Flags;
 	Ar << Header.BlockSize;
 	Ar << Header.CompressionFormat;
 	Ar << Header.ChunksDirectory;
 
+	if (Ar.IsLoading() && (Header.Version < uint32(EOnDemandTocVersion::TocFlags)))
+	{
+		Header.Flags = 0; 
+	}
+
 	return Ar;
 }
 
@@ -119,7 +125,7 @@ FCbWriter& operator<<(FCbWriter& Writer, const FOnDemandTocHeader& Header)
 	Writer.BeginObject();
 	Writer.AddInteger(UTF8TEXTVIEW("Magic"), Header.Magic);
 	Writer.AddInteger(UTF8TEXTVIEW("Version"), Header.Version);
-	Writer.AddInteger(UTF8TEXTVIEW("ChunkVersion"), Header.ChunkVersion);
+	Writer.AddInteger(UTF8TEXTVIEW("Flags"), Header.Flags);
 	Writer.AddInteger(UTF8TEXTVIEW("BlockSize"), Header.BlockSize);
 	Writer.AddString(UTF8TEXTVIEW("CompressionFormat"), Header.CompressionFormat);
 	Writer.AddString(UTF8TEXTVIEW("ChunksDirectory"), Header.ChunksDirectory);
@@ -134,7 +140,7 @@ bool LoadFromCompactBinary(FCbFieldView Field, FOnDemandTocHeader& OutTocHeader)
 	{
 		OutTocHeader.Magic = Obj["Magic"].AsUInt64();
 		OutTocHeader.Version = Obj["Version"].AsUInt32();
-		OutTocHeader.ChunkVersion = Obj["ChunkVersion"].AsUInt32();
+		OutTocHeader.Flags = Obj["Flags"].AsUInt32();
 		OutTocHeader.BlockSize = Obj["BlockSize"].AsUInt32();
 		OutTocHeader.CompressionFormat = FString(Obj["CompressionFormat"].AsString());
 		OutTocHeader.ChunksDirectory = FString(Obj["ChunksDirectory"].AsString());
@@ -656,3 +662,46 @@ TIoStatusOr<FOnDemandToc> FOnDemandToc::LoadFromFile(const FString& FilePath, bo
 }
 
 } // namespace UE::IoStore
+
+////////////////////////////////////////////////////////////////////////////////
+void LexToString(UE::IoStore::EOnDemandTocFlags TocFlags, FStringBuilderBase& Out)
+{
+	using namespace UE::IoStore;
+
+	static const TCHAR* FlagText[]
+	{
+		TEXT("None"),
+		TEXT("InstallOnDemand"),
+		TEXT("StreamOnDemand"),
+	};
+
+	if (TocFlags == EOnDemandTocFlags::None)
+	{
+		Out << FlagText[0];
+		return;
+	}
+
+	constexpr uint32 BitCount = 1 + FMath::CountTrailingZerosConstExpr(
+		static_cast<std::underlying_type_t<EOnDemandTocFlags>>(EOnDemandTocFlags::Last));
+	static_assert(UE_ARRAY_COUNT(FlagText) == BitCount + 1, "Please update flag text list");
+
+	for (int32 Idx = 0; Idx < BitCount; ++Idx)
+	{
+		const EOnDemandTocFlags FlagToTest = static_cast<EOnDemandTocFlags>(1 << Idx);
+		if (EnumHasAnyFlags(TocFlags, FlagToTest))
+		{
+			if (Out.Len())
+			{
+				Out << TEXT("|");
+			}
+			Out << FlagText[Idx + 1];
+		}
+	}
+}
+
+FString LexToString(UE::IoStore::EOnDemandTocFlags TocFlags)
+{
+	TStringBuilder<128> Sb;
+	LexToString(TocFlags, Sb);
+	return FString::ConstructFromPtrSize(Sb.ToString(), Sb.Len());
+}
diff --git a/Engine/Source/Runtime/Experimental/IoStore/OnDemandCore/Public/IO/OnDemandToc.h b/Engine/Source/Runtime/Experimental/IoStore/OnDemandCore/Public/IO/OnDemandToc.h
--- a/Engine/Source/Runtime/Experimental/IoStore/OnDemandCore/Public/IO/OnDemandToc.h
+++ b/Engine/Source/Runtime/Experimental/IoStore/OnDemandCore/Public/IO/OnDemandToc.h
@@ -2,10 +2,12 @@
 
 #pragma once
 
+#include "Containers/StringFwd.h"
 #include "IO/IoChunkId.h"
 #include "IO/IoContainerId.h"
 #include "IO/IoHash.h"
 #include "IO/IoStatus.h"
+#include "Misc/EnumClassFlags.h"
 
 #define UE_API IOSTOREONDEMANDCORE_API
 
@@ -30,20 +32,22 @@ enum class EOnDemandTocVersion : uint32
 	AdditionalFiles	= 7,
 	TagSets			= 8,
 	ContainerFlags	= 9,
+	TocFlags		= 10,
 
 	LatestPlusOne,
 	Latest			= (LatestPlusOne - 1)
 };
 
 ///////////////////////////////////////////////////////////////////////////////
-enum class EOnDemandChunkVersion : uint32
+enum class EOnDemandTocFlags : uint32
 {
-	Invalid			= 0,
-	Initial			= 1,
+	None			= 0,
+	InstallOnDemand	= (1 << 0),
+	StreamOnDemand	= (1 << 1),
 
-	LatestPlusOne,
-	Latest			= (LatestPlusOne - 1)
+	Last			= StreamOnDemand
 };
+ENUM_CLASS_FLAGS(EOnDemandTocFlags);
 
 ///////////////////////////////////////////////////////////////////////////////
 struct FTocMeta
@@ -65,7 +69,7 @@ struct FOnDemandTocHeader
 
 	uint64	Magic = ExpectedMagic;
 	uint32	Version = uint32(EOnDemandTocVersion::Latest);
-	uint32	ChunkVersion = uint32(EOnDemandChunkVersion::Latest);
+	uint32	Flags = uint32(EOnDemandTocFlags::None);
 	uint32	BlockSize = 0;
 	FString CompressionFormat;
 	FString ChunksDirectory;
@@ -193,4 +197,8 @@ UE_API bool LoadFromCompactBinary(FCbFieldView Field, FOnDemandToc& OutToc);
 
 } // namespace UE::IoStore
 
+////////////////////////////////////////////////////////////////////////////////
+UE_API void LexToString(UE::IoStore::EOnDemandTocFlags TocFlags, FStringBuilderBase& Out);
+UE_API FString LexToString(UE::IoStore::EOnDemandTocFlags TocFlags);
+
 #undef UE_API
