Move LexToString to global namespace to fix build errors.

#rnx

diff --git a/Engine/Source/Runtime/Experimental/IoStore/OnDemand/Private/OnDemandIoStore.cpp b/Engine/Source/Runtime/Experimental/IoStore/OnDemand/Private/OnDemandIoStore.cpp
--- a/Engine/Source/Runtime/Experimental/IoStore/OnDemand/Private/OnDemandIoStore.cpp
+++ b/Engine/Source/Runtime/Experimental/IoStore/OnDemand/Private/OnDemandIoStore.cpp
@@ -359,51 +359,6 @@ static FIoStatus BuildInstallData(
 
 } // namespace UE::IoStore::Private
 
-///////////////////////////////////////////////////////////////////////////////
-void LexToString(EOnDemandContainerFlags Flags, FStringBuilderBase& Out)
-{
-	static const TCHAR* Names[]
-	{
-		TEXT("None"),
-		TEXT("PendingEncryptionKey"),
-		TEXT("Mounted"),
-		TEXT("StreamOnDemand"),
-		TEXT("InstallOnDemand"),
-		TEXT("Encrypted"),
-		TEXT("WithSoftReferences")
-	};
-
-	if (Flags == EOnDemandContainerFlags::None)
-	{
-		Out << TEXT("None");
-		return;
-	}
-
-	constexpr uint32 BitCount = 1 + FMath::CountTrailingZerosConstExpr(
-		static_cast<std::underlying_type_t<EOnDemandContainerFlags>>(EOnDemandContainerFlags::Last));
-	static_assert(UE_ARRAY_COUNT(Names) == BitCount + 1, "Please update names list");
-
-	for (int32 Idx = 0; Idx < BitCount; ++Idx)
-	{
-		const EOnDemandContainerFlags FlagToTest = static_cast<EOnDemandContainerFlags>(1 << Idx);
-		if (EnumHasAnyFlags(Flags, FlagToTest))
-		{
-			if (Out.Len())
-			{
-				Out << TEXT("|");
-			}
-			Out << Names[Idx + 1];
-		}
-	}
-}
-
-FString LexToString(EOnDemandContainerFlags Flags)
-{
-	TStringBuilder<128> Sb;
-	LexToString(Flags, Sb);
-	return FString::ConstructFromPtrSize(Sb.ToString(), Sb.Len());
-}
-
 ///////////////////////////////////////////////////////////////////////////////
 const FOnDemandChunkEntry FOnDemandChunkEntry::Null = {};
 
@@ -1401,7 +1356,7 @@ FIoStatus FOnDemandIoStore::TickMountRequest(FMountRequest& MountRequest)
 		}
 
 		TStringBuilder<128> Sb;
-		LexToString(Container->Flags, Sb);
+		Sb << Container->Flags;
 		UE_LOG(LogIoStoreOnDemand, Log, TEXT("Mounting container '%s', Entries=%d, Flags='%s'"),
 			*Container->Name, Container->ChunkEntries.Num(), Sb.ToString());
 		EnumAddFlags(Container->Flags, EOnDemandContainerFlags::Mounted);
@@ -2254,3 +2209,53 @@ TUniquePtr<IAnalyticsRecording> FOnDemandIoStore::StartAnalyticsRecording() cons
 }
 
 } // namespace UE::IoStore
+
+////////////////////////////////////////////////////////////////////////////////
+FStringBuilderBase& operator<<(FStringBuilderBase& Sb, UE::IoStore::EOnDemandContainerFlags Flags)
+{
+	using namespace UE::IoStore;
+
+	static const TCHAR* Names[]
+	{
+		TEXT("None"),
+		TEXT("PendingEncryptionKey"),
+		TEXT("Mounted"),
+		TEXT("StreamOnDemand"),
+		TEXT("InstallOnDemand"),
+		TEXT("Encrypted"),
+		TEXT("WithSoftReferences")
+	};
+
+	if (Flags == EOnDemandContainerFlags::None)
+	{
+		Sb << TEXT("None");
+		return Sb;
+	}
+
+	constexpr uint32 BitCount = 1 + FMath::CountTrailingZerosConstExpr(
+		static_cast<std::underlying_type_t<EOnDemandContainerFlags>>(EOnDemandContainerFlags::Last));
+	static_assert(UE_ARRAY_COUNT(Names) == BitCount + 1, "Please update names list");
+
+	for (int32 Idx = 0; Idx < BitCount; ++Idx)
+	{
+		const EOnDemandContainerFlags FlagToTest = static_cast<EOnDemandContainerFlags>(1 << Idx);
+		if (EnumHasAnyFlags(Flags, FlagToTest))
+		{
+			if (Sb.Len())
+			{
+				Sb << TEXT("|");
+			}
+			Sb << Names[Idx + 1];
+		}
+	}
+
+	return Sb;
+}
+
+FString LexToString(UE::IoStore::EOnDemandContainerFlags Flags)
+{
+	TStringBuilder<128> Sb;
+	Sb << Flags;
+	return FString::ConstructFromPtrSize(Sb.ToString(), Sb.Len());
+}
+
diff --git a/Engine/Source/Runtime/Experimental/IoStore/OnDemand/Private/OnDemandIoStore.h b/Engine/Source/Runtime/Experimental/IoStore/OnDemand/Private/OnDemandIoStore.h
--- a/Engine/Source/Runtime/Experimental/IoStore/OnDemand/Private/OnDemandIoStore.h
+++ b/Engine/Source/Runtime/Experimental/IoStore/OnDemand/Private/OnDemandIoStore.h
@@ -48,9 +48,6 @@ enum class EOnDemandContainerFlags : uint8
 };
 ENUM_CLASS_FLAGS(EOnDemandContainerFlags);
 
-void LexToString(EOnDemandContainerFlags Flags, FStringBuilderBase& Out);
-FString LexToString(EOnDemandContainerFlags Flags);
-
 ///////////////////////////////////////////////////////////////////////////////
 struct FOnDemandChunkEntry
 {
@@ -351,3 +348,8 @@ private:
 };
 
 } // namespace UE::IoStore
+
+////////////////////////////////////////////////////////////////////////////////
+FStringBuilderBase& operator<<(FStringBuilderBase& Sb, UE::IoStore::EOnDemandContainerFlags Flags);
+FString LexToString(UE::IoStore::EOnDemandContainerFlags Flags);
+
diff --git a/Engine/Source/Runtime/Experimental/IoStore/OnDemandCore/Internal/IO/IoStoreOnDemandInternals.h b/Engine/Source/Runtime/Experimental/IoStore/OnDemandCore/Internal/IO/IoStoreOnDemandInternals.h
--- a/Engine/Source/Runtime/Experimental/IoStore/OnDemandCore/Internal/IO/IoStoreOnDemandInternals.h
+++ b/Engine/Source/Runtime/Experimental/IoStore/OnDemandCore/Internal/IO/IoStoreOnDemandInternals.h
@@ -30,8 +30,6 @@ public:
 	FWeakOnDemandIoStore	IoStore;
 };
 
-UE_API FString LexToString(const FOnDemandInternalContentHandle& Handle);
-
 ////////////////////////////////////////////////////////////////////////////////
 class FOnDemandInternalInstallRequest
 {
@@ -69,8 +67,11 @@ public:
 	std::atomic<EIoErrorCode>	ErrorCode{EIoErrorCode::Unknown};
 };
 
-UE_API FString LexToString(const FOnDemandInternalInstallRequest& Request);
-
 } // namespace UE::IoStore
 
+////////////////////////////////////////////////////////////////////////////////////////////////////////////////////////////////////////////////////////////////
+UE_API FString LexToString(const UE::IoStore::FOnDemandInternalContentHandle& Handle);
+UE_API FString LexToString(const UE::IoStore::FOnDemandInternalInstallRequest& Request);
+
+
 #undef UE_API
diff --git a/Engine/Source/Runtime/Experimental/IoStore/OnDemandCore/Private/IoStoreOnDemandInternals.cpp b/Engine/Source/Runtime/Experimental/IoStore/OnDemandCore/Private/IoStoreOnDemandInternals.cpp
--- a/Engine/Source/Runtime/Experimental/IoStore/OnDemandCore/Private/IoStoreOnDemandInternals.cpp
+++ b/Engine/Source/Runtime/Experimental/IoStore/OnDemandCore/Private/IoStoreOnDemandInternals.cpp
@@ -14,11 +14,6 @@ FOnDemandInternalContentHandle::~FOnDemandInternalContentHandle()
 	}
 }
 
-FString LexToString(const FOnDemandInternalContentHandle& Handle)
-{
-	return FString::Printf(TEXT("Id=0x%llX DebugName='%s'"), Handle.HandleId(), *Handle.DebugName);
-}
-
 ///////////////////////////////////////////////////////////////////////////////
 FOnDemandContentHandle::FOnDemandContentHandle()
 {
@@ -49,18 +44,7 @@ FOnDemandContentHandle FOnDemandContentHandle::Create(FStringView DebugName)
 
 FString LexToString(const FOnDemandContentHandle& Handle)
 {
-	return Handle.IsValid() ? LexToString(*Handle.Handle) : TEXT("Invalid");
-}
-
-////////////////////////////////////////////////////////////////////////////////
-FString LexToString(const FOnDemandInternalInstallRequest& Request)
-{
-	return FString::Printf(TEXT("Id=0x%llX, DebugName='%s', ContentHandle='%s'"),
-		UPTRINT(&Request),
-		Request.Args.DebugName.IsEmpty()
-			? TEXT("NoName")
-			: *Request.Args.DebugName,
-		*LexToString(Request.Args.ContentHandle));
+	return Handle.IsValid() ? ::LexToString(*Handle.Handle) : TEXT("Invalid");
 }
 
 uint32 FOnDemandInternalInstallRequest::NextSeqNo = 0;
@@ -121,3 +105,20 @@ void FOnDemandInstallRequest::UpdatePriority(int32 NewPriority)
 }
 
 } // namespace UE::IoStore
+
+////////////////////////////////////////////////////////////////////////////////
+FString LexToString(const UE::IoStore::FOnDemandInternalContentHandle& Handle)
+{
+	return FString::Printf(TEXT("Id=0x%llX DebugName='%s'"), Handle.HandleId(), *Handle.DebugName);
+}
+
+FString LexToString(const UE::IoStore::FOnDemandInternalInstallRequest& Request)
+{
+	return FString::Printf(TEXT("Id=0x%llX, DebugName='%s', ContentHandle='%s'"),
+		UPTRINT(&Request),
+		Request.Args.DebugName.IsEmpty()
+			? TEXT("NoName")
+			: *Request.Args.DebugName,
+		*LexToString(Request.Args.ContentHandle));
+}
+
diff --git a/Engine/Source/Runtime/Experimental/IoStore/OnDemandCore/Private/OnDemandToc.cpp b/Engine/Source/Runtime/Experimental/IoStore/OnDemandCore/Private/OnDemandToc.cpp
--- a/Engine/Source/Runtime/Experimental/IoStore/OnDemandCore/Private/OnDemandToc.cpp
+++ b/Engine/Source/Runtime/Experimental/IoStore/OnDemandCore/Private/OnDemandToc.cpp
@@ -664,7 +664,7 @@ TIoStatusOr<FOnDemandToc> FOnDemandToc::LoadFromFile(const FString& FilePath, bo
 } // namespace UE::IoStore
 
 ////////////////////////////////////////////////////////////////////////////////
-void LexToString(UE::IoStore::EOnDemandTocFlags TocFlags, FStringBuilderBase& Out)
+FStringBuilderBase& operator<<(FStringBuilderBase& Sb, UE::IoStore::EOnDemandTocFlags TocFlags)
 {
 	using namespace UE::IoStore;
 
@@ -677,8 +677,8 @@ void LexToString(UE::IoStore::EOnDemandTocFlags TocFlags, FStringBuilderBase& Ou
 
 	if (TocFlags == EOnDemandTocFlags::None)
 	{
-		Out << FlagText[0];
-		return;
+		Sb << FlagText[0];
+		return Sb;
 	}
 
 	constexpr uint32 BitCount = 1 + FMath::CountTrailingZerosConstExpr(
@@ -690,18 +690,20 @@ void LexToString(UE::IoStore::EOnDemandTocFlags TocFlags, FStringBuilderBase& Ou
 		const EOnDemandTocFlags FlagToTest = static_cast<EOnDemandTocFlags>(1 << Idx);
 		if (EnumHasAnyFlags(TocFlags, FlagToTest))
 		{
-			if (Out.Len())
+			if (Sb.Len())
 			{
-				Out << TEXT("|");
+				Sb << TEXT("|");
 			}
-			Out << FlagText[Idx + 1];
+			Sb << FlagText[Idx + 1];
 		}
 	}
+
+	return Sb;
 }
 
 FString LexToString(UE::IoStore::EOnDemandTocFlags TocFlags)
 {
 	TStringBuilder<128> Sb;
-	LexToString(TocFlags, Sb);
+	Sb << TocFlags;
 	return FString::ConstructFromPtrSize(Sb.ToString(), Sb.Len());
 }
diff --git a/Engine/Source/Runtime/Experimental/IoStore/OnDemandCore/Public/IO/OnDemandToc.h b/Engine/Source/Runtime/Experimental/IoStore/OnDemandCore/Public/IO/OnDemandToc.h
--- a/Engine/Source/Runtime/Experimental/IoStore/OnDemandCore/Public/IO/OnDemandToc.h
+++ b/Engine/Source/Runtime/Experimental/IoStore/OnDemandCore/Public/IO/OnDemandToc.h
@@ -198,7 +198,7 @@ UE_API bool LoadFromCompactBinary(FCbFieldView Field, FOnDemandToc& OutToc);
 } // namespace UE::IoStore
 
 ////////////////////////////////////////////////////////////////////////////////
-UE_API void LexToString(UE::IoStore::EOnDemandTocFlags TocFlags, FStringBuilderBase& Out);
+UE_API FStringBuilderBase& operator<<(FStringBuilderBase& Sb, UE::IoStore::EOnDemandTocFlags TocFlags);
 UE_API FString LexToString(UE::IoStore::EOnDemandTocFlags TocFlags);
 
 #undef UE_API
