Iris
* Changed FObjectReferenceCache stringify method so it's clearer which object name is tied to which ref handle ex: [RefHandle1](ObjectName1).[RefHandle2](ObjectName2)

#rnx

diff --git a/Engine/Source/Runtime/Experimental/Iris/Core/Private/Iris/ReplicationSystem/ObjectReferenceCache.cpp b/Engine/Source/Runtime/Experimental/Iris/Core/Private/Iris/ReplicationSystem/ObjectReferenceCache.cpp
--- a/Engine/Source/Runtime/Experimental/Iris/Core/Private/Iris/ReplicationSystem/ObjectReferenceCache.cpp
+++ b/Engine/Source/Runtime/Experimental/Iris/Core/Private/Iris/ReplicationSystem/ObjectReferenceCache.cpp
@@ -1871,7 +1871,7 @@ void FObjectReferenceCache::GenerateFullPath_r(FNetRefHandle RefHandle, const FN
 	{
 		// Doh, this shouldn't be possible, but if this happens, we can't continue
 		// So warn, and return
-		FullPath += FString::Printf(TEXT("[%s]NOT_IN_CACHE" ), *RefHandle.ToString());
+		FullPath += FString::Printf(TEXT("[%s](NOT_IN_CACHE)" ), *RefHandle.ToString());
 		return;
 	}
 
@@ -1893,19 +1893,19 @@ void FObjectReferenceCache::GenerateFullPath_r(FNetRefHandle RefHandle, const FN
 		//	UE_LOG_REFERENCECACHE_WARNING(TEXT( "GenerateFullPath_r: Name mismatch! %s != %s" ), *CacheObject->PathName.ToString(), *CacheObject->Object->GetName() );
 		//}
 
-		FullPath += FString::Printf(TEXT("[%s]%s"), *RefHandle.ToString(), *CacheObject->Object->GetName());
+		FullPath += FString::Printf(TEXT("[%s](%s)"), *RefHandle.ToString(), *CacheObject->Object->GetName());
 	}
 	else
 	{
 		if (!CacheObject->RelativePath.IsValid())
 		{
 			// This can happen when a non stably named object is nullptr
-			FullPath += FString::Printf(TEXT("[%s]EMPTY" ), *RefHandle.ToString());
+			FullPath += FString::Printf(TEXT("[%s](EMPTY)" ), *RefHandle.ToString());
 		}
 		else
 		{
 			const TCHAR* ResolvedPath = StringTokenStore->ResolveRemoteToken(CacheObject->RelativePath, *ResolveContext.RemoteNetTokenStoreState);
-			FullPath += FString::Printf(TEXT("[%s]%s"), *RefHandle.ToString(), ResolvedPath);
+			FullPath += FString::Printf(TEXT("[%s](%s)"), *RefHandle.ToString(), ResolvedPath);
 		}
 	}
 }
