Dataflow - Fixed a minor typo in a DataflowInstance doc comment.

#rnx

diff --git a/Engine/Source/Runtime/Experimental/Dataflow/Engine/Public/Dataflow/DataflowInstance.h b/Engine/Source/Runtime/Experimental/Dataflow/Engine/Public/Dataflow/DataflowInstance.h
--- a/Engine/Source/Runtime/Experimental/Dataflow/Engine/Public/Dataflow/DataflowInstance.h
+++ b/Engine/Source/Runtime/Experimental/Dataflow/Engine/Public/Dataflow/DataflowInstance.h
@@ -116,7 +116,7 @@ public:
 
 	/** 
 	* Update the asset by re-evaluating the dataflow
-	* return true if the dataflow has been proeprly evaluated and the asset updated
+	* return true if the dataflow has been properly evaluated and the asset updated
 	*/
 	DATAFLOWENGINE_API bool UpdateOwnerAsset() const;
 
