Added CVar iostore.MinInstallDurationSeconds to simluate long install times.

#rnx

diff --git a/Engine/Source/Runtime/Experimental/IoStore/OnDemand/Private/OnDemandIoStore.cpp b/Engine/Source/Runtime/Experimental/IoStore/OnDemand/Private/OnDemandIoStore.cpp
--- a/Engine/Source/Runtime/Experimental/IoStore/OnDemand/Private/OnDemandIoStore.cpp
+++ b/Engine/Source/Runtime/Experimental/IoStore/OnDemand/Private/OnDemandIoStore.cpp
@@ -53,6 +53,15 @@ static FAutoConsoleVariableRef CVar_IoStoreOnDemandTreatMissingChunksAsError(
 	ECVF_ReadOnly
 );
 
+#if !UE_BUILD_SHIPPING
+float GIoStoreOnDemandMinInstallDurationSeconds = 0.0f;
+static FAutoConsoleVariableRef CVar_IoStoreOnDemandMinInstallDurationSeconds(
+	TEXT("iostore.MinInstallDurationSeconds"),
+	GIoStoreOnDemandMinInstallDurationSeconds,
+	TEXT("Sets the minimum install duration in seconds.")
+);
+#endif
+
 ///////////////////////////////////////////////////////////////////////////////
 namespace UE::IoStore
 {
@@ -1701,8 +1710,19 @@ FOnDemandInstallResult FOnDemandIoStore::TickInstallRequest(FOnDemandInternalIns
 		}
 	}
 
-	UE_LOG(LogIoStoreOnDemand, Log, TEXT("Downloaded %llu (%.2lf MiB) of total %llu (%.2lf MiB) requested chunk(s), ContentHandle='%s'"),
-		ToDownloadChunkCount, double(ToDownloadBytes) / 1024.0 / 1024.0, TotalChunkCount, double(TotalBytes) / 1024.0 / 1024.0, *LexToString(*ContentHandle));
+	double DurationInSeconds = FPlatformTime::Seconds() - StartTime;
+	#if !UE_BUILD_SHIPPING
+	if (GIoStoreOnDemandMinInstallDurationSeconds > 0.0f)
+	{
+		if (float SleepTime = GIoStoreOnDemandMinInstallDurationSeconds - float(DurationInSeconds); SleepTime > 0.0f)
+		{
+			FPlatformProcess::SleepNoStats(SleepTime);
+			DurationInSeconds = FPlatformTime::Seconds() - StartTime;
+		}
+	}
+	#endif
+	UE_LOG(LogIoStoreOnDemand, Log, TEXT("Downloaded %llu (%.2lf MiB) of total %llu (%.2lf MiB) requested chunk(s), Duration=%.2lfs, ContentHandle='%s'"),
+		ToDownloadChunkCount, double(ToDownloadBytes) / 1024.0 / 1024.0, TotalChunkCount, double(TotalBytes) / 1024.0 / 1024.0, DurationInSeconds, *LexToString(*ContentHandle));
 
 	return OutResult;
 }
