Dataflow : add method to clear the cache on the context

diff --git a/Engine/Source/Runtime/Experimental/Dataflow/Core/Private/Dataflow/DataflowNodeParameters.cpp b/Engine/Source/Runtime/Experimental/Dataflow/Core/Private/Dataflow/DataflowNodeParameters.cpp
--- a/Engine/Source/Runtime/Experimental/Dataflow/Core/Private/Dataflow/DataflowNodeParameters.cpp
+++ b/Engine/Source/Runtime/Experimental/Dataflow/Core/Private/Dataflow/DataflowNodeParameters.cpp
@@ -149,6 +149,12 @@ namespace UE::Dataflow
 		return DataStore.IsEmpty();
 	}
 
+	void FContext::ClearAllData()
+	{
+		FScopedOptionalLock Lock(DataLock.Get());
+		DataStore.Reset();
+	}
+
 	FTimestamp FContext::GetTimestamp(FContextCacheKey Key) const
 	{
 		if (const TUniquePtr<FContextCacheElementBase>* const Cache = const_cast<FContext*>(this)->GetDataImpl(Key))
diff --git a/Engine/Source/Runtime/Experimental/Dataflow/Core/Public/Dataflow/DataflowNodeParameters.h b/Engine/Source/Runtime/Experimental/Dataflow/Core/Public/Dataflow/DataflowNodeParameters.h
--- a/Engine/Source/Runtime/Experimental/Dataflow/Core/Public/Dataflow/DataflowNodeParameters.h
+++ b/Engine/Source/Runtime/Experimental/Dataflow/Core/Public/Dataflow/DataflowNodeParameters.h
@@ -178,6 +178,8 @@ namespace UE::Dataflow
 		DATAFLOWCORE_API virtual void Evaluate(const FDataflowNode* Node, const FDataflowOutput* Output);
 		DATAFLOWCORE_API virtual bool Evaluate(const FDataflowOutput& Connection);
 
+		DATAFLOWCORE_API void ClearAllData();
+
 		template<typename T>
 		void SetData(FContextCacheKey InKey, const FProperty* InProperty, T&& InValue, const FGuid& InNodeGuid, uint32 InNodeHash, const FTimestamp& InTimestamp)
 		{
