Iris
- Fixed crash in NetObjectGridFilter if culldistance is too large.
- Added warnings if culldistance/culldistance overrride is set outside of configurable max to detect usage that can lead to bad performance

diff --git a/Engine/Source/Runtime/Experimental/Iris/Core/Private/Iris/ReplicationSystem/Filtering/NetObjectGridFilter.cpp b/Engine/Source/Runtime/Experimental/Iris/Core/Private/Iris/ReplicationSystem/Filtering/NetObjectGridFilter.cpp
--- a/Engine/Source/Runtime/Experimental/Iris/Core/Private/Iris/ReplicationSystem/Filtering/NetObjectGridFilter.cpp
+++ b/Engine/Source/Runtime/Experimental/Iris/Core/Private/Iris/ReplicationSystem/Filtering/NetObjectGridFilter.cpp
@@ -20,6 +20,11 @@ void UNetObjectGridFilter::OnInit(const FNetObjectFilterInitParams& Params)
 	AssignedObjectInfoIndices.Init(Params.CurrentMaxInternalIndex);
 
 	PerConnectionInfos.SetNum(Params.MaxConnectionCount + 1);
+
+	MinWorldPos = GetDefault<UWorldLocationsConfig>()->MinPos;
+	MaxWorldPos = GetDefault<UWorldLocationsConfig>()->MaxPos;
+
+	MaxNetCullDistance = GetDefault<UWorldLocationsConfig>()->MaxNetCullDistance;
 }
 
 void UNetObjectGridFilter::OnDeinit()
@@ -88,12 +93,10 @@ bool UNetObjectGridFilter::AddObject(uint32 ObjectIndex, FNetObjectFilterAddObje
 	PerObjectInfo.ObjectIndex = ObjectIndex;
 	PerObjectInfo.FrameCountBeforeCulling = GetFrameCountBeforeCulling(Params.ProfileName);
 
-	AddCellInfoForObject(ObjectLocationInfo, Params.InstanceProtocol);
-	
-	if (Config->MaxCullDistance > 0.0f && PerObjectInfo.GetCullDistance() > Config->MaxCullDistance)
+	if (!AddCellInfoForObject(ObjectLocationInfo, Params.InstanceProtocol))
 	{
 		// Too big an object. We expect it to be costly to move it across cells.
-		UE_LOG(LogIris, Warning, TEXT("ReplicatedObject %s cull distance %f is above the max %f. Object will become always relevant instead"), *NetRefHandleManager->PrintObjectFromIndex(ObjectIndex), PerObjectInfo.GetCullDistance(), Config->MaxCullDistance);
+		UE_LOG(LogIris, Warning, TEXT("ReplicatedObject %s cull distance %f is above the max %f. Object will become always relevant instead"), *NetRefHandleManager->PrintObjectFromIndex(ObjectIndex), PerObjectInfo.GetCullDistance(), MaxNetCullDistance);
 		RemoveObject(ObjectIndex, ObjectLocationInfo);
 		return false;
 	}
@@ -279,7 +282,7 @@ void UNetObjectGridFilter::FreeObjectInfo(uint32 Index)
 	ObjectInfos[Index] = FPerObjectInfo();
 }
 
-void UNetObjectGridFilter::AddCellInfoForObject(const FObjectLocationInfo& ObjectLocationInfo, const UE::Net::FReplicationInstanceProtocol* InstanceProtocol)
+bool UNetObjectGridFilter::AddCellInfoForObject(const FObjectLocationInfo& ObjectLocationInfo, const UE::Net::FReplicationInstanceProtocol* InstanceProtocol)
 {
 	// Called for completely new objects
 	FPerObjectInfo& PerObjectInfo = ObjectInfos[ObjectLocationInfo.GetInfoIndex()];
@@ -290,6 +293,11 @@ void UNetObjectGridFilter::AddCellInfoForObject(const FObjectLocationInfo& Objec
 		UpdatePositionAndCullDistance(ObjectLocationInfo, PerObjectInfo, InstanceProtocol);
 	}
 
+	if (MaxNetCullDistance > 0.0f && PerObjectInfo.GetCullDistance() > MaxNetCullDistance)
+	{
+		return false;
+	}
+
 	FCellBox NewCellBox;
 
 	// Set cellbox info
@@ -314,6 +322,8 @@ void UNetObjectGridFilter::AddCellInfoForObject(const FObjectLocationInfo& Objec
 			}
 		}
 	}
+
+	return true;
 }
 
 void UNetObjectGridFilter::RemoveCellInfoForObject(const FObjectLocationInfo& ObjectLocationInfo)
@@ -475,7 +485,8 @@ void UNetObjectGridFilter::UpdatePositionAndCullDistance(const UNetObjectGridFil
 
 void UNetObjectGridFilter::CalculateCellBox(const UNetObjectGridFilter::FPerObjectInfo& PerObjectInfo, UNetObjectGridFilter::FCellBox& OutCellBox)
 {
-	const double CullDistance = PerObjectInfo.GetCullDistance();
+	// Clamp culldistance to sane value if MaxNetCullDistance is set.
+	const double CullDistance = FMath::Min(FMath::Max(PerObjectInfo.GetCullDistance(), MaxNetCullDistance), PerObjectInfo.GetCullDistance());
 	const FVector Position = PerObjectInfo.Position;
 
 	ensureMsgf(!Position.ContainsNaN(), TEXT("Calculating grid cell box based off a position with a NaN for object: %s"), ToCStr(NetRefHandleManager->PrintObjectFromIndex(PerObjectInfo.ObjectIndex)));
@@ -483,6 +494,10 @@ void UNetObjectGridFilter::CalculateCellBox(const UNetObjectGridFilter::FPerObje
 	FVector MinPosition = Position - CullDistance;
 	FVector MaxPosition = Position + CullDistance;
 
+	// Clamp cells to the worldsize.
+	MinPosition = MinPosition.ComponentMax(MinWorldPos);
+	MaxPosition = MaxPosition.ComponentMin(MaxWorldPos);
+
 	const int64 MinX = FPlatformMath::FloorToInt(MinPosition.X/Config->CellSizeX);
 	const int64 MinY = FPlatformMath::FloorToInt(MinPosition.Y/Config->CellSizeY);
 	const int64 MaxX = FPlatformMath::FloorToInt(MaxPosition.X/Config->CellSizeX);
diff --git a/Engine/Source/Runtime/Experimental/Iris/Core/Private/Iris/ReplicationSystem/ReplicationSystem.cpp b/Engine/Source/Runtime/Experimental/Iris/Core/Private/Iris/ReplicationSystem/ReplicationSystem.cpp
--- a/Engine/Source/Runtime/Experimental/Iris/Core/Private/Iris/ReplicationSystem/ReplicationSystem.cpp
+++ b/Engine/Source/Runtime/Experimental/Iris/Core/Private/Iris/ReplicationSystem/ReplicationSystem.cpp
@@ -202,6 +202,7 @@ public:
 			FWorldLocations& WorldLocations = ReplicationSystemInternal.GetWorldLocations();
 			FWorldLocationsInitParams InitParams;
 			InitParams.MaxInternalNetRefIndex = CurrentMaxInternalNetRefIndex;
+			InitParams.ReplicationSystem = ReplicationSystem;
 			WorldLocations.Init(InitParams);
 		}
 	
diff --git a/Engine/Source/Runtime/Experimental/Iris/Core/Private/Iris/ReplicationSystem/WorldLocations.cpp b/Engine/Source/Runtime/Experimental/Iris/Core/Private/Iris/ReplicationSystem/WorldLocations.cpp
--- a/Engine/Source/Runtime/Experimental/Iris/Core/Private/Iris/ReplicationSystem/WorldLocations.cpp
+++ b/Engine/Source/Runtime/Experimental/Iris/Core/Private/Iris/ReplicationSystem/WorldLocations.cpp
@@ -2,6 +2,9 @@
 
 #include "Iris/ReplicationSystem/WorldLocations.h"
 #include "Iris/Core/IrisMemoryTracker.h"
+#include "Iris/Core/IrisLog.h"
+#include "Iris/ReplicationSystem/ReplicationSystem.h"
+#include "Iris/ReplicationSystem/ReplicationSystemInternal.h"
 #include "Iris/ReplicationSystem/NetRefHandleManager.h"
 
 #include UE_INLINE_GENERATED_CPP_BY_NAME(WorldLocations)
@@ -21,6 +24,9 @@ void FWorldLocations::Init(const FWorldLocationsInitParams& InitParams)
 
 	MinWorldPos = GetDefault<UWorldLocationsConfig>()->MinPos;
 	MaxWorldPos = GetDefault<UWorldLocationsConfig>()->MaxPos;
+	MaxNetCullDistance = GetDefault<UWorldLocationsConfig>()->MaxNetCullDistance;
+	
+	NetRefHandleManager = &InitParams.ReplicationSystem->GetReplicationSystemInternal()->GetNetRefHandleManager();
 }
 
 void FWorldLocations::PostSendUpdate()
@@ -97,9 +103,18 @@ void FWorldLocations::SetObjectInfo(UE::Net::Private::FInternalNetRefIndex Objec
 
 	const FVector ClampedLoc = ClampPositionToBoundary(Location);
 	
-	const bool bHasInfoChanged = ObjectsWithDirtyInfo.GetBit(ObjectIndex) || TargetObjectInfo.WorldLocation != ClampedLoc || TargetObjectInfo.CullDistance != NetCullDistance;
+	const bool bHasCullDistanceChanged = TargetObjectInfo.CullDistance != NetCullDistance;
+	const bool bHasInfoChanged = ObjectsWithDirtyInfo.GetBit(ObjectIndex) || TargetObjectInfo.WorldLocation != ClampedLoc || bHasCullDistanceChanged;
 
 	TargetObjectInfo.WorldLocation = ClampedLoc;
+
+	// For now we just warn, this will be clamped by filter.
+	if (bHasCullDistanceChanged && MaxNetCullDistance > 0.f && NetCullDistance > MaxNetCullDistance)
+	{
+		UE_LOG(LogIrisNetCull, Verbose, TEXT("FWorldLocations::SetObjectInfo ReplicatedObject %s cull distance %f is above the max %f. Consider making object always relevant instead"), *NetRefHandleManager->PrintObjectFromIndex(ObjectIndex), NetCullDistance, MaxNetCullDistance);
+		ensureMsgf(false, TEXT("FWorldLocations::SetObjectInfo ReplicatedObject %s cull distance %f is above the max %f. Consider making object always relevant instead"), *NetRefHandleManager->PrintObjectFromIndex(ObjectIndex), NetCullDistance, MaxNetCullDistance);
+	}
+
 	TargetObjectInfo.CullDistance = NetCullDistance;
 
 	ObjectsWithDirtyInfo.OrBitValue(ObjectIndex, bHasInfoChanged);
@@ -111,9 +126,16 @@ bool FWorldLocations::SetCullDistanceOverride(UE::Net::Private::FInternalNetRefI
 
 	if (ValidInfoIndexes.IsBitSet(ObjectIndex))
 	{
-		//TODO: Check for zero or clamp down on huge values ?
+		// TODO: Check for zero or clamp down on huge values ?
 		if (GetObjectInfo(ObjectIndex).CullDistanceOverride != CullDistance)
 		{
+			// For now we just warn, this will be clamped by filter.
+			if (MaxNetCullDistance > 0.f && CullDistance > MaxNetCullDistance)
+			{
+				UE_LOG(LogIrisNetCull, Verbose, TEXT("FWorldLocations::SetCullDistanceOverride ReplicatedObject %s cull distance %f is above the max %f. Consider making object always relevant instead"), *NetRefHandleManager->PrintObjectFromIndex(ObjectIndex), CullDistance, MaxNetCullDistance);	
+				ensureMsgf(false, TEXT("FWorldLocations::SetCullDistanceOverride ReplicatedObject %s cull distance %f is above the max %f. Consider making object always relevant instead"), *NetRefHandleManager->PrintObjectFromIndex(ObjectIndex), CullDistance, MaxNetCullDistance);	
+			}
+
 			GetObjectInfo(ObjectIndex).CullDistanceOverride = CullDistance;
 			ObjectsWithDirtyInfo.SetBitValue(ObjectIndex, true);
 		}
diff --git a/Engine/Source/Runtime/Experimental/Iris/Core/Public/Iris/ReplicationSystem/Filtering/NetObjectGridFilter.h b/Engine/Source/Runtime/Experimental/Iris/Core/Public/Iris/ReplicationSystem/Filtering/NetObjectGridFilter.h
--- a/Engine/Source/Runtime/Experimental/Iris/Core/Public/Iris/ReplicationSystem/Filtering/NetObjectGridFilter.h
+++ b/Engine/Source/Runtime/Experimental/Iris/Core/Public/Iris/ReplicationSystem/Filtering/NetObjectGridFilter.h
@@ -62,10 +62,6 @@ public:
 	UPROPERTY(Config)
 	float CellSizeY = 20000.0f;
 
-	/** Objects with larger sqrt(NetCullDistanceSqr) will be rejected. Disabled when value is zero. */
-	UPROPERTY(Config)
-	float MaxCullDistance = 0.0f;
-
 	/** Objects without a NetCullDistanceSquared property will assume to have this value but squared unless there's a cull distance override. */
 	UPROPERTY(Config)
 	float DefaultCullDistance = 15000.0f;
@@ -125,7 +121,7 @@ protected:
 		uint32 GetInfoIndex() const		{ return (uint32(Data[3]) << 16U) | uint32(Data[2]); }
 	};
 
-	IRISCORE_API void AddCellInfoForObject(const FObjectLocationInfo& ObjectInfo, const UE::Net::FReplicationInstanceProtocol* InstanceProtocol);
+	IRISCORE_API bool AddCellInfoForObject(const FObjectLocationInfo& ObjectInfo, const UE::Net::FReplicationInstanceProtocol* InstanceProtocol);
 	IRISCORE_API void RemoveCellInfoForObject(const FObjectLocationInfo& ObjectInfo);
 	IRISCORE_API void UpdateCellInfoForObject(const FObjectLocationInfo& ObjectInfo, const UE::Net::FReplicationInstanceProtocol* InstanceProtocol);
 
@@ -287,6 +283,11 @@ private:
 
 	// Cached state from UpdateCellInfoForObject().
 	FDebugUpdateCellInfo DebugUpdateCellInfo;
+
+	/** World boundaries (min and max position). */
+	FVector MinWorldPos;
+	FVector MaxWorldPos;
+	float MaxNetCullDistance = 0.f;
 };
 
 /**
diff --git a/Engine/Source/Runtime/Experimental/Iris/Core/Public/Iris/ReplicationSystem/WorldLocations.h b/Engine/Source/Runtime/Experimental/Iris/Core/Public/Iris/ReplicationSystem/WorldLocations.h
--- a/Engine/Source/Runtime/Experimental/Iris/Core/Public/Iris/ReplicationSystem/WorldLocations.h
+++ b/Engine/Source/Runtime/Experimental/Iris/Core/Public/Iris/ReplicationSystem/WorldLocations.h
@@ -9,9 +9,12 @@
 #include "Math/Vector.h"
 #include "WorldLocations.generated.h"
 
+class UReplicationSystem;
+
 namespace UE::Net::Private
 {
 	typedef uint32 FInternalNetRefIndex;
+	class FNetRefHandleManager;
 }
 
 /**
@@ -30,6 +33,10 @@ public:
 	/** All world positions will be clamped to MinPos and MaxPos. */
 	UPROPERTY(Config)
 	FVector MaxPos = { +0.5f * 2097152.0f, +0.5f * 2097152.0f, +0.5f * 2097152.0f };
+
+	/** We will issue a warning if user sets a higher NetCullDistance or NetCullDistanceOverride than the MaxNetCullDistance. */
+	UPROPERTY(Config)
+	float MaxNetCullDistance = 150000.f;
 };
 
 namespace UE::Net
@@ -37,6 +44,8 @@ namespace UE::Net
 
 struct FWorldLocationsInitParams
 {
+	TObjectPtr<UReplicationSystem> ReplicationSystem = nullptr;
+
 	UE::Net::Private::FInternalNetRefIndex MaxInternalNetRefIndex = 0;
 
 	/** How many world info storage slots to preallocate. */
@@ -203,9 +212,12 @@ private:
 
 	TNetChunkedArray<FObjectInfo, StorageElementsPerChunk> StoredObjectInfo;
 
+	const UE::Net::Private::FNetRefHandleManager* NetRefHandleManager = nullptr;
+
 	/** World boundaries (min and max position). */
 	FVector MinWorldPos;
 	FVector MaxWorldPos;
+	float MaxNetCullDistance;
 
 	/** Controls if the dirty list can be modified */
 	bool bLockdownDirtyList = false;
