[NetPhysics] Cherrypick CL 41061000 - Crash fixes during resimulation

diff --git a/Engine/Source/Runtime/Experimental/Chaos/Private/Chaos/ChaosMarshallingManager.cpp b/Engine/Source/Runtime/Experimental/Chaos/Private/Chaos/ChaosMarshallingManager.cpp
--- a/Engine/Source/Runtime/Experimental/Chaos/Private/Chaos/ChaosMarshallingManager.cpp
+++ b/Engine/Source/Runtime/Experimental/Chaos/Private/Chaos/ChaosMarshallingManager.cpp
@@ -21,6 +21,7 @@ FChaosMarshallingManager::FChaosMarshallingManager()
 , ExternalTimestamp_External(0)
 , SimTime_External(0)
 , InternalStep_External(0)
+, ConsumerData(nullptr)
 , ProducerData(nullptr)
 , CurPullData(nullptr)
 , Delay(SimDelay)
