Dataflow - Added Struct support for variables.

#rnx

diff --git a/Engine/Source/Runtime/Experimental/Dataflow/Core/Public/Dataflow/DataflowContextCache.h b/Engine/Source/Runtime/Experimental/Dataflow/Core/Public/Dataflow/DataflowContextCache.h
--- a/Engine/Source/Runtime/Experimental/Dataflow/Core/Public/Dataflow/DataflowContextCache.h
+++ b/Engine/Source/Runtime/Experimental/Dataflow/Core/Public/Dataflow/DataflowContextCache.h
@@ -3,6 +3,8 @@
 #pragma once
 
 #include "HAL/Platform.h"
+#include "StructUtils/InstancedStruct.h"
+#include "StructUtils/StructArrayView.h"
 #include "Templates/Decay.h"
 #include "Templates/UniquePtr.h"
 #include "UObject/UnrealType.h"
@@ -22,6 +24,12 @@ namespace UE::Dataflow
 		static constexpr bool Value = (std::is_pointer_v<Type> && std::is_convertible_v<Type, const UObjectBase*>) || TIsTObjectPtr_V<Type>;
 	};
 
+	/** Trait used to select the UStruct cache element path code. */
+	template <typename T, typename = void, typename = void>
+	struct TIsReflectedStruct { static constexpr bool Value = false; };
+	template <typename T>
+	struct TIsReflectedStruct<T, std::void_t<decltype(&TBaseStructure<typename TDecay<T>::Type>::Get)>, std::void_t<decltype(&T::StaticStruct)>> { static constexpr bool Value = true; };
+
 	struct FTimestamp
 	{
 		typedef uint64 Type;
@@ -55,7 +63,9 @@ namespace UE::Dataflow
 			CacheElementReference,
 			CacheElementNull,
 			CacheElementUObject, // UObjectPtr
-			CacheElementUObjectArray // TArray of UObjectPtr
+			CacheElementUObjectArray, // TArray of UObjectPtr
+			CacheElementUStruct,
+			CacheElementUStructArray
 		};
 
 		FContextCacheElementBase(EType CacheElementType, FGuid InNodeGuid = FGuid(), const FProperty* InProperty = nullptr, uint32 InNodeHash = 0, FTimestamp InTimestamp = FTimestamp::Invalid)
@@ -68,7 +78,7 @@ namespace UE::Dataflow
 		virtual ~FContextCacheElementBase() {}
 
 		// InReferenceDataKey is the key of the cache element this function is called on 
-		virtual TUniquePtr<FContextCacheElementBase> CreateReference(FContextCacheKey InReferenceDataKey) const = 0;
+		inline TUniquePtr<FContextCacheElementBase> CreateReference(FContextCacheKey InReferenceDataKey) const;
 
 		// clone the cache entry
 		virtual TUniquePtr<FContextCacheElementBase> Clone(const IContextCacheStore& Context) const = 0;
@@ -125,8 +135,6 @@ namespace UE::Dataflow
 
 		const T& GetDataDirect() const { return Data; }
 
-		inline virtual TUniquePtr<FContextCacheElementBase> CreateReference(FContextCacheKey InReferenceDataKey) const override;
-
 		inline virtual TUniquePtr<FContextCacheElementBase> Clone(const IContextCacheStore& Context) const override;
 
 	private:
@@ -137,20 +145,18 @@ namespace UE::Dataflow
 	//--------------------------------------------------------------------
 	// Reference to another context cache entry 
 	//--------------------------------------------------------------------
-	template<class T>
-	struct TContextCacheElementReference : public FContextCacheElementBase
+	struct FContextCacheElementReference : public FContextCacheElementBase
 	{
-		TContextCacheElementReference(FGuid InNodeGuid, const FProperty* InProperty, FContextCacheKey InDataKey, uint32 InNodeHash, FTimestamp Timestamp)
+		FContextCacheElementReference(FGuid InNodeGuid, const FProperty* InProperty, FContextCacheKey InDataKey, uint32 InNodeHash, FTimestamp Timestamp)
 			: FContextCacheElementBase(EType::CacheElementReference, InNodeGuid, InProperty, InNodeHash, Timestamp)
 			, DataKey(InDataKey)
 		{}
 
+		template<class T>
 		inline const T& GetData(const IContextCacheStore& Context, const FProperty* PropertyIn, const T& Default) const;
 
 		inline virtual const void* GetUntypedData(const IContextCacheStore& Context, const FProperty* PropertyIn) const;
 
-		inline virtual TUniquePtr<FContextCacheElementBase> CreateReference(FContextCacheKey InReferenceDataKey) const override;
-
 		inline virtual TUniquePtr<FContextCacheElementBase> Clone(const IContextCacheStore& Context) const override;
 
 	private:
@@ -175,8 +181,6 @@ namespace UE::Dataflow
 			return nullptr;
 		}
 
-		inline virtual TUniquePtr<FContextCacheElementBase> CreateReference(FContextCacheKey InReferenceDataKey) const override;
-
 		inline virtual TUniquePtr<FContextCacheElementBase> Clone(const IContextCacheStore& Context) const override;
 	};
 
@@ -206,8 +210,6 @@ namespace UE::Dataflow
 			return (void*)ObjectPtr;
 		}
 
-		inline virtual TUniquePtr<FContextCacheElementBase> CreateReference(FContextCacheKey InReferenceDataKey) const override;
-
 		inline virtual TUniquePtr<FContextCacheElementBase> Clone(const IContextCacheStore& Context) const override;
 
 		//~ Begin FGCObject interface
@@ -234,7 +236,7 @@ namespace UE::Dataflow
 
 
 		TContextCacheElementUObjectArray(const TContextCacheElementUObjectArray<T>& Other)
-			: FContextCacheElementBase(EType::CacheElementUObject, Other.GetNodeGuid(), Other.GetProperty(), Other.GetNodeHash(), Other.GetTimestamp())
+			: FContextCacheElementBase(EType::CacheElementUObjectArray, Other.GetNodeGuid(), Other.GetProperty(), Other.GetNodeHash(), Other.GetTimestamp())
 			, ObjectPtrArray(Other.ObjectPtrArray)
 		{}
 
@@ -245,11 +247,9 @@ namespace UE::Dataflow
 
 		virtual const void* GetUntypedData(const IContextCacheStore& /*Context*/, const FProperty* /*PropertyIn*/) const
 		{
-			return (void*)&ObjectPtrArray;
+			return &ObjectPtrArray;
 		}
 
-		inline virtual TUniquePtr<FContextCacheElementBase> CreateReference(FContextCacheKey InReferenceDataKey) const override;
-
 		inline virtual TUniquePtr<FContextCacheElementBase> Clone(const IContextCacheStore& Context) const override;
 
 		//~ Begin FGCObject interface
@@ -262,6 +262,148 @@ namespace UE::Dataflow
 		TArray<FDataType> ObjectPtrArray;               // Decaying T removes any reference and gets the correct underlying storage data type
 	};
 
+	//--------------------------------------------------------------------
+	// UStruct cache element
+	//--------------------------------------------------------------------
+	struct FContextCacheElementUStruct : public FContextCacheElementBase
+	{
+		FContextCacheElementUStruct(FGuid InNodeGuid, const FProperty* InProperty, const FConstStructView& StructView, uint32 InNodeHash, FTimestamp Timestamp)
+			: FContextCacheElementBase(EType::CacheElementUStruct, InNodeGuid, InProperty, InNodeHash, Timestamp)
+			, InstancedStruct(StructView)
+		{}
+
+		template<typename T>
+		FContextCacheElementUStruct(FGuid InNodeGuid, const FProperty* InProperty, T&& InStruct, uint32 InNodeHash, FTimestamp Timestamp)
+			: FContextCacheElementBase(EType::CacheElementUStruct, InNodeGuid, InProperty, InNodeHash, Timestamp)
+		{
+			InstancedStruct.InitializeAs<typename TDecay<T>::Type>(Forward<T>(InStruct));
+		}
+
+		FContextCacheElementUStruct(const FContextCacheElementUStruct& Other)
+			: FContextCacheElementBase(EType::CacheElementUStruct, Other.GetNodeGuid(), Other.GetProperty(), Other.GetNodeHash(), Other.GetTimestamp())
+			, InstancedStruct(Other.InstancedStruct)
+		{}
+
+		template<typename T>
+		const T& GetData(const IContextCacheStore& /*Context*/, const FProperty* /*PropertyIn*/, const T& /*Default*/) const
+		{
+			return InstancedStruct.Get<T>();
+		}
+
+		virtual const void* GetUntypedData(const IContextCacheStore& /*Context*/, const FProperty* /*PropertyIn*/) const
+		{
+			return InstancedStruct.GetMemory();
+		}
+
+		inline virtual TUniquePtr<FContextCacheElementBase> Clone(const IContextCacheStore& Context) const override
+		{
+			return MakeUnique<FContextCacheElementUStruct>(*this);
+		}
+
+	private:
+		FInstancedStruct InstancedStruct;
+	};
+
+	//--------------------------------------------------------------------
+	// UStruct array cache element
+	//--------------------------------------------------------------------
+	struct FContextCacheElementUStructArray : public FContextCacheElementBase
+	{
+		FContextCacheElementUStructArray(FGuid InNodeGuid, const FProperty* InProperty, const FConstStructArrayView& InStructArrayView, uint32 InNodeHash, FTimestamp Timestamp)
+			: FContextCacheElementBase(EType::CacheElementUStructArray, InNodeGuid, InProperty, InNodeHash, Timestamp)
+			, InstancedStructArray(InStructArrayView)
+		{}
+
+		template<typename T>
+		FContextCacheElementUStructArray(FGuid InNodeGuid, const FProperty* InProperty, TArray<T>&& InStructArray, uint32 InNodeHash, FTimestamp Timestamp)
+			: FContextCacheElementBase(EType::CacheElementUStructArray, InNodeGuid, InProperty, InNodeHash, Timestamp)
+			, InstancedStructArray(TBaseStructure<T>::Get())
+		{
+			InstancedStructArray.Get<T>() = Forward<TArray<T>>(InStructArray);
+		}
+
+		FContextCacheElementUStructArray(const FContextCacheElementUStructArray& Other)
+			: FContextCacheElementUStructArray(Other.GetNodeGuid(), Other.GetProperty(), Other.GetStructArrayView(), Other.GetNodeHash(), Other.GetTimestamp())
+		{}
+	
+		virtual ~FContextCacheElementUStructArray() override = default;
+
+		template<typename T>
+		const TArray<T>& GetData(const IContextCacheStore& /*Context*/, const FProperty* /*PropertyIn*/, const TArray<T>& /*Default*/) const
+		{
+			return InstancedStructArray.Get<T>();
+		}
+
+		virtual const void* GetUntypedData(const IContextCacheStore& /*Context*/, const FProperty* /*PropertyIn*/) const
+		{
+			return InstancedStructArray.GetMemory();
+		}
+
+		inline virtual TUniquePtr<FContextCacheElementBase> Clone(const IContextCacheStore& Context) const override
+		{
+			return MakeUnique<FContextCacheElementUStructArray>(*this);
+		}
+
+	private:
+		FConstStructArrayView GetStructArrayView() const
+		{
+			return InstancedStructArray.GetScriptStruct() ?
+				FConstStructArrayView(*InstancedStructArray.GetScriptStruct(), InstancedStructArray.GetData(), InstancedStructArray.Num()) :
+				FConstStructArrayView();
+		}
+
+		// Implements a FInstancedStruct for arrays (FInstancedStructContainer cannot be cast to a TArray)
+		struct FInstancedStructArray final : public TArray<uint8>
+		{
+			explicit FInstancedStructArray(const UScriptStruct* const InScriptStruct)
+				: ScriptStruct(InScriptStruct)
+			{}
+			explicit FInstancedStructArray(const FConstStructArrayView& StructArrayView)
+				: ScriptStruct(StructArrayView.GetScriptStruct())
+			{
+				if (ScriptStruct)  // Null if the array has 0 element
+				{
+					SetNumUnsafeInternal(StructArrayView.Num());
+					ArrayMax = Num();
+					GetAllocatorInstance().ResizeAllocation(0, Num(), FMath::Max(1, ScriptStruct->GetStructureSize()), ScriptStruct->GetMinAlignment());
+					ScriptStruct->InitializeStruct(GetData(), Num());
+					ScriptStruct->CopyScriptStruct(GetData(), StructArrayView.GetData(), Num());
+				}
+			}
+			~FInstancedStructArray()
+			{
+				if (ScriptStruct)
+				{
+					ScriptStruct->DestroyStruct(GetData(), Num());
+					GetAllocatorInstance().ResizeAllocation(ArrayMax, 0, FMath::Max(1, ScriptStruct->GetStructureSize()), ScriptStruct->GetMinAlignment());
+				}
+			}
+			const UScriptStruct* GetScriptStruct() const
+			{
+				return ScriptStruct;
+			}
+			const void* GetMemory() const
+			{
+				return this;
+			}
+			template<typename T>
+			const TArray<T>& Get() const
+			{
+				check(!ScriptStruct || TBaseStructure<T>::Get() == ScriptStruct);
+				return reinterpret_cast<const TArray<T>&>(*this);
+			}
+			template<typename T>
+			TArray<T>& Get()
+			{
+				check(!ScriptStruct || TBaseStructure<T>::Get() == ScriptStruct);
+				return reinterpret_cast<TArray<T>&>(*this);
+			}
+		private:
+			const UScriptStruct* const ScriptStruct;
+		};
+
+		FInstancedStructArray InstancedStructArray;
+	};
 
 	// cache element method implementation 
 	template<class T>
@@ -277,7 +419,7 @@ namespace UE::Dataflow
 		}
 		if (Type == EType::CacheElementReference)
 		{
-			return static_cast<const TContextCacheElementReference<T>&>(*this).GetData(Context, PropertyIn, Default);
+			return static_cast<const FContextCacheElementReference&>(*this).GetData(Context, PropertyIn, Default);
 		}
 		if (Type == EType::CacheElementNull)
 		{
@@ -300,6 +442,23 @@ namespace UE::Dataflow
 				return static_cast<const TContextCacheElementUObject<T>&>(*this).GetData(Context, PropertyIn, Default);
 			}
 		}
+		if (Type == EType::CacheElementUStructArray)
+		{
+			if constexpr (TIsTArray<T>::Value)
+			{
+				if constexpr (TIsReflectedStruct<typename T::ElementType>::Value)
+				{
+					return static_cast<const FContextCacheElementUStructArray&>(*this).GetData<typename T::ElementType>(Context, PropertyIn, Default);
+				}
+			}
+		}
+		if (Type == EType::CacheElementUStruct)
+		{
+			if constexpr (TIsReflectedStruct<T>::Value)
+			{
+				return static_cast<const FContextCacheElementUStruct&>(*this).GetData<T>(Context, PropertyIn, Default);
+			}
+		}
 		check(false); // should never happen
 		return Default;
 	}
@@ -312,16 +471,15 @@ namespace UE::Dataflow
 	// cache classes implemetation 
 	// this needs to be after the FContext definition because they access its methods
 
-	template<class T>
-	const T& TContextCacheElement<T>::GetData(const IContextCacheStore& Context, const FProperty* PropertyIn, const T& Default) const
+	TUniquePtr<FContextCacheElementBase> FContextCacheElementBase::CreateReference(FContextCacheKey InReferenceDataKey) const
 	{
-		return Data;
+		return MakeUnique<FContextCacheElementReference>(GetNodeGuid(), GetProperty(), InReferenceDataKey, GetNodeHash(), GetTimestamp());
 	}
 
 	template<class T>
-	TUniquePtr<FContextCacheElementBase> TContextCacheElement<T>::CreateReference(FContextCacheKey InReferenceDataKey) const
+	const T& TContextCacheElement<T>::GetData(const IContextCacheStore& Context, const FProperty* PropertyIn, const T& Default) const
 	{
-		return MakeUnique<TContextCacheElementReference<T>>(GetNodeGuid(), GetProperty(), InReferenceDataKey, GetNodeHash(), GetTimestamp());
+		return Data;
 	}
 
 	template<class T>
@@ -331,7 +489,7 @@ namespace UE::Dataflow
 	}
 
 	template<class T>
-	const T& TContextCacheElementReference<T>::GetData(const IContextCacheStore& Context, const FProperty* InProperty, const T& Default) const
+	const T& FContextCacheElementReference::GetData(const IContextCacheStore& Context, const FProperty* InProperty, const T& Default) const
 	{
 		if (const TUniquePtr<FContextCacheElementBase>* Cache = Context.FindCacheElement(DataKey))
 		{
@@ -340,8 +498,7 @@ namespace UE::Dataflow
 		return Default;
 	}
 
-	template<class T>
-	const void* TContextCacheElementReference<T>::GetUntypedData(const IContextCacheStore& Context, const FProperty* InProperty) const
+	const void* FContextCacheElementReference::GetUntypedData(const IContextCacheStore& Context, const FProperty* InProperty) const
 	{
 		if (const TUniquePtr<FContextCacheElementBase>* Cache = Context.FindCacheElement(DataKey))
 		{
@@ -350,14 +507,7 @@ namespace UE::Dataflow
 		return nullptr;
 	}
 
-	template<class T>
-	TUniquePtr<FContextCacheElementBase> TContextCacheElementReference<T>::CreateReference(FContextCacheKey InReferenceDataKey) const
-	{
-		return MakeUnique<TContextCacheElementReference<T>>(GetNodeGuid(), GetProperty(), InReferenceDataKey, GetNodeHash(), GetTimestamp());
-	}
-
-	template<class T>
-	TUniquePtr<FContextCacheElementBase> TContextCacheElementReference<T>::Clone(const IContextCacheStore& Context) const
+	TUniquePtr<FContextCacheElementBase> FContextCacheElementReference::Clone(const IContextCacheStore& Context) const
 	{
 		if (const TUniquePtr<FContextCacheElementBase>* ReferencedCacheEntry = Context.FindCacheElement(DataKey))
 		{
@@ -369,36 +519,17 @@ namespace UE::Dataflow
 		return MakeUnique<FContextCacheElementNull>(GetNodeGuid(), GetProperty(), FContextCacheKey(), GetNodeHash(), GetTimestamp());
 	}
 
-	inline TUniquePtr<FContextCacheElementBase> FContextCacheElementNull::CreateReference(FContextCacheKey InReferenceDataKey) const
-	{
-		// a null entry will always return a null entry as a reference
-		return MakeUnique<FContextCacheElementNull>(GetNodeGuid(), GetProperty(), InReferenceDataKey, GetNodeHash(), GetTimestamp());
-	}
-
-
 	TUniquePtr<FContextCacheElementBase> FContextCacheElementNull::Clone(const IContextCacheStore& Context) const
 	{
 		return MakeUnique<FContextCacheElementNull>(GetNodeGuid(), GetProperty(), FContextCacheKey(), GetNodeHash(), GetTimestamp());
 	}
 
-	template<class T>
-	TUniquePtr<FContextCacheElementBase> TContextCacheElementUObject<T>::CreateReference(FContextCacheKey InReferenceDataKey) const
-	{
-		return MakeUnique<TContextCacheElementReference<T>>(GetNodeGuid(), GetProperty(), InReferenceDataKey, GetNodeHash(), GetTimestamp());
-	}
-
 	template<class T>
 	TUniquePtr<FContextCacheElementBase> TContextCacheElementUObject<T>::Clone(const IContextCacheStore& Context) const
 	{
 		return MakeUnique<TContextCacheElementUObject<T>>(*this);
 	}
 
-	template<class T>
-	TUniquePtr<FContextCacheElementBase> TContextCacheElementUObjectArray<T>::CreateReference(FContextCacheKey InReferenceDataKey) const
-	{
-		return MakeUnique<TContextCacheElementReference<T>>(GetNodeGuid(), GetProperty(), InReferenceDataKey, GetNodeHash(), GetTimestamp());
-	}
-
 	template<class T>
 	TUniquePtr<FContextCacheElementBase> TContextCacheElementUObjectArray<T>::Clone(const IContextCacheStore& Context) const
 	{
diff --git a/Engine/Source/Runtime/Experimental/Dataflow/Core/Public/Dataflow/DataflowInputOutput.h b/Engine/Source/Runtime/Experimental/Dataflow/Core/Public/Dataflow/DataflowInputOutput.h
--- a/Engine/Source/Runtime/Experimental/Dataflow/Core/Public/Dataflow/DataflowInputOutput.h
+++ b/Engine/Source/Runtime/Experimental/Dataflow/Core/Public/Dataflow/DataflowInputOutput.h
@@ -220,6 +220,22 @@ public:
 		}
 	}
 
+	void SetValueFromStructView(const FConstStructView& StructView, UE::Dataflow::FContext& Context) const
+	{
+		if (Property)
+		{
+			Context.SetDataFromStructView(CacheKey(), Property, StructView, GetOwningNodeGuid(), GetOwningNodeValueHash(), GetOwningNodeTimestamp());
+		}
+	}
+
+	void SetValueFromStructArrayView(const FConstStructArrayView& StructArrayView, UE::Dataflow::FContext& Context) const
+	{
+		if (Property)
+		{
+			Context.SetDataFromStructArrayView(CacheKey(), Property, StructArrayView, GetOwningNodeGuid(), GetOwningNodeValueHash(), GetOwningNodeTimestamp());
+		}
+	}
+
 	template<typename TAnyType>
 	void SetValueFromAnyType(const typename TAnyType::FStorageType& InVal, UE::Dataflow::FContext& Context) const
 	{
diff --git a/Engine/Source/Runtime/Experimental/Dataflow/Core/Public/Dataflow/DataflowNodeParameters.h b/Engine/Source/Runtime/Experimental/Dataflow/Core/Public/Dataflow/DataflowNodeParameters.h
--- a/Engine/Source/Runtime/Experimental/Dataflow/Core/Public/Dataflow/DataflowNodeParameters.h
+++ b/Engine/Source/Runtime/Experimental/Dataflow/Core/Public/Dataflow/DataflowNodeParameters.h
@@ -193,6 +193,10 @@ namespace UE::Dataflow
 					bMadeEntry = true;
 					DataStoreEntry = MakeUnique<TContextCacheElementUObjectArray<typename ArrayType::ElementType>>(InNodeGuid, InProperty, Forward<T>(InValue), InNodeHash, InTimestamp);
 				}
+				else if constexpr (TIsReflectedStruct<ArrayType>::Value)
+				{
+					DataStoreEntry = MakeUnique<FContextCacheElementUStructArray>(InNodeGuid, InProperty, Forward<T>(InValue), InNodeHash, InTimestamp);
+				}
 			}
 			if (!bMadeEntry)
 			{
@@ -200,6 +204,10 @@ namespace UE::Dataflow
 				{
 					DataStoreEntry = MakeUnique<TContextCacheElementUObject<T>>(InNodeGuid, InProperty, Forward<T>(InValue), InNodeHash, InTimestamp);
 				}
+				else if constexpr (TIsReflectedStruct<T>::Value)
+				{
+					DataStoreEntry = MakeUnique<FContextCacheElementUStruct>(InNodeGuid, InProperty, Forward<T>(InValue), InNodeHash, InTimestamp);
+				}
 				else
 				{
 					DataStoreEntry = MakeUnique<TContextCacheElement<T>>(InNodeGuid, InProperty, Forward<T>(InValue), InNodeHash, InTimestamp);
@@ -209,6 +217,18 @@ namespace UE::Dataflow
 			SetDataImpl(InKey, MoveTemp(DataStoreEntry));
 		}
 
+		void SetDataFromStructView(FContextCacheKey InKey, const FProperty* InProperty, const FConstStructView& StructView, const FGuid& InNodeGuid, uint32 InNodeHash, const FTimestamp& InTimestamp)
+		{
+			TUniquePtr<FContextCacheElementBase> DataStoreEntry = MakeUnique<FContextCacheElementUStruct>(InNodeGuid, InProperty, StructView, InNodeHash, InTimestamp);
+			SetDataImpl(InKey, MoveTemp(DataStoreEntry));
+		}
+
+		void SetDataFromStructArrayView(FContextCacheKey InKey, const FProperty* InProperty, const FConstStructArrayView& StructArrayView, const FGuid& InNodeGuid, uint32 InNodeHash, const FTimestamp& InTimestamp)
+		{
+			TUniquePtr<FContextCacheElementBase> DataStoreEntry = MakeUnique<FContextCacheElementUStructArray>(InNodeGuid, InProperty, StructArrayView, InNodeHash, InTimestamp);
+			SetDataImpl(InKey, MoveTemp(DataStoreEntry));
+		}
+
 		void SetDataReference(FContextCacheKey Key, const FProperty* Property, FContextCacheKey ReferenceKey, const FTimestamp& InTimestamp)
 		{
 			// find the reference key to get 
diff --git a/Engine/Source/Runtime/Experimental/Dataflow/Engine/Private/Dataflow/DataflowVariableNodes.cpp b/Engine/Source/Runtime/Experimental/Dataflow/Engine/Private/Dataflow/DataflowVariableNodes.cpp
--- a/Engine/Source/Runtime/Experimental/Dataflow/Engine/Private/Dataflow/DataflowVariableNodes.cpp
+++ b/Engine/Source/Runtime/Experimental/Dataflow/Engine/Private/Dataflow/DataflowVariableNodes.cpp
@@ -189,9 +189,12 @@ void FGetDataflowVariableNode::UpdateOutputTypes(const FPropertyBagPropertyDesc&
 		ChangeOutputType(Output, ConcreteType);
 		EvaluateFunction = &EvaluateObject;
 		break;
+	case EPropertyBagPropertyType::Struct:
+		ChangeOutputType(Output, ConcreteType);
+		EvaluateFunction = &EvaluateStruct;
+		break;
 		// unsupported type for now 
 	case EPropertyBagPropertyType::Enum:
-	case EPropertyBagPropertyType::Struct:
 	case EPropertyBagPropertyType::SoftObject:
 	case EPropertyBagPropertyType::Class:
 	case EPropertyBagPropertyType::SoftClass:
@@ -589,7 +592,44 @@ void FGetDataflowVariableNode::Evaluate(UE::Dataflow::FContext& Context, const F
 		break;
 	}
 	}
+}
 
-
+/*static*/ void FGetDataflowVariableNode::EvaluateStruct(const FGetDataflowVariableNode& Node, const FInstancedPropertyBag& Variables, UE::Dataflow::FContext& Context, const FDataflowOutput& Out)
+{
+	switch (Node.ContainerType)
+	{
+	case EPropertyBagContainerType::None:
+	{
+		FConstStructView OutStruct;
+		if (const FStructView* const Value = Variables.GetValueStruct(Node.VariableName).TryGetValue())
+		{
+			OutStruct = *Value;
+		}
+		Out.SetValueFromStructView(OutStruct, Context);
+		break;
+	}
+	case EPropertyBagContainerType::Array:
+	{
+		FConstStructArrayView OutStructArrayView;
+		if (const FPropertyBagArrayRef* const ArrayRef = Variables.GetArrayRef(Node.VariableName).TryGetValue())
+		{
+			if (const int32 Count = ArrayRef->Num())
+			{
+				if (const FStructView* const Value = ArrayRef->GetValueStruct(0).TryGetValue())
+				{
+					if (const UScriptStruct* const ScriptStruct = Value->GetScriptStruct())
+					{
+						OutStructArrayView = FConstStructArrayView(
+							*ScriptStruct,
+							(const void*)Value->GetMemory(),
+							Count);
+					}
+				}
+			}
+		}
+		Out.SetValueFromStructArrayView(OutStructArrayView, Context);
+		break;
+	}
+	}
 }
 
diff --git a/Engine/Source/Runtime/Experimental/Dataflow/Engine/Public/Dataflow/DataflowVariableNodes.h b/Engine/Source/Runtime/Experimental/Dataflow/Engine/Public/Dataflow/DataflowVariableNodes.h
--- a/Engine/Source/Runtime/Experimental/Dataflow/Engine/Public/Dataflow/DataflowVariableNodes.h
+++ b/Engine/Source/Runtime/Experimental/Dataflow/Engine/Public/Dataflow/DataflowVariableNodes.h
@@ -48,6 +48,7 @@ private:
 	static void EvaluateString(const FGetDataflowVariableNode& Node, const FInstancedPropertyBag& Variables, UE::Dataflow::FContext& Context, const FDataflowOutput& Out);
 	static void EvaluateText(const FGetDataflowVariableNode& Node, const FInstancedPropertyBag& Variables, UE::Dataflow::FContext& Context, const FDataflowOutput& Out);
 	static void EvaluateObject(const FGetDataflowVariableNode& Node, const FInstancedPropertyBag& Variables, UE::Dataflow::FContext& Context, const FDataflowOutput& Out);
+	static void EvaluateStruct(const FGetDataflowVariableNode& Node, const FInstancedPropertyBag& Variables, UE::Dataflow::FContext& Context, const FDataflowOutput& Out);
 
 	void UpdateOutputTypes(const FPropertyBagPropertyDesc& Desc);
 	void ChangeOutputType(FDataflowOutput* Output, FName NewType);
