IoDispatcher error improvement
- Add EIoErrorCode::NotInstalled
- In FIoRequestImpl replace bCancelled and bFailed with LastBackendError
- FOnDemandInstallCache now handles any request for a chunk it knows about. If a request is not in the cache, it fails with NotInstalled.

Next steps will be to bubble up errors through the async loader.

diff --git a/Engine/Source/Runtime/Core/Private/IO/IoDispatcher.cpp b/Engine/Source/Runtime/Core/Private/IO/IoDispatcher.cpp
--- a/Engine/Source/Runtime/Core/Private/IO/IoDispatcher.cpp
+++ b/Engine/Source/Runtime/Core/Private/IO/IoDispatcher.cpp
@@ -56,7 +56,8 @@ static const TCHAR* const GetIoErrorText_ErrorCodeTextArray[] =
 	TEXT("Compression Error"),
 	TEXT("Pending Fork"),
 	TEXT("Pending Encryption Key"),
-	TEXT("Disabled")
+	TEXT("Disabled"),
+	TEXT("Not Installed")
 };
 
 CORE_API const TCHAR* const* GetIoErrorText_ErrorCodeText = GetIoErrorText_ErrorCodeTextArray;
@@ -715,13 +716,9 @@ private:
 			while (CompletedRequestsHead)
 			{
 				FIoRequestImpl* NextRequest = CompletedRequestsHead->NextRequest;
-				if (CompletedRequestsHead->bCancelled)
+				if (CompletedRequestsHead->LastBackendError != EIoErrorCode::Ok)
 				{
-					CompleteRequest(CompletedRequestsHead, EIoErrorCode::Cancelled);
-				}
-				else if (CompletedRequestsHead->bFailed)
-				{
-					CompleteRequest(CompletedRequestsHead, EIoErrorCode::ReadError);
+					CompleteRequest(CompletedRequestsHead, CompletedRequestsHead->LastBackendError);
 				}
 				else
 				{
@@ -810,9 +807,9 @@ private:
 			}
 			for (FIoRequestImpl* RequestToCancel : LocalRequestsToCancel)
 			{
-				if (!RequestToCancel->bCancelled)
+				if (!RequestToCancel->IsCancelled())
 				{
-					RequestToCancel->bCancelled = true;
+					RequestToCancel->SetLastBackendError(EIoErrorCode::Cancelled);
 					if (RequestToCancel->Backend)
 					{
 						RequestToCancel->Backend->CancelIoRequest(RequestToCancel);
@@ -840,7 +837,7 @@ private:
 				check(Request->NextRequest == nullptr);
 				RequestStats.OnRequestStarted(*Request);
 
-				if (Request->bCancelled)
+				if (Request->IsCancelled())
 				{
 					CompleteRequest(Request, EIoErrorCode::Cancelled);
 					Request->ReleaseRef();
diff --git a/Engine/Source/Runtime/Core/Public/IO/IoDispatcherBackend.h b/Engine/Source/Runtime/Core/Public/IO/IoDispatcherBackend.h
--- a/Engine/Source/Runtime/Core/Public/IO/IoDispatcherBackend.h
+++ b/Engine/Source/Runtime/Core/Public/IO/IoDispatcherBackend.h
@@ -42,19 +42,32 @@ public:
 	/** Returns whether the request has been cancelled. */
 	bool IsCancelled() const
 	{
-		return bCancelled;
+		return LastBackendError == EIoErrorCode::Cancelled;
 	}
 
 	/** Returns whether the request failed. */
 	bool IsFailed() const
 	{
-		return bFailed;
+		return 
+			LastBackendError != EIoErrorCode::Ok &&
+			LastBackendError != EIoErrorCode::Cancelled;
 	}
 
-	/** Mark the request as failed. */
+	/** Mark the request as failed (EIoErrorCode::ReadError). */
 	void SetFailed()
 	{
-		bFailed = true;
+		if (!IsCancelled())
+		{
+			LastBackendError = EIoErrorCode::ReadError;
+		}
+	}
+
+	void SetLastBackendError(EIoErrorCode InError)
+	{
+		if (!IsCancelled())
+		{
+			LastBackendError = InError;
+		}
 	}
 
 	/** Returns whether request has a valid buffer. */
@@ -118,8 +131,7 @@ private:
 	FIoReadCallback Callback;
 	TAtomic<uint32> RefCount{ 0 };
 	TAtomic<EIoErrorCode> ErrorCode{ EIoErrorCode::Unknown };
-	bool bCancelled = false;
-	bool bFailed = false;
+	EIoErrorCode LastBackendError = EIoErrorCode::Ok;
 };
 
 namespace UE::Private
diff --git a/Engine/Source/Runtime/Core/Public/IO/IoStatus.h b/Engine/Source/Runtime/Core/Public/IO/IoStatus.h
--- a/Engine/Source/Runtime/Core/Public/IO/IoStatus.h
+++ b/Engine/Source/Runtime/Core/Public/IO/IoStatus.h
@@ -4,7 +4,6 @@
 
 #include "Containers/StringView.h"
 #include "Containers/UnrealString.h"
-#include "Experimental/UnifiedError/UnifiedError.h"
 
 template <typename CharType> class TStringBuilderBase;
 
@@ -33,7 +32,8 @@ enum class EIoErrorCode
 	CompressionError,
 	PendingFork,
 	PendingEncryptionKey,
-	Disabled
+	Disabled,
+	NotInstalled
 };
 
 /** Get I/O error code description. */
diff --git a/Engine/Source/Runtime/Experimental/IoStore/OnDemand/Private/OnDemandInstallCache.cpp b/Engine/Source/Runtime/Experimental/IoStore/OnDemand/Private/OnDemandInstallCache.cpp
--- a/Engine/Source/Runtime/Experimental/IoStore/OnDemand/Private/OnDemandInstallCache.cpp
+++ b/Engine/Source/Runtime/Experimental/IoStore/OnDemand/Private/OnDemandInstallCache.cpp
@@ -1253,7 +1253,7 @@ private:
 													FCasBlockInfoMap& BlockInfo) const;
 	FIoStatus									Purge(const FCasBlockInfoMap& BlockInfo, uint64 TotalBytesToPurge, uint64& OutTotalPurgedBytes);
 	bool										Resolve(FIoRequestImpl* Request);
-	void										CompleteRequest(FIoRequestImpl* Request, bool bFileReadWasCancelled);
+	void										CompleteRequest(FIoRequestImpl* Request, EIoErrorCode Status);
 	FIoStatus									FlushPendingChunks(FPendingChunks& Chunks);
 	FIoStatus									FlushPendingChunksImpl(FPendingChunks& Chunks);
 	FString										GetJournalFilename() const { return CacheDirectory / TEXT("cas.jrn"); }
@@ -1556,7 +1556,8 @@ bool FOnDemandInstallCache::Resolve(FIoRequestImpl* Request)
 	const FCasLocation CasLoc = Cas.FindChunk(ChunkInfo.Hash());
 	if (CasLoc.IsValid() == false)
 	{
-		return false;
+		CompleteRequest(Request, EIoErrorCode::NotInstalled);
+		return true;
 	}
 
 	const uint64 RequestSize = FMath::Min<uint64>(
@@ -1573,7 +1574,8 @@ bool FOnDemandInstallCache::Resolve(FIoRequestImpl* Request)
 	if (ChunkRange.IsOk() == false)
 	{
 		UE_LOG(LogIoStoreOnDemand, Error, TEXT("Failed to get chunk range"));
-		return false;
+		CompleteRequest(Request, ChunkRange.Status().GetErrorCode());
+		return true;
 	}
 
 	TRACE_IOSTORE_BACKEND_REQUEST_STARTED(Request, this);
@@ -1630,10 +1632,10 @@ bool FOnDemandInstallCache::Resolve(FIoRequestImpl* Request)
 					*Filename, *FileOpenResult.GetError().GetMessage());
 			}
 
-			const bool bWasCancelled = bOk == false;
-			UE::Tasks::Launch(UE_SOURCE_LOCATION, [this, Request, bWasCancelled]
+			UE::Tasks::Launch(UE_SOURCE_LOCATION, [this, Request, bOk]
 			{
-				CompleteRequest(Request, bWasCancelled);
+				const EIoErrorCode Status = bOk ? EIoErrorCode::Ok : EIoErrorCode::ReadError;
+				CompleteRequest(Request, Status);
 			});
 		}, UE::Tasks::ETaskPriority::BackgroundHigh);
 
@@ -1647,7 +1649,8 @@ bool FOnDemandInstallCache::Resolve(FIoRequestImpl* Request)
 		const FString Filename = Cas.GetBlockFilename(CasLoc.BlockId);
 		UE_LOG(LogIoStoreOnDemand, Error, TEXT("Failed to open CAS block '%s' for async reading, reason '%s'"), *Filename, *FileOpenResult.GetError().GetMessage());
 		TUniquePtr<FChunkRequest> Detached = FChunkRequest::Detach(*Request);
-		return false;
+		CompleteRequest(Request, EIoErrorCode::FileOpenFailed);
+		return true;
 	}
 	
 	FSharedAsyncFileHandle FileHandle(FileOpenResult.StealValue());
@@ -1665,7 +1668,8 @@ bool FOnDemandInstallCache::Resolve(FIoRequestImpl* Request)
 	{
 		UE::Tasks::Launch(UE_SOURCE_LOCATION, [this, Request, bWasCancelled]
 		{
-			CompleteRequest(Request, bWasCancelled);
+			const EIoErrorCode Status = bWasCancelled ? EIoErrorCode::ReadError : EIoErrorCode::Ok;
+			CompleteRequest(Request, Status);
 		});
 	};
 
@@ -1680,7 +1684,8 @@ bool FOnDemandInstallCache::Resolve(FIoRequestImpl* Request)
 	{
 		TRACE_IOSTORE_BACKEND_REQUEST_FAILED(Request);
 		TUniquePtr<FChunkRequest> Detached = FChunkRequest::Detach(*Request);
-		return false;
+		CompleteRequest(Request, EIoErrorCode::ReadError);
+		return true;
 	}
 
 	return true;
@@ -2049,36 +2054,44 @@ FIoStatus FOnDemandInstallCache::FlushPendingChunksImpl(FPendingChunks& Chunks)
 	return FIoStatus::Ok;
 }
 
-void FOnDemandInstallCache::CompleteRequest(FIoRequestImpl* Request, bool bFileReadWasCancelled)
+void FOnDemandInstallCache::CompleteRequest(FIoRequestImpl* Request, EIoErrorCode Status)
 {
-	FChunkRequest& ChunkRequest			= FChunkRequest::GetRef(*Request);
-	const FOnDemandChunkInfo& ChunkInfo	= ChunkRequest.ChunkInfo;
-	FIoBuffer EncodedChunk				= MoveTemp(ChunkRequest.EncodedChunk);
-	bool bSucceeded						= EncodedChunk.GetSize() > 0 && !bFileReadWasCancelled && !Request->IsCancelled();
-
-	if (bSucceeded)
+	bool bSucceeded	= false;
+	if (Status == EIoErrorCode::Ok && !Request->IsCancelled())
 	{
-		FIoChunkDecodingParams Params;
-		Params.CompressionFormat	= ChunkInfo.CompressionFormat();
-		Params.EncryptionKey		= ChunkInfo.EncryptionKey();
-		Params.BlockSize			= ChunkInfo.BlockSize();
-		Params.TotalRawSize			= ChunkInfo.RawSize();
-		Params.RawOffset			= Request->Options.GetOffset();
-		Params.EncodedOffset		= ChunkRequest.ChunkRange.GetOffset();
-		Params.EncodedBlockSize		= ChunkInfo.Blocks();
-		Params.BlockHash			= ChunkInfo.BlockHashes();
+		FChunkRequest& ChunkRequest = FChunkRequest::GetRef(*Request);
+		const FOnDemandChunkInfo& ChunkInfo = ChunkRequest.ChunkInfo;
+		FIoBuffer EncodedChunk = MoveTemp(ChunkRequest.EncodedChunk);
+
+		if (EncodedChunk.GetSize() > 0)
+		{
+			FIoChunkDecodingParams Params;
+			Params.CompressionFormat = ChunkInfo.CompressionFormat();
+			Params.EncryptionKey = ChunkInfo.EncryptionKey();
+			Params.BlockSize = ChunkInfo.BlockSize();
+			Params.TotalRawSize = ChunkInfo.RawSize();
+			Params.RawOffset = Request->Options.GetOffset();
+			Params.EncodedOffset = ChunkRequest.ChunkRange.GetOffset();
+			Params.EncodedBlockSize = ChunkInfo.Blocks();
+			Params.BlockHash = ChunkInfo.BlockHashes();
 
-		Request->CreateBuffer(ChunkRequest.RawSize);
-		FMutableMemoryView RawChunk = Request->GetBuffer().GetMutableView();
+			Request->CreateBuffer(ChunkRequest.RawSize);
+			FMutableMemoryView RawChunk = Request->GetBuffer().GetMutableView();
 
-		bSucceeded = FIoChunkEncoding::Decode(Params, EncodedChunk.GetView(), RawChunk);
-		UE_CLOG(!bSucceeded, LogIoStoreOnDemand, Error, TEXT("Failed to decode chunk, ChunkId='%s'"), *LexToString(Request->ChunkId));
+			bSucceeded = FIoChunkEncoding::Decode(Params, EncodedChunk.GetView(), RawChunk);
+			UE_CLOG(!bSucceeded, LogIoStoreOnDemand, Error, TEXT("Failed to decode chunk, ChunkId='%s'"), *LexToString(Request->ChunkId));
+		}
 	}
 
 	if (bSucceeded == false)
 	{
+		if (Status == EIoErrorCode::Ok)
+		{
+			Status = EIoErrorCode::ReadError;
+		}
+		Request->SetLastBackendError(Status);
+
 		Request->SetResult(FIoBuffer());
-		Request->SetFailed();
 		TRACE_IOSTORE_BACKEND_REQUEST_FAILED(Request);
 	}
 	else
