Dataflow : show variable name bound to the variable node in the details for easier debugging

diff --git a/Engine/Source/Runtime/Experimental/Dataflow/Engine/Public/Dataflow/DataflowVariableNodes.h b/Engine/Source/Runtime/Experimental/Dataflow/Engine/Public/Dataflow/DataflowVariableNodes.h
--- a/Engine/Source/Runtime/Experimental/Dataflow/Engine/Public/Dataflow/DataflowVariableNodes.h
+++ b/Engine/Source/Runtime/Experimental/Dataflow/Engine/Public/Dataflow/DataflowVariableNodes.h
@@ -57,7 +57,7 @@ private:
 	UPROPERTY(EditAnywhere, Category = Variables, meta = (FixedLayout, ShowOnlyInnerProperties))
 	FInstancedPropertyBag VariablePropertyBag;
 
-	UPROPERTY()
+	UPROPERTY(VisibleAnywhere, Category = Variables)
 	FName VariableName;
 
 	void RegisterHandlers();
