[NetPhysics] Change to Predictive Interpolation if received target state is outside of rewind history when running resimulation replication.

diff --git a/Engine/Source/Runtime/Engine/Private/PhysicsEngine/NetworkPhysicsComponent.cpp b/Engine/Source/Runtime/Engine/Private/PhysicsEngine/NetworkPhysicsComponent.cpp
--- a/Engine/Source/Runtime/Engine/Private/PhysicsEngine/NetworkPhysicsComponent.cpp
+++ b/Engine/Source/Runtime/Engine/Private/PhysicsEngine/NetworkPhysicsComponent.cpp
@@ -73,6 +73,9 @@ namespace PhysicsReplicationCVars
 		static FAutoConsoleVariableRef CVarResimUseDefaultForDeltaSourceReplication(TEXT("np2.Resim.StatefulDeltaSerialization.UseDefaultForDeltaSourceReplication"), bUseDefaultDeltaForDeltaSourceReplication, TEXT("When false delta sources will use standard serialization when being replicated. When true there will be a valid delta source pointer to default data which can be used for delta serialization when replicating delta sources."));
 		float TimeToSyncStatefulDeltaSource = 5.0f;
 		static FAutoConsoleVariableRef CVarResimTimeToSyncStatefulDeltaSource(TEXT("np2.Resim.StatefulDeltaSerialization.TimeToSyncStatefulDeltaSource"), TimeToSyncStatefulDeltaSource, TEXT("The time in seconds between synchronizing the stateful delta source from server to clients."));
+		
+		bool bApplyPredictiveInterpolationWhenBehindServer = true;
+		static FAutoConsoleVariableRef CVarResimApplyPredictiveInterpolationWhenBehindServer(TEXT("np2.Resim.ApplyPredictiveInterpolationWhenBehindServer"), bApplyPredictiveInterpolationWhenBehindServer, TEXT("When true, switch over to replicating with Predictive Interpolation temporarily, when the client receive target states from the server for frames that have not yet simulated on the client. When false apply the received target via a resimulation when the client has caught up and simulated the corresponding frame."));
 	}
 }
 
@@ -335,8 +338,11 @@ int32 FNetworkPhysicsCallback::TriggerRewindIfNeeded_Internal(int32 LatestStepCo
 
 	if (RewindData)
 	{
-		const int32 TargetStateComparisonFrame = RewindData->CompareTargetsToLastFrame();
-		ResimFrame = (ResimFrame == INDEX_NONE) ? TargetStateComparisonFrame : (TargetStateComparisonFrame == INDEX_NONE) ? ResimFrame : FMath::Min(TargetStateComparisonFrame, ResimFrame);
+		if (!PhysicsReplicationCVars::ResimulationCVars::bApplyPredictiveInterpolationWhenBehindServer)
+		{
+			const int32 TargetStateComparisonFrame = RewindData->CompareTargetsToLastFrame();
+			ResimFrame = (ResimFrame == INDEX_NONE) ? TargetStateComparisonFrame : (TargetStateComparisonFrame == INDEX_NONE) ? ResimFrame : FMath::Min(TargetStateComparisonFrame, ResimFrame);
+		}
 
 		const int32 ReplicationFrame = RewindData->GetResimFrame();
 		ResimFrame = (ResimFrame == INDEX_NONE) ? ReplicationFrame : (ReplicationFrame == INDEX_NONE) ? ResimFrame : FMath::Min(ReplicationFrame, ResimFrame);
diff --git a/Engine/Source/Runtime/Engine/Private/PhysicsEngine/PhysicsReplication.cpp b/Engine/Source/Runtime/Engine/Private/PhysicsEngine/PhysicsReplication.cpp
--- a/Engine/Source/Runtime/Engine/Private/PhysicsEngine/PhysicsReplication.cpp
+++ b/Engine/Source/Runtime/Engine/Private/PhysicsEngine/PhysicsReplication.cpp
@@ -104,6 +104,8 @@ namespace PhysicsReplicationCVars
 
 	namespace ResimulationCVars
 	{
+		extern bool bApplyPredictiveInterpolationWhenBehindServer;
+
 		bool bRuntimeCorrectionEnabled = false;
 		static FAutoConsoleVariableRef CVarResimRuntimeCorrectionEnabled(TEXT("np2.Resim.RuntimeCorrectionEnabled"), bRuntimeCorrectionEnabled, TEXT("Apply positional and rotational runtime corrections while within resim trigger distance."));
 
@@ -1060,6 +1062,7 @@ void FPhysicsReplicationAsync::UpdateAsyncTarget(const FPhysicsRepAsyncInputData
 		Target->PrevLinVel = Input.TargetState.LinVel;
 		Target->RepModeOverride = Input.RepMode;
 	}
+	check(Target);
 
 	/** Target Update Description
 	* @param Input = incoming state target for replication.
@@ -1136,6 +1139,9 @@ void FPhysicsReplicationAsync::UpdateAsyncTarget(const FPhysicsRepAsyncInputData
 		// Apply full Replication LOD on received target
 		ApplyPhysicsReplicationLOD(Input.PhysicsObject, *Target, EPhysicsReplicationLODFlags::LODFlag_All);
 
+		// Check if target is valid to use for resimulation and perform actions if not
+		CheckTargetResimValidity(*Target);
+
 		if (Target->RepMode == EPhysicsReplicationMode::PredictiveInterpolation)
 		{
 #if !(UE_BUILD_SHIPPING || UE_BUILD_TEST)
@@ -1329,6 +1335,48 @@ void FPhysicsReplicationAsync::ApplyTargetStatesAsync(const float DeltaSeconds)
 	}
 }
 
+void FPhysicsReplicationAsync::CheckTargetResimValidity(FReplicatedPhysicsTargetAsync& Target)
+{
+	if (Target.RepMode != EPhysicsReplicationMode::Resimulation)
+	{
+		return;
+	}
+
+	Chaos::FPBDRigidsSolver* RigidsSolver = static_cast<Chaos::FPBDRigidsSolver*>(GetSolver());
+	if (RigidsSolver == nullptr)
+	{
+		return;
+	}
+
+	Chaos::FRewindData* RewindData = RigidsSolver->GetRewindData();
+	if (RewindData == nullptr)
+	{
+		return;
+	}
+
+	const int32 LocalFrame = Target.ServerFrame - Target.FrameOffset;
+	if (!RewindData->IsFrameWithinRewindHistory(LocalFrame))
+	{
+		if (LocalFrame < RewindData->GetEarliestFrame_Internal())
+		{
+			// Client is far ahead of the server, switch over to Predictive Interpolation since it can't use incoming target states from the server to perform resimulations with
+
+			Target.RepMode = EPhysicsReplicationMode::PredictiveInterpolation;
+		}
+		else if (PhysicsReplicationCVars::ResimulationCVars::bApplyPredictiveInterpolationWhenBehindServer)
+		{
+			/** NOTE: If the server is ahead of the client we receive target states for frames we have not yet simulated on the client, target states are stored in FRewindData still though.
+			* If PhysicsReplicationCVars::ResimulationCVars::bApplyPredictiveInterpolationWhenBehindServer is true switch over to using PredictiveInterpolation temporarily.
+			* else FRewindData::CompareTargetsToLastFrame will check for already cached targets to resim with when the server has simulated the corresponding frame */
+
+			Target.RepMode = EPhysicsReplicationMode::PredictiveInterpolation;
+		}
+
+		UE_LOG(LogPhysics, Warning, TEXT("FPhysicsReplication received target frame (%d) out of rewind data bounds (%d, %d) - %s - Target will use EPhysicsReplicationMode: %s"),
+			LocalFrame, RewindData->GetEarliestFrame_Internal(), RewindData->CurrentFrame(), (LocalFrame < RewindData->GetEarliestFrame_Internal()) ? TEXT("Client is far ahead of the server, server might be dropping frames.") : TEXT("Client is behind the server, client might be dropping frames."), *UEnum::GetValueAsString(Target.RepMode));
+	}	
+}
+
 void FPhysicsReplicationAsync::ApplyPhysicsReplicationLOD(Chaos::FConstPhysicsObjectHandle PhysicsObjectHandle, FReplicatedPhysicsTargetAsync& Target, const uint32 LODFLags)
 {
 	Chaos::FPBDRigidsSolver& RigidsSolver = GetSolver()->CastChecked();
@@ -2117,13 +2165,8 @@ bool FPhysicsReplicationAsync::ResimulationReplication(Chaos::FPBDRigidParticleH
 
 	const int32 LocalFrame = Target.ServerFrame - Target.FrameOffset;
 
-	if (LocalFrame >= RewindData->CurrentFrame() || LocalFrame < RewindData->GetEarliestFrame_Internal())
+	if (!RewindData->IsFrameWithinRewindHistory(LocalFrame))
 	{
-		if (LocalFrame > 0 && (RewindData->CurrentFrame() - RewindData->GetEarliestFrame_Internal()) == RewindData->Capacity())
-		{
-			UE_LOG(LogPhysics, Warning, TEXT("FPhysicsReplication::ResimulationReplication target frame (%d) out of rewind data bounds (%d,%d)"),
-				LocalFrame, RewindData->GetEarliestFrame_Internal(), RewindData->CurrentFrame());
-		}
 		return true;
 	}
 
diff --git a/Engine/Source/Runtime/Engine/Public/PhysicsReplication.h b/Engine/Source/Runtime/Engine/Public/PhysicsReplication.h
--- a/Engine/Source/Runtime/Engine/Public/PhysicsReplication.h
+++ b/Engine/Source/Runtime/Engine/Public/PhysicsReplication.h
@@ -230,6 +230,7 @@ private:
 	// Sets SettingsCurrent to either the objects custom settings or to the default settings
 	void FetchObjectSettings(Chaos::FConstPhysicsObjectHandle PhysicsObject); 
 	bool UsePhysicsReplicationLOD();
+	void CheckTargetResimValidity(FReplicatedPhysicsTargetAsync& Target);
 	void ApplyPhysicsReplicationLOD(Chaos::FConstPhysicsObjectHandle PhysicsObjectHandle, FReplicatedPhysicsTargetAsync& Target, const uint32 LODFlags);
 	
 	/** Static function to extrapolate a target for N ticks using X DeltaSeconds */
diff --git a/Engine/Source/Runtime/Experimental/Chaos/Private/RewindData.cpp b/Engine/Source/Runtime/Experimental/Chaos/Private/RewindData.cpp
--- a/Engine/Source/Runtime/Experimental/Chaos/Private/RewindData.cpp
+++ b/Engine/Source/Runtime/Experimental/Chaos/Private/RewindData.cpp
@@ -305,46 +305,43 @@ const int32 FRewindData::CompareTargetsToLastFrame()
 	for (FDirtyParticleInfo& DirtyParticleInfo : DirtyParticles)
 	{
 		// TODO: Only iterate source target states, i.e. states that are not predicted/interpolated to fill in gaps
-		FGeometryParticleHandle* PTParticle = DirtyParticleInfo.GetObjectPtr();
-		if (PTParticle)
+		
+		FGeometryParticleStateBase& History = DirtyParticleInfo.GetHistory();
+		if ((bCompareX || bCompareR) && !History.TargetPositions.IsEmpty())
 		{
-			FGeometryParticleStateBase& History = DirtyParticleInfo.GetHistory();
-			if ((bCompareX || bCompareR) && !History.TargetPositions.IsEmpty())
+			// Compare with particle for this frame and mark resim if needed from CurrentFrame()
+			if (const FParticlePositionRotation* TargetState = History.TargetPositions.Read(FrameAndPhase, PropertiesPool))
 			{
-				// Compare with particle for this frame and mark resim if needed from CurrentFrame()
-				if (const FParticlePositionRotation* TargetState = History.TargetPositions.Read(FrameAndPhase, PropertiesPool))
+				if (const FParticlePositionRotation* PastState = History.ParticlePositionRotation.Read(FrameAndPhase, PropertiesPool))
 				{
-					if (const FParticlePositionRotation* PastState = History.ParticlePositionRotation.Read(FrameAndPhase, PropertiesPool))
+					if (bCompareX)
 					{
-						if (bCompareX)
-						{
-							ShouldTriggerResim |= FRewindData::CheckVectorThreshold(TargetState->GetX(), PastState->GetX(), FPhysicsSolverBase::GetResimulationErrorPositionThreshold()); // TODO: Take per actor settings into consideration via NetworkPhysicsSettingsComponent
-						}
+						ShouldTriggerResim |= FRewindData::CheckVectorThreshold(TargetState->GetX(), PastState->GetX(), FPhysicsSolverBase::GetResimulationErrorPositionThreshold()); // TODO: Take per actor settings into consideration via NetworkPhysicsSettingsComponent
+					}
 
-						if (bCompareR)
-						{
-							ShouldTriggerResim |= FRewindData::CheckQuaternionThreshold(TargetState->GetR(), PastState->GetR(), FPhysicsSolverBase::GetResimulationErrorRotationThreshold()); // TODO: Take per actor settings into consideration via NetworkPhysicsSettingsComponent
-						}
+					if (bCompareR)
+					{
+						ShouldTriggerResim |= FRewindData::CheckQuaternionThreshold(TargetState->GetR(), PastState->GetR(), FPhysicsSolverBase::GetResimulationErrorRotationThreshold()); // TODO: Take per actor settings into consideration via NetworkPhysicsSettingsComponent
 					}
 				}
 			}
+		}
 
-			if (!ShouldTriggerResim && (bCompareV || bCompareW) && !History.TargetVelocities.IsEmpty())
+		if (!ShouldTriggerResim && (bCompareV || bCompareW) && !History.TargetVelocities.IsEmpty())
+		{
+			// Compare with particle for this frame and mark resim if needed from CurrentFrame()
+			if (const FParticleVelocities* TargetState = History.TargetVelocities.Read(FrameAndPhase, PropertiesPool))
 			{
-				// Compare with particle for this frame and mark resim if needed from CurrentFrame()
-				if (const FParticleVelocities* TargetState = History.TargetVelocities.Read(FrameAndPhase, PropertiesPool))
+				if (const FParticleVelocities* PastState = History.Velocities.Read(FrameAndPhase, PropertiesPool))
 				{
-					if (const FParticleVelocities* PastState = History.Velocities.Read(FrameAndPhase, PropertiesPool))
+					if (bCompareV)
 					{
-						if (bCompareV)
-						{
-							ShouldTriggerResim |= FRewindData::CheckVectorThreshold(TargetState->GetV(), PastState->GetV(), FPhysicsSolverBase::GetResimulationErrorLinearVelocityThreshold()); // TODO: Take per actor settings into consideration via NetworkPhysicsSettingsComponent
-						}
+						ShouldTriggerResim |= FRewindData::CheckVectorThreshold(TargetState->GetV(), PastState->GetV(), FPhysicsSolverBase::GetResimulationErrorLinearVelocityThreshold()); // TODO: Take per actor settings into consideration via NetworkPhysicsSettingsComponent
+					}
 
-						if (bCompareW)
-						{
-							ShouldTriggerResim |= FRewindData::CheckVectorThreshold(TargetState->GetW(), PastState->GetW(), FPhysicsSolverBase::GetResimulationErrorAngularVelocityThreshold()); // TODO: Take per actor settings into consideration via NetworkPhysicsSettingsComponent
-						}
+					if (bCompareW)
+					{
+						ShouldTriggerResim |= FRewindData::CheckVectorThreshold(TargetState->GetW(), PastState->GetW(), FPhysicsSolverBase::GetResimulationErrorAngularVelocityThreshold()); // TODO: Take per actor settings into consideration via NetworkPhysicsSettingsComponent
 					}
 				}
 			}
diff --git a/Engine/Source/Runtime/Experimental/Chaos/Public/RewindData.h b/Engine/Source/Runtime/Experimental/Chaos/Public/RewindData.h
--- a/Engine/Source/Runtime/Experimental/Chaos/Public/RewindData.h
+++ b/Engine/Source/Runtime/Experimental/Chaos/Public/RewindData.h
@@ -1773,6 +1773,9 @@ public:
 	const int32 GetResimFrame() const { return ResimFrame; }
 	void SetResimFrame(int32 Frame) { ResimFrame = Frame; }
 
+	/** Check if a frame number is within the current rewind history */
+	bool IsFrameWithinRewindHistory(int32 Frame) { return Frame < CurrentFrame() && Frame >= GetEarliestFrame_Internal(); }
+
 	/** Request a resimulation by setting a requested frame to rewind to
 	* @param Particle can be used to mark a specific particle as the one triggering the resimulation */
 	void CHAOS_API RequestResimulation(int32 RequestedFrame, Chaos::FGeometryParticleHandle* Particle = nullptr);
