Multithreaded End Collision
All cookies information are updated in parallel
Collision End Detect is 75% faster which bring up to 8 to 10 % faster the entire collision detection (BP+NP+EndDetect)

diff --git a/Engine/Source/Runtime/Experimental/Chaos/Private/Chaos/Collision/CollisionConstraintAllocator.cpp b/Engine/Source/Runtime/Experimental/Chaos/Private/Chaos/Collision/CollisionConstraintAllocator.cpp
--- a/Engine/Source/Runtime/Experimental/Chaos/Private/Chaos/Collision/CollisionConstraintAllocator.cpp
+++ b/Engine/Source/Runtime/Experimental/Chaos/Private/Chaos/Collision/CollisionConstraintAllocator.cpp
@@ -1,6 +1,7 @@
 // Copyright Epic Games, Inc. All Rights Reserved.
 #include "Chaos/Collision/CollisionConstraintAllocator.h"
 #include "Chaos/PBDCollisionConstraints.h"
+#include "Tasks/Task.h"
 
 namespace Chaos
 {
@@ -77,7 +78,7 @@ namespace Chaos
 
 				// NOTE: It may look weird that we are passing the particle's self to the AddMidPhase method. 
 				// The midphase represents two particles and it needs to know which one we're currently adding it to.
-				// We could make this API nicer but providing a helper on the Particle, but it's a bit awkard to do so
+				// We could make this API nicer but providing a helper on the Particle, but it's a bit awkward to do so
 				// and since AddMidPahse is only called from here we'll leave it like this for now.
 				Particle0->ParticleCollisions().AddMidPhase(Particle0, MidPhase);
 				Particle1->ParticleCollisions().AddMidPhase(Particle1, MidPhase);
@@ -90,6 +91,79 @@ namespace Chaos
 			}
 		}
 
+		void FCollisionConstraintAllocator::ProcessNewNonCcdConstraints(const TArray<FPBDCollisionConstraint*>& NewConstraints, TArray<TArray<FPBDCollisionConstraint*>>& CCDConstraintsBatch)
+		{
+			int32 NumTasks = !!GSingleThreadedPhysics ? 1 : FMath::Min(FTaskGraphInterface::Get().GetNumWorkerThreads(), Chaos::MaxNumWorkers);
+			NumTasks = FMath::Max(FMath::Min(NumTasks, NewConstraints.Num()), 1);
+
+			const int32 NumConstraints = NewConstraints.Num();
+			const int32 ConstraintByTask = FMath::Max(FMath::DivideAndRoundUp(NumConstraints, NumTasks), 1);
+			const int32 NumBatches = FMath::DivideAndRoundUp(NumConstraints, ConstraintByTask);
+			TArray<UE::Tasks::FTask> PendingTasks;
+			PendingTasks.Reserve(NumBatches);
+
+			CCDConstraintsBatch.SetNum(NumBatches);
+
+			for (int32 BatchIndex = 0; BatchIndex < NumBatches; BatchIndex++)
+			{
+				const int32 StartIndex = BatchIndex * ConstraintByTask;
+				int32 EndIndex = (BatchIndex + 1) * ConstraintByTask;
+				EndIndex = BatchIndex == NumBatches - 1 ? FMath::Min(NumConstraints, EndIndex) : EndIndex;
+
+				TArray<FPBDCollisionConstraint*>& CCDConstraints = CCDConstraintsBatch[BatchIndex];
+
+				UE::Tasks::FTask PendingTask = UE::Tasks::Launch(UE_SOURCE_LOCATION, [this, StartIndex, EndIndex, &NewConstraints, &CCDConstraints]
+				{
+					QUICK_SCOPE_CYCLE_COUNTER(STAT_Collisions_SubProcessNewConstraints);
+					for (int32 Index = StartIndex; Index < EndIndex; Index++)
+					{
+						FPBDCollisionConstraint* Constraint = NewConstraints[Index];
+						SetConstraintCookie(Constraint, Index);
+						// If the constraint uses CCD, keep it in another list so we don't have to search the full list for CCD contacts
+						if (Constraint->GetCCDEnabled())
+						{
+							CCDConstraints.Add(Constraint);
+						}
+					}
+				}, UE::Tasks::ETaskPriority::High);
+				PendingTasks.Add(PendingTask);
+			}
+			UE::Tasks::Wait(PendingTasks);
+		}
+
+		void FCollisionConstraintAllocator::ProcessNewCcdConstraints()
+		{
+			int32 NumTasks = !!GSingleThreadedPhysics ? 1 : FMath::Min(FTaskGraphInterface::Get().GetNumWorkerThreads(), Chaos::MaxNumWorkers);
+			NumTasks = FMath::Max(FMath::Min(NumTasks, ActiveCCDConstraints.Num()), 1);
+
+			const int32 NumConstraints = ActiveCCDConstraints.Num();
+			const int32 ConstraintByTask = FMath::Max(FMath::DivideAndRoundUp(NumConstraints, NumTasks), 1);
+			const int32 NumBatches = FMath::DivideAndRoundUp(NumConstraints, ConstraintByTask);
+			TArray<UE::Tasks::FTask> PendingTasks;
+			PendingTasks.Reserve(NumBatches);
+
+			for (int32 BatchIndex = 0; BatchIndex < NumBatches; BatchIndex++)
+			{
+				const int32 StartIndex = BatchIndex * ConstraintByTask;
+				int32 EndIndex = (BatchIndex + 1) * ConstraintByTask;
+				EndIndex = BatchIndex == NumBatches - 1 ? FMath::Min(NumConstraints, EndIndex) : EndIndex;
+
+				UE::Tasks::FTask PendingTask = UE::Tasks::Launch(UE_SOURCE_LOCATION, [this, StartIndex, EndIndex]
+					{
+						QUICK_SCOPE_CYCLE_COUNTER(STAT_Collisions_SubCCDProcessNewConstraints);
+						for (int32 Index = StartIndex; Index < EndIndex; Index++)
+						{
+							FPBDCollisionConstraint* Constraint = ActiveCCDConstraints[Index];
+							Constraint->GetContainerCookie().CCDConstraintIndex = Index;
+
+						}
+					}, UE::Tasks::ETaskPriority::High);
+				PendingTasks.Add(PendingTask);
+			}
+			ContextAllocators[0]->NewActiveConstraints.Reset();
+			UE::Tasks::Wait(PendingTasks);
+		}
+
 		void FCollisionConstraintAllocator::ProcessNewConstraints()
 		{
 			QUICK_SCOPE_CYCLE_COUNTER(STAT_Collisions_ProcessNewConstraints);
@@ -98,15 +172,31 @@ namespace Chaos
 			if (!ContextAllocators.IsEmpty())
 			{
 				const TArray<FPBDCollisionConstraint*>& NewConstraints = ContextAllocators[0]->NewActiveConstraints;
-				for (FPBDCollisionConstraint* Constraint : NewConstraints)
+				ActiveConstraints.Append(NewConstraints);
+
+				TArray<TArray<FPBDCollisionConstraint*>> CCDConstraintsBatch;
+				ProcessNewNonCcdConstraints(NewConstraints, CCDConstraintsBatch);
+
+				// Single Threaded and protect the insertion order for determinism 
+				for (int32 BatchIndex = 0; BatchIndex < CCDConstraintsBatch.Num(); BatchIndex++)
 				{
-					// Add the constraint to the active list (and the CCD list if appropriate)
-					AddActiveConstraintImpl(Constraint);
+					ActiveCCDConstraints.Append(CCDConstraintsBatch[BatchIndex]);
 				}
-				ContextAllocators[0]->NewActiveConstraints.Reset();
+
+				ProcessNewCcdConstraints();
 			}
 		}
 
+		void FCollisionConstraintAllocator::ProcessNewItems()
+		{
+			UE::Tasks::FTask PendingTask = UE::Tasks::Launch(UE_SOURCE_LOCATION, [this] 
+			{
+				ProcessNewMidPhases();
+			}, UE::Tasks::ETaskPriority::High);
+
+			ProcessNewConstraints();
+			UE::Tasks::Wait(PendingTask);
+		}
 
 		void FCollisionConstraintAllocator::EndDetectCollisions()
 		{
@@ -193,6 +283,17 @@ namespace Chaos
 			}
 		}
 
+		void FCollisionConstraintAllocator::SetConstraintCookie(FPBDCollisionConstraint* Constraint, int32 Index)
+		{
+			FPBDCollisionConstraintContainerCookie& Cookie = Constraint->GetContainerCookie();
+
+			Cookie.LastUsedEpoch = CurrentEpoch;
+			Cookie.CCDConstraintIndex = INDEX_NONE;
+
+			// Add to the active list and remember the index for O(1) removal
+			Cookie.ConstraintIndex = Index;
+		}
+
 		void FCollisionConstraintAllocator::RemoveActiveConstraint(FPBDCollisionConstraint& Constraint)
 		{
 			FPBDCollisionConstraintContainerCookie& Cookie = Constraint.GetContainerCookie();
diff --git a/Engine/Source/Runtime/Experimental/Chaos/Public/Chaos/Collision/CollisionConstraintAllocator.h b/Engine/Source/Runtime/Experimental/Chaos/Public/Chaos/Collision/CollisionConstraintAllocator.h
--- a/Engine/Source/Runtime/Experimental/Chaos/Public/Chaos/Collision/CollisionConstraintAllocator.h
+++ b/Engine/Source/Runtime/Experimental/Chaos/Public/Chaos/Collision/CollisionConstraintAllocator.h
@@ -529,12 +529,7 @@ namespace Chaos
 			friend class FCollisionContextAllocator;
 
 			// Gather all the newly created items and finalize their setup
-			void ProcessNewItems()
-			{
-				ProcessNewMidPhases();
-
-				ProcessNewConstraints();
-			}
+			void ProcessNewItems();
 
 			void AddMidPhase(FParticlePairMidPhasePtr&& MidPhase)
 			{
@@ -562,8 +557,11 @@ namespace Chaos
 
 			// Collect all the constraints activated in the collision tasks and register them (calls ActivateConstraintImp on each)
 			CHAOS_API void ProcessNewConstraints();
+			CHAOS_API void ProcessNewNonCcdConstraints(const TArray<FPBDCollisionConstraint*>& NewConstraints, TArray<TArray<FPBDCollisionConstraint*>>& CCDConstraintsBatch);
+			CHAOS_API void ProcessNewCcdConstraints();
 
 			CHAOS_API void AddActiveConstraintImpl(FPBDCollisionConstraint* Constraint);
+			CHAOS_API void SetConstraintCookie(FPBDCollisionConstraint* Constraint, int32 Index);
 			CHAOS_API void RemoveActiveConstraint(FPBDCollisionConstraint& Constraint);
 
 			CHAOS_API void SortActiveConstraints();
