Expose some of the destruction clustering parameters to the project settings

diff --git a/Engine/Source/Runtime/Experimental/Chaos/Private/Chaos/PBDRigidClustering.cpp b/Engine/Source/Runtime/Experimental/Chaos/Private/Chaos/PBDRigidClustering.cpp
--- a/Engine/Source/Runtime/Experimental/Chaos/Private/Chaos/PBDRigidClustering.cpp
+++ b/Engine/Source/Runtime/Experimental/Chaos/Private/Chaos/PBDRigidClustering.cpp
@@ -134,21 +134,17 @@ namespace Chaos
 			return Current;
 		}
 
-		bool CVarShouldThrottleParticleRelease()
-		{
-			return (ClusteringParticleReleaseThrottlingMinCount >= 0 && ClusteringParticleReleaseThrottlingMaxCount >= 0);
-		}
-
 		// compute a ratio (between 0 and 1) of released particles to release
 		float GetRatioOfReleasedParticlesToDisable(const FRigidClustering::FRigidEvolution& Evolution)
 		{
+			const FChaosSolverDestructionSettings& DestructionSettings = Evolution.GetRigidClustering().GetDestructionSettings();
 			const FPBDRigidsSOAs& ParticleStructures = Evolution.GetParticles();
 			int32 NumActiveParticles = 0;
 			NumActiveParticles += ParticleStructures.GetSleepingGeometryCollectionArray().Num();
 			NumActiveParticles += ParticleStructures.GetDynamicGeometryCollectionArray().Num();
 
-			const int32 Range = FMath::Max(0, (ClusteringParticleReleaseThrottlingMaxCount - ClusteringParticleReleaseThrottlingMinCount));
-			const int32 OverMinCount = FMath::Max(0, (NumActiveParticles - ClusteringParticleReleaseThrottlingMinCount));
+			const int32 Range = FMath::Max(0, (DestructionSettings.ClusteringParticleReleaseThrottlingMaxCount - DestructionSettings.ClusteringParticleReleaseThrottlingMinCount));
+			const int32 OverMinCount = FMath::Max(0, (NumActiveParticles - DestructionSettings.ClusteringParticleReleaseThrottlingMinCount));
 
 			if (Range > 0)
 			{
@@ -162,7 +158,7 @@ namespace Chaos
 		template<typename TParticleContainer>
 		void GenericThrottleReleasedParticlesIfNecessary(TParticleContainer& Container, FRigidClustering& Clustering)
 		{
-			if (!CVarShouldThrottleParticleRelease())
+			if (!Clustering.ShouldThrottleParticleRelease())
 			{
 				return;
 			}
@@ -1652,7 +1648,7 @@ namespace Chaos
 
 	static int32 GPerAdvanceBreaksAllowed = TNumericLimits<int32>::Max();
 	FAutoConsoleVariableRef CVarPerAdvanceBreaksAllowed(TEXT("p.Chaos.Clustering.PerAdvanceBreaksAllowed"), GPerAdvanceBreaksAllowed,
-		TEXT("Number of breaks allowed to occur for each invokation of AdvanceClustering"));
+		TEXT("Number of breaks allowed to occur for each invocation of AdvanceClustering"));
 
 	static int32 GPerAdvanceBreaksRescheduleLimit = TNumericLimits<int32>::Max();
 	FAutoConsoleVariableRef CVarPerAdvanceBreaksRescheduleLimit(TEXT("p.Chaos.Clustering.PerAdvanceBreaksRescheduleLimit"), GPerAdvanceBreaksRescheduleLimit,
@@ -1849,7 +1845,7 @@ namespace Chaos
 						if (BreakingModel({ &ParticleToProcess , 1 }))
 						{
 							NumBreaks++;
-							if (NumBreaks >= GPerAdvanceBreaksAllowed)
+							if (NumBreaks >= DestructionSettings.PerAdvanceBreaksAllowed)
 							{
 								LastProcessedIndex = Index;
 								break;
@@ -1858,10 +1854,10 @@ namespace Chaos
 					}
 
 					// Add back the rest of the particles to process back in the strained array for later processing
-					if (GPerAdvanceBreaksRescheduleLimit > 0)
+					if (DestructionSettings.PerAdvanceBreaksRescheduleLimit > 0)
 					{
 						// Add back the non processed parent clusters for next tick
-						for (int32 Index = 0; Index < GPerAdvanceBreaksRescheduleLimit; Index++)
+						for (int32 Index = 0; Index < DestructionSettings.PerAdvanceBreaksRescheduleLimit; Index++)
 						{
 							const int32 ParticleToProcessIndex = Index + LastProcessedIndex + 1;
 							if (ParticleToProcessIndex >= ParticlesToProcess.Num())
@@ -3090,7 +3086,7 @@ namespace Chaos
 
 	bool FRigidClustering::ShouldThrottleParticleRelease() const
 	{
-		return CVarShouldThrottleParticleRelease();
+		return (DestructionSettings.ClusteringParticleReleaseThrottlingMinCount >= 0 && DestructionSettings.ClusteringParticleReleaseThrottlingMaxCount >= 0);
 	}
 
 	void FRigidClustering::ThrottleReleasedParticlesIfNecessary(TSet<FPBDRigidParticleHandle*>& Particles)
@@ -3103,4 +3099,34 @@ namespace Chaos
 		GenericThrottleReleasedParticlesIfNecessary(Particles, *this);
 	}
 
+	FClusterCreationParameters::EConnectionMethod ToInternalConnectionMethod(EClusterUnionMethod InMethod)
+	{
+		using ETargetEnum = FClusterCreationParameters::EConnectionMethod;
+		switch (InMethod)
+		{
+		case EClusterUnionMethod::PointImplicit:
+			return ETargetEnum::PointImplicit;
+		case EClusterUnionMethod::DelaunayTriangulation:
+			return ETargetEnum::DelaunayTriangulation;
+		case EClusterUnionMethod::MinimalSpanningSubsetDelaunayTriangulation:
+			return ETargetEnum::MinimalSpanningSubsetDelaunayTriangulation;
+		case EClusterUnionMethod::PointImplicitAugmentedWithMinimalDelaunay:
+			return ETargetEnum::PointImplicitAugmentedWithMinimalDelaunay;
+		case EClusterUnionMethod::BoundsOverlapFilteredDelaunayTriangulation:
+			return ETargetEnum::BoundsOverlapFilteredDelaunayTriangulation;
+		}
+
+		return ETargetEnum::None;
+	}
+
+	void FRigidClustering::ApplySettings(const FChaosSolverConfiguration& Settings)
+	{
+		SetClusterConnectionFactor(Settings.ClusterConnectionFactor);
+		SetClusterUnionConnectionType(ToInternalConnectionMethod(Settings.ClusterUnionConnectionType));
+		DestructionSettings.PerAdvanceBreaksAllowed = FMath::Min(Settings.DestructionSettings.PerAdvanceBreaksAllowed, GPerAdvanceBreaksAllowed);
+		DestructionSettings.PerAdvanceBreaksRescheduleLimit = FMath::Min(Settings.DestructionSettings.PerAdvanceBreaksRescheduleLimit, GPerAdvanceBreaksRescheduleLimit);
+		DestructionSettings.ClusteringParticleReleaseThrottlingMinCount = FMath::Max(Settings.DestructionSettings.ClusteringParticleReleaseThrottlingMinCount, ClusteringParticleReleaseThrottlingMinCount);
+		DestructionSettings.ClusteringParticleReleaseThrottlingMaxCount = FMath::Max(Settings.DestructionSettings.ClusteringParticleReleaseThrottlingMaxCount, ClusteringParticleReleaseThrottlingMaxCount);
+	}
+
 } // namespace Chaos
diff --git a/Engine/Source/Runtime/Experimental/Chaos/Private/PBDRigidsSolver.cpp b/Engine/Source/Runtime/Experimental/Chaos/Private/PBDRigidsSolver.cpp
--- a/Engine/Source/Runtime/Experimental/Chaos/Private/PBDRigidsSolver.cpp
+++ b/Engine/Source/Runtime/Experimental/Chaos/Private/PBDRigidsSolver.cpp
@@ -2896,31 +2896,9 @@ TRACE_COUNTER_SET(ChaosTraceCounter_##Name, Value)
 		GetEvolution()->SetParticleSleepType(Rigid, DynamicMisc.SleepType());
 	}
 
-
-	FClusterCreationParameters::EConnectionMethod ToInternalConnectionMethod(EClusterUnionMethod InMethod)
-	{
-		using ETargetEnum = FClusterCreationParameters::EConnectionMethod;
-		switch(InMethod)
-		{
-		case EClusterUnionMethod::PointImplicit:
-			return ETargetEnum::PointImplicit;
-		case EClusterUnionMethod::DelaunayTriangulation:
-			return ETargetEnum::DelaunayTriangulation;
-		case EClusterUnionMethod::MinimalSpanningSubsetDelaunayTriangulation:
-			return ETargetEnum::MinimalSpanningSubsetDelaunayTriangulation;
-		case EClusterUnionMethod::PointImplicitAugmentedWithMinimalDelaunay:
-			return ETargetEnum::PointImplicitAugmentedWithMinimalDelaunay;
-		case EClusterUnionMethod::BoundsOverlapFilteredDelaunayTriangulation:
-			return ETargetEnum::BoundsOverlapFilteredDelaunayTriangulation;
-		}
-
-		return ETargetEnum::None;
-	}
-
 	void FPBDRigidsSolver::ApplyConfig(const FChaosSolverConfiguration& InConfig)
 	{
-		GetEvolution()->GetRigidClustering().SetClusterConnectionFactor(InConfig.ClusterConnectionFactor);
-		GetEvolution()->GetRigidClustering().SetClusterUnionConnectionType(ToInternalConnectionMethod(InConfig.ClusterUnionConnectionType));
+		GetEvolution()->GetRigidClustering().ApplySettings(InConfig);
 		SetPositionIterations(InConfig.PositionIterations);
 		SetVelocityIterations(InConfig.VelocityIterations);
 		SetProjectionIterations(InConfig.ProjectionIterations);
diff --git a/Engine/Source/Runtime/Experimental/Chaos/Public/Chaos/PBDRigidClustering.h b/Engine/Source/Runtime/Experimental/Chaos/Public/Chaos/PBDRigidClustering.h
--- a/Engine/Source/Runtime/Experimental/Chaos/Public/Chaos/PBDRigidClustering.h
+++ b/Engine/Source/Runtime/Experimental/Chaos/Public/Chaos/PBDRigidClustering.h
@@ -2,6 +2,7 @@
 #pragma once
 
 #include "Chaos/ClusterUnionManager.h"
+#include "ChaosSolverConfiguration.h"
 #include "Chaos/PBDRigidClusteredParticles.h"
 #include "Chaos/PBDCollisionConstraints.h"
 #include "Chaos/Transform.h"
@@ -375,6 +376,8 @@ public:
 	void SetClusterUnionConnectionType(FClusterCreationParameters::EConnectionMethod ClusterConnectionType) { MClusterUnionConnectionType = ClusterConnectionType; }
 	FClusterCreationParameters::EConnectionMethod GetClusterUnionConnectionType() const { return MClusterUnionConnectionType; }
 
+	void ApplySettings(const FChaosSolverConfiguration& Settings);
+
 	CHAOS_API void GenerateConnectionGraph(
 		const TArray<FPBDRigidParticleHandle*>& Particles,
 		const FClusterCreationParameters& Parameters = FClusterCreationParameters(),
@@ -499,6 +502,8 @@ public:
 	/** disable a particle and make sure the related internal structure are up to date */
 	void DisableParticle(FPBDRigidParticleHandle* ParticleToDisable, bool bRemoveFromChildrenMap = true);
 
+	const FChaosSolverDestructionSettings& GetDestructionSettings() const { return DestructionSettings; }
+
  protected:
 
 	CHAOS_API void ComputeStrainFromCollision(const FPBDCollisionConstraints& CollisionRule, const FReal Dt);
@@ -632,6 +637,8 @@ private:
 
 	// Sim callback objects which implement cluster modification steps
 	const TArray<ISimCallbackObject*>* StrainModifiers;
+
+	FChaosSolverDestructionSettings DestructionSettings;
 };
 
 } // namespace Chaos
diff --git a/Engine/Source/Runtime/Experimental/Chaos/Public/ChaosSolverConfiguration.h b/Engine/Source/Runtime/Experimental/Chaos/Public/ChaosSolverConfiguration.h
--- a/Engine/Source/Runtime/Experimental/Chaos/Public/ChaosSolverConfiguration.h
+++ b/Engine/Source/Runtime/Experimental/Chaos/Public/ChaosSolverConfiguration.h
@@ -17,6 +17,28 @@ enum class EClusterUnionMethod : uint8
 	None
 };
 
+USTRUCT()
+struct FChaosSolverDestructionSettings
+{
+	GENERATED_BODY();
+
+	/** Number of breaks allowed to occur for each invocation of AdvanceClustering */
+	UPROPERTY(EditAnywhere, Category = "SolverConfiguration|Destruction", meta = (ClampMin = "0"))
+	int32 PerAdvanceBreaksAllowed = TNumericLimits<int32>::Max();
+
+	/** Number of breaks allowed to be rescheduled for next frame if any */
+	UPROPERTY(EditAnywhere, Category = "SolverConfiguration|Destruction", meta = (ClampMin = "0"))
+	int32 PerAdvanceBreaksRescheduleLimit = TNumericLimits<int32>::Max();
+
+	/** Minimum number of active geometry collection to reach before clustering start to disable a percentage of the released particle per cluster */
+	UPROPERTY(EditAnywhere, Category = "SolverConfiguration|Destruction", meta = (ClampMin = "-1"))
+	int32 ClusteringParticleReleaseThrottlingMinCount = INDEX_NONE;
+
+	/** Maximum number of active geometry collection to reach before all released clustering disable all released particle instantly */
+	UPROPERTY(EditAnywhere, Category = "SolverConfiguration|Destruction", meta = (ClampMin = "-1"))
+	int32 ClusteringParticleReleaseThrottlingMaxCount = INDEX_NONE;
+};
+
 USTRUCT()
 struct FChaosSolverConfiguration
 {
@@ -78,6 +100,9 @@ struct FChaosSolverConfiguration
 	UPROPERTY(EditAnywhere, Category = "SolverConfiguration|Clustering")
 	EClusterUnionMethod ClusterUnionConnectionType;
 
+	UPROPERTY(EditAnywhere, Category = "SolverConfiguration|Destruction")
+	FChaosSolverDestructionSettings DestructionSettings;
+
 	UPROPERTY(EditAnywhere, Category = "SolverConfiguration|GeneratedData")
 	bool bGenerateCollisionData;
 
