[ChaosVD] Adding new CVD stage scopes for Network Rewind & Resim

Currently some of the particle updates in these stages were being recorded, but CVD auto-labeled them as "In-Between" stage as this happens after the Evolution End stage and before the next the Evolution Start stage.

This CL creates a proper stage scope in them so CVD can label them correctly. It also ensures we record the updated particle data for the all particles that are modified in these stages.

diff --git a/Engine/Source/Runtime/Experimental/Chaos/Private/RewindData.cpp b/Engine/Source/Runtime/Experimental/Chaos/Private/RewindData.cpp
--- a/Engine/Source/Runtime/Experimental/Chaos/Private/RewindData.cpp
+++ b/Engine/Source/Runtime/Experimental/Chaos/Private/RewindData.cpp
@@ -197,6 +197,8 @@ void FRewindData::RewindStates(const int32 RewindFrame, const bool bResetSolver)
 
 void FRewindData::ApplyTargets(const int32 Frame, const bool bResetSimulation)
 {
+	CVD_SCOPE_TRACE_SOLVER_STEP(CVDDC_Default, TEXT("Rewind Apply Targets"));
+
 	PRAGMA_DISABLE_DEPRECATION_WARNINGS
 	RewindStates(Frame, bResetSimulation);
 	PRAGMA_ENABLE_DEPRECATION_WARNINGS
@@ -272,6 +274,9 @@ void FRewindData::ApplyTargets(const int32 Frame, const bool bResetSimulation)
 				break;
 			}
 		});
+
+		CVD_TRACE_PARTICLE(PTParticle);
+
 #if !(UE_BUILD_SHIPPING || UE_BUILD_TEST)
 		if (!History.TargetPositions.IsClean(RewindFrameAndPhase) && Chaos::FPhysicsSolverBase::CanDebugNetworkPhysicsPrediction())
 		{
@@ -385,6 +390,8 @@ bool FRewindData::RewindToFrame(int32 Frame)
 {
 	QUICK_SCOPE_CYCLE_COUNTER(RewindToFrame);
 
+	CVD_SCOPE_TRACE_SOLVER_STEP(CVDDC_Default, TEXT("Rewind To Frame"));
+
 	EnsureIsInPhysicsThreadContext();
 	//Can't go too far back
 	const int32 EarliestFrame = GetEarliestFrame_Internal();
@@ -475,6 +482,11 @@ bool FRewindData::RewindToFrame(int32 Frame)
 		// Set this after rewinding, since ResimType gets overwritten if NonFrequentData is cached
 		PTParticle->SetResimType(bResimAsFollower ? EResimType::ResimAsFollower : EResimType::FullResim);
 
+		if (bAnyChange)
+		{
+			CVD_TRACE_PARTICLE(PTParticle);
+		}
+
 		if (!bResimAsFollower)
 		{
 			if (bAnyChange)
