Associate an estimated/expected size of a request's data.
#rnx
#ushell-unshelve of 41195365

diff --git a/Engine/Source/Runtime/Experimental/IoStore/HttpClient/Internal/IO/Http/Client.h b/Engine/Source/Runtime/Experimental/IoStore/HttpClient/Internal/IO/Http/Client.h
--- a/Engine/Source/Runtime/Experimental/IoStore/HttpClient/Internal/IO/Http/Client.h
+++ b/Engine/Source/Runtime/Experimental/IoStore/HttpClient/Internal/IO/Http/Client.h
@@ -223,7 +223,8 @@ public:
 	struct FRequestParams
 	{
 		FCertRootsRef		VerifyCert		= {};
-		uint32				BufferSize		= 256;
+		uint16				BufferSize		= 256;
+		uint16				ContentEstKiB	= 0;
 		bool				bAutoRedirect	= false;
 		bool				bAllowChunked	= true;
 	};
diff --git a/Engine/Source/Runtime/Experimental/IoStore/HttpClient/Private/Activity.inl b/Engine/Source/Runtime/Experimental/IoStore/HttpClient/Private/Activity.inl
--- a/Engine/Source/Runtime/Experimental/IoStore/HttpClient/Private/Activity.inl
+++ b/Engine/Source/Runtime/Experimental/IoStore/HttpClient/Private/Activity.inl
@@ -81,7 +81,8 @@ struct alignas(16) FActivity
 	uint8				NoContent : 1;
 	uint8				bFollow30x : 1;
 	uint8				bAllowChunked : 1;
-	uint8				_Unused0 : 4;
+	uint8				LengthScore : 3;
+	uint8				_Unused : 1;
 	uint32				StateParam = 0;
 #if IAS_HTTP_WITH_PERF
 	FStopwatch			Stopwatch;
@@ -195,6 +196,21 @@ static void Activity_SetError(FActivity* Activity, const FOutcome& Outcome)
 	Activity_SetError(Activity, Outcome.GetMessage().GetData(), Outcome.GetErrorCode());
 }
 
+////////////////////////////////////////////////////////////////////////////////
+static void Activity_SetScore(FActivity* Activity, uint32 ContentEstKiB)
+{
+	if (ContentEstKiB == 0)
+	{
+		Activity->LengthScore = 0;
+		return;
+	}
+
+	ContentEstKiB |= 2;
+	uint32 Pow2 = FMath::FloorLog2(uint32(ContentEstKiB));
+	Pow2 = FMath::Min(Pow2, 7u);
+	Activity->LengthScore = uint8(Pow2);
+}
+
 // }}}
 
 } // namespace UE::IoStore::HTTP
diff --git a/Engine/Source/Runtime/Experimental/IoStore/HttpClient/Private/Loop.inl b/Engine/Source/Runtime/Experimental/IoStore/HttpClient/Private/Loop.inl
--- a/Engine/Source/Runtime/Experimental/IoStore/HttpClient/Private/Loop.inl
+++ b/Engine/Source/Runtime/Experimental/IoStore/HttpClient/Private/Loop.inl
@@ -1559,6 +1559,8 @@ FRequest FEventLoop::Request(
 
 	Buffer.Fix();
 
+	Activity_SetScore(Activity, Params->ContentEstKiB);
+
 	memcpy(HostNamePtr, HostName.GetData(), HostNameLength);
 	HostNamePtr[HostNameLength] = '\0';
 	new (Host) FHost({
@@ -1590,6 +1592,9 @@ FRequest FEventLoop::Request(
 	Activity->IsKeepAlive = 1;
 	Activity->bFollow30x = (Params->bAutoRedirect == true);
 	Activity->bAllowChunked = (Params->bAllowChunked == true);
+	Activity->LengthScore = 0;
+
+	Activity_SetScore(Activity, Params->ContentEstKiB);
 
 	return Impl->Request(Method, Path, Activity);
 }
