Route the calling of an activity's sink through a function.
#rnx
#ushell-unshelve of 41195365

diff --git a/Engine/Source/Runtime/Experimental/IoStore/HttpClient/Private/Activity.inl b/Engine/Source/Runtime/Experimental/IoStore/HttpClient/Private/Activity.inl
--- a/Engine/Source/Runtime/Experimental/IoStore/HttpClient/Private/Activity.inl
+++ b/Engine/Source/Runtime/Experimental/IoStore/HttpClient/Private/Activity.inl
@@ -152,6 +152,13 @@ static uint32 Activity_RemainingKiB(FActivity* Activity)
 	return (ContentLength - Activity->StateParam) >> 10;
 }
 
+////////////////////////////////////////////////////////////////////////////////
+static void Activity_CallSink(FActivity* Activity)
+{
+	FTicketStatus& SinkArg = *(FTicketStatus*)Activity;
+	Activity->Sink(SinkArg);
+}
+
 ////////////////////////////////////////////////////////////////////////////////
 static FActivity* Activity_Alloc(uint32 BufferSize)
 {
diff --git a/Engine/Source/Runtime/Experimental/IoStore/HttpClient/Private/Loop.inl b/Engine/Source/Runtime/Experimental/IoStore/HttpClient/Private/Loop.inl
--- a/Engine/Source/Runtime/Experimental/IoStore/HttpClient/Private/Loop.inl
+++ b/Engine/Source/Runtime/Experimental/IoStore/HttpClient/Private/Loop.inl
@@ -229,10 +229,7 @@ static FOutcome DoRecvMessage(FActivity* Activity, FHttpPeer& Peer)
 	FIoBuffer* PriorDest = Activity->Dest; // to retain unioned Host ptr (redirect uses it in sink)
 	Internal.Code = -1;
 	Internal.ContentLength = ContentLength;
-	{
-		FTicketStatus& SinkArg = *(FTicketStatus*)Activity;
-		Activity->Sink(SinkArg);
-	}
+	Activity_CallSink(Activity);
 	Activity->NoContent |= (ContentLength == 0);
 
 	// Check the user gave us a destination for content
@@ -368,8 +365,7 @@ static FOutcome DoRecvStream(FActivity* Activity, FHttpPeer& Peer, int32& MaxRec
 		FIoBuffer Slice(View, Dest);
 		Swap(Dest, Slice);
 		
-		FTicketStatus& SinkArg = *(FTicketStatus*)Activity;
-		Activity->Sink(SinkArg);
+		Activity_CallSink(Activity);
 
 		Swap(Dest, Slice);
 	};
@@ -561,8 +557,7 @@ static FOutcome DoRecv(FActivity* Activity, FHttpPeer& Peer, int32& MaxRecvSize)
 static void DoRecvDone(FActivity* Activity)
 {
 	// Notify the user we've received everything
-	FTicketStatus& SinkArg = *(FTicketStatus*)Activity;
-	Activity->Sink(SinkArg);
+	Activity_CallSink(Activity);
 
 	Activity_ChangeState(Activity, FActivity::EState::Completed);
 }
@@ -577,8 +572,7 @@ static void DoCancel(FActivity* Activity)
 
 	Activity_ChangeState(Activity, FActivity::EState::Cancelled);
 
-	FTicketStatus& SinkArg = *(FTicketStatus*)Activity;
-	Activity->Sink(SinkArg);
+	Activity_CallSink(Activity);
 }
 
 ////////////////////////////////////////////////////////////////////////////////
@@ -587,8 +581,7 @@ static void DoFail(FActivity* Activity)
 	check(Activity->State == FActivity::EState::Failed);
 
 	// Notify the user we've received everything
-	FTicketStatus& SinkArg = *(FTicketStatus*)Activity;
-	Activity->Sink(SinkArg);
+	Activity_CallSink(Activity);
 }
 
 
