Iris
* Fix UPackage missing include

#rnx

diff --git a/Engine/Source/Runtime/Experimental/Iris/Core/Private/Iris/ReplicationSystem/Filtering/ReplicationFiltering.cpp b/Engine/Source/Runtime/Experimental/Iris/Core/Private/Iris/ReplicationSystem/Filtering/ReplicationFiltering.cpp
--- a/Engine/Source/Runtime/Experimental/Iris/Core/Private/Iris/ReplicationSystem/Filtering/ReplicationFiltering.cpp
+++ b/Engine/Source/Runtime/Experimental/Iris/Core/Private/Iris/ReplicationSystem/Filtering/ReplicationFiltering.cpp
@@ -3,6 +3,7 @@
 #include "ReplicationFiltering.h"
 #include "HAL/PlatformMath.h"
 #include "HAL/IConsoleManager.h"
+#include "UObject/Package.h"
 
 #include "Iris/IrisConfigInternal.h"
 #include "Iris/IrisConstants.h"
@@ -2671,14 +2672,14 @@ void FReplicationFiltering::InitFilters()
 
 		
 		FFilterInfo& Info = DynamicFilterInfos.Emplace_GetRef();
-		Info.Filter = TStrongObjectPtr<UNetObjectFilter>(NewObject<UNetObjectFilter>((UObject*)GetTransientPackage(), NetObjectFilterClass, MakeUniqueObjectName(nullptr, NetObjectFilterClass, FilterDefinition.FilterName)));
+		Info.Filter = TStrongObjectPtr<UNetObjectFilter>(NewObject<UNetObjectFilter>(GetTransientPackage(), NetObjectFilterClass, MakeUniqueObjectName(nullptr, NetObjectFilterClass, FilterDefinition.FilterName)));
 		check(Info.Filter.IsValid());
 		Info.Name = FilterDefinition.FilterName;
 		Info.ObjectCount = 0;
 
 		FNetObjectFilterInitParams InitParams;
 		InitParams.ReplicationSystem = ReplicationSystem;
-		InitParams.Config = (NetObjectFilterConfigClass ? NewObject<UNetObjectFilterConfig>((UObject*)GetTransientPackage(), NetObjectFilterConfigClass) : nullptr);
+		InitParams.Config = (NetObjectFilterConfigClass ? NewObject<UNetObjectFilterConfig>(GetTransientPackage(), NetObjectFilterConfigClass) : nullptr);
 		InitParams.AbsoluteMaxNetObjectCount = NetRefHandleManager->GetMaxActiveObjectCount();
 		InitParams.CurrentMaxInternalIndex = MaxInternalNetRefIndex;
 		InitParams.MaxConnectionCount = Connections->GetMaxConnectionCount();
