IoStoreOnDemand
- Fix bug that could prevent soft reference discovery if there were no other imports
- Fix format strings for some errors

diff --git a/Engine/Source/Runtime/Experimental/IoStore/OnDemand/Private/OnDemandIoStore.cpp b/Engine/Source/Runtime/Experimental/IoStore/OnDemand/Private/OnDemandIoStore.cpp
--- a/Engine/Source/Runtime/Experimental/IoStore/OnDemand/Private/OnDemandIoStore.cpp
+++ b/Engine/Source/Runtime/Experimental/IoStore/OnDemand/Private/OnDemandIoStore.cpp
@@ -154,13 +154,13 @@ static TIoStatusOr<FSharedContainerHeader> DeserializeContainerHeader(const FOnD
 		if (OutHeader->SoftPackageReferencesSerialInfo.Offset < 0)
 		{
 			FIoStatus Status = FIoStatusBuilder(EIoErrorCode::ReadError)
-				<< FString::Printf(TEXT("Invalid soft package reference offset '" INT64_FMT "'"), OutHeader->SoftPackageReferencesSerialInfo.Offset);
+				<< FString::Printf(TEXT("Invalid soft package reference offset '%" INT64_FMT "'"), OutHeader->SoftPackageReferencesSerialInfo.Offset);
 			return Status;
 		}
 		if ((OutHeader->SoftPackageReferencesSerialInfo.Offset + OutHeader->SoftPackageReferencesSerialInfo.Size) > Ar.TotalSize())
 		{
 			FIoStatus Status = FIoStatusBuilder(EIoErrorCode::ReadError)
-				<< FString::Printf(TEXT("Soft package reference offset '" INT64_FMT "' and size '" INT64_FMT "' will seek past the end of archive size '" INT64_FMT "'"),
+				<< FString::Printf(TEXT("Soft package reference offset '%" INT64_FMT "' and size '%" INT64_FMT "' will seek past the end of archive size '%" INT64_FMT "'"),
 					OutHeader->SoftPackageReferencesSerialInfo.Offset,
 					OutHeader->SoftPackageReferencesSerialInfo.Size,
 					Ar.TotalSize());
@@ -282,18 +282,18 @@ static FIoStatus BuildInstallData(
 					{
 						Queue.Add(ImportedPackageId);
 					}
+				}
 
-					if (EnumHasAnyFlags(Options, EBuildInstallDataOptions::IncludeSoftReferences))
+				if (EnumHasAnyFlags(Options, EBuildInstallDataOptions::IncludeSoftReferences))
+				{
+					TConstArrayView<FPackageId> SoftReferences;
+					TConstArrayView<uint32> Indices = FPackageStore::Get().GetSoftReferences(PackageId, SoftReferences);
+					for (uint32 Idx : Indices)
 					{
-						TConstArrayView<FPackageId> SoftReferences;
-						TConstArrayView<uint32> Indices = FPackageStore::Get().GetSoftReferences(PackageId, SoftReferences);
-						for (uint32 Idx : Indices)
+						const FPackageId& SoftRef = SoftReferences[Idx];
+						if (!Visitied.Contains(SoftRef))
 						{
-							const FPackageId& SoftRef = SoftReferences[Idx];
-							if (!Visitied.Contains(SoftRef))
-							{
-								Queue.Add(SoftRef);
-							}
+							Queue.Add(SoftRef);
 						}
 					}
 				}
