Interpolate kinematics when rewinding physics in immediate physics. Very small/minor fix.

diff --git a/Engine/Source/Runtime/Experimental/Chaos/Private/Chaos/Evolution/PBDMinEvolution.cpp b/Engine/Source/Runtime/Experimental/Chaos/Private/Chaos/Evolution/PBDMinEvolution.cpp
--- a/Engine/Source/Runtime/Experimental/Chaos/Private/Chaos/Evolution/PBDMinEvolution.cpp
+++ b/Engine/Source/Runtime/Experimental/Chaos/Private/Chaos/Evolution/PBDMinEvolution.cpp
@@ -192,6 +192,13 @@ namespace Chaos
 					Particle.SetWf(FVec3f::Lerp(Particle.GetPreWf(), Particle.GetWf(), FRealSingle(T)));
 				}
 			}
+			else if (Particle.ObjectState() == EObjectStateType::Kinematic)
+			{
+				// We can't rewind velocities for kinematics, and anyway they will get recalculated
+				// anyway when the kinematic targets are applied.
+				Particle.SetX(Particle.GetX() - Particle.GetV() * RewindDt);
+				Particle.SetRf(FRotation3f::IntegrateRotationWithAngularVelocity(Particle.GetRf(), -Particle.GetWf(), FRealSingle(RewindDt)));
+			}
 		}
 
 		for (auto& Particle : Particles.GetActiveKinematicParticlesView())
