Small preformance improvement in the merge sort constraint in collision detection
Use relexed memory ordered

diff --git a/Engine/Source/Runtime/Experimental/Chaos/Public/Chaos/Collision/SpatialAccelerationBroadPhase.h b/Engine/Source/Runtime/Experimental/Chaos/Public/Chaos/Collision/SpatialAccelerationBroadPhase.h
--- a/Engine/Source/Runtime/Experimental/Chaos/Public/Chaos/Collision/SpatialAccelerationBroadPhase.h
+++ b/Engine/Source/Runtime/Experimental/Chaos/Public/Chaos/Collision/SpatialAccelerationBroadPhase.h
@@ -506,7 +506,7 @@ namespace Chaos
 						int32 NextIndexToMerge = CurrentIndex;
 						while (true) 
 						{
-							int32 NextExpected = NextBatchToMerge.load();
+							int32 NextExpected = NextBatchToMerge.load(std::memory_order_relaxed);
 							// If there is a sorted batch ready to be merged, lets merge it
 							if (NextExpected != INDEX_NONE)
 							{
