Managed Array collection : add IsEmpty method

diff --git a/Engine/Source/Runtime/Experimental/Chaos/Public/GeometryCollection/ManagedArrayCollection.h b/Engine/Source/Runtime/Experimental/Chaos/Public/GeometryCollection/ManagedArrayCollection.h
--- a/Engine/Source/Runtime/Experimental/Chaos/Public/GeometryCollection/ManagedArrayCollection.h
+++ b/Engine/Source/Runtime/Experimental/Chaos/Public/GeometryCollection/ManagedArrayCollection.h
@@ -433,6 +433,14 @@ public:
 	*/
 	CHAOS_API TArray<FName> GroupNames() const;
 
+	/**
+	* Whether the collection is ampty ( no group / no attributes ) 
+	*/
+	CHAOS_API bool IsEmpty() const
+	{
+		return GroupInfo.IsEmpty() && Map.IsEmpty();
+	}
+
 	/**
 	* Check for the existence of a attribute
 	* @param Name - The name of the attribute
