Partial Island Sleeping - Bug fix: edge flag not updated

diff --git a/Engine/Source/Runtime/Experimental/Chaos/Private/Chaos/Island/IslandManager.cpp b/Engine/Source/Runtime/Experimental/Chaos/Private/Chaos/Island/IslandManager.cpp
--- a/Engine/Source/Runtime/Experimental/Chaos/Private/Chaos/Island/IslandManager.cpp
+++ b/Engine/Source/Runtime/Experimental/Chaos/Private/Chaos/Island/IslandManager.cpp
@@ -2624,6 +2624,8 @@ namespace Chaos::Private
 							{
 								if (!Edge->Flags.bWasProcessed)
 								{
+									Edge->Flags.bWasProcessed = true;
+
 									FPBDIslandParticle* OtherNode = (Node == Edge->Nodes[0]) ? Edge->Nodes[1] : Edge->Nodes[0];
 									// Wake up the OtherNode if the node level is greater than or equal to the edge level.
 									if ((OtherNode != nullptr) && OtherNode->Flags.bIsDynamic && (OtherNode->Level >= Edge->Level))
