Chaos
- fix MidphaseModifier DisableConvexOptimization so that it resets collisions if we change from using optimized geometry to non-optimized or vice versa after collisions have been generated on a previous frame
- clean up CCD TOI so it is always reset (mainly to make debugging CCD issues easier)

diff --git a/Engine/Source/Runtime/Experimental/Chaos/Private/Chaos/Collision/ParticlePairMidPhase.cpp b/Engine/Source/Runtime/Experimental/Chaos/Private/Chaos/Collision/ParticlePairMidPhase.cpp
--- a/Engine/Source/Runtime/Experimental/Chaos/Private/Chaos/Collision/ParticlePairMidPhase.cpp
+++ b/Engine/Source/Runtime/Experimental/Chaos/Private/Chaos/Collision/ParticlePairMidPhase.cpp
@@ -480,6 +480,7 @@ namespace Chaos
 			// Constraint may have been previously used with CCD enabled (e.g., a midphase modifier)
 			// so we need to make sure that the CCD flag is disabled
 			Constraint->SetCCDEnabled(false);
+			Constraint->ResetCCDTimeOfImpact();
 
 			// If the constraint was not used last frame, it needs to be reset, otherwise we will try to reuse
 			if (!bWasUpdatedLastTick || (Constraint->GetManifoldPoints().Num() == 0))
@@ -560,6 +561,7 @@ namespace Chaos
 			check(!Constraint->IsEnabled());
 			Constraint->SetCCDEnabled(true);
 			Constraint->SetCCDSweepEnabled(bEnableCCDSweep);
+			Constraint->ResetCCDTimeOfImpact();
 		}
 
 		if (!bEnableCCDSweep)
@@ -833,6 +835,7 @@ namespace Chaos
 		Flags.bIsSleeping = false;
 		Flags.bIsModified = false;
 		Flags.bIsConvexOptimizationActive = true;
+		Flags.bIsConvexOptimizationUsed = false;
 
 		ResetImpl();
 	}
@@ -1244,6 +1247,21 @@ namespace Chaos
 			? ClusteredHandle1->ConvexOptimizer().Get()
 			: nullptr;
 
+		// We need to reset collisions if we enable/disable use of the convex optimizer
+		// because we are switching to a different set of collision shapes.
+		// bIsConvexOptimizationActive can be changed in a modifier which gets reset every frame so we
+		// need to remember what we actually used rather than just resetting on change.
+		// NOTE: We don't have to handle the case where bHasConvexOptimizer0 or bHasConvexOptimizer1
+		// changes because we already reset collisions whenever geometry is replaced or rebuilt.
+		const bool bHasConvexOptimizer0 = (ConvexOptimizer0 != nullptr) && ConvexOptimizer0->IsValid();
+		const bool bHasConvexOptimizer1 = (ConvexOptimizer1 != nullptr) && ConvexOptimizer1->IsValid();
+		const bool bShouldUseConvexOptimization = Flags.bIsConvexOptimizationActive && (bHasConvexOptimizer0 || bHasConvexOptimizer1);
+		if (Flags.bIsConvexOptimizationUsed != bShouldUseConvexOptimization)
+		{
+			Flags.bIsConvexOptimizationUsed = bShouldUseConvexOptimization;
+			Constraints.Reset();
+		}
+
 		// See if we have a BVH for either/both of the particles
 		const Private::FImplicitBVH* BVH0 = nullptr;
 		const Private::FImplicitBVH* BVH1 = nullptr;
@@ -1251,7 +1269,6 @@ namespace Chaos
 		{
 			if (const FImplicitObjectUnion* Union0 = Implicit0->template AsA<FImplicitObjectUnion>())
 			{
-				const bool bHasConvexOptimizer0 = (ConvexOptimizer0 != nullptr) && ConvexOptimizer0->IsValid();
 				if(!bHasConvexOptimizer0)
 				{
 					BVH0 = Union0->GetBVH();
@@ -1259,7 +1276,6 @@ namespace Chaos
 			}
 			if (const FImplicitObjectUnion* Union1 = Implicit1->template AsA<FImplicitObjectUnion>())
 			{
-				const bool bHasConvexOptimizer1 = (ConvexOptimizer1 != nullptr) && ConvexOptimizer1->IsValid();
 				if(!bHasConvexOptimizer1)
 				{
 					BVH1 = Union1->GetBVH();
@@ -1817,6 +1833,7 @@ namespace Chaos
 			// CCD may be temporarily disabled by the user (via a midphase modifier) or because we are moving slowly.
 			Constraint->SetCCDEnabled(Flags.bIsCCDActive);
 			Constraint->SetCCDSweepEnabled(bUseCCDSweep);
+			Constraint->ResetCCDTimeOfImpact();
 
 			// NOTE: Probe constraints are always active and we run collision detection for them at the end opf the frame
 			bool bIsActive = true;
diff --git a/Engine/Source/Runtime/Experimental/Chaos/Private/Chaos/CollisionResolution.cpp b/Engine/Source/Runtime/Experimental/Chaos/Private/Chaos/CollisionResolution.cpp
--- a/Engine/Source/Runtime/Experimental/Chaos/Private/Chaos/CollisionResolution.cpp
+++ b/Engine/Source/Runtime/Experimental/Chaos/Private/Chaos/CollisionResolution.cpp
@@ -1046,8 +1046,6 @@ namespace Chaos
 			const FReal LengthCCD = DirCCD.SafeNormalize();
 			const bool bUseCCD = (LengthCCD > 0);
 
-			Constraint.ResetCCDTimeOfImpact();
-
 			if (bUseCCD)
 			{
 				const EContactShapesType ShapePairType = Constraint.GetShapesType();
diff --git a/Engine/Source/Runtime/Experimental/Chaos/Public/Chaos/Collision/ParticlePairMidPhase.h b/Engine/Source/Runtime/Experimental/Chaos/Public/Chaos/Collision/ParticlePairMidPhase.h
--- a/Engine/Source/Runtime/Experimental/Chaos/Public/Chaos/Collision/ParticlePairMidPhase.h
+++ b/Engine/Source/Runtime/Experimental/Chaos/Public/Chaos/Collision/ParticlePairMidPhase.h
@@ -416,7 +416,8 @@ namespace Chaos
 				uint16 bUseSweep : 1;    // True if CCD is active (this frame) and we are moving fast enough to require a sweep
 				uint16 bIsMACD : 1;      // True if MACD (movement-aware collision detection) is enabled for this pair
 				uint16 bIsConvexOptimizationActive : 1; // True if convex optimization is active for this midphase
-				uint16 bIsSleeping : 1;
+				uint16 bIsConvexOptimizationUsed : 1;   // True if we generated collisions with convex optimization active (used to detect changes to bIsConvexOptimizationActive in modifiers). See FGenericParticlePairMidPhase::GenerateCollisionsImpl
+				uint16 bIsSleeping : 1;  // True is the particles in the midphase are sleeping
 				uint16 bIsModified : 1;  // True if a modifier applied any changes to this midphase
 			};
 			uint16 Bits;
