Chaos - add flag to check if a particle's collision geometry is being modified after being shared

diff --git a/Engine/Source/Runtime/Experimental/Chaos/Private/Chaos/ClusterUnionManager.cpp b/Engine/Source/Runtime/Experimental/Chaos/Private/Chaos/ClusterUnionManager.cpp
--- a/Engine/Source/Runtime/Experimental/Chaos/Private/Chaos/ClusterUnionManager.cpp
+++ b/Engine/Source/Runtime/Experimental/Chaos/Private/Chaos/ClusterUnionManager.cpp
@@ -1144,6 +1144,10 @@ namespace Chaos
 	void FClusterUnionManager::ForceRegenerateGeometry(FClusterUnion& ClusterUnion, const TSet<FPBDRigidParticleHandle*>& FullChildrenSet)
 	{
 		SCOPE_CYCLE_COUNTER(STAT_ForceRegenerateGeometry);
+
+		// Must destroy collisions before modifying geometry
+		MEvolution.InvalidateParticle(ClusterUnion.InternalCluster);
+
 		ClusterUnion.InternalCluster->SetGeometry(MakeImplicitObjectPtr<FImplicitObjectUnionClustered>());
 
 		ModifyAdditionOfChildrenToClusterUnionGeometry(
diff --git a/Engine/Source/Runtime/Experimental/Chaos/Private/Chaos/ImplicitObjectUnion.cpp b/Engine/Source/Runtime/Experimental/Chaos/Private/Chaos/ImplicitObjectUnion.cpp
--- a/Engine/Source/Runtime/Experimental/Chaos/Private/Chaos/ImplicitObjectUnion.cpp
+++ b/Engine/Source/Runtime/Experimental/Chaos/Private/Chaos/ImplicitObjectUnion.cpp
@@ -102,10 +102,16 @@ bool FImplicitObjectUnion::HasValidBVH() const
 	// No need to create BVH if we are using the simple shape
 	return BVH.IsValid() && CVars::bChaosUnionBVHEnabled;
 }
+
+void FImplicitObjectUnion::CheckNotLocked() const
+{
+	ensureMsgf(!IsLocked(), TEXT("Attempt to modify an ImplicitObjectUnion after it has been shared and locked"));
+}
 	
 void FImplicitObjectUnion::Combine(TArray<Chaos::FImplicitObjectPtr>& OtherObjects)
 {
 	ensure(MObjects.Num());
+	CheckNotLocked();
 
 	const FAABB3 OtherBounds = CalculateObjectsBounds(MakeArrayView(OtherObjects));
 	if (MObjects.Num() > 0)
@@ -127,6 +133,8 @@ void FImplicitObjectUnion::Combine(TArray<Chaos::FImplicitObjectPtr>& OtherObjec
 
 void FImplicitObjectUnion::RemoveAt(int32 RemoveIndex)
 {
+	CheckNotLocked();
+
 	if (RemoveIndex < MObjects.Num())
 	{
 		SetNumLeafObjects(GetNumLeafObjects() - Private::FImplicitBVH::CountLeafObjects(MakeArrayView(&MObjects[RemoveIndex], 1)));
@@ -140,6 +148,8 @@ void FImplicitObjectUnion::RemoveAt(int32 RemoveIndex)
 
 void FImplicitObjectUnion::RemoveAtSortedIndices(const TArrayView<const int32>& InIndices)
 {
+	CheckNotLocked();
+
 	if (InIndices.IsEmpty())
 	{
 		return;
diff --git a/Engine/Source/Runtime/Experimental/Chaos/Private/Chaos/ParticleHandle.cpp b/Engine/Source/Runtime/Experimental/Chaos/Private/Chaos/ParticleHandle.cpp
--- a/Engine/Source/Runtime/Experimental/Chaos/Private/Chaos/ParticleHandle.cpp
+++ b/Engine/Source/Runtime/Experimental/Chaos/Private/Chaos/ParticleHandle.cpp
@@ -121,6 +121,18 @@ namespace Chaos
 		}
 	}
 
+	template <typename T, int d, bool bPersistent>
+	void Chaos::TGeometryParticleHandleImp<T, d, bPersistent>::LockGeometry()
+	{
+		if (FImplicitObjectRef CurrentGeometry = GetGeometry())
+		{
+			if (FImplicitObjectUnion* Union = CurrentGeometry->AsA<FImplicitObjectUnion>())
+			{
+				Union->SetLockedRecursive(true);
+			}
+		}
+	}
+
 	template <typename T, int d>
 	void Chaos::TGeometryParticle<T, d>::RemoveShape(FPerShapeData* InShape, bool bWakeTouching)
 	{
@@ -209,10 +221,28 @@ namespace Chaos
 	{
 		if (MNonFrequentData.IsDirty(MDirtyFlags))
 		{
-			if (const FImplicitObjectUnion* Union = MNonFrequentData.Read().GetGeometry()->template GetObject<FImplicitObjectUnion>())
+			if (FImplicitObjectRef Implicit = MNonFrequentData.Read().GetGeometry())
 			{
-				// This will rebuild the BVH if the geometry is new, otherwise do nothing
-				const_cast<FImplicitObjectUnion*>(Union)->SetAllowBVH(true);
+				if (const FImplicitObjectUnion* Union = Implicit->AsA<FImplicitObjectUnion>())
+				{
+					// This will rebuild the BVH if the geometry is new, otherwise do nothing
+					const_cast<FImplicitObjectUnion*>(Union)->SetAllowBVH(true);
+				}
+			}
+		}
+	}
+
+	template <typename T, int d>
+	void Chaos::TGeometryParticle<T, d>::LockGeometry()
+	{
+		if (MNonFrequentData.IsDirty(MDirtyFlags))
+		{
+			if (FImplicitObjectRef Implicit = MNonFrequentData.Read().GetGeometry())
+			{
+				if (FImplicitObjectUnion* Union = Implicit->AsA<FImplicitObjectUnion>())
+				{
+					Union->SetLockedRecursive(true);
+				}
 			}
 		}
 	}
diff --git a/Engine/Source/Runtime/Experimental/Chaos/Private/PBDRigidsSolver.cpp b/Engine/Source/Runtime/Experimental/Chaos/Private/PBDRigidsSolver.cpp
--- a/Engine/Source/Runtime/Experimental/Chaos/Private/PBDRigidsSolver.cpp
+++ b/Engine/Source/Runtime/Experimental/Chaos/Private/PBDRigidsSolver.cpp
@@ -1440,6 +1440,7 @@ namespace Chaos
 				}
 
 				Particle->PrepareBVH();
+				Particle->LockGeometry();
 				Particle->SyncRemoteData(*Manager, DataIdx, Dirty.PropertyData, Dirty.ShapeDataIndices, ShapeDirtyData);
 				Proxy->ClearAccumulatedData();
 				Proxy->ResetDirtyIdx();
@@ -1456,6 +1457,7 @@ namespace Chaos
 			{
 				FClusterUnionPhysicsProxy* Proxy = static_cast<FClusterUnionPhysicsProxy*>(Dirty.Proxy);
 				FClusterUnionPhysicsProxy::FExternalParticle* Particle = Proxy->GetParticle_External();
+				Particle->LockGeometry();
 				Proxy->SyncRemoteData(*Manager, DataIdx, Dirty.PropertyData);
 				Proxy->ClearAccumulatedData();
 				Proxy->ResetDirtyIdx();
diff --git a/Engine/Source/Runtime/Experimental/Chaos/Private/PhysicsProxy/ClusterUnionPhysicsProxy.cpp b/Engine/Source/Runtime/Experimental/Chaos/Private/PhysicsProxy/ClusterUnionPhysicsProxy.cpp
--- a/Engine/Source/Runtime/Experimental/Chaos/Private/PhysicsProxy/ClusterUnionPhysicsProxy.cpp
+++ b/Engine/Source/Runtime/Experimental/Chaos/Private/PhysicsProxy/ClusterUnionPhysicsProxy.cpp
@@ -776,6 +776,12 @@ namespace Chaos
 					// Fallback to deep copy geometry hierarchy
 					BufferData.Geometry = ClusterUnion->Geometry->DeepCopyGeometry();
 				}
+
+				if (FImplicitObjectUnion* Union = BufferData.Geometry->AsA<FImplicitObjectUnion>())
+				{
+					Union->SetLockedRecursive(true);
+				}
+
 				ClusterUnion->bGeometryModified = false;
 			}
 		}
diff --git a/Engine/Source/Runtime/Experimental/Chaos/Private/PhysicsProxy/GeometryCollectionPhysicsProxy.cpp b/Engine/Source/Runtime/Experimental/Chaos/Private/PhysicsProxy/GeometryCollectionPhysicsProxy.cpp
--- a/Engine/Source/Runtime/Experimental/Chaos/Private/PhysicsProxy/GeometryCollectionPhysicsProxy.cpp
+++ b/Engine/Source/Runtime/Experimental/Chaos/Private/PhysicsProxy/GeometryCollectionPhysicsProxy.cpp
@@ -1289,10 +1289,12 @@ void FGeometryCollectionPhysicsProxy::InitializeBodiesPT(Chaos::FPBDRigidsSolver
 					if (const Chaos::FImplicitObjectUnion* Union = ClusterHandle->GetGeometry()->GetObject<Chaos::FImplicitObjectUnion>())
 					{
 						const_cast<Chaos::FImplicitObjectUnion*>(Union)->SetAllowBVH(true);
+						const_cast<Chaos::FImplicitObjectUnion*>(Union)->SetLocked(true);
 					}
 					else if (const Chaos::FImplicitObjectUnion* UnionClustered = ClusterHandle->GetGeometry()->GetObject<Chaos::FImplicitObjectUnionClustered>())
 					{
 						const_cast<Chaos::FImplicitObjectUnion*>(UnionClustered)->SetAllowBVH(true);
+						const_cast<Chaos::FImplicitObjectUnion*>(UnionClustered)->SetLocked(true);
 					}
 				}
 
diff --git a/Engine/Source/Runtime/Experimental/Chaos/Public/Chaos/ImplicitObjectUnion.h b/Engine/Source/Runtime/Experimental/Chaos/Public/Chaos/ImplicitObjectUnion.h
--- a/Engine/Source/Runtime/Experimental/Chaos/Public/Chaos/ImplicitObjectUnion.h
+++ b/Engine/Source/Runtime/Experimental/Chaos/Public/Chaos/ImplicitObjectUnion.h
@@ -65,10 +65,39 @@ class FImplicitObjectUnion : public FImplicitObject
 	{
 		if (bInAllowBVH != Flags.bAllowBVH)
 		{
+			CheckNotLocked();
+
 			Flags.bAllowBVH = bInAllowBVH;
 			RebuildBVH();
 		}
 	}
+
+	// Recurively locked the union. Once locked, it is an error to make a change to the Union or any children
+	void SetLocked(const bool bInLocked)
+	{
+		Flags.bIsLocked = bInLocked;
+	}
+
+	void SetLockedRecursive(const bool bInLocked)
+	{
+		if (bInLocked != IsLocked())
+		{
+			for (FImplicitObject* Object : MObjects)
+			{
+				if (FImplicitObjectUnion* Union = Object->template GetObject<FImplicitObjectUnion>())
+				{
+					Union->SetLockedRecursive(bInLocked);
+				}
+			}
+
+			SetLocked(bInLocked);
+		}
+	}
+
+	bool IsLocked() const
+	{
+		return !!Flags.bIsLocked;
+	}
 	
 	CHAOS_API virtual Chaos::FImplicitObjectPtr CopyGeometry() const;
 	CHAOS_API virtual Chaos::FImplicitObjectPtr CopyGeometryWithScale(const FVec3& Scale) const override;
@@ -281,6 +310,7 @@ protected:
 			// NOTE: Flags are serialized. Ordering must be retained
 			uint8 bAllowBVH : 1;				// Are we allowed to use a BVH for this Union? (We generally only want a BVH in the root union)
 			uint8 bHasBVH : 1;					// Do we currently have a BVH? Equivalent to checking BVH pointer, but used by serialization to know whether to load a BVH
+			uint8 bIsLocked : 1;				// If the union is locked, it is an error to change children
 		};
 		uint8 Bits;
 	};
@@ -288,6 +318,9 @@ protected:
 	// Check if the BVH could be used and valid
 	bool HasValidBVH() const;
 
+	// Log an error we try to modify a union after it has been shared
+	CHAOS_API void CheckNotLocked() const;
+
 	// list of implicit objects that are part of the union
 	TArray<Chaos::FImplicitObjectPtr> MObjects;
 	
diff --git a/Engine/Source/Runtime/Experimental/Chaos/Public/Chaos/ParticleHandle.h b/Engine/Source/Runtime/Experimental/Chaos/Public/Chaos/ParticleHandle.h
--- a/Engine/Source/Runtime/Experimental/Chaos/Public/Chaos/ParticleHandle.h
+++ b/Engine/Source/Runtime/Experimental/Chaos/Public/Chaos/ParticleHandle.h
@@ -591,6 +591,10 @@ public:
 	void SetGeometry(const FImplicitObjectPtr& InGeometry) { GeometryParticles->SetGeometry(ParticleIdx, InGeometry); }
 	void MergeGeometry(TArray<Chaos::FImplicitObjectPtr>&& Objects);
 
+	// Set a flag on all Unions to indicate that any add/remove operations are errors.
+	// This is called when the shapes are shared between GT/PT
+	void LockGeometry();
+
 	UE_DEPRECATED(5.4, "Use GetGeometry instead")
 	TSerializablePtr<FImplicitObject> Geometry() const { check(false); return TSerializablePtr<FImplicitObject>(); }
 	
@@ -2785,6 +2789,10 @@ public:
 
 	const FImplicitObjectRef GetGeometry() const { return MNonFrequentData.Read().GetGeometry(); }
 
+	// Set a flag on all Unions to indicate that any add/remove operations are errors.
+	// This is called when the shapes are shared between GT/PT
+	void LockGeometry();
+
 	UE_DEPRECATED(5.4, "Use SetGeometry with an array of Chaos::FImplicitObjectPtr instead")
 	void SetGeometry(TUniquePtr<FImplicitObject>&& UniqueGeometry) { check(false); }
 	
