ChaosVD - add particle inflated bounds serialization and visualization

#tests Tested in PIE, Tested cook

diff --git a/Engine/Plugins/ChaosVD/Source/ChaosVD/Private/ChaosVDSceneParticle.cpp b/Engine/Plugins/ChaosVD/Source/ChaosVD/Private/ChaosVDSceneParticle.cpp
--- a/Engine/Plugins/ChaosVD/Source/ChaosVD/Private/ChaosVDSceneParticle.cpp
+++ b/Engine/Plugins/ChaosVD/Source/ChaosVD/Private/ChaosVDSceneParticle.cpp
@@ -567,6 +567,15 @@ FBox FChaosVDSceneParticle::GetBoundingBox() const
 	return CachedBounds;
 }
 
+FBox FChaosVDSceneParticle::GetInflatedBoundingBox() const
+{
+	if (ParticleDataPtr && ParticleDataPtr->ParticleInflatedBounds.HasValidData())
+	{
+		return FBox(ParticleDataPtr->ParticleInflatedBounds.MMin, ParticleDataPtr->ParticleInflatedBounds.MMax);
+	}
+	return GetBoundingBox();
+}
+
 Chaos::TAABB<double, 3> FChaosVDSceneParticle::GetChaosBoundingBox() const
 {
 	FBox Bounds = GetBoundingBox();
diff --git a/Engine/Plugins/ChaosVD/Source/ChaosVD/Private/Settings/ChaosVDParticleVisualizationSettings.h b/Engine/Plugins/ChaosVD/Source/ChaosVD/Private/Settings/ChaosVDParticleVisualizationSettings.h
--- a/Engine/Plugins/ChaosVD/Source/ChaosVD/Private/Settings/ChaosVDParticleVisualizationSettings.h
+++ b/Engine/Plugins/ChaosVD/Source/ChaosVD/Private/Settings/ChaosVDParticleVisualizationSettings.h
@@ -43,6 +43,7 @@ enum class EChaosVDParticleDataVisualizationFlags : uint32
 	DrawDataOnlyForSelectedParticle	= 1 << 8,
 	
 	Bounds = 1 << 10,
+	InflatedBounds = 1 << 11,
 
 	EnableDraw	= 1 << 9
 };
diff --git a/Engine/Plugins/ChaosVD/Source/ChaosVD/Private/Visualizers/ChaosVDParticleDataComponentVisualizer.cpp b/Engine/Plugins/ChaosVD/Source/ChaosVD/Private/Visualizers/ChaosVDParticleDataComponentVisualizer.cpp
--- a/Engine/Plugins/ChaosVD/Source/ChaosVD/Private/Visualizers/ChaosVDParticleDataComponentVisualizer.cpp
+++ b/Engine/Plugins/ChaosVD/Source/ChaosVD/Private/Visualizers/ChaosVDParticleDataComponentVisualizer.cpp
@@ -285,6 +285,19 @@ void FChaosVDParticleDataComponentVisualizer::DrawVisualizationForParticleData(c
 		FChaosVDDebugDrawUtils::DrawBox(PDI, ParticleBounds.GetExtent(), FColor::Red, Location, FText::GetEmpty(), InVisualizationContext.DebugDrawSettings->DepthPriority, LineThickness);
 	}
 
+	if (InVisualizationContext.IsVisualizationFlagEnabled(EChaosVDParticleDataVisualizationFlags::InflatedBounds))
+	{
+		FBox ParticleInflatedBounds = InParticleInstance.GetInflatedBoundingBox();
+		bool bShowInflatedBounds = (ParticleInflatedBounds != ParticleBounds) || !InVisualizationContext.IsVisualizationFlagEnabled(EChaosVDParticleDataVisualizationFlags::Bounds);
+		if (bShowInflatedBounds)
+		{
+			FTransform Location;
+			Location.SetLocation(ParticleInflatedBounds.GetCenter());
+			FColor InflatedBoundsColour = FColor(0xC7, 0x6E, 0x10);	// Dark Orange
+			FChaosVDDebugDrawUtils::DrawBox(PDI, ParticleInflatedBounds.GetExtent(), InflatedBoundsColour, Location, FText::GetEmpty(), InVisualizationContext.DebugDrawSettings->DepthPriority, LineThickness);
+		}
+	}
+
 	if (ParticleData->ParticleVelocities.HasValidData())
 	{
 		DrawParticleVector(PDI, OwnerCoMLocation, ParticleData->ParticleVelocities.MV,EChaosVDParticleDataVisualizationFlags::Velocity, InVisualizationContext, LineThickness); 
diff --git a/Engine/Plugins/ChaosVD/Source/ChaosVD/Public/ChaosVDSceneParticle.h b/Engine/Plugins/ChaosVD/Source/ChaosVD/Public/ChaosVDSceneParticle.h
--- a/Engine/Plugins/ChaosVD/Source/ChaosVD/Public/ChaosVDSceneParticle.h
+++ b/Engine/Plugins/ChaosVD/Source/ChaosVD/Public/ChaosVDSceneParticle.h
@@ -77,6 +77,7 @@ public:
 	}
 
 	FBox GetBoundingBox() const;
+	FBox GetInflatedBoundingBox() const;
 	Chaos::TAABB<double, 3> GetChaosBoundingBox() const;
 
 	// BEGIN IChaosVDCollisionDataProvider Interface
diff --git a/Engine/Source/Runtime/Core/Public/UObject/FortniteMainBranchObjectVersions.inl b/Engine/Source/Runtime/Core/Public/UObject/FortniteMainBranchObjectVersions.inl
--- a/Engine/Source/Runtime/Core/Public/UObject/FortniteMainBranchObjectVersions.inl
+++ b/Engine/Source/Runtime/Core/Public/UObject/FortniteMainBranchObjectVersions.inl
@@ -622,5 +622,8 @@ FORTNITE_MAIN_VERSION(PCGSplineDirectionClockwiseFix, 202)
 // Rect Lights set in EV units had the wrong intensity (older files need a flag set to keep the old look)
 FORTNITE_MAIN_VERSION(RectLightFixedEVUnitConversion, 203)
 
+// Add particle bounds to data exported to CVD
+FORTNITE_MAIN_VERSION(ParticleInflatedBoundsInChaosVisualDebugger, 204)
+
 // -----<new versions can be added above this line>-------------------------------------------------
 #undef FORTNITE_MAIN_VERSION
\ No newline at end of file
diff --git a/Engine/Source/Runtime/Experimental/Chaos/Private/ChaosVisualDebugger/ChaosVDDataWrapperUtils.cpp b/Engine/Source/Runtime/Experimental/Chaos/Private/ChaosVisualDebugger/ChaosVDDataWrapperUtils.cpp
--- a/Engine/Source/Runtime/Experimental/Chaos/Private/ChaosVisualDebugger/ChaosVDDataWrapperUtils.cpp
+++ b/Engine/Source/Runtime/Experimental/Chaos/Private/ChaosVisualDebugger/ChaosVDDataWrapperUtils.cpp
@@ -122,6 +122,8 @@ FChaosVDParticleDataWrapper FChaosVDDataWrapperUtils::BuildParticleDataWrapperFr
 		WrappedParticleData.ParticlePositionRotation.CopyFrom<Chaos::FGeometryParticleHandle, FChaosVDParticlePositionRotation::EAccessorType::XR>(*ParticleHandlePtr);
 	}
 
+	WrappedParticleData.ParticleInflatedBounds.CopyFrom(*ParticleHandlePtr);
+
 	if (const Chaos::TPBDRigidClusteredParticleHandleImp<Chaos::FReal, 3, true>* ClusteredParticle = ParticleHandlePtr->CastToClustered())
 	{
 		WrappedParticleData.ParticleCluster.CopyFrom(*ClusteredParticle);
diff --git a/Engine/Source/Runtime/Experimental/ChaosVisualDebugger/Private/DataWrappers/ChaosVDParticleDataWrapper.cpp b/Engine/Source/Runtime/Experimental/ChaosVisualDebugger/Private/DataWrappers/ChaosVDParticleDataWrapper.cpp
--- a/Engine/Source/Runtime/Experimental/ChaosVisualDebugger/Private/DataWrappers/ChaosVDParticleDataWrapper.cpp
+++ b/Engine/Source/Runtime/Experimental/ChaosVisualDebugger/Private/DataWrappers/ChaosVDParticleDataWrapper.cpp
@@ -74,6 +74,21 @@ bool FChaosVDParticleVelocities::Serialize(FArchive& Ar)
 	return !Ar.IsError();
 }
 
+bool FChaosVDParticleBounds::Serialize(FArchive& Ar)
+{
+	Ar << bHasValidData;
+
+	if (!bHasValidData)
+	{
+		return !Ar.IsError();
+	}
+
+	Ar << MMin;
+	Ar << MMax;
+
+	return !Ar.IsError();
+}
+
 bool FChaosVDParticleDynamics::Serialize(FArchive& Ar)
 {
 	Ar << bHasValidData;
@@ -186,6 +201,8 @@ bool FChaosVDVSmooth::Serialize(FArchive& Ar)
 
 bool FChaosVDParticleDataWrapper::Serialize(FArchive& Ar)
 {
+	Ar.UsingCustomVersion(FFortniteMainBranchObjectVersion::GUID);
+
 	Ar << bHasValidData;
 	if (!bHasValidData)
 	{
@@ -220,6 +237,12 @@ bool FChaosVDParticleDataWrapper::Serialize(FArchive& Ar)
 	Ar << ParticlePositionRotation;
 	Ar << ParticleVelocities;
 
+	// Bounds data was not exported prior to this version
+	if (Ar.CustomVer(FFortniteMainBranchObjectVersion::GUID) >= FFortniteMainBranchObjectVersion::ParticleInflatedBoundsInChaosVisualDebugger)
+	{
+		Ar << ParticleInflatedBounds;
+	}
+
 	Ar << ParticleDynamics;
 	Ar << ParticleDynamicsMisc;
 	Ar << ParticleMassProps;
diff --git a/Engine/Source/Runtime/Experimental/ChaosVisualDebugger/Public/DataWrappers/ChaosVDParticleDataWrapper.h b/Engine/Source/Runtime/Experimental/ChaosVisualDebugger/Public/DataWrappers/ChaosVDParticleDataWrapper.h
--- a/Engine/Source/Runtime/Experimental/ChaosVisualDebugger/Public/DataWrappers/ChaosVDParticleDataWrapper.h
+++ b/Engine/Source/Runtime/Experimental/ChaosVisualDebugger/Public/DataWrappers/ChaosVDParticleDataWrapper.h
@@ -254,6 +254,40 @@ struct FChaosVDParticleVelocities : public FChaosVDWrapperDataBase
 
 CVD_IMPLEMENT_SERIALIZER(FChaosVDParticleVelocities)
 
+USTRUCT()
+struct FChaosVDParticleBounds : public FChaosVDWrapperDataBase
+{
+	GENERATED_BODY()
+
+	FChaosVDParticleBounds()
+	{
+	}
+
+	CHAOSVDRUNTIME_API bool Serialize(FArchive& Ar);
+
+	template <typename TOther>
+	void CopyFrom(const TOther& Other)
+	{
+		MMin = Other.WorldSpaceInflatedBounds().Min();
+		MMax = Other.WorldSpaceInflatedBounds().Max();
+		bHasValidData = true;
+	}
+
+	template <typename TOther>
+	void CopyTo(TOther& Other) const
+	{
+		// TODO:
+	}
+
+	UPROPERTY(VisibleAnywhere, Category = "Particle Bounds Min")
+	FVector MMin = FVector(ForceInit);
+
+	UPROPERTY(VisibleAnywhere, Category = "Particle Bounds Max")
+	FVector MMax = FVector(ForceInit);
+};
+
+CVD_IMPLEMENT_SERIALIZER(FChaosVDParticleBounds)
+
 /** Simplified UStruct version of FParticleDynamics.
  * Used to be able to show the values in the editor and allow changes via the Property Editor.
  */
@@ -707,6 +741,9 @@ struct FChaosVDParticleDataWrapper : public FChaosVDWrapperDataBase
 	UPROPERTY(VisibleAnywhere, Category= "Particle Velocities")
 	FChaosVDParticleVelocities ParticleVelocities;
 
+	UPROPERTY(VisibleAnywhere, Category = "Particle Inflated Bounds")
+	FChaosVDParticleBounds ParticleInflatedBounds;
+
 	UPROPERTY(VisibleAnywhere, Category= "Particle Kinematic Target")
 	FChaosVDKinematicTarget ParticleKinematicTarget;
 
