Chaos - add cvar to enable sorting manifold points for solver convergence

diff --git a/Engine/Source/Runtime/Experimental/Chaos/Private/Chaos/CollisionOneShotManifolds.cpp b/Engine/Source/Runtime/Experimental/Chaos/Private/Chaos/CollisionOneShotManifolds.cpp
--- a/Engine/Source/Runtime/Experimental/Chaos/Private/Chaos/CollisionOneShotManifolds.cpp
+++ b/Engine/Source/Runtime/Experimental/Chaos/Private/Chaos/CollisionOneShotManifolds.cpp
@@ -159,10 +159,14 @@ namespace Chaos
 	bool bChaos_Collision_UseConvexTriangleGJKSAT = true;
 	FAutoConsoleVariableRef CVarChaos_Collision_UseConvexTriangleSAT(TEXT("p.Chaos.Collision.UseConvexTriangleGJKSAT"), bChaos_Collision_UseConvexTriangleGJKSAT, TEXT(""));
 
-	// Whether to sort sphere/capsule vs mesh contacts by depth. Doing so leads to more consistent behaviour but is more expensive
+	// Whether to sort sphere/capsule/convex vs mesh contacts by depth. Doing so leads to more consistent behaviour but is more expensive
 	bool bChaos_Collision_ConvexTriMeshSortByPhi = false;
 	FAutoConsoleVariableRef CVarChaos_Collision_ConvexTriMeshSortByPhi(TEXT("p.Chaos.Collision.ConvexTriMeshSortByPhi"), bChaos_Collision_ConvexTriMeshSortByPhi, TEXT(""));
 
+	// Whether to sort sphere/capsule/convex vs mesh contacts by distance from center of mass in the contact plane. Slow but better for convergence.
+	bool bChaos_Collision_ConvexTriMeshSortByDistance = false;
+	FAutoConsoleVariableRef CVarChaos_Collision_ConvexTriMeshSortByDistance(TEXT("p.Chaos.Collision.ConvexTriMeshSortByDistance"), bChaos_Collision_ConvexTriMeshSortByDistance, TEXT(""));
+
 	namespace Collisions
 	{
 		// Forward delarations we need from CollisionRestitution.cpp
diff --git a/Engine/Source/Runtime/Experimental/Chaos/Private/Chaos/CollisionOneShotManifoldsMeshShapes.cpp b/Engine/Source/Runtime/Experimental/Chaos/Private/Chaos/CollisionOneShotManifoldsMeshShapes.cpp
--- a/Engine/Source/Runtime/Experimental/Chaos/Private/Chaos/CollisionOneShotManifoldsMeshShapes.cpp
+++ b/Engine/Source/Runtime/Experimental/Chaos/Private/Chaos/CollisionOneShotManifoldsMeshShapes.cpp
@@ -43,6 +43,7 @@ namespace Chaos
 	extern bool bChaos_Collision_ConvexTriMeshBackFaceCull;
 
 	extern bool bChaos_Collision_ConvexTriMeshSortByPhi;
+	extern bool bChaos_Collision_ConvexTriMeshSortByDistance;
 
 	namespace CVars
 	{
@@ -912,6 +913,7 @@ namespace Chaos
 				ContactGeneratorSettings.FaceNormalDotThreshold = 0.9999;	// ~0.8deg Normals must be accurate or rolling will not work correctly
 				ContactGeneratorSettings.bUseTwoPassLoop = false;			// two-pass loop is not helpful for capsules and spheres
 				ContactGeneratorSettings.bSortByPhi = bChaos_Collision_ConvexTriMeshSortByPhi;
+				ContactGeneratorSettings.bSortForSolverConvergence = bChaos_Collision_ConvexTriMeshSortByDistance && !bChaos_Collision_ConvexTriMeshSortByPhi;
 				Private::FMeshContactGenerator ContactGenerator(ContactGeneratorSettings);
 
 				const FVec3 RelativeMovement = FVec3(Constraint.GetRelativeMovement());
@@ -971,6 +973,7 @@ namespace Chaos
 				ContactGeneratorSettings.FaceNormalDotThreshold = 0.9999;	// ~0.8deg Normals must be accurate or rolling will not work correctly
 				ContactGeneratorSettings.bUseTwoPassLoop = false;			// two-pass loop is not helpful for capsules and spheres
 				ContactGeneratorSettings.bSortByPhi = bChaos_Collision_ConvexTriMeshSortByPhi;
+				ContactGeneratorSettings.bSortForSolverConvergence = bChaos_Collision_ConvexTriMeshSortByDistance && !bChaos_Collision_ConvexTriMeshSortByPhi;
 				Private::FMeshContactGenerator ContactGenerator(ContactGeneratorSettings);
 
 				if (const FImplicitSphere3* Sphere = Quadratic.template GetObject<FImplicitSphere3>())
@@ -1037,6 +1040,7 @@ namespace Chaos
 				const FVec3 RelativeMovement = FVec3(Constraint.GetRelativeMovement());
 				Private::FMeshContactGeneratorSettings ContactGeneratorSettings;
 				ContactGeneratorSettings.bSortByPhi = bChaos_Collision_ConvexTriMeshSortByPhi;
+				ContactGeneratorSettings.bSortForSolverConvergence = bChaos_Collision_ConvexTriMeshSortByDistance && !bChaos_Collision_ConvexTriMeshSortByPhi;
 				Private::FMeshContactGenerator ContactGenerator(ContactGeneratorSettings);
 
 				if (const FImplicitBox3* RawBox = Convex.template GetObject<FImplicitBox3>())
@@ -1111,6 +1115,7 @@ namespace Chaos
 				// New version uses a two-pass loop over triangles to avoid visiting triangles whose vertices are all colliding as a result of checking adjacent triangles
 				Private::FMeshContactGeneratorSettings ContactGeneratorSettings;
 				ContactGeneratorSettings.bSortByPhi = bChaos_Collision_ConvexTriMeshSortByPhi;
+				ContactGeneratorSettings.bSortForSolverConvergence = bChaos_Collision_ConvexTriMeshSortByDistance && !bChaos_Collision_ConvexTriMeshSortByPhi;
 				Private::FMeshContactGenerator ContactGenerator(ContactGeneratorSettings);
 				const FVec3 RelativeMovement = FVec3(Constraint.GetRelativeMovement());
 
