Include soft references when installing a package via console command.

#rnx

diff --git a/Engine/Source/Runtime/Experimental/IoStore/OnDemandCore/Private/IoStoreOnDemand.cpp b/Engine/Source/Runtime/Experimental/IoStore/OnDemandCore/Private/IoStoreOnDemand.cpp
--- a/Engine/Source/Runtime/Experimental/IoStore/OnDemandCore/Private/IoStoreOnDemand.cpp
+++ b/Engine/Source/Runtime/Experimental/IoStore/OnDemandCore/Private/IoStoreOnDemand.cpp
@@ -282,6 +282,7 @@ static void InstallPackage(const TArray<FString>& Args)
 	FOnDemandInstallArgs InstallArgs;
 	InstallArgs.PackageIds.Add(PackageId);
 	InstallArgs.ContentHandle = DefaultContentHandle;
+	InstallArgs.Options = EOnDemandInstallOptions::InstallSoftReferences;
 
 	IoStore->Install(
 		MoveTemp(InstallArgs),
