Check for package redirects when install on-demand packages.

#rnx

diff --git a/Engine/Source/Runtime/Experimental/IoStore/OnDemand/Private/OnDemandIoStore.cpp b/Engine/Source/Runtime/Experimental/IoStore/OnDemand/Private/OnDemandIoStore.cpp
--- a/Engine/Source/Runtime/Experimental/IoStore/OnDemand/Private/OnDemandIoStore.cpp
+++ b/Engine/Source/Runtime/Experimental/IoStore/OnDemand/Private/OnDemandIoStore.cpp
@@ -258,11 +258,20 @@ static FIoStatus BuildInstallData(
 			Queue.Add(PackageId);
 		}
 
-		FPackageStoreReadScope _(FPackageStore::Get());
+		FPackageStore& PackageStore = FPackageStore::Get();
+		FPackageStoreReadScope _(PackageStore);
 
 		while (!Queue.IsEmpty())
 		{
 			FPackageId PackageId = Queue.PopFrontValue();
+			{
+				FName		SourcePackageName;
+				FPackageId	RedirectedToPackageId;
+				if (PackageStore.GetPackageRedirectInfo(PackageId, SourcePackageName, RedirectedToPackageId))
+				{
+					PackageId = RedirectedToPackageId;
+				}
+			}
 
 			bool bIsAlreadyInSet = false;
 			Visitied.Add(PackageId, &bIsAlreadyInSet);
@@ -272,7 +281,7 @@ static FIoStatus BuildInstallData(
 			}
 
 			FPackageStoreEntry PackageStoreEntry;
-			const EPackageStoreEntryStatus EntryStatus = FPackageStore::Get().GetPackageStoreEntry(PackageId, NAME_None, PackageStoreEntry);
+			const EPackageStoreEntryStatus EntryStatus = PackageStore.GetPackageStoreEntry(PackageId, NAME_None, PackageStoreEntry);
 			if (EntryStatus != EPackageStoreEntryStatus::Missing)
 			{
 				ResolvedPackageIds.Add(PackageId);
@@ -287,7 +296,7 @@ static FIoStatus BuildInstallData(
 				if (EnumHasAnyFlags(Options, EBuildInstallDataOptions::IncludeSoftReferences))
 				{
 					TConstArrayView<FPackageId> SoftReferences;
-					TConstArrayView<uint32> Indices = FPackageStore::Get().GetSoftReferences(PackageId, SoftReferences);
+					TConstArrayView<uint32> Indices = PackageStore.GetSoftReferences(PackageId, SoftReferences);
 					for (uint32 Idx : Indices)
 					{
 						const FPackageId& SoftRef = SoftReferences[Idx];
