Compile fix when UE_ENABLE_REQUEST_THREAD is enabled.

#rnx

- Moved LexToString(EThreadPriority) to the global namespace and made it static although the long term plan is to eventually move it to GenericPlatformAffinity.h.

diff --git a/Engine/Source/Runtime/Experimental/IoStore/OnDemand/Private/OnDemandHttpThread.cpp b/Engine/Source/Runtime/Experimental/IoStore/OnDemand/Private/OnDemandHttpThread.cpp
--- a/Engine/Source/Runtime/Experimental/IoStore/OnDemand/Private/OnDemandHttpThread.cpp
+++ b/Engine/Source/Runtime/Experimental/IoStore/OnDemand/Private/OnDemandHttpThread.cpp
@@ -13,6 +13,32 @@
 #include "OnDemandHttpClient.h"
 #include "Statistics.h"
 
+// TODO - Maybe do LexToString/LexFromString (dupe code in UnrealEngine.cpp ThreadPriorityToString)
+static const TCHAR* LexToString(EThreadPriority Priority)
+{
+	switch (Priority)
+	{
+		case EThreadPriority::TPri_Normal:
+			return TEXT("TPri_Normal");
+		case EThreadPriority::TPri_AboveNormal:
+			return TEXT("TPri_AboveNormal");
+		case EThreadPriority::TPri_BelowNormal:
+			return TEXT("TPri_BelowNormal");
+		case EThreadPriority::TPri_Highest:
+			return TEXT("TPri_Highest");
+		case EThreadPriority::TPri_Lowest:
+			return TEXT("TPri_Lowest");
+		case EThreadPriority::TPri_SlightlyBelowNormal:
+			return TEXT("TPri_SlightlyBelowNormal");
+		case EThreadPriority::TPri_TimeCritical:
+			return TEXT("TPri_TimeCritical");
+		case EThreadPriority::TPri_Num:
+		default:
+			return TEXT("TPri_Undefined");
+			break;
+	};
+}
+
 namespace UE::IoStore
 {
 
@@ -34,7 +60,6 @@ extern int32 GIasHttpPollTimeoutMs;
 
 extern UE::Tasks::ETaskPriority GetRequestCompletionTaskPriority();
 
-// TODO - Maybe do LexToString/LexFromString (dupe code in UnrealEngine.cpp ThreadPriorityToString)
 const int32 GREPLACEOnDemandBackendThreadPriorities[5] =
 {
 	EThreadPriority::TPri_Lowest,
@@ -44,31 +69,6 @@ const int32 GREPLACEOnDemandBackendThreadPriorities[5] =
 	EThreadPriority::TPri_AboveNormal
 };
 
-const TCHAR* LexToString(EThreadPriority Priority)
-{
-	switch (Priority)
-	{
-		case EThreadPriority::TPri_Normal:
-			return TEXT("TPri_Normal");
-		case EThreadPriority::TPri_AboveNormal:
-			return TEXT("TPri_AboveNormal");
-		case EThreadPriority::TPri_BelowNormal:
-			return TEXT("TPri_BelowNormal");
-		case EThreadPriority::TPri_Highest:
-			return TEXT("TPri_Highest");
-		case EThreadPriority::TPri_Lowest:
-			return TEXT("TPri_Lowest");
-		case EThreadPriority::TPri_SlightlyBelowNormal:
-			return TEXT("TPri_SlightlyBelowNormal");
-		case EThreadPriority::TPri_TimeCritical:
-			return TEXT("TPri_TimeCritical");
-		case EThreadPriority::TPri_Num:
-		default:
-			return TEXT("TPri_Undefined");
-			break;
-	};
-}
-
 ////////////////////////////////////////
 // END EXTERN CVARS
 ////////////////////////////////////////
