Geometry collection : fix issue where one section with an invalid material for nanite can cause all subsequent section to show the default  materials

diff --git a/Engine/Source/Runtime/Experimental/GeometryCollectionEngine/Private/GeometryCollection/GeometryCollectionSceneProxy.cpp b/Engine/Source/Runtime/Experimental/GeometryCollectionEngine/Private/GeometryCollection/GeometryCollectionSceneProxy.cpp
--- a/Engine/Source/Runtime/Experimental/GeometryCollectionEngine/Private/GeometryCollection/GeometryCollectionSceneProxy.cpp
+++ b/Engine/Source/Runtime/Experimental/GeometryCollectionEngine/Private/GeometryCollection/GeometryCollectionSceneProxy.cpp
@@ -1183,18 +1183,6 @@ FNaniteGeometryCollectionSceneProxy::FNaniteGeometryCollectionSceneProxy(UGeomet
 	// Nanite always uses GPUScene, so we can skip expensive primitive uniform buffer updates.
 	bVFRequiresPrimitiveUniformBuffer = false;
 
-	// Indicates if 1 or more materials contain settings not supported by Nanite.
-	bHasMaterialErrors = false;
-
-	// Check if the assigned material can be rendered in Nanite. If not, default.
-	// TODO: Handle cases like geometry collections adding a "selected geometry" material with translucency.
-	const bool IsRenderable = true;// Nanite::FSceneProxy::IsNaniteRenderable(MaterialRelevance);
-
-	if (!IsRenderable)
-	{
-		bHasMaterialErrors = true;
-	}
-
 	const TSharedPtr<FGeometryCollection, ESPMode::ThreadSafe> Collection = GeometryCollection->GetGeometryCollection();
 	const TManagedArray<int32>& TransformToGeometryIndices = Collection->TransformToGeometryIndex;
 	const TManagedArray<int32>& SimulationType = Collection->SimulationType;
@@ -1216,7 +1204,6 @@ FNaniteGeometryCollectionSceneProxy::FNaniteGeometryCollectionSceneProxy(UGeomet
 		const bool bInvalidMaterial = !MaterialInterface || !IsOpaqueOrMaskedBlendMode(*MaterialInterface) || MaterialInterface->GetShadingModels().HasShadingModel(MSM_SingleLayerWater);
 		if (bInvalidMaterial)
 		{
-			bHasMaterialErrors = true;
 			if (MaterialInterface)
 			{
 				UE_LOG
@@ -1231,9 +1218,9 @@ FNaniteGeometryCollectionSceneProxy::FNaniteGeometryCollectionSceneProxy(UGeomet
 			}
 		}
 
-		const bool bForceDefaultMaterial = /*!!FORCE_NANITE_DEFAULT_MATERIAL ||*/ bHasMaterialErrors;
-		if (bForceDefaultMaterial)
+		if (bInvalidMaterial)
 		{
+			// force default material 
 			MaterialInterface = UMaterial::GetDefaultMaterial(MD_Surface);
 		}
 
