Fix accidentally checking in commented-out code block

diff --git a/Engine/Source/Runtime/Experimental/Dataflow/Engine/Private/Dataflow/DataflowTextureAssetNodes.cpp b/Engine/Source/Runtime/Experimental/Dataflow/Engine/Private/Dataflow/DataflowTextureAssetNodes.cpp
--- a/Engine/Source/Runtime/Experimental/Dataflow/Engine/Private/Dataflow/DataflowTextureAssetNodes.cpp
+++ b/Engine/Source/Runtime/Experimental/Dataflow/Engine/Private/Dataflow/DataflowTextureAssetNodes.cpp
@@ -93,14 +93,14 @@ void FDataflowTextureToImageNode::Evaluate(UE::Dataflow::FContext& Context, cons
 			{
 				CpuCopy->CopyTo(TempImage, ERawImageFormat::RGBA32F, EGammaSpace::Linear);
 			}
-//#if WITH_EDITORONLY_DATA
-//			else if (InTextureAsset->Source.IsValid())
-//			{
-//				FImage MipImage;
-//				InTextureAsset->Source.GetMipImage(MipImage, 0);
-//				MipImage.CopyTo(TempImage, ERawImageFormat::RGBA32F, EGammaSpace::Linear);
-//			}
-//#endif
+#if WITH_EDITORONLY_DATA
+			else if (InTextureAsset->Source.IsValid())
+			{
+				FImage MipImage;
+				InTextureAsset->Source.GetMipImage(MipImage, 0);
+				MipImage.CopyTo(TempImage, ERawImageFormat::RGBA32F, EGammaSpace::Linear);
+			}
+#endif
 			else
 			{
 				// TODO: Handle more ways to get image data from textures
@@ -146,7 +146,14 @@ void FDataflowImageToTextureNode::Evaluate(UE::Dataflow::FContext& Context, cons
 
 		const FDataflowImage InImage = GetValue(Context, &Image);
 
-		DataflowTextureAssetNodes::Private::UpdateTexture2DFromImage(*OutTexture, InImage.GetImage());
+		if (InImage.GetWidth() > 0 && InImage.GetHeight() > 0)
+		{
+			DataflowTextureAssetNodes::Private::UpdateTexture2DFromImage(*OutTexture, InImage.GetImage());
+		}
+		else
+		{
+			Context.Warning(TEXT("Input image is empty"));
+		}
 
 		SetValue(Context, OutTexture, &TransientTexture);
 		return;
