Fix GJK on XSX and Switch
Make GJK more robust and avoid using fma operation in cross product.
More robust cross product, and symetric computation.

diff --git a/Engine/Source/Runtime/Experimental/ChaosCore/Public/Chaos/VectorUtility.h b/Engine/Source/Runtime/Experimental/ChaosCore/Public/Chaos/VectorUtility.h
--- a/Engine/Source/Runtime/Experimental/ChaosCore/Public/Chaos/VectorUtility.h
+++ b/Engine/Source/Runtime/Experimental/ChaosCore/Public/Chaos/VectorUtility.h
@@ -230,18 +230,18 @@ namespace Chaos::Private
 	 */
 	FORCEINLINE VectorRegister4Float VectorCrossNoFMA(const VectorRegister4Float& Vec1, const VectorRegister4Float& Vec2)
 	{
-#if PLATFORM_ENABLE_VECTORINTRINSICS_NEON
-		VectorRegister4Float C = VectorMultiply(Vec1, VectorSwizzle(Vec2, 1, 2, 0, 3));
-		C = VectorSubtract(C, VectorMultiply(VectorSwizzle(Vec1, 1, 2, 0, 3), Vec2));
-		return VectorSwizzle(C, 1, 2, 0, 3);
-#elif PLATFORM_ENABLE_VECTORINTRINSICS
+#if PLATFORM_ENABLE_VECTORINTRINSICS
 		// YZX
 		VectorRegister4Float A = VectorSwizzle(Vec2, 1, 2, 0, 3);
 		VectorRegister4Float B = VectorSwizzle(Vec1, 1, 2, 0, 3);
 		// XY, YZ, ZX
-		A = VectorMultiply(A, Vec1);
+		// This is the only way found to avoid the compiler on XSX using FMA.
+		// By forcing two FMA in a row, the subtract cannot be a FMA at the end.
+		// This allow to have a symmetric and reliable cross product. 
+		A = VectorMultiplyAdd(A, Vec1, VectorZero());
+		B = VectorMultiplyAdd(B, Vec2, VectorZero());
 		// XY-YX, YZ-ZY, ZX-XZ
-		A = VectorSubtract(A, VectorMultiply(B, Vec2));
+		A = VectorSubtract(A, B);
 		// YZ-ZY, ZX-XZ, XY-YX
 		return VectorSwizzle(A, 1, 2, 0, 3);
 #else
