Add WarpWeftScale to FPBDEdgeSpringConstraints

diff --git a/Engine/Plugins/ChaosCloth/Source/ChaosCloth/Private/ChaosCloth/ChaosClothConstraints.cpp b/Engine/Plugins/ChaosCloth/Source/ChaosCloth/Private/ChaosCloth/ChaosClothConstraints.cpp
--- a/Engine/Plugins/ChaosCloth/Source/ChaosCloth/Private/ChaosCloth/ChaosClothConstraints.cpp
+++ b/Engine/Plugins/ChaosCloth/Source/ChaosCloth/Private/ChaosCloth/ChaosClothConstraints.cpp
@@ -799,7 +799,7 @@ void FClothConstraints::AddRules(
 	CreateBendingConstraints(ConfigProperties, WeightMaps, TriangleMesh, PatternData);
 
 	// Area constraints
-	CreateAreaConstraints(ConfigProperties, WeightMaps, TriangleMesh);
+	CreateAreaConstraints(ConfigProperties, WeightMaps, TriangleMesh, PatternData);
 
 	// Long range constraints
 	CreateLongRangeConstraints(ConfigProperties, WeightMaps, Tethers, MeshScale);
@@ -980,12 +980,26 @@ void FClothConstraints::CreateStretchConstraints(
 	{
 		if (Evolution)
 		{
-			EdgeConstraints = MakeShared<Softs::FPBDEdgeSpringConstraints>(
-				Evolution->GetSoftBodyParticles(ParticleRangeId),
-				TriangleMesh.GetSurfaceElements(),
-				WeightMaps,
-				ConfigProperties,
-				/*bTrimKinematicConstraints =*/ true);
+			if (PatternData && PatternData->PatternPositions.Num())
+			{
+				EdgeConstraints = MakeShared<Softs::FPBDEdgeSpringConstraints>(
+					Evolution->GetSoftBodyParticles(ParticleRangeId),
+					TriangleMesh,
+					PatternData->WeldedFaceVertexPatternPositions,
+					WeightMaps,
+					ConfigProperties,
+					/*bTrimKinematicConstraints =*/ true);
+
+			}
+			else
+			{
+				EdgeConstraints = MakeShared<Softs::FPBDEdgeSpringConstraints>(
+					Evolution->GetSoftBodyParticles(ParticleRangeId),
+					TriangleMesh.GetSurfaceElements(),
+					WeightMaps,
+					ConfigProperties,
+					/*bTrimKinematicConstraints =*/ true);
+			}
 		}
 		else
 		{
@@ -1174,7 +1188,8 @@ void FClothConstraints::CreateBendingConstraints(
 void FClothConstraints::CreateAreaConstraints(
 	const Softs::FCollectionPropertyConstFacade& ConfigProperties,
 	const TMap<FString, TConstArrayView<FRealSingle>>& WeightMaps,
-	const FTriangleMesh& TriangleMesh)
+	const FTriangleMesh& TriangleMesh,
+	const FClothingPatternData* PatternData)
 {
 	if (Softs::FXPBDAreaSpringConstraints::IsEnabled(ConfigProperties))
 	{
@@ -1206,12 +1221,26 @@ void FClothConstraints::CreateAreaConstraints(
 	{
 		if (Evolution)
 		{
-			AreaConstraints = MakeShared<Softs::FPBDAreaSpringConstraints>(
-				Evolution->GetSoftBodyParticles(ParticleRangeId),
-				TriangleMesh.GetSurfaceElements(),
-				WeightMaps,
-				ConfigProperties,
-				/*bTrimKinematicConstraints =*/ true);
+			if (PatternData && PatternData->PatternPositions.Num())
+			{
+				AreaConstraints = MakeShared<Softs::FPBDAreaSpringConstraints>(
+					Evolution->GetSoftBodyParticles(ParticleRangeId),
+					TriangleMesh,
+					PatternData->WeldedFaceVertexPatternPositions,
+					WeightMaps,
+					ConfigProperties,
+					/*bTrimKinematicConstraints =*/ true);
+
+			}
+			else
+			{
+				AreaConstraints = MakeShared<Softs::FPBDAreaSpringConstraints>(
+					Evolution->GetSoftBodyParticles(ParticleRangeId),
+					TriangleMesh.GetSurfaceElements(),
+					WeightMaps,
+					ConfigProperties,
+					/*bTrimKinematicConstraints =*/ true);
+			}
 		}
 		else
 		{
diff --git a/Engine/Plugins/ChaosCloth/Source/ChaosCloth/Public/ChaosCloth/ChaosClothConstraints.h b/Engine/Plugins/ChaosCloth/Source/ChaosCloth/Public/ChaosCloth/ChaosClothConstraints.h
--- a/Engine/Plugins/ChaosCloth/Source/ChaosCloth/Public/ChaosCloth/ChaosClothConstraints.h
+++ b/Engine/Plugins/ChaosCloth/Source/ChaosCloth/Public/ChaosCloth/ChaosClothConstraints.h
@@ -148,7 +148,8 @@ namespace Chaos
 		void CreateAreaConstraints(
 			const Softs::FCollectionPropertyConstFacade& ConfigProperties,
 			const TMap<FString, TConstArrayView<FRealSingle>>& WeightMaps,
-			const FTriangleMesh& TriangleMesh);
+			const FTriangleMesh& TriangleMesh,
+			const FClothingPatternData* PatternData);
 		void CreateLongRangeConstraints(
 			const Softs::FCollectionPropertyConstFacade& ConfigProperties,
 			const TMap<FString, TConstArrayView<FRealSingle>>& WeightMaps,
diff --git a/Engine/Plugins/ChaosClothAsset/Source/ChaosClothAssetEngine/Private/ChaosClothAsset/ClothAssetInteractor.cpp b/Engine/Plugins/ChaosClothAsset/Source/ChaosClothAssetEngine/Private/ChaosClothAsset/ClothAssetInteractor.cpp
--- a/Engine/Plugins/ChaosClothAsset/Source/ChaosClothAssetEngine/Private/ChaosClothAsset/ClothAssetInteractor.cpp
+++ b/Engine/Plugins/ChaosClothAsset/Source/ChaosClothAssetEngine/Private/ChaosClothAsset/ClothAssetInteractor.cpp
@@ -91,11 +91,15 @@ namespace UE::Chaos::ClothAsset::Private
 	{
 		TEXT("XPBDAnisoStretchWarpScale"),
 		TEXT("XPBDAnisoSpringWarpScale"),
+		TEXT("EdgeSpringWarpScale"),
+		TEXT("AreaSpringWarpScale"),
 	};
 	static const TArray<FString> StretchWeftScaleAliases =
 	{
 		TEXT("XPBDAnisoStretchWeftScale"),
 		TEXT("XPBDAnisoSpringWeftScale"),
+		TEXT("EdgeSpringWeftScale"),
+		TEXT("AreaSpringWeftScale"),
 	};
 	static const TArray<FString> AreaStiffnessAliases =
 	{
diff --git a/Engine/Plugins/ChaosClothAssetEditor/Source/ChaosClothAssetDataflowNodes/Private/ChaosClothAsset/SimulationStretchConfigNode.cpp b/Engine/Plugins/ChaosClothAssetEditor/Source/ChaosClothAssetDataflowNodes/Private/ChaosClothAsset/SimulationStretchConfigNode.cpp
--- a/Engine/Plugins/ChaosClothAssetEditor/Source/ChaosClothAssetDataflowNodes/Private/ChaosClothAsset/SimulationStretchConfigNode.cpp
+++ b/Engine/Plugins/ChaosClothAssetEditor/Source/ChaosClothAssetDataflowNodes/Private/ChaosClothAsset/SimulationStretchConfigNode.cpp
@@ -75,8 +75,16 @@ void FChaosClothAssetSimulationStretchConfigNode::AddProperties(FPropertyHelper&
 			 }, { FName(TEXT("XPBDEdgeSpringDamping")),
 				FName(TEXT("XPBDAnisoStretchDamping")) });
 
-			PropertyHelper.SetPropertyWeighted(FName(TEXT("XPBDAnisoSpringWarpScale")), StretchWarpScale, { FName(TEXT("XPBDAnisoStretchWarpScale")) });
-			PropertyHelper.SetPropertyWeighted(FName(TEXT("XPBDAnisoSpringWeftScale")), StretchWeftScale, { FName(TEXT("XPBDAnisoStretchWeftScale")) });
+			PropertyHelper.SetPropertyWeighted(FName(TEXT("XPBDAnisoSpringWarpScale")), StretchWarpScale, 
+				{ 
+					FName(TEXT("XPBDAnisoStretchWarpScale")),
+					FName(TEXT("EdgeSpringWarpScale"))
+				});
+			PropertyHelper.SetPropertyWeighted(FName(TEXT("XPBDAnisoSpringWeftScale")), StretchWeftScale, 
+				{ 
+					FName(TEXT("XPBDAnisoStretchWeftScale")),
+					FName(TEXT("EdgeSpringWeftScale"))
+				});
 		}
 		else
 		{
@@ -102,11 +110,31 @@ void FChaosClothAssetSimulationStretchConfigNode::AddProperties(FPropertyHelper&
 			FName(TEXT("XPBDEdgeSpringStiffness")), 
 			FName(TEXT("XPBDAnisoStretchStiffnessWarp")),
 			FName(TEXT("XPBDAnisoSpringStiffnessWarp"))});
+
+		if (bEnableStretchWarpAndWeftScale)
+		{
+			PropertyHelper.SetPropertyWeighted(FName(TEXT("EdgeSpringWarpScale")), StretchWarpScale,
+				{
+					FName(TEXT("XPBDAnisoStretchWarpScale")),
+					FName(TEXT("XPBDAnisoSpringWarpScale"))
+				});
+			PropertyHelper.SetPropertyWeighted(FName(TEXT("EdgeSpringWeftScale")), StretchWeftScale,
+				{
+					FName(TEXT("XPBDAnisoStretchWeftScale")),
+					FName(TEXT("XPBDAnisoSpringWeftScale"))
+				});
+		}
 		
 		if(bAddAreaConstraint)
 		{
 			PropertyHelper.SetPropertyWeighted(FName(TEXT("AreaSpringStiffness")), AreaStiffness,{
 				FName(TEXT("XPBDAreaSpringStiffness"))});
+
+			if (bEnableStretchWarpAndWeftScale)
+			{
+				PropertyHelper.SetPropertyWeighted(FName(TEXT("AreaSpringWarpScale")), StretchWarpScale);
+				PropertyHelper.SetPropertyWeighted(FName(TEXT("AreaSpringWeftScale")), StretchWeftScale);
+			}
 		}
 	}
 }
diff --git a/Engine/Plugins/ChaosClothAssetEditor/Source/ChaosClothAssetDataflowNodes/Public/ChaosClothAsset/SimulationStretchConfigNode.h b/Engine/Plugins/ChaosClothAssetEditor/Source/ChaosClothAssetDataflowNodes/Public/ChaosClothAsset/SimulationStretchConfigNode.h
--- a/Engine/Plugins/ChaosClothAssetEditor/Source/ChaosClothAssetDataflowNodes/Public/ChaosClothAsset/SimulationStretchConfigNode.h
+++ b/Engine/Plugins/ChaosClothAssetEditor/Source/ChaosClothAssetDataflowNodes/Public/ChaosClothAsset/SimulationStretchConfigNode.h
@@ -87,13 +87,20 @@ public:
 	UPROPERTY(EditAnywhere, Category = "Stretch Properties", Meta = (UIMin = "0", UIMax = "10", ClampMin = "0", ClampMax = "1000", EditCondition = "SolverType == EChaosClothAssetConstraintSolverType::XPBD && DistributionType == EChaosClothAssetConstraintDistributionType::Isotropic", EditConditionHides, InteractorName = "StretchDamping"))
 	FChaosClothAssetWeightedValue StretchDamping = { true, 1.f, 1.f, TEXT("StretchDamping") };
 
+	/**
+	* Enable StretchWarp and WeftScale.
+	* This adds a small amount of memory and computational cost, so this is optional for the least expensive constraint type (PBD)
+	*/
+	UPROPERTY(EditAnywhere, Category = "Stretch Properties", Meta = (EditCondition = "SolverType == EChaosClothAssetConstraintSolverType::PBD", EditConditionHides))
+	bool bEnableStretchWarpAndWeftScale = false;
+
 	/**
 	 * The scale of the stretch constraints at rest in the warp (vertical) direction.
 	 * If a valid weight map is found with the given Weight Map name, then both Low and High values
 	 * are interpolated with the per particle weight to make the final value used for the simulation.
 	 * Otherwise all particles are considered to have a zero weight, and only the Low value is meaningful.
 	 */
-	UPROPERTY(EditAnywhere, Category = "Stretch Properties", Meta = (UIMin = "0", UIMax = "10", ClampMin = "0", ClampMax = "10000000", EditCondition = "(SolverType == EChaosClothAssetConstraintSolverType::XPBD && DistributionType == EChaosClothAssetConstraintDistributionType::Anisotropic)", EditConditionHides, InteractorName = "StretchWarpScale"))
+	UPROPERTY(EditAnywhere, Category = "Stretch Properties", Meta = (UIMin = "0", UIMax = "10", ClampMin = "0", ClampMax = "10000000", EditCondition = "(SolverType == EChaosClothAssetConstraintSolverType::XPBD && DistributionType == EChaosClothAssetConstraintDistributionType::Anisotropic) || (bEnableStretchWarpAndWeftScale && SolverType == EChaosClothAssetConstraintSolverType::PBD)", EditConditionHides, InteractorName = "StretchWarpScale"))
 	FChaosClothAssetWeightedValue StretchWarpScale = { true, 1.f, 1.f, TEXT("StretchWarpScale") };
 
 	/**
@@ -102,7 +109,7 @@ public:
 	 * are interpolated with the per particle weight to make the final value used for the simulation.
 	 * Otherwise all particles are considered to have a zero weight, and only the Low value is meaningful.
 	 */
-	UPROPERTY(EditAnywhere, Category = "Stretch Properties", Meta = (UIMin = "0", UIMax = "10", ClampMin = "0", ClampMax = "10000000", EditCondition = "(SolverType == EChaosClothAssetConstraintSolverType::XPBD && DistributionType == EChaosClothAssetConstraintDistributionType::Anisotropic)", EditConditionHides, InteractorName = "StretchWeftScale"))
+	UPROPERTY(EditAnywhere, Category = "Stretch Properties", Meta = (UIMin = "0", UIMax = "10", ClampMin = "0", ClampMax = "10000000", EditCondition = "(SolverType == EChaosClothAssetConstraintSolverType::XPBD && DistributionType == EChaosClothAssetConstraintDistributionType::Anisotropic) || (bEnableStretchWarpAndWeftScale && SolverType == EChaosClothAssetConstraintSolverType::PBD)", EditConditionHides, InteractorName = "StretchWeftScale"))
 	FChaosClothAssetWeightedValue StretchWeftScale = { true, 1.f, 1.f, TEXT("StretchWeftScale") };
 
 	/**
diff --git a/Engine/Source/Runtime/Experimental/Chaos/Private/Chaos/PBDAxialSpringConstraints.cpp b/Engine/Source/Runtime/Experimental/Chaos/Private/Chaos/PBDAxialSpringConstraints.cpp
--- a/Engine/Source/Runtime/Experimental/Chaos/Private/Chaos/PBDAxialSpringConstraints.cpp
+++ b/Engine/Source/Runtime/Experimental/Chaos/Private/Chaos/PBDAxialSpringConstraints.cpp
@@ -2,6 +2,7 @@
 #include "Chaos/PBDAxialSpringConstraints.h"
 #include "Chaos/GraphColoring.h"
 #include "Chaos/Framework/Parallel.h"
+#include "Chaos/TriangleMesh.h"
 #include "ChaosStats.h"
 #if INTEL_ISPC
 #include "PBDAxialSpringConstraints.ispc.generated.h"
@@ -26,8 +27,10 @@ int32 Chaos_AxialSpring_ParallelConstraintCount = 100;
 FAutoConsoleVariableRef CVarChaosAxialSpringParallelConstraintCount(TEXT("p.Chaos.AxialSpring.ParallelConstraintCount"), Chaos_AxialSpring_ParallelConstraintCount, TEXT("If we have more constraints than this, use parallel-for in Apply."));
 
 template<typename SolverParticlesOrRange>
-void FPBDAxialSpringConstraints::InitColor(const SolverParticlesOrRange& InParticles)
+TArray<int32> FPBDAxialSpringConstraints::InitColor(const SolverParticlesOrRange& InParticles)
 {
+	TArray<int32> OrigToReorderedIndices; // used to reorder stiffness indices
+
 	// In dev builds we always color so we can tune the system without restarting. See Apply()
 #if UE_BUILD_SHIPPING || UE_BUILD_TEST
 	if (Constraints.Num() > Chaos_AxialSpring_ParallelConstraintCount)
@@ -39,7 +42,6 @@ void FPBDAxialSpringConstraints::InitColor(const SolverParticlesOrRange& InParti
 		TArray<TVec3<int32>> ReorderedConstraints;
 		TArray<FSolverReal> ReorderedBarys;
 		TArray<FSolverReal> ReorderedDists;
-		TArray<int32> OrigToReorderedIndices; // used to reorder stiffness indices
 		ReorderedConstraints.SetNumUninitialized(Constraints.Num());
 		ReorderedBarys.SetNumUninitialized(Constraints.Num());
 		ReorderedDists.SetNumUninitialized(Constraints.Num());
@@ -68,9 +70,10 @@ void FPBDAxialSpringConstraints::InitColor(const SolverParticlesOrRange& InParti
 		Dists = MoveTemp(ReorderedDists);
 		Stiffness.ReorderIndices(OrigToReorderedIndices);
 	}
+	return OrigToReorderedIndices;
 }
-template CHAOS_API void FPBDAxialSpringConstraints::InitColor(const FSolverParticles& InParticles);
-template CHAOS_API void FPBDAxialSpringConstraints::InitColor(const FSolverParticlesRange& InParticles);
+template CHAOS_API TArray<int32> FPBDAxialSpringConstraints::InitColor(const FSolverParticles& InParticles);
+template CHAOS_API TArray<int32> FPBDAxialSpringConstraints::InitColor(const FSolverParticlesRange& InParticles);
 
 template<typename SolverParticlesOrRange>
 void FPBDAxialSpringConstraints::ApplyHelper(SolverParticlesOrRange& Particles, const FSolverReal Dt, const int32 ConstraintIndex, const FSolverReal ExpStiffnessValue) const
@@ -199,6 +202,115 @@ void FPBDAxialSpringConstraints::Apply(SolverParticlesOrRange& Particles, const
 template CHAOS_API void FPBDAxialSpringConstraints::Apply(FSolverParticles& Particles, const FSolverReal Dt) const;
 template CHAOS_API void FPBDAxialSpringConstraints::Apply(FSolverParticlesRange& Particles, const FSolverReal Dt) const;
 
+FPBDAreaSpringConstraints::FPBDAreaSpringConstraints(
+	const FSolverParticlesRange& Particles,
+	const FTriangleMesh& TriangleMesh,
+	const TArray<TVec3<FVec2f>>& FaceVertexPatternPositions,
+	const TMap<FString, TConstArrayView<FRealSingle>>& WeightMaps,
+	const FCollectionPropertyConstFacade& PropertyCollection,
+	bool bTrimKinematicConstraints)
+	: FPBDAxialSpringConstraints(
+		Particles,
+		TriangleMesh.GetElements(),
+		WeightMaps.FindRef(GetAreaSpringStiffnessString(PropertyCollection, AreaSpringStiffnessName.ToString())),
+		FSolverVec2(GetWeightedFloatAreaSpringStiffness(PropertyCollection, 1.f)),
+		bTrimKinematicConstraints,
+		false /*bInitColor*/
+	)
+	, WarpScale(FSolverVec2(GetWeightedFloatAreaSpringWarpScale(PropertyCollection, 1.f)),
+		WeightMaps.FindRef(GetAreaSpringWarpScaleString(PropertyCollection, AreaSpringWarpScaleName.ToString())),
+		TConstArrayView<TVec3<int32>>(Constraints),
+		ParticleOffset,
+		ParticleCount)
+	, WeftScale(FSolverVec2(GetWeightedFloatAreaSpringWeftScale(PropertyCollection, 1.f)),
+		WeightMaps.FindRef(GetAreaSpringWeftScaleString(PropertyCollection, AreaSpringWeftScaleName.ToString())),
+		TConstArrayView<TVec3<int32>>(Constraints),
+		ParticleOffset,
+		ParticleCount)
+	, AreaSpringStiffnessIndex(PropertyCollection)
+	, AreaSpringWarpScaleIndex(PropertyCollection)
+	, AreaSpringWeftScaleIndex(PropertyCollection)
+{
+	if (AreaSpringWarpScaleIndex != INDEX_NONE || AreaSpringWeftScaleIndex != INDEX_NONE)
+	{
+		InitFromPatternData(FaceVertexPatternPositions, TriangleMesh);
+	}
+
+	const TArray<int32> OrigToReorderedIndices = InitColor(Particles);
+	if (bWarpWeftScaleEnabled && OrigToReorderedIndices.Num() == Constraints.Num())
+	{
+		TArray<FSolverReal> ReorderedBaseDists;
+		TArray<FSolverVec2> ReorderedWarpWeftScaleBaseMultipliers;
+		ReorderedBaseDists.SetNumUninitialized(Constraints.Num());
+		ReorderedWarpWeftScaleBaseMultipliers.SetNumUninitialized(Constraints.Num());
+		for (int32 OrigIndex = 0; OrigIndex < Constraints.Num(); ++OrigIndex)
+		{
+			const int32 ReorderedIndex = OrigToReorderedIndices[OrigIndex];
+			ReorderedBaseDists[ReorderedIndex] = BaseDists[OrigIndex];
+			ReorderedWarpWeftScaleBaseMultipliers[ReorderedIndex] = WarpWeftScaleBaseMultipliers[OrigIndex];
+		}
+
+		BaseDists = MoveTemp(ReorderedBaseDists);
+		WarpWeftScaleBaseMultipliers = MoveTemp(ReorderedWarpWeftScaleBaseMultipliers);
+		WarpScale.ReorderIndices(OrigToReorderedIndices);
+		WeftScale.ReorderIndices(OrigToReorderedIndices);
+	}
+}
+
+void FPBDAreaSpringConstraints::InitFromPatternData(const TArray<TVec3<FVec2f>>& FaceVertexPatternPositions, const FTriangleMesh& TriangleMesh)
+{
+	BaseDists = Dists; // Use Dists calculated by base class using 3D positions
+
+	WarpWeftScaleBaseMultipliers.SetNumUninitialized(Constraints.Num());
+
+	TConstArrayView<TArray<int32>> PointToTriangleMap = TriangleMesh.GetPointToTriangleMap();
+	auto SortedElement = [](const TVec3<int32>& Element)
+		{
+			return TVec3<int32>(Element.Min(), Element.Mid(), Element.Max());
+		};
+	const TArray<TVec3<int32>>& Elements = TriangleMesh.GetElements();
+
+	for (int32 ConstraintIdx = 0; ConstraintIdx < Constraints.Num(); ++ConstraintIdx)
+	{
+		const TVec3<int32>& Constraint = Constraints[ConstraintIdx];
+		// Find corresponding triangle in original mesh (constraints get reordered by Base class init)
+		const TArray<int32>& PossibleTriangles = PointToTriangleMap[Constraint[0]];
+		const TVec3<int32> SortedConstraint = SortedElement(Constraint);
+		int32 TriangleIndex = INDEX_NONE;
+		for (const int32 Triangle : PossibleTriangles)
+		{
+			if (SortedConstraint == SortedElement(Elements[Triangle]))
+			{
+				TriangleIndex = Triangle;
+				break;
+			}
+		}
+		check(TriangleIndex != INDEX_NONE);
+
+		const TVec3<int32>& Element = Elements[TriangleIndex];
+		auto MatchIndex = [&Constraint, &Element](const int32 ConstraintAxis)
+			{
+				check(Element[0] == Constraint[ConstraintAxis] || Element[1] == Constraint[ConstraintAxis] || Element[2] == Constraint[ConstraintAxis]);
+				return Element[0] == Constraint[ConstraintAxis] ? 0 : Element[1] == Constraint[ConstraintAxis] ? 1 : 2;
+			};
+		const int32 FaceIndex1 = MatchIndex(0);
+		const int32 FaceIndex2 = MatchIndex(1);
+		const int32 FaceIndex3 = MatchIndex(2);
+
+		const FVec2f& UV1 = FaceVertexPatternPositions[TriangleIndex][FaceIndex1];
+		const FVec2f& UV2 = FaceVertexPatternPositions[TriangleIndex][FaceIndex2];
+		const FVec2f& UV3 = FaceVertexPatternPositions[TriangleIndex][FaceIndex3];
+
+		const FVec2f UV = (UV2 - UV3) * Barys[ConstraintIdx] + UV3;
+		const FSolverVec2 UVDiff = UV - UV1;
+		const FSolverVec2 UVDiffAbs = UVDiff.GetAbs();
+		const FSolverReal UVLength = UVDiffAbs.Length();
+		WarpWeftScaleBaseMultipliers[ConstraintIdx] = UVLength > UE_SMALL_NUMBER ? UVDiffAbs / UVLength : FSolverVec2(UE_INV_SQRT_2, UE_INV_SQRT_2); // Default to equally scaling warp and weft directions if zero length
+	}
+
+	bWarpWeftScaleEnabled = true;
+}
+
 void FPBDAreaSpringConstraints::SetProperties(
 	const FCollectionPropertyConstFacade& PropertyCollection,
 	const TMap<FString, TConstArrayView<FRealSingle>>& WeightMaps)
@@ -221,6 +333,81 @@ void FPBDAreaSpringConstraints::SetProperties(
 			Stiffness.SetWeightedValue(WeightedValue);
 		}
 	}
+	if (bWarpWeftScaleEnabled)
+	{
+		if (IsAreaSpringWarpScaleMutable(PropertyCollection))
+		{
+			const FSolverVec2 WeightedValue(GetWeightedFloatAreaSpringWarpScale(PropertyCollection));
+			if (IsAreaSpringWarpScaleStringDirty(PropertyCollection))
+			{
+				const FString& WeightMapName = GetAreaSpringWarpScaleString(PropertyCollection);
+				WarpScale = FPBDWeightMap(
+					WeightedValue,
+					WeightMaps.FindRef(WeightMapName),
+					TConstArrayView<TVec3<int32>>(Constraints),
+					ParticleOffset,
+					ParticleCount);
+			}
+			else
+			{
+				WarpScale.SetWeightedValue(WeightedValue);
+			}
+		}
+		if (IsAreaSpringWeftScaleMutable(PropertyCollection))
+		{
+			const FSolverVec2 WeightedValue(GetWeightedFloatAreaSpringWeftScale(PropertyCollection));
+			if (IsAreaSpringWeftScaleStringDirty(PropertyCollection))
+			{
+				const FString& WeightMapName = GetAreaSpringWeftScaleString(PropertyCollection);
+				WeftScale = FPBDWeightMap(
+					WeightedValue,
+					WeightMaps.FindRef(WeightMapName),
+					TConstArrayView<TVec3<int32>>(Constraints),
+					ParticleOffset,
+					ParticleCount);
+			}
+			else
+			{
+				WeftScale.SetWeightedValue(WeightedValue);
+			}
+		}
+	}
+}
+
+void FPBDAreaSpringConstraints::ApplyProperties(const FSolverReal Dt, const int32 NumIterations)
+{
+	FPBDAxialSpringConstraints::ApplyProperties(Dt, NumIterations);
+
+	if (bWarpWeftScaleEnabled)
+	{
+		bool bWarpScaleChanged = false;
+		WarpScale.ApplyValues(&bWarpScaleChanged);
+		bool bWeftScaleChanged = false;
+		WeftScale.ApplyValues(&bWeftScaleChanged);
+		if (bWarpScaleChanged || bWeftScaleChanged)
+		{
+			// Need to update distances
+			UpdateDists();
+		}
+	}
+}
+
+void FPBDAreaSpringConstraints::UpdateDists()
+{
+	if (bWarpWeftScaleEnabled)
+	{
+		const bool WarpScaleHasWeightMap = WarpScale.HasWeightMap();
+		const bool WeftScaleHasWeightMap = WeftScale.HasWeightMap();
+		const FSolverReal WarpScaleNoMap = (FSolverReal)WarpScale;
+		const FSolverReal WeftScaleNoMap = (FSolverReal)WeftScale;
+		for (int32 ConstraintIndex = 0; ConstraintIndex < BaseDists.Num(); ++ConstraintIndex)
+		{
+			const FSolverReal WarpScaleValue = WarpScaleHasWeightMap ? WarpScale[ConstraintIndex] : WarpScaleNoMap;
+			const FSolverReal WeftScaleValue = WeftScaleHasWeightMap ? WeftScale[ConstraintIndex] : WeftScaleNoMap;
+
+			Dists[ConstraintIndex] = BaseDists[ConstraintIndex] * FMath::Sqrt(FMath::Square(WeftScaleValue * WarpWeftScaleBaseMultipliers[ConstraintIndex][0]) + FMath::Square(WarpScaleValue * WarpWeftScaleBaseMultipliers[ConstraintIndex][1]));
+		}
+	}
 }
 
 }  // End namespace Chaos::Softs
diff --git a/Engine/Source/Runtime/Experimental/Chaos/Private/Chaos/PBDSpringConstraints.cpp b/Engine/Source/Runtime/Experimental/Chaos/Private/Chaos/PBDSpringConstraints.cpp
--- a/Engine/Source/Runtime/Experimental/Chaos/Private/Chaos/PBDSpringConstraints.cpp
+++ b/Engine/Source/Runtime/Experimental/Chaos/Private/Chaos/PBDSpringConstraints.cpp
@@ -4,6 +4,7 @@
 #include "Chaos/PBDRigidParticles.h"
 #include "Chaos/GraphColoring.h"
 #include "Chaos/Framework/Parallel.h"
+#include "Chaos/TriangleMesh.h"
 #include "ChaosStats.h"
 #include "ChaosLog.h"
 #include "HAL/IConsoleManager.h"
@@ -33,8 +34,10 @@ FAutoConsoleVariableRef CVarChaosSpringParallelConstraintCount(TEXT("p.Chaos.Spr
 #endif
 
 template<typename SolverParticlesOrRange>
-void FPBDSpringConstraints::InitColor(const SolverParticlesOrRange& Particles)
+TArray<int32> FPBDSpringConstraints::InitColor(const SolverParticlesOrRange& Particles)
 {
+	TArray<int32> OrigToReorderedIndices; // used to reorder stiffness indices
+	
 	// In dev builds we always color so we can tune the system without restarting. See Apply()
 #if UE_BUILD_SHIPPING || UE_BUILD_TEST
 	if (Constraints.Num() > Chaos_Spring_ParallelConstraintCount)
@@ -45,7 +48,6 @@ void FPBDSpringConstraints::InitColor(const SolverParticlesOrRange& Particles)
 		// Reorder constraints based on color so each array in ConstraintsPerColor contains contiguous elements.
 		TArray<TVec2<int32>> ReorderedConstraints;
 		TArray<FSolverReal> ReorderedDists;
-		TArray<int32> OrigToReorderedIndices; // used to reorder stiffness indices
 		ReorderedConstraints.SetNumUninitialized(Constraints.Num());
 		ReorderedDists.SetNumUninitialized(Dists.Num());
 		OrigToReorderedIndices.SetNumUninitialized(Constraints.Num());
@@ -72,9 +74,10 @@ void FPBDSpringConstraints::InitColor(const SolverParticlesOrRange& Particles)
 		Dists = MoveTemp(ReorderedDists);
 		Stiffness.ReorderIndices(OrigToReorderedIndices);
 	}
+	return OrigToReorderedIndices;
 }
-template CHAOS_API void FPBDSpringConstraints::InitColor(const FSolverParticles& Particles);
-template CHAOS_API void FPBDSpringConstraints::InitColor(const FSolverParticlesRange& Particles);
+template CHAOS_API TArray<int32> FPBDSpringConstraints::InitColor(const FSolverParticles& Particles);
+template CHAOS_API TArray<int32> FPBDSpringConstraints::InitColor(const FSolverParticlesRange& Particles);
 
 template<typename SolverParticlesOrRange>
 void FPBDSpringConstraints::ApplyHelper(SolverParticlesOrRange& Particles, const FSolverReal Dt, const int32 ConstraintIndex, const FSolverReal ExpStiffnessValue) const
@@ -193,6 +196,100 @@ void FPBDSpringConstraints::Apply(SolverParticlesOrRange& Particles, const FSolv
 template CHAOS_API void FPBDSpringConstraints::Apply(FSolverParticles& Particles, const FSolverReal Dt) const;
 template CHAOS_API void FPBDSpringConstraints::Apply(FSolverParticlesRange& Particles, const FSolverReal Dt) const;
 
+FPBDEdgeSpringConstraints::FPBDEdgeSpringConstraints(
+	const FSolverParticlesRange& Particles,
+	const FTriangleMesh& TriangleMesh,
+	const TArray<TVec3<FVec2f>>& FaceVertexPatternPositions,
+	const TMap<FString, TConstArrayView<FRealSingle>>& WeightMaps,
+	const FCollectionPropertyConstFacade& PropertyCollection,
+	bool bTrimKinematicConstraints)
+	: FPBDSpringConstraints(
+		Particles,
+		TriangleMesh.GetElements(),
+		WeightMaps.FindRef(GetEdgeSpringStiffnessString(PropertyCollection, EdgeSpringStiffnessName.ToString())),
+		FSolverVec2(GetWeightedFloatEdgeSpringStiffness(PropertyCollection, 1.f)),
+		bTrimKinematicConstraints,
+		false /*bInitColor*/)
+	, WarpScale(FSolverVec2(GetWeightedFloatEdgeSpringWarpScale(PropertyCollection, 1.f)),
+		WeightMaps.FindRef(GetEdgeSpringWarpScaleString(PropertyCollection, EdgeSpringWarpScaleName.ToString())),
+		TConstArrayView<TVec2<int32>>(Constraints),
+		ParticleOffset,
+		ParticleCount)
+	, WeftScale(FSolverVec2(GetWeightedFloatEdgeSpringWeftScale(PropertyCollection, 1.f)),
+		WeightMaps.FindRef(GetEdgeSpringWeftScaleString(PropertyCollection, EdgeSpringWeftScaleName.ToString())),
+		TConstArrayView<TVec2<int32>>(Constraints),
+		ParticleOffset,
+		ParticleCount)
+	, EdgeSpringStiffnessIndex(PropertyCollection)
+	, EdgeSpringWarpScaleIndex(PropertyCollection)
+	, EdgeSpringWeftScaleIndex(PropertyCollection)
+{
+	if (EdgeSpringWarpScaleIndex != INDEX_NONE || EdgeSpringWeftScaleIndex != INDEX_NONE)
+	{
+		InitFromPatternData(FaceVertexPatternPositions, TriangleMesh);
+	}
+
+	const TArray<int32> OrigToReorderedIndices = InitColor(Particles);
+	if (bWarpWeftScaleEnabled && OrigToReorderedIndices.Num() == Constraints.Num())
+	{
+		TArray<FSolverReal> ReorderedBaseDists;
+		TArray<FSolverVec2> ReorderedWarpWeftScaleBaseMultipliers;
+		ReorderedBaseDists.SetNumUninitialized(Constraints.Num());
+		ReorderedWarpWeftScaleBaseMultipliers.SetNumUninitialized(Constraints.Num());
+		for (int32 OrigIndex = 0; OrigIndex < Constraints.Num(); ++OrigIndex)
+		{
+			const int32 ReorderedIndex = OrigToReorderedIndices[OrigIndex];
+			ReorderedBaseDists[ReorderedIndex] = BaseDists[OrigIndex];
+			ReorderedWarpWeftScaleBaseMultipliers[ReorderedIndex] = WarpWeftScaleBaseMultipliers[OrigIndex];
+		}
+
+		BaseDists = MoveTemp(ReorderedBaseDists);
+		WarpWeftScaleBaseMultipliers = MoveTemp(ReorderedWarpWeftScaleBaseMultipliers);
+		WarpScale.ReorderIndices(OrigToReorderedIndices);
+		WeftScale.ReorderIndices(OrigToReorderedIndices);
+	}
+}
+
+void FPBDEdgeSpringConstraints::InitFromPatternData(const TArray<TVec3<FVec2f>>& FaceVertexPatternPositions, const FTriangleMesh& TriangleMesh)
+{
+	// Calculate multipliers per face. Average the multipliers for faces that share the same edge.
+	TMap<TVec2<int32> /*Edge*/, TArray<FSolverVec2>> EdgeBasedWarpWeftScaleBaseMultiplier;
+	auto SortedEdge = [](int32 P0, int32 P1) { return P0 <= P1 ? TVec2<int32>(P0, P1) : TVec2<int32>(P1, P0); };
+	auto Multiplier = [](const FVec2f& UV0, const FVec2f& UV1)
+		{
+			const FSolverVec2 UVDiff = UV1 - UV0;
+			const FSolverVec2 UVDiffAbs = UVDiff.GetAbs();
+			const FSolverReal UVLength = UVDiffAbs.Length();
+			return UVLength > UE_SMALL_NUMBER ? UVDiffAbs / UVLength : FSolverVec2(UE_INV_SQRT_2, UE_INV_SQRT_2); // Default to equally scaling warp and weft directions if zero length
+		};
+
+	const TArray<TVec3<int32>>& Elements = TriangleMesh.GetElements();
+	for (int32 ElemIdx = 0; ElemIdx < Elements.Num(); ++ElemIdx)
+	{
+		const TVec3<int32>& Element = Elements[ElemIdx];
+		const TVec3<FVec2f>& UVs = FaceVertexPatternPositions[ElemIdx];
+		EdgeBasedWarpWeftScaleBaseMultiplier.FindOrAdd(SortedEdge(Element[0], Element[1])).Add(Multiplier(UVs[0], UVs[1]));
+		EdgeBasedWarpWeftScaleBaseMultiplier.FindOrAdd(SortedEdge(Element[1], Element[2])).Add(Multiplier(UVs[1], UVs[2]));
+		EdgeBasedWarpWeftScaleBaseMultiplier.FindOrAdd(SortedEdge(Element[2], Element[0])).Add(Multiplier(UVs[2], UVs[0]));
+	}
+
+	BaseDists = Dists;
+
+	WarpWeftScaleBaseMultipliers.SetNumZeroed(Constraints.Num());
+	for (int32 ConstraintIdx = 0; ConstraintIdx < Constraints.Num(); ++ConstraintIdx)
+	{
+		const TArray<FSolverVec2>& EdgeMultipliers = EdgeBasedWarpWeftScaleBaseMultiplier.FindChecked(SortedEdge(Constraints[ConstraintIdx][0], Constraints[ConstraintIdx][1]));
+		check(EdgeMultipliers.Num());
+		for (const FSolverVec2& EdgeMult : EdgeMultipliers)
+		{
+			WarpWeftScaleBaseMultipliers[ConstraintIdx] += EdgeMult;
+		}
+		WarpWeftScaleBaseMultipliers[ConstraintIdx].Normalize();
+	}
+
+	bWarpWeftScaleEnabled = true;
+}
+
 void FPBDEdgeSpringConstraints::SetProperties(
 	const FCollectionPropertyConstFacade& PropertyCollection,
 	const TMap<FString, TConstArrayView<FRealSingle>>& WeightMaps)
@@ -215,6 +312,81 @@ void FPBDEdgeSpringConstraints::SetProperties(
 			Stiffness.SetWeightedValue(WeightedValue);
 		}
 	}
+	if (bWarpWeftScaleEnabled)
+	{
+		if (IsEdgeSpringWarpScaleMutable(PropertyCollection))
+		{
+			const FSolverVec2 WeightedValue(GetWeightedFloatEdgeSpringWarpScale(PropertyCollection));
+			if (IsEdgeSpringWarpScaleStringDirty(PropertyCollection))
+			{
+				const FString& WeightMapName = GetEdgeSpringWarpScaleString(PropertyCollection);
+				WarpScale = FPBDWeightMap(
+					WeightedValue,
+					WeightMaps.FindRef(WeightMapName),
+					TConstArrayView<TVec2<int32>>(Constraints),
+					ParticleOffset,
+					ParticleCount);
+			}
+			else
+			{
+				WarpScale.SetWeightedValue(WeightedValue);
+			}
+		}
+		if (IsEdgeSpringWeftScaleMutable(PropertyCollection))
+		{
+			const FSolverVec2 WeightedValue(GetWeightedFloatEdgeSpringWeftScale(PropertyCollection));
+			if (IsEdgeSpringWeftScaleStringDirty(PropertyCollection))
+			{
+				const FString& WeightMapName = GetEdgeSpringWeftScaleString(PropertyCollection);
+				WeftScale = FPBDWeightMap(
+					WeightedValue,
+					WeightMaps.FindRef(WeightMapName),
+					TConstArrayView<TVec2<int32>>(Constraints),
+					ParticleOffset,
+					ParticleCount);
+			}
+			else
+			{
+				WeftScale.SetWeightedValue(WeightedValue);
+			}
+		}
+	}
+}
+
+void FPBDEdgeSpringConstraints::ApplyProperties(const FSolverReal Dt, const int32 NumIterations)
+{
+	FPBDSpringConstraints::ApplyProperties(Dt, NumIterations);
+
+	if (bWarpWeftScaleEnabled)
+	{
+		bool bWarpScaleChanged = false;
+		WarpScale.ApplyValues(&bWarpScaleChanged);
+		bool bWeftScaleChanged = false;
+		WeftScale.ApplyValues(&bWeftScaleChanged);
+		if (bWarpScaleChanged || bWeftScaleChanged)
+		{
+			// Need to update distances
+			UpdateDists();
+		}
+	}
+}
+
+void FPBDEdgeSpringConstraints::UpdateDists()
+{
+	if (bWarpWeftScaleEnabled)
+	{
+		const bool WarpScaleHasWeightMap = WarpScale.HasWeightMap();
+		const bool WeftScaleHasWeightMap = WeftScale.HasWeightMap();
+		const FSolverReal WarpScaleNoMap = (FSolverReal)WarpScale;
+		const FSolverReal WeftScaleNoMap = (FSolverReal)WeftScale;
+		for (int32 ConstraintIndex = 0; ConstraintIndex < BaseDists.Num(); ++ConstraintIndex)
+		{
+			const FSolverReal WarpScaleValue = WarpScaleHasWeightMap ? WarpScale[ConstraintIndex] : WarpScaleNoMap;
+			const FSolverReal WeftScaleValue = WeftScaleHasWeightMap ? WeftScale[ConstraintIndex] : WeftScaleNoMap;
+
+			Dists[ConstraintIndex] = BaseDists[ConstraintIndex] * FMath::Sqrt(FMath::Square(WeftScaleValue * WarpWeftScaleBaseMultipliers[ConstraintIndex][0]) + FMath::Square(WarpScaleValue * WarpWeftScaleBaseMultipliers[ConstraintIndex][1]));
+		}
+	}
 }
 
 void FPBDBendingSpringConstraints::SetProperties(
diff --git a/Engine/Source/Runtime/Experimental/Chaos/Public/Chaos/PBDAxialSpringConstraints.h b/Engine/Source/Runtime/Experimental/Chaos/Public/Chaos/PBDAxialSpringConstraints.h
--- a/Engine/Source/Runtime/Experimental/Chaos/Public/Chaos/PBDAxialSpringConstraints.h
+++ b/Engine/Source/Runtime/Experimental/Chaos/Public/Chaos/PBDAxialSpringConstraints.h
@@ -6,13 +6,17 @@
 #include "Chaos/PBDAxialSpringConstraintsBase.h"
 #include "Chaos/CollectionPropertyFacade.h"
 
+namespace Chaos
+{
+	class FTriangleMesh;
+}
+
 namespace Chaos::Softs
 {
 
 class FPBDAxialSpringConstraints : public FPBDAxialSpringConstraintsBase
 {
 	typedef FPBDAxialSpringConstraintsBase Base;
-	using Base::Barys;
 
 public:
 	FPBDAxialSpringConstraints(
@@ -20,7 +24,8 @@ public:
 		const TArray<TVec3<int32>>& InConstraints,
 		const TConstArrayView<FRealSingle>& StiffnessMultipliers,
 		const FSolverVec2& InStiffness,
-		bool bTrimKinematicConstraints)
+		bool bTrimKinematicConstraints,
+		bool bInitColor = true)
 		: Base(
 			Particles,
 			InConstraints,
@@ -28,7 +33,10 @@ public:
 			InStiffness,
 			bTrimKinematicConstraints)
 	{
-		InitColor(Particles);
+		if (bInitColor)
+		{
+			InitColor(Particles);
+		}
 	}
 
 	FPBDAxialSpringConstraints(
@@ -38,7 +46,8 @@ public:
 		const TArray<TVec3<int32>>& InConstraints,
 		const TConstArrayView<FRealSingle>& StiffnessMultipliers,
 		const FSolverVec2& InStiffness,
-		bool bTrimKinematicConstraints)
+		bool bTrimKinematicConstraints,
+		bool bInitColor = true)
 		: Base(
 			Particles,
 			InParticleOffset,
@@ -48,7 +57,10 @@ public:
 			InStiffness,
 			bTrimKinematicConstraints)
 	{
-		InitColor(Particles);
+		if (bInitColor)
+		{
+			InitColor(Particles);
+		}
 	}
 
 	virtual ~FPBDAxialSpringConstraints() override {}
@@ -58,13 +70,15 @@ public:
 
 protected:
 	using Base::Constraints;
+	using Base::Barys;
 	using Base::Stiffness;
 	using Base::ParticleOffset;
 	using Base::ParticleCount;
 
-private:
 	template<typename SolverParticlesOrRange>
-	CHAOS_API void InitColor(const SolverParticlesOrRange& InParticles);
+	CHAOS_API TArray<int32> InitColor(const SolverParticlesOrRange& InParticles);
+
+private:
 	template<typename SolverParticlesOrRange>
 	void ApplyHelper(SolverParticlesOrRange& InParticles, const FSolverReal Dt, const int32 ConstraintIndex, const FSolverReal ExpStiffnessValue) const;
 
@@ -91,7 +105,11 @@ public:
 			WeightMaps.FindRef(GetAreaSpringStiffnessString(PropertyCollection, AreaSpringStiffnessName.ToString())),
 			FSolverVec2(GetWeightedFloatAreaSpringStiffness(PropertyCollection, 1.f)),
 			bTrimKinematicConstraints)
+		, WarpScale(FSolverVec2(1.f))
+		, WeftScale(FSolverVec2(1.f))
 		, AreaSpringStiffnessIndex(PropertyCollection)
+		, AreaSpringWarpScaleIndex(ForceInit)
+		, AreaSpringWeftScaleIndex(ForceInit)
 	{}
 
 	FPBDAreaSpringConstraints(
@@ -110,22 +128,50 @@ public:
 			WeightMaps.FindRef(GetAreaSpringStiffnessString(PropertyCollection, AreaSpringStiffnessName.ToString())),
 			FSolverVec2(GetWeightedFloatAreaSpringStiffness(PropertyCollection, 1.f)),
 			bTrimKinematicConstraints)
+		, WarpScale(FSolverVec2(1.f))
+		, WeftScale(FSolverVec2(1.f))
 		, AreaSpringStiffnessIndex(PropertyCollection)
+		, AreaSpringWarpScaleIndex(ForceInit)
+		, AreaSpringWeftScaleIndex(ForceInit)
 	{}
 
+	/** This version of the constructor supports Warp and Weft Scale */
+	CHAOS_API FPBDAreaSpringConstraints(
+		const FSolverParticlesRange & Particles,
+		const FTriangleMesh& TriangleMesh,
+		const TArray<TVec3<FVec2f>>&FaceVertexPatternPositions,
+		const TMap<FString, TConstArrayView<FRealSingle>>&WeightMaps,
+		const FCollectionPropertyConstFacade & PropertyCollection,
+		bool bTrimKinematicConstraints = false);
+
 	virtual ~FPBDAreaSpringConstraints() override = default;
 
 	CHAOS_API void SetProperties(
 		const FCollectionPropertyConstFacade& PropertyCollection,
 		const TMap<FString, TConstArrayView<FRealSingle>>& WeightMaps);
 
+	CHAOS_API void ApplyProperties(const FSolverReal Dt, const int32 NumIterations);
+
 private:
+
+	void InitFromPatternData(const TArray<TVec3<FVec2f>>& FaceVertexPatternPositions, const FTriangleMesh& TriangleMesh);
+	void UpdateDists();
+
 	using FPBDAxialSpringConstraints::Constraints;
+	using FPBDAxialSpringConstraints::Barys;
 	using FPBDAxialSpringConstraints::Stiffness;
 	using FPBDAxialSpringConstraints::ParticleOffset;
 	using FPBDAxialSpringConstraints::ParticleCount;
 
+	bool bWarpWeftScaleEnabled = false;
+	FPBDWeightMap WarpScale;
+	FPBDWeightMap WeftScale;
+	TArray<FSolverReal> BaseDists; // Without Warp/Weft Scale applied
+	TArray<FSolverVec2> WarpWeftScaleBaseMultipliers;
+
 	UE_CHAOS_DECLARE_PROPERTYCOLLECTION_NAME(AreaSpringStiffness, float);
+	UE_CHAOS_DECLARE_PROPERTYCOLLECTION_NAME(AreaSpringWarpScale, float);
+	UE_CHAOS_DECLARE_PROPERTYCOLLECTION_NAME(AreaSpringWeftScale, float);
 };
 
 }  // End namespace Chaos::Softs
diff --git a/Engine/Source/Runtime/Experimental/Chaos/Public/Chaos/PBDSpringConstraints.h b/Engine/Source/Runtime/Experimental/Chaos/Public/Chaos/PBDSpringConstraints.h
--- a/Engine/Source/Runtime/Experimental/Chaos/Public/Chaos/PBDSpringConstraints.h
+++ b/Engine/Source/Runtime/Experimental/Chaos/Public/Chaos/PBDSpringConstraints.h
@@ -4,6 +4,11 @@
 #include "Chaos/PBDSpringConstraintsBase.h"
 #include "Chaos/CollectionPropertyFacade.h"
 
+namespace Chaos
+{
+	class FTriangleMesh;
+}
+
 namespace Chaos::Softs
 {
 
@@ -17,6 +22,7 @@ public:
 		const TConstArrayView<FRealSingle>& StiffnessMultipliers,
 		const FSolverVec2& InStiffness,
 		bool bTrimKinematicConstraints = false,
+		bool bInitColor = true,
 		typename TEnableIf<Valence >= 2 && Valence <= 4>::Type* = nullptr)
 		: Base(
 			Particles,
@@ -25,7 +31,10 @@ public:
 			InStiffness,
 			bTrimKinematicConstraints)
 	{
-		InitColor(Particles);
+		if (bInitColor)
+		{
+			InitColor(Particles);
+		}
 	}
 
 	template<int32 Valence>
@@ -37,6 +46,7 @@ public:
 		const TConstArrayView<FRealSingle>& StiffnessMultipliers,
 		const FSolverVec2& InStiffness,
 		bool bTrimKinematicConstraints = false,
+		bool bInitColor = true,
 		typename TEnableIf<Valence >= 2 && Valence <= 4>::Type* = nullptr)
 		: Base(
 			Particles,
@@ -47,7 +57,10 @@ public:
 			InStiffness,
 			bTrimKinematicConstraints)
 	{
-		InitColor(Particles);
+		if (bInitColor)
+		{
+			InitColor(Particles);
+		}
 	}
 
 	virtual ~FPBDSpringConstraints() override {}
@@ -58,6 +71,12 @@ public:
 	const TArray<int32>& GetConstraintsPerColorStartIndex() const { return ConstraintsPerColorStartIndex; }
 
 protected:
+
+	/* Returns OrigToReorderedVertices */
+	template<typename SolverParticlesOrRange>
+	CHAOS_API TArray<int32> InitColor(const SolverParticlesOrRange& InParticles);
+
+
 	typedef FPBDSpringConstraintsBase Base;
 	using Base::Constraints;
 	using Base::Stiffness;
@@ -65,8 +84,6 @@ protected:
 	using Base::ParticleCount;
 
 private:
-	template<typename SolverParticlesOrRange>
-	CHAOS_API void InitColor(const SolverParticlesOrRange& InParticles);
 	template<typename SolverParticlesOrRange>
 	void ApplyHelper(SolverParticlesOrRange& Particles, const FSolverReal Dt, const int32 ConstraintIndex, const FSolverReal ExpStiffnessValue) const;
 
@@ -93,9 +110,15 @@ public:
 			InConstraints,
 			WeightMaps.FindRef(GetEdgeSpringStiffnessString(PropertyCollection, EdgeSpringStiffnessName.ToString())),
 			FSolverVec2(GetWeightedFloatEdgeSpringStiffness(PropertyCollection, 1.f)),
-			bTrimKinematicConstraints)
+			bTrimKinematicConstraints
+		)
+		, WarpScale(FSolverVec2(1.f))
+		, WeftScale(FSolverVec2(1.f))
 		, EdgeSpringStiffnessIndex(PropertyCollection)
-	{}
+		, EdgeSpringWarpScaleIndex(ForceInit)
+		, EdgeSpringWeftScaleIndex(ForceInit)
+	{
+	}
 
 	FPBDEdgeSpringConstraints(
 		const FSolverParticles& Particles,
@@ -113,8 +136,22 @@ public:
 			WeightMaps.FindRef(GetEdgeSpringStiffnessString(PropertyCollection, EdgeSpringStiffnessName.ToString())),
 			FSolverVec2(GetWeightedFloatEdgeSpringStiffness(PropertyCollection, 1.f)),
 			bTrimKinematicConstraints)
+		, WarpScale(FSolverVec2(1.f))
+		, WeftScale(FSolverVec2(1.f))
 		, EdgeSpringStiffnessIndex(PropertyCollection)
-	{}
+		, EdgeSpringWarpScaleIndex(ForceInit)
+		, EdgeSpringWeftScaleIndex(ForceInit)
+	{
+	}
+
+	/** This version of the constructor supports Warp and Weft Scale */
+	CHAOS_API FPBDEdgeSpringConstraints(
+		const FSolverParticlesRange& Particles,
+		const FTriangleMesh& TriangleMesh,
+		const TArray<TVec3<FVec2f>>& FaceVertexPatternPositions,
+		const TMap<FString, TConstArrayView<FRealSingle>>& WeightMaps,
+		const FCollectionPropertyConstFacade& PropertyCollection,
+		bool bTrimKinematicConstraints = false);
 
 	virtual ~FPBDEdgeSpringConstraints() override = default;
 
@@ -122,10 +159,24 @@ public:
 		const FCollectionPropertyConstFacade& PropertyCollection,
 		const TMap<FString, TConstArrayView<FRealSingle>>& WeightMaps);
 
+	CHAOS_API void ApplyProperties(const FSolverReal Dt, const int32 NumIterations);
+
 private:
+
+	void InitFromPatternData(const TArray<TVec3<FVec2f>>& FaceVertexPatternPositions, const FTriangleMesh& TriangleMesh);
+	void UpdateDists();
+
 	using FPBDSpringConstraints::ParticleCount;
 
+	bool bWarpWeftScaleEnabled = false;
+	FPBDWeightMap WarpScale;
+	FPBDWeightMap WeftScale;
+	TArray<FSolverReal> BaseDists; // Without Warp/Weft Scale applied
+	TArray<FSolverVec2> WarpWeftScaleBaseMultipliers;
+
 	UE_CHAOS_DECLARE_PROPERTYCOLLECTION_NAME(EdgeSpringStiffness, float);
+	UE_CHAOS_DECLARE_PROPERTYCOLLECTION_NAME(EdgeSpringWarpScale, float);
+	UE_CHAOS_DECLARE_PROPERTYCOLLECTION_NAME(EdgeSpringWeftScale, float);
 };
 
 class FPBDBendingSpringConstraints final : public FPBDSpringConstraints
@@ -149,7 +200,9 @@ public:
 			FSolverVec2(GetWeightedFloatBendingSpringStiffness(PropertyCollection, 1.f)),
 			bTrimKinematicConstraints)
 		, BendingSpringStiffnessIndex(PropertyCollection)
-	{}
+	{
+		InitColor(Particles);
+	}
 
 	FPBDBendingSpringConstraints(
 		const FSolverParticles& Particles,
@@ -168,7 +221,9 @@ public:
 			FSolverVec2(GetWeightedFloatBendingSpringStiffness(PropertyCollection, 1.f)),
 			bTrimKinematicConstraints)
 		, BendingSpringStiffnessIndex(PropertyCollection)
-	{}
+	{
+		InitColor(Particles);
+	}
 
 	virtual ~FPBDBendingSpringConstraints() override = default;
 
