Chaos - expose ConstConstraints() on collision container for convenience (without having to const-cast when you have a non-const reference to the container).

diff --git a/Engine/Source/Runtime/Experimental/Chaos/Private/Chaos/PBDCollisionConstraints.cpp b/Engine/Source/Runtime/Experimental/Chaos/Private/Chaos/PBDCollisionConstraints.cpp
--- a/Engine/Source/Runtime/Experimental/Chaos/Private/Chaos/PBDCollisionConstraints.cpp
+++ b/Engine/Source/Runtime/Experimental/Chaos/Private/Chaos/PBDCollisionConstraints.cpp
@@ -734,7 +734,14 @@ namespace Chaos
 	const FPBDCollisionConstraint& FPBDCollisionConstraints::GetConstraint(int32 Index) const
 	{
 		check(Index < NumConstraints());
-		
+
+		return *GetConstraints()[Index];
+	}
+
+	const FPBDCollisionConstraint& FPBDCollisionConstraints::GetConstConstraint(int32 Index) const
+	{
+		check(Index < NumConstraints());
+
 		return *GetConstraints()[Index];
 	}
 
diff --git a/Engine/Source/Runtime/Experimental/Chaos/Public/Chaos/PBDCollisionConstraints.h b/Engine/Source/Runtime/Experimental/Chaos/Public/Chaos/PBDCollisionConstraints.h
--- a/Engine/Source/Runtime/Experimental/Chaos/Public/Chaos/PBDCollisionConstraints.h
+++ b/Engine/Source/Runtime/Experimental/Chaos/Public/Chaos/PBDCollisionConstraints.h
@@ -298,6 +298,7 @@ public:
 	CHAOS_API FConstHandles GetConstConstraintHandles() const;
 
 	CHAOS_API const FPBDCollisionConstraint& GetConstraint(int32 Index) const;
+	CHAOS_API const FPBDCollisionConstraint& GetConstConstraint(int32 Index) const;
 
 	Private::FCollisionConstraintAllocator& GetConstraintAllocator() { return ConstraintAllocator; }
 
