Only return valid chunk info for chunks that is referenced by a content handle.

- prevents the on-demand install cache I/O backend to resolve chunks that is not referenced by a content handle, i.e. pinned in the cache and
further prevents DoesChunkExist/GetSizeForChunk to return valid values for non referenced chunks

#rnx

diff --git a/Engine/Source/Runtime/Experimental/IoStore/OnDemand/Private/OnDemandInstallCache.cpp b/Engine/Source/Runtime/Experimental/IoStore/OnDemand/Private/OnDemandInstallCache.cpp
--- a/Engine/Source/Runtime/Experimental/IoStore/OnDemand/Private/OnDemandInstallCache.cpp
+++ b/Engine/Source/Runtime/Experimental/IoStore/OnDemand/Private/OnDemandInstallCache.cpp
@@ -1515,7 +1515,8 @@ void FOnDemandInstallCache::UpdatePriorityForIoRequest(FIoRequestImpl* Request)
 
 bool FOnDemandInstallCache::DoesChunkExist(const FIoChunkId& ChunkId) const
 {
-	if (FOnDemandChunkInfo ChunkInfo = IoStore.GetInstalledChunkInfo(ChunkId))
+	EIoErrorCode ErrorCode = EIoErrorCode::UnknownChunkID;
+	if (FOnDemandChunkInfo ChunkInfo = IoStore.GetInstalledChunkInfo(ChunkId, ErrorCode))
 	{
 		const FCasLocation CasLoc = Cas.FindChunk(ChunkInfo.Hash());
 		return CasLoc.IsValid();
@@ -1526,12 +1527,13 @@ bool FOnDemandInstallCache::DoesChunkExist(const FIoChunkId& ChunkId) const
 
 TIoStatusOr<uint64> FOnDemandInstallCache::GetSizeForChunk(const FIoChunkId& ChunkId) const
 {
-	if (FOnDemandChunkInfo ChunkInfo = IoStore.GetInstalledChunkInfo(ChunkId))
+	EIoErrorCode ErrorCode = EIoErrorCode::UnknownChunkID;
+	if (FOnDemandChunkInfo ChunkInfo = IoStore.GetInstalledChunkInfo(ChunkId, ErrorCode))
 	{
 		return ChunkInfo.RawSize();
 	}
 
-	return FIoStatus(EIoErrorCode::UnknownChunkID);
+	return FIoStatus(ErrorCode);
 }
 
 TIoStatusOr<FIoMappedRegion> FOnDemandInstallCache::OpenMapped(const FIoChunkId& ChunkId, const FIoReadOptions& Options)
@@ -1547,9 +1549,15 @@ const TCHAR* FOnDemandInstallCache::GetName() const
 
 bool FOnDemandInstallCache::Resolve(FIoRequestImpl* Request)
 {
-	FOnDemandChunkInfo ChunkInfo = IoStore.GetInstalledChunkInfo(Request->ChunkId);
+	EIoErrorCode ErrorCode = EIoErrorCode::UnknownChunkID;
+	FOnDemandChunkInfo ChunkInfo = IoStore.GetInstalledChunkInfo(Request->ChunkId, ErrorCode);
 	if (ChunkInfo.IsValid() == false)
 	{
+		if (ErrorCode == EIoErrorCode::NotInstalled)
+		{
+			CompleteRequest(Request, EIoErrorCode::NotInstalled);
+			return true;
+		}
 		return false;
 	}
 
diff --git a/Engine/Source/Runtime/Experimental/IoStore/OnDemand/Private/OnDemandIoStore.cpp b/Engine/Source/Runtime/Experimental/IoStore/OnDemand/Private/OnDemandIoStore.cpp
--- a/Engine/Source/Runtime/Experimental/IoStore/OnDemand/Private/OnDemandIoStore.cpp
+++ b/Engine/Source/Runtime/Experimental/IoStore/OnDemand/Private/OnDemandIoStore.cpp
@@ -896,12 +896,55 @@ FIoStatus FOnDemandIoStore::GetInstallSizesByMountId(const FOnDemandGetInstallSi
 
 FOnDemandChunkInfo FOnDemandIoStore::GetStreamingChunkInfo(const FIoChunkId& ChunkId)
 {
-	return GetChunkInfo(ChunkId, EOnDemandContainerFlags::Mounted | EOnDemandContainerFlags::StreamOnDemand);
+	const EOnDemandContainerFlags ContainerFlags = EOnDemandContainerFlags::Mounted | EOnDemandContainerFlags::StreamOnDemand;
+
+	TUniqueLock Lock(ContainerMutex);
+
+	for (FSharedOnDemandContainer& Container : Containers)
+	{
+		if (!EnumHasAllFlags(Container->Flags, ContainerFlags))
+		{
+			continue;
+		}
+
+		if (const FOnDemandChunkEntry* Entry = Container->FindChunkEntry(ChunkId))
+		{
+			return FOnDemandChunkInfo(Container, *Entry);
+		}
+	}
+
+	return FOnDemandChunkInfo();
 }
 
-FOnDemandChunkInfo FOnDemandIoStore::GetInstalledChunkInfo(const FIoChunkId& ChunkId)
+FOnDemandChunkInfo FOnDemandIoStore::GetInstalledChunkInfo(const FIoChunkId& ChunkId, EIoErrorCode& OutErrorCode)
 {
-	return GetChunkInfo(ChunkId, EOnDemandContainerFlags::Mounted | EOnDemandContainerFlags::InstallOnDemand);
+	const EOnDemandContainerFlags ContainerFlags = EOnDemandContainerFlags::Mounted | EOnDemandContainerFlags::InstallOnDemand;
+
+	TUniqueLock Lock(ContainerMutex);
+
+	for (FSharedOnDemandContainer& Container : Containers)
+	{
+		if (!EnumHasAllFlags(Container->Flags, ContainerFlags))
+		{
+			continue;
+		}
+
+		int32 EntryIndex = INDEX_NONE;
+		if (const FOnDemandChunkEntry* Entry = Container->FindChunkEntry(ChunkId, &EntryIndex))
+		{
+			if (Container->IsReferenced(EntryIndex))
+			{
+				OutErrorCode = EIoErrorCode::Ok;
+				return FOnDemandChunkInfo(Container, *Entry);
+			}
+
+			OutErrorCode = EIoErrorCode::NotInstalled;
+			return FOnDemandChunkInfo();
+		}
+	}
+
+	OutErrorCode = EIoErrorCode::UnknownChunkID;
+	return FOnDemandChunkInfo();
 }
 
 #if !UE_BUILD_SHIPPING
@@ -1019,24 +1062,6 @@ FIoStatus FOnDemandIoStore::InitializeInstallCache()
 	return EIoErrorCode::Ok;
 }
 
-FOnDemandChunkInfo FOnDemandIoStore::GetChunkInfo(const FIoChunkId& ChunkId, EOnDemandContainerFlags ContainerFlags)
-{
-	TUniqueLock Lock(ContainerMutex);
-
-	for (FSharedOnDemandContainer& Container : Containers)
-	{
-		if (EnumHasAllFlags(Container->Flags, ContainerFlags))
-		{
-			if (const FOnDemandChunkEntry* Entry = Container->FindChunkEntry(ChunkId))
-			{
-				return FOnDemandChunkInfo(Container, *Entry);
-			}
-		}
-	}
-
-	return FOnDemandChunkInfo();
-}
-
 void FOnDemandIoStore::TryEnterTickLoop()
 {
 	bool bEnterTickLoop = false;
diff --git a/Engine/Source/Runtime/Experimental/IoStore/OnDemand/Private/OnDemandIoStore.h b/Engine/Source/Runtime/Experimental/IoStore/OnDemand/Private/OnDemandIoStore.h
--- a/Engine/Source/Runtime/Experimental/IoStore/OnDemand/Private/OnDemandIoStore.h
+++ b/Engine/Source/Runtime/Experimental/IoStore/OnDemand/Private/OnDemandIoStore.h
@@ -81,10 +81,10 @@ struct FOnDemandContainer
 {
 	FString									UniqueName() const;
 	inline int32							FindChunkEntryIndex(const FIoChunkId& ChunkId) const;
-	inline const FOnDemandChunkEntry*		FindChunkEntry(const FIoChunkId& ChunkId) const;
-	inline FOnDemandChunkEntry*				FindChunkEntry(const FIoChunkId& ChunkId);
+	inline const FOnDemandChunkEntry*		FindChunkEntry(const FIoChunkId& ChunkId, int32* OutIndex = nullptr) const;
 	inline FOnDemandChunkEntryReferences&	FindOrAddChunkEntryReferences(const FOnDemandInternalContentHandle& ContentHandle);
 	inline TBitArray<>						GetReferencedChunkEntries() const;
+	inline bool								IsReferenced(int32 ChunkEntryIndex) const;
 	bool									HasAnyFlags(EOnDemandContainerFlags Contains) const { return EnumHasAnyFlags(Flags, Contains); }
 	bool									HasAllFlags(EOnDemandContainerFlags Contains) const { return EnumHasAllFlags(Flags, Contains); }
 
@@ -117,27 +117,21 @@ int32 FOnDemandContainer::FindChunkEntryIndex(const FIoChunkId& ChunkId) const
 	{
 		if (ChunkIds[Index] == ChunkId)
 		{
-			return Index; 
+			return Index;
 		}
 	}
 
 	return INDEX_NONE;
 }
 
-const FOnDemandChunkEntry* FOnDemandContainer::FindChunkEntry(const FIoChunkId& ChunkId) const
-{
-	if (int32 Index = FindChunkEntryIndex(ChunkId); Index != INDEX_NONE)
-	{
-		return &ChunkEntries[Index];
-	}
-
-	return nullptr;
-}
-
-FOnDemandChunkEntry* FOnDemandContainer::FindChunkEntry(const FIoChunkId& ChunkId)
+const FOnDemandChunkEntry* FOnDemandContainer::FindChunkEntry(const FIoChunkId& ChunkId, int32* OutIndex) const
 {
 	if (int32 Index = FindChunkEntryIndex(ChunkId); Index != INDEX_NONE)
 	{
+		if (OutIndex != nullptr)
+		{
+			*OutIndex = Index;
+		}
 		return &ChunkEntries[Index];
 	}
 
@@ -173,6 +167,19 @@ TBitArray<> FOnDemandContainer::GetReferencedChunkEntries() const
 	return Indices;
 }
 
+bool FOnDemandContainer::IsReferenced(int32 ChunkEntryIndex) const
+{
+	for (const FOnDemandChunkEntryReferences& Refs : ChunkEntryReferences)
+	{
+		if (Refs.Indices[ChunkEntryIndex])
+		{
+			return true;
+		}
+	}
+
+	return false;
+}
+
 ///////////////////////////////////////////////////////////////////////////////
 struct FOnDemandChunkInfo
 {
@@ -268,7 +275,7 @@ public:
 	FIoStatus							InitializeStreamingBackend(const FOnDemandEndpointConfig& EndpointConfig);
 
 	FOnDemandChunkInfo					GetStreamingChunkInfo(const FIoChunkId& ChunkId);
-	FOnDemandChunkInfo					GetInstalledChunkInfo(const FIoChunkId& ChunkId);
+	FOnDemandChunkInfo					GetInstalledChunkInfo(const FIoChunkId& ChunkId, EIoErrorCode& OutErrorCode);
 
 #if !UE_BUILD_SHIPPING
 	/** Used to access a number of FIoChunkId that exist in the system when running tests and is not intended for production use */
@@ -304,7 +311,6 @@ private:
 											TSet<FSharedOnDemandContainer>& OutContainersForInstallation,
 											TSet<FPackageId>& OutPackageIdsToInstall) const;
 	FIoStatus							InitializeInstallCache();
-	FOnDemandChunkInfo					GetChunkInfo(const FIoChunkId& ChunkId, EOnDemandContainerFlags ContainerFlags);
 	void								TryEnterTickLoop();
 	void								TickLoop();
 	bool								Tick();
