Extended on-demand TOC with host group name.

#rnx

diff --git a/Engine/Source/Runtime/Experimental/IoStore/OnDemand/Private/Tool/ChunkPlugin.cpp b/Engine/Source/Runtime/Experimental/IoStore/OnDemand/Private/Tool/ChunkPlugin.cpp
--- a/Engine/Source/Runtime/Experimental/IoStore/OnDemand/Private/Tool/ChunkPlugin.cpp
+++ b/Engine/Source/Runtime/Experimental/IoStore/OnDemand/Private/Tool/ChunkPlugin.cpp
@@ -247,6 +247,7 @@ static int32 ChunkPluginCommandEntry(const FContext& Context)
 	const FString IntermediateFolder	= FString(Context.Get<FStringView>(TEXT("-IntermediateFolder"), FString()));
 	FString SettingsFile				= FString(Context.Get<FStringView>(TEXT("-SettingsFile"), FString()));
 	FString OutputStatsJson				= FString(Context.Get<FStringView>(TEXT("-OutputStatsJson"), FString()));
+	const FString HostGroupName			= FString(Context.Get<FStringView>(TEXT("-HostGroupName"), FString()));
 
 	const bool bStreamOnDemand			= Context.Get<bool>(TEXT("-StreamOnDemand"), false);
 	const bool bInstallOnDemand			= Context.Get<bool>(TEXT("-InstallOnDemand"), false);
@@ -379,6 +380,7 @@ static int32 ChunkPluginCommandEntry(const FContext& Context)
 
 	FOnDemandToc OnDemandToc;
 	OnDemandToc.Header.Flags = static_cast<uint32>(TocFlags);
+	OnDemandToc.Header.HostGroupName = HostGroupName;
 	//OnDemandToc.Header.ChunksDirectory = TODO 
 	OnDemandToc.Containers.Reserve(ContainerFilenames.Num());
 
@@ -756,6 +758,7 @@ static FCommand ChunkPluginCommand(
 		TArgument<FStringView>(TEXT("-IntermediateFolder"),	TEXT("Intermediate folder.")),
 		TArgument<FStringView>(TEXT("-SettingsFile"),		TEXT("Optional settings file.")),
 		TArgument<FStringView>(TEXT("-OutputStatsJson"),	TEXT("Path to write a json file with statistics.")),
+		TArgument<FStringView>(TEXT("-HostGroupName"),		TEXT("Host group name or URL")),
 		TArgument<bool>(TEXT("-IncludeSigPak"),				TEXT("Include .sig and .pak file in the uondemandtoc")),
 		TArgument<bool>(TEXT("-KeepContainerFiles"),		TEXT("Should we keep the container files after processing them.")),
 		TArgument<bool>(TEXT("-StreamOnDemand"),			TEXT("Set the content to be streamed on-demand")),
diff --git a/Engine/Source/Runtime/Experimental/IoStore/OnDemand/Private/Tool/Upload.cpp b/Engine/Source/Runtime/Experimental/IoStore/OnDemand/Private/Tool/Upload.cpp
--- a/Engine/Source/Runtime/Experimental/IoStore/OnDemand/Private/Tool/Upload.cpp
+++ b/Engine/Source/Runtime/Experimental/IoStore/OnDemand/Private/Tool/Upload.cpp
@@ -50,6 +50,7 @@ struct FUploadParams
 	FString BuildVersion;
 	FString TargetPlatform;
 	FString EncryptionKeyName;
+	FString HostGroupName;
 	int32 MaxConcurrentUploads = 16;
 	int32 MaxTocListCount = 10000;
 	int32 MaxTocDownloadCount = 100;
@@ -130,6 +131,7 @@ static FUploadParams BuildUploadParams(const FContext& Context)
 	Ret.CredentialsFileKeyName	= Context.Get<FStringView>(TEXT("-CredentialsFileKeyName"),	Ret.CredentialsFileKeyName);
 	Ret.BuildVersion			= Context.Get<FStringView>(TEXT("-BuildVersion"),			Ret.BuildVersion);
 	Ret.TargetPlatform			= Context.Get<FStringView>(TEXT("-TargetPlatform"),			Ret.TargetPlatform);
+	Ret.HostGroupName			= Context.Get<FStringView>(TEXT("-HostGroupName"),			Ret.HostGroupName);
 	Ret.bWriteTocToDisk			= Context.Get<bool>(TEXT("-WriteTocToDisk"),				Ret.bWriteTocToDisk);
 	Ret.bPerContainerTocs		= Context.Get<bool>(TEXT("-PerContainerTocs"),				Ret.bPerContainerTocs);
 	Ret.MaxConcurrentUploads	= Context.Get<int32>(TEXT("-MaxConcurrentUploads"),			Ret.MaxConcurrentUploads);
@@ -567,6 +569,7 @@ static TIoStatusOr<FUploadResult> UploadContainerFiles(
 
 	FOnDemandToc OnDemandToc;
 	OnDemandToc.Header.Flags = static_cast<uint32>(UploadParams.TocFlags);
+	OnDemandToc.Header.HostGroupName = UploadParams.HostGroupName;
 	OnDemandToc.Containers.Reserve(ContainerFiles.Num());
 
 	// Map of the .utoc paths that we have created ondemand containers for, indexed by their hash so that the paths can
@@ -842,6 +845,7 @@ static TIoStatusOr<FUploadResult> UploadContainerFiles(
 		UE_LOG(LogIas, Display, TEXT("%-15s: %s"), TEXT("Bucket"), *UploadParams.Bucket);
 		UE_LOG(LogIas, Display, TEXT("%-15s: %s"), TEXT("TargetPlatform"), *UploadParams.TargetPlatform);
 		UE_LOG(LogIas, Display, TEXT("%-15s: %s"), TEXT("BuildVersion"), *UploadParams.BuildVersion);
+		UE_LOG(LogIas, Display, TEXT("%-15s: %s"), TEXT("HostGroupName"), *UploadParams.HostGroupName);
 		UE_LOG(LogIas, Display, TEXT("%-15s: %s"), TEXT("TOC path"), *UploadResult.TocPath);
 		UE_LOG(LogIas, Display, TEXT("%-15s: %s"), TEXT("TOC flags"), *LexToString(UploadParams.TocFlags));
 		UE_LOG(LogIas, Display, TEXT("%-15s: %.2lf KiB"), TEXT("TOC size"), double(UploadResult.TocSize) / 1024.0);
@@ -955,6 +959,7 @@ static FCommand UploadCommand(
 		TArgument<FStringView>(TEXT("-DistributionUrl"),TEXT("URL for IAS to use at runtime. Output to config file")),
 		TArgument<FStringView>(TEXT("-FallbackUrl"),	TEXT("Alternative fallack for DistributionUrl")),
 		TArgument<FStringView>(TEXT("-BucketPrefix"),	TEXT("Path to prefix to bucket objects")),
+		TArgument<FStringView>(TEXT("-HostGroupName"),	TEXT("Host group name or URL")),
 		TArgument<bool>(TEXT("-KeepContainerFiles"),	TEXT("Do not delete container files after upload")),
 		TArgument<bool>(TEXT("-KeepPakFiles"),			TEXT("Do not delete the springboard pak files")),
 		TArgument<bool>(TEXT("-WriteTocToDisk"),		TEXT("Output the TOC to disk as well as uploading")),
diff --git a/Engine/Source/Runtime/Experimental/IoStore/OnDemandCore/Private/OnDemandToc.cpp b/Engine/Source/Runtime/Experimental/IoStore/OnDemandCore/Private/OnDemandToc.cpp
--- a/Engine/Source/Runtime/Experimental/IoStore/OnDemandCore/Private/OnDemandToc.cpp
+++ b/Engine/Source/Runtime/Experimental/IoStore/OnDemandCore/Private/OnDemandToc.cpp
@@ -111,8 +111,13 @@ FArchive& operator<<(FArchive& Ar, FOnDemandTocHeader& Header)
 	Ar << Header.BlockSize;
 	Ar << Header.CompressionFormat;
 	Ar << Header.ChunksDirectory;
+	
+	if (Ar.IsSaving() || (Header.Version >= static_cast<uint32>(EOnDemandTocVersion::HostGroupName)))
+	{
+		Ar << Header.HostGroupName;
+	}
 
-	if (Ar.IsLoading() && (Header.Version < uint32(EOnDemandTocVersion::TocFlags)))
+	if (Ar.IsLoading() && (Header.Version < static_cast<uint32>(EOnDemandTocVersion::TocFlags)))
 	{
 		Header.Flags = 0; 
 	}
@@ -129,6 +134,7 @@ FCbWriter& operator<<(FCbWriter& Writer, const FOnDemandTocHeader& Header)
 	Writer.AddInteger(UTF8TEXTVIEW("BlockSize"), Header.BlockSize);
 	Writer.AddString(UTF8TEXTVIEW("CompressionFormat"), Header.CompressionFormat);
 	Writer.AddString(UTF8TEXTVIEW("ChunksDirectory"), Header.ChunksDirectory);
+	Writer.AddString(UTF8TEXTVIEW("HostGroupName"), Header.HostGroupName);
 	Writer.EndObject();
 
 	return Writer;
@@ -144,6 +150,7 @@ bool LoadFromCompactBinary(FCbFieldView Field, FOnDemandTocHeader& OutTocHeader)
 		OutTocHeader.BlockSize = Obj["BlockSize"].AsUInt32();
 		OutTocHeader.CompressionFormat = FString(Obj["CompressionFormat"].AsString());
 		OutTocHeader.ChunksDirectory = FString(Obj["ChunksDirectory"].AsString());
+		OutTocHeader.HostGroupName = FString(Obj["HostGroupName"].AsString());
 
 		return OutTocHeader.Magic == FOnDemandTocHeader::ExpectedMagic &&
 			static_cast<EOnDemandTocVersion>(OutTocHeader.Version) != EOnDemandTocVersion::Invalid;
diff --git a/Engine/Source/Runtime/Experimental/IoStore/OnDemandCore/Public/IO/OnDemandToc.h b/Engine/Source/Runtime/Experimental/IoStore/OnDemandCore/Public/IO/OnDemandToc.h
--- a/Engine/Source/Runtime/Experimental/IoStore/OnDemandCore/Public/IO/OnDemandToc.h
+++ b/Engine/Source/Runtime/Experimental/IoStore/OnDemandCore/Public/IO/OnDemandToc.h
@@ -33,6 +33,7 @@ enum class EOnDemandTocVersion : uint32
 	TagSets			= 8,
 	ContainerFlags	= 9,
 	TocFlags		= 10,
+	HostGroupName	= 11,
 
 	LatestPlusOne,
 	Latest			= (LatestPlusOne - 1)
@@ -73,7 +74,8 @@ struct FOnDemandTocHeader
 	uint32	BlockSize = 0;
 	FString CompressionFormat;
 	FString ChunksDirectory;
-	
+	FString HostGroupName;
+
 	UE_API friend FArchive& operator<<(FArchive& Ar, FOnDemandTocHeader& Header);
 	UE_API friend FCbWriter& operator<<(FCbWriter& Writer, const FOnDemandTocHeader& Header);
 };
