Choose the correct triangle for height calculation when directly queried from a heightfield.

A heightfield cell is represented by two virtual triangles so a diagonal must be chosen for the split. Spatial queries were correct but the calculation for which triangle to use when directly querying a position on the heightfield was inverted.

diff --git a/Engine/Source/Runtime/Experimental/Chaos/Private/Chaos/HeightField.cpp b/Engine/Source/Runtime/Experimental/Chaos/Private/Chaos/HeightField.cpp
--- a/Engine/Source/Runtime/Experimental/Chaos/Private/Chaos/HeightField.cpp
+++ b/Engine/Source/Runtime/Experimental/Chaos/Private/Chaos/HeightField.cpp
@@ -934,7 +934,7 @@ namespace Chaos
 		const Chaos::FReal FractionX = FMath::Frac(ClampedGridLocationLocal[0]);
 		const Chaos::FReal FractionY = FMath::Frac(ClampedGridLocationLocal[1]);
 
-		if(FractionX > FractionY)
+		if(FractionX < FractionY)
 		{
 			if(bFillHeight)
 			{
