Chaos - minor improvement to friction behaviour for dynamic-kinematic collisions, and dynamic-dynamic collisions with large mass ratios.

The world-space contact point was the average of the surface contact points of the two colliding shapes. This causes a penetration-depth dependency in tangential impulses, because the constraint effective mass depends on the constraint relative position.

Now the world-space contact point is the weighted average surface contact point, so dynamic-kinematic collisions always use a contact point on the dynamic surface.

This will not address the depth-dependence of dynamic-dynamic friction for bodies of similar masses.

diff --git a/Engine/Source/Runtime/Experimental/Chaos/Private/Chaos/Collision/PBDCollisionContainerSolver.cpp b/Engine/Source/Runtime/Experimental/Chaos/Private/Chaos/Collision/PBDCollisionContainerSolver.cpp
--- a/Engine/Source/Runtime/Experimental/Chaos/Private/Chaos/Collision/PBDCollisionContainerSolver.cpp
+++ b/Engine/Source/Runtime/Experimental/Chaos/Private/Chaos/Collision/PBDCollisionContainerSolver.cpp
@@ -91,9 +91,11 @@ namespace Chaos
 		const FRigidTransform3& ShapeWorldTransform1 = Constraint->GetShapeWorldTransform1();
 
 		// World-space contact points on each shape
+		const FReal ContactWeight0 = Body0.InvM();
+		const FReal ContactWeight1 = Body1.InvM();
 		const FVec3 WorldContact0 = ShapeWorldTransform0.TransformPositionNoScale(FVec3(ManifoldPoint.ContactPoint.ShapeContactPoints[0]));
 		const FVec3 WorldContact1 = ShapeWorldTransform1.TransformPositionNoScale(FVec3(ManifoldPoint.ContactPoint.ShapeContactPoints[1]));
-		const FVec3 WorldContact = FReal(0.5) * (WorldContact0 + WorldContact1);
+		const FVec3 WorldContact = (ContactWeight0 * WorldContact0 + ContactWeight1 * WorldContact1) / (ContactWeight0 + ContactWeight1);
 		const FVec3f WorldRelativeContact0 = FVec3f(WorldContact - Body0.P());
 		const FVec3f WorldRelativeContact1 = FVec3f(WorldContact - Body1.P());
 
