Chaos - fix velocity-phase friction calculation so that it clamps total frictional impulse and converges correctly with multiple iterations.

diff --git a/Engine/Source/Runtime/Experimental/Chaos/Public/Chaos/Collision/PBDCollisionSolver.h b/Engine/Source/Runtime/Experimental/Chaos/Public/Chaos/Collision/PBDCollisionSolver.h
--- a/Engine/Source/Runtime/Experimental/Chaos/Public/Chaos/Collision/PBDCollisionSolver.h
+++ b/Engine/Source/Runtime/Experimental/Chaos/Public/Chaos/Collision/PBDCollisionSolver.h
@@ -650,15 +650,20 @@ namespace Chaos
 					+ (ManifoldPoint.NetSoftPushOutNormal
 					+  ManifoldPoint.NetPushOutNormal) / Dt;
 
+				// Clamp the total tangential impulse including positional contribution to the friction cone
 				const FSolverReal MinMaxFrictionImpulse = State.MinMaxFrictionPushout / Dt;
 				const FSolverReal MaxImpulseTangent = DynamicFriction * FMath::Max(MinMaxFrictionImpulse, TotalImpulseNormal);
 				const FSolverReal MaxImpulseTangentSq = FMath::Square(MaxImpulseTangent);
-				const FSolverReal ImpulseTangentSq = FMath::Square(ImpulseTangentU) + FMath::Square(ImpulseTangentV);
-				if (ImpulseTangentSq > (MaxImpulseTangentSq + UE_SMALL_NUMBER))
+				const FSolverReal NetImpulseTangentU = ManifoldPoint.NetImpulseTangentU + ManifoldPoint.NetPushOutTangentU / Dt;
+				const FSolverReal NetImpulseTangentV = ManifoldPoint.NetImpulseTangentV + ManifoldPoint.NetPushOutTangentV / Dt;
+				const FSolverReal NetImpulseTangentSq = FMath::Square(NetImpulseTangentU + ImpulseTangentU) + FMath::Square(NetImpulseTangentV + ImpulseTangentV);
+				if (NetImpulseTangentSq > (MaxImpulseTangentSq + UE_SMALL_NUMBER))
 				{
-					const FSolverReal ImpulseTangentScale = MaxImpulseTangent * FMath::InvSqrt(ImpulseTangentSq);
-					ImpulseTangentU *= ImpulseTangentScale;
-					ImpulseTangentV *= ImpulseTangentScale;
+					const FSolverReal ImpulseTangentScale = MaxImpulseTangent * FMath::InvSqrt(NetImpulseTangentSq);
+					ImpulseTangentU = (NetImpulseTangentU + ImpulseTangentU) * ImpulseTangentScale - NetImpulseTangentU;
+					ImpulseTangentV = (NetImpulseTangentV + ImpulseTangentV) * ImpulseTangentScale - NetImpulseTangentV;
+					// TODO: Should we update the StaticFrictionRatio if we clamp here? This is used to move the
+					// static friction anchors to the edge of the friction cone
 				}
 
 				ManifoldPoint.NetImpulseTangentU += ImpulseTangentU;
