Chaos - enable velocity-phase friction by default

This follows on from the fix to make the velocity-phase friciton calculation clamp correctly and be independent of number of iterations.

diff --git a/Engine/Source/Runtime/Experimental/Chaos/Public/Chaos/Collision/PBDCollisionSolverSettings.h b/Engine/Source/Runtime/Experimental/Chaos/Public/Chaos/Collision/PBDCollisionSolverSettings.h
--- a/Engine/Source/Runtime/Experimental/Chaos/Public/Chaos/Collision/PBDCollisionSolverSettings.h
+++ b/Engine/Source/Runtime/Experimental/Chaos/Public/Chaos/Collision/PBDCollisionSolverSettings.h
@@ -15,7 +15,7 @@ namespace Chaos
 			: MaxPushOutVelocity(0)
 			, DepenetrationVelocity(0)
 			, NumPositionFrictionIterations(4)
-			, NumVelocityFrictionIterations(0)
+			, NumVelocityFrictionIterations(1)
 			, NumPositionShockPropagationIterations(3)
 			, NumVelocityShockPropagationIterations(2)
 		{
