FOnDemandPackageStoreBackend now tracks which packages are installed/referenced.

diff --git a/Engine/Source/Runtime/Experimental/IoStore/OnDemand/Private/OnDemandIoStore.cpp b/Engine/Source/Runtime/Experimental/IoStore/OnDemand/Private/OnDemandIoStore.cpp
--- a/Engine/Source/Runtime/Experimental/IoStore/OnDemand/Private/OnDemandIoStore.cpp
+++ b/Engine/Source/Runtime/Experimental/IoStore/OnDemand/Private/OnDemandIoStore.cpp
@@ -797,12 +797,9 @@ FIoStatus FOnDemandIoStore::Unmount(FStringView MountId)
 				UE_LOG(LogIoStoreOnDemand, Log, TEXT("Unmounting container, ContainerName='%s', MountId='%s'"),
 					*WriteToString<128>(Container->Name), *WriteToString<128>(Container->MountId));
 
-				const FString ContainerName = Container->UniqueName();
-				if (PackageStoreBackend.IsValid())
-				{
-					PackageStoreBackend->Unmount(ContainerName);
-				}
+				PackageStoreBackend->NeedsUpdate();
 
+				const FString ContainerName = Container->UniqueName();
 				PendingContainerHeaders.Remove(ContainerName);
 
 				return true;
@@ -1000,7 +997,7 @@ FIoStatus FOnDemandIoStore::InitializeInstallCache()
 		}
 #endif
 		FIoDispatcher::Get().Mount(InstallCache.ToSharedRef(), BackendPriority);
-		PackageStoreBackend = MakeOnDemandPackageStoreBackend();
+		PackageStoreBackend = MakeOnDemandPackageStoreBackend(AsShared().ToWeakPtr());
 		FPackageStore::Get().Mount(PackageStoreBackend.ToSharedRef());
 	}
 	else
@@ -1390,11 +1387,7 @@ FIoStatus FOnDemandIoStore::TickMountRequest(FMountRequest& MountRequest)
 					return Header.Status();
 				}
 
-				Container->Header = Header.ConsumeValueOrDie(); 
-				if (FIoStatus MountStatus = PackageStoreBackend->Mount(Container->UniqueName(), Container->Header); !MountStatus.IsOk())
-				{
-					return MountStatus;
-				}
+				Container->Header = Header.ConsumeValueOrDie();
 			}
 		}
 
@@ -1412,6 +1405,8 @@ FIoStatus FOnDemandIoStore::TickMountRequest(FMountRequest& MountRequest)
 		PendingContainerHeaders.Append(MoveTemp(RequestedContainerHeaders));
 	}
 
+	PackageStoreBackend->NeedsUpdate();
+
 	return bAnyPendingEncryptionKey ? EIoErrorCode::PendingEncryptionKey : EIoErrorCode::Ok;
 }
 
@@ -1732,7 +1727,7 @@ FOnDemandInstallResult FOnDemandIoStore::TickInstallRequest(FOnDemandInternalIns
 	{
 		FOnDemandContainer& Container				= *Kv.Key;
 		const Private::FContainerInstallData& Data	= Kv.Value;
-		// TODO: Mark installed packages in package store backend 
+
 		{
 			UE::TUniqueLock Lock(ContainerMutex);
 			FOnDemandChunkEntryReferences& References = Container.FindOrAddChunkEntryReferences(*ContentHandle);
@@ -1744,6 +1739,12 @@ FOnDemandInstallResult FOnDemandIoStore::TickInstallRequest(FOnDemandInternalIns
 		}
 	}
 
+	if (InstallData.IsEmpty() == false)
+	{
+		// Referencers were added so the package store must be refreshed
+		PackageStoreBackend->NeedsUpdate();
+	}
+
 	double DurationInSeconds = FPlatformTime::Seconds() - StartTime;
 	#if !UE_BUILD_SHIPPING
 	if (GIoStoreOnDemandMinInstallDurationSeconds > 0.0f)
@@ -1851,6 +1852,8 @@ void FOnDemandIoStore::OnEncryptionKeyAdded(const FGuid& Id, const FAES::FAESKey
 
 	TUniqueLock Lock(ContainerMutex);
 
+	bool bAddedContainerHeaders = false;
+
 	for (FSharedOnDemandContainer& Container : Containers)
 	{
 		if (EnumHasAnyFlags(Container->Flags, EOnDemandContainerFlags::PendingEncryptionKey) == false)
@@ -1890,16 +1893,8 @@ void FOnDemandIoStore::OnEncryptionKeyAdded(const FGuid& Id, const FAES::FAESKey
 
 				if (ensure(Header.IsOk()))
 				{
-					Container->Header = Header.ConsumeValueOrDie(); 
-					const FIoStatus MountStatus = PackageStoreBackend->Mount(Container->UniqueName(), Container->Header);
-					if (MountStatus.IsOk())
-					{
-						EnumAddFlags(Container->Flags, EOnDemandContainerFlags::Mounted);
-					}
-					else
-					{
-						UE_LOG(LogIoStoreOnDemand, Error, TEXT("Failed to add container '%s' to package store backend"), *Container->UniqueName()); 
-					}
+					Container->Header = Header.ConsumeValueOrDie();
+					bAddedContainerHeaders = true;
 				}
 				else
 				{
@@ -1911,6 +1906,11 @@ void FOnDemandIoStore::OnEncryptionKeyAdded(const FGuid& Id, const FAES::FAESKey
 			PendingContainerHeaders.RemoveByHash(ContainerNameHash, ContainerName);
 		}
 	}
+
+	if (bAddedContainerHeaders)
+	{
+		PackageStoreBackend->NeedsUpdate();
+	}
 }
 
 void FOnDemandIoStore::CreateContainersFromToc(
@@ -2065,7 +2065,7 @@ FIoStatus FOnDemandIoStore::GetContainersAndPackagesForInstall(
 			if (Container->HasAnyFlags(EOnDemandContainerFlags::PendingEncryptionKey))
 			{
 				check(Container->HasAnyFlags(EOnDemandContainerFlags::Mounted) == false);
-				if(MountId == Container->MountId)
+				if (MountId == Container->MountId)
 				{
 					check(EnumHasAnyFlags(Container->Flags, EOnDemandContainerFlags::Mounted) == false);
 					FIoStatus Status = FIoStatusBuilder(EIoErrorCode::PendingEncryptionKey)
@@ -2167,6 +2167,9 @@ void FOnDemandIoStore::ReleaseContent(FOnDemandInternalContentHandle& ContentHan
 					return Refs.ContentHandleId == ContentHandleId;
 				}));
 	}
+
+	// Package store needs to be aware of install state
+	PackageStoreBackend->NeedsUpdate();
 }
 
 void FOnDemandIoStore::GetReferencedContent(TArray<FSharedOnDemandContainer>& OutContainers, TArray<TBitArray<>>& OutChunkEntryIndices)
@@ -2179,6 +2182,14 @@ void FOnDemandIoStore::GetReferencedContent(TArray<FSharedOnDemandContainer>& Ou
 			continue;
 		}
 
+		if (Container->HasAnyFlags(EOnDemandContainerFlags::PendingEncryptionKey))
+		{
+			continue;
+		}
+
+		// Header should be available if decrypted
+		check(Container->Header);
+
 		TBitArray<> Indices = Container->GetReferencedChunkEntries();
 		if (Indices.IsEmpty() == false)
 		{
diff --git a/Engine/Source/Runtime/Experimental/IoStore/OnDemand/Private/OnDemandPackageStoreBackend.cpp b/Engine/Source/Runtime/Experimental/IoStore/OnDemand/Private/OnDemandPackageStoreBackend.cpp
--- a/Engine/Source/Runtime/Experimental/IoStore/OnDemand/Private/OnDemandPackageStoreBackend.cpp
+++ b/Engine/Source/Runtime/Experimental/IoStore/OnDemand/Private/OnDemandPackageStoreBackend.cpp
@@ -10,6 +10,7 @@
 #include "IO/IoStoreOnDemand.h"
 #include "IO/PackageId.h"
 #include "Internationalization/PackageLocalizationManager.h"
+#include "OnDemandIoStore.h"
 
 namespace UE::IoStore
 {
@@ -32,20 +33,23 @@ class FOnDemandPackageStoreBackend final
 		FSoftPackageReferenceMap	SoftRefs;
 	};
 
+	struct FOnDemandPackageStoreEntry
+	{
+		const FFilePackageStoreEntry* FilePackageStoreEntry = nullptr;
+		bool bIsInstalled = false; // uninstalled packages are not 'missing' and not 'pending'
+	};
+
 	using FSharedBackendContext = TSharedPtr<const FPackageStoreBackendContext>;
-	using FEntryMap				= TMap<FPackageId, const FFilePackageStoreEntry*>;
-	using FUniqueContainer		= TUniquePtr<FContainer>;
+	using FEntryMap				= TMap<FPackageId, FOnDemandPackageStoreEntry>;
 	using FRedirect				= TTuple<FName, FPackageId>;
 	using FLocalizedMap			= TMap<FPackageId, FName>;
 	using FRedirectMap			= TMap<FPackageId, FRedirect>;
 
 public:
-						FOnDemandPackageStoreBackend();
+						FOnDemandPackageStoreBackend(TWeakPtr<FOnDemandIoStore> OnDemandIoStore);
 						virtual ~FOnDemandPackageStoreBackend ();
 
-	virtual FIoStatus	Mount(FString ContainerName, FSharedContainerHeader ContainerHeader) override;
-	virtual FIoStatus	Unmount(const FString& ContainerName) override;
-	virtual FIoStatus	UnmountAll() override;
+	virtual void		NeedsUpdate() override;
 
 	virtual void		OnMounted(TSharedRef<const FPackageStoreBackendContext> Context) override;
 	virtual void 		BeginRead() override;
@@ -66,17 +70,18 @@ public:
 private:
 	void				UpdateLookupTables();
 
-	FSharedBackendContext		BackendContext;
-	TArray<FUniqueContainer>	Containers;
+	TWeakPtr<FOnDemandIoStore>	OnDemandIoStore;
+	TArray<FContainer>			Containers;
 	FEntryMap					EntryMap;
 	FLocalizedMap				LocalizedMap;
 	FRedirectMap				RedirectMap;
 	UE::FMutex					Mutex;
-	bool						bNeedsUpdate = false;
+	bool						bNeedsUpdate = true;
 };
 
 ///////////////////////////////////////////////////////////////////////////////
-FOnDemandPackageStoreBackend::FOnDemandPackageStoreBackend()
+FOnDemandPackageStoreBackend::FOnDemandPackageStoreBackend(TWeakPtr<FOnDemandIoStore> OnDemandIoStore)
+	: OnDemandIoStore(MoveTemp(OnDemandIoStore))
 {
 }
 
@@ -84,57 +89,14 @@ FOnDemandPackageStoreBackend::~FOnDemandPackageStoreBackend()
 {
 }
 
-FIoStatus FOnDemandPackageStoreBackend::Mount(FString ContainerName, FSharedContainerHeader ContainerHeader)
-{
-	{
-		UE::TUniqueLock Lock(Mutex);
-
-		const FUniqueContainer* Existing =
-			Algo::FindByPredicate(
-				Containers,
-				[&ContainerName](const FUniqueContainer& C) { return C->Name == ContainerName; });
-
-		if (Existing != nullptr)
-		{
-			return EIoErrorCode::Ok;
-		}
-
-		Containers.Add(MakeUnique<FContainer>(MoveTemp(ContainerName), MoveTemp(ContainerHeader)));
-		bNeedsUpdate = true;
-	}
-
-	if (BackendContext)
-	{
-		BackendContext->PendingEntriesAdded.Broadcast();
-	}
-
-	return EIoErrorCode::Ok;
-}
-
-FIoStatus FOnDemandPackageStoreBackend::Unmount(const FString& ContainerName)
+void FOnDemandPackageStoreBackend::NeedsUpdate()
 {
 	UE::TUniqueLock Lock(Mutex);
-	const int32 NumRemoved = Containers.RemoveAll([&ContainerName](const TUniquePtr<FContainer>& C)
-	{
-		return C->Name == ContainerName;
-	});
-	check(NumRemoved <= 1);
-
 	bNeedsUpdate = true;
-	return NumRemoved == 0 ? EIoErrorCode::NotFound : EIoErrorCode::Ok;
-}
-
-FIoStatus FOnDemandPackageStoreBackend::UnmountAll()
-{
-	UE::TUniqueLock Lock(Mutex);
-	Containers.Empty();
-	bNeedsUpdate = true;
-	return EIoErrorCode::Ok;
 }
 
 void FOnDemandPackageStoreBackend::OnMounted(TSharedRef<const FPackageStoreBackendContext> Context)
 {
-	BackendContext = Context;
 }
 
 void FOnDemandPackageStoreBackend::BeginRead()
@@ -153,14 +115,27 @@ EPackageStoreEntryStatus FOnDemandPackageStoreBackend::GetPackageStoreEntry(
 	FName PackageName,
 	FPackageStoreEntry& OutPackageStoreEntry)
 {
-	if (const FFilePackageStoreEntry* Entry = EntryMap.FindRef(PackageId))
+	if (const FOnDemandPackageStoreEntry* Entry = EntryMap.Find(PackageId))
 	{
+		const FFilePackageStoreEntry* FilePackageStoreEntry = Entry->FilePackageStoreEntry;
 		OutPackageStoreEntry.ImportedPackageIds =
-			MakeArrayView(Entry->ImportedPackages.Data(), Entry->ImportedPackages.Num());
+			MakeArrayView(FilePackageStoreEntry->ImportedPackages.Data(), FilePackageStoreEntry->ImportedPackages.Num());
 		OutPackageStoreEntry.ShaderMapHashes =
-			MakeArrayView(Entry->ShaderMapHashes.Data(), Entry->ShaderMapHashes.Num());
+			MakeArrayView(FilePackageStoreEntry->ShaderMapHashes.Data(), FilePackageStoreEntry->ShaderMapHashes.Num());
+
+		// TODO: this depends on some package loading changes that aren't ready to go in yet.
+		// We can't return missing in this case because we have code paths that need to be able
+		// to find NotInstalled packages.
+#if 0
+		if (Entry->bIsInstalled)
+		{
+			return EPackageStoreEntryStatus::Ok;
+		}
 
+		return EPackageStoreEntryStatus::NotInstalled;
+#else
 		return EPackageStoreEntryStatus::Ok;
+#endif
 	}
 
 	return EPackageStoreEntryStatus::Missing;
@@ -198,11 +173,11 @@ bool FOnDemandPackageStoreBackend::GetPackageRedirectInfo(
 
 TConstArrayView<uint32> FOnDemandPackageStoreBackend::GetSoftReferences(FPackageId PackageId, TConstArrayView<FPackageId>& OutPackageIds) 
 {
-	for (const FUniqueContainer& Container : Containers)
+	for (const FContainer& Container : Containers)
 	{
-		if (const FFilePackageStoreEntrySoftReferences* SoftRefs = Container->SoftRefs.FindRef(PackageId))
+		if (const FFilePackageStoreEntrySoftReferences* SoftRefs = Container.SoftRefs.FindRef(PackageId))
 		{
-			OutPackageIds = Container->Header->SoftPackageReferences.PackageIds;
+			OutPackageIds = Container.Header->SoftPackageReferences.PackageIds;
 			return TConstArrayView<uint32>(SoftRefs->Indices.Data(), SoftRefs->Indices.Num());
 		}
 	}
@@ -219,22 +194,31 @@ void FOnDemandPackageStoreBackend::UpdateLookupTables()
 	bNeedsUpdate = false;
 	check(Mutex.IsLocked());
 
-	EntryMap.Empty();
 	LocalizedMap.Empty();
 	RedirectMap.Empty();
 
-	const int32 PackageCount = Algo::TransformAccumulate(
-		Containers,
-		[](const FUniqueContainer& C) { return C->Header->PackageIds.Num(); },
-		0);
+	TArray<FSharedOnDemandContainer> AllContainers;
+	TArray<TBitArray<>> ReferencedChunkEntryIndices;
+	int32 PackageCount = 0;
 
-	EntryMap.Reserve(PackageCount);
+	if (TSharedPtr<FOnDemandIoStore> PinOnDemandIoStore = OnDemandIoStore.Pin())
+	{
+		PinOnDemandIoStore->GetReferencedContent(AllContainers, ReferencedChunkEntryIndices);
+
+		for (const FSharedOnDemandContainer& C : AllContainers)
+		{
+			PackageCount += C->Header->PackageIds.Num();
+		}
+	}
+
+	Containers.Empty(AllContainers.Num());
+	EntryMap.Empty(PackageCount);
 
-	for (const FUniqueContainer& Container : Containers)
+	for (int32 ContainerIdx = 0; const FSharedOnDemandContainer& SharedContainer : AllContainers)
 	{
-		Container->SoftRefs.Reset();
+		FContainer& Container = Containers.Emplace_GetRef(CopyTemp(SharedContainer->UniqueName()), SharedContainer->Header);
 
-		const FIoContainerHeader& Hdr = *Container->Header;
+		const FIoContainerHeader& Hdr = *Container.Header;
 		TConstArrayView<FFilePackageStoreEntry> Entries(
 			reinterpret_cast<const FFilePackageStoreEntry*>(Hdr.StoreEntries.GetData()),
 			Hdr.PackageIds.Num());
@@ -245,24 +229,29 @@ void FOnDemandPackageStoreBackend::UpdateLookupTables()
 			AllSoftReferences = MakeArrayView<const FFilePackageStoreEntrySoftReferences>(
 				reinterpret_cast<const FFilePackageStoreEntrySoftReferences*>(Hdr.SoftPackageReferences.PackageIndices.GetData()),
 				Hdr.PackageIds.Num());
-			Container->SoftRefs.Reserve(Hdr.PackageIds.Num());
+			Container.SoftRefs.Reserve(Hdr.PackageIds.Num());
 		}
 
-		for (int32 Idx = 0; const FFilePackageStoreEntry& Entry : Entries)
+		for (int32 EntryIdx = 0; const FFilePackageStoreEntry& Entry : Entries)
 		{
-			const FPackageId PkgId = Hdr.PackageIds[Idx];
-			EntryMap.Add(PkgId, &Entry);
+			const FPackageId PkgId = Hdr.PackageIds[EntryIdx];
+			const FIoChunkId PkgChunkId = CreatePackageDataChunkId(PkgId);
+			const int32 PkgChunkIdx = SharedContainer->FindChunkEntryIndex(PkgChunkId);
+			check(PkgChunkIdx != INDEX_NONE);
+			const bool IsReferenced = ReferencedChunkEntryIndices[ContainerIdx][PkgChunkIdx];
+
+			EntryMap.Add(PkgId, FOnDemandPackageStoreEntry{ .FilePackageStoreEntry = &Entry, .bIsInstalled = IsReferenced });
 			if (AllSoftReferences.IsEmpty() == false)
 			{
-				const FFilePackageStoreEntrySoftReferences& SoftRefs = AllSoftReferences[Idx];
+				const FFilePackageStoreEntrySoftReferences& SoftRefs = AllSoftReferences[EntryIdx];
 				if (SoftRefs.Indices.Num() > 0)
 				{
-					Container->SoftRefs.Add(PkgId, &SoftRefs);
+					Container.SoftRefs.Add(PkgId, &SoftRefs);
 				}
 			}
-			++Idx;
+			++EntryIdx;
 		}
-		Container->SoftRefs.Shrink();
+		Container.SoftRefs.Shrink();
 
 		for (const FIoContainerHeaderLocalizedPackage& Localized : Hdr.LocalizedPackages)
 		{
@@ -285,6 +274,8 @@ void FOnDemandPackageStoreBackend::UpdateLookupTables()
 				RedirectEntry.Value = Redirect.TargetPackageId;
 			}
 		}
+
+		++ContainerIdx;
 	}
 
 	EntryMap.Shrink();
@@ -293,9 +284,9 @@ void FOnDemandPackageStoreBackend::UpdateLookupTables()
 }
 
 ///////////////////////////////////////////////////////////////////////////////
-TSharedPtr<IOnDemandPackageStoreBackend> MakeOnDemandPackageStoreBackend()
+TSharedPtr<IOnDemandPackageStoreBackend> MakeOnDemandPackageStoreBackend(TWeakPtr<FOnDemandIoStore> OnDemandIoStore)
 {
-	return MakeShareable<IOnDemandPackageStoreBackend>(new FOnDemandPackageStoreBackend());
+	return MakeShared<FOnDemandPackageStoreBackend>(MoveTemp(OnDemandIoStore));
 }
 
 } // namespace UE
diff --git a/Engine/Source/Runtime/Experimental/IoStore/OnDemand/Private/OnDemandPackageStoreBackend.h b/Engine/Source/Runtime/Experimental/IoStore/OnDemand/Private/OnDemandPackageStoreBackend.h
--- a/Engine/Source/Runtime/Experimental/IoStore/OnDemand/Private/OnDemandPackageStoreBackend.h
+++ b/Engine/Source/Runtime/Experimental/IoStore/OnDemand/Private/OnDemandPackageStoreBackend.h
@@ -7,6 +7,7 @@
 #include "Templates/SharedPointer.h"
 
 struct FIoContainerHeader;
+class FOnDemandIoStore;
 using FSharedContainerHeader = TSharedPtr<FIoContainerHeader>;
 
 namespace UE::IoStore
@@ -18,11 +19,9 @@ class IOnDemandPackageStoreBackend
 public:
 	virtual ~IOnDemandPackageStoreBackend() = default;
 
-	virtual FIoStatus Mount(FString ContainerName, FSharedContainerHeader ContainerHeader) = 0;
-	virtual FIoStatus Unmount(const FString& ContainerName) = 0;
-	virtual FIoStatus UnmountAll() = 0;
+	virtual void NeedsUpdate() = 0;
 };
 
-TSharedPtr<IOnDemandPackageStoreBackend> MakeOnDemandPackageStoreBackend();
+TSharedPtr<IOnDemandPackageStoreBackend> MakeOnDemandPackageStoreBackend(TWeakPtr<FOnDemandIoStore> OnDemandIoStore);
 
 } // namespace UE
