Geometry collection : fix crash when importing skin weight from skeletal mesh with several mesh islands
- Extract the skin weights in the same loop as the positions and colors
- Removed dead code accessing render and LOD data

diff --git a/Engine/Source/Runtime/Experimental/GeometryCollectionEngine/Private/GeometryCollection/GeometryCollectionEngineConversion.cpp b/Engine/Source/Runtime/Experimental/GeometryCollectionEngine/Private/GeometryCollection/GeometryCollectionEngineConversion.cpp
--- a/Engine/Source/Runtime/Experimental/GeometryCollectionEngine/Private/GeometryCollection/GeometryCollectionEngineConversion.cpp
+++ b/Engine/Source/Runtime/Experimental/GeometryCollectionEngine/Private/GeometryCollection/GeometryCollectionEngineConversion.cpp
@@ -1290,9 +1290,14 @@ bool FGeometryCollectionEngineConversion::AppendSkeletalMesh(const USkeletalMesh
 	TArray<FTransform> ComponentTransform;
 	GeometryCollectionAlgo::GlobalMatrices(LocalSpaceTransform, Parent, ComponentTransform);
 
+	FSkeletalMeshConstAttributes SkeletalMeshConstAttributes(MeshDescription);
+
+	// normal detection
+	const bool bHasNormals = MeshDescription.VertexInstanceAttributes().HasAttribute(MeshAttribute::VertexInstance::Normal);
+
+	// color attribute detection 
 	bool bHasVertexColors = false;
 	FLinearColor DefaultColor = FLinearColor::White;
-	FSkeletalMeshConstAttributes SkeletalMeshConstAttributes(MeshDescription);
 	TVertexInstanceAttributesConstRef<FVector4f> InstanceColors = SkeletalMeshConstAttributes.GetVertexInstanceColors();
 	if (InstanceColors.IsValid())
 	{
@@ -1300,20 +1305,27 @@ bool FGeometryCollectionEngineConversion::AppendSkeletalMesh(const USkeletalMesh
 		DefaultColor = InstanceColors.GetDefaultValue();
 	}
 
+	// vertex weight attribute detection
+	FSkinWeightsVertexAttributesConstRef VertexSkinWeights = SkeletalMeshConstAttributes.GetVertexSkinWeights();
+	const bool bHasVertexSkinWeights = VertexSkinWeights.IsValid();
+	GeometryCollection::Facades::FVertexBoneWeightsFacade VertexBoneWeightsFacade(GeometryCollection);
+	TArray<int32> VertexBoneIndex;
+	TArray<float> VertexBoneWeight;
+
 	// @todo(GeometryCollectionConversion) : Add support for UV's, Normals  
 	//const int32 NumUVLayers = VertexBuffers.StaticMeshVertexBuffer.GetNumTexCoords();
 	//GeometryCollection::UV::SetNumUVLayers(*GeometryCollection, NumUVLayers);
 	//GeometryCollection::UV::FUVLayers UVLayers = GeometryCollection::UV::FindActiveUVLayers(*GeometryCollection);
 	TArray<bool> TargetVertexVisited; TargetVertexVisited.Init(false, NumTargetVertices);
-	for (int ComponentIndex = 0;ComponentIndex<ComponentsSourceIndices.Num();ComponentIndex++)
+	for (int32 ComponentIndex = 0;ComponentIndex<ComponentsSourceIndices.Num(); ComponentIndex++)
 	{
 		for (int32 TriangleIndex = 0; TriangleIndex < ComponentsSourceIndices[ComponentIndex].Num(); TriangleIndex++)
 		{
-			int32 SourceTriangleIndex = SourceToTargetTriangleMap[ComponentIndex][TriangleIndex][0];
-			FIntVector& Triangle = ComponentsSourceIndices[ComponentIndex][TriangleIndex];
-			for (int k = 0; k < 3; k++)
+			const int32 SourceTriangleIndex = SourceToTargetTriangleMap[ComponentIndex][TriangleIndex][0];
+			const FIntVector& Triangle = ComponentsSourceIndices[ComponentIndex][TriangleIndex];
+			for (int TriVtxIdx = 0; TriVtxIdx < 3; TriVtxIdx++)
 			{
-				int SourceVertexIndex = Triangle[k];
+				int SourceVertexIndex = Triangle[TriVtxIdx];
 				if (ensure(0 <= SourceVertexIndex && SourceVertexIndex < SourceVertexToComponentMap.Num()))
 				{
 					int TargetVertexIndex = SourceVertexToComponentMap[SourceVertexIndex] + VertexBaseIndex;
@@ -1321,27 +1333,55 @@ bool FGeometryCollectionEngineConversion::AppendSkeletalMesh(const USkeletalMesh
 					{
 						if (!TargetVertexVisited[TargetVertexIndex])
 						{
-							// Vertex position
-							//  @todo(Skin Wieghts) : Currently its a static pose in component space. 
+							TargetVertexVisited[TargetVertexIndex] = true;
+
+							// Vertex positions
 							BoneMap[TargetVertexIndex] = ComponentToTransformGroupIndex[ComponentIndex];
 
 							FVector SourceVertex = (FVector)MeshDescription.GetVertexPosition(SourceVertexIndex);
 							FMatrix M = ComponentTransform[BoneMap[TargetVertexIndex]].ToInverseMatrixWithScale();
 							Vertex[TargetVertexIndex] = (FVector4f)M.TransformPosition(SourceVertex);
 
-							TargetVertexVisited[TargetVertexIndex] = true;
+							// Get Avg Vertex Normals
+							// TODO(chaos) - once we support UV and split vertices accordingly, we can use the actual normal instead of averaging it 
+							if (bHasNormals)
+							{
+								Normal[TargetVertexIndex] = FVector3f::ZeroVector;
+
+								const TArrayView<const FVertexInstanceID> VertexInstances = MeshDescription.GetVertexVertexInstanceIDs(SourceVertexIndex);
+								for (const FVertexInstanceID VertexInstanceID : VertexInstances)
+								{
+									const FVector3f VertexInstanceNormal = MeshDescription.VertexInstanceAttributes().GetAttribute<FVector3f>(VertexInstanceID, MeshAttribute::VertexInstance::Normal);
+									Normal[TargetVertexIndex] += VertexInstanceNormal / VertexInstances.Num();
+								}
+								Normal[TargetVertexIndex].Normalize();
+							}
 
-							// transfer color
+							// Vertex Colors
 							Color[TargetVertexIndex] = DefaultColor;
 							if (bHasVertexColors)
 							{
 								TArrayView<const FVertexInstanceID> SourceInstanceTri = MeshDescription.GetTriangleVertexInstances(SourceTriangleIndex);
-								FVector4f InstColor = InstanceColors[SourceInstanceTri[k]];
+								FVector4f InstColor = InstanceColors[SourceInstanceTri[TriVtxIdx]];
 								UE::Geometry::LinearColors::SRGBToLinear(InstColor);
 								Color[TargetVertexIndex] = FLinearColor(InstColor);
 							}
 
-							// @todo(GeometryCollectionConversion) : Add support for UV's, Normals  
+							// Vertex skin weights
+							if (bHasVertexSkinWeights)
+							{
+								const FVertexBoneWeightsConst BoneWeights = VertexSkinWeights.Get(FVertexID(SourceVertexIndex));
+								const int32 InfluenceCount = BoneWeights.Num();
+								VertexBoneIndex.SetNum(InfluenceCount);
+								VertexBoneWeight.SetNum(InfluenceCount);
+								for (int32 InfluenceIndex = 0; InfluenceIndex < InfluenceCount; ++InfluenceIndex)
+								{
+									VertexBoneIndex[InfluenceIndex] = BoneWeights[InfluenceIndex].GetBoneIndex();
+									VertexBoneWeight[InfluenceIndex] = BoneWeights[InfluenceIndex].GetWeight();
+								}
+								VertexBoneWeightsFacade.ModifyBoneWeight(TargetVertexIndex, VertexBoneIndex, VertexBoneWeight);
+							}
+							// @todo(GeometryCollectionConversion) : Add support for UV's  
 							//TangentU[VertexOffset] = VertexBuffers.StaticMeshVertexBuffer.VertexTangentX(VertexIndex);
 							//TangentV[VertexOffset] = VertexBuffers.StaticMeshVertexBuffer.VertexTangentY(VertexIndex);
 							//Normal[VertexOffset] = VertexBuffers.StaticMeshVertexBuffer.VertexTangentZ(VertexIndex);
@@ -1370,9 +1410,9 @@ bool FGeometryCollectionEngineConversion::AppendSkeletalMesh(const USkeletalMesh
 			SourceToTargetTriangleMap[ComponentIndex][TriangleIndex][1] = CurrentIndex;
 
 			FIntVector& Triangle = ComponentsSourceIndices[ComponentIndex][TriangleIndex];
-			for (int k = 0; k < 3; k++)
+			for (int TriVtxIdx = 0; TriVtxIdx < 3; TriVtxIdx++)
 			{
-				Indices[CurrentIndex][k] = SourceVertexToComponentMap[Triangle[k]] + VertexBaseIndex;
+				Indices[CurrentIndex][TriVtxIdx] = SourceVertexToComponentMap[Triangle[TriVtxIdx]] + VertexBaseIndex;
 			}
 
 			Visible[CurrentIndex] = true;
@@ -1390,39 +1430,7 @@ bool FGeometryCollectionEngineConversion::AppendSkeletalMesh(const USkeletalMesh
 	GeometryCollection.RemoveElements(FGeometryCollection::GeometryGroup, GeometryIndices);
 	::GeometryCollection::AddGeometryProperties(&GeometryCollection);
 
-	//Extract skin weights
-	int32 NumVertices = GeometryCollection.NumElements(FGeometryCollection::VerticesGroup);
-	FSkeletalMeshRenderData* RenderData = InSkeletalMesh->GetResourceForRendering();
-	if (RenderData->LODRenderData.Num())
-	{
-		FSkeletalMeshLODRenderData* LODRenderData = &RenderData->LODRenderData[LODIndex];
-		const FSkinWeightVertexBuffer* SkinWeightVertexBuffer = LODRenderData->GetSkinWeightVertexBuffer();
-
-		TArray<int32> BlendedSkeletalMeshVertices;
-		TArray<TArray<float>> Weights;
-		Weights.SetNum(SkinWeightVertexBuffer->GetNumVertices());
-		FSkeletalMeshAttributes MeshAttribs(MeshDescription);
-		FSkinWeightsVertexAttributesRef VertexSkinWeights = MeshAttribs.GetVertexSkinWeights();
-		if (ensure(NumVertices == MeshDescription.Vertices().Num()))
-		{
-			GeometryCollection::Facades::FVertexBoneWeightsFacade VertexBoneWeightsFacade(GeometryCollection);
-			for (int32 VertexIndex = 0; VertexIndex < NumVertices; ++VertexIndex)
-			{
-				FVertexBoneWeights BoneWeights = VertexSkinWeights.Get(FVertexID(VertexIndex));
-				const int32 InfluenceCount = BoneWeights.Num();
-				TArray<int32> VertexBoneIndex;
-				TArray<float> VertexBoneWeight;
-				VertexBoneIndex.SetNum(InfluenceCount);
-				VertexBoneWeight.SetNum(InfluenceCount);
-				for (int32 InfluenceIndex = 0; InfluenceIndex < InfluenceCount; ++InfluenceIndex)
-				{
-					VertexBoneIndex[InfluenceIndex] = BoneWeights[InfluenceIndex].GetBoneIndex();
-					VertexBoneWeight[InfluenceIndex] = BoneWeights[InfluenceIndex].GetWeight();
-				}
-				VertexBoneWeightsFacade.ModifyBoneWeight(SourceVertexToComponentMap[VertexIndex], VertexBoneIndex, VertexBoneWeight);
-			}
-		}
-	}
+	// copy to the target collection
 	GeometryCollection.CopyTo(InManagedArrayCollection);
 	return true;
 #else
