Iris
* Changed WorldLocation out of bounds detection from an ensure to a configurable warning.
* By default the warning will trigger once per-class
* Use net.Iris.InvalidLocationWarnings to modify the warning behavior
* 0 = disable warnings and OOB checking  (default for shipping builds)
* 1 = detect OOB but warn once per class (default for non-shipping builds)
* 2 = detect OOB but warn once per object
* 3 = warn everytime we detect an OOB location (very spammy)

#rnx

diff --git a/Engine/Source/Runtime/Experimental/Iris/Core/Private/Iris/ReplicationSystem/ObjectReplicationBridge.cpp b/Engine/Source/Runtime/Experimental/Iris/Core/Private/Iris/ReplicationSystem/ObjectReplicationBridge.cpp
--- a/Engine/Source/Runtime/Experimental/Iris/Core/Private/Iris/ReplicationSystem/ObjectReplicationBridge.cpp
+++ b/Engine/Source/Runtime/Experimental/Iris/Core/Private/Iris/ReplicationSystem/ObjectReplicationBridge.cpp
@@ -94,6 +94,17 @@ static FAutoConsoleVariableRef CVarEnableForceNetUpdate(
 	TEXT("When true the system only allows ForceNetUpdate to skip the poll frequency of objects. When false any MarkDirty object will be immediately polled.")
 );
 
+#if UE_BUILD_SHIPPING
+static int32 GInvalidLocationWarnings = 0; // don't check OOB locations in Shipping
+#else
+static int32 GInvalidLocationWarnings = 1; // check OOB locations and warn once per class 
+#endif
+static FAutoConsoleVariableRef CVarInvalidLocationWarnings(
+	TEXT("net.Iris.InvalidLocationWarnings"),
+	GInvalidLocationWarnings,
+	TEXT("Chooses how to warn when we detect an invalid world location on a rootobject. 0=never warn. 1=warn once per class. 2=warn once per object. 3=warn on every update")
+);
+
 static int32 GEnsureNetRefHandleError = 0;
 static FAutoConsoleVariableRef CvarEnsureNetRefHandleError(
 	TEXT("net.Iris.EnsureNetRefHandleError"),
@@ -1484,24 +1495,69 @@ void UObjectReplicationBridge::ForceUpdateWorldLocation(FNetRefHandle NetRefHand
 		return;
 	}
 
-	UNetObjectFactory* NetFactory = GetNetFactory(NetRefHandleManager->GetReplicatedObjectData(InternalObjectIndex).NetFactoryId);
+	const FUpdateWorldInfoContext UpdateContext
+	{
+		.WorldLocations = WorldLocations,
+		.ObjectIndex = InternalObjectIndex,
+		.NetRefHandle = NetRefHandle,
+	};
+	UpdateRootObjectWorldInfo(UpdateContext);
+}
 
-	const UNetObjectFactory::FWorldInfoContext Context
+void UObjectReplicationBridge::UpdateRootObjectWorldInfo(const FUpdateWorldInfoContext& UpdateContext)
+{
+	using namespace UE::Net::Private;
+
+	const FNetRefHandleManager::FReplicatedObjectData& ObjectData = NetRefHandleManager->GetReplicatedObjectDataNoCheck(UpdateContext.ObjectIndex);
+
+	// Ask the NetFactory for the object's current world info
+	UNetObjectFactory* NetFactory = GetNetFactory(ObjectData.NetFactoryId);
+
+	const UNetObjectFactory::FWorldInfoContext WorldInfoContext
 	{
-		.Instance = NetRefHandleManager->GetReplicatedObjectInstance(InternalObjectIndex),
-		.Handle = NetRefHandle,
+		.Instance = NetRefHandleManager->GetReplicatedObjectInstance(UpdateContext.ObjectIndex),
+		.Handle = UpdateContext.NetRefHandle,
 		.InfoRequested = UNetObjectFactory::EWorldInfoRequested::All,
 	};
 
 	UNetObjectFactory::FWorldInfoData WorldInfo;
-	NetFactory->GetWorldInfo(Context, WorldInfo);
+	NetFactory->GetWorldInfo(WorldInfoContext, WorldInfo);
+
+	// Check if we need to warn about invalid world locations.
+	// Note that FWorldLocations will store the invalid location but only return a valid clamped position from it's API
+	if (UE_LOG_ACTIVE(LogIrisBridge, Warning) && GInvalidLocationWarnings != 0)
+	{
+		const bool bIsValidLocation = UpdateContext.WorldLocations.IsValidLocation(WorldInfo.WorldLocation);
+		
+		if (!bIsValidLocation)
+		{
+			FObjectKey ClassOrObject;
+			
+			// Warn once per class type
+			if (GInvalidLocationWarnings == 1)
+			{
+				ClassOrObject = WorldInfoContext.Instance ? WorldInfoContext.Instance->GetClass() : nullptr;
+			}
+			// Warn once per object
+			else if (GInvalidLocationWarnings >= 2)
+			{
+				ClassOrObject = WorldInfoContext.Instance;
+			}
+
+			bool bAlreadyWarned(false);
+			WorldLocationOOBWarnings.FindOrAdd(ClassOrObject, &bAlreadyWarned);
 
-	// Check the position lies within world boundaries.
-	ensureMsgf(	WorldInfo.WorldLocation.X >= WorldLocations.GetWorldMinPos().X && WorldInfo.WorldLocation.X <= WorldLocations.GetWorldMaxPos().X &&
-				WorldInfo.WorldLocation.Y >= WorldLocations.GetWorldMinPos().Y && WorldInfo.WorldLocation.Y <= WorldLocations.GetWorldMaxPos().Y,
-				TEXT("Object %s with position %s lies outside configured world boundary."), ToCStr(NetRefHandleManager->PrintObjectFromIndex(InternalObjectIndex)), *WorldInfo.WorldLocation.ToString());
+			// If it's the first warning or it's set to warn on every update
+			if (!bAlreadyWarned || GInvalidLocationWarnings == 3)
+			{
+				UE_LOG_BRIDGEID(LogIrisBridge, Warning, TEXT("RootObject %s has a location %s that is out of bounds."), 
+					ToCStr(NetRefHandleManager->PrintObjectFromIndex(UpdateContext.ObjectIndex)), *WorldInfo.WorldLocation.ToCompactString());
+				bAlreadyWarned = true;
+			}
+		}
+	}
 
-	WorldLocations.SetObjectInfo(InternalObjectIndex, WorldInfo.WorldLocation, WorldInfo.CullDistance);
+	UpdateContext.WorldLocations.SetObjectInfo(UpdateContext.ObjectIndex, WorldInfo.WorldLocation, WorldInfo.CullDistance);
 }
 
 void UObjectReplicationBridge::UpdateInstancesWorldLocation()
@@ -1512,33 +1568,20 @@ void UObjectReplicationBridge::UpdateInstancesWorldLocation()
 	FReplicationSystemInternal* ReplicationSystemInternal = GetReplicationSystem()->GetReplicationSystemInternal();
 	const FNetRefHandleManager& LocalNetRefHandleManager = ReplicationSystemInternal->GetNetRefHandleManager();
 	FWorldLocations& WorldLocations = ReplicationSystemInternal->GetWorldLocations();
-	const TNetChunkedArray<TObjectPtr<UObject>>& ReplicatedInstances = LocalNetRefHandleManager.GetReplicatedInstances();
+
+	FUpdateWorldInfoContext UpdateContext
+	{
+		.WorldLocations = WorldLocations
+	};
 
 	// Retrieve the world location for instances that supports it. Only dirty objects are considered.
-	auto UpdateInstanceWorldLocation = [this, &ReplicatedInstances, &LocalNetRefHandleManager, &WorldLocations](uint32 InternalObjectIndex)
+	auto UpdateInstanceWorldLocation = [this, &LocalNetRefHandleManager, &UpdateContext](uint32 InternalObjectIndex)
 	{
-		if (WorldLocations.HasInfoForObject(InternalObjectIndex))
+		if (UpdateContext.WorldLocations.HasInfoForObject(InternalObjectIndex))
 		{
-			const FNetRefHandleManager::FReplicatedObjectData& ObjectData = LocalNetRefHandleManager.GetReplicatedObjectDataNoCheck(InternalObjectIndex);
-
-            const UNetObjectFactory::FWorldInfoContext Context
-            {
-                .Instance = ReplicatedInstances[InternalObjectIndex],
-                .Handle = ObjectData.RefHandle,
-                //TODO: Set InfoRequested based on the what was set dirty by the user
-            };
-
-            UNetObjectFactory::FWorldInfoData WorldInfo;
-
-            UNetObjectFactory* NetFactory = GetNetFactory(ObjectData.NetFactoryId);
-            NetFactory->GetWorldInfo(Context, WorldInfo);
-
-			// Check the position lies within world boundaries.
-			ensureMsgf(WorldInfo.WorldLocation.X >= WorldLocations.GetWorldMinPos().X && WorldInfo.WorldLocation.X <= WorldLocations.GetWorldMaxPos().X && 
-					   WorldInfo.WorldLocation.Y >= WorldLocations.GetWorldMinPos().Y && WorldInfo.WorldLocation.Y <= WorldLocations.GetWorldMaxPos().Y,
-                       TEXT("Object %s with position %s lies outside configured world boundary."), ToCStr(NetRefHandleManager->PrintObjectFromIndex((InternalObjectIndex))), *WorldInfo.WorldLocation.ToString());
-			
-            WorldLocations.SetObjectInfo(InternalObjectIndex, WorldInfo.WorldLocation, WorldInfo.CullDistance);
+			UpdateContext.ObjectIndex = InternalObjectIndex;
+			UpdateContext.NetRefHandle = LocalNetRefHandleManager.GetNetRefHandleFromInternalIndex(InternalObjectIndex);
+			UpdateRootObjectWorldInfo(UpdateContext);
 		}
 	};
 
diff --git a/Engine/Source/Runtime/Experimental/Iris/Core/Public/Iris/ReplicationSystem/ObjectReplicationBridge.h b/Engine/Source/Runtime/Experimental/Iris/Core/Public/Iris/ReplicationSystem/ObjectReplicationBridge.h
--- a/Engine/Source/Runtime/Experimental/Iris/Core/Public/Iris/ReplicationSystem/ObjectReplicationBridge.h
+++ b/Engine/Source/Runtime/Experimental/Iris/Core/Public/Iris/ReplicationSystem/ObjectReplicationBridge.h
@@ -10,6 +10,7 @@
 #include "Iris/ReplicationSystem/NetObjectFactoryRegistry.h"
 
 #include "Containers/ArrayView.h"
+#include "Containers/Set.h"
 
 #include "Delegates/IDelegateInstance.h"
 
@@ -34,6 +35,7 @@ namespace UE::Net
 	typedef uint32 FReplicationProtocolIdentifier;
 
 	class FNetObjectReference;
+	class FWorldLocations;
 
 	namespace Private
 	{
@@ -429,6 +431,15 @@ private:
 	/** Marks a spatially filtered object as requiring or not requiring frequent world location updates independent of it having dirty replicated properties */
 	void OptionallySetObjectRequiresFrequentWorldLocationUpdate(FNetRefHandle RefHandle, bool bDesiresFrequentWorldLocationUpdate);
 
+	struct FUpdateWorldInfoContext
+	{
+		UE::Net::FWorldLocations& WorldLocations;
+		UE::Net::Private::FInternalNetRefIndex ObjectIndex;
+		FNetRefHandle NetRefHandle;
+	};
+	/** Update the WorldLocationInfo of a specific root object */
+	void UpdateRootObjectWorldInfo(const FUpdateWorldInfoContext& Context);
+
 	/** Returns the TypeStatsIndex this class should use */
 	int32 GetTypeStatsIndex(const UClass* Class);
 
@@ -527,8 +538,12 @@ private:
 
 protected:
 	bool bSuppressCreateInstanceFailedEnsure = false;
+
 private:
 	TMap<FObjectKey, bool> ArchetypesAlreadyPrinted;
+
+	/** Tells if we already logged a warning about a given object or a class type that had a location that is out of bounds. */
+	TSet<FObjectKey> WorldLocationOOBWarnings;
 };
 
 ENUM_CLASS_FLAGS(UObjectReplicationBridge::EPrintDebugInfoTraits);
\ No newline at end of file
diff --git a/Engine/Source/Runtime/Experimental/Iris/Core/Public/Iris/ReplicationSystem/WorldLocations.h b/Engine/Source/Runtime/Experimental/Iris/Core/Public/Iris/ReplicationSystem/WorldLocations.h
--- a/Engine/Source/Runtime/Experimental/Iris/Core/Public/Iris/ReplicationSystem/WorldLocations.h
+++ b/Engine/Source/Runtime/Experimental/Iris/Core/Public/Iris/ReplicationSystem/WorldLocations.h
@@ -149,11 +149,18 @@ public:
 	const FVector& GetWorldMaxPos() const { return MaxWorldPos; };
 	
 	/** Return a position clamped to the configured world boundary. */
-	FVector ClampPositionToBoundary(const FVector& Position)
+	FVector ClampPositionToBoundary(const FVector& Position) const
 	{
 		return Position.BoundToBox(GetWorldMinPos(), GetWorldMaxPos());
 	}
 
+	/** Is the location without the configured Min/Max WorldPos*/
+	bool IsValidLocation(const FVector& Location) const
+	{
+		return (Location.X >= MinWorldPos.X && Location.Y >= MinWorldPos.Y && Location.Z >= MinWorldPos.Z &&
+				Location.X <= MaxWorldPos.X && Location.Y <= MaxWorldPos.Y && Location.Z <= MaxWorldPos.Z);
+	}
+
 	void OnMaxInternalNetRefIndexIncreased(UE::Net::Private::FInternalNetRefIndex NewMaxInternalIndex);
 
 private:
