Include container header in the on-demand TOC file.

- simplifies mounting on-demand container files used for installing content where the container header is always needed
- enabled by default but can be disabled by passing -IgnoreContainerHeader to IasTool.exe Upload command

#rnx

diff --git a/Engine/Source/Runtime/Experimental/IoStore/OnDemand/Private/Tool/ChunkPlugin.cpp b/Engine/Source/Runtime/Experimental/IoStore/OnDemand/Private/Tool/ChunkPlugin.cpp
--- a/Engine/Source/Runtime/Experimental/IoStore/OnDemand/Private/Tool/ChunkPlugin.cpp
+++ b/Engine/Source/Runtime/Experimental/IoStore/OnDemand/Private/Tool/ChunkPlugin.cpp
@@ -175,7 +175,7 @@ public:
 
 		if (TUniquePtr<FArchive> Ar(IFileManager::Get().CreateFileWriter(Sb.ToString())); Ar.IsValid())
 		{
-			UE_LOG(LogIoStore, Display, TEXT("Writing file '%s' (%.2lf KiB)"),
+			UE_LOG(LogIoStoreOnDemand, Display, TEXT("Writing file '%s' (%.2lf KiB)"),
 				Sb.ToString(), double(Chunk.GetView().GetSize()) / 1024.0);
 			Ar->Serialize((void*)Chunk.GetView().GetData(), Chunk.GetView().GetSize());
 
@@ -218,7 +218,7 @@ static FIoStatus WriteChunk(const FString& Directory, FMemoryView Chunk, const F
 
 	if (TUniquePtr<FArchive> Ar(IFileManager::Get().CreateFileWriter(Sb.ToString())); Ar.IsValid())
 	{
-		UE_LOG(LogIoStore, Display, TEXT("Writing file '%s' (%.2lf KiB)"),
+		UE_LOG(LogIoStoreOnDemand, Display, TEXT("Writing file '%s' (%.2lf KiB)"),
 			Sb.ToString(), double(Chunk.GetSize()) / 1024.0);
 		Ar->Serialize((void*)Chunk.GetData(), Chunk.GetSize());
 
@@ -249,6 +249,7 @@ static int32 ChunkPluginCommandEntry(const FContext& Context)
 	FString OutputStatsJson				= FString(Context.Get<FStringView>(TEXT("-OutputStatsJson"), FString()));
 	const FString HostGroupName			= FString(Context.Get<FStringView>(TEXT("-HostGroupName"), FString()));
 
+	const bool bIgnoreContainerHeader	= Context.Get<bool>(TEXT("-IgnoreContainerHeader"), false);
 	const bool bStreamOnDemand			= Context.Get<bool>(TEXT("-StreamOnDemand"), false);
 	const bool bInstallOnDemand			= Context.Get<bool>(TEXT("-InstallOnDemand"), false);
 
@@ -264,7 +265,7 @@ static int32 ChunkPluginCommandEntry(const FContext& Context)
 	else
 	{
 		const EOnDemandTocFlags FallbackFlags = EOnDemandTocFlags::InstallOnDemand;
-		UE_LOG(LogIoStore, Warning, TEXT("No TOC flags was specified, falling back to '%s'"), *LexToString(FallbackFlags));
+		UE_LOG(LogIoStoreOnDemand, Warning, TEXT("No TOC flags was specified, falling back to '%s'"), *LexToString(FallbackFlags));
 		TocFlags = FallbackFlags;
 	}
 
@@ -282,19 +283,20 @@ static int32 ChunkPluginCommandEntry(const FContext& Context)
 	FPaths::NormalizeFilename(SettingsFile);
 	FPaths::NormalizeFilename(OutputStatsJson);
 
-	UE_LOG(LogIoStore, Display, TEXT("I/O store chunk plugin:"));
-	UE_LOG(LogIoStore, Display, TEXT("----------------------------------------"));
-	UE_LOG(LogIoStore, Display, TEXT("\tBuildVersion: %s"), *BuildVersion);
-	UE_LOG(LogIoStore, Display, TEXT("\tPlatform: %s"), *Platform);
-	UE_LOG(LogIoStore, Display, TEXT("\tOnDemandTocName: %s"), *OnDemandTocName);
-	UE_LOG(LogIoStore, Display, TEXT("\tInputFolder: %s"), *InputFolder);
-	UE_LOG(LogIoStore, Display, TEXT("\tOutputFolder: %s"), *OutputFolder);
-	UE_LOG(LogIoStore, Display, TEXT("\tIntermediateFolder: %s"), *IntermediateFolder);
-	UE_LOG(LogIoStore, Display, TEXT("\tSettingsFile: %s"), *SettingsFile);
-	UE_LOG(LogIoStore, Display, TEXT("\tOutputStatsJson: %s"), *OutputStatsJson);
-	UE_LOG(LogIoStore, Display, TEXT("\tIncludeSigPak: %s"), bIncludeSigPak ? TEXT("true") : TEXT("false"));
-	UE_LOG(LogIoStore, Display, TEXT("\tDeleteContainerFiles: %s"), bDeleteContainerFiles ? TEXT("true") : TEXT("false"));
-	UE_LOG(LogIoStore, Display, TEXT("\tTOC flags: '%s'"), *LexToString(TocFlags));
+	UE_LOG(LogIoStoreOnDemand, Display, TEXT("I/O store chunk plugin:"));
+	UE_LOG(LogIoStoreOnDemand, Display, TEXT("----------------------------------------"));
+	UE_LOG(LogIoStoreOnDemand, Display, TEXT("\tBuildVersion: %s"), *BuildVersion);
+	UE_LOG(LogIoStoreOnDemand, Display, TEXT("\tPlatform: %s"), *Platform);
+	UE_LOG(LogIoStoreOnDemand, Display, TEXT("\tOnDemandTocName: %s"), *OnDemandTocName);
+	UE_LOG(LogIoStoreOnDemand, Display, TEXT("\tInputFolder: %s"), *InputFolder);
+	UE_LOG(LogIoStoreOnDemand, Display, TEXT("\tOutputFolder: %s"), *OutputFolder);
+	UE_LOG(LogIoStoreOnDemand, Display, TEXT("\tIntermediateFolder: %s"), *IntermediateFolder);
+	UE_LOG(LogIoStoreOnDemand, Display, TEXT("\tSettingsFile: %s"), *SettingsFile);
+	UE_LOG(LogIoStoreOnDemand, Display, TEXT("\tOutputStatsJson: %s"), *OutputStatsJson);
+	UE_LOG(LogIoStoreOnDemand, Display, TEXT("\tIncludeSigPak: %s"), bIncludeSigPak ? TEXT("true") : TEXT("false"));
+	UE_LOG(LogIoStoreOnDemand, Display, TEXT("\tDeleteContainerFiles: %s"), bDeleteContainerFiles ? TEXT("true") : TEXT("false"));
+	UE_LOG(LogIoStoreOnDemand, Display, TEXT("\tIgnoreContainerHeader: %s"), bIgnoreContainerHeader ? TEXT("true") : TEXT("false"));
+	UE_LOG(LogIoStoreOnDemand, Display, TEXT("\tTOC flags: '%s'"), *LexToString(TocFlags));
 
 	IFileManager& FileMgr = IFileManager::Get();
 
@@ -316,7 +318,7 @@ static int32 ChunkPluginCommandEntry(const FContext& Context)
 		TUniquePtr<FArchive> Ar(FileMgr.CreateFileReader(*SettingsFile));
 		if (!Ar || Ar->IsError())
 		{
-			UE_LOG(LogIoStore, Error, TEXT("Failed to open settings file '%s'"), *SettingsFile);
+			UE_LOG(LogIoStoreOnDemand, Error, TEXT("Failed to open settings file '%s'"), *SettingsFile);
 			return -1;
 		}
 
@@ -325,14 +327,14 @@ static int32 ChunkPluginCommandEntry(const FContext& Context)
 		TSharedPtr<FJsonValue> JsonSettings;
 		if (!FJsonSerializer::Deserialize(*JsonReader, JsonSettings))
 		{
-			UE_LOG(LogIoStore, Error, TEXT("Failed to read settings file '%s'"), *SettingsFile);
+			UE_LOG(LogIoStoreOnDemand, Error, TEXT("Failed to read settings file '%s'"), *SettingsFile);
 			return -1;
 		}
 
 		TSharedPtr<FJsonObject> JsonSettingsObject = JsonSettings->AsObject();
 		if (!JsonSettingsObject)
 		{
-			UE_LOG(LogIoStore, Error, TEXT("Bad settings file '%s'"), *SettingsFile);
+			UE_LOG(LogIoStoreOnDemand, Error, TEXT("Bad settings file '%s'"), *SettingsFile);
 			return -1;
 		}
 
@@ -341,7 +343,7 @@ static int32 ChunkPluginCommandEntry(const FContext& Context)
 			TSharedPtr<FJsonObject> PackageSetsObject = PackageSets->AsObject();
 			if (!PackageSetsObject)
 			{
-				UE_LOG(LogIoStore, Error, TEXT("Bad settings file '%s'"), *SettingsFile);
+				UE_LOG(LogIoStoreOnDemand, Error, TEXT("Bad settings file '%s'"), *SettingsFile);
 				return -1;
 			}
 
@@ -350,14 +352,14 @@ static int32 ChunkPluginCommandEntry(const FContext& Context)
 				TArray<FString> Packages;
 				if (!PackageSetsObject->TryGetStringArrayField(Pair.Key, Packages))
 				{
-					UE_LOG(LogIoStore, Error, TEXT("Bad settings file '%s'"), *SettingsFile);
+					UE_LOG(LogIoStoreOnDemand, Error, TEXT("Bad settings file '%s'"), *SettingsFile);
 					return -1;
 				}
 
-				UE_LOG(LogIoStore, Display, TEXT("Found Package Set '%s'"), *Pair.Key);
+				UE_LOG(LogIoStoreOnDemand, Display, TEXT("Found Package Set '%s'"), *Pair.Key);
 				for (const FString& Package : Packages)
 				{
-					UE_LOG(LogIoStore, Display, TEXT("-> \t'%s'"), *Package);
+					UE_LOG(LogIoStoreOnDemand, Display, TEXT("-> \t'%s'"), *Package);
 				}
 
 				Settings.PackageSets.Add(Pair.Key, MoveTemp(Packages));
@@ -370,13 +372,13 @@ static int32 ChunkPluginCommandEntry(const FContext& Context)
 	TMap<FGuid, FAES::FAESKey> EncryptionKeys;
 	if (FileMgr.DirectoryExists(*ContainerFolder) == false)
 	{
-		UE_LOG(LogIoStore, Error, TEXT("Directory '%s' does not exist"), *ContainerFolder);
+		UE_LOG(LogIoStoreOnDemand, Error, TEXT("Directory '%s' does not exist"), *ContainerFolder);
 		return -1;
 	}
 
 	TArray<FString> ContainerFilenames;
 	FileMgr.FindFiles(ContainerFilenames, *ContainerFolder, TEXT("*.utoc"));
-	UE_LOG(LogIoStore, Display, TEXT("Found %d container files(s)"), ContainerFilenames.Num());
+	UE_LOG(LogIoStoreOnDemand, Display, TEXT("Found %d container files(s)"), ContainerFilenames.Num());
 
 	FOnDemandToc OnDemandToc;
 	OnDemandToc.Header.Flags = static_cast<uint32>(TocFlags);
@@ -397,18 +399,18 @@ static int32 ChunkPluginCommandEntry(const FContext& Context)
 			FIoStatus Status = ContainerFileReader.Initialize(*FPaths::ChangeExtension(FullPath, TEXT("")), EncryptionKeys);
 			if (!Status.IsOk())
 			{
-				UE_LOG(LogIoStore, Error, TEXT("Failed to open container '%s' for reading"), *FullPath);
+				UE_LOG(LogIoStoreOnDemand, Error, TEXT("Failed to open container '%s' for reading"), *FullPath);
 				continue;
 			}
 		}
 
-		UE_LOG(LogIoStore, Display, TEXT("Processing container '%s'"), *FullPath);
+		UE_LOG(LogIoStoreOnDemand, Display, TEXT("Processing container '%s'"), *FullPath);
 
 		/*TArray<FString> ContainerFileList;
 		ContainerFileReader.GetFilenames(ContainerFileList);
 		for (const FString& File : ContainerFileList)
 		{
-			UE_LOG(LogIoStore, Display, TEXT("Found file in container '%s'"), *File);
+			UE_LOG(LogIoStoreOnDemand, Display, TEXT("Found file in container '%s'"), *File);
 		}*/
 
 		const uint32 BlockSize = ContainerFileReader.GetCompressionBlockSize();
@@ -437,11 +439,11 @@ static int32 ChunkPluginCommandEntry(const FContext& Context)
 			{
 				if (Status.Status().GetErrorCode() == EIoErrorCode::NotFound)
 				{
-					UE_LOG(LogIoStore, Display, TEXT("No header chunk for container '%s'"), *FullPath);
+					UE_LOG(LogIoStoreOnDemand, Display, TEXT("No header chunk for container '%s'"), *FullPath);
 				}
 				else
 				{
-					UE_LOG(LogIoStore, Error, TEXT("Failed to read header chunk for container '%s'"), *FullPath);
+					UE_LOG(LogIoStoreOnDemand, Error, TEXT("Failed to read header chunk for container '%s'"), *FullPath);
 					return -1;
 				}
 			}
@@ -453,7 +455,7 @@ static int32 ChunkPluginCommandEntry(const FContext& Context)
 				Ar.Close();
 				if (Ar.IsError() || Ar.IsCriticalError())
 				{
-					UE_LOG(LogIoStore, Error, TEXT("Failed to deserialize header chunk for container '%s'"), *FullPath);
+					UE_LOG(LogIoStoreOnDemand, Error, TEXT("Failed to deserialize header chunk for container '%s'"), *FullPath);
 					return -1;
 				}
 
@@ -470,7 +472,7 @@ static int32 ChunkPluginCommandEntry(const FContext& Context)
 						int32 PackageIndex = INDEX_NONE;
 						if (ContainerHeader.PackageIds.Find(PackageId, PackageIndex))
 						{
-							UE_LOG(LogIoStore, Display, TEXT("Found package %s for tag %s in container %s"), *Package, *Tag, *Filename)
+							UE_LOG(LogIoStoreOnDemand, Display, TEXT("Found package %s for tag %s in container %s"), *Package, *Tag, *Filename)
 							TagSetPackageList.PackageIndicies.Add(PackageIndex);
 						}
 					}
@@ -490,7 +492,7 @@ static int32 ChunkPluginCommandEntry(const FContext& Context)
 			return true;
 		});
 
-		UE_LOG(LogIoStore, Display, TEXT("Serializing %d chunks"), ChunkInfos.Num());
+		UE_LOG(LogIoStoreOnDemand, Display, TEXT("Serializing %d chunks"), ChunkInfos.Num());
 
 		for (const FIoStoreTocChunkInfo& ChunkInfo : ChunkInfos)
 		{
@@ -498,7 +500,7 @@ static int32 ChunkPluginCommandEntry(const FContext& Context)
 			TIoStatusOr<FIoStoreCompressedReadResult> Status = ContainerFileReader.ReadCompressed(ChunkInfo.Id, FIoReadOptions(), bDecrypt);
 			if (!Status.IsOk())
 			{
-				UE_LOG(LogIoStore, Error, TEXT("Failed to read container chunk, Container='%s', Reason='%s'"),
+				UE_LOG(LogIoStoreOnDemand, Error, TEXT("Failed to read container chunk, Container='%s', Reason='%s'"),
 					*FullPath, *Status.Status().ToString());
 				return -1;
 			}
@@ -534,7 +536,7 @@ static int32 ChunkPluginCommandEntry(const FContext& Context)
 
 			if (EncodedChunkSize != ReadResult.IoBuffer.GetSize())
 			{
-				UE_LOG(LogIoStore, Error, TEXT("Chunk size mismatch, Container='%s', ChunkId='%s'"),
+				UE_LOG(LogIoStoreOnDemand, Error, TEXT("Chunk size mismatch, Container='%s', ChunkId='%s'"),
 					*FullPath, *LexToString(ChunkInfo.Id));
 				return -1;
 			}
@@ -542,7 +544,7 @@ static int32 ChunkPluginCommandEntry(const FContext& Context)
 			FIoStatus WriteStatus = ChunkWriter->WriteChunk(ChunksRelativeFolder, ReadResult.IoBuffer, ChunkHash);
 			if (WriteStatus.IsOk() == false)
 			{
-				UE_LOG(LogIoStore, Error, TEXT("%s"), *WriteStatus.ToString());
+				UE_LOG(LogIoStoreOnDemand, Error, TEXT("%s"), *WriteStatus.ToString());
 				return -1;
 			}
 
@@ -557,6 +559,32 @@ static int32 ChunkPluginCommandEntry(const FContext& Context)
 			Stats.TotalChunksSizeKb += (double)EncodedChunkSize / 1024.0;
 		}
 
+
+		if (bIgnoreContainerHeader)
+		{
+			UE_LOG(LogIoStoreOnDemand, Display, TEXT("Ignoring container header for '%s'"), *Filename)
+		}
+		else
+		{
+			const bool bDecrypt = false;
+			const FIoChunkId ChunkId = CreateContainerHeaderChunkId(ContainerFileReader.GetContainerId());
+			TIoStatusOr<FIoStoreCompressedReadResult> ReadResult = ContainerFileReader.ReadCompressed(ChunkId, FIoReadOptions(), bDecrypt);
+			if (ReadResult.IsOk())
+			{
+				const FIoBuffer& Chunk = ReadResult.ValueOrDie().IoBuffer;
+				ContainerEntry.Header = MakeArrayView<const uint8>(Chunk.GetData(), static_cast<int32>(Chunk.GetSize()));
+			}
+			else
+			{
+				const EIoErrorCode Err = ReadResult.Status().GetErrorCode();
+				if (Err != EIoErrorCode::UnknownChunkID && Err != EIoErrorCode::NotFound)
+				{
+					UE_LOG(LogIoStoreOnDemand, Error, TEXT("Failed to read container header, reason '%s'"), *ReadResult.Status().ToString())
+					return -1; 
+				}
+			}
+		}
+
 		if (bDeleteContainerFiles)
 		{
 			FilesToDelete.Add(FullPath);
@@ -574,13 +602,12 @@ static int32 ChunkPluginCommandEntry(const FContext& Context)
 
 	for (const FString& Path : FilesToDelete)
 	{
-		//UE_LOG(LogIas, Display, TEXT("Attempt Deleting '%s'"), *Path);
 		if (FileMgr.FileExists(*Path))
 		{
-			UE_LOG(LogIas, Display, TEXT("Deleting '%s'"), *Path);
+			UE_LOG(LogIoStoreOnDemand, Display, TEXT("Deleting '%s'"), *Path);
 			if (!FileMgr.Delete(*Path, /*RequireExists*/true))
 			{
-				UE_LOG(LogIas, Error, TEXT("Failed to delete '%s'"), *Path);
+				UE_LOG(LogIoStoreOnDemand, Error, TEXT("Failed to delete '%s'"), *Path);
 			}
 		}
 	}
@@ -598,7 +625,7 @@ static int32 ChunkPluginCommandEntry(const FContext& Context)
 		TArray<FString> AdditionalFiles;
 		FileMgr.FindFiles(AdditionalFiles, *ContainerFolder, TEXT("*.*"));
 
-		UE_LOG(LogIoStore, Display, TEXT("Serializing %d additional file(s)"), AdditionalFiles.Num());
+		UE_LOG(LogIoStoreOnDemand, Display, TEXT("Serializing %d additional file(s)"), AdditionalFiles.Num());
 		for (const FString& Filename : AdditionalFiles)
 		{
 			const bool bIncludeDot = false;
@@ -615,7 +642,7 @@ static int32 ChunkPluginCommandEntry(const FContext& Context)
 			TArray<uint8> FileData;
 			if (!FFileHelper::LoadFileToArray(FileData, *FullPath))
 			{
-				UE_LOG(LogIoStore, Error, TEXT("Failed reading file '%s'"), *FullPath);
+				UE_LOG(LogIoStoreOnDemand, Error, TEXT("Failed reading file '%s'"), *FullPath);
 				return -1;
 			}
 
@@ -628,11 +655,11 @@ static int32 ChunkPluginCommandEntry(const FContext& Context)
 
 			if (WriteStatus.IsOk() == false)
 			{
-				UE_LOG(LogIoStore, Error, TEXT("%s"), *WriteStatus.ToString());
+				UE_LOG(LogIoStoreOnDemand, Error, TEXT("%s"), *WriteStatus.ToString());
 				return -1;
 			}
 
-			UE_LOG(LogIoStore, Display, TEXT("Adding additional file '%s'"), *Filename);
+			UE_LOG(LogIoStoreOnDemand, Display, TEXT("Adding additional file '%s'"), *Filename);
 			OnDemandToc.AdditionalFiles.Add(FOnDemandTocAdditionalFile
 			{
 				.Hash = ChunkHash,
@@ -658,19 +685,19 @@ static int32 ChunkPluginCommandEntry(const FContext& Context)
 			Ar->Close();
 			if (Ar->IsError())
 			{
-				UE_LOG(LogIoStore, Error, TEXT("Failed to serialize TOC '%s'"), *TocPath);
+				UE_LOG(LogIoStoreOnDemand, Error, TEXT("Failed to serialize TOC '%s'"), *TocPath);
 				return -1;
 			}
 			else
 			{
 				const int64 TocSize = Ar->Tell();
-				UE_LOG(LogIoStore, Display, TEXT("Writing file '%s' (%.2lf KiB)"),
+				UE_LOG(LogIoStoreOnDemand, Display, TEXT("Writing file '%s' (%.2lf KiB)"),
 					*TocPath, double(TocSize) / 1024);
 			}
 		}
 		else
 		{
-			UE_LOG(LogIoStore, Display, TEXT("Failed writing file '%s'"), *TocPath );
+			UE_LOG(LogIoStoreOnDemand, Display, TEXT("Failed writing file '%s'"), *TocPath );
 			return -1;
 		}
 	}
@@ -712,18 +739,18 @@ static int32 ChunkPluginCommandEntry(const FContext& Context)
 			TIoStatusOr<uint64> Status = FIoStoreTocResource::Write(*FullPath, Toc, 0, 0, ContainerSettings);
 			if (!Status.IsOk())
 			{
-				UE_LOG(LogIoStore, Error, TEXT("Failed to write dummy container '%s' (%s)"), *FullPath, *Status.Status().ToString());
+				UE_LOG(LogIoStoreOnDemand, Error, TEXT("Failed to write dummy container '%s' (%s)"), *FullPath, *Status.Status().ToString());
 				return -1;
 			}
 
-			UE_LOG(LogIoStore, Display, TEXT("Wrote dummy file '%s' (%.2lf KiB)"), *FullPath, double(Status.ValueOrDie()) / 1024);
+			UE_LOG(LogIoStoreOnDemand, Display, TEXT("Wrote dummy file '%s' (%.2lf KiB)"), *FullPath, double(Status.ValueOrDie()) / 1024);
 		}
 	}
 
 
 	if (!ChunkWriter->Flush())
 	{
-		UE_LOG(LogIoStore, Error, TEXT("Writer error: Failed to upload chunk(s)"));
+		UE_LOG(LogIoStoreOnDemand, Error, TEXT("Writer error: Failed to upload chunk(s)"));
 		return -1; 
 	}
 
@@ -732,7 +759,7 @@ static int32 ChunkPluginCommandEntry(const FContext& Context)
 		TUniquePtr<FArchive> FileWriter(FileMgr.CreateFileWriter(*OutputStatsJson));
 		if (!FileWriter)
 		{
-			UE_LOG(LogIoStore, Display, TEXT("Failed writing stats file '%s'"), *OutputStatsJson);
+			UE_LOG(LogIoStoreOnDemand, Display, TEXT("Failed writing stats file '%s'"), *OutputStatsJson);
 			return -1;
 		}
 		TSharedRef<TJsonWriter<UTF8CHAR>> JsonWriter = TJsonWriterFactory<UTF8CHAR>::Create(FileWriter.Get());
diff --git a/Engine/Source/Runtime/Experimental/IoStore/OnDemand/Private/Tool/Upload.cpp b/Engine/Source/Runtime/Experimental/IoStore/OnDemand/Private/Tool/Upload.cpp
--- a/Engine/Source/Runtime/Experimental/IoStore/OnDemand/Private/Tool/Upload.cpp
+++ b/Engine/Source/Runtime/Experimental/IoStore/OnDemand/Private/Tool/Upload.cpp
@@ -60,6 +60,7 @@ struct FUploadParams
 	bool bDeletePakFiles = true;
 	bool bPerContainerTocs = false;
 	bool bIgnoreContainerFlags = false;
+	bool bIgnoreContainerHeader = false; 
 
 	/** If we should write out the .iochunktoc to disk as well as uploading it. */
 	bool bWriteTocToDisk = false;
@@ -134,6 +135,7 @@ static FUploadParams BuildUploadParams(const FContext& Context)
 	Ret.HostGroupName			= Context.Get<FStringView>(TEXT("-HostGroupName"),			Ret.HostGroupName);
 	Ret.bWriteTocToDisk			= Context.Get<bool>(TEXT("-WriteTocToDisk"),				Ret.bWriteTocToDisk);
 	Ret.bPerContainerTocs		= Context.Get<bool>(TEXT("-PerContainerTocs"),				Ret.bPerContainerTocs);
+	Ret.bIgnoreContainerHeader	= Context.Get<bool>(TEXT("-IgnoreContainerHeader"),			Ret.bIgnoreContainerHeader);
 	Ret.MaxConcurrentUploads	= Context.Get<int32>(TEXT("-MaxConcurrentUploads"),			Ret.MaxConcurrentUploads);
 	Ret.MaxTocListCount			= Context.Get<int32>(TEXT("-MaxTocListCount"),				Ret.MaxTocListCount);
 	Ret.MaxTocDownloadCount		= Context.Get<int32>(TEXT("-MaxTocDownloadCount"),			Ret.MaxTocDownloadCount);
@@ -708,6 +710,30 @@ static TIoStatusOr<FUploadResult> UploadContainerFiles(
 			}
 		}
 
+		if (UploadParams.bIgnoreContainerHeader)
+		{
+			UE_LOG(LogIoStoreOnDemand, Display, TEXT("Ignoring container header for '%s'"), *Path)
+		}
+		else
+		{
+			const bool bDecrypt = false;
+			const FIoChunkId ChunkId = CreateContainerHeaderChunkId(ContainerFileReader.GetContainerId());
+			TIoStatusOr<FIoStoreCompressedReadResult> ReadResult = ContainerFileReader.ReadCompressed(ChunkId, FIoReadOptions(), bDecrypt);
+			if (ReadResult.IsOk())
+			{
+				const FIoBuffer& Chunk = ReadResult.ValueOrDie().IoBuffer;
+				ContainerEntry.Header = MakeArrayView<const uint8>(Chunk.GetData(), static_cast<int32>(Chunk.GetSize()));
+			}
+			else
+			{
+				const EIoErrorCode Err = ReadResult.Status().GetErrorCode();
+				if (Err != EIoErrorCode::UnknownChunkID && Err != EIoErrorCode::NotFound)
+				{
+					return ReadResult.Status();
+				}
+			}
+		}
+
 		{
 			const FString UTocFilePath = FPaths::ChangeExtension(Path, TEXT(".utoc"));
 			TArray<uint8> Buffer;
@@ -967,6 +993,7 @@ static FCommand UploadCommand(
 		TArgument<bool>(TEXT("-IgnoreContainerFlags"),	TEXT("Whether to ignore the OnDemand container flag")),
 		TArgument<bool>(TEXT("-StreamOnDemand"),		TEXT("Set the content to be streamed on-demand")),
 		TArgument<bool>(TEXT("-InstallOnDemand"),		TEXT("Set the content to be installed on-demand")),
+		TArgument<bool>(TEXT("-IgnoreContainerHeader"),	TEXT("Do not serialize the container header as part of the TOC file.")),
 		TArgument<int32>(TEXT("-MaxConcurrentUploads"),	TEXT("Number of simultaneous uploads")),
 		TArgument<int32>(TEXT("-MaxTocListCount"),		TEXT("Maximum number of TOC file(s) to list from the bucket")),
 		TArgument<int32>(TEXT("-MaxTocDownloadCount"),	TEXT("Maximum number of TOC file(s) to download")),
diff --git a/Engine/Source/Runtime/Experimental/IoStore/OnDemandCore/Private/OnDemandToc.cpp b/Engine/Source/Runtime/Experimental/IoStore/OnDemandCore/Private/OnDemandToc.cpp
--- a/Engine/Source/Runtime/Experimental/IoStore/OnDemandCore/Private/OnDemandToc.cpp
+++ b/Engine/Source/Runtime/Experimental/IoStore/OnDemandCore/Private/OnDemandToc.cpp
@@ -4,6 +4,7 @@
 
 #include "HAL/FileManager.h"
 #include "HAL/FileManagerGeneric.h"
+#include "Memory/SharedBuffer.h"
 #include "Misc/EnumClassFlags.h"
 #include "Misc/Paths.h"
 #include "Misc/PathViews.h"
@@ -243,6 +244,11 @@ FArchive& operator<<(FArchive& Ar, FOnDemandTocContainerEntry& ContainerEntry)
 		Ar << ContainerEntry.ContainerFlags;
 	}
 
+	if (!Ar.IsLoading() || (TocVersion >= EOnDemandTocVersion::ContainerHeader))
+	{
+		Ar << ContainerEntry.Header;
+	}
+
 	return Ar;
 }
 
@@ -276,6 +282,11 @@ FCbWriter& operator<<(FCbWriter& Writer, const FOnDemandTocContainerEntry& Conta
 
 	Writer.AddHash(UTF8TEXTVIEW("UTocHash"), ContainerEntry.UTocHash);
 
+	if (!ContainerEntry.Header.IsEmpty())
+	{
+		Writer.AddBinary(UTF8TEXTVIEW("Header"), ContainerEntry.Header.GetData(), ContainerEntry.Header.Num());
+	}
+
 	Writer.EndObject();
 
 	return Writer;
@@ -322,6 +333,14 @@ bool LoadFromCompactBinary(FCbFieldView Field, FOnDemandTocContainerEntry& OutCo
 
 		OutContainer.UTocHash = Obj["UTocHash"].AsHash();
 
+		FMemoryView Header = Obj["Header"].AsBinaryView();
+		if (!Header.IsEmpty())
+		{
+			OutContainer.Header = MakeArrayView<const uint8>(
+				reinterpret_cast<const uint8*>(Header.GetData()),
+				IntCastChecked<int32>(Header.GetSize()));
+		}
+
 		return true;
 	}
 
diff --git a/Engine/Source/Runtime/Experimental/IoStore/OnDemandCore/Public/IO/OnDemandToc.h b/Engine/Source/Runtime/Experimental/IoStore/OnDemandCore/Public/IO/OnDemandToc.h
--- a/Engine/Source/Runtime/Experimental/IoStore/OnDemandCore/Public/IO/OnDemandToc.h
+++ b/Engine/Source/Runtime/Experimental/IoStore/OnDemandCore/Public/IO/OnDemandToc.h
@@ -34,6 +34,7 @@ enum class EOnDemandTocVersion : uint32
 	ContainerFlags	= 9,
 	TocFlags		= 10,
 	HostGroupName	= 11,
+	ContainerHeader	= 12,
 
 	LatestPlusOne,
 	Latest			= (LatestPlusOne - 1)
@@ -107,6 +108,7 @@ struct FOnDemandTocContainerEntry
 	TArray<FOnDemandTocEntry>	Entries;
 	TArray<uint32>				BlockSizes;
 	TArray<FIoBlockHash>		BlockHashes;
+	TArray<uint8>				Header;	
 
 	/** Hash of the .utoc file (on disk) used to generate this data */
 	FIoHash						UTocHash;
