[HWRT] Fix GeometryCollection using one BLAS per part not working in cooked builds.
- SubSections array was being stripped from cooked builds.

diff --git a/Engine/Source/Runtime/Experimental/GeometryCollectionEngine/Private/GeometryCollection/GeometryCollectionRenderData.cpp b/Engine/Source/Runtime/Experimental/GeometryCollectionEngine/Private/GeometryCollection/GeometryCollectionRenderData.cpp
--- a/Engine/Source/Runtime/Experimental/GeometryCollectionEngine/Private/GeometryCollection/GeometryCollectionRenderData.cpp
+++ b/Engine/Source/Runtime/Experimental/GeometryCollectionEngine/Private/GeometryCollection/GeometryCollectionRenderData.cpp
@@ -127,6 +127,9 @@ void FGeometryCollectionMeshDescription::Serialize(FArchive& Ar)
 {
 	Ar << NumVertices << NumTriangles << PreSkinnedBounds << Sections << SectionsNoInternal;
 	
+	// need subsection data on cooked builds when using ray tracing
+	// to be able to build one BLAS per subsection
+#if !RHI_RAYTRACING
 	if (Ar.IsCooking())
 	{
 		// No need for SubSections in cooked build.
@@ -134,6 +137,7 @@ void FGeometryCollectionMeshDescription::Serialize(FArchive& Ar)
 		Ar << Dummy;
 	}
 	else
+#endif
 	{
 		Ar << SubSections;
 	}
