[NetPhysics] Don't call PostInitialize_Internal if the ISimCallbackObject is set to be deleted on the same frame it's created.

diff --git a/Engine/Source/Runtime/Experimental/Chaos/Private/PBDRigidsSolver.cpp b/Engine/Source/Runtime/Experimental/Chaos/Private/PBDRigidsSolver.cpp
--- a/Engine/Source/Runtime/Experimental/Chaos/Private/PBDRigidsSolver.cpp
+++ b/Engine/Source/Runtime/Experimental/Chaos/Private/PBDRigidsSolver.cpp
@@ -1966,7 +1966,10 @@ namespace Chaos
 		
 		for (ISimCallbackObject* SimCallbackObject : PushData.SimCallbackObjectsToAdd)
 		{
-			SimCallbackObject->PostInitialize_Internal();
+			if (!SimCallbackObject->bPendingDelete)
+			{
+				SimCallbackObject->PostInitialize_Internal();
+			}
 		}
 
 		if (MRewindCallback && !IsShuttingDown())
