[HWRT] Added cvar to disable WPO evaluation for Geometry Collection ray tracing representation.

diff --git a/Engine/Source/Runtime/Experimental/GeometryCollectionEngine/Private/GeometryCollection/GeometryCollectionSceneProxy.cpp b/Engine/Source/Runtime/Experimental/GeometryCollectionEngine/Private/GeometryCollection/GeometryCollectionSceneProxy.cpp
--- a/Engine/Source/Runtime/Experimental/GeometryCollectionEngine/Private/GeometryCollection/GeometryCollectionSceneProxy.cpp
+++ b/Engine/Source/Runtime/Experimental/GeometryCollectionEngine/Private/GeometryCollection/GeometryCollectionSceneProxy.cpp
@@ -70,14 +70,22 @@ FAutoConsoleVariableRef CVarGeometryCollectionTripleBufferUploads(
 	ECVF_Default
 );
 
-static int32 GRayTracingGeometryCollectionProxyMeshes = 0;
-FAutoConsoleVariableRef CVarRayTracingGeometryCollectionProxyMeshes(
+static int32 GRayTracingGeometryCollection = 0;
+FAutoConsoleVariableRef CVarRayTracingGeometryCollection(
 	TEXT("r.RayTracing.Geometry.GeometryCollection"),
-	GRayTracingGeometryCollectionProxyMeshes,
+	GRayTracingGeometryCollection,
 	TEXT("Include geometry collection proxy meshes in ray tracing effects (default = 0 (Geometry collection meshes disabled in ray tracing))"),
 	ECVF_RenderThreadSafe
 );
 
+static bool GRayTracingGeometryCollectionWPO = true;
+static FAutoConsoleVariableRef CVarRayTracingGeometryCollectionWPO(
+	TEXT("r.RayTracing.Geometry.GeometryCollection.WPO"),
+	GRayTracingGeometryCollectionWPO,
+	TEXT("Whether to update geometry collection ray tracing representation based on material World Position Offset."),
+	ECVF_RenderThreadSafe
+);
+
 static TAutoConsoleVariable<int32> CVarRayTracingGeometryCollectionForceUpdate(
 	TEXT("r.RayTracing.Geometry.GeometryCollection.ForceUpdate"),
 	0,
@@ -541,7 +549,7 @@ void FGeometryCollectionSceneProxyBase::GetDynamicRayTracingInstances(FRayTracin
 {
 	checkf(bHasRayTracingRepresentation, TEXT("Shouldn't try to get ray tracing instances from proxy that doesn't have a ray tracing representation."));
 
-	if (GRayTracingGeometryCollectionProxyMeshes == 0)
+	if (GRayTracingGeometryCollection == 0)
 	{
 		return;
 	}
@@ -553,6 +561,11 @@ void FGeometryCollectionSceneProxyBase::GetDynamicRayTracingInstances(FRayTracin
 
 	QUICK_SCOPE_CYCLE_COUNTER(STAT_GeometryCollectionSceneProxyBase_GetDynamicRayTracingInstances);
 
+	if (!GRayTracingGeometryCollectionWPO)
+	{
+		bAnyMaterialHasWorldPositionOffset = false;
+	}
+
 	const uint32 LODIndex = 0;
 	const bool bWireframe = false;
 
