Fix WPO on nanite geometry collection.
The material relevance was not being set correctly.

diff --git a/Engine/Source/Runtime/Experimental/GeometryCollectionEngine/Private/GeometryCollection/GeometryCollectionSceneProxy.cpp b/Engine/Source/Runtime/Experimental/GeometryCollectionEngine/Private/GeometryCollection/GeometryCollectionSceneProxy.cpp
--- a/Engine/Source/Runtime/Experimental/GeometryCollectionEngine/Private/GeometryCollection/GeometryCollectionSceneProxy.cpp
+++ b/Engine/Source/Runtime/Experimental/GeometryCollectionEngine/Private/GeometryCollection/GeometryCollectionSceneProxy.cpp
@@ -1249,6 +1249,8 @@ FNaniteGeometryCollectionSceneProxy::FNaniteGeometryCollectionSceneProxy(UGeomet
 		MaterialSections[SectionIndex].bCastShadow = true;
 	}
 
+	OnMaterialsUpdated();
+
 	const bool bHasGeometryBoundingBoxes = 
 		Collection->HasAttribute("BoundingBox", FGeometryCollection::GeometryGroup) &&
 		Collection->NumElements(FGeometryCollection::GeometryGroup);
