Refactor the FOnDemandThread so it no longer relies on implementation details from IAS. At this point it should be possible for IAD to use it too.

#rnx

- FOnDemandThread::QueueRequest no longer accepts FChunkRequests as input, instead it takes the parameters needed to make the http request along with a callback. It will return a handle that the caller can use to reference the request in future if they wish to change it's priority or cancel the request.
-- Note that this handle is no longer valid once the callback has been completed.
- Added FOnDemandThread::CancelRequest which allows the caller to cancel an existing request using the handle that was provided when the request was first queued.

### Misc
- FOnDemandThread still lacks proper documentation as there is likely to be one more piece of work before the api is finalized.
- Move UE_ENABLE_REQUEST_THREAD from the Build.cs to OnDemandHttpThread.h.
- Move FChunk code back to OnDemandIoDispatcherBackend.cpp and TThreadSafeIntrusiveQueue to it's own header.
-- Now that we are no longer sharing chunk code between FOnDemandIoBackend and FOnDemandThread we can delete the shared header.

diff --git a/Engine/Source/Runtime/Experimental/IoStore/OnDemand/IoStoreOnDemand.Build.cs b/Engine/Source/Runtime/Experimental/IoStore/OnDemand/IoStoreOnDemand.Build.cs
--- a/Engine/Source/Runtime/Experimental/IoStore/OnDemand/IoStoreOnDemand.Build.cs
+++ b/Engine/Source/Runtime/Experimental/IoStore/OnDemand/IoStoreOnDemand.Build.cs
@@ -35,7 +35,5 @@ public class IoStoreOnDemand : ModuleRules
 		{
 			PrivateDefinitions.Add("UE_HTTPCLIENT_THREADSAFE_CHECKS=0");
 		}
-
-		PrivateDefinitions.Add("UE_ENABLE_REQUEST_THREAD=0");
 	}
 }
diff --git a/Engine/Source/Runtime/Experimental/IoStore/OnDemand/Private/ChunkRequestShared.h b/Engine/Source/Runtime/Experimental/IoStore/OnDemand/Private/ChunkRequestShared.h
deleted file mode 100644
diff --git a/Engine/Source/Runtime/Experimental/IoStore/OnDemand/Private/OnDemandHttpThread.cpp b/Engine/Source/Runtime/Experimental/IoStore/OnDemand/Private/OnDemandHttpThread.cpp
--- a/Engine/Source/Runtime/Experimental/IoStore/OnDemand/Private/OnDemandHttpThread.cpp
+++ b/Engine/Source/Runtime/Experimental/IoStore/OnDemand/Private/OnDemandHttpThread.cpp
@@ -4,6 +4,7 @@
 
 #if UE_ENABLE_REQUEST_THREAD
 
+#include "Async/UniqueLock.h"
 #include "HAL/IConsoleManager.h"
 #include "HAL/LowLevelMemTracker.h"
 #include "HAL/RunnableThread.h"
@@ -11,7 +12,9 @@
 #include "IasHostGroup.h"
 #include "Logging/StructuredLog.h"
 #include "OnDemandHttpClient.h"
+#include "OnDemandIoStore.h"
 #include "Statistics.h"
+#include "Tasks/Task.h"
 
 // TODO - Maybe do LexToString/LexFromString (dupe code in UnrealEngine.cpp ThreadPriorityToString)
 static const TCHAR* LexToString(EThreadPriority Priority)
@@ -73,6 +76,31 @@ const int32 GREPLACEOnDemandBackendThreadPriorities[5] =
 // END EXTERN CVARS
 ////////////////////////////////////////
 
+struct FHttpRequest
+{
+	void OnRequestCompleted(FIoBuffer&& ChunkData)
+	{
+		CompletionCallback(ChunkData);
+	}
+
+	FIASHostGroup HostGroup()
+	{
+		return ChunkInfo->HostGroup();
+	}
+
+	void GetUrl(FAnsiStringBuilderBase& Url) const
+	{
+		ChunkInfo->GetUrl(Url);
+	}
+
+	FHttpRequest* NextRequest = nullptr;
+	FOnDemandThread::FCompletionCallback CompletionCallback;
+	const FOnDemandChunkInfo* ChunkInfo = nullptr;
+	FIoOffsetAndLength ChunkRange;
+	bool bCancelled = false;
+	int32 Priority = 0;
+};
+
 FOnDemandThread* FOnDemandThread::SingletonInstance = nullptr;
 
 FOnDemandThread::FOnDemandThread()
@@ -108,20 +136,41 @@ FOnDemandThread& FOnDemandThread::Get()
 	return *SingletonInstance;
 }
 
-void FOnDemandThread::QueueRequest(FChunkRequest* Request)
+FOnDemandThread::FRequestHandle FOnDemandThread::IssueRequest(const FOnDemandChunkInfo* ChunkInfo, const FIoOffsetAndLength& ReadRange, FCompletionCallback&& CompletionCallback)
 {
-	TRACE_CPUPROFILER_EVENT_SCOPE(FOnDemandThread::QueueRequest);
+	TRACE_CPUPROFILER_EVENT_SCOPE(FOnDemandThread::IssueRequest);
+
+	checkf(ChunkInfo != nullptr, TEXT("FOnDemandThread::IssueRequest requires a valid ChunkInfo"));
+
+	FHttpRequest* Request = AllocateRequest();
+
+	Request->ChunkInfo = ChunkInfo;
+	Request->ChunkRange = ReadRange;
+	Request->CompletionCallback = MoveTemp(CompletionCallback);
 
 	FOnDemandIoBackendStats::Get()->OnHttpEnqueue();
 
 	HttpRequests.EnqueueByPriority(Request);
 	TickThreadEvent->Trigger();
+
+	return Request;
 }
 
-void FOnDemandThread::ReprioritizeRequest(FChunkRequest* Request)
+void FOnDemandThread::ReprioritizeRequest(FRequestHandle Request)
 {
 	TRACE_CPUPROFILER_EVENT_SCOPE(FOnDemandThread::ReprioritizeRequest);
-	HttpRequests.Reprioritize(Request);
+	if (Request != 0)
+	{
+		HttpRequests.Reprioritize(static_cast<FHttpRequest*>(Request));
+	}
+}
+
+void FOnDemandThread::CancelRequest(FRequestHandle Request)
+{
+	if (Request != 0)
+	{
+		static_cast<FHttpRequest*>(Request)->bCancelled = true;
+	}
 }
 
 bool FOnDemandThread::Init()
@@ -181,46 +230,48 @@ void FOnDemandThread::TickRequests()
 	const int32 MaxConcurrentRequests = FMath::Min(GIasHttpConcurrentRequests, 32);
 
 	int32 NumConcurrentRequests = 0;
-	FChunkRequest* NextChunkRequest = HttpRequests.Dequeue();
+	FHttpRequest* NextHttpRequest = HttpRequests.Dequeue();
 
-	while (NextChunkRequest)
+	while (NextHttpRequest)
 	{
-		while (NextChunkRequest)
+		while (NextHttpRequest)
 		{
 			{
 				TRACE_CPUPROFILER_EVENT_SCOPE(FOnDemandThread::IssueHttpGet);
-				FChunkRequest* ChunkRequest = NextChunkRequest;
-				NextChunkRequest = ChunkRequest->NextRequest;
-				ChunkRequest->NextRequest = nullptr;
+				FHttpRequest* HttpRequest = NextHttpRequest;
+				NextHttpRequest = HttpRequest->NextRequest;
+				HttpRequest->NextRequest = nullptr;
 
 				FOnDemandIoBackendStats::Get()->OnHttpDequeue();
 
-				if (ChunkRequest->bCancelled)
+				if (HttpRequest->bCancelled)
 				{
-					ChunkRequest->OnRequestCompleted();
+					HttpRequest->OnRequestCompleted(FIoBuffer());
 					FOnDemandIoBackendStats::Get()->OnHttpCancel();
+
+					DestroyRequest(HttpRequest);
 				}
-				else if (!ChunkRequest->HostGroup.IsConnected() || !bHttpEnabled)
+				else if (!HttpRequest->HostGroup().IsConnected() || !bHttpEnabled)
 				{
-					ChunkRequest->OnRequestCompleted();
+					HttpRequest->OnRequestCompleted(FIoBuffer());
 
 					// Technically this request is being skipped because of a pre-existing error. It is not
 					// an error itself and it is not being canceled by higher level code. However we do not
 					// currently have a statistic for that and we have to call one of the existing types in
 					// order to correctly reduce the pending count.
 					FOnDemandIoBackendStats::Get()->OnHttpCancel();
+
+					DestroyRequest(HttpRequest);
 				}
 				else
 				{
-					TAnsiStringBuilder<256> Url;
-					ChunkRequest->Params.GetUrl(Url);
-
 					NumConcurrentRequests++;
 
-					auto AnsiUrl = StringCast<ANSICHAR>(Url.GetData(), Url.Len());
+					TAnsiStringBuilder<128> ChunkUrl;
+					HttpRequest->GetUrl(ChunkUrl);
 
-					HttpClient->Get(ChunkRequest->HostGroup.GetUnderlyingHostGroup(), AnsiUrl, ChunkRequest->Params.ChunkRange,
-						[this, ChunkRequest, &NumConcurrentRequests]
+					HttpClient->Get(HttpRequest->HostGroup().GetUnderlyingHostGroup(), ChunkUrl, HttpRequest->ChunkRange,
+						[this, HttpRequest, &NumConcurrentRequests]
 						(FMultiEndpointHttpClientResponse&& HttpResponse)
 						{
 							TRACE_CPUPROFILER_EVENT_SCOPE(FOnDemandThread::RequestCallback);
@@ -236,27 +287,28 @@ void FOnDemandThread::TickRequests()
 
 							if (HttpResponse.IsOk())
 							{
-								ChunkRequest->HostGroup.OnSuccessfulResponse();
+								HttpRequest->HostGroup().OnSuccessfulResponse();
 
-								ChunkRequest->Chunk = MoveTemp(HttpResponse.Body);
-								FOnDemandIoBackendStats::Get()->OnHttpGet(ChunkRequest->Chunk.DataSize(), HttpResponse.DurationMilliseconds);
+								FOnDemandIoBackendStats::Get()->OnHttpGet(HttpResponse.Body.DataSize(), HttpResponse.DurationMilliseconds);
 
-								UE::Tasks::Launch(UE_SOURCE_LOCATION, [ ChunkRequest]()
+								UE::Tasks::Launch(UE_SOURCE_LOCATION, [this, HttpRequest, ChunkData = MoveTemp(HttpResponse.Body)]() mutable
 									{
-										ChunkRequest->OnRequestCompleted();
+										HttpRequest->OnRequestCompleted(MoveTemp(ChunkData));
+										DestroyRequest(HttpRequest);
 									}, GetRequestCompletionTaskPriority());
 							}
 							else
 							{
 								FOnDemandIoBackendStats::Get()->OnHttpError();
 
-								if (ChunkRequest->HostGroup.OnFailedResponse())
+								if (HttpRequest->HostGroup().OnFailedResponse())
 								{
 									// A disconnect was triggered
 									FOnDemandIoBackendStats::Get()->OnHttpDisconnected();
 								}
 
-								ChunkRequest->OnRequestCompleted();
+								HttpRequest->OnRequestCompleted(FIoBuffer());
+								DestroyRequest(HttpRequest);
 							}
 						});
 				}
@@ -271,9 +323,9 @@ void FOnDemandThread::TickRequests()
 				}
 			}
 
-			if (!NextChunkRequest)
+			if (!NextHttpRequest)
 			{
-				NextChunkRequest = HttpRequests.Dequeue();
+				NextHttpRequest = HttpRequests.Dequeue();
 			}
 		}
 
@@ -282,11 +334,12 @@ void FOnDemandThread::TickRequests()
 			TRACE_CPUPROFILER_EVENT_SCOPE(FOnDemandThread::TickHttp);
 			while (HttpClient->Tick(GIasHttpPollTimeoutMs, GIasHttpRateLimitKiBPerSecond))
 			{
-				if (!NextChunkRequest)
+				if (!NextHttpRequest)
 				{
-					NextChunkRequest = HttpRequests.Dequeue();
+					NextHttpRequest = HttpRequests.Dequeue();
 				}
-				if (NextChunkRequest)
+
+				if (NextHttpRequest)
 				{
 					break;
 				}
@@ -297,14 +350,15 @@ void FOnDemandThread::TickRequests()
 
 void FOnDemandThread::DrainHttpRequests()
 {
-	FChunkRequest* Iterator = HttpRequests.Dequeue();
+	FHttpRequest* Iterator = HttpRequests.Dequeue();
 	while (Iterator != nullptr)
 	{
-		FChunkRequest* Request = Iterator;
+		FHttpRequest* Request = Iterator;
 		Iterator = Iterator->NextRequest;
 
 		FOnDemandIoBackendStats::Get()->OnHttpDequeue();
-		Request->OnRequestCompleted();
+		Request->OnRequestCompleted(FIoBuffer());
+		DestroyRequest(Request);
 		FOnDemandIoBackendStats::Get()->OnHttpCancel();
 	}
 }
@@ -341,6 +395,18 @@ void FOnDemandThread::UpdateThreadPriorityIfNeeded()
 	}
 }
 
+FHttpRequest* FOnDemandThread::AllocateRequest()
+{
+	UE::TUniqueLock _(AllocatorMutex);
+	return RequestAllocator.Construct();
+}
+
+void FOnDemandThread::DestroyRequest(FHttpRequest* Request)
+{
+	UE::TUniqueLock _(AllocatorMutex);
+	RequestAllocator.Destroy(Request);
+}
+
 } //namespace UE::IoStore
 
 #endif // UE_ENABLE_REQUEST_THREAD
diff --git a/Engine/Source/Runtime/Experimental/IoStore/OnDemand/Private/OnDemandHttpThread.h b/Engine/Source/Runtime/Experimental/IoStore/OnDemand/Private/OnDemandHttpThread.h
--- a/Engine/Source/Runtime/Experimental/IoStore/OnDemand/Private/OnDemandHttpThread.h
+++ b/Engine/Source/Runtime/Experimental/IoStore/OnDemand/Private/OnDemandHttpThread.h
@@ -2,14 +2,23 @@
 
 #pragma once
 
+#ifndef UE_ENABLE_REQUEST_THREAD
+	#define UE_ENABLE_REQUEST_THREAD 0
+#endif
+
 #if UE_ENABLE_REQUEST_THREAD
 
-#include "ChunkRequestShared.h" // Very hefty
+#include "Async/Mutex.h"
+#include "Containers/StringFwd.h"
 #include "Delegates/Delegate.h"
 #include "HAL/Event.h"
 #include "HAL/PlatformAffinity.h"
 #include "HAL/Runnable.h"
+#include "IO/IoAllocators.h"
+#include "IO/IoBuffer.h"
+#include "IO/IoOffsetLength.h"
 #include "Templates/UniquePtr.h"
+#include "ThreadSafeIntrusiveQueue.h"
 
 class FRunnableThread;
 
@@ -17,21 +26,26 @@ namespace UE::IoStore
 {
 
 class FMultiEndpointHttpClient;
+struct FHttpRequest;
+struct FOnDemandChunkInfo;
 
 class FOnDemandThread final
 	: public FRunnable
 {
 public:
+	using FCompletionCallback = TUniqueFunction<void(FIoBuffer /*ChunkData*/)>;
+	using FRequestHandle = void*;
+	DECLARE_MULTICAST_DELEGATE(FOnTickIdle);
 
 	FOnDemandThread();
 	virtual ~FOnDemandThread();
 
 	static FOnDemandThread& Get();
 
-	void QueueRequest(FChunkRequest* Request);
-	void ReprioritizeRequest(FChunkRequest* Request);
+	FRequestHandle IssueRequest(const FOnDemandChunkInfo* ChunkInfo, const FIoOffsetAndLength& ReadRange, FCompletionCallback&& CompletionCallback);
+	void ReprioritizeRequest(FRequestHandle Request);
 
-	DECLARE_MULTICAST_DELEGATE(FOnTickIdle);
+	void CancelRequest(FRequestHandle Request);
 
 	FOnTickIdle& OnTickIdle()
 	{
@@ -51,6 +65,9 @@ private:
 	bool TryCreateHttpClient();
 	void UpdateThreadPriorityIfNeeded();
 
+	FHttpRequest* AllocateRequest();
+	void DestroyRequest(FHttpRequest* Request);
+
 private:
 	static FOnDemandThread* SingletonInstance;
 	TUniquePtr<FRunnableThread> Thread;
@@ -60,11 +77,28 @@ private:
 	FOnTickIdle OnTickIdleDelegate;
 	TUniquePtr<FMultiEndpointHttpClient> HttpClient;
 
-	FChunkRequestQueue HttpRequests;
+	TThreadSafeIntrusiveQueue<FHttpRequest> HttpRequests;
 
 	EThreadPriority ThreadPriority = EThreadPriority::TPri_Normal;
+
 	std::atomic_bool bStopRequested = false;
 	std::atomic_bool bHttpEnabled = true;
+
+	UE::FMutex AllocatorMutex;
+	TSingleThreadedSlabAllocator<FHttpRequest, 32> RequestAllocator;
+};
+
+} //namespace UE::IoStore
+
+#else
+
+namespace UE::IoStore
+{
+
+// Dummy class for FOnDemandIoStore::OnDemandHttpThread
+// Remove when we delete UE_ENABLE_REQUEST_THREAD
+class FOnDemandThread final
+{
 };
 
 } //namespace UE::IoStore
diff --git a/Engine/Source/Runtime/Experimental/IoStore/OnDemand/Private/OnDemandIoDispatcherBackend.cpp b/Engine/Source/Runtime/Experimental/IoStore/OnDemand/Private/OnDemandIoDispatcherBackend.cpp
--- a/Engine/Source/Runtime/Experimental/IoStore/OnDemand/Private/OnDemandIoDispatcherBackend.cpp
+++ b/Engine/Source/Runtime/Experimental/IoStore/OnDemand/Private/OnDemandIoDispatcherBackend.cpp
@@ -2,8 +2,6 @@
 
 #include "OnDemandIoDispatcherBackend.h"
 
-#include "AnalyticsEventAttribute.h"
-#include "ChunkRequestShared.h"
 #include "Containers/BitArray.h"
 #include "Containers/StringView.h"
 #include "DistributionEndpoints.h"
@@ -40,6 +38,7 @@
 #include "Statistics.h"
 #include "String/Numeric.h"
 #include "Tasks/Task.h"
+#include "ThreadSafeIntrusiveQueue.h"
 
 #include <atomic>
 
@@ -260,6 +259,7 @@ FAutoConsoleVariableRef CVarOnDemandBackendThreadPriority(
 	ECVF_Default);
 
 ///////////////////////////////////////////////////////////////////////////////
+
 [[nodiscard]] UE::Tasks::ETaskPriority GetRequestCompletionTaskPriority()
 {
 	return GCompleteMaterializeTaskPriorities[FMath::Clamp(GCompleteMaterializeTaskPriority, 0, UE_ARRAY_COUNT(GCompleteMaterializeTaskPriorities) - 1)];
@@ -267,6 +267,196 @@ FAutoConsoleVariableRef CVarOnDemandBackendThreadPriority(
 
 ///////////////////////////////////////////////////////////////////////////////
 
+static FIoHash GetChunkKey(const FIoHash& ChunkHash, const FIoOffsetAndLength& Range)
+{
+	FIoHashBuilder HashBuilder;
+	HashBuilder.Update(ChunkHash.GetBytes(), sizeof(FIoHash::ByteArray));
+	HashBuilder.Update(&Range, sizeof(FIoOffsetAndLength));
+
+	return HashBuilder.Finalize();
+}
+
+///////////////////////////////////////////////////////////////////////////////
+struct FChunkRequestParams
+{
+	static FChunkRequestParams Create(const FIoOffsetAndLength& OffsetLength, const FOnDemandChunkInfo& ChunkInfo)
+	{
+		FIoOffsetAndLength ChunkRange;
+		if (ChunkInfo.EncodedSize() <= (uint64(GIasHttpRangeRequestMinSizeKiB) << 10))
+		{
+			ChunkRange = FIoOffsetAndLength(0, ChunkInfo.EncodedSize());
+		}
+		else
+		{
+			const uint64 RawSize = FMath::Min<uint64>(OffsetLength.GetLength(), ChunkInfo.RawSize() - OffsetLength.GetOffset());
+
+			ChunkRange = FIoChunkEncoding::GetChunkRange(
+				ChunkInfo.RawSize(),
+				ChunkInfo.BlockSize(),
+				ChunkInfo.Blocks(),
+				OffsetLength.GetOffset(),
+				RawSize).ConsumeValueOrDie();
+		}
+
+		return FChunkRequestParams{ GetChunkKey(ChunkInfo.Hash(), ChunkRange), ChunkRange, ChunkInfo };
+	}
+
+	static FChunkRequestParams Create(FIoRequestImpl* Request, const FOnDemandChunkInfo& ChunkInfo)
+	{
+		check(Request);
+		check(Request->NextRequest == nullptr);
+		return Create(FIoOffsetAndLength(Request->Options.GetOffset(), Request->Options.GetSize()), ChunkInfo);
+	}
+
+	const FIoHash& GetUrlHash() const
+	{
+		return ChunkInfo.Hash();
+	}
+
+	void GetUrl(FAnsiStringBuilderBase& Url) const
+	{
+		ChunkInfo.GetUrl(Url);
+	}
+
+	FIoChunkDecodingParams GetDecodingParams() const
+	{
+		FIoChunkDecodingParams Params;
+		Params.EncryptionKey = ChunkInfo.EncryptionKey();
+		Params.CompressionFormat = ChunkInfo.CompressionFormat();
+		Params.BlockSize = ChunkInfo.BlockSize();
+		Params.TotalRawSize = ChunkInfo.RawSize();
+		Params.EncodedBlockSize = ChunkInfo.Blocks();
+		Params.BlockHash = ChunkInfo.BlockHashes();
+		Params.EncodedOffset = ChunkRange.GetOffset();
+
+		return Params;
+	}
+
+	FIoHash ChunkKey;
+	FIoOffsetAndLength ChunkRange;
+	FOnDemandChunkInfo ChunkInfo;
+};
+
+///////////////////////////////////////////////////////////////////////////////
+
+struct FChunkRequest
+{
+	explicit FChunkRequest(FIoRequestImpl* Request, const FChunkRequestParams& RequestParams)
+		: NextRequest()
+		, Params(RequestParams)
+		, RequestHead(Request)
+		, RequestTail(Request)
+		, StartTime(FPlatformTime::Cycles64())
+		, Priority(Request->Priority)
+		, RequestCount(1)
+		, bCached(false)
+	{
+		check(Request && NextRequest == nullptr);
+	}
+
+	bool AddDispatcherRequest(FIoRequestImpl* Request)
+	{
+		check(RequestHead && RequestTail);
+		check(Request && !Request->NextRequest);
+
+		const bool bPriorityChanged = Request->Priority > RequestHead->Priority;
+		if (bPriorityChanged)
+		{
+			Priority = Request->Priority;
+			Request->NextRequest = RequestHead;
+			RequestHead = Request;
+		}
+		else
+		{
+			FIoRequestImpl* It = RequestHead;
+			while (It->NextRequest != nullptr && Request->Priority <= It->NextRequest->Priority)
+			{
+				It = It->NextRequest;
+			}
+
+			if (RequestTail == It)
+			{
+				check(It->NextRequest == nullptr);
+				RequestTail = Request;
+			}
+
+			Request->NextRequest = It->NextRequest;
+			It->NextRequest = Request;
+		}
+
+		RequestCount++;
+		return bPriorityChanged;
+	}
+
+	int32 RemoveDispatcherRequest(FIoRequestImpl* Request)
+	{
+		check(Request != nullptr);
+		check(RequestCount > 0);
+
+		if (RequestHead == Request)
+		{
+			RequestHead = Request->NextRequest;
+			if (RequestTail == Request)
+			{
+				check(RequestHead == nullptr);
+				RequestTail = nullptr;
+			}
+		}
+		else
+		{
+			FIoRequestImpl* It = RequestHead;
+			while (It->NextRequest != Request)
+			{
+				It = It->NextRequest;
+				if (It == nullptr)
+				{
+					return INDEX_NONE; // Not found
+				}
+			}
+			check(It->NextRequest == Request);
+			It->NextRequest = It->NextRequest->NextRequest;
+
+			if (RequestTail == Request)
+			{
+				check(It->NextRequest == nullptr);
+				RequestTail = It;
+			}
+		}
+
+		Request->NextRequest = nullptr;
+		RequestCount--;
+
+		return RequestCount;
+	}
+
+	FIoRequestImpl* DeqeueDispatcherRequests()
+	{
+		FIoRequestImpl* Head = RequestHead;
+		RequestHead = RequestTail = nullptr;
+		RequestCount = 0;
+
+		return Head;
+	}
+
+	FChunkRequest* NextRequest;
+	FChunkRequestParams Params;
+	FIoRequestImpl* RequestHead;
+	FIoRequestImpl* RequestTail;
+#if UE_ENABLE_REQUEST_THREAD
+	FOnDemandThread::FRequestHandle HttpRequest = 0;
+#endif //UE_ENABLE_REQUEST_THREAD
+	FIASHostGroup HostGroup; // Still used in a few places even when UE_ENABLE_REQUEST_THREAD is enabled
+	FIoBuffer Chunk;
+	uint64 StartTime;
+	int32 Priority;
+	uint16 RequestCount;
+	bool bCached;
+	bool bCancelled = false;
+	EIoErrorCode CacheGetStatus;
+};
+
+//////////////////////////////////////////////////////////////////////////////
+
 static void LogIoResult(
 	const FIoChunkId& ChunkId,
 	const FIoHash& UrlHash,
@@ -327,6 +517,7 @@ class FOnDemandIoBackend final
 	, public IOnDemandIoDispatcherBackend
 {
 	using FIoRequestQueue = TThreadSafeIntrusiveQueue<FIoRequestImpl>;
+	using FChunkRequestQueue = TThreadSafeIntrusiveQueue<FChunkRequest>;
 
 	struct FBackendData
 	{
@@ -444,6 +635,10 @@ class FOnDemandIoBackend final
 					check(ChunkRequest.RequestTail == nullptr);
 
 					ChunkRequest.bCancelled = true;
+#if UE_ENABLE_REQUEST_THREAD
+					FOnDemandThread::Get().CancelRequest(ChunkRequest.HttpRequest);
+#endif // UE_ENABLE_REQUEST_THREAD
+
 					if (TheCache != nullptr)
 					{
 						TheCache->Cancel(ChunkRequest.Chunk);
@@ -891,8 +1086,15 @@ void FOnDemandIoBackend::CompleteCacheRead(FChunkRequest* ChunkRequest)
 	}
 
 #if UE_ENABLE_REQUEST_THREAD
-	ChunkRequest->CompletionCallback.BindRaw(this, &FOnDemandIoBackend::CompleteRequest);
-	FOnDemandThread::Get().QueueRequest(ChunkRequest);
+	ChunkRequest->HttpRequest = FOnDemandThread::Get().IssueRequest(&ChunkRequest->Params.ChunkInfo, ChunkRequest->Params.ChunkRange,
+		[this, ChunkRequest](FIoBuffer ChunkData)
+		{
+			ChunkRequest->Chunk = MoveTemp(ChunkData);
+			CompleteRequest(ChunkRequest);
+		});
+
+	check(ChunkRequest->HttpRequest != nullptr);
+
 #else
 	Stats.OnHttpEnqueue();
 	HttpRequests.EnqueueByPriority(ChunkRequest);
@@ -933,7 +1135,7 @@ bool FOnDemandIoBackend::Resolve(FIoRequestImpl* Request)
 		if (bUpdatePriority)
 		{
 #if UE_ENABLE_REQUEST_THREAD
-			FOnDemandThread::Get().ReprioritizeRequest(ChunkRequest);
+			FOnDemandThread::Get().ReprioritizeRequest(ChunkRequest->HttpRequest);
 #else
 			TRACE_CPUPROFILER_EVENT_SCOPE(IasBackend::UpdatePriorityForIoRequest);
 			HttpRequests.Reprioritize(ChunkRequest);
@@ -984,8 +1186,14 @@ bool FOnDemandIoBackend::Resolve(FIoRequestImpl* Request)
 	ChunkRequest->HostGroup = ChunkInfo.HostGroup();
 
 #if UE_ENABLE_REQUEST_THREAD
-	ChunkRequest->CompletionCallback.BindRaw(this, &FOnDemandIoBackend::CompleteRequest);
-	FOnDemandThread::Get().QueueRequest(ChunkRequest);
+	ChunkRequest->HttpRequest = FOnDemandThread::Get().IssueRequest(&ChunkRequest->Params.ChunkInfo, ChunkRequest->Params.ChunkRange,
+		[this, ChunkRequest](FIoBuffer ChunkData)
+		{
+			ChunkRequest->Chunk = MoveTemp(ChunkData);
+			CompleteRequest(ChunkRequest);
+		});
+
+	check(ChunkRequest->HttpRequest != nullptr);
 #else
 	Stats.OnHttpEnqueue();
 	HttpRequests.EnqueueByPriority(ChunkRequest);
@@ -1021,7 +1229,7 @@ void FOnDemandIoBackend::UpdatePriorityForIoRequest(FIoRequestImpl* Request)
 	if (FChunkRequest* ChunkRequest = ChunkRequests.TryUpdatePriority(Request))
 	{
 #if UE_ENABLE_REQUEST_THREAD
-		FOnDemandThread::Get().ReprioritizeRequest(ChunkRequest);
+		FOnDemandThread::Get().ReprioritizeRequest(ChunkRequest->HttpRequest);
 #else
 		HttpRequests.Reprioritize(ChunkRequest);
 #endif // UE_ENABLE_REQUEST_THREAD
diff --git a/Engine/Source/Runtime/Experimental/IoStore/OnDemand/Private/OnDemandIoStore.h b/Engine/Source/Runtime/Experimental/IoStore/OnDemand/Private/OnDemandIoStore.h
--- a/Engine/Source/Runtime/Experimental/IoStore/OnDemand/Private/OnDemandIoStore.h
+++ b/Engine/Source/Runtime/Experimental/IoStore/OnDemand/Private/OnDemandIoStore.h
@@ -197,13 +197,15 @@ struct FOnDemandChunkInfo
 	inline TConstArrayView<FIoBlockHash>	BlockHashes() const;
 	FAnsiStringView							ChunksDirectory() const { return SharedContainer->ChunksDirectory; }
 	const FOnDemandChunkEntry&				ChunkEntry() const { return Entry; }
-	inline FIASHostGroup					HostGroup() const;
+	inline const FIASHostGroup&				HostGroup() const;
 
 	bool									IsValid() const { return SharedContainer.IsValid(); }
 	operator								bool() const { return IsValid(); }
 	
 	inline static FOnDemandChunkInfo		Find(FSharedOnDemandContainer Container, const FIoChunkId& ChunkId);
 
+	inline void								GetUrl(FAnsiStringBuilderBase& Url) const;
+
 private:
 	friend class FOnDemandIoStore;
 
@@ -228,7 +230,7 @@ TConstArrayView<FIoBlockHash> FOnDemandChunkInfo::BlockHashes() const
 		: TConstArrayView<FIoBlockHash>(SharedContainer->BlockHashes.GetData() + Entry.BlockOffset, Entry.BlockCount);
 }
 
-FIASHostGroup FOnDemandChunkInfo::HostGroup() const
+const FIASHostGroup& FOnDemandChunkInfo::HostGroup() const
 {
 	return SharedContainer->IasHostGroup;
 }
@@ -244,6 +246,17 @@ FOnDemandChunkInfo FOnDemandChunkInfo::Find(FSharedOnDemandContainer Container,
 	return FOnDemandChunkInfo();
 }
 
+void FOnDemandChunkInfo::GetUrl(FAnsiStringBuilderBase& Url) const
+{
+	TAnsiStringBuilder<41> HashString;
+	HashString << Hash();
+
+	Url.Reset();
+	Url << "/" << ChunksDirectory()
+		<< "/" << HashString.ToView().Left(2)
+		<< "/" << HashString << ANSITEXTVIEW(".iochunk");
+}
+
 ///////////////////////////////////////////////////////////////////////////////
 class FOnDemandIoStore
 	: public TSharedFromThis<FOnDemandIoStore, ESPMode::ThreadSafe>
@@ -329,9 +342,7 @@ private:
 											const TArray<FString>& ExistingContainers,
 											TArray<FSharedOnDemandContainer>& Out);
 
-#if UE_ENABLE_REQUEST_THREAD
 	TUniquePtr<FOnDemandThread>				OnDemandHttpThread;
-#endif // UE_ENABLE_REQUEST_THREAD
 
 	FSharedInstallCache						InstallCache;
 	FSharedPackageStoreBackend				PackageStoreBackend;
diff --git a/Engine/Source/Runtime/Experimental/IoStore/OnDemand/Private/ThreadSafeIntrusiveQueue.h b/Engine/Source/Runtime/Experimental/IoStore/OnDemand/Private/ThreadSafeIntrusiveQueue.h
new file mode 100644
--- /dev/null
+++ b/Engine/Source/Runtime/Experimental/IoStore/OnDemand/Private/ThreadSafeIntrusiveQueue.h
@@ -0,0 +1,145 @@
+// Copyright Epic Games, Inc. All Rights Reserved.
+
+#pragma once
+
+#include "HAL/CriticalSection.h"
+#include "Misc/ScopeLock.h"
+
+namespace UE::IoStore
+{
+
+template<typename T>
+class TThreadSafeIntrusiveQueue
+{
+public:
+	void Enqueue(T* Request)
+	{
+		check(Request->NextRequest == nullptr);
+		FScopeLock _(&CriticalSection);
+
+		if (Tail)
+		{
+			Tail->NextRequest = Request;
+		}
+		else
+		{
+			check(Head == nullptr);
+			Head = Request;
+		}
+
+		Tail = Request;
+	}
+
+	void EnqueueByPriority(T* Request)
+	{
+		FScopeLock _(&CriticalSection);
+		EnqueueByPriorityInternal(Request);
+	}
+
+	T* Dequeue()
+	{
+		FScopeLock _(&CriticalSection);
+
+		T* Requests = Head;
+		Head = Tail = nullptr;
+
+		return Requests;
+	}
+
+	void Reprioritize(T* Request)
+	{
+		// Switch to double linked list/array if this gets too expensive
+		FScopeLock _(&CriticalSection);
+		if (RemoveInternal(Request))
+		{
+			EnqueueByPriorityInternal(Request);
+		}
+	}
+
+private:
+	void EnqueueByPriorityInternal(T* Request)
+	{
+		check(Request->NextRequest == nullptr);
+
+		if (Head == nullptr || Request->Priority > Head->Priority)
+		{
+			if (Head == nullptr)
+			{
+				check(Tail == nullptr);
+				Tail = Request;
+			}
+
+			Request->NextRequest = Head;
+			Head = Request;
+		}
+		else if (Request->Priority <= Tail->Priority)
+		{
+			check(Tail != nullptr);
+			Tail->NextRequest = Request;
+			Tail = Request;
+		}
+		else
+		{
+			// NOTE: This can get expensive if the queue gets too long, might be better to have x number of bucket(s)
+			TRACE_CPUPROFILER_EVENT_SCOPE(IasBackend::EnqueueByPriority);
+			T* It = Head;
+			while (It->NextRequest != nullptr && Request->Priority <= It->NextRequest->Priority)
+			{
+				It = It->NextRequest;
+			}
+
+			Request->NextRequest = It->NextRequest;
+			It->NextRequest = Request;
+		}
+	}
+
+	bool RemoveInternal(T* Request)
+	{
+		check(Request != nullptr);
+		if (Head == nullptr)
+		{
+			check(Tail == nullptr);
+			return false;
+		}
+
+		if (Head == Request)
+		{
+			Head = Request->NextRequest;
+			if (Tail == Request)
+			{
+				check(Head == nullptr);
+				Tail = nullptr;
+			}
+
+			Request->NextRequest = nullptr;
+			return true;
+		}
+		else
+		{
+			T* It = Head;
+			while (It->NextRequest && It->NextRequest != Request)
+			{
+				It = It->NextRequest;
+			}
+
+			if (It->NextRequest == Request)
+			{
+				It->NextRequest = It->NextRequest->NextRequest;
+				Request->NextRequest = nullptr;
+				if (Tail == Request)
+				{
+					Tail = It;
+				}
+				return true;
+			}
+		}
+
+		return false;
+	}
+
+	FCriticalSection CriticalSection;
+	T* Head = nullptr;
+	T* Tail = nullptr;
+};
+
+} //namespace UE::IoStore
