Iris serializer for FRemoteObjectReference.

diff --git a/Engine/Plugins/Runtime/ReplicationSystemTestPlugin/Source/Private/Tests/ReplicationState/TestReplicationStateDescriptorBuilder.cpp b/Engine/Plugins/Runtime/ReplicationSystemTestPlugin/Source/Private/Tests/ReplicationState/TestReplicationStateDescriptorBuilder.cpp
--- a/Engine/Plugins/Runtime/ReplicationSystemTestPlugin/Source/Private/Tests/ReplicationState/TestReplicationStateDescriptorBuilder.cpp
+++ b/Engine/Plugins/Runtime/ReplicationSystemTestPlugin/Source/Private/Tests/ReplicationState/TestReplicationStateDescriptorBuilder.cpp
@@ -9,6 +9,7 @@
 #include "Iris/ReplicationState/PropertyNetSerializerInfoRegistry.h"
 #include "Iris/Serialization/InternalNetSerializers.h"
 #include "Iris/Serialization/NetSerializers.h"
+#include "Iris/Serialization/RemoteObjectReferenceNetSerializer.h"
 // UnrealNetwork.h needed for DOREPLIFETIME macros.
 #include "Net/UnrealNetwork.h"
 
@@ -1187,6 +1188,18 @@ UE_NET_TEST_FIXTURE(FTestReplicationStateDescriptorBuilderFixture, UTestReplicat
 	UE_NET_ASSERT_EQ(SerializerDescriptors[0].Serializer, &UE_NET_GET_SERIALIZER(FFieldPathNetSerializer));
 }
 
+UE_NET_TEST_FIXTURE(FTestReplicationStateDescriptorBuilderFixture, UTestReplicationStateDescriptor_TestClassWithRemoteObjectReferenceProperty)
+{
+	InitDescriptorsFromClass(UTestReplicationStateDescriptor_TestClassWithRemoteObjectReferenceProperty::StaticClass());
+	UE_NET_ASSERT_EQ(Descriptors.Num(), 1);
+
+	UE_NET_ASSERT_EQ(Descriptor->MemberCount, uint16(1));
+
+	const FReplicationStateMemberSerializerDescriptor* SerializerDescriptors = Descriptor->MemberSerializerDescriptors;
+
+	UE_NET_ASSERT_EQ(SerializerDescriptors[0].Serializer, &UE_NET_GET_SERIALIZER(FRemoteObjectReferenceNetSerializer));
+}
+
 UE_NET_TEST_FIXTURE(FTestReplicationStateDescriptorBuilderFixture, TestClassWithOnlyPushBasedPropertiesHasPushBasedDirtinessTrait)
 {
 	InitDescriptorsFromClass(UTestReplicationStateDescriptor_TestClassWithOnlyPushBasedProperties::StaticClass());
@@ -1476,6 +1489,11 @@ void UTestReplicationStateDescriptor_TestClassWithFieldPathProperty::GetLifetime
 	DOREPLIFETIME(ThisClass, FieldPathProperty);
 }
 
+void UTestReplicationStateDescriptor_TestClassWithRemoteObjectReferenceProperty::GetLifetimeReplicatedProps(TArray<class FLifetimeProperty>& OutLifetimeProps) const
+{
+ 	Super::GetLifetimeReplicatedProps(OutLifetimeProps);
+}
+
 void UTestReplicationStateDescriptor_TestClassWithOnlyPushBasedProperties::GetLifetimeReplicatedProps(TArray<FLifetimeProperty>& OutLifetimeProps) const
 {
 	FDoRepLifetimeParams LifetimeParams;
diff --git a/Engine/Plugins/Runtime/ReplicationSystemTestPlugin/Source/Private/Tests/ReplicationState/TestReplicationStateDescriptorBuilder.h b/Engine/Plugins/Runtime/ReplicationSystemTestPlugin/Source/Private/Tests/ReplicationState/TestReplicationStateDescriptorBuilder.h
--- a/Engine/Plugins/Runtime/ReplicationSystemTestPlugin/Source/Private/Tests/ReplicationState/TestReplicationStateDescriptorBuilder.h
+++ b/Engine/Plugins/Runtime/ReplicationSystemTestPlugin/Source/Private/Tests/ReplicationState/TestReplicationStateDescriptorBuilder.h
@@ -4,6 +4,7 @@
 
 #include "NetworkAutomationTest.h"
 #include "UObject/ObjectMacros.h"
+#include "UObject/RemoteObjectTransfer.h"
 #include "Engine/EngineTypes.h"
 #include "Iris/Serialization/NetSerializer.h"
 #include "Tests/Serialization/EnumTestTypes.h"
@@ -777,6 +778,20 @@ public:
 	TFieldPath<FProperty> FieldPathProperty;
 };
 
+UCLASS()
+class UTestReplicationStateDescriptor_TestClassWithRemoteObjectReferenceProperty : public UObject
+{
+	GENERATED_BODY()
+
+public:
+	UTestReplicationStateDescriptor_TestClassWithRemoteObjectReferenceProperty() : UObject() {}
+
+	void GetLifetimeReplicatedProps(TArray<class FLifetimeProperty>& OutLifetimeProps) const;
+
+	UPROPERTY(Replicated)
+	FRemoteObjectReference RemoteObjectReferenceProperty;
+};
+
 UCLASS()
 class UTestReplicationStateDescriptor_TestClassWithOnlyPushBasedProperties : public UObject
 {
diff --git a/Engine/Plugins/Runtime/ReplicationSystemTestPlugin/Source/Private/Tests/Serialization/TestRemoteObjectReferenceNetSerializer.cpp b/Engine/Plugins/Runtime/ReplicationSystemTestPlugin/Source/Private/Tests/Serialization/TestRemoteObjectReferenceNetSerializer.cpp
new file mode 100644
--- /dev/null
+++ b/Engine/Plugins/Runtime/ReplicationSystemTestPlugin/Source/Private/Tests/Serialization/TestRemoteObjectReferenceNetSerializer.cpp
@@ -0,0 +1,141 @@
+// Copyright Epic Games, Inc. All Rights Reserved.
+
+#include "TestRemoteObjectReferenceNetSerializer.h"
+#include "Iris/Core/NetObjectReference.h"
+#include "Iris/Serialization/InternalNetSerializationContext.h"
+#include "Iris/Serialization/ObjectNetSerializer.h"
+#include "Tests/ReplicationSystem/RPC/RPCTestFixture.h"
+#include "Tests/ReplicationSystem/ReplicationSystemServerClientTestFixture.h"
+#include "Tests/ReplicationSystem/ReplicatedTestObject.h"
+#include "UObject/Package.h"
+
+void UTestReplicatedObjectWithRemoteReference::GetLifetimeReplicatedProps(TArray<class FLifetimeProperty>& OutLifetimeProps) const
+{
+	Super::GetLifetimeReplicatedProps(OutLifetimeProps);
+}
+
+void UTestReplicatedObjectWithRemoteReference::RemoteRPCWithRemoteReferenceParam_Implementation(FRemoteObjectReference RemoteReference)
+{
+	RemoteRPCWithRemoteReferenceParamCallCount++;
+	LastReceivedRemoteReference = RemoteReference;
+}
+
+namespace UE::Net::Private
+{
+
+class FScopedRemoteDelegateOverride
+{
+public:
+	FScopedRemoteDelegateOverride()
+		: OldStoreRemoteObjectDataDelegate(UE::RemoteObject::Transfer::StoreRemoteObjectDataDelegate)
+		, OldRestoreRemoteObjectDataDelegate(UE::RemoteObject::Transfer::RestoreRemoteObjectDataDelegate)
+	{
+		// Just use empty delegates, we don't need them for the test
+		// Avoids crashing in FPlatformFile while it's only a stub for the test
+		UE::RemoteObject::Transfer::StoreRemoteObjectDataDelegate.BindLambda([](const UE::RemoteObject::Transfer::FMigrateSendParams&) {});
+		UE::RemoteObject::Transfer::RestoreRemoteObjectDataDelegate.BindLambda([](const FUObjectMigrationContext&) {});
+	}
+
+	~FScopedRemoteDelegateOverride()
+	{
+		UE::RemoteObject::Transfer::StoreRemoteObjectDataDelegate = OldStoreRemoteObjectDataDelegate;
+		UE::RemoteObject::Transfer::RestoreRemoteObjectDataDelegate = OldRestoreRemoteObjectDataDelegate;
+	}
+
+private:
+	TDelegate<void(const UE::RemoteObject::Transfer::FMigrateSendParams&)> OldStoreRemoteObjectDataDelegate;
+	TDelegate<void(const FUObjectMigrationContext&)> OldRestoreRemoteObjectDataDelegate;
+};
+
+UE_NET_TEST_FIXTURE(FRPCTestFixture, TestRemoteObjectReference)
+{
+	// Add a client
+	FReplicationSystemTestClient* Client = CreateClient();
+
+	// Spawn objects on server
+	UTestReplicatedObjectWithRemoteReference* TestObject = Server->CreateObject<UTestReplicatedObjectWithRemoteReference>();
+	UReplicatedTestObject* ReferencedObject = Server->CreateObject<UReplicatedTestObject>();
+
+	TestObject->bIsServerObject = true;
+	TestObject->ReplicationSystem = Server->GetReplicationSystem();
+	Server->ReplicationSystem->SetOwningNetConnection(TestObject->NetRefHandle, Client->ConnectionIdOnServer);
+
+	TestObject->RemoteReferenceProperty = FRemoteObjectReference(ReferencedObject);
+	TestObject->RemoteRPCWithRemoteReferenceParam(FRemoteObjectReference(ReferencedObject));
+	
+	Server->PreSendUpdate();
+	Server->SendAndDeliverTo(Client, true);
+	Server->PostSendUpdate();
+
+	UTestReplicatedObjectWithRemoteReference* ClientTestObject = Client->GetObjectAs<UTestReplicatedObjectWithRemoteReference>(TestObject->NetRefHandle);
+
+	UE_NET_ASSERT_NE(ClientTestObject, nullptr);
+	UE_NET_ASSERT_EQ(ClientTestObject->RemoteRPCWithRemoteReferenceParamCallCount, 1);
+
+	// Verify that we received valid references
+	UE_NET_ASSERT_TRUE(ClientTestObject->LastReceivedRemoteReference.GetRemoteId().IsValid());
+	UE_NET_ASSERT_TRUE(ClientTestObject->RemoteReferenceProperty.GetRemoteId().IsValid());
+
+	// Verify that the replicated references resolve to the correct object
+	UReplicatedTestObject* ResolvedRPCParameter = Cast<UReplicatedTestObject>(ClientTestObject->LastReceivedRemoteReference.Resolve());
+	UReplicatedTestObject* ResolvedProperty = Cast<UReplicatedTestObject>(ClientTestObject->RemoteReferenceProperty.Resolve());
+
+	// These are expected to resolve to to server object in this test case, since both the client and server
+	// exist in the same engine instance.
+	UE_NET_ASSERT_EQ(ResolvedRPCParameter, ReferencedObject);
+	UE_NET_ASSERT_EQ(ResolvedProperty, ReferencedObject);
+}
+
+UE_NET_TEST_FIXTURE(FRPCTestFixture, TestRemoteObjectReferencePaths)
+{
+	// Add a client
+	FReplicationSystemTestClient* Client = CreateClient();
+
+	// Don't want these to trigger in this test
+	FScopedRemoteDelegateOverride ScopedDelegateOverride;
+
+	// Spawn replicated object on server to do the replication
+	UTestReplicatedObjectWithRemoteReference* TestObject = Server->CreateObject<UTestReplicatedObjectWithRemoteReference>();
+
+	// Spawn a non-replicated object with a deterministic path, it will be the one referenced by the FRemoteObjectReferece
+	UTestNamedObject* ServerReferencedObject = NewObject<UTestNamedObject>(GetTransientPackage(), UTestNamedObject::StaticClass(), "TestNamedObject");
+
+	TestObject->bIsServerObject = true;
+	TestObject->ReplicationSystem = Server->GetReplicationSystem();
+	Server->ReplicationSystem->SetOwningNetConnection(TestObject->NetRefHandle, Client->ConnectionIdOnServer);
+
+	TestObject->RemoteReferenceProperty = FRemoteObjectReference(ServerReferencedObject);
+	TestObject->RemoteRPCWithRemoteReferenceParam(FRemoteObjectReference(ServerReferencedObject));
+
+	Server->PreSendUpdate();
+	Server->SendTo(Client);
+	Server->PostSendUpdate();
+
+	// Simulate having different instances of objects on the client & server, this forces the FRemoteObjectReference
+	// to look up the object by path on the client. We accomplish this by destroying, GCing, and recreating the object.
+	ServerReferencedObject->MarkAsGarbage();
+	CollectGarbage(GARBAGE_COLLECTION_KEEPFLAGS);
+	ServerReferencedObject = nullptr;
+
+	UTestNamedObject* ClientReferencedObject = NewObject<UTestNamedObject>(GetTransientPackage(), UTestNamedObject::StaticClass(), "TestNamedObject");
+
+	// Deliver to client now that the referenced object has been recreated
+	Server->DeliverTo(Client, true);
+
+	UTestReplicatedObjectWithRemoteReference* ClientTestObject = Client->GetObjectAs<UTestReplicatedObjectWithRemoteReference>(TestObject->NetRefHandle);
+
+	UE_NET_ASSERT_NE(ClientTestObject, nullptr);
+	UE_NET_ASSERT_EQ(ClientTestObject->RemoteRPCWithRemoteReferenceParamCallCount, 1);
+
+	// Verify that we received a valid reference
+	UE_NET_ASSERT_TRUE(ClientTestObject->LastReceivedRemoteReference.GetRemoteId().IsValid());
+
+	// Client should be able to resolve the remote reference by path
+	UTestNamedObject* ClientResolvedObjectRPC = Cast<UTestNamedObject>(ClientTestObject->LastReceivedRemoteReference.Resolve());
+	UE_NET_ASSERT_EQ(ClientResolvedObjectRPC, ClientReferencedObject);
+
+	UTestNamedObject* ClientResolvedObjectProperty = Cast<UTestNamedObject>(ClientTestObject->RemoteReferenceProperty.Resolve());
+	UE_NET_ASSERT_EQ(ClientResolvedObjectProperty, ClientReferencedObject);
+}
+
+}
diff --git a/Engine/Plugins/Runtime/ReplicationSystemTestPlugin/Source/Private/Tests/Serialization/TestRemoteObjectReferenceNetSerializer.h b/Engine/Plugins/Runtime/ReplicationSystemTestPlugin/Source/Private/Tests/Serialization/TestRemoteObjectReferenceNetSerializer.h
new file mode 100644
--- /dev/null
+++ b/Engine/Plugins/Runtime/ReplicationSystemTestPlugin/Source/Private/Tests/Serialization/TestRemoteObjectReferenceNetSerializer.h
@@ -0,0 +1,32 @@
+// Copyright Epic Games, Inc. All Rights Reserved.
+
+#pragma once
+#include "Tests/ReplicationSystem/RPC/ReplicatedTestObjectWithRPC.h"
+#include "UObject/RemoteObjectTransfer.h"
+#include "TestRemoteObjectReferenceNetSerializer.generated.h"
+
+/** Just need an empty object the test can spawn with a stable name. Can't use UObject directly since it' abstract. */
+UCLASS()
+class UTestNamedObject : public UObject
+{
+	GENERATED_BODY()
+};
+
+UCLASS()
+class UTestReplicatedObjectWithRemoteReference : public UTestReplicatedObjectWithRPC
+{
+	GENERATED_BODY()
+
+public:
+	void GetLifetimeReplicatedProps(TArray<class FLifetimeProperty>& OutLifetimeProps) const;
+
+	UPROPERTY(Replicated)
+	FRemoteObjectReference RemoteReferenceProperty;
+
+	UFUNCTION(Client, Reliable)
+	void RemoteRPCWithRemoteReferenceParam(FRemoteObjectReference RemoteReference);
+
+	int32 RemoteRPCWithRemoteReferenceParamCallCount = 0;
+
+	FRemoteObjectReference LastReceivedRemoteReference;
+};
diff --git a/Engine/Plugins/Runtime/ReplicationSystemTestPlugin/Source/Private/Tests/TestNetworkAutomationTestFramework.cpp b/Engine/Plugins/Runtime/ReplicationSystemTestPlugin/Source/Private/Tests/TestNetworkAutomationTestFramework.cpp
--- a/Engine/Plugins/Runtime/ReplicationSystemTestPlugin/Source/Private/Tests/TestNetworkAutomationTestFramework.cpp
+++ b/Engine/Plugins/Runtime/ReplicationSystemTestPlugin/Source/Private/Tests/TestNetworkAutomationTestFramework.cpp
@@ -242,4 +242,11 @@ UE_NET_TEST(NetworkAutomationTest, TestPushModelEnabled)
 	UE_NET_ASSERT_TRUE(PushModelCVar->GetBool());
 }
 
+UE_NET_TEST(NetworkAutomationTest, TestRemoteObjectHandleEnabled)
+{
+#if !UE_WITH_REMOTE_OBJECT_HANDLE
+	UE_NET_FAIL("Remote object handle support not compiled in.");
+#endif
+}
+
 }
diff --git a/Engine/Source/Programs/ReplicationSystemLowLevelTests/ReplicationSystemLowLevelTests.Target.cs b/Engine/Source/Programs/ReplicationSystemLowLevelTests/ReplicationSystemLowLevelTests.Target.cs
--- a/Engine/Source/Programs/ReplicationSystemLowLevelTests/ReplicationSystemLowLevelTests.Target.cs
+++ b/Engine/Source/Programs/ReplicationSystemLowLevelTests/ReplicationSystemLowLevelTests.Target.cs
@@ -24,5 +24,6 @@ public class ReplicationSystemLowLevelTestsTarget : TestTargetRules
 		GlobalDefinitions.Add("UE_NET_TEST_FAKE_REP_TAGS=1");
 		// Load time profiling brings object construction to a crawl.
 		GlobalDefinitions.Add("LOADTIMEPROFILERTRACE_ENABLED=0");
+		GlobalDefinitions.Add("UE_WITH_REMOTE_OBJECT_HANDLE=1");
 	}
 }
diff --git a/Engine/Source/Programs/ReplicationSystemTest/ReplicationSystemTest.Target.cs b/Engine/Source/Programs/ReplicationSystemTest/ReplicationSystemTest.Target.cs
--- a/Engine/Source/Programs/ReplicationSystemTest/ReplicationSystemTest.Target.cs
+++ b/Engine/Source/Programs/ReplicationSystemTest/ReplicationSystemTest.Target.cs
@@ -43,6 +43,7 @@ public class ReplicationSystemTestTarget : TargetRules
 		GlobalDefinitions.Add("UE_NET_TEST_FAKE_REP_TAGS=1");
 		// Load time profiling brings object construction to a crawl.
 		GlobalDefinitions.Add("LOADTIMEPROFILERTRACE_ENABLED=0");
+		GlobalDefinitions.Add("UE_WITH_REMOTE_OBJECT_HANDLE=1");
 	}
 }
 
diff --git a/Engine/Source/Runtime/CoreUObject/Private/UObject/RemoteObjectTransfer.cpp b/Engine/Source/Runtime/CoreUObject/Private/UObject/RemoteObjectTransfer.cpp
--- a/Engine/Source/Runtime/CoreUObject/Private/UObject/RemoteObjectTransfer.cpp
+++ b/Engine/Source/Runtime/CoreUObject/Private/UObject/RemoteObjectTransfer.cpp
@@ -1323,6 +1323,30 @@ bool FRemoteObjectReference::NetSerialize(FArchive& Ar, class UPackageMap* Map,
 	return bOutSuccess;
 }
 
+void FRemoteObjectReference::NetDequantize(FRemoteObjectId InObjectId, FRemoteServerId InServerId, const FRemoteObjectPathName& InPath)
+{
+	ObjectId = InObjectId;
+	ServerId = InServerId;
+
+	UE::RemoteObject::Private::RegisterRemoteObjectId(ObjectId, ServerId);
+
+	if (InPath.Num())
+	{
+		// We currently can't differentiate between assets that exist on both servers which we don't want to migrate
+		// so try to find the object in memory first by its id and if it doesn't exist try to find the object by pathname
+		// and if such object does exist use its remote id to resolve this reference to avoid migrating the asset from another server
+		UObject* ExistingObject = StaticFindObjectFastInternal(ObjectId);
+		if (!ExistingObject)
+		{
+			ExistingObject = InPath.Resolve();
+			if (ExistingObject)
+			{
+				ObjectId = FRemoteObjectId(ExistingObject);
+			}
+		}
+	}
+}
+
 FArchive& operator<<(FArchive& Ar, FRemoteObjectReference& Ref)
 {
 	Ar << Ref.ObjectId;
diff --git a/Engine/Source/Runtime/CoreUObject/Public/UObject/NoExportTypes.h b/Engine/Source/Runtime/CoreUObject/Public/UObject/NoExportTypes.h
--- a/Engine/Source/Runtime/CoreUObject/Public/UObject/NoExportTypes.h
+++ b/Engine/Source/Runtime/CoreUObject/Public/UObject/NoExportTypes.h
@@ -3434,6 +3434,12 @@ struct FRemoteObjectTables
 	TArray<FRemoteObjectId> RemoteIds;
 };
 
+USTRUCT(noexport)
+struct FRemoteObjectPathName : public FRemoteObjectTables
+{
+	GENERATED_BODY()
+};
+
 USTRUCT(noexport, BlueprintType)
 struct FRemoteObjectData
 {
diff --git a/Engine/Source/Runtime/CoreUObject/Public/UObject/RemoteObjectPathName.h b/Engine/Source/Runtime/CoreUObject/Public/UObject/RemoteObjectPathName.h
--- a/Engine/Source/Runtime/CoreUObject/Public/UObject/RemoteObjectPathName.h
+++ b/Engine/Source/Runtime/CoreUObject/Public/UObject/RemoteObjectPathName.h
@@ -73,9 +73,12 @@ struct FRemoteObjectPathName : public FRemoteObjectTables
 {
 	FRemoteObjectPathName() = default;
 	FRemoteObjectPathName(FRemoteObjectPathName&&) = default;
-	explicit FRemoteObjectPathName(UObject* InObject);
+	COREUOBJECT_API explicit FRemoteObjectPathName(UObject* InObject);
 	FRemoteObjectPathName& operator=(FRemoteObjectPathName&&) = default;
 
+	FRemoteObjectPathName(const FRemoteObjectPathName&) = default;
+	FRemoteObjectPathName& operator=(const FRemoteObjectPathName&) = default;
+
 	FORCEINLINE FName GetSegmentName(int32 InSegmentIndex, const FRemoteObjectTables&) const
 	{
 		return Names[InSegmentIndex];
diff --git a/Engine/Source/Runtime/CoreUObject/Public/UObject/RemoteObjectTransfer.h b/Engine/Source/Runtime/CoreUObject/Public/UObject/RemoteObjectTransfer.h
--- a/Engine/Source/Runtime/CoreUObject/Public/UObject/RemoteObjectTransfer.h
+++ b/Engine/Source/Runtime/CoreUObject/Public/UObject/RemoteObjectTransfer.h
@@ -8,6 +8,11 @@
 #include "Async/Async.h"
 #include "UObject/RemoteExecutor.h"
 
+namespace UE::Net
+{
+	struct FRemoteObjectReferenceNetSerializer;
+}
+
 struct FUObjectMigrationContext;
 
 /**
@@ -181,6 +186,9 @@ private:
 public:
 	FRemoteObjectReference() = default;
 
+	FRemoteObjectReference(const FRemoteObjectReference&) = default;
+	FRemoteObjectReference& operator=(const FRemoteObjectReference&) = default;
+
 	COREUOBJECT_API explicit FRemoteObjectReference(const FObjectPtr& Ptr);
 
 	template <typename T>
@@ -228,6 +236,10 @@ public:
 	COREUOBJECT_API bool NetSerialize(FArchive& Ar, class UPackageMap* Map, bool& bOutSuccess);
 
 	friend COREUOBJECT_API FArchive& operator<<(FArchive& Ar, FRemoteObjectReference& Ref);
+
+private:
+	friend struct UE::Net::FRemoteObjectReferenceNetSerializer;
+	COREUOBJECT_API void NetDequantize(FRemoteObjectId InObjectId, FRemoteServerId InServerId, const FRemoteObjectPathName& InPath);
 };
 
 template<>
diff --git a/Engine/Source/Runtime/Experimental/Iris/Core/Private/Iris/ReplicationState/DefaultPropertyNetSerializerInfos.cpp b/Engine/Source/Runtime/Experimental/Iris/Core/Private/Iris/ReplicationState/DefaultPropertyNetSerializerInfos.cpp
--- a/Engine/Source/Runtime/Experimental/Iris/Core/Private/Iris/ReplicationState/DefaultPropertyNetSerializerInfos.cpp
+++ b/Engine/Source/Runtime/Experimental/Iris/Core/Private/Iris/ReplicationState/DefaultPropertyNetSerializerInfos.cpp
@@ -6,6 +6,7 @@
 #include "Iris/Serialization/NetSerializers.h"
 #include "Iris/Serialization/InternalNetSerializers.h"
 #include "Iris/Serialization/DateTimeNetSerializer.h"
+#include "Iris/Serialization/RemoteObjectReferenceNetSerializer.h"
 #include "UObject/TextProperty.h"
 
 namespace UE::Net::Private
@@ -31,6 +32,7 @@ UE_NET_IMPLEMENT_SIMPLE_NETSERIALIZER_INFO(FDoubleProperty, FDoubleNetSerializer
 // Objects and fields
 static const FName PropertyNetSerializerRegistry_NAME_SoftObjectPath("SoftObjectPath");
 static const FName PropertyNetSerializerRegistry_NAME_SoftClassPath("SoftClassPath");
+static const FName PropertyNetSerializerRegistry_NAME_RemoteObjectReference("RemoteObjectReference");
 
 UE_NET_IMPLEMENT_SIMPLE_NETSERIALIZER_INFO(FObjectProperty, FObjectNetSerializer);
 UE_NET_IMPLEMENT_SIMPLE_NETSERIALIZER_INFO(FWeakObjectProperty, FWeakObjectNetSerializer);
@@ -39,6 +41,7 @@ UE_NET_IMPLEMENT_SIMPLE_NETSERIALIZER_INFO(FSoftClassProperty, FSoftObjectNetSer
 UE_NET_IMPLEMENT_SIMPLE_NETSERIALIZER_INFO(FFieldPathProperty, FFieldPathNetSerializer);
 UE_NET_IMPLEMENT_NAMED_STRUCT_NETSERIALIZER_INFO(PropertyNetSerializerRegistry_NAME_SoftObjectPath, FSoftObjectPathNetSerializer);
 UE_NET_IMPLEMENT_NAMED_STRUCT_NETSERIALIZER_INFO(PropertyNetSerializerRegistry_NAME_SoftClassPath, FSoftClassPathNetSerializer);
+UE_NET_IMPLEMENT_NAMED_STRUCT_NETSERIALIZER_INFO(PropertyNetSerializerRegistry_NAME_RemoteObjectReference, FRemoteObjectReferenceNetSerializer);
 
 // Strings
 UE_NET_IMPLEMENT_SIMPLE_NETSERIALIZER_INFO(FStrProperty, FStringNetSerializer);
@@ -193,7 +196,8 @@ void RegisterDefaultPropertyNetSerializerInfos()
 	UE_NET_REGISTER_NETSERIALIZER_INFO(FFieldPathProperty);
 	UE_NET_REGISTER_NETSERIALIZER_INFO(PropertyNetSerializerRegistry_NAME_SoftObjectPath);
 	UE_NET_REGISTER_NETSERIALIZER_INFO(PropertyNetSerializerRegistry_NAME_SoftClassPath);
-	
+	UE_NET_REGISTER_NETSERIALIZER_INFO(PropertyNetSerializerRegistry_NAME_RemoteObjectReference);
+
 	// String types
 	UE_NET_REGISTER_NETSERIALIZER_INFO(FNameProperty);
 	UE_NET_REGISTER_NETSERIALIZER_INFO(FStrProperty);
diff --git a/Engine/Source/Runtime/Experimental/Iris/Core/Private/Iris/Serialization/RemoteObjectReferenceNetSerializer.cpp b/Engine/Source/Runtime/Experimental/Iris/Core/Private/Iris/Serialization/RemoteObjectReferenceNetSerializer.cpp
new file mode 100644
--- /dev/null
+++ b/Engine/Source/Runtime/Experimental/Iris/Core/Private/Iris/Serialization/RemoteObjectReferenceNetSerializer.cpp
@@ -0,0 +1,222 @@
+// Copyright Epic Games, Inc. All Rights Reserved.
+
+#include "RemoteObjectReferenceNetSerializer.h"
+#include "Iris/ReplicationState/ReplicationStateDescriptorBuilder.h"
+#include "Iris/Serialization/NetSerializerDelegates.h"
+#include "Iris/Serialization/NetSerializers.h"
+#include "UObject/RemoteObjectTransfer.h"
+
+namespace UE::Net
+{
+
+struct FRemoteObjectReferenceNetSerializer
+{
+	// Version
+	static const uint32 Version = 0;
+
+	// Traits
+	static constexpr bool bIsForwardingSerializer = true; // Triggers asserts if a function is missing
+	static constexpr bool bHasDynamicState = true;
+	static constexpr bool bHasCustomNetReference = true;
+
+	// Types
+	struct FQuantizedType
+	{
+		alignas(8) uint8 QuantizedStruct[48];
+	};
+
+	typedef FRemoteObjectReference SourceType;
+	typedef FQuantizedType QuantizedType;
+	typedef FRemoteObjectReferenceNetSerializerConfig ConfigType;
+
+	static const ConfigType DefaultConfig;
+
+	static void Serialize(FNetSerializationContext&, const FNetSerializeArgs& Args);
+	static void Deserialize(FNetSerializationContext&, const FNetDeserializeArgs& Args);
+
+	static void SerializeDelta(FNetSerializationContext&, const FNetSerializeDeltaArgs& Args);
+	static void DeserializeDelta(FNetSerializationContext&, const FNetDeserializeDeltaArgs& Args);
+
+	static void Quantize(FNetSerializationContext&, const FNetQuantizeArgs& Args);
+	static void Dequantize(FNetSerializationContext&, const FNetDequantizeArgs& Args);
+
+	static bool IsEqual(FNetSerializationContext&, const FNetIsEqualArgs& Args);
+	static bool Validate(FNetSerializationContext&, const FNetValidateArgs& Args);
+
+	static void CloneDynamicState(FNetSerializationContext&, const FNetCloneDynamicStateArgs&);
+	static void FreeDynamicState(FNetSerializationContext&, const FNetFreeDynamicStateArgs&);
+
+	static void CollectNetReferences(FNetSerializationContext&, const FNetCollectReferencesArgs&);
+
+private:
+	class FNetSerializerRegistryDelegates final : private UE::Net::FNetSerializerRegistryDelegates
+	{
+	private:
+		virtual void OnPostFreezeNetSerializerRegistry() override;
+	};
+
+	static void RemoteObjectReferenceToRemoteObjectReferenceNetSerializationHelper(const SourceType& RemoteReference, FRemoteObjectReferenceNetSerializationHelper& OutStruct);
+
+	static FRemoteObjectReferenceNetSerializer::FNetSerializerRegistryDelegates NetSerializerRegistryDelegates;
+	static FStructNetSerializerConfig StructNetSerializerConfig;
+	static const FNetSerializer* StructNetSerializer;
+};
+UE_NET_IMPLEMENT_SERIALIZER_INTERNAL(FRemoteObjectReferenceNetSerializer);
+
+const FRemoteObjectReferenceNetSerializer::ConfigType FRemoteObjectReferenceNetSerializer::DefaultConfig;
+
+FRemoteObjectReferenceNetSerializer::FNetSerializerRegistryDelegates FRemoteObjectReferenceNetSerializer::NetSerializerRegistryDelegates;
+FStructNetSerializerConfig FRemoteObjectReferenceNetSerializer::StructNetSerializerConfig;
+const FNetSerializer* FRemoteObjectReferenceNetSerializer::StructNetSerializer = &UE_NET_GET_SERIALIZER(FStructNetSerializer);
+
+void FRemoteObjectReferenceNetSerializer::Serialize(FNetSerializationContext& Context, const FNetSerializeArgs& Args)
+{
+	FNetSerializeArgs InternalArgs = Args;
+	InternalArgs.NetSerializerConfig = &StructNetSerializerConfig;
+	return StructNetSerializer->Serialize(Context, InternalArgs);
+}
+
+void FRemoteObjectReferenceNetSerializer::Deserialize(FNetSerializationContext& Context, const FNetDeserializeArgs& Args)
+{
+	FNetDeserializeArgs InternalArgs = Args;
+	InternalArgs.NetSerializerConfig = &StructNetSerializerConfig;
+	return StructNetSerializer->Deserialize(Context, InternalArgs);
+}
+
+void FRemoteObjectReferenceNetSerializer::SerializeDelta(FNetSerializationContext& Context, const FNetSerializeDeltaArgs& Args)
+{
+	FNetSerializeDeltaArgs InternalArgs = Args;
+	InternalArgs.NetSerializerConfig = &StructNetSerializerConfig;
+	return StructNetSerializer->SerializeDelta(Context, InternalArgs);
+}
+
+void FRemoteObjectReferenceNetSerializer::DeserializeDelta(FNetSerializationContext& Context, const FNetDeserializeDeltaArgs& Args)
+{
+	FNetDeserializeDeltaArgs InternalArgs = Args;
+	InternalArgs.NetSerializerConfig = &StructNetSerializerConfig;
+	return StructNetSerializer->DeserializeDelta(Context, InternalArgs);
+}
+
+void FRemoteObjectReferenceNetSerializer::Quantize(FNetSerializationContext& Context, const FNetQuantizeArgs& Args)
+{
+	const SourceType& SourceValue = *reinterpret_cast<const SourceType*>(Args.Source);
+
+	FRemoteObjectReferenceNetSerializationHelper IntermediateValue;
+	RemoteObjectReferenceToRemoteObjectReferenceNetSerializationHelper(SourceValue, IntermediateValue);
+
+	if (IntermediateValue.ObjectId.IsValid())
+	{
+		if (UObject* ExistingObject = StaticFindObjectFastInternal(IntermediateValue.ObjectId))
+		{
+			UE::RemoteObject::Transfer::RegisterSharedObject(ExistingObject);
+			IntermediateValue.Path = FRemoteObjectPathName(ExistingObject);
+		}
+	}
+
+	FNetQuantizeArgs InternalArgs = Args;
+	InternalArgs.Source = NetSerializerValuePointer(&IntermediateValue);
+	InternalArgs.NetSerializerConfig = &StructNetSerializerConfig;
+	return StructNetSerializer->Quantize(Context, InternalArgs);
+}
+
+void FRemoteObjectReferenceNetSerializer::Dequantize(FNetSerializationContext& Context, const FNetDequantizeArgs& Args)
+{
+	SourceType& TargetValue = *reinterpret_cast<SourceType*>(Args.Target);
+
+	FRemoteObjectReferenceNetSerializationHelper IntermediateValue;
+
+	FNetDequantizeArgs InternalArgs = Args;
+	InternalArgs.Target = NetSerializerValuePointer(&IntermediateValue);
+	InternalArgs.NetSerializerConfig = &StructNetSerializerConfig;
+
+	StructNetSerializer->Dequantize(Context, InternalArgs);
+
+	TargetValue.NetDequantize(IntermediateValue.ObjectId, IntermediateValue.ServerId, IntermediateValue.Path);
+}
+
+bool FRemoteObjectReferenceNetSerializer::IsEqual(FNetSerializationContext& Context, const FNetIsEqualArgs& Args)
+{
+	if (Args.bStateIsQuantized)
+	{
+		FNetIsEqualArgs InternalArgs = Args;
+		InternalArgs.NetSerializerConfig = &StructNetSerializerConfig;
+		return StructNetSerializer->IsEqual(Context, InternalArgs);
+	}
+	else
+	{
+		const SourceType& SourceValue0 = *reinterpret_cast<const SourceType*>(Args.Source0);
+		const SourceType& SourceValue1 = *reinterpret_cast<const SourceType*>(Args.Source1);
+
+		return SourceValue0 == SourceValue1;
+	}
+}
+
+bool FRemoteObjectReferenceNetSerializer::Validate(FNetSerializationContext& Context, const FNetValidateArgs& Args)
+{
+	const SourceType& SourceValue = *reinterpret_cast<const SourceType*>(Args.Source);
+
+	// The remote reference's ObjectId and ServerId should either both be valid or both be invalid
+	if (SourceValue.GetRemoteId().IsValid() != SourceValue.GetSharingServerId().IsValid())
+	{
+		return false;
+	}
+
+	FRemoteObjectReferenceNetSerializationHelper IntermediateValue;
+	RemoteObjectReferenceToRemoteObjectReferenceNetSerializationHelper(SourceValue, IntermediateValue);
+
+	FNetValidateArgs InternalArgs = Args;
+	InternalArgs.Source = NetSerializerValuePointer(&IntermediateValue);
+	InternalArgs.NetSerializerConfig = &StructNetSerializerConfig;
+
+	return StructNetSerializer->Validate(Context, InternalArgs);
+}
+
+void FRemoteObjectReferenceNetSerializer::CollectNetReferences(FNetSerializationContext& Context, const FNetCollectReferencesArgs& Args)
+{
+	using namespace Private;
+
+	const FReplicationStateDescriptor* Descriptor = StructNetSerializerConfig.StateDescriptor;
+	if (Descriptor && EnumHasAnyFlags(Descriptor->Traits, EReplicationStateTraits::HasObjectReference))
+	{
+		FNetCollectReferencesArgs InternalArgs = Args;
+		InternalArgs.NetSerializerConfig = &StructNetSerializerConfig;
+		StructNetSerializer->CollectNetReferences(Context, InternalArgs);
+	}
+}
+
+void FRemoteObjectReferenceNetSerializer::CloneDynamicState(FNetSerializationContext& Context, const FNetCloneDynamicStateArgs& Args)
+{
+	FNetCloneDynamicStateArgs InternalArgs = Args;
+	InternalArgs.NetSerializerConfig = &StructNetSerializerConfig;
+	return StructNetSerializer->CloneDynamicState(Context, InternalArgs);
+}
+
+void FRemoteObjectReferenceNetSerializer::FreeDynamicState(FNetSerializationContext& Context, const FNetFreeDynamicStateArgs& Args)
+{
+	FNetFreeDynamicStateArgs InternalArgs = Args;
+	InternalArgs.NetSerializerConfig = &StructNetSerializerConfig;
+	return StructNetSerializer->FreeDynamicState(Context, InternalArgs);
+}
+
+void FRemoteObjectReferenceNetSerializer::RemoteObjectReferenceToRemoteObjectReferenceNetSerializationHelper(const SourceType& RemoteObjectReference, FRemoteObjectReferenceNetSerializationHelper& OutStruct)
+{
+	OutStruct.ObjectId = RemoteObjectReference.ObjectId;
+	OutStruct.ServerId = RemoteObjectReference.ServerId;
+}
+
+void FRemoteObjectReferenceNetSerializer::FNetSerializerRegistryDelegates::OnPostFreezeNetSerializerRegistry()
+{
+	UStruct* Struct = FRemoteObjectReferenceNetSerializationHelper::StaticStruct();
+	StructNetSerializerConfig.StateDescriptor = FReplicationStateDescriptorBuilder::CreateDescriptorForStruct(Struct);
+	const FReplicationStateDescriptor* Descriptor = StructNetSerializerConfig.StateDescriptor.GetReference();
+	check(Descriptor != nullptr);
+
+	// Validate our assumptions regarding quantized state size and alignment.
+	static_assert(offsetof(FQuantizedType, QuantizedStruct) == 0U, "Expected buffer for struct to be first member of FQuantizedType.");
+	if (sizeof(FQuantizedType::QuantizedStruct) < Descriptor->InternalSize || alignof(FQuantizedType) < Descriptor->InternalAlignment)
+	{
+		LowLevelFatalError(TEXT("FQuantizedType::QuantizedStruct for FRemoteObjectReferenceNetSerializer has size %u and alignment %u but requires size %u and alignment %u."), uint32(sizeof(FQuantizedType::QuantizedStruct)), uint32(alignof(FQuantizedType)), uint32(Descriptor->InternalSize), uint32(Descriptor->InternalAlignment)); 
+	}
+}
+
+}
diff --git a/Engine/Source/Runtime/Experimental/Iris/Core/Private/Iris/Serialization/RemoteObjectReferenceNetSerializer.h b/Engine/Source/Runtime/Experimental/Iris/Core/Private/Iris/Serialization/RemoteObjectReferenceNetSerializer.h
new file mode 100644
--- /dev/null
+++ b/Engine/Source/Runtime/Experimental/Iris/Core/Private/Iris/Serialization/RemoteObjectReferenceNetSerializer.h
@@ -0,0 +1,35 @@
+// Copyright Epic Games, Inc. All Rights Reserved.
+
+#pragma once
+
+#include "Iris/Serialization/InternalNetSerializer.h"
+#include "UObject/RemoteObjectPathName.h"
+#include "RemoteObjectReferenceNetSerializer.generated.h"
+
+USTRUCT()
+struct FRemoteObjectReferenceNetSerializerConfig : public FNetSerializerConfig
+{
+	GENERATED_BODY()
+};
+
+USTRUCT()
+struct FRemoteObjectReferenceNetSerializationHelper
+{
+	GENERATED_BODY()
+
+	UPROPERTY()
+	FRemoteObjectId ObjectId;
+
+	UPROPERTY()
+	FRemoteServerId ServerId;
+
+	UPROPERTY()
+	FRemoteObjectPathName Path;
+};
+
+namespace UE::Net
+{
+
+UE_NET_DECLARE_SERIALIZER_INTERNAL(FRemoteObjectReferenceNetSerializer, IRISCORE_API);
+
+}
