Move global on-demand TOC file extension to FOnDemandToc.

diff --git a/Engine/Source/Runtime/Experimental/IoStore/OnDemand/Private/OnDemandIoDispatcherBackend.cpp b/Engine/Source/Runtime/Experimental/IoStore/OnDemand/Private/OnDemandIoDispatcherBackend.cpp
--- a/Engine/Source/Runtime/Experimental/IoStore/OnDemand/Private/OnDemandIoDispatcherBackend.cpp
+++ b/Engine/Source/Runtime/Experimental/IoStore/OnDemand/Private/OnDemandIoDispatcherBackend.cpp
@@ -62,10 +62,6 @@ namespace UE::IoStore
 
 ///////////////////////////////////////////////////////////////////////////////
 
-extern FString GIasOnDemandTocExt;
-
-///////////////////////////////////////////////////////////////////////////////
-
 /** Note that GIasHttpPrimaryEndpoint has no effect after initial start up */
 int32 GIasHttpPrimaryEndpoint = 0;
 static FAutoConsoleVariableRef CVar_IasHttpPrimaryEndpoint(
diff --git a/Engine/Source/Runtime/Experimental/IoStore/OnDemand/Private/OnDemandIoStore.cpp b/Engine/Source/Runtime/Experimental/IoStore/OnDemand/Private/OnDemandIoStore.cpp
--- a/Engine/Source/Runtime/Experimental/IoStore/OnDemand/Private/OnDemandIoStore.cpp
+++ b/Engine/Source/Runtime/Experimental/IoStore/OnDemand/Private/OnDemandIoStore.cpp
@@ -66,8 +66,6 @@ static FAutoConsoleVariableRef CVar_IoStoreOnDemandMinInstallDurationSeconds(
 namespace UE::IoStore
 {
 
-extern FString GIasOnDemandTocExt;
-
 ///////////////////////////////////////////////////////////////////////////////
 namespace Private
 {
@@ -434,7 +432,7 @@ FIoStatus FOnDemandIoStore::Initialize()
 		[this](EMountOperation Operation, const TCHAR* ContainerPath, int32 Order) -> void
 		{
 			IPlatformFile& Ipf = IPlatformFile::GetPlatformPhysical();
-			const FString OnDemandTocPath = FPathViews::ChangeExtension(ContainerPath, GIasOnDemandTocExt);
+			const FString OnDemandTocPath = FPathViews::ChangeExtension(ContainerPath, *FOnDemandToc::FileExt);
 
 			if (Ipf.FileExists(*OnDemandTocPath) == false)
 			{
@@ -482,11 +480,11 @@ FIoStatus FOnDemandIoStore::Initialize()
 			for (const FMountedPakInfo& Info : PakInfo)
 			{
 				check(Info.PakFile != nullptr);
-				const FString OnDemandTocPath = FPathViews::ChangeExtension(Info.PakFile->PakGetPakFilename(), GIasOnDemandTocExt);
+				const FString OnDemandTocPath = FPathViews::ChangeExtension(Info.PakFile->PakGetPakFilename(), *FOnDemandToc::FileExt);
 				if (Ipf.FileExists(*OnDemandTocPath))
 				{
 					FSharedMountRequest MountRequest = MakeShared<FMountRequest>();
-					MountRequest->MountArgs = FOnDemandMountArgs
+					MountRequest->Args = FOnDemandMountArgs
 					{
 						.MountId	= OnDemandTocPath,
 						.FilePath	= OnDemandTocPath
@@ -496,7 +494,7 @@ FIoStatus FOnDemandIoStore::Initialize()
 						UE_CLOG(!Result.IsOk(), LogIoStoreOnDemand, Error, TEXT("Failed to mount container, reason '%s'"),
 							*Result.Status().ToString());
 					};
-					MountRequests.Add(OnDemandTocPath, MoveTemp(MountRequest));
+					MountRequests.Add(MoveTemp(MountRequest));
 				}
 			}
 
diff --git a/Engine/Source/Runtime/Experimental/IoStore/OnDemand/Private/Tool/Upload.cpp b/Engine/Source/Runtime/Experimental/IoStore/OnDemand/Private/Tool/Upload.cpp
--- a/Engine/Source/Runtime/Experimental/IoStore/OnDemand/Private/Tool/Upload.cpp
+++ b/Engine/Source/Runtime/Experimental/IoStore/OnDemand/Private/Tool/Upload.cpp
@@ -24,12 +24,6 @@
 #include "Serialization/LargeMemoryWriter.h"
 #include "UploadQueue.h"
 
-////////////////////////////////////////////////////////////////////////////////
-namespace UE::IoStore
-{
-extern FString GIasOnDemandTocExt;
-}
-
 namespace UE::IoStore::Tool
 {
 
@@ -217,7 +211,7 @@ static bool SaveArrayToFile(TArrayView64<uint8> Data, const TCHAR* Filename, uin
 ////////////////////////////////////////////////////////////////////////////////
 static FIoStatus WriteContainerFiles(FOnDemandToc& OnDemandToc, const TMap<FIoHash, FString>& UTocPaths)
 {
-	UE_LOG(LogIas, Display, TEXT("Attempting to write out %d '%s' files"), OnDemandToc.Containers.Num(), *GIasOnDemandTocExt);
+	UE_LOG(LogIas, Display, TEXT("Attempting to write out %d '%s' files"), OnDemandToc.Containers.Num(), *FOnDemandToc::FileExt);
 
 	for (FOnDemandTocContainerEntry& Container : OnDemandToc.Containers)
 	{
@@ -235,7 +229,7 @@ static FIoStatus WriteContainerFiles(FOnDemandToc& OnDemandToc, const TMap<FIoHa
 		ContainerToc.Meta = OnDemandToc.Meta;
 		ContainerToc.Containers.Emplace(MoveTemp(Container));
 
-		const FString OutputPath = FPathViews::ChangeExtension(*UTocPath, GIasOnDemandTocExt);
+		const FString OutputPath = FPathViews::ChangeExtension(*UTocPath, FOnDemandToc::FileExt);
 
 		TUniquePtr<FArchive> Ar(IFileManager::Get().CreateFileWriter(*OutputPath));
 
diff --git a/Engine/Source/Runtime/Experimental/IoStore/OnDemandCore/Private/IoStoreOnDemand.cpp b/Engine/Source/Runtime/Experimental/IoStore/OnDemandCore/Private/IoStoreOnDemand.cpp
--- a/Engine/Source/Runtime/Experimental/IoStore/OnDemandCore/Private/IoStoreOnDemand.cpp
+++ b/Engine/Source/Runtime/Experimental/IoStore/OnDemandCore/Private/IoStoreOnDemand.cpp
@@ -331,9 +331,6 @@ static FAutoConsoleCommand InstallPackageCommand(
 } // namespace Commands
 #endif // !UE_BUILD_SHIPPING
 
-////////////////////////////////////////////////////////////////////////////////
-FString GIasOnDemandTocExt = TEXT(".uondemandtoc");
-
 static const TCHAR* NotInitializedError = TEXT("I/O store on-demand not initialized");
 
 ////////////////////////////////////////////////////////////////////////////////
diff --git a/Engine/Source/Runtime/Experimental/IoStore/OnDemandCore/Private/OnDemandToc.cpp b/Engine/Source/Runtime/Experimental/IoStore/OnDemandCore/Private/OnDemandToc.cpp
--- a/Engine/Source/Runtime/Experimental/IoStore/OnDemandCore/Private/OnDemandToc.cpp
+++ b/Engine/Source/Runtime/Experimental/IoStore/OnDemandCore/Private/OnDemandToc.cpp
@@ -572,7 +572,8 @@ FCbWriter& operator<<(FCbWriter& Writer, const FOnDemandToc& Toc)
 	return Writer;
 }
 
-FGuid FOnDemandToc::VersionGuid = FGuid("C43DD98F353F499D9A0767F6EA0155EB");
+const FGuid FOnDemandToc::VersionGuid	= FGuid("C43DD98F353F499D9A0767F6EA0155EB");
+const FString FOnDemandToc::FileExt		= FString(TEXT(".uondemandtoc"));
 
 bool LoadFromCompactBinary(FCbFieldView Field, FOnDemandToc& OutToc)
 {
diff --git a/Engine/Source/Runtime/Experimental/IoStore/OnDemandCore/Public/IO/OnDemandToc.h b/Engine/Source/Runtime/Experimental/IoStore/OnDemandCore/Public/IO/OnDemandToc.h
--- a/Engine/Source/Runtime/Experimental/IoStore/OnDemandCore/Public/IO/OnDemandToc.h
+++ b/Engine/Source/Runtime/Experimental/IoStore/OnDemandCore/Public/IO/OnDemandToc.h
@@ -192,7 +192,8 @@ struct FOnDemandToc
 	UE_API friend FArchive& operator<<(FArchive& Ar, FOnDemandToc& Toc);
 	UE_API friend FCbWriter& operator<<(FCbWriter& Writer, const FOnDemandToc& Toc);
 
-	static FGuid VersionGuid;
+	UE_API static const FGuid	VersionGuid;
+	UE_API static const FString	FileExt;
 
 	UE_API static TIoStatusOr<FOnDemandToc> LoadFromFile(const FString& FilePath, bool bValidate);
 };
