Serialize container header from on-demand TOC rather then fetching via HTTP

diff --git a/Engine/Source/Runtime/Experimental/IoStore/OnDemand/Private/OnDemandIoStore.cpp b/Engine/Source/Runtime/Experimental/IoStore/OnDemand/Private/OnDemandIoStore.cpp
--- a/Engine/Source/Runtime/Experimental/IoStore/OnDemand/Private/OnDemandIoStore.cpp
+++ b/Engine/Source/Runtime/Experimental/IoStore/OnDemand/Private/OnDemandIoStore.cpp
@@ -1347,45 +1347,62 @@ FIoStatus FOnDemandIoStore::TickMountRequest(FMountRequest& MountRequest)
 		// Try fetch and deserialize the container header if the container is used for installing content.
 		if (EnumHasAnyFlags(Container->Flags, EOnDemandContainerFlags::InstallOnDemand))
 		{
-			if (Container->HostGroup.IsEmpty())
-			{
-				return FIoStatusBuilder(EIoErrorCode::InvalidParameter) << TEXT("No valid endpoint(s)");
-			}
+			const FIoChunkId ChunkId			= CreateContainerHeaderChunkId(Container->ContainerId);
+			const FOnDemandChunkInfo ChunkInfo	= FOnDemandChunkInfo::Find(Container, ChunkId);
 
-			const FIoChunkId ChunkId = CreateContainerHeaderChunkId(Container->ContainerId);
-			if (const FOnDemandChunkInfo ChunkInfo = FOnDemandChunkInfo::Find(Container, ChunkId); ChunkInfo.IsValid())
+			if (ChunkInfo.IsValid())
 			{
-				UE_LOG(LogIoStoreOnDemand, Verbose, TEXT("Fetching container header, ContainerName='%s', ChunkId='%s'"),
-					*Container->Name, *LexToString(ChunkId));
+				const FOnDemandTocContainerEntry* ContainerEntry = Args.Toc->Containers.FindByPredicate(
+						[&Container](const FOnDemandTocContainerEntry& E) { return E.ContainerId == Container->ContainerId; });
+				UE_CLOG(ContainerEntry == nullptr, LogIoStoreOnDemand, Error, TEXT("Failed to find TOC container entry for container '%s'"), *Container->Name);
 
-				const uint32 RetryCount = 2;
-				TIoStatusOr<FMultiEndpointHttpClientResponse> Response = FMultiEndpointHttpClient::Get(
-					Private::GetChunkUrl(Args.HostGroup.PrimaryHost(), *Container, ChunkInfo.ChunkEntry(), ChunkUrlBuilder).ToView(),
-					FMultiEndpointHttpClientConfig{.MaxRetryCount = RetryCount}
-					);
-
-				if (Response.IsOk() == false)
+				if (ContainerEntry != nullptr)
 				{
-					return Response.Status();
-				}
+					FIoBuffer HeaderBuffer;
+					if (ContainerEntry->Header.IsEmpty() == false)
+					{
+						HeaderBuffer = FIoBuffer(FIoBuffer::Wrap, ContainerEntry->Header.GetData(), ContainerEntry->Header.Num());
+					}
+					else
+					{
+						if (Container->HostGroup.IsEmpty())
+						{
+							return FIoStatusBuilder(EIoErrorCode::InvalidParameter) << TEXT("No valid endpoint(s)");
+						}
 
-				if (EnumHasAnyFlags(Container->Flags, EOnDemandContainerFlags::PendingEncryptionKey))
-				{
-					RequestedContainerHeaders.Add(Container->UniqueName(), Response.ConsumeValueOrDie().Body);
-					continue;
-				}
+						UE_LOG(LogIoStoreOnDemand, Warning, TEXT("Fetching container header, ContainerName='%s', ChunkId='%s'"),
+								*Container->Name, *LexToString(ChunkId));
 
-				TIoStatusOr<FSharedContainerHeader> Header = Private::DeserializeContainerHeader(
-					ChunkInfo,
-					Response.ConsumeValueOrDie().Body.GetView(),
-					EnumHasAnyFlags(Container->Flags, EOnDemandContainerFlags::WithSoftReferences));
+						const uint32 RetryCount = 2;
+						TIoStatusOr<FMultiEndpointHttpClientResponse> Response = FMultiEndpointHttpClient::Get(
+								Private::GetChunkUrl(Args.HostGroup.PrimaryHost(), *Container, ChunkInfo.ChunkEntry(), ChunkUrlBuilder).ToView(),
+								FMultiEndpointHttpClientConfig{.MaxRetryCount = RetryCount});
 
-				if (Header.IsOk() == false)
-				{
-					return Header.Status();
-				}
+						if (Response.IsOk() == false)
+						{
+							return Response.Status();
+						}
+						HeaderBuffer = Response.ConsumeValueOrDie().Body;
+					}
+
+					if (EnumHasAnyFlags(Container->Flags, EOnDemandContainerFlags::PendingEncryptionKey))
+					{
+						HeaderBuffer.EnsureOwned();
+						RequestedContainerHeaders.Add(Container->UniqueName(), HeaderBuffer);
+						continue;
+					}
 
-				Container->Header = Header.ConsumeValueOrDie();
+					TIoStatusOr<FSharedContainerHeader> Header = Private::DeserializeContainerHeader(
+							ChunkInfo,
+							HeaderBuffer.GetView(),
+							EnumHasAnyFlags(Container->Flags, EOnDemandContainerFlags::WithSoftReferences));
+
+					if (Header.IsOk() == false)
+					{
+						return Header.Status();
+					}
+					Container->Header = Header.ConsumeValueOrDie();
+				}
 			}
 		}
 
@@ -2170,7 +2187,7 @@ void FOnDemandIoStore::ReleaseContent(FOnDemandInternalContentHandle& ContentHan
 	PackageStoreBackend->NeedsUpdate();
 }
 
-void FOnDemandIoStore::GetReferencedContent(TArray<FSharedOnDemandContainer>& OutContainers, TArray<TBitArray<>>& OutChunkEntryIndices)
+void FOnDemandIoStore::GetReferencedContent(TArray<FSharedOnDemandContainer>& OutContainers, TArray<TBitArray<>>& OutChunkEntryIndices, bool bContainersWithPackagesOnly)
 {
 	UE::TUniqueLock Lock(ContainerMutex);
 	for (FSharedOnDemandContainer& Container : Containers)
@@ -2185,8 +2202,10 @@ void FOnDemandIoStore::GetReferencedContent(TArray<FSharedOnDemandContainer>& Ou
 			continue;
 		}
 
-		// Header should be available if decrypted
-		check(Container->Header);
+		if (bContainersWithPackagesOnly && Container->Header.IsValid() == false)
+		{
+			continue;
+		}
 
 		TBitArray<> Indices = Container->GetReferencedChunkEntries();
 		if (Indices.IsEmpty() == false)
diff --git a/Engine/Source/Runtime/Experimental/IoStore/OnDemand/Private/OnDemandIoStore.h b/Engine/Source/Runtime/Experimental/IoStore/OnDemand/Private/OnDemandIoStore.h
--- a/Engine/Source/Runtime/Experimental/IoStore/OnDemand/Private/OnDemandIoStore.h
+++ b/Engine/Source/Runtime/Experimental/IoStore/OnDemand/Private/OnDemandIoStore.h
@@ -294,7 +294,7 @@ public:
 	/** Used to access a number of FIoChunkId that exist in the system when running tests and is not intended for production use */
 	TArray<FIoChunkId>					DebugFindStreamingChunkIds(int32 NumToFind);
 #endif //!UE_BUILD_SHIPPING
-	void								GetReferencedContent(TArray<FSharedOnDemandContainer>& OutContainers, TArray<TBitArray<>>& OutChunkEntryIndices);
+	void								GetReferencedContent(TArray<FSharedOnDemandContainer>& OutContainers, TArray<TBitArray<>>& OutChunkEntryIndices, bool bContainersWithPackagesOnly = false);
 
 	// IOnDemandIoStore interface
 	virtual FIoStatus					Initialize() override;
diff --git a/Engine/Source/Runtime/Experimental/IoStore/OnDemand/Private/OnDemandPackageStoreBackend.cpp b/Engine/Source/Runtime/Experimental/IoStore/OnDemand/Private/OnDemandPackageStoreBackend.cpp
--- a/Engine/Source/Runtime/Experimental/IoStore/OnDemand/Private/OnDemandPackageStoreBackend.cpp
+++ b/Engine/Source/Runtime/Experimental/IoStore/OnDemand/Private/OnDemandPackageStoreBackend.cpp
@@ -203,7 +203,8 @@ void FOnDemandPackageStoreBackend::UpdateLookupTables()
 
 	if (TSharedPtr<FOnDemandIoStore> PinOnDemandIoStore = OnDemandIoStore.Pin())
 	{
-		PinOnDemandIoStore->GetReferencedContent(AllContainers, ReferencedChunkEntryIndices);
+		const bool bContainersWithPackagesOnly = true;
+		PinOnDemandIoStore->GetReferencedContent(AllContainers, ReferencedChunkEntryIndices, bContainersWithPackagesOnly);
 
 		for (const FSharedOnDemandContainer& C : AllContainers)
 		{
