[ChaosVD] Fix deadlock that is likely to happen when you connect to +50 running client/servers at the same time

diff --git a/Engine/Plugins/ChaosVD/Source/ChaosVD/Private/Trace/ChaosVDTraceProvider.cpp b/Engine/Plugins/ChaosVD/Source/ChaosVD/Private/Trace/ChaosVDTraceProvider.cpp
--- a/Engine/Plugins/ChaosVD/Source/ChaosVD/Private/Trace/ChaosVDTraceProvider.cpp
+++ b/Engine/Plugins/ChaosVD/Source/ChaosVD/Private/Trace/ChaosVDTraceProvider.cpp
@@ -134,7 +134,9 @@ int32 FChaosVDTraceProvider::RemapSolverID(int32 SolverID)
 		while (InternalRecording->HasSolverID_AssumesLocked(RemappedSolverID))
 		{
 			constexpr int32 MaxValue = std::numeric_limits<int32>::max() -1;
-			RemappedSolverID = FMath::RandRange(RemappedSolverID, MaxValue);
+			check(RemappedSolverID < MaxValue);
+	
+			RemappedSolverID = InternalRecording->GetAvailableTrackIDForRemapping();
 		}
 
 		InternalRecording->ReserveSolverID_AssumesLocked(RemappedSolverID);
diff --git a/Engine/Source/Runtime/Experimental/ChaosVDData/Public/ChaosVDRecording.h b/Engine/Source/Runtime/Experimental/ChaosVDData/Public/ChaosVDRecording.h
--- a/Engine/Source/Runtime/Experimental/ChaosVDData/Public/ChaosVDRecording.h
+++ b/Engine/Source/Runtime/Experimental/ChaosVDData/Public/ChaosVDRecording.h
@@ -665,6 +665,11 @@ struct CHAOSVDDATA_API FChaosVDRecording
 	void CommitSolverID(int32 SolverID);
 	void CommitSolverID_AssumesLocked(int32 SolverID);
 
+	int32 GetAvailableTrackIDForRemapping()
+	{
+		return AvailableTrackIDForRemapping++;
+	}
+
 protected:
 
 	/** Adds an Implicit Object to the recording and takes ownership of it */
@@ -704,6 +709,8 @@ protected:
 	TSet<int32> ReservedSolverIDs;
 	TSet<int32> SolverIDs;
 
+	std::atomic<int32> AvailableTrackIDForRemapping = 1;
+
 	friend class FChaosVDTraceProvider;
 	friend class FChaosVDTraceImplicitObjectProcessor;
 };
