Chaos - Add FlushToZero to FVec3

diff --git a/Engine/Source/Runtime/Experimental/ChaosCore/Public/Chaos/Vector.h b/Engine/Source/Runtime/Experimental/ChaosCore/Public/Chaos/Vector.h
--- a/Engine/Source/Runtime/Experimental/ChaosCore/Public/Chaos/Vector.h
+++ b/Engine/Source/Runtime/Experimental/ChaosCore/Public/Chaos/Vector.h
@@ -538,6 +538,12 @@ namespace Chaos
 		{
 			return TVector<FRealSingle, 3>(V1.X > V2.X ? V1.X : V2.X, V1.Y > V2.Y ? V1.Y : V2.Y, V1.Z > V2.Z ? V1.Z : V2.Z);
 		}
+		inline void FlushToZero(FRealSingle Epsilon)
+		{
+			if ((X >= -Epsilon) && (X <= Epsilon)) { X = 0.0f; }
+			if ((Y >= -Epsilon) && (Y <= Epsilon)) { Y = 0.0f; }
+			if ((Z >= -Epsilon) && (Z <= Epsilon)) { Z = 0.0f; }
+		}
 		static TVector<FRealSingle, 3> AxisVector(const int32 Axis)
 		{ return Axis == 0 ? TVector<FRealSingle, 3>(1.f, 0.f, 0.f) : (Axis == 1 ? TVector<FRealSingle, 3>(0.f, 1.f, 0.f) : TVector<FRealSingle, 3>(0.f, 0.f, 1.f)); }
 		static Pair<FRealSingle, int32> MaxAndAxis(const TVector<FRealSingle, 3>& V1, const TVector<FRealSingle, 3>& V2)
@@ -738,6 +744,12 @@ namespace Chaos
 		{
 			return TVector<FRealDouble, 3>(V1.X > V2.X ? V1.X : V2.X, V1.Y > V2.Y ? V1.Y : V2.Y, V1.Z > V2.Z ? V1.Z : V2.Z);
 		}
+		inline void FlushToZero(FRealDouble Epsilon)
+		{
+			if ((X >= -Epsilon) && (X <= Epsilon)) { X = 0.0; }
+			if ((Y >= -Epsilon) && (Y <= Epsilon)) { Y = 0.0; }
+			if ((Z >= -Epsilon) && (Z <= Epsilon)) { Z = 0.0; }
+		}
 		static TVector<FRealDouble, 3> AxisVector(const int32 Axis)
 		{
 			return Axis == 0 ? TVector<FRealDouble, 3>(1.f, 0.f, 0.f) : (Axis == 1 ? TVector<FRealDouble, 3>(0.f, 1.f, 0.f) : TVector<FRealDouble, 3>(0.f, 0.f, 1.f));
