Control cloth collision sphere parallelism with a cvar. By default, only enable parallelism if there are at least 256 spheres.

diff --git a/Engine/Source/Runtime/Experimental/Chaos/Private/Chaos/PBDSelfCollisionSphereConstraints.cpp b/Engine/Source/Runtime/Experimental/Chaos/Private/Chaos/PBDSelfCollisionSphereConstraints.cpp
--- a/Engine/Source/Runtime/Experimental/Chaos/Private/Chaos/PBDSelfCollisionSphereConstraints.cpp
+++ b/Engine/Source/Runtime/Experimental/Chaos/Private/Chaos/PBDSelfCollisionSphereConstraints.cpp
@@ -4,11 +4,15 @@
 #if !COMPILE_WITHOUT_UNREAL_SUPPORT
 #include "Chaos/Framework/Parallel.h"
 #include "Chaos/HierarchicalSpatialHash.h"
+#include "HAL/IConsoleManager.h"
 
 namespace Chaos::Softs
 {
 	namespace Private
 	{
+		static int32 Chaos_SelfCollisionSphere_ParallelCount = 256;
+		static FAutoConsoleVariableRef CVarChaosSelfCollisionSphereParallelismEnabled(TEXT("p.Chaos.SelfCollisionSphere.ParallelCount"), Chaos_SelfCollisionSphere_ParallelCount, TEXT("Min number of spheres to enable parallelism for SelfCollisionSpheres constraint"));
+
 		struct FSphereSpatialEntry
 		{
 			const TConstArrayView<FSolverVec3>* Points;
@@ -48,7 +52,7 @@ namespace Chaos::Softs
 	{
 		Constraints.Reset();
 
-		if (NumParticles == 0 || (VertexSetNoOffset && VertexSetNoOffset->IsEmpty()))
+		if (NumParticles == 0 || (VertexSetNoOffset && VertexSetNoOffset->IsEmpty()) || Stiffness <= 0.f)
 		{
 			return;
 		}
@@ -74,7 +78,7 @@ namespace Chaos::Softs
 			}
 		}
 		TSpatialHashGridPoints<int32, FSolverReal> SpatialHash(Diameter);
-		SpatialHash.InitializePoints(Entries);
+		SpatialHash.InitializePoints(Entries, Entries.Num() < Private::Chaos_SelfCollisionSphere_ParallelCount);
 		
 		const FSolverReal DiamSq = FMath::Square(Diameter);
 		constexpr int32 CellRadius = 1; // We set the cell size of the spatial hash such that we only need to look 1 cell away to find proximities.
@@ -98,7 +102,7 @@ namespace Chaos::Softs
 				}
 			}
 			return true;
-		}
+		}, Entries.Num() < Private::Chaos_SelfCollisionSphere_ParallelCount
 		);
 	}
 	template CHAOS_API void FPBDSelfCollisionSphereConstraintsBase::Init(const FSolverParticles& Particles);
@@ -143,8 +147,14 @@ namespace Chaos::Softs
 		const FCollectionPropertyConstFacade& PropertyCollection,
 		const TMap<FString, const TSet<int32>*>& VertexSets)
 	{
-		Radius = (FSolverReal)FMath::Max(GetSelfCollisionSphereRadius(PropertyCollection), 0.f);
-		Stiffness = (FSolverReal)FMath::Clamp(GetSelfCollisionSphereStiffness(PropertyCollection), 0.f, 1.f);
+		if (IsSelfCollisionSphereRadiusMutable(PropertyCollection))
+		{
+			Radius = (FSolverReal)FMath::Max(GetSelfCollisionSphereRadius(PropertyCollection), 0.f);
+		}
+		if (IsSelfCollisionSphereStiffnessMutable(PropertyCollection))
+		{
+			Stiffness = (FSolverReal)FMath::Clamp(GetSelfCollisionSphereStiffness(PropertyCollection), 0.f, 1.f);
+		}
 	}
 }  // End namespace Chaos::Softs
 
diff --git a/Engine/Source/Runtime/Experimental/Chaos/Public/Chaos/HierarchicalSpatialHash.h b/Engine/Source/Runtime/Experimental/Chaos/Public/Chaos/HierarchicalSpatialHash.h
--- a/Engine/Source/Runtime/Experimental/Chaos/Public/Chaos/HierarchicalSpatialHash.h
+++ b/Engine/Source/Runtime/Experimental/Chaos/Public/Chaos/HierarchicalSpatialHash.h
@@ -562,7 +562,7 @@ public:
 	}
 
 	template<typename ParticlesType>
-	void InitializePoints(const ParticlesType& Particles)
+	void InitializePoints(const ParticlesType& Particles, bool bForceSingleThreaded = false)
 	{
 		TRACE_CPUPROFILER_EVENT_SCOPE(SpatialHashGridConstantBounds_Init);
 		Reset();
@@ -579,13 +579,13 @@ public:
 			const FPayloadAndBounds Value = { Payload, ParticleIdx };
 			const int32 ElementIdx = HashMap.ConcurrentAddElement(CellIdx, INDEX_NONE, Value);
 			ParticleIndexToElement[ParticleIdx] = ElementIdx;
-		}
+		}, bForceSingleThreaded
 		);
 
 		HashMap.ShrinkElementsAfterConcurrentAdd();
 	}
 
-	TArray<TVec2<PayloadType>> FindAllSelfProximities(int32 CellRadius, int32 MaxNumExpectedConnections, TFunctionRef<bool(const int32 Payload0, const int32 Payload1)> NarrowTest) const
+	TArray<TVec2<PayloadType>> FindAllSelfProximities(int32 CellRadius, int32 MaxNumExpectedConnections, TFunctionRef<bool(const int32 Payload0, const int32 Payload1)> NarrowTest, bool bForceSingleThreaded = false) const
 	{
 		check(CellRadius >= 0);
 
@@ -698,7 +698,7 @@ public:
 			
 		} ConcurrentResults{ CellRadius, MaxNumExpectedConnections, NarrowTest, HashMap, ParticleIndexToElement };
 
-		PhysicsParallelFor(ParticleIndexToElement.Num(), ConcurrentResults);
+		PhysicsParallelFor(ParticleIndexToElement.Num(), ConcurrentResults, bForceSingleThreaded);
 
 		// Shrink Result array to actual number of found proximities
 		const int32 ResultNum = ConcurrentResults.ResultIndex.load();
