Estimate in bytes instead of KiB.

#rnx

diff --git a/Engine/Source/Runtime/Experimental/IoStore/HttpClient/Private/Activity.inl b/Engine/Source/Runtime/Experimental/IoStore/HttpClient/Private/Activity.inl
--- a/Engine/Source/Runtime/Experimental/IoStore/HttpClient/Private/Activity.inl
+++ b/Engine/Source/Runtime/Experimental/IoStore/HttpClient/Private/Activity.inl
@@ -204,14 +204,15 @@ static void Activity_SetError(FActivity* Activity, const FOutcome& Outcome)
 }
 
 ////////////////////////////////////////////////////////////////////////////////
-static void Activity_SetScore(FActivity* Activity, uint32 ContentEstKiB)
+static void Activity_SetScore(FActivity* Activity, uint32 ContentSizeEst)
 {
-	if (ContentEstKiB == 0)
+	if (ContentSizeEst == 0)
 	{
 		Activity->LengthScore = 0;
 		return;
 	}
 
+	uint32 ContentEstKiB = (ContentSizeEst + 1023) >> 10;
 	ContentEstKiB |= 2;
 	uint32 Pow2 = FMath::FloorLog2(uint32(ContentEstKiB));
 	Pow2 = FMath::Min(Pow2, 7u);
diff --git a/Engine/Source/Runtime/Experimental/IoStore/HttpClient/Private/Loop.inl b/Engine/Source/Runtime/Experimental/IoStore/HttpClient/Private/Loop.inl
--- a/Engine/Source/Runtime/Experimental/IoStore/HttpClient/Private/Loop.inl
+++ b/Engine/Source/Runtime/Experimental/IoStore/HttpClient/Private/Loop.inl
@@ -1552,7 +1552,7 @@ FRequest FEventLoop::Request(
 
 	Buffer.Fix();
 
-	Activity_SetScore(Activity, Params->ContentEstKiB);
+	Activity_SetScore(Activity, Params->ContentSizeEst);
 
 	memcpy(HostNamePtr, HostName.GetData(), HostNameLength);
 	HostNamePtr[HostNameLength] = '\0';
@@ -1587,7 +1587,7 @@ FRequest FEventLoop::Request(
 	Activity->bAllowChunked = (Params->bAllowChunked == true);
 	Activity->LengthScore = 0;
 
-	Activity_SetScore(Activity, Params->ContentEstKiB);
+	Activity_SetScore(Activity, Params->ContentSizeEst);
 
 	return Impl->Request(Method, Path, Activity);
 }
