Missed a file in previous submit.

#rnx

diff --git a/Engine/Source/Runtime/Experimental/IoStore/HttpClient/Internal/IO/Http/Client.h b/Engine/Source/Runtime/Experimental/IoStore/HttpClient/Internal/IO/Http/Client.h
--- a/Engine/Source/Runtime/Experimental/IoStore/HttpClient/Internal/IO/Http/Client.h
+++ b/Engine/Source/Runtime/Experimental/IoStore/HttpClient/Internal/IO/Http/Client.h
@@ -223,8 +223,8 @@ public:
 	struct FRequestParams
 	{
 		FCertRootsRef		VerifyCert		= {};
+		uint32				ContentSizeEst	= 0;
 		uint16				BufferSize		= 256;
-		uint16				ContentEstKiB	= 0;
 		bool				bAutoRedirect	= false;
 		bool				bAllowChunked	= true;
 	};
