Add a size estimate to requests that IAX makes.

#rnx

diff --git a/Engine/Source/Runtime/Experimental/IoStore/OnDemand/Private/OnDemandHttpClient.cpp b/Engine/Source/Runtime/Experimental/IoStore/OnDemand/Private/OnDemandHttpClient.cpp
--- a/Engine/Source/Runtime/Experimental/IoStore/OnDemand/Private/OnDemandHttpClient.cpp
+++ b/Engine/Source/Runtime/Experimental/IoStore/OnDemand/Private/OnDemandHttpClient.cpp
@@ -329,6 +329,7 @@ void FMultiEndpointHttpClient::IssueRequest(FRequest&& Request)
 
 	FEventLoop::FRequestParams RequestParams = 
 		{
+			.ContentSizeEst = uint32(Request.Range.GetLength()),
 			.bAutoRedirect = Config.Redirects == EHttpRedirects::Follow,
 			.bAllowChunked = Config.bAllowChunkedTransfer
 		};
