Add a new set of IAS analytics to try and better track request time and the data throughput rate.

#rnx

- When recording analytics about IAS requests we now use the length of time spent in ::recv rather than the total time since the request was launched as it is a better indicator of network performance.
-- FMultiEndpointHttpClientResponse now contains a HTTP::FTicketPerf::FSample which records these sorts of timings so that the data can be forwarded to the statistics code.
-- Enable IAS_HTTP_WITH_PERF by default as we want to use the timings this generates, even in shipping. Note that if this is disabled then the newer timings will all return 0.
- Dropping 'IasHttpDurationStdDev' as we have never used it and renamed 'IasHttpDurationMeanAvg' to 'IasHttpDurationMean' to better match other analytics names.
- We now reports the total sum of time spent in requests via 'IasHttpDurationSum'
-- GHttpDurationMsSum is reset each time analytics are recorded so that we only count the requests for the given recording period.
- 'IasHttpDataRateMean' records the average data throughput of the system via "total bytes downloaded" / "time spent in recv(s)"
-- Fixed FAnalyticsRecording to match the global analytics reporting as the plan is to replace global reporting with these recording scopes in the future.

diff --git a/Engine/Source/Runtime/Experimental/IoStore/HttpClient/Internal/IO/Http/Client.h b/Engine/Source/Runtime/Experimental/IoStore/HttpClient/Internal/IO/Http/Client.h
--- a/Engine/Source/Runtime/Experimental/IoStore/HttpClient/Internal/IO/Http/Client.h
+++ b/Engine/Source/Runtime/Experimental/IoStore/HttpClient/Internal/IO/Http/Client.h
@@ -8,7 +8,7 @@
 #endif
 
 #if !defined(IAS_HTTP_WITH_PERF)
-#	define IAS_HTTP_WITH_PERF !UE_BUILD_SHIPPING
+#	define IAS_HTTP_WITH_PERF 1
 #endif
 
 #define UE_API IOSTOREHTTPCLIENT_API
diff --git a/Engine/Source/Runtime/Experimental/IoStore/OnDemand/Private/OnDemandHttpClient.cpp b/Engine/Source/Runtime/Experimental/IoStore/OnDemand/Private/OnDemandHttpClient.cpp
--- a/Engine/Source/Runtime/Experimental/IoStore/OnDemand/Private/OnDemandHttpClient.cpp
+++ b/Engine/Source/Runtime/Experimental/IoStore/OnDemand/Private/OnDemandHttpClient.cpp
@@ -300,14 +300,16 @@ void FMultiEndpointHttpClient::IssueRequest(FRequest&& Request)
 				}
 				else
 				{
+					
 					FMultiEndpointHttpClientResponse Response
 					{
 						.Body					= MoveTemp(Request.Body),
+						.Sample					= TicketStatus.GetPerf().GetSample(),
 						.DurationMilliseconds	= DurationMs,
 						.StatusCode				= Request.StatusCode,
 						.RetryCount				= Request.RetryCount,
 						.HostIndex				= Request.Host,
-						.CDNCacheStatus = Request.CDNCacheStatus
+						.CDNCacheStatus			= Request.CDNCacheStatus
 					};
 
 					if (bError)
diff --git a/Engine/Source/Runtime/Experimental/IoStore/OnDemand/Private/OnDemandHttpClient.h b/Engine/Source/Runtime/Experimental/IoStore/OnDemand/Private/OnDemandHttpClient.h
--- a/Engine/Source/Runtime/Experimental/IoStore/OnDemand/Private/OnDemandHttpClient.h
+++ b/Engine/Source/Runtime/Experimental/IoStore/OnDemand/Private/OnDemandHttpClient.h
@@ -55,13 +55,14 @@ struct FMultiEndpointHttpClientResponse
 {
 	const bool	IsOk() const { return StatusCode > 199 && StatusCode < 299; }
 
-	FIoBuffer	Body;
-	FString		Reason;
-	uint64		DurationMilliseconds = 0;
-	uint32		StatusCode = 0;
-	uint32		RetryCount = 0;
-	int32		HostIndex = INDEX_NONE;
-	int8		CDNCacheStatus = -1;
+	FIoBuffer					Body;
+	FString						Reason;
+	HTTP::FTicketPerf::FSample	Sample;
+	uint64						DurationMilliseconds = 0;
+	uint32						StatusCode = 0;
+	uint32						RetryCount = 0;
+	int32						HostIndex = INDEX_NONE;
+	int8						CDNCacheStatus = -1;
 };
 
 /**
diff --git a/Engine/Source/Runtime/Experimental/IoStore/OnDemand/Private/OnDemandIoDispatcherBackend.cpp b/Engine/Source/Runtime/Experimental/IoStore/OnDemand/Private/OnDemandIoDispatcherBackend.cpp
--- a/Engine/Source/Runtime/Experimental/IoStore/OnDemand/Private/OnDemandIoDispatcherBackend.cpp
+++ b/Engine/Source/Runtime/Experimental/IoStore/OnDemand/Private/OnDemandIoDispatcherBackend.cpp
@@ -1540,7 +1540,7 @@ void FOnDemandIoBackend::ProcessHttpRequests(FMultiEndpointHttpClient& HttpClien
 								ChunkRequest->HostGroup.OnSuccessfulResponse();
 
 								ChunkRequest->Chunk = MoveTemp(HttpResponse.Body);
-								Stats.OnHttpGet(ChunkRequest->Chunk.DataSize(), HttpResponse.DurationMilliseconds);
+								Stats.OnHttpGet(ChunkRequest->Chunk.DataSize(), HttpResponse.Sample.RecvMs);
 
 								UE::Tasks::Launch(UE_SOURCE_LOCATION, [this, ChunkRequest]()
 									{
diff --git a/Engine/Source/Runtime/Experimental/IoStore/OnDemand/Private/Statistics.cpp b/Engine/Source/Runtime/Experimental/IoStore/OnDemand/Private/Statistics.cpp
--- a/Engine/Source/Runtime/Experimental/IoStore/OnDemand/Private/Statistics.cpp
+++ b/Engine/Source/Runtime/Experimental/IoStore/OnDemand/Private/Statistics.cpp
@@ -538,14 +538,19 @@ void FOnDemandIoBackendStats::ReportEndPointAnalytics(TArray<FAnalyticsEventAttr
 
 	if (GIasReportHttpAnalyticsEnabled)
 	{
+		const int64 ByteCount = GDeltaTracking.Get(TEXTVIEW("IasHttpDownloadedBytes"), GHttpDownloadedBytes.Get());
+		const double DataRateBPS = GHttpDurationMsSum > 0 ? static_cast<double>(ByteCount) / (static_cast<double>(GHttpDurationMsSum) / 1000.0) : 0.0;
+
 		AppendAnalyticsEventAttributeArray(OutAnalyticsArray
 			,TRACK_DELTA("IasHttpErrorCount", GHttpErrorCount.Get())
 			,TRACK_DELTA("IasHttpDecodeErrors", GHttpDecodeErrorCount.Get())
 			,TRACK_DELTA("IasHttpRetryCount", GHttpRetryCount.Get())
 			,TRACK_DELTA("IasHttpGetCount", GHttpGetCount.Get())
-			,TRACK_DELTA("IasHttpDownloadedBytes", GHttpDownloadedBytes.Get())
-			,TEXT("IasHttpDurationMeanAvg"), GHttpAvgDuration.GetMean()
-			,TEXT("IasHttpDurationStdDev"), GHttpAvgDuration.GetDeviation()
+			//,TRACK_DELTA("IasHttpDownloadedBytes", GHttpDownloadedBytes.Get()) // Disabled for now as we need ByteCount for DataRateBPS above.
+			, TEXT("IasHttpDownloadedBytes"), ByteCount
+			,TEXT("IasHttpDurationMean"), GHttpAvgDuration.GetMean()
+			,TEXT("IasHttpDurationSum"), GHttpDurationMsSum
+			,TEXT("IasHttpDataRateMean"), DataRateBPS
 		);
 	}
 
@@ -575,6 +580,7 @@ void FOnDemandIoBackendStats::ReportEndPointAnalytics(TArray<FAnalyticsEventAttr
 
 		// These values we can just reset as they are only being used with analytics
 		GHttpAvgDuration.Reset();
+		GHttpDurationMsSum = 0;
 	}
 
 #undef TRACK_DELTA
@@ -625,17 +631,25 @@ private:
 	void Report(TArray<FAnalyticsEventAttribute>& OutAnalyticsArray) const override
 	{
 		// Report if the system has ever actually managed to make a connection
-		AppendAnalyticsEventAttributeArray(OutAnalyticsArray, TEXT("IasHttpHasEverConnected"), GHttpConnectCount.Get() > 0);
+		AppendAnalyticsEventAttributeArray(OutAnalyticsArray
+			,TEXT("IasHttpDistributedEndpointResolved"), GHttpDistributedEndpointResolved
+			,TEXT("IasHttpHasEverConnected"), GHttpConnectCount.Get() > 0
+		);
 
 		if (GIasReportHttpAnalyticsEnabled)
 		{
+			const double DataRateBPS = Http.TotalDuration.GetValue() > 0 ? static_cast<double>(Http.DownloadedBytes.GetValue()) / (static_cast<double>(Http.TotalDuration.GetValue()) / 1000.0) : 0.0;
+			const double DurationMean = Http.GetCount.GetValue() ? static_cast<double>(Http.TotalDuration.GetValue()) / static_cast<double>(Http.GetCount.GetValue()) : 0.0;
+
 			AppendAnalyticsEventAttributeArray(OutAnalyticsArray
 				, TEXT("IasHttpErrorCount"), Http.ErrorCount.GetValue()
 				, TEXT("IasHttpDecodeErrorCount"), Http.DecodeErrorCount.GetValue()
 				, TEXT("IasHttpRetryCount"), Http.RetryCount.GetValue()
 				, TEXT("IasHttpGetCount"), Http.GetCount.GetValue()
 				, TEXT("IasHttpDownloadedBytes"), Http.DownloadedBytes.GetValue()
-				, TEXT("IasHttpTotalDuration"), Http.TotalDuration.GetValue()
+				, TEXT("IasHttpDurationMean"), DurationMean
+				, TEXT("IasHttpDurationSum"), Http.TotalDuration.GetValue()
+				, TEXT("IasHttpDataRateMean"), DataRateBPS
 			);
 		}
 
