Constraints evaluation: remove deprecated constraint's evaluation notification

#rnx

diff --git a/Engine/Source/Runtime/Experimental/Animation/Constraints/Private/Transform/TransformConstraint.cpp b/Engine/Source/Runtime/Experimental/Animation/Constraints/Private/Transform/TransformConstraint.cpp
--- a/Engine/Source/Runtime/Experimental/Animation/Constraints/Private/Transform/TransformConstraint.cpp
+++ b/Engine/Source/Runtime/Experimental/Animation/Constraints/Private/Transform/TransformConstraint.cpp
@@ -476,16 +476,22 @@ void UTickableTransformConstraint::OnHandleModified(UTransformableHandle* InHand
 	{
 		return;
 	}
-	
-	if(InHandle == ChildTRSHandle || InHandle == ParentTRSHandle)
+
+	const bool bIsThisChildHandle = InHandle == ChildTRSHandle;
+	const bool bIsThisParentHandle = InHandle == ParentTRSHandle;
+
+	// if the incoming handle has nothing to do with this constraint then exit
+	if (!bIsThisChildHandle && !bIsThisParentHandle)
 	{
-		if (InNotification == EHandleEvent::ComponentUpdated)
-		{
-			SetupDependencies(World);
-			// update dependencies now the component has been updated. 
-			FDependencyBuilder::BuildDependencies(World, this);
-			return;
-		}
+		return;
+	}
+
+	// update dependencies now the component has been updated.
+	if (InNotification == EHandleEvent::ComponentUpdated)
+	{
+		SetupDependencies(World);
+		FDependencyBuilder::BuildDependencies(World, this);
+		return;
 	}
 
 	auto MarkForEvaluation = [this, World]()
@@ -494,65 +500,26 @@ void UTickableTransformConstraint::OnHandleModified(UTransformableHandle* InHand
 		Controller.MarkConstraintForEvaluation(this);
 	};
 
-	auto DirectEvaluation = [this, World]()
-	{
-		const FConstraintTickFunction* ConstraintTick = ConstraintTicks.Find(World->GetCurrentLevel());
-		if (IsFullyActive() && ConstraintTick->IsTickFunctionRegistered() && ConstraintTick->IsTickFunctionEnabled())
-		{
-			Evaluate(true);
-		}
-	};
-	
-	if (InHandle == ChildTRSHandle)
+	if (bIsThisChildHandle)
 	{
 		if (InNotification == EHandleEvent::UpperDependencyUpdated)
 		{
 			const UObject* ParentTarget = ConstraintLocals::GetHandleTarget(ParentTRSHandle);
 			if (ParentTarget && ParentTarget != Target)
 			{
-				if (FConstraintsEvaluationGraph::UseEvaluationGraph())
-				{
-					MarkForEvaluation();
-				}
-				else
-				{
-					DirectEvaluation();
-				}
-				return;
+				return MarkForEvaluation();
 			}
 		}
-	}
-
-	// update the constraint if the parent's transform has been modified 
-	if (InHandle == ParentTRSHandle && InNotification == EHandleEvent::GlobalTransformUpdated)
-	{
-		if (FConstraintsEvaluationGraph::UseEvaluationGraph())
-		{
-			MarkForEvaluation();
-		}
-		else
-		{
-			DirectEvaluation();
-		}
-		return;
-	}
 
-	if (InHandle == ParentTRSHandle && InNotification == EHandleEvent::UpperDependencyUpdated)
-	{
-		if (FConstraintsEvaluationGraph::UseEvaluationGraph())
+		if (InNotification == EHandleEvent::GlobalTransformUpdated)
 		{
-			MarkForEvaluation();
-		}
-		else
-		{
-			DirectEvaluation();
+			return MarkForEvaluation();
 		}
-		return;
 	}
 
-	if (InHandle == ChildTRSHandle && InNotification == EHandleEvent::GlobalTransformUpdated)
+	if (bIsThisParentHandle)
 	{
-		if (FConstraintsEvaluationGraph::UseEvaluationGraph())
+		if (InNotification == EHandleEvent::GlobalTransformUpdated || InNotification == EHandleEvent::UpperDependencyUpdated)
 		{
 			return MarkForEvaluation();
 		}
