OnDemand package store can now update just refs when content is installed instead of fully rebuilding its state.
Fix bug where OnDemand Package store could return Missing instead of NotInstalled if no packages in a container were referenced.

diff --git a/Engine/Source/Runtime/Experimental/IoStore/OnDemand/Private/OnDemandIoStore.cpp b/Engine/Source/Runtime/Experimental/IoStore/OnDemand/Private/OnDemandIoStore.cpp
--- a/Engine/Source/Runtime/Experimental/IoStore/OnDemand/Private/OnDemandIoStore.cpp
+++ b/Engine/Source/Runtime/Experimental/IoStore/OnDemand/Private/OnDemandIoStore.cpp
@@ -1757,7 +1757,7 @@ FOnDemandInstallResult FOnDemandIoStore::TickInstallRequest(FOnDemandInternalIns
 	if (InstallData.IsEmpty() == false)
 	{
 		// Referencers were added so the package store must be refreshed
-		PackageStoreBackend->NeedsUpdate();
+		PackageStoreBackend->NeedsUpdate(EOnDemandPackageStoreUpdateFlags::ReferencedPackages);
 	}
 
 	double DurationInSeconds = FPlatformTime::Seconds() - StartTime;
@@ -2184,10 +2184,10 @@ void FOnDemandIoStore::ReleaseContent(FOnDemandInternalContentHandle& ContentHan
 	}
 
 	// Package store needs to be aware of install state
-	PackageStoreBackend->NeedsUpdate();
+	PackageStoreBackend->NeedsUpdate(EOnDemandPackageStoreUpdateFlags::ReferencedPackages);
 }
 
-void FOnDemandIoStore::GetReferencedContent(TArray<FSharedOnDemandContainer>& OutContainers, TArray<TBitArray<>>& OutChunkEntryIndices, bool bContainersWithPackagesOnly)
+void FOnDemandIoStore::GetReferencedContent(TArray<FSharedOnDemandContainer>& OutContainers, TArray<TBitArray<>>& OutChunkEntryIndices, bool bPackageStore)
 {
 	UE::TUniqueLock Lock(ContainerMutex);
 	for (FSharedOnDemandContainer& Container : Containers)
@@ -2202,12 +2202,19 @@ void FOnDemandIoStore::GetReferencedContent(TArray<FSharedOnDemandContainer>& Ou
 			continue;
 		}
 
-		if (bContainersWithPackagesOnly && Container->Header.IsValid() == false)
+		if (bPackageStore && Container->Header.IsValid() == false)
 		{
+			// Package store doesn't care about containers without packages
 			continue;
 		}
 
 		TBitArray<> Indices = Container->GetReferencedChunkEntries();
+		if (Indices.IsEmpty() && bPackageStore)
+		{
+			// Package store wants to know about about containers with packages even if there are not references yet
+			Indices.SetNum(Container->ChunkEntries.Num(), false);
+		}
+
 		if (Indices.IsEmpty() == false)
 		{
 			OutContainers.Add(Container);
diff --git a/Engine/Source/Runtime/Experimental/IoStore/OnDemand/Private/OnDemandIoStore.h b/Engine/Source/Runtime/Experimental/IoStore/OnDemand/Private/OnDemandIoStore.h
--- a/Engine/Source/Runtime/Experimental/IoStore/OnDemand/Private/OnDemandIoStore.h
+++ b/Engine/Source/Runtime/Experimental/IoStore/OnDemand/Private/OnDemandIoStore.h
@@ -294,7 +294,7 @@ public:
 	/** Used to access a number of FIoChunkId that exist in the system when running tests and is not intended for production use */
 	TArray<FIoChunkId>					DebugFindStreamingChunkIds(int32 NumToFind);
 #endif //!UE_BUILD_SHIPPING
-	void								GetReferencedContent(TArray<FSharedOnDemandContainer>& OutContainers, TArray<TBitArray<>>& OutChunkEntryIndices, bool bContainersWithPackagesOnly = false);
+	void								GetReferencedContent(TArray<FSharedOnDemandContainer>& OutContainers, TArray<TBitArray<>>& OutChunkEntryIndices, bool bPackageStore = false);
 
 	// IOnDemandIoStore interface
 	virtual FIoStatus					Initialize() override;
diff --git a/Engine/Source/Runtime/Experimental/IoStore/OnDemand/Private/OnDemandPackageStoreBackend.cpp b/Engine/Source/Runtime/Experimental/IoStore/OnDemand/Private/OnDemandPackageStoreBackend.cpp
--- a/Engine/Source/Runtime/Experimental/IoStore/OnDemand/Private/OnDemandPackageStoreBackend.cpp
+++ b/Engine/Source/Runtime/Experimental/IoStore/OnDemand/Private/OnDemandPackageStoreBackend.cpp
@@ -49,7 +49,7 @@ public:
 						FOnDemandPackageStoreBackend(TWeakPtr<FOnDemandIoStore> OnDemandIoStore);
 						virtual ~FOnDemandPackageStoreBackend ();
 
-	virtual void		NeedsUpdate() override;
+	virtual void		NeedsUpdate(EOnDemandPackageStoreUpdateFlags Flags) override;
 
 	virtual void		OnMounted(TSharedRef<const FPackageStoreBackendContext> Context) override;
 	virtual void 		BeginRead() override;
@@ -69,14 +69,15 @@ public:
 
 private:
 	void				UpdateLookupTables();
-
-	TWeakPtr<FOnDemandIoStore>	OnDemandIoStore;
-	TArray<FContainer>			Containers;
-	FEntryMap					EntryMap;
-	FLocalizedMap				LocalizedMap;
-	FRedirectMap				RedirectMap;
-	UE::FMutex					Mutex;
-	bool						bNeedsUpdate = true;
+	void				UpdateReferencedPackages();
+
+	TWeakPtr<FOnDemandIoStore>			OnDemandIoStore;
+	TArray<FContainer>					Containers;
+	FEntryMap							EntryMap;
+	FLocalizedMap						LocalizedMap;
+	FRedirectMap						RedirectMap;
+	UE::FMutex							Mutex;
+	EOnDemandPackageStoreUpdateFlags	NeedsUpdateFlags = EOnDemandPackageStoreUpdateFlags::Full;
 };
 
 ///////////////////////////////////////////////////////////////////////////////
@@ -89,10 +90,10 @@ FOnDemandPackageStoreBackend::~FOnDemandPackageStoreBackend()
 {
 }
 
-void FOnDemandPackageStoreBackend::NeedsUpdate()
+void FOnDemandPackageStoreBackend::NeedsUpdate(EOnDemandPackageStoreUpdateFlags Flags)
 {
 	UE::TUniqueLock Lock(Mutex);
-	bNeedsUpdate = true;
+	NeedsUpdateFlags |= Flags;
 }
 
 void FOnDemandPackageStoreBackend::OnMounted(TSharedRef<const FPackageStoreBackendContext> Context)
@@ -179,13 +180,27 @@ TConstArrayView<uint32> FOnDemandPackageStoreBackend::GetSoftReferences(FPackage
 
 void FOnDemandPackageStoreBackend::UpdateLookupTables()
 {
-	if (!bNeedsUpdate)
+	check(Mutex.IsLocked());
+
+	EOnDemandPackageStoreUpdateFlags LocalNeedsUpdate = NeedsUpdateFlags;
+	NeedsUpdateFlags = EOnDemandPackageStoreUpdateFlags::None;
+
+	if (LocalNeedsUpdate == EOnDemandPackageStoreUpdateFlags::None)
 	{
 		return;
 	}
 
-	bNeedsUpdate = false;
-	check(Mutex.IsLocked());
+	if (!EnumHasAnyFlags(LocalNeedsUpdate, EOnDemandPackageStoreUpdateFlags::Full))
+	{
+		if (EnumHasAnyFlags(LocalNeedsUpdate, EOnDemandPackageStoreUpdateFlags::ReferencedPackages))
+		{
+			UpdateReferencedPackages();
+		}
+
+		return;
+	}
+
+	// Full Update
 
 	LocalizedMap.Empty();
 	RedirectMap.Empty();
@@ -196,8 +211,8 @@ void FOnDemandPackageStoreBackend::UpdateLookupTables()
 
 	if (TSharedPtr<FOnDemandIoStore> PinOnDemandIoStore = OnDemandIoStore.Pin())
 	{
-		const bool bContainersWithPackagesOnly = true;
-		PinOnDemandIoStore->GetReferencedContent(AllContainers, ReferencedChunkEntryIndices, bContainersWithPackagesOnly);
+		const bool bPackageStore = true;
+		PinOnDemandIoStore->GetReferencedContent(AllContainers, ReferencedChunkEntryIndices, bPackageStore);
 
 		for (const FSharedOnDemandContainer& C : AllContainers)
 		{
@@ -277,6 +292,40 @@ void FOnDemandPackageStoreBackend::UpdateLookupTables()
 	RedirectMap.Shrink();
 }
 
+void FOnDemandPackageStoreBackend::UpdateReferencedPackages()
+{
+	TArray<FSharedOnDemandContainer> AllContainers;
+	TArray<TBitArray<>> ReferencedChunkEntryIndices;
+
+	if (TSharedPtr<FOnDemandIoStore> PinOnDemandIoStore = OnDemandIoStore.Pin())
+	{
+		const bool bPackageStore = true;
+		PinOnDemandIoStore->GetReferencedContent(AllContainers, ReferencedChunkEntryIndices, bPackageStore);
+	}
+
+	for (int32 ContainerIdx = 0; const FSharedOnDemandContainer& SharedContainer : AllContainers)
+	{
+		const FIoContainerHeader& Hdr = *SharedContainer->Header;
+		TConstArrayView<FFilePackageStoreEntry> Entries(
+			reinterpret_cast<const FFilePackageStoreEntry*>(Hdr.StoreEntries.GetData()),
+			Hdr.PackageIds.Num());
+
+		for (int32 EntryIdx = 0; const FFilePackageStoreEntry & Entry : Entries)
+		{
+			const FPackageId PkgId = Hdr.PackageIds[EntryIdx];
+			const FIoChunkId PkgChunkId = CreatePackageDataChunkId(PkgId);
+			const int32 PkgChunkIdx = SharedContainer->FindChunkEntryIndex(PkgChunkId);
+			check(PkgChunkIdx != INDEX_NONE);
+			const bool IsReferenced = ReferencedChunkEntryIndices[ContainerIdx][PkgChunkIdx];
+
+			EntryMap.FindChecked(PkgId).bIsInstalled = IsReferenced;
+			++EntryIdx;
+		}
+
+		++ContainerIdx;
+	}
+}
+
 ///////////////////////////////////////////////////////////////////////////////
 TSharedPtr<IOnDemandPackageStoreBackend> MakeOnDemandPackageStoreBackend(TWeakPtr<FOnDemandIoStore> OnDemandIoStore)
 {
diff --git a/Engine/Source/Runtime/Experimental/IoStore/OnDemand/Private/OnDemandPackageStoreBackend.h b/Engine/Source/Runtime/Experimental/IoStore/OnDemand/Private/OnDemandPackageStoreBackend.h
--- a/Engine/Source/Runtime/Experimental/IoStore/OnDemand/Private/OnDemandPackageStoreBackend.h
+++ b/Engine/Source/Runtime/Experimental/IoStore/OnDemand/Private/OnDemandPackageStoreBackend.h
@@ -13,13 +13,21 @@ using FSharedContainerHeader = TSharedPtr<FIoContainerHeader>;
 namespace UE::IoStore
 {
 
+enum class EOnDemandPackageStoreUpdateFlags : uint8
+{
+	None				= 0,
+	Full				= (1 << 0),
+	ReferencedPackages	= (1 << 1),
+};
+ENUM_CLASS_FLAGS(EOnDemandPackageStoreUpdateFlags);
+
 class IOnDemandPackageStoreBackend
 	: public IPackageStoreBackend
 {
 public:
 	virtual ~IOnDemandPackageStoreBackend() = default;
 
-	virtual void NeedsUpdate() = 0;
+	virtual void NeedsUpdate(EOnDemandPackageStoreUpdateFlags Flags = EOnDemandPackageStoreUpdateFlags::Full) = 0;
 };
 
 TSharedPtr<IOnDemandPackageStoreBackend> MakeOnDemandPackageStoreBackend(TWeakPtr<FOnDemandIoStore> OnDemandIoStore);
