Iris - Don't send state with subobject destroy. Behavior can be controlled via cvar net.Iris.AllowSendStateWithEndReplicationForSubObjects which defaults to false.

diff --git a/Engine/Plugins/Runtime/ReplicationSystemTestPlugin/Source/Private/Tests/ReplicationSystem/ReplicatedTestObject.cpp b/Engine/Plugins/Runtime/ReplicationSystemTestPlugin/Source/Private/Tests/ReplicationSystem/ReplicatedTestObject.cpp
--- a/Engine/Plugins/Runtime/ReplicationSystemTestPlugin/Source/Private/Tests/ReplicationSystem/ReplicatedTestObject.cpp
+++ b/Engine/Plugins/Runtime/ReplicationSystemTestPlugin/Source/Private/Tests/ReplicationSystem/ReplicatedTestObject.cpp
@@ -248,9 +248,15 @@ void UTestReplicatedIrisObject::GetLifetimeReplicatedProps(TArray<FLifetimePrope
 	DOREPLIFETIME_WITH_PARAMS_FAST(ThisClass, IntA, Params);
 	DOREPLIFETIME_WITH_PARAMS_FAST(ThisClass, IntB, Params);
 	DOREPLIFETIME_WITH_PARAMS_FAST(ThisClass, IntC, Params);
+	DOREPLIFETIME_WITH_PARAMS_FAST(ThisClass, IntDWithOnRep, Params);
 	DOREPLIFETIME_WITH_PARAMS_FAST(ThisClass, StructD, Params);
 }
 
+void UTestReplicatedIrisObject::OnRep_IntD()
+{
+	bIntDHitOnRep = true;
+}
+
 void UTestReplicatedIrisObject::AddComponents(const UTestReplicatedIrisObject::FComponents& InComponents)
 {
 	check(!NetRefHandle.IsValid())
diff --git a/Engine/Plugins/Runtime/ReplicationSystemTestPlugin/Source/Private/Tests/ReplicationSystem/ReplicatedTestObject.h b/Engine/Plugins/Runtime/ReplicationSystemTestPlugin/Source/Private/Tests/ReplicationSystem/ReplicatedTestObject.h
--- a/Engine/Plugins/Runtime/ReplicationSystemTestPlugin/Source/Private/Tests/ReplicationSystem/ReplicatedTestObject.h
+++ b/Engine/Plugins/Runtime/ReplicationSystemTestPlugin/Source/Private/Tests/ReplicationSystem/ReplicatedTestObject.h
@@ -328,9 +328,17 @@ public:
 	UPROPERTY(Transient, Replicated)
 	int8 IntC;
 
+	UPROPERTY(Transient, ReplicatedUsing=OnRep_IntD)
+	int32 IntDWithOnRep;
+	bool bIntDHitOnRep = false;
+
 	UPROPERTY(Transient, Replicated)
 	FTestReplicatedIrisObject_Struct StructD;
 
+protected:
+	UFUNCTION()
+	void OnRep_IntD();
+
 public:
 	TArray<TStrongObjectPtr<UTestReplicatedIrisPropertyComponent>> Components;
 	TArray<TUniquePtr<FTestReplicatedIrisComponent>> IrisComponents;
diff --git a/Engine/Plugins/Runtime/ReplicationSystemTestPlugin/Source/Private/Tests/ReplicationSystem/TestFlushBeforeDestroy.cpp b/Engine/Plugins/Runtime/ReplicationSystemTestPlugin/Source/Private/Tests/ReplicationSystem/TestFlushBeforeDestroy.cpp
--- a/Engine/Plugins/Runtime/ReplicationSystemTestPlugin/Source/Private/Tests/ReplicationSystem/TestFlushBeforeDestroy.cpp
+++ b/Engine/Plugins/Runtime/ReplicationSystemTestPlugin/Source/Private/Tests/ReplicationSystem/TestFlushBeforeDestroy.cpp
@@ -2244,4 +2244,43 @@ UE_NET_TEST_FIXTURE(FTestFlushBeforeDestroyFixture, TestSubObjectIsNotDestroyedB
 	UE_NET_ASSERT_TRUE(Client->IsResolvableNetRefHandle(ServerSubObject1->NetRefHandle));
 }
 
+// The following test can crash if the state update is sent with the destroy
+UE_NET_TEST_FIXTURE(FTestFlushBeforeDestroyFixture, TestSubObjectStateIsNotSentWithDestroy)
+{
+	FReplicationSystemTestClient* Client = CreateClient();
+	RegisterNetBlobHandlers(Client);
+
+	UTestReplicatedIrisObject* ServerObject = Server->CreateObject();
+	FNetRefHandle ObjectHandle = ServerObject->NetRefHandle;
+	UTestReplicatedIrisObject* ServerSubObject = Server->CreateSubObject(ObjectHandle);
+	FNetRefHandle SubObjectHandle = ServerSubObject->NetRefHandle;
+
+	Server->PreSendUpdate();
+	Server->SendAndDeliverTo(Client, DeliverPacket);
+	Server->PostSendUpdate();
+
+	// Verify that objects is created
+	UE_NET_ASSERT_TRUE(Client->IsResolvableNetRefHandle(ObjectHandle));
+	UE_NET_ASSERT_TRUE(Client->IsResolvableNetRefHandle(SubObjectHandle));
+
+	// Modify state on SubObject
+	ServerSubObject->IntDWithOnRep += 3;
+
+	// Update but do not ack/nak the packet just yet
+	Server->PreSendUpdate();
+	Server->SendTo(Client);
+	Server->PostSendUpdate();
+
+	// Destroy subobject without flush
+	Server->DestroyObject(ServerSubObject);
+
+	// Fail to deliver subobject update
+	Server->DeliverTo(Client, DoNotDeliverPacket);
+
+	// Update and deliver packet
+	Server->UpdateAndSend({Client}, DeliverPacket);
+
+	UE_NET_ASSERT_FALSE(Client->IsValidNetRefHandle(SubObjectHandle));
+}
+
 }
diff --git a/Engine/Plugins/Runtime/ReplicationSystemTestPlugin/Source/Private/Tests/ReplicationSystem/TestReplicationBridge.cpp b/Engine/Plugins/Runtime/ReplicationSystemTestPlugin/Source/Private/Tests/ReplicationSystem/TestReplicationBridge.cpp
--- a/Engine/Plugins/Runtime/ReplicationSystemTestPlugin/Source/Private/Tests/ReplicationSystem/TestReplicationBridge.cpp
+++ b/Engine/Plugins/Runtime/ReplicationSystemTestPlugin/Source/Private/Tests/ReplicationSystem/TestReplicationBridge.cpp
@@ -34,9 +34,9 @@ UE_NET_TEST_FIXTURE(FReplicationBridgeTestFixture, CreateObjectWithSingleFragmen
 	UE_NET_ASSERT_TRUE(Protocol->ProtocolIdentifier != 0);
 	UE_NET_ASSERT_EQ(1u, Protocol->ReplicationStateCount);
 	UE_NET_ASSERT_EQ(4u, Protocol->InternalTotalAlignment);
-	UE_NET_ASSERT_EQ(16u, Protocol->InternalTotalSize);
-	UE_NET_ASSERT_EQ(28u, Protocol->MaxExternalStateSize);
-	UE_NET_ASSERT_EQ(4u, Protocol->ChangeMaskBitCount);
+	UE_NET_ASSERT_EQ(20u, Protocol->InternalTotalSize);
+	UE_NET_ASSERT_EQ(32u, Protocol->MaxExternalStateSize);
+	UE_NET_ASSERT_EQ(5u, Protocol->ChangeMaskBitCount);
 	UE_NET_ASSERT_EQ(EReplicationProtocolTraits::SupportsDeltaCompression, Protocol->ProtocolTraits);
 
 	// Destroy the handle
@@ -64,7 +64,7 @@ UE_NET_TEST_FIXTURE(FReplicationBridgeTestFixture, CreateObjectWithMultipleFragm
 	UE_NET_ASSERT_EQ(8u, Protocol->InternalTotalAlignment);
 	UE_NET_ASSERT_EQ(16u + PropertyComponentCount*PropertyComponentSize, Protocol->InternalTotalSize);
 	UE_NET_ASSERT_EQ(40u, Protocol->MaxExternalStateSize);
-	UE_NET_ASSERT_EQ(4U + PropertyComponentCount*RegularStatePropertyCountInPropertyComponent, Protocol->ChangeMaskBitCount);
+	UE_NET_ASSERT_EQ(5U + PropertyComponentCount*RegularStatePropertyCountInPropertyComponent, Protocol->ChangeMaskBitCount);
 	UE_NET_ASSERT_EQ(EReplicationProtocolTraits::SupportsDeltaCompression, Protocol->ProtocolTraits);
 
 	// Destroy the handle
@@ -92,7 +92,7 @@ UE_NET_TEST_FIXTURE(FReplicationBridgeTestFixture, CreateObjectWithMixedFragment
 	UE_NET_ASSERT_EQ(8u, Protocol->InternalTotalAlignment);
 	UE_NET_ASSERT_EQ(16u + PropertyComponentCount*PropertyComponentSize + IrisComponentCount*12U, Protocol->InternalTotalSize);
 	UE_NET_ASSERT_EQ(40u, Protocol->MaxExternalStateSize);
-	UE_NET_ASSERT_EQ(4U + PropertyComponentCount*RegularStatePropertyCountInPropertyComponent + IrisComponentCount*3U, Protocol->ChangeMaskBitCount);
+	UE_NET_ASSERT_EQ(5U + PropertyComponentCount*RegularStatePropertyCountInPropertyComponent + IrisComponentCount*3U, Protocol->ChangeMaskBitCount);
 	UE_NET_ASSERT_EQ(EReplicationProtocolTraits::SupportsDeltaCompression, Protocol->ProtocolTraits);
 
 	// Destroy the handle
diff --git a/Engine/Source/Runtime/Experimental/Iris/Core/Private/Iris/ReplicationSystem/ReplicationWriter.cpp b/Engine/Source/Runtime/Experimental/Iris/Core/Private/Iris/ReplicationSystem/ReplicationWriter.cpp
--- a/Engine/Source/Runtime/Experimental/Iris/Core/Private/Iris/ReplicationSystem/ReplicationWriter.cpp
+++ b/Engine/Source/Runtime/Experimental/Iris/Core/Private/Iris/ReplicationSystem/ReplicationWriter.cpp
@@ -94,6 +94,9 @@ static int32 GReplicationWriterCannotSendWarningInterval = 256;
 static FAutoConsoleVariableRef CVarReplicationWriterCannotSendWarningInterval(TEXT("net.Iris.ReplicationWriterCannotSendWarningInterval"), GReplicationWriterCannotSendWarningInterval,
 	TEXT("Warn if we are prevented from Sending for more than a certain number of updates, setting it to 0 will disable warning"));
 
+static bool bAllowSendStateWithEndReplicationForSubObjects = false;
+static FAutoConsoleVariableRef CvarAllowSendStateWithEndReplicationForSubObjects(TEXT("net.Iris.AllowSendStateWithEndReplicationForSubObjects"), bAllowSendStateWithEndReplicationForSubObjects, TEXT("Allow state to be sent together with subobject destroy. Default is false. Cannot be true if net.Iris.ImmediateDispatchEndReplicationForSubObjects is true."));
+
 static const FName NetError_ObjectStateTooLarge("Object state is too large to be split.");
 
 const TCHAR* FReplicationWriter::LexToString(const EReplicatedObjectState State)
@@ -697,6 +700,13 @@ void FReplicationWriter::SetPendingDestroyOrSubObjectPendingDestroyState(uint32
 				MarkObjectDirty(InternalIndex, "SetPendingDestroyOrSubObjectPendingDestroyState2");
 				ObjectsPendingDestroy.SetBit(InternalIndex);
 				Info.SubObjectPendingDestroy = 1U;
+				if (!bAllowSendStateWithEndReplicationForSubObjects)
+				{
+					// For clarity clear HasDirtyChangeMask as we do not intend to replicate state. The object itself needs to remain dirty in ObjectsWithDirtyChanges though.
+					Info.HasDirtyChangeMask = 0U;
+					// One cannot cancel a subobject destroy so we should not send any attachments for this subobject going forward
+					Info.HasAttachments = 0U;
+				}
 				ensure(!Info.TearOff || (GetFlushStatus(InternalIndex, Info, Info.FlushFlags) == FlushFlags_None));
 				return;
 			}
@@ -2385,10 +2395,11 @@ FReplicationWriter::EWriteObjectStatus FReplicationWriter::WriteObjectAndSubObje
 	}
 	Info.HasDirtyConditionals = 0;
 
+	const bool bAllowStateAndAttachmentSending = (!Info.SubObjectPendingDestroy || bAllowSendStateWithEndReplicationForSubObjects);
 	const bool bIsObjectIndexForAttachment = IsObjectIndexForOOBAttachment(InternalIndex);
-	const bool bHasState = (bIsInitialState || Info.HasDirtyChangeMask) && !!(WriteObjectFlags & EWriteObjectFlag::WriteObjectFlag_State);
-	const bool bHasAttachments = (Info.HasAttachments || bIsObjectIndexForAttachment);
-	const bool bWriteAttachments = bHasAttachments && !!(WriteObjectFlags & EWriteObjectFlag::WriteObjectFlag_Attachments);
+	const bool bHasState = (bIsInitialState || Info.HasDirtyChangeMask) && !!(WriteObjectFlags & EWriteObjectFlag::WriteObjectFlag_State) && bAllowStateAndAttachmentSending;
+	const bool bHasAttachments = (Info.HasAttachments || bIsObjectIndexForAttachment) && bAllowStateAndAttachmentSending;
+	const bool bWriteAttachments = bHasAttachments && !!(WriteObjectFlags & EWriteObjectFlag::WriteObjectFlag_Attachments) && bAllowStateAndAttachmentSending;
 	BatchEntry.bHasUnsentAttachments = bHasAttachments;
 
 	// Check if we must defer tearoff until after flush
