[NetPhysics] Fix bug where RewindData would put state into the input history.

diff --git a/Engine/Source/Runtime/Experimental/Chaos/Public/RewindData.h b/Engine/Source/Runtime/Experimental/Chaos/Public/RewindData.h
--- a/Engine/Source/Runtime/Experimental/Chaos/Public/RewindData.h
+++ b/Engine/Source/Runtime/Experimental/Chaos/Public/RewindData.h
@@ -1735,7 +1735,7 @@ public:
 	/** Add state history for particle to the rewind data */
 	void AddStateHistory(const TSharedPtr<FBaseRewindHistory>& StateHistory, Chaos::FGeometryParticleHandle* Particle)
 	{
-		AddInputHistory(StateHistory);
+		AddStateHistory(StateHistory);
 		if (Particle != nullptr)
 		{
 			StateParticleHistories.Add(Particle, StateHistory.ToWeakPtr());
@@ -1745,7 +1745,7 @@ public:
 	/** Remove state history for particle from the rewind data */
 	void RemoveStateHistory(const TSharedPtr<FBaseRewindHistory>& StateHistory, Chaos::FGeometryParticleHandle* Particle)
 	{
-		RemoveInputHistory(StateHistory);
+		RemoveStateHistory(StateHistory);
 		if (Particle != nullptr)
 		{
 			StateParticleHistories.Remove(Particle);
