Extend FOnDemandThread::IssueRequest so that it will be more suitable for IAD.

#rnx

- FOnDemandThread::IssueRequest will now take FOnDemandChunkInfo by reference or rtype and store it as part of it's internal request rather than keeping track of a pointer.
- FOnDemandChunkInfo now requires copy constructors.

diff --git a/Engine/Source/Runtime/Experimental/IoStore/OnDemand/Private/OnDemandHttpThread.cpp b/Engine/Source/Runtime/Experimental/IoStore/OnDemand/Private/OnDemandHttpThread.cpp
--- a/Engine/Source/Runtime/Experimental/IoStore/OnDemand/Private/OnDemandHttpThread.cpp
+++ b/Engine/Source/Runtime/Experimental/IoStore/OnDemand/Private/OnDemandHttpThread.cpp
@@ -78,6 +78,15 @@ const int32 GREPLACEOnDemandBackendThreadPriorities[5] =
 
 struct FHttpRequest
 {
+	FHttpRequest() = default;
+	FHttpRequest(FOnDemandChunkInfo&& InChunkInfo)
+		: ChunkInfo(MoveTemp(InChunkInfo))
+	{
+
+	}
+
+	~FHttpRequest() = default;
+
 	void OnRequestCompleted(FIoBuffer&& ChunkData)
 	{
 		CompletionCallback(ChunkData);
@@ -85,17 +94,17 @@ struct FHttpRequest
 
 	FIASHostGroup HostGroup()
 	{
-		return ChunkInfo->HostGroup();
+		return ChunkInfo.HostGroup();
 	}
 
 	void GetUrl(FAnsiStringBuilderBase& Url) const
 	{
-		ChunkInfo->GetUrl(Url);
+		ChunkInfo.GetUrl(Url);
 	}
 
 	FHttpRequest* NextRequest = nullptr;
 	FOnDemandThread::FCompletionCallback CompletionCallback;
-	const FOnDemandChunkInfo* ChunkInfo = nullptr;
+	FOnDemandChunkInfo ChunkInfo;
 	FIoOffsetAndLength ChunkRange;
 	bool bCancelled = false;
 	int32 Priority = 0;
@@ -136,15 +145,17 @@ FOnDemandThread& FOnDemandThread::Get()
 	return *SingletonInstance;
 }
 
-FOnDemandThread::FRequestHandle FOnDemandThread::IssueRequest(const FOnDemandChunkInfo* ChunkInfo, const FIoOffsetAndLength& ReadRange, FCompletionCallback&& CompletionCallback)
+FOnDemandThread::FRequestHandle FOnDemandThread::IssueRequest(const FOnDemandChunkInfo& ChunkInfo, const FIoOffsetAndLength& ReadRange, FCompletionCallback&& CompletionCallback)
 {
-	TRACE_CPUPROFILER_EVENT_SCOPE(FOnDemandThread::IssueRequest);
+	return IssueRequest(FOnDemandChunkInfo(ChunkInfo), ReadRange, MoveTemp(CompletionCallback));
+}
 
-	checkf(ChunkInfo != nullptr, TEXT("FOnDemandThread::IssueRequest requires a valid ChunkInfo"));
+FOnDemandThread::FRequestHandle FOnDemandThread::IssueRequest(FOnDemandChunkInfo&& ChunkInfo, const FIoOffsetAndLength& ReadRange, FCompletionCallback&& CompletionCallback)
+{
+	TRACE_CPUPROFILER_EVENT_SCOPE(FOnDemandThread::IssueRequest);
 
-	FHttpRequest* Request = AllocateRequest();
+	FHttpRequest* Request = AllocateRequest(MoveTemp(ChunkInfo));
 
-	Request->ChunkInfo = ChunkInfo;
 	Request->ChunkRange = ReadRange;
 	Request->CompletionCallback = MoveTemp(CompletionCallback);
 
@@ -395,10 +406,10 @@ void FOnDemandThread::UpdateThreadPriorityIfNeeded()
 	}
 }
 
-FHttpRequest* FOnDemandThread::AllocateRequest()
+FHttpRequest* FOnDemandThread::AllocateRequest(FOnDemandChunkInfo&& ChunkInfo)
 {
 	UE::TUniqueLock _(AllocatorMutex);
-	return RequestAllocator.Construct();
+	return RequestAllocator.Construct(MoveTemp(ChunkInfo));
 }
 
 void FOnDemandThread::DestroyRequest(FHttpRequest* Request)
diff --git a/Engine/Source/Runtime/Experimental/IoStore/OnDemand/Private/OnDemandHttpThread.h b/Engine/Source/Runtime/Experimental/IoStore/OnDemand/Private/OnDemandHttpThread.h
--- a/Engine/Source/Runtime/Experimental/IoStore/OnDemand/Private/OnDemandHttpThread.h
+++ b/Engine/Source/Runtime/Experimental/IoStore/OnDemand/Private/OnDemandHttpThread.h
@@ -41,8 +41,10 @@ public:
 	virtual ~FOnDemandThread();
 
 	static FOnDemandThread& Get();
+	
+	FRequestHandle IssueRequest(const FOnDemandChunkInfo& ChunkInfo, const FIoOffsetAndLength& ReadRange, FCompletionCallback&& CompletionCallback);
+	FRequestHandle IssueRequest(FOnDemandChunkInfo&& ChunkInfo, const FIoOffsetAndLength& ReadRange, FCompletionCallback&& CompletionCallback);
 
-	FRequestHandle IssueRequest(const FOnDemandChunkInfo* ChunkInfo, const FIoOffsetAndLength& ReadRange, FCompletionCallback&& CompletionCallback);
 	void ReprioritizeRequest(FRequestHandle Request);
 
 	void CancelRequest(FRequestHandle Request);
@@ -65,7 +67,8 @@ private:
 	bool TryCreateHttpClient();
 	void UpdateThreadPriorityIfNeeded();
 
-	FHttpRequest* AllocateRequest();
+	FHttpRequest* AllocateRequest(FOnDemandChunkInfo&& ChunkInfo);
+
 	void DestroyRequest(FHttpRequest* Request);
 
 private:
diff --git a/Engine/Source/Runtime/Experimental/IoStore/OnDemand/Private/OnDemandIoDispatcherBackend.cpp b/Engine/Source/Runtime/Experimental/IoStore/OnDemand/Private/OnDemandIoDispatcherBackend.cpp
--- a/Engine/Source/Runtime/Experimental/IoStore/OnDemand/Private/OnDemandIoDispatcherBackend.cpp
+++ b/Engine/Source/Runtime/Experimental/IoStore/OnDemand/Private/OnDemandIoDispatcherBackend.cpp
@@ -1082,7 +1082,7 @@ void FOnDemandIoBackend::CompleteCacheRead(FChunkRequest* ChunkRequest)
 	}
 
 #if UE_ENABLE_REQUEST_THREAD
-	ChunkRequest->HttpRequest = FOnDemandThread::Get().IssueRequest(&ChunkRequest->Params.ChunkInfo, ChunkRequest->Params.ChunkRange,
+	ChunkRequest->HttpRequest = FOnDemandThread::Get().IssueRequest(ChunkRequest->Params.ChunkInfo, ChunkRequest->Params.ChunkRange,
 		[this, ChunkRequest](FIoBuffer ChunkData)
 		{
 			ChunkRequest->Chunk = MoveTemp(ChunkData);
@@ -1182,7 +1182,7 @@ bool FOnDemandIoBackend::Resolve(FIoRequestImpl* Request)
 	ChunkRequest->HostGroup = ChunkInfo.HostGroup();
 
 #if UE_ENABLE_REQUEST_THREAD
-	ChunkRequest->HttpRequest = FOnDemandThread::Get().IssueRequest(&ChunkRequest->Params.ChunkInfo, ChunkRequest->Params.ChunkRange,
+	ChunkRequest->HttpRequest = FOnDemandThread::Get().IssueRequest(ChunkRequest->Params.ChunkInfo, ChunkRequest->Params.ChunkRange,
 		[this, ChunkRequest](FIoBuffer ChunkData)
 		{
 			ChunkRequest->Chunk = MoveTemp(ChunkData);
diff --git a/Engine/Source/Runtime/Experimental/IoStore/OnDemand/Private/OnDemandIoStore.h b/Engine/Source/Runtime/Experimental/IoStore/OnDemand/Private/OnDemandIoStore.h
--- a/Engine/Source/Runtime/Experimental/IoStore/OnDemand/Private/OnDemandIoStore.h
+++ b/Engine/Source/Runtime/Experimental/IoStore/OnDemand/Private/OnDemandIoStore.h
@@ -187,6 +187,21 @@ struct FOnDemandChunkInfo
 		: Entry(FOnDemandChunkEntry::Null)
 	{ }
 
+	FOnDemandChunkInfo(const FOnDemandChunkInfo& Other)
+		: SharedContainer(Other.SharedContainer)
+		, Entry(Other.Entry)
+	{
+
+	}
+
+	FOnDemandChunkInfo(FOnDemandChunkInfo&& Other)
+		: SharedContainer(MoveTemp(Other.SharedContainer))
+		, Entry(Other.Entry)
+	{
+
+	}
+
+
 	const									FIoHash& Hash() const { return Entry.Hash; }
 	uint32									RawSize() const { return Entry.RawSize; }
 	uint32									EncodedSize() const { return Entry.EncodedSize; }
