Extend the FOnDemandThread completion callback to return the http status and error reason.

#rnx

- This makes it a bit clearer when the http operation has actually failed rather than relying on checking if the FIoBuffer is empty, it should also allow us to provide better error messages to the caller.
- Added error reasons for things like request cancellation or the host group failing. The http client will also return error messages under certain conditions.
- Note that the task we launch for OnRequestCompleted is now right up against the max sized allowed by the small task allocator. If we capture more data the task will end up being allocated on the heap.

diff --git a/Engine/Source/Runtime/Experimental/IoStore/OnDemand/Private/OnDemandHttpThread.cpp b/Engine/Source/Runtime/Experimental/IoStore/OnDemand/Private/OnDemandHttpThread.cpp
--- a/Engine/Source/Runtime/Experimental/IoStore/OnDemand/Private/OnDemandHttpThread.cpp
+++ b/Engine/Source/Runtime/Experimental/IoStore/OnDemand/Private/OnDemandHttpThread.cpp
@@ -87,9 +87,9 @@ struct FHttpRequest
 
 	~FHttpRequest() = default;
 
-	void OnRequestCompleted(FIoBuffer&& ChunkData)
+	void OnRequestCompleted(int32 StatusCode, FStringView ErrorReason, FIoBuffer&& Data)
 	{
-		CompletionCallback(ChunkData);
+		CompletionCallback(StatusCode, ErrorReason, MoveTemp(Data));
 	}
 
 	FIASHostGroup HostGroup()
@@ -257,14 +257,14 @@ void FOnDemandThread::TickRequests()
 
 				if (HttpRequest->bCancelled)
 				{
-					HttpRequest->OnRequestCompleted(FIoBuffer());
+					HttpRequest->OnRequestCompleted(0, TEXTVIEW("Request cancelled"), FIoBuffer());
 					FOnDemandIoBackendStats::Get()->OnHttpCancel();
 
 					DestroyRequest(HttpRequest);
 				}
 				else if (!HttpRequest->HostGroup().IsConnected() || !bHttpEnabled)
 				{
-					HttpRequest->OnRequestCompleted(FIoBuffer());
+					HttpRequest->OnRequestCompleted(0, TEXTVIEW("Hostgroup is disconnected"), FIoBuffer());
 
 					// Technically this request is being skipped because of a pre-existing error. It is not
 					// an error itself and it is not being canceled by higher level code. However we do not
@@ -302,10 +302,10 @@ void FOnDemandThread::TickRequests()
 
 								FOnDemandIoBackendStats::Get()->OnHttpGet(HttpResponse.Body.DataSize(), HttpResponse.DurationMilliseconds);
 
-								UE::Tasks::Launch(UE_SOURCE_LOCATION, [this, HttpRequest, ChunkData = MoveTemp(HttpResponse.Body)]() mutable
+								UE::Tasks::Launch(UE_SOURCE_LOCATION, [HttpRequest, StatusCode = HttpResponse.StatusCode, ErrorReason = MoveTemp(HttpResponse.Reason), Data = MoveTemp(HttpResponse.Body)]() mutable
 									{
-										HttpRequest->OnRequestCompleted(MoveTemp(ChunkData));
-										DestroyRequest(HttpRequest);
+										HttpRequest->OnRequestCompleted(StatusCode, ErrorReason, MoveTemp(Data));
+										FOnDemandThread::Get().DestroyRequest(HttpRequest); // Don't pass the this pointer into the lambda to keep it below the small task allocation size
 									}, GetRequestCompletionTaskPriority());
 							}
 							else
@@ -318,7 +318,7 @@ void FOnDemandThread::TickRequests()
 									FOnDemandIoBackendStats::Get()->OnHttpDisconnected();
 								}
 
-								HttpRequest->OnRequestCompleted(FIoBuffer());
+								HttpRequest->OnRequestCompleted(HttpResponse.StatusCode, HttpResponse.Reason, FIoBuffer());
 								DestroyRequest(HttpRequest);
 							}
 						});
@@ -368,7 +368,7 @@ void FOnDemandThread::DrainHttpRequests()
 		Iterator = Iterator->NextRequest;
 
 		FOnDemandIoBackendStats::Get()->OnHttpDequeue();
-		Request->OnRequestCompleted(FIoBuffer());
+		Request->OnRequestCompleted(0, TEXTVIEW("Request cancelled due to shutdown"), FIoBuffer());
 		DestroyRequest(Request);
 		FOnDemandIoBackendStats::Get()->OnHttpCancel();
 	}
diff --git a/Engine/Source/Runtime/Experimental/IoStore/OnDemand/Private/OnDemandHttpThread.h b/Engine/Source/Runtime/Experimental/IoStore/OnDemand/Private/OnDemandHttpThread.h
--- a/Engine/Source/Runtime/Experimental/IoStore/OnDemand/Private/OnDemandHttpThread.h
+++ b/Engine/Source/Runtime/Experimental/IoStore/OnDemand/Private/OnDemandHttpThread.h
@@ -33,7 +33,7 @@ class FOnDemandThread final
 	: public FRunnable
 {
 public:
-	using FCompletionCallback = TUniqueFunction<void(FIoBuffer /*ChunkData*/)>;
+	using FCompletionCallback = TUniqueFunction<void(int32 /*StatusCode*/ , FStringView/*ErrorReason*/, FIoBuffer&& /*Data*/)>;
 	using FRequestHandle = void*;
 	DECLARE_MULTICAST_DELEGATE(FOnTickIdle);
 
diff --git a/Engine/Source/Runtime/Experimental/IoStore/OnDemand/Private/OnDemandIoDispatcherBackend.cpp b/Engine/Source/Runtime/Experimental/IoStore/OnDemand/Private/OnDemandIoDispatcherBackend.cpp
--- a/Engine/Source/Runtime/Experimental/IoStore/OnDemand/Private/OnDemandIoDispatcherBackend.cpp
+++ b/Engine/Source/Runtime/Experimental/IoStore/OnDemand/Private/OnDemandIoDispatcherBackend.cpp
@@ -1083,9 +1083,9 @@ void FOnDemandIoBackend::CompleteCacheRead(FChunkRequest* ChunkRequest)
 
 #if UE_ENABLE_REQUEST_THREAD
 	ChunkRequest->HttpRequest = FOnDemandThread::Get().IssueRequest(ChunkRequest->Params.ChunkInfo, ChunkRequest->Params.ChunkRange,
-		[this, ChunkRequest](FIoBuffer ChunkData)
+		[this, ChunkRequest](int32 StatusCode, FStringView ErrorReason, FIoBuffer&& Data)
 		{
-			ChunkRequest->Chunk = MoveTemp(ChunkData);
+			ChunkRequest->Chunk = MoveTemp(Data);
 			CompleteRequest(ChunkRequest);
 		});
 
@@ -1183,9 +1183,9 @@ bool FOnDemandIoBackend::Resolve(FIoRequestImpl* Request)
 
 #if UE_ENABLE_REQUEST_THREAD
 	ChunkRequest->HttpRequest = FOnDemandThread::Get().IssueRequest(ChunkRequest->Params.ChunkInfo, ChunkRequest->Params.ChunkRange,
-		[this, ChunkRequest](FIoBuffer ChunkData)
+		[this, ChunkRequest](int32 StatusCode, FStringView ErrorReason, FIoBuffer&& Data)
 		{
-			ChunkRequest->Chunk = MoveTemp(ChunkData);
+			ChunkRequest->Chunk = MoveTemp(Data);
 			CompleteRequest(ChunkRequest);
 		});
 
