Update FOnDemandThread callback to use the correct type for the status code to avoid type casting.

#rnx

- The callback was passing the status code around as a int32 but we originally provide it as uint32.
- Removed duplicate code validating status codes and moved it to a common utility function IsHttpStatusOk.

diff --git a/Engine/Source/Runtime/Experimental/IoStore/OnDemand/Private/OnDemandHttpClient.cpp b/Engine/Source/Runtime/Experimental/IoStore/OnDemand/Private/OnDemandHttpClient.cpp
--- a/Engine/Source/Runtime/Experimental/IoStore/OnDemand/Private/OnDemandHttpClient.cpp
+++ b/Engine/Source/Runtime/Experimental/IoStore/OnDemand/Private/OnDemandHttpClient.cpp
@@ -156,7 +156,7 @@ TIoStatusOr<FMultiEndpointHttpClientResponse> FMultiEndpointHttpClient::Get(FAns
 			// Busy loop
 		}
 
-		if (StatusCode > 199 && StatusCode < 300)
+		if (IsHttpStatusOk(StatusCode))
 		{
 			FMultiEndpointHttpClientResponse Response
 			{
diff --git a/Engine/Source/Runtime/Experimental/IoStore/OnDemand/Private/OnDemandHttpClient.h b/Engine/Source/Runtime/Experimental/IoStore/OnDemand/Private/OnDemandHttpClient.h
--- a/Engine/Source/Runtime/Experimental/IoStore/OnDemand/Private/OnDemandHttpClient.h
+++ b/Engine/Source/Runtime/Experimental/IoStore/OnDemand/Private/OnDemandHttpClient.h
@@ -26,6 +26,10 @@ namespace UE::IoStore
 {
 
 FIoStatus LoadDefaultHttpCertificates(bool& bWasLoaded); 
+inline bool IsHttpStatusOk(uint32 StatusCode)
+{
+	return StatusCode > 199 && StatusCode < 300;
+}
 
 enum class EHttpRedirects
 {
@@ -53,7 +57,10 @@ struct FMultiEndpointHttpClientConfig
 
 struct FMultiEndpointHttpClientResponse
 {
-	const bool	IsOk() const { return StatusCode > 199 && StatusCode < 299; }
+	const bool	IsOk() const
+	{
+		return IsHttpStatusOk(StatusCode);
+	}
 
 	FIoBuffer					Body;
 	FString						Reason;
diff --git a/Engine/Source/Runtime/Experimental/IoStore/OnDemand/Private/OnDemandHttpThread.cpp b/Engine/Source/Runtime/Experimental/IoStore/OnDemand/Private/OnDemandHttpThread.cpp
--- a/Engine/Source/Runtime/Experimental/IoStore/OnDemand/Private/OnDemandHttpThread.cpp
+++ b/Engine/Source/Runtime/Experimental/IoStore/OnDemand/Private/OnDemandHttpThread.cpp
@@ -87,7 +87,7 @@ struct FHttpRequest
 
 	~FHttpRequest() = default;
 
-	void OnRequestCompleted(int32 StatusCode, FStringView ErrorReason, FIoBuffer&& Data)
+	void OnRequestCompleted(uint32 StatusCode, FStringView ErrorReason, FIoBuffer&& Data)
 	{
 		CompletionCallback(StatusCode, ErrorReason, MoveTemp(Data));
 	}
diff --git a/Engine/Source/Runtime/Experimental/IoStore/OnDemand/Private/OnDemandHttpThread.h b/Engine/Source/Runtime/Experimental/IoStore/OnDemand/Private/OnDemandHttpThread.h
--- a/Engine/Source/Runtime/Experimental/IoStore/OnDemand/Private/OnDemandHttpThread.h
+++ b/Engine/Source/Runtime/Experimental/IoStore/OnDemand/Private/OnDemandHttpThread.h
@@ -3,7 +3,7 @@
 #pragma once
 
 #ifndef UE_ENABLE_REQUEST_THREAD
-	#define UE_ENABLE_REQUEST_THREAD 0
+	#define UE_ENABLE_REQUEST_THREAD 1
 #endif
 
 #if UE_ENABLE_REQUEST_THREAD
@@ -33,7 +33,7 @@ class FOnDemandThread final
 	: public FRunnable
 {
 public:
-	using FCompletionCallback = TUniqueFunction<void(int32 /*StatusCode*/ , FStringView/*ErrorReason*/, FIoBuffer&& /*Data*/)>;
+	using FCompletionCallback = TUniqueFunction<void(uint32 /*StatusCode*/ , FStringView/*ErrorReason*/, FIoBuffer&& /*Data*/)>;
 	using FRequestHandle = void*;
 	DECLARE_MULTICAST_DELEGATE(FOnTickIdle);
 
diff --git a/Engine/Source/Runtime/Experimental/IoStore/OnDemand/Private/OnDemandIoDispatcherBackend.cpp b/Engine/Source/Runtime/Experimental/IoStore/OnDemand/Private/OnDemandIoDispatcherBackend.cpp
--- a/Engine/Source/Runtime/Experimental/IoStore/OnDemand/Private/OnDemandIoDispatcherBackend.cpp
+++ b/Engine/Source/Runtime/Experimental/IoStore/OnDemand/Private/OnDemandIoDispatcherBackend.cpp
@@ -1083,7 +1083,7 @@ void FOnDemandIoBackend::CompleteCacheRead(FChunkRequest* ChunkRequest)
 
 #if UE_ENABLE_REQUEST_THREAD
 	ChunkRequest->HttpRequest = FOnDemandThread::Get().IssueRequest(ChunkRequest->Params.ChunkInfo, ChunkRequest->Params.ChunkRange,
-		[this, ChunkRequest](int32 StatusCode, FStringView ErrorReason, FIoBuffer&& Data)
+		[this, ChunkRequest](uint32 StatusCode, FStringView ErrorReason, FIoBuffer&& Data)
 		{
 			ChunkRequest->Chunk = MoveTemp(Data);
 			CompleteRequest(ChunkRequest);
@@ -1183,7 +1183,7 @@ bool FOnDemandIoBackend::Resolve(FIoRequestImpl* Request)
 
 #if UE_ENABLE_REQUEST_THREAD
 	ChunkRequest->HttpRequest = FOnDemandThread::Get().IssueRequest(ChunkRequest->Params.ChunkInfo, ChunkRequest->Params.ChunkRange,
-		[this, ChunkRequest](int32 StatusCode, FStringView ErrorReason, FIoBuffer&& Data)
+		[this, ChunkRequest](uint32 StatusCode, FStringView ErrorReason, FIoBuffer&& Data)
 		{
 			ChunkRequest->Chunk = MoveTemp(Data);
 			CompleteRequest(ChunkRequest);
