Safeguard NetRPC_GetDebugFunctionName from null pointer access

#rnx

diff --git a/Engine/Source/Runtime/Experimental/Iris/Core/Private/Iris/ReplicationSystem/NetBlob/NetRPC.cpp b/Engine/Source/Runtime/Experimental/Iris/Core/Private/Iris/ReplicationSystem/NetBlob/NetRPC.cpp
--- a/Engine/Source/Runtime/Experimental/Iris/Core/Private/Iris/ReplicationSystem/NetBlob/NetRPC.cpp
+++ b/Engine/Source/Runtime/Experimental/Iris/Core/Private/Iris/ReplicationSystem/NetBlob/NetRPC.cpp
@@ -748,7 +748,7 @@ static FString NetRPC_GetDebugFunctionName(const UReplicationSystem* Replication
 		return TEXT("InvalidFunctionIndex");
 	}
 
-	return Descriptor->MemberFunctionDescriptors[FunctionLocator.FunctionIndex].Function->GetName();
+	return GetNameSafe(Descriptor->MemberFunctionDescriptors[FunctionLocator.FunctionIndex].Function);
 }
 
 static FString NetRPC_GetDebugObjectRefName(FNetSerializationContext& Context, const FNetObjectReference& ObjectReference, UObject* ObjectPtr)
