Disable UE_ENABLE_REQUEST_THREAD that I once again accidently enabled.

#rnx

diff --git a/Engine/Source/Runtime/Experimental/IoStore/OnDemand/Private/OnDemandHttpThread.h b/Engine/Source/Runtime/Experimental/IoStore/OnDemand/Private/OnDemandHttpThread.h
--- a/Engine/Source/Runtime/Experimental/IoStore/OnDemand/Private/OnDemandHttpThread.h
+++ b/Engine/Source/Runtime/Experimental/IoStore/OnDemand/Private/OnDemandHttpThread.h
@@ -3,7 +3,7 @@
 #pragma once
 
 #ifndef UE_ENABLE_REQUEST_THREAD
-	#define UE_ENABLE_REQUEST_THREAD 1
+	#define UE_ENABLE_REQUEST_THREAD 0
 #endif
 
 #if UE_ENABLE_REQUEST_THREAD
