Fix EPA bug on PS5
Makde the algorithm more robust by adding an absoolute epsilon

diff --git a/Engine/Source/Runtime/Experimental/Chaos/Public/Chaos/EPAVectorized.h b/Engine/Source/Runtime/Experimental/Chaos/Public/Chaos/EPAVectorized.h
--- a/Engine/Source/Runtime/Experimental/Chaos/Public/Chaos/EPAVectorized.h
+++ b/Engine/Source/Runtime/Experimental/Chaos/Public/Chaos/EPAVectorized.h
@@ -424,7 +424,6 @@ FORCEINLINE_DEBUGGABLE void VectorComputeEPAResults(const VectorRegister4Float*
 template <typename TSupportA, typename TSupportB>
 EEPAResult VectorEPA(TArray<VectorRegister4Float>& VertsABuffer, TArray<VectorRegister4Float>& VertsBBuffer, const TSupportA& SupportA, const TSupportB& SupportB, VectorRegister4Float& OutPenetration, VectorRegister4Float& OutDir, VectorRegister4Float& WitnessA, VectorRegister4Float& WitnessB)
 {
-	constexpr VectorRegister4Float EpsRel = MakeVectorRegisterFloatConstant(1.e-2f, 1.e-2f, 1.e-2f, 1.e-2f);
 	const VectorRegister4Float OriginInsideEps = VectorZero();
 
 	TEPAWorkingArray<VectorTEPAEntry> Entries;
@@ -511,9 +510,13 @@ EEPAResult VectorEPA(TArray<VectorRegister4Float>& VertsABuffer, TArray<VectorRe
 		UpperBound = VectorMin(DistanceToSupportPlane, UpperBound);
 		LowerBound = Entry.Distance;
 
+		// Add one absolute epsilon to be more robust with small values, fixed some issue with small UpperBound
+		constexpr VectorRegister4Float SmallFloatEps = MakeVectorRegisterFloatConstant(1.e-5f, 1.e-5f, 1.e-5f, 1.e-5f);
+		// Use a relative epsilon to have a robust check in the main case
+		constexpr VectorRegister4Float OnePlusRelativeEps = MakeVectorRegisterFloatConstant(1.01f, 1.01f, 1.01f, 1.01f);
 		// It's possible the origin is not contained by the CSO, probably because of numerical error. 
 		// In this case the upper bound will be negative, at which point we should just exit. 
-		const VectorRegister4Float UpperBoundTolerance = VectorMultiply(VectorAdd(VectorOneFloat(), EpsRel), VectorAbs(LowerBound));
+		const VectorRegister4Float UpperBoundTolerance = VectorAdd(VectorMultiply(OnePlusRelativeEps, VectorAbs(LowerBound)), SmallFloatEps);
 		const VectorRegister4Float UpperBoundLEUpperBoundTolerance = VectorCompareLE(UpperBound, UpperBoundTolerance);
 		if (VectorMaskBits(UpperBoundLEUpperBoundTolerance))
 		{
