Dataflow : Add GetArrayElement and GeArraySize generic nodes
- Note: not yet supported for arrays of structures

diff --git a/Engine/Plugins/Experimental/GeometryCollectionPlugin/Source/GeometryCollectionNodes/Private/Dataflow/GeometryCollectionArrayNodes.cpp b/Engine/Plugins/Experimental/GeometryCollectionPlugin/Source/GeometryCollectionNodes/Private/Dataflow/GeometryCollectionArrayNodes.cpp
--- a/Engine/Plugins/Experimental/GeometryCollectionPlugin/Source/GeometryCollectionNodes/Private/Dataflow/GeometryCollectionArrayNodes.cpp
+++ b/Engine/Plugins/Experimental/GeometryCollectionPlugin/Source/GeometryCollectionNodes/Private/Dataflow/GeometryCollectionArrayNodes.cpp
@@ -23,6 +23,9 @@ namespace UE::Dataflow
 		DATAFLOW_NODE_REGISTER_CREATION_FACTORY(FRemoveFloatArrayElementDataflowNode);
 		DATAFLOW_NODE_REGISTER_CREATION_FACTORY(FFloatArrayComputeStatisticsDataflowNode);
 		DATAFLOW_NODE_REGISTER_CREATION_FACTORY(FRandomizeFloatArrayDataflowNode);
+		DATAFLOW_NODE_REGISTER_CREATION_FACTORY(FDataflowGetArraySizeNode);
+		DATAFLOW_NODE_REGISTER_CREATION_FACTORY(FDataflowGetArrayElementNode);
+		
 	}
 }
 
@@ -610,3 +613,63 @@ void FRandomizeFloatArrayDataflowNode::Evaluate(UE::Dataflow::FContext& Context,
 		SetValue(Context, MoveTemp(OutFloatArray), &FloatArray);
 	}
 }
+
+///////////////////////////////////////////////////////////////////////////////////////////////
+
+FDataflowGetArraySizeNode::FDataflowGetArraySizeNode(const UE::Dataflow::FNodeParameters& InParam, FGuid InGuid)
+	: FDataflowNode(InParam, InGuid)
+{
+	RegisterInputConnection(&Array);
+	RegisterOutputConnection(&Size);
+}
+
+void FDataflowGetArraySizeNode::Evaluate(UE::Dataflow::FContext& Context, const FDataflowOutput* Out) const
+{
+	if (Out->IsA(&Size))
+	{
+		SetArraySizeFromInput(Context, &Array, &Size);
+	}
+}
+
+///////////////////////////////////////////////////////////////////////////////////////////////
+
+FDataflowGetArrayElementNode::FDataflowGetArrayElementNode(const UE::Dataflow::FNodeParameters& InParam, FGuid InGuid)
+	: FDataflowNode(InParam, InGuid)
+{
+	static FName ArrayElementDependencyGroup("ArrayElement");
+	
+	RegisterInputConnection(&Array)
+		.SetTypeDependencyGroup(ArrayElementDependencyGroup);
+	RegisterInputConnection(&Index);
+	RegisterOutputConnection(&Array, &Array)
+		.SetTypeDependencyGroup(ArrayElementDependencyGroup);
+	RegisterOutputConnection(&Element)
+		.SetTypeDependencyGroup(ArrayElementDependencyGroup);
+}
+
+FName FDataflowGetArrayElementNode::GetDependentConnectionType(const FDataflowConnection& SourceConnection, const FDataflowConnection& DependentConnection) const
+{
+	if (SourceConnection.IsA(&Array) && DependentConnection.IsA(&Element))
+	{
+		return FDataflowArrayTypePolicy::GetElementType(SourceConnection.GetType());
+	}
+	else if (SourceConnection.IsA(&Element) && DependentConnection.IsA(&Array))
+	{
+		return FDataflowArrayTypePolicy::GetArrayType(SourceConnection.GetType());
+	}
+	// default for all the rest 
+	return FDataflowNode::GetDependentConnectionType(SourceConnection, DependentConnection);
+}
+
+void FDataflowGetArrayElementNode::Evaluate(UE::Dataflow::FContext& Context, const FDataflowOutput* Out) const
+{
+	if (Out->IsA(&Array))
+	{
+		ForwardInput(Context, &Array, &Array);
+	}
+	else if (Out->IsA(&Element))
+	{
+		const int32 InIndex = GetValue(Context, &Index);
+		SetArrayElementFromInput(Context, &Array, InIndex, &Element);
+	}
+}
\ No newline at end of file
diff --git a/Engine/Plugins/Experimental/GeometryCollectionPlugin/Source/GeometryCollectionNodes/Public/Dataflow/GeometryCollectionArrayNodes.h b/Engine/Plugins/Experimental/GeometryCollectionPlugin/Source/GeometryCollectionNodes/Public/Dataflow/GeometryCollectionArrayNodes.h
--- a/Engine/Plugins/Experimental/GeometryCollectionPlugin/Source/GeometryCollectionNodes/Public/Dataflow/GeometryCollectionArrayNodes.h
+++ b/Engine/Plugins/Experimental/GeometryCollectionPlugin/Source/GeometryCollectionNodes/Public/Dataflow/GeometryCollectionArrayNodes.h
@@ -533,6 +533,60 @@ public:
 	virtual void Evaluate(UE::Dataflow::FContext& Context, const FDataflowOutput* Out) const override;
 };
 
+/**
+ * Ge size of an array
+ */
+USTRUCT()
+struct FDataflowGetArraySizeNode : public FDataflowNode
+{
+	GENERATED_USTRUCT_BODY()
+	DATAFLOW_NODE_DEFINE_INTERNAL(FDataflowGetArraySizeNode, "GetArraySize", "Utilities|Array", "")
+
+private:
+	/** Array to get size from */
+	UPROPERTY(meta = (DataflowInput, DataflowIntrinsic))
+	FDataflowArrayTypes Array;
+
+	/** Computed value */
+	UPROPERTY(meta = (DataflowOutput))
+	int32 Size = 0;
+
+	virtual void Evaluate(UE::Dataflow::FContext& Context, const FDataflowOutput* Out) const override;
+
+public:
+	FDataflowGetArraySizeNode(const UE::Dataflow::FNodeParameters& InParam, FGuid InGuid = FGuid::NewGuid());
+};
+
+/**
+ * Get an element from an array
+ */
+USTRUCT()
+struct FDataflowGetArrayElementNode : public FDataflowNode
+{
+	GENERATED_USTRUCT_BODY()
+	DATAFLOW_NODE_DEFINE_INTERNAL(FDataflowGetArrayElementNode, "GetArrayElement", "Utilities|Array", "")
+
+private:
+	/** Array to get size from */
+	UPROPERTY(meta = (DataflowInput, DataflowOutput, DataflowPassthrough = "Array", DataflowIntrinsic))
+	FDataflowArrayTypes Array;
+
+	/** Array to get size from */
+	UPROPERTY(EditAnywhere, Category = "Array", meta = (DataflowInput))
+	int32 Index = 0;
+
+	/** Element from the array at the specified index */
+	UPROPERTY(meta = (DataflowOutput))
+	FDataflowAllTypes Element;
+
+	virtual void Evaluate(UE::Dataflow::FContext& Context, const FDataflowOutput* Out) const override;
+
+	virtual FName GetDependentConnectionType(const FDataflowConnection& SourceConnection, const FDataflowConnection& DependentConnection) const override;
+
+public:
+	FDataflowGetArrayElementNode(const UE::Dataflow::FNodeParameters& InParam, FGuid InGuid = FGuid::NewGuid());
+};
+
 namespace UE::Dataflow
 {
 	void GeometryCollectionArrayNodes();
diff --git a/Engine/Source/Runtime/Experimental/Dataflow/Core/Private/Dataflow/DataflowAnyType.cpp b/Engine/Source/Runtime/Experimental/Dataflow/Core/Private/Dataflow/DataflowAnyType.cpp
--- a/Engine/Source/Runtime/Experimental/Dataflow/Core/Private/Dataflow/DataflowAnyType.cpp
+++ b/Engine/Source/Runtime/Experimental/Dataflow/Core/Private/Dataflow/DataflowAnyType.cpp
@@ -9,6 +9,7 @@ namespace UE::Dataflow
 	{
 		UE_DATAFLOW_REGISTER_ANYTYPE(FDataflowAnyType);
 		UE_DATAFLOW_REGISTER_ANYTYPE(FDataflowAllTypes);
+		UE_DATAFLOW_REGISTER_ANYTYPE(FDataflowArrayTypes);
 		UE_DATAFLOW_REGISTER_ANYTYPE(FDataflowNumericTypes);
 		UE_DATAFLOW_REGISTER_ANYTYPE(FDataflowVectorTypes);
 		UE_DATAFLOW_REGISTER_ANYTYPE(FDataflowStringTypes);
diff --git a/Engine/Source/Runtime/Experimental/Dataflow/Core/Private/Dataflow/DataflowNode.cpp b/Engine/Source/Runtime/Experimental/Dataflow/Core/Private/Dataflow/DataflowNode.cpp
--- a/Engine/Source/Runtime/Experimental/Dataflow/Core/Private/Dataflow/DataflowNode.cpp
+++ b/Engine/Source/Runtime/Experimental/Dataflow/Core/Private/Dataflow/DataflowNode.cpp
@@ -1443,6 +1443,45 @@ void FDataflowNode::ForwardInput(UE::Dataflow::FContext& Context, const UE::Data
 	}
 }
 
+void FDataflowNode::SetArraySizeFromInput(UE::Dataflow::FContext& Context, const UE::Dataflow::FConnectionReference& InputReference, const int32* OutputReference) const
+{
+	int32 OutSize = 0;
+	if (const FDataflowInput* ArrayInput = FindInput(InputReference))
+	{
+		if (ArrayInput->IsConnected())
+		{
+			// Pull the value to populate the cache
+			ArrayInput->PullValue(Context);
+
+			OutSize = Context.GetArraySizeFromData(ArrayInput->GetConnection()->CacheKey());
+		}
+	}
+	SetValue(Context, OutSize, OutputReference);
+}
+
+void FDataflowNode::SetArrayElementFromInput(UE::Dataflow::FContext& Context, const UE::Dataflow::FConnectionReference& InputReference, int32 Index, const void* OutputReference) const
+{
+	if (const FDataflowOutput* Output = FindOutput(OutputReference))
+	{
+		if (const FDataflowInput* ArrayInput = FindInput(InputReference))
+		{
+			if (ArrayInput->IsConnected())
+			{
+				// Pull the value to populate the cache
+				ArrayInput->PullValue(Context);
+
+				const UE::Dataflow::FContextCacheKey ArrayKey = ArrayInput->GetConnection()->CacheKey();
+				const UE::Dataflow::FContextCacheKey ElementKey = Output->CacheKey();
+
+				Context.SetArrayElementFromData(ArrayKey, Index, ElementKey, Output->GetProperty(), Output->GetOwningNodeGuid(), Output->GetOwningNodeValueHash(), Output->GetOwningNodeTimestamp());
+				return;
+			}
+		}
+		// at least return a default value 
+		Output->SetNullValue(Context);
+	}
+}
+
 bool FDataflowNode::TrySetConnectionType(FDataflowConnection* Connection, FName NewType)
 {
 	if (Connection)
@@ -1478,6 +1517,12 @@ PRAGMA_ENABLE_DEPRECATION_WARNINGS
 	}
 }
 
+FName FDataflowNode::GetDependentConnectionType(const FDataflowConnection& SourceConnection, const FDataflowConnection& DependentConnection) const
+{
+	// by default return the same type 
+	return SourceConnection.GetType();
+}
+
 bool FDataflowNode::OnConnectionTypeChanged(const FDataflowConnection& Connection)
 {
 	// detect if any of the dependent connections is connected
@@ -1496,12 +1541,12 @@ bool FDataflowNode::OnConnectionTypeChanged(const FDataflowConnection& Connectio
 
 	// Change all connection that share the same type dependecy group
 	bool bTypeChanged = false;
-	const FName NewType = Connection.GetType();
 	ForEachConnection(
-		[this, &bTypeChanged, &DependencyGroup, &NewType](FDataflowConnection* OtherConnection)
+		[this, &Connection, &bTypeChanged, &DependencyGroup](FDataflowConnection* OtherConnection)
 		{
 			if (OtherConnection && (OtherConnection->GetTypeDependencyGroup() == DependencyGroup))
 			{
+				const FName NewType = GetDependentConnectionType(Connection, *OtherConnection);
 				SetConnectionConcreteType(OtherConnection, NewType);
 				bTypeChanged = true;
 			}
diff --git a/Engine/Source/Runtime/Experimental/Dataflow/Core/Private/Dataflow/DataflowNodeParameters.cpp b/Engine/Source/Runtime/Experimental/Dataflow/Core/Private/Dataflow/DataflowNodeParameters.cpp
--- a/Engine/Source/Runtime/Experimental/Dataflow/Core/Private/Dataflow/DataflowNodeParameters.cpp
+++ b/Engine/Source/Runtime/Experimental/Dataflow/Core/Private/Dataflow/DataflowNodeParameters.cpp
@@ -385,6 +385,46 @@ namespace UE::Dataflow
 #endif // EDITOR_ONLY
 	}
 
+	int32 FContext::GetArraySizeFromData(const FContextCacheKey InKey) const
+	{
+		// Query the size without knowning the type ( the cache elements do know it )
+		const TUniquePtr<UE::Dataflow::FContextCacheElementBase>* CacheEntryPr = GetDataImpl(InKey);
+		if (CacheEntryPr && CacheEntryPr->IsValid())
+		{
+			return (*CacheEntryPr)->GetNumArrayElements(*this);
+		}
+		return 0;
+	}
+
+	void FContext::SetArrayElementFromData(const FContextCacheKey InArrayKey, int32 Index, const FContextCacheKey InElementKey, const FProperty* InProperty, const FGuid& InNodeGuid, uint32 InNodeHash, const FTimestamp& InTimestamp)
+	{
+		// get the element from the cache 
+		const TUniquePtr<UE::Dataflow::FContextCacheElementBase>* CacheEntryPr = GetDataImpl(InArrayKey);
+		if (CacheEntryPr && CacheEntryPr->IsValid())
+		{
+			TUniquePtr<UE::Dataflow::FContextCacheElementBase> ArrayElementCacheEntry = (*CacheEntryPr)->CreateFromArrayElement(*this, Index, InProperty, InNodeGuid, InNodeHash, InTimestamp);
+			if (ArrayElementCacheEntry.IsValid())
+			{
+				SetDataImpl(InElementKey, MoveTemp(ArrayElementCacheEntry));
+				return;
+			}
+		}
+		// fallback 
+		SetNullData(InElementKey, InProperty, InNodeGuid, InNodeHash, InTimestamp);
+	}
+
+	bool FContext::CopyDataToAnotherContext(const FContextCacheKey InSourceKey, FContext& TargetContext, const FContextCacheKey InTargetKey) const
+	{
+		const TUniquePtr<UE::Dataflow::FContextCacheElementBase>* CacheEntryToClone = GetDataImpl(InSourceKey);
+		if (CacheEntryToClone && CacheEntryToClone->IsValid())
+		{
+			TUniquePtr<UE::Dataflow::FContextCacheElementBase> ClonedCacheEntry = (*CacheEntryToClone)->Clone(*this);
+			TargetContext.SetDataImpl(InTargetKey, MoveTemp(ClonedCacheEntry));
+			return true;
+		}
+		return false;
+	}
+
 	void FContext::BeginContextEvaluation(const FDataflowNode * Node, const FDataflowOutput * Output)
 	{
 		if (Output)
diff --git a/Engine/Source/Runtime/Experimental/Dataflow/Core/Public/Dataflow/DataflowAnyType.h b/Engine/Source/Runtime/Experimental/Dataflow/Core/Public/Dataflow/DataflowAnyType.h
--- a/Engine/Source/Runtime/Experimental/Dataflow/Core/Public/Dataflow/DataflowAnyType.h
+++ b/Engine/Source/Runtime/Experimental/Dataflow/Core/Public/Dataflow/DataflowAnyType.h
@@ -27,6 +27,16 @@ struct FDataflowAllTypes : public FDataflowAnyType
 	GENERATED_BODY()
 };
 
+
+/** Generic Array types */
+USTRUCT()
+struct FDataflowArrayTypes : public FDataflowAnyType
+{
+	using FPolicyType = FDataflowArrayTypePolicy;
+	using FStorageType = void;
+	GENERATED_BODY()
+};
+
 /**
 * Numeric types
 * (double, float, int64, uint64, int32, uint32, int16, uint16, int8, uint8)"
diff --git a/Engine/Source/Runtime/Experimental/Dataflow/Core/Public/Dataflow/DataflowContextCache.h b/Engine/Source/Runtime/Experimental/Dataflow/Core/Public/Dataflow/DataflowContextCache.h
--- a/Engine/Source/Runtime/Experimental/Dataflow/Core/Public/Dataflow/DataflowContextCache.h
+++ b/Engine/Source/Runtime/Experimental/Dataflow/Core/Public/Dataflow/DataflowContextCache.h
@@ -13,6 +13,7 @@ namespace UE::Dataflow
 {
 	class FContext;
 	struct FContextCacheElementBase;
+	struct FContextCacheElementNull;
 
 	typedef uint32 FContextCacheKey;
 
@@ -88,6 +89,9 @@ namespace UE::Dataflow
 
 		virtual const void* GetUntypedData(const IContextCacheStore& Context, const FProperty* PropertyIn) const = 0;
 
+		virtual int32 GetNumArrayElements(const IContextCacheStore& Context) const = 0;
+		virtual TUniquePtr<FContextCacheElementBase> CreateFromArrayElement(const IContextCacheStore& Context, int32 Index, const FProperty* InProperty, const FGuid& InNodeGuid, uint32 InNodeHash, const FTimestamp& InTimestamp) const = 0;
+
 		EType GetType() const {	return Type; }
 
 		const FProperty* GetProperty() const { return Property; }
@@ -133,6 +137,34 @@ namespace UE::Dataflow
 			return &Data;
 		}
 
+		virtual int32 GetNumArrayElements(const IContextCacheStore& Context) const override
+		{
+			if constexpr (TIsTArray<FDataType>::Value)
+			{
+				return Data.Num();
+			}
+			else
+			{
+				return 0;
+			}
+		}
+
+		virtual TUniquePtr<FContextCacheElementBase> CreateFromArrayElement(const IContextCacheStore& Context, int32 Index, const FProperty* InProperty, const FGuid& InNodeGuid, uint32 InNodeHash, const FTimestamp& InTimestamp) const override
+		{
+			if constexpr (TIsTArray<FDataType>::Value)
+			{
+				if (Data.IsValidIndex(Index))
+				{
+					return MakeUnique<TContextCacheElement<decltype(Data[Index])>>(InNodeGuid, InProperty, Data[Index], InNodeHash, InTimestamp);
+				}
+				return {};
+			}
+			else
+			{
+				return {};
+			}
+		}
+
 		const T& GetDataDirect() const { return Data; }
 
 		inline virtual TUniquePtr<FContextCacheElementBase> Clone(const IContextCacheStore& Context) const override;
@@ -155,7 +187,11 @@ namespace UE::Dataflow
 		template<class T>
 		inline const T& GetData(const IContextCacheStore& Context, const FProperty* PropertyIn, const T& Default) const;
 
-		inline virtual const void* GetUntypedData(const IContextCacheStore& Context, const FProperty* PropertyIn) const;
+		inline virtual const void* GetUntypedData(const IContextCacheStore& Context, const FProperty* PropertyIn) const override;
+
+		inline virtual int32 GetNumArrayElements(const IContextCacheStore& Context) const override;
+
+		inline virtual TUniquePtr<FContextCacheElementBase> CreateFromArrayElement(const IContextCacheStore& Context, int32 Index, const FProperty* InProperty, const FGuid& InNodeGuid, uint32 InNodeHash, const FTimestamp& InTimestamp) const override;
 
 		inline virtual TUniquePtr<FContextCacheElementBase> Clone(const IContextCacheStore& Context) const override;
 
@@ -172,15 +208,30 @@ namespace UE::Dataflow
 		// IMPORTANT: 
 		// Timestamp must be set to (Timestamp.Value - 1) to make sure that this type of entry is always invalid
 		//
+		UE_DEPRECATED(5.6, "Use the other constructor that does not pass a DataKey (the key is not needed)")
 		FContextCacheElementNull(FGuid InNodeGuid, const FProperty* InProperty, FContextCacheKey InDataKey, uint32 InNodeHash, FTimestamp Timestamp)
 			: FContextCacheElementBase(EType::CacheElementNull, InNodeGuid, InProperty, InNodeHash, FTimestamp((Timestamp.Value == 0)? 0: (Timestamp.Value - 1)))
 		{}
 
-		virtual const void* GetUntypedData(const IContextCacheStore& Context, const FProperty* PropertyIn) const
+		FContextCacheElementNull(FGuid InNodeGuid, const FProperty* InProperty, uint32 InNodeHash, FTimestamp Timestamp)
+			: FContextCacheElementBase(EType::CacheElementNull, InNodeGuid, InProperty, InNodeHash, FTimestamp((Timestamp.Value == 0) ? 0 : (Timestamp.Value - 1)))
+		{}
+
+		virtual const void* GetUntypedData(const IContextCacheStore& Context, const FProperty* PropertyIn) const override
 		{
 			return nullptr;
 		}
 
+		virtual int32 GetNumArrayElements(const IContextCacheStore& Context) const override
+		{
+			return 0;
+		}
+
+		virtual TUniquePtr<FContextCacheElementBase> CreateFromArrayElement(const IContextCacheStore& Context, int32 Index, const FProperty* InProperty, const FGuid& InNodeGuid, uint32 InNodeHash, const FTimestamp& InTimestamp) const override
+		{
+			return {};
+		}
+
 		inline virtual TUniquePtr<FContextCacheElementBase> Clone(const IContextCacheStore& Context) const override;
 	};
 
@@ -205,11 +256,21 @@ namespace UE::Dataflow
 			return ObjectPtr;
 		}
 
-		virtual const void* GetUntypedData(const IContextCacheStore& /*Context*/, const FProperty* /*PropertyIn*/) const
+		virtual const void* GetUntypedData(const IContextCacheStore& /*Context*/, const FProperty* /*PropertyIn*/) const override
 		{
 			return (void*)ObjectPtr;
 		}
 
+		virtual int32 GetNumArrayElements(const IContextCacheStore& Context) const override
+		{
+			return 0;
+		}
+
+		virtual TUniquePtr<FContextCacheElementBase> CreateFromArrayElement(const IContextCacheStore& Context, int32 Index, const FProperty* InProperty, const FGuid& InNodeGuid, uint32 InNodeHash, const FTimestamp& InTimestamp) const override
+		{
+			return {};
+		}
+
 		inline virtual TUniquePtr<FContextCacheElementBase> Clone(const IContextCacheStore& Context) const override;
 
 		//~ Begin FGCObject interface
@@ -245,11 +306,26 @@ namespace UE::Dataflow
 			return ObjectPtrArray;
 		}
 
-		virtual const void* GetUntypedData(const IContextCacheStore& /*Context*/, const FProperty* /*PropertyIn*/) const
+		virtual const void* GetUntypedData(const IContextCacheStore& /*Context*/, const FProperty* /*PropertyIn*/) const override
 		{
 			return &ObjectPtrArray;
 		}
 
+		virtual int32 GetNumArrayElements(const IContextCacheStore& Context) const override
+		{
+			return ObjectPtrArray.Num();
+		}
+
+		virtual TUniquePtr<FContextCacheElementBase> CreateFromArrayElement(const IContextCacheStore& Context, int32 Index, const FProperty* InProperty, const FGuid& InNodeGuid, uint32 InNodeHash, const FTimestamp& InTimestamp) const override
+		{
+			if (ObjectPtrArray.IsValidIndex(Index))
+			{
+				FDataType Element = ObjectPtrArray[Index];
+				return MakeUnique<TContextCacheElementUObject<FDataType>>(InNodeGuid, InProperty, MoveTemp(Element), InNodeHash, InTimestamp);
+			}
+			return {};
+		}
+
 		inline virtual TUniquePtr<FContextCacheElementBase> Clone(const IContextCacheStore& Context) const override;
 
 		//~ Begin FGCObject interface
@@ -290,11 +366,21 @@ namespace UE::Dataflow
 			return InstancedStruct.Get<T>();
 		}
 
-		virtual const void* GetUntypedData(const IContextCacheStore& /*Context*/, const FProperty* /*PropertyIn*/) const
+		virtual const void* GetUntypedData(const IContextCacheStore& /*Context*/, const FProperty* /*PropertyIn*/) const override
 		{
 			return InstancedStruct.GetMemory();
 		}
 
+		virtual int32 GetNumArrayElements(const IContextCacheStore& Context) const override
+		{
+			return 0;
+		}
+
+		virtual TUniquePtr<FContextCacheElementBase> CreateFromArrayElement(const IContextCacheStore& Context, int32 Index, const FProperty* InProperty, const FGuid& InNodeGuid, uint32 InNodeHash, const FTimestamp& InTimestamp) const override
+		{
+			return {};
+		}
+
 		inline virtual TUniquePtr<FContextCacheElementBase> Clone(const IContextCacheStore& Context) const override
 		{
 			return MakeUnique<FContextCacheElementUStruct>(*this);
@@ -334,11 +420,27 @@ namespace UE::Dataflow
 			return InstancedStructArray.Get<T>();
 		}
 
-		virtual const void* GetUntypedData(const IContextCacheStore& /*Context*/, const FProperty* /*PropertyIn*/) const
+		virtual const void* GetUntypedData(const IContextCacheStore& /*Context*/, const FProperty* /*PropertyIn*/) const override
 		{
 			return InstancedStructArray.GetMemory();
 		}
 
+		virtual int32 GetNumArrayElements(const IContextCacheStore& Context) const override
+		{
+			return InstancedStructArray.Num();
+		}
+
+		virtual TUniquePtr<FContextCacheElementBase> CreateFromArrayElement(const IContextCacheStore& Context, int32 Index, const FProperty* InProperty, const FGuid& InNodeGuid, uint32 InNodeHash, const FTimestamp& InTimestamp) const override
+		{
+			if (InstancedStructArray.IsValidIndex(Index))
+			{
+				//return MakeUnique<FContextCacheElementUStruct>(GetNodeGuid(), GetProperty(), Forward<T>(ObjectPtrArray[Index]), GetNodeHash(), GetTimestamp());
+				// todo : support extracting structure from array of structures
+				return {};
+			}
+			return {};
+		}
+
 		inline virtual TUniquePtr<FContextCacheElementBase> Clone(const IContextCacheStore& Context) const override
 		{
 			return MakeUnique<FContextCacheElementUStructArray>(*this);
@@ -507,6 +609,24 @@ namespace UE::Dataflow
 		return nullptr;
 	}
 
+	int32 FContextCacheElementReference::GetNumArrayElements(const IContextCacheStore& Context) const
+	{
+		if (const TUniquePtr<FContextCacheElementBase>* Cache = Context.FindCacheElement(DataKey))
+		{
+			return (*Cache)->GetNumArrayElements(Context);
+		}
+		return 0;
+	}
+
+	TUniquePtr<FContextCacheElementBase> FContextCacheElementReference::CreateFromArrayElement(const IContextCacheStore& Context, int32 Index, const FProperty* InProperty, const FGuid& InNodeGuid, uint32 InNodeHash, const FTimestamp& InTimestamp) const
+	{
+		if (const TUniquePtr<FContextCacheElementBase>* Cache = Context.FindCacheElement(DataKey))
+		{
+			return (*Cache)->CreateFromArrayElement(Context, Index, InProperty, InNodeGuid, InNodeHash, InTimestamp);
+		}
+		return 0;
+	}
+
 	TUniquePtr<FContextCacheElementBase> FContextCacheElementReference::Clone(const IContextCacheStore& Context) const
 	{
 		if (const TUniquePtr<FContextCacheElementBase>* ReferencedCacheEntry = Context.FindCacheElement(DataKey))
@@ -516,12 +636,12 @@ namespace UE::Dataflow
 				return (*ReferencedCacheEntry)->Clone(Context);
 			}
 		}
-		return MakeUnique<FContextCacheElementNull>(GetNodeGuid(), GetProperty(), FContextCacheKey(), GetNodeHash(), GetTimestamp());
+		return MakeUnique<FContextCacheElementNull>(GetNodeGuid(), GetProperty(), GetNodeHash(), GetTimestamp());
 	}
 
 	TUniquePtr<FContextCacheElementBase> FContextCacheElementNull::Clone(const IContextCacheStore& Context) const
 	{
-		return MakeUnique<FContextCacheElementNull>(GetNodeGuid(), GetProperty(), FContextCacheKey(), GetNodeHash(), GetTimestamp());
+		return MakeUnique<FContextCacheElementNull>(GetNodeGuid(), GetProperty(), GetNodeHash(), GetTimestamp());
 	}
 
 	template<class T>
diff --git a/Engine/Source/Runtime/Experimental/Dataflow/Core/Public/Dataflow/DataflowNode.h b/Engine/Source/Runtime/Experimental/Dataflow/Core/Public/Dataflow/DataflowNode.h
--- a/Engine/Source/Runtime/Experimental/Dataflow/Core/Public/Dataflow/DataflowNode.h
+++ b/Engine/Source/Runtime/Experimental/Dataflow/Core/Public/Dataflow/DataflowNode.h
@@ -614,6 +614,20 @@ PRAGMA_ENABLE_DEPRECATION_WARNINGS
 		}
 	}
 
+	/**
+	* This method infers the array size from an input and set it on the specified output.
+	* When called, the input will be pulled as a result.
+	* If the input is not an array, the value set on the output will be 0
+	*/
+	DATAFLOWCORE_API void SetArraySizeFromInput(UE::Dataflow::FContext& Context, const UE::Dataflow::FConnectionReference& InputReference, const int32* OutputReference)  const;
+
+	/**
+	* This method extracts an element from the specified input at a specific index and set it as output to the specified output.
+	* When called, the input will be pulled as a result.
+	* If the input is not an array or if the index is out of bounds, a null value will be set on the output (which will result in a default value when read from) 
+	*/
+	DATAFLOWCORE_API void SetArrayElementFromInput(UE::Dataflow::FContext& Context, const UE::Dataflow::FConnectionReference& InputReference, int32 Index, const void* OutputReference) const;
+
 	/*
 	* Go through all connections and execute a function
 	* Signature of the function is bool Function(FDataflowConnection* Connection)
@@ -754,6 +768,12 @@ PRAGMA_ENABLE_DEPRECATION_WARNINGS
 protected:
 	bool OnConnectionTypeChanged(const FDataflowConnection& Connection);
 
+	/** 
+	* When enforcing type dependency returns the type of the target connection from a source one being changed
+	* the default implementation returns the same type 
+	*/
+	DATAFLOWCORE_API virtual FName GetDependentConnectionType(const FDataflowConnection& SourceConnection, const FDataflowConnection& DependentConnection) const;
+
 	UE_DEPRECATED(5.6, "Use inputs and outputs SetTypeDependencyGroup() to declare dependent types")
 	virtual bool OnInputTypeChanged(const FDataflowInput* Input) { return false; }
 
diff --git a/Engine/Source/Runtime/Experimental/Dataflow/Core/Public/Dataflow/DataflowNodeParameters.h b/Engine/Source/Runtime/Experimental/Dataflow/Core/Public/Dataflow/DataflowNodeParameters.h
--- a/Engine/Source/Runtime/Experimental/Dataflow/Core/Public/Dataflow/DataflowNodeParameters.h
+++ b/Engine/Source/Runtime/Experimental/Dataflow/Core/Public/Dataflow/DataflowNodeParameters.h
@@ -251,10 +251,16 @@ namespace UE::Dataflow
 		// in that case posting an invalid reference, will allow the evaluatino to go through and the node reading it will get a default value instead
 		void SetNullData(FContextCacheKey InKey, const FProperty* InProperty, const FGuid& InNodeGuid, uint32 InNodeHash, const FTimestamp& InTimestamp)
 		{
-			TUniquePtr<FContextCacheElementNull> CacheNullElement = MakeUnique<FContextCacheElementNull>(InNodeGuid, InProperty, InKey, InNodeHash, InTimestamp);
+			TUniquePtr<FContextCacheElementNull> CacheNullElement = MakeUnique<FContextCacheElementNull>(InNodeGuid, InProperty, InNodeHash, InTimestamp);
 			SetDataImpl(InKey, MoveTemp(CacheNullElement));
 		}
 
+		DATAFLOWCORE_API int32 GetArraySizeFromData(const FContextCacheKey InKey) const;
+
+		DATAFLOWCORE_API void SetArrayElementFromData(const FContextCacheKey InArrayKey, int32 Index, const FContextCacheKey InElementKey, const FProperty* InProperty, const FGuid& InNodeGuid, uint32 InNodeHash, const FTimestamp& InTimestamp);
+
+		DATAFLOWCORE_API bool CopyDataToAnotherContext(const FContextCacheKey InSourceKey, FContext& TargetContext, const FContextCacheKey InTargetKey) const;
+
 		template<class T>
 		const T& GetData(FContextCacheKey Key, const FProperty* InProperty, const T& Default = T()) const
 		{
diff --git a/Engine/Source/Runtime/Experimental/Dataflow/Core/Public/Dataflow/DataflowTypePolicy.h b/Engine/Source/Runtime/Experimental/Dataflow/Core/Public/Dataflow/DataflowTypePolicy.h
--- a/Engine/Source/Runtime/Experimental/Dataflow/Core/Public/Dataflow/DataflowTypePolicy.h
+++ b/Engine/Source/Runtime/Experimental/Dataflow/Core/Public/Dataflow/DataflowTypePolicy.h
@@ -218,6 +218,46 @@ struct TDataflowMultiTypePolicy<T, TTypes...>: public TDataflowMultiTypePolicy<T
 	}
 };
 
+struct FDataflowArrayTypePolicy: public IDataflowTypePolicy
+{
+	static constexpr TCHAR ArrayPrefix[] = TEXT("TArray<");
+	static constexpr TCHAR ArrayFormat[] = TEXT("TArray<{0}>");
+
+	virtual bool SupportsType(FName InType) const override
+	{
+		return FDataflowArrayTypePolicy::SupportsTypeStatic(InType);
+	}
+
+	static bool SupportsTypeStatic(FName InType)
+	{
+		return (InType.ToString().StartsWith(ArrayPrefix));
+	}
+
+	static IDataflowTypePolicy* GetInterface()
+	{
+		static FDataflowArrayTypePolicy Instance;
+		return &Instance;
+	}
+
+	static FName GetElementType(FName InType)
+	{
+		if (SupportsTypeStatic(InType))
+		{
+			const int32 ArrayPrefixLen = FCString::Strlen(ArrayPrefix);
+			const FString RightStr = InType.ToString().RightChop(ArrayPrefixLen);
+			check(!RightStr.IsEmpty() && RightStr[RightStr.Len() - 1] == '>');
+			return FName(RightStr.LeftChop(1)); // remove the last ">"
+		}
+		// not really an array just return the original type
+		return InType;
+	}
+
+	static FName GetArrayType(FName InType)
+	{
+		return FName(FString::Format(ArrayFormat, { InType.ToString() }));
+	}
+};
+
 struct FDataflowNumericTypePolicy : 
 	public TDataflowMultiTypePolicy<double, float, int64, uint64, int32, uint32, int16, uint16, int8, uint8>
 {
