Iris
* Added PostTickDispatch callback in NetObjectFactory

SceneGraph
* Correctly destroy entities on client when the server sends a destruction message
* Now queuing all destroy entities request and doing them simultaneously after processing all incoming packets in PostTickDispatch

#rnx

diff --git a/Engine/Source/Runtime/Engine/Private/NetDriver.cpp b/Engine/Source/Runtime/Engine/Private/NetDriver.cpp
--- a/Engine/Source/Runtime/Engine/Private/NetDriver.cpp
+++ b/Engine/Source/Runtime/Engine/Private/NetDriver.cpp
@@ -7463,7 +7463,13 @@ void UNetDriver::PostDispatchSendUpdate()
 			return;
 		}
 
-		ReplicationSystem->PreSendUpdate(UReplicationSystem::FSendUpdateParams { .SendPass = UE::Net::EReplicationSystemSendPass::PostTickDispatch });
+		const UReplicationSystem::FSendUpdateParams Params
+		{ 
+			.SendPass = UE::Net::EReplicationSystemSendPass::PostTickDispatch 
+			//no delta time passed for PostTickDispatch
+		};
+
+		ReplicationSystem->PreSendUpdate(Params);
 
 		ReplicationSystem->SendUpdate([this](TArrayView<uint32> ConnectionsToSend)
 		{
diff --git a/Engine/Source/Runtime/Experimental/Iris/Core/Private/Iris/ReplicationSystem/NetObjectFactory.cpp b/Engine/Source/Runtime/Experimental/Iris/Core/Private/Iris/ReplicationSystem/NetObjectFactory.cpp
--- a/Engine/Source/Runtime/Experimental/Iris/Core/Private/Iris/ReplicationSystem/NetObjectFactory.cpp
+++ b/Engine/Source/Runtime/Experimental/Iris/Core/Private/Iris/ReplicationSystem/NetObjectFactory.cpp
@@ -36,6 +36,11 @@ void UNetObjectFactory::Deinit()
 	Bridge = nullptr;
 }
 
+void UNetObjectFactory::PostTickDispatch()
+{
+	OnPostTickDispatch();
+}
+
 TUniquePtr<UE::Net::FNetObjectCreationHeader> UNetObjectFactory::CreateHeader(UE::Net::FNetRefHandle Handle, UE::Net::FReplicationProtocolIdentifier ProtocolId)
 {
 	using namespace UE::Net;
diff --git a/Engine/Source/Runtime/Experimental/Iris/Core/Private/Iris/ReplicationSystem/ObjectReplicationBridge.cpp b/Engine/Source/Runtime/Experimental/Iris/Core/Private/Iris/ReplicationSystem/ObjectReplicationBridge.cpp
--- a/Engine/Source/Runtime/Experimental/Iris/Core/Private/Iris/ReplicationSystem/ObjectReplicationBridge.cpp
+++ b/Engine/Source/Runtime/Experimental/Iris/Core/Private/Iris/ReplicationSystem/ObjectReplicationBridge.cpp
@@ -1108,6 +1108,14 @@ void UObjectReplicationBridge::OnPostSendUpdate()
 	bBlockStartRootObjectReplication = false;
 }
 
+void UObjectReplicationBridge::OnPostTickDispatch()
+{
+	for (UNetObjectFactory* NetFactory : NetObjectFactories)
+	{
+		NetFactory->PostTickDispatch();
+	}
+}
+
 void UObjectReplicationBridge::PruneStaleObjects()
 {
 	using namespace UE::Net;
diff --git a/Engine/Source/Runtime/Experimental/Iris/Core/Private/Iris/ReplicationSystem/ReplicationBridge.cpp b/Engine/Source/Runtime/Experimental/Iris/Core/Private/Iris/ReplicationSystem/ReplicationBridge.cpp
--- a/Engine/Source/Runtime/Experimental/Iris/Core/Private/Iris/ReplicationSystem/ReplicationBridge.cpp
+++ b/Engine/Source/Runtime/Experimental/Iris/Core/Private/Iris/ReplicationSystem/ReplicationBridge.cpp
@@ -96,6 +96,10 @@ void UReplicationBridge::OnPostSendUpdate()
 {
 }
 
+void UReplicationBridge::OnPostTickDispatch()
+{
+}
+
 void UReplicationBridge::UpdateInstancesWorldLocation()
 {
 }
diff --git a/Engine/Source/Runtime/Experimental/Iris/Core/Private/Iris/ReplicationSystem/ReplicationSystem.cpp b/Engine/Source/Runtime/Experimental/Iris/Core/Private/Iris/ReplicationSystem/ReplicationSystem.cpp
--- a/Engine/Source/Runtime/Experimental/Iris/Core/Private/Iris/ReplicationSystem/ReplicationSystem.cpp
+++ b/Engine/Source/Runtime/Experimental/Iris/Core/Private/Iris/ReplicationSystem/ReplicationSystem.cpp
@@ -8,15 +8,14 @@
 #include "Iris/Core/IrisProfiler.h"
 #include "Iris/Core/NetObjectReference.h"
 
-#include "Net/Core/Trace/NetTrace.h"
-#include "Net/Core/Misc/NetConditionGroupManager.h"
 #include "Net/Core/Connection/NetEnums.h"
+#include "Net/Core/Misc/NetConditionGroupManager.h"
 #include "Net/Core/NetToken/NetToken.h"
+#include "Net/Core/Trace/NetTrace.h"
 
 #include "Iris/ReplicationState/ReplicationStateUtil.h"
+
 #include "Iris/ReplicationSystem/ChangeMaskCache.h"
-#include "Iris/ReplicationSystem/Filtering/NetObjectFilter.h"
-#include "Iris/ReplicationSystem/Filtering/ReplicationFiltering.h"
 #include "Iris/ReplicationSystem/ReplicationBridge.h"
 #include "Iris/ReplicationSystem/ReplicationOperationsInternal.h"
 #include "Iris/ReplicationSystem/ReplicationReader.h"
@@ -25,6 +24,10 @@
 #include "Iris/ReplicationSystem/ReplicationSystemTypes.h"
 #include "Iris/ReplicationSystem/ReplicationTypes.h"
 #include "Iris/ReplicationSystem/ReplicationWriter.h"
+
+#include "Iris/ReplicationSystem/Filtering/NetObjectFilter.h"
+#include "Iris/ReplicationSystem/Filtering/ReplicationFiltering.h"
+
 #include "Iris/DataStream/DataStreamManager.h"
 #include "Iris/Serialization/InternalNetSerializationContext.h"
 #include "Iris/Serialization/NetBitStreamWriter.h"
@@ -126,9 +129,9 @@ public:
 		const uint32 ReplicationSystemId = ReplicationSystem->GetId();
 
 		// Verify that we got a NetTokenStore and that it is configured as we expect.
-		bool bHasValidNetTokenStore = ensureAlwaysMsgf(Params.NetTokenStore, TEXT("ReplicationSystem cannot be initialized without a valid NetTokenStore"));
-		bHasValidNetTokenStore = bHasValidNetTokenStore && ensureAlwaysMsgf(Params.NetTokenStore->GetDataStore<FStringTokenStore>(), TEXT("ReplicationSystem cannot be initialized without a StringTokenStore"));
-		bHasValidNetTokenStore = bHasValidNetTokenStore && ensureAlwaysMsgf(Params.NetTokenStore->GetDataStore<FNameTokenStore>(), TEXT("ReplicationSystem cannot be initialized without a NameTokenStore"));
+		bool bHasValidNetTokenStore = ensureMsgf(Params.NetTokenStore, TEXT("ReplicationSystem cannot be initialized without a valid NetTokenStore"));
+		bHasValidNetTokenStore = bHasValidNetTokenStore && ensureMsgf(Params.NetTokenStore->GetDataStore<FStringTokenStore>(), TEXT("ReplicationSystem cannot be initialized without a StringTokenStore"));
+		bHasValidNetTokenStore = bHasValidNetTokenStore && ensureMsgf(Params.NetTokenStore->GetDataStore<FNameTokenStore>(), TEXT("ReplicationSystem cannot be initialized without a NameTokenStore"));
 
 		if (!bHasValidNetTokenStore)
 		{
@@ -578,6 +581,11 @@ public:
 		NetBlobManager.ProcessOOBNetObjectAttachmentSendQueue(ConnectionsPendingPostTickDispatchSend);
 	}
 
+	void PostTickDispatch()
+	{
+		ReplicationSystemInternal.GetReplicationBridge()->OnPostTickDispatch();
+	}
+
 	void ResetNetObjectAttachmentSendQueue()
 	{
 		FNetBlobManager& NetBlobManager = ReplicationSystemInternal.GetNetBlobManager();
@@ -831,7 +839,7 @@ void UReplicationSystem::PreSendUpdate(const FSendUpdateParams& Params)
 
 	IRIS_CSV_PROFILER_SCOPE(Iris, ReplicationSystem_PreSendUpdate);
 
-	ensureAlways(Impl->CurrentSendPass == EReplicationSystemSendPass::Invalid);
+	ensureMsgf(Impl->CurrentSendPass == EReplicationSystemSendPass::Invalid, TEXT("PostSendUpdate was not called after the last PreSendUpdate."));
 	Impl->CurrentSendPass = Params.SendPass;
 
 	FReplicationSystemInternal& InternalSys = Impl->ReplicationSystemInternal;
@@ -929,12 +937,19 @@ void UReplicationSystem::PreSendUpdate(const FSendUpdateParams& Params)
 	{
 		// Forward attachments scheduled to use the OOBChannel and mark connetions needing immediate send
 		Impl->ProcessOOBNetObjectAttachmentSendQueue();
+
+		Impl->PostTickDispatch();
 	}
 }
 
 void UReplicationSystem::PreSendUpdate(float DeltaSeconds)
 {
-	PreSendUpdate(FSendUpdateParams {.SendPass = UE::Net::EReplicationSystemSendPass::TickFlush, .DeltaSeconds = DeltaSeconds });
+	const FSendUpdateParams Params
+	{ 
+		.SendPass = UE::Net::EReplicationSystemSendPass::TickFlush, 
+		.DeltaSeconds = DeltaSeconds 
+	};
+	PreSendUpdate(Params);
 }
 
 IRISCORE_API void UReplicationSystem::SendUpdate(TFunctionRef<void(TArrayView<uint32>)> SendFunction)
diff --git a/Engine/Source/Runtime/Experimental/Iris/Core/Public/Iris/ReplicationSystem/NetObjectFactory.h b/Engine/Source/Runtime/Experimental/Iris/Core/Public/Iris/ReplicationSystem/NetObjectFactory.h
--- a/Engine/Source/Runtime/Experimental/Iris/Core/Public/Iris/ReplicationSystem/NetObjectFactory.h
+++ b/Engine/Source/Runtime/Experimental/Iris/Core/Public/Iris/ReplicationSystem/NetObjectFactory.h
@@ -82,6 +82,7 @@ public:
 
 	void Init(UE::Net::FNetObjectFactoryId InId, UObjectReplicationBridge* InBridge);
 	void Deinit();
+	void PostTickDispatch();
 
 	struct FInstantiateResult;
 	struct FInstantiateContext;
@@ -263,9 +264,15 @@ public:
 
 protected:
 
+	/** Called when the netfactory is created */
 	virtual void OnInit() {}
+
+	/** Called before the netfactory will be destroyed */
 	virtual void OnDeinit() {}
 
+	/** Called after we finished processing all incoming packets */
+	virtual void OnPostTickDispatch() {}
+
 	/**
 	 * Create the correct header type for a given replicated object and fill the header with the information representing it.
 	 * @param Handle The net handle of the replicated object.
diff --git a/Engine/Source/Runtime/Experimental/Iris/Core/Public/Iris/ReplicationSystem/ObjectReplicationBridge.h b/Engine/Source/Runtime/Experimental/Iris/Core/Public/Iris/ReplicationSystem/ObjectReplicationBridge.h
--- a/Engine/Source/Runtime/Experimental/Iris/Core/Public/Iris/ReplicationSystem/ObjectReplicationBridge.h
+++ b/Engine/Source/Runtime/Experimental/Iris/Core/Public/Iris/ReplicationSystem/ObjectReplicationBridge.h
@@ -287,6 +287,7 @@ protected:
 	IRISCORE_API virtual void PreSendUpdate() override;	
 	IRISCORE_API virtual void OnStartPreSendUpdate() override;
 	IRISCORE_API virtual void OnPostSendUpdate() override;
+	IRISCORE_API virtual void OnPostTickDispatch() override;
 	IRISCORE_API virtual void UpdateInstancesWorldLocation() override;
 	IRISCORE_API virtual void PruneStaleObjects() override;	
 	IRISCORE_API virtual bool WriteNetRefHandleCreationInfo(FReplicationBridgeSerializationContext& Context, FNetRefHandle Handle) override;
diff --git a/Engine/Source/Runtime/Experimental/Iris/Core/Public/Iris/ReplicationSystem/ReplicationBridge.h b/Engine/Source/Runtime/Experimental/Iris/Core/Public/Iris/ReplicationSystem/ReplicationBridge.h
--- a/Engine/Source/Runtime/Experimental/Iris/Core/Public/Iris/ReplicationSystem/ReplicationBridge.h
+++ b/Engine/Source/Runtime/Experimental/Iris/Core/Public/Iris/ReplicationSystem/ReplicationBridge.h
@@ -142,6 +142,9 @@ protected:
 	
 	/** Invoked after we sent data to all connections. */
 	IRISCORE_API virtual void OnPostSendUpdate();
+
+	/** Invoked after we processed all incoming data */
+	IRISCORE_API virtual void OnPostTickDispatch();
 	
 	/** Invoked before ReplicationSystem copies dirty state data for a single replicated object. */
 	IRISCORE_API virtual void PreSendUpdateSingleHandle(FNetRefHandle Handle);
