Fix lock order issues between FOnDemandPackageStoreBackend and FOnDemandIoStore that could lead to hang.

diff --git a/Engine/Source/Runtime/Experimental/IoStore/OnDemand/Private/OnDemandIoStore.cpp b/Engine/Source/Runtime/Experimental/IoStore/OnDemand/Private/OnDemandIoStore.cpp
--- a/Engine/Source/Runtime/Experimental/IoStore/OnDemand/Private/OnDemandIoStore.cpp
+++ b/Engine/Source/Runtime/Experimental/IoStore/OnDemand/Private/OnDemandIoStore.cpp
@@ -70,6 +70,29 @@ namespace UE::IoStore
 namespace Private
 {
 
+///////////////////////////////////////////////////////////////////////////////
+// Enforce lock order
+struct FPackageStoreUpdateScope : public FOnDemandPackageStoreUpdateScope
+{
+	FPackageStoreUpdateScope(
+		UE::FMutex& InMutex,
+		IOnDemandPackageStoreBackend& InPackageStore,
+		TUniqueFunction<EOnDemandPackageStoreUpdateFlags()> InGetFlags = nullptr)
+		: FOnDemandPackageStoreUpdateScope(InPackageStore, MoveTemp(InGetFlags))
+		, Mutex(InMutex)
+	{
+		Mutex.Lock();
+	}
+
+	~FPackageStoreUpdateScope()
+	{
+		Mutex.Unlock();
+	}
+
+private:
+	UE::FMutex& Mutex;
+};
+
 ///////////////////////////////////////////////////////////////////////////////
 static FAnsiStringBuilderBase& GetChunkUrl(
 	const FAnsiStringView& Host,
@@ -786,7 +809,7 @@ FIoStatus FOnDemandIoStore::Unmount(FStringView MountId)
 	}
 
 	{
-		TUniqueLock Lock(ContainerMutex);
+		Private::FPackageStoreUpdateScope(ContainerMutex, *PackageStoreBackend);
 
 		Containers.SetNum(Algo::RemoveIf(Containers, [this, &MountId](const FSharedOnDemandContainer& Container)
 		{
@@ -795,8 +818,6 @@ FIoStatus FOnDemandIoStore::Unmount(FStringView MountId)
 				UE_LOG(LogIoStoreOnDemand, Log, TEXT("Unmounting container, ContainerName='%s', MountId='%s'"),
 					*WriteToString<128>(Container->Name), *WriteToString<128>(Container->MountId));
 
-				PackageStoreBackend->NeedsUpdate();
-
 				const FString ContainerName = Container->UniqueName();
 				PendingContainerHeaders.Remove(ContainerName);
 
@@ -1415,13 +1436,12 @@ FIoStatus FOnDemandIoStore::TickMountRequest(FMountRequest& MountRequest)
 	}
 
 	{
-		UE::TUniqueLock Lock(ContainerMutex);
+		Private::FPackageStoreUpdateScope(ContainerMutex, *PackageStoreBackend);
+
 		Containers.Append(MoveTemp(RequestedContainers));
 		PendingContainerHeaders.Append(MoveTemp(RequestedContainerHeaders));
 	}
 
-	PackageStoreBackend->NeedsUpdate();
-
 	return bAnyPendingEncryptionKey ? EIoErrorCode::PendingEncryptionKey : EIoErrorCode::Ok;
 }
 
@@ -1738,28 +1758,32 @@ FOnDemandInstallResult FOnDemandIoStore::TickInstallRequest(FOnDemandInternalIns
 		ContentHandle->IoStore = AsWeak(); 
 	}
 
-	for (const auto& Kv : InstallData)
+
+	if (InstallData.IsEmpty() == false)
 	{
-		FOnDemandContainer& Container				= *Kv.Key;
-		const Private::FContainerInstallData& Data	= Kv.Value;
+		Private::FPackageStoreUpdateScope(ContainerMutex, *PackageStoreBackend,
+			[]
+			{
+				// Referencers were added so the package store must be refreshed
+				return EOnDemandPackageStoreUpdateFlags::ReferencedPackages;
+			});
 
+		for (const auto& Kv : InstallData)
 		{
-			UE::TUniqueLock Lock(ContainerMutex);
-			FOnDemandChunkEntryReferences& References = Container.FindOrAddChunkEntryReferences(*ContentHandle);
-			for (const TPair<uint32, Private::EChunkCacheState>& ChunkPair : Data.ResolvedChunks)
+			FOnDemandContainer& Container = *Kv.Key;
+			const Private::FContainerInstallData& Data = Kv.Value;
+
 			{
-				const int32 EntryIndex = ChunkPair.Key;
-				References.Indices[EntryIndex] = true;
+				FOnDemandChunkEntryReferences& References = Container.FindOrAddChunkEntryReferences(*ContentHandle);
+				for (const TPair<uint32, Private::EChunkCacheState>& ChunkPair : Data.ResolvedChunks)
+				{
+					const int32 EntryIndex = ChunkPair.Key;
+					References.Indices[EntryIndex] = true;
+				}
 			}
 		}
 	}
 
-	if (InstallData.IsEmpty() == false)
-	{
-		// Referencers were added so the package store must be refreshed
-		PackageStoreBackend->NeedsUpdate(EOnDemandPackageStoreUpdateFlags::ReferencedPackages);
-	}
-
 	double DurationInSeconds = FPlatformTime::Seconds() - StartTime;
 	#if !UE_BUILD_SHIPPING
 	if (GIoStoreOnDemandMinInstallDurationSeconds > 0.0f)
@@ -1865,9 +1889,12 @@ void FOnDemandIoStore::OnEncryptionKeyAdded(const FGuid& Id, const FAES::FAESKey
 {
 	TRACE_CPUPROFILER_EVENT_SCOPE(FOnDemandIoStore::OnEncryptionKeyAdded);
 
-	TUniqueLock Lock(ContainerMutex);
-
 	bool bAddedContainerHeaders = false;
+	Private::FPackageStoreUpdateScope(ContainerMutex, *PackageStoreBackend,
+		[&bAddedContainerHeaders]
+		{
+			return bAddedContainerHeaders ? EOnDemandPackageStoreUpdateFlags::Full : EOnDemandPackageStoreUpdateFlags::None;
+		});
 
 	for (FSharedOnDemandContainer& Container : Containers)
 	{
@@ -1921,11 +1948,6 @@ void FOnDemandIoStore::OnEncryptionKeyAdded(const FGuid& Id, const FAES::FAESKey
 			PendingContainerHeaders.RemoveByHash(ContainerNameHash, ContainerName);
 		}
 	}
-
-	if (bAddedContainerHeaders)
-	{
-		PackageStoreBackend->NeedsUpdate();
-	}
 }
 
 void FOnDemandIoStore::CreateContainersFromToc(
@@ -2166,7 +2188,13 @@ void FOnDemandIoStore::ReleaseContent(FOnDemandInternalContentHandle& ContentHan
 	LLM_SCOPE_BYNAME(TEXT("FileSystem/OnDemandIoStore"));
 	UE_LOG(LogIoStoreOnDemand, Log, TEXT("Releasing content handle, %s"), *LexToString(ContentHandle));
 
-	UE::TUniqueLock Lock(ContainerMutex);
+	Private::FPackageStoreUpdateScope(ContainerMutex, *PackageStoreBackend, 
+		[]
+		{
+			// Package store needs to be aware of install state
+			return EOnDemandPackageStoreUpdateFlags::ReferencedPackages; 
+		});
+
 	for (FSharedOnDemandContainer& Container : Containers)
 	{
 		if (EnumHasAnyFlags(Container->Flags, EOnDemandContainerFlags::StreamOnDemand))
@@ -2182,9 +2210,6 @@ void FOnDemandIoStore::ReleaseContent(FOnDemandInternalContentHandle& ContentHan
 					return Refs.ContentHandleId == ContentHandleId;
 				}));
 	}
-
-	// Package store needs to be aware of install state
-	PackageStoreBackend->NeedsUpdate(EOnDemandPackageStoreUpdateFlags::ReferencedPackages);
 }
 
 void FOnDemandIoStore::GetReferencedContent(TArray<FSharedOnDemandContainer>& OutContainers, TArray<TBitArray<>>& OutChunkEntryIndices, bool bPackageStore)
diff --git a/Engine/Source/Runtime/Experimental/IoStore/OnDemand/Private/OnDemandPackageStoreBackend.cpp b/Engine/Source/Runtime/Experimental/IoStore/OnDemand/Private/OnDemandPackageStoreBackend.cpp
--- a/Engine/Source/Runtime/Experimental/IoStore/OnDemand/Private/OnDemandPackageStoreBackend.cpp
+++ b/Engine/Source/Runtime/Experimental/IoStore/OnDemand/Private/OnDemandPackageStoreBackend.cpp
@@ -49,7 +49,8 @@ public:
 						FOnDemandPackageStoreBackend(TWeakPtr<FOnDemandIoStore> OnDemandIoStore);
 						virtual ~FOnDemandPackageStoreBackend ();
 
-	virtual void		NeedsUpdate(EOnDemandPackageStoreUpdateFlags Flags) override;
+	virtual void		BeginNeedsUpdate() override;
+	virtual void		EndNeedsUpdate(EOnDemandPackageStoreUpdateFlags Flags) override;
 
 	virtual void		OnMounted(TSharedRef<const FPackageStoreBackendContext> Context) override;
 	virtual void 		BeginRead() override;
@@ -90,10 +91,15 @@ FOnDemandPackageStoreBackend::~FOnDemandPackageStoreBackend()
 {
 }
 
-void FOnDemandPackageStoreBackend::NeedsUpdate(EOnDemandPackageStoreUpdateFlags Flags)
+void FOnDemandPackageStoreBackend::BeginNeedsUpdate()
+{
+	Mutex.Lock();
+}
+
+void FOnDemandPackageStoreBackend::EndNeedsUpdate(EOnDemandPackageStoreUpdateFlags Flags)
 {
-	UE::TUniqueLock Lock(Mutex);
 	NeedsUpdateFlags |= Flags;
+	Mutex.Unlock();
 }
 
 void FOnDemandPackageStoreBackend::OnMounted(TSharedRef<const FPackageStoreBackendContext> Context)
@@ -180,6 +186,8 @@ TConstArrayView<uint32> FOnDemandPackageStoreBackend::GetSoftReferences(FPackage
 
 void FOnDemandPackageStoreBackend::UpdateLookupTables()
 {
+	// This must be locked before FOnDemandIoStore::ContainerMutex is locked in the call to GetReferencedContent below
+	// We need to maintain this lock order to not deadlock
 	check(Mutex.IsLocked());
 
 	EOnDemandPackageStoreUpdateFlags LocalNeedsUpdate = NeedsUpdateFlags;
@@ -326,6 +334,21 @@ void FOnDemandPackageStoreBackend::UpdateReferencedPackages()
 	}
 }
 
+///////////////////////////////////////////////////////////////////////////////
+FOnDemandPackageStoreUpdateScope::FOnDemandPackageStoreUpdateScope(
+	IOnDemandPackageStoreBackend& InPackageStore, TUniqueFunction<EOnDemandPackageStoreUpdateFlags()> InGetFlags)
+	: PackageStore(InPackageStore)
+	, GetFlags(MoveTemp(InGetFlags))
+{
+	PackageStore.BeginNeedsUpdate();
+}
+
+FOnDemandPackageStoreUpdateScope::~FOnDemandPackageStoreUpdateScope()
+{
+	EOnDemandPackageStoreUpdateFlags Flags = GetFlags ? GetFlags() : EOnDemandPackageStoreUpdateFlags::Full;
+	PackageStore.EndNeedsUpdate(Flags);
+}
+
 ///////////////////////////////////////////////////////////////////////////////
 TSharedPtr<IOnDemandPackageStoreBackend> MakeOnDemandPackageStoreBackend(TWeakPtr<FOnDemandIoStore> OnDemandIoStore)
 {
diff --git a/Engine/Source/Runtime/Experimental/IoStore/OnDemand/Private/OnDemandPackageStoreBackend.h b/Engine/Source/Runtime/Experimental/IoStore/OnDemand/Private/OnDemandPackageStoreBackend.h
--- a/Engine/Source/Runtime/Experimental/IoStore/OnDemand/Private/OnDemandPackageStoreBackend.h
+++ b/Engine/Source/Runtime/Experimental/IoStore/OnDemand/Private/OnDemandPackageStoreBackend.h
@@ -27,7 +27,22 @@ class IOnDemandPackageStoreBackend
 public:
 	virtual ~IOnDemandPackageStoreBackend() = default;
 
-	virtual void NeedsUpdate(EOnDemandPackageStoreUpdateFlags Flags = EOnDemandPackageStoreUpdateFlags::Full) = 0;
+	virtual void BeginNeedsUpdate() = 0;
+
+	virtual void EndNeedsUpdate(EOnDemandPackageStoreUpdateFlags Flags) = 0;
+};
+
+// Update wrapper to help enforce mutex locking order between this and FOnDemandIoStore
+struct FOnDemandPackageStoreUpdateScope
+{
+	FOnDemandPackageStoreUpdateScope(
+		IOnDemandPackageStoreBackend& InPackageStore, 
+		TUniqueFunction<EOnDemandPackageStoreUpdateFlags()> InGetFlags = nullptr);
+	~FOnDemandPackageStoreUpdateScope();
+	
+private:
+	IOnDemandPackageStoreBackend& PackageStore;
+	TUniqueFunction<EOnDemandPackageStoreUpdateFlags()> GetFlags;
 };
 
 TSharedPtr<IOnDemandPackageStoreBackend> MakeOnDemandPackageStoreBackend(TWeakPtr<FOnDemandIoStore> OnDemandIoStore);
