Dataflow : Update the fracture editor dataflow evaluation to evaluate all dependent assets terminal nodes

diff --git a/Engine/Plugins/Experimental/ChaosEditor/Source/FractureEditor/Private/FractureEditorModeToolkit.cpp b/Engine/Plugins/Experimental/ChaosEditor/Source/FractureEditor/Private/FractureEditorModeToolkit.cpp
--- a/Engine/Plugins/Experimental/ChaosEditor/Source/FractureEditor/Private/FractureEditorModeToolkit.cpp
+++ b/Engine/Plugins/Experimental/ChaosEditor/Source/FractureEditor/Private/FractureEditorModeToolkit.cpp
@@ -311,7 +311,7 @@ FReply FFractureEditorModeToolkit::OnDataflowOverridesUpdateAsset()
 	UFractureSettings* FractureSettings = GetMutableDefault<UFractureSettings>();
 	if (FractureSettings && FractureSettings->RestCollection.IsValid())
 	{
-		FractureSettings->RestCollection->GetDataflowInstance().UpdateOwnerAsset();
+		FractureSettings->RestCollection->GetDataflowInstance().UpdateOwnerAsset(/* bUpdateDependentAssets */true);
 	}
 	return FReply::Handled();
 }
diff --git a/Engine/Source/Runtime/Experimental/Dataflow/Engine/Private/Dataflow/DataflowInstance.cpp b/Engine/Source/Runtime/Experimental/Dataflow/Engine/Private/Dataflow/DataflowInstance.cpp
--- a/Engine/Source/Runtime/Experimental/Dataflow/Engine/Private/Dataflow/DataflowInstance.cpp
+++ b/Engine/Source/Runtime/Experimental/Dataflow/Engine/Private/Dataflow/DataflowInstance.cpp
@@ -368,15 +368,29 @@ void FDataflowInstance::SyncVariables()
 	VariableOverrides.SyncVariables();
 }
 
-bool FDataflowInstance::UpdateOwnerAsset() const
+bool FDataflowInstance::UpdateOwnerAsset(bool bUpdateDependentAssets) const
 {
+	bool bSuccess = false;
 	if (DataflowAsset && DataflowAsset->Dataflow && Owner)
 	{
-		if (const TSharedPtr<FDataflowNode> Node = DataflowAsset->Dataflow->FindFilteredNode(FDataflowTerminalNode::StaticType(), DataflowTerminal))
+		TArray<TSharedPtr<FDataflowNode>> NodesToEvaluate;
+
+		if (bUpdateDependentAssets)
+		{
+			// find all terminal nodes
+			NodesToEvaluate = DataflowAsset->Dataflow->GetFilteredNodes(FDataflowTerminalNode::StaticType());
+		}
+		else if (const TSharedPtr<FDataflowNode> Node = DataflowAsset->Dataflow->FindFilteredNode(FDataflowTerminalNode::StaticType(), DataflowTerminal))
 		{
-			if (const FDataflowTerminalNode* TerminalNode = Node->AsType<const FDataflowTerminalNode>())
+			// find only the specified one 
+			NodesToEvaluate.Add(Node);
+		}
+
+		UE::Dataflow::FEngineContext Context(Owner.Get());
+		for (TSharedPtr<FDataflowNode> Node: NodesToEvaluate)
+		{
+			if (const FDataflowTerminalNode* TerminalNode = Node? Node->AsType<const FDataflowTerminalNode>(): nullptr)
 			{
-				UE::Dataflow::FEngineContext Context(Owner.Get());
 				// Note: If the node is deactivated and has any outputs, then these outputs might still need to be forwarded.
 				//       Therefore the Evaluate method has to be called for whichever value of bActive.
 				//       This however isn't the case of SetAssetValue() for which the active state needs to be checked before the call.
@@ -384,10 +398,10 @@ bool FDataflowInstance::UpdateOwnerAsset() const
 				if (TerminalNode->IsActive() && Owner)
 				{
 					TerminalNode->SetAssetValue(Owner, Context);
-					return true;
+					bSuccess = true;
 				}
 			}
 		}
 	}
-	return false;
+	return bSuccess;
 }
\ No newline at end of file
diff --git a/Engine/Source/Runtime/Experimental/Dataflow/Engine/Public/Dataflow/DataflowInstance.h b/Engine/Source/Runtime/Experimental/Dataflow/Engine/Public/Dataflow/DataflowInstance.h
--- a/Engine/Source/Runtime/Experimental/Dataflow/Engine/Public/Dataflow/DataflowInstance.h
+++ b/Engine/Source/Runtime/Experimental/Dataflow/Engine/Public/Dataflow/DataflowInstance.h
@@ -121,7 +121,7 @@ public:
 	* Update the asset by re-evaluating the dataflow
 	* return true if the dataflow has been properly evaluated and the asset updated
 	*/
-	DATAFLOWENGINE_API bool UpdateOwnerAsset() const;
+	DATAFLOWENGINE_API bool UpdateOwnerAsset(bool bUpdateDependentAssets = false) const;
 
 #if WITH_EDITOR
 	DATAFLOWENGINE_API TSharedPtr<FStructOnScope> MakeStructOnScope() const;
