Changed I/O store on-demand module feature to both create and destroy the I/O store instance to avoid cross dll memory issues in non-unity builds.
- renabled I/O store on-demand in editor builds

diff --git a/Engine/Source/Runtime/Engine/Engine.Build.cs b/Engine/Source/Runtime/Engine/Engine.Build.cs
--- a/Engine/Source/Runtime/Engine/Engine.Build.cs
+++ b/Engine/Source/Runtime/Engine/Engine.Build.cs
@@ -124,7 +124,7 @@ public class Engine : ModuleRules
 			PublicDependencyModuleNames.Add("ApplicationCore");
 		}
 
-		if (Target.bCompileIoStoreOnDemand && !Target.bBuildEditor)
+		if (Target.bCompileIoStoreOnDemand)
 		{
 			PrivateDependencyModuleNames.Add("IoStoreOnDemand");
 			PublicDefinitions.Add("UE_WITH_IOSTOREONDEMAND=1");
diff --git a/Engine/Source/Runtime/Experimental/IoStore/OnDemand/Private/OnDemandModule.cpp b/Engine/Source/Runtime/Experimental/IoStore/OnDemand/Private/OnDemandModule.cpp
--- a/Engine/Source/Runtime/Experimental/IoStore/OnDemand/Private/OnDemandModule.cpp
+++ b/Engine/Source/Runtime/Experimental/IoStore/OnDemand/Private/OnDemandModule.cpp
@@ -61,10 +61,20 @@ struct FPlatformSocketSystem
 class FOnDemandIoStoreFactory final
 	: public IOnDemandIoStoreFactory
 {
-	virtual FSharedOnDemandIoStore CreateInstance() override
+public:
+	virtual IOnDemandIoStore* CreateInstance() override
+	{
+		IoStore = MakeShared<FOnDemandIoStore>();
+		return IoStore.Get();
+	}
+
+	virtual void DestroyInstance(IOnDemandIoStore* Instance)
 	{
-		return MakeShared<FOnDemandIoStore>();
+		IoStore.Reset();
 	}
+
+private:
+	FSharedOnDemandIoStore IoStore;
 };
 
 ////////////////////////////////////////////////////////////////////////////////
diff --git a/Engine/Source/Runtime/Experimental/IoStore/OnDemandCore/Private/IoStoreOnDemand.cpp b/Engine/Source/Runtime/Experimental/IoStore/OnDemandCore/Private/IoStoreOnDemand.cpp
--- a/Engine/Source/Runtime/Experimental/IoStore/OnDemandCore/Private/IoStoreOnDemand.cpp
+++ b/Engine/Source/Runtime/Experimental/IoStore/OnDemandCore/Private/IoStoreOnDemand.cpp
@@ -23,7 +23,7 @@ DEFINE_LOG_CATEGORY(LogIas);
 namespace UE::IoStore
 {
 
-FSharedOnDemandIoStore GOnDemandIoStore;
+IOnDemandIoStore* GOnDemandIoStore = nullptr;
 
 ///////////////////////////////////////////////////////////////////////////////
 #if !UE_BUILD_SHIPPING
@@ -368,15 +368,24 @@ IMPLEMENT_MODULE(UE::IoStore::FIoStoreOnDemandCoreModule, IoStoreOnDemandCore);
 ////////////////////////////////////////////////////////////////////////////////
 void FIoStoreOnDemandCoreModule::HandleModularFeatureRegistered(const FName& Type, IModularFeature* ModularFeature)
 {
-	if (Type != IOnDemandIoStoreFactory::FeatureName || GOnDemandIoStore.IsValid())
+	if (Type != IOnDemandIoStoreFactory::FeatureName || GOnDemandIoStore != nullptr)
 	{
 		return;
 	}
 
-	IOnDemandIoStoreFactory* Factory = static_cast<IOnDemandIoStoreFactory*>(ModularFeature);
-	FSharedOnDemandIoStore IoStore = Factory->CreateInstance();
+	IOnDemandIoStoreFactory* Factory	= static_cast<IOnDemandIoStoreFactory*>(ModularFeature);
+	IOnDemandIoStore* IoStore			= Factory->CreateInstance();
+
+	if (IoStore == nullptr)
+	{
+		UE_LOG(LogIoStoreOnDemand, Warning, TEXT("I/O store on-demand disabled, reason '%s'"), TEXT("Failed to create I/O store"));
+		return;
+	}
+
 	if (FIoStatus Status = IoStore->Initialize(); !Status.IsOk())
 	{
+		Factory->DestroyInstance(IoStore);
+
 		if (Status.GetErrorCode() == EIoErrorCode::Disabled || Status.GetErrorCode() == EIoErrorCode::NotFound)
 		{
 			UE_LOG(LogIoStoreOnDemand, Log, TEXT("I/O store on-demand disabled, reason '%s'"), *Status.ToString());
@@ -420,21 +429,36 @@ void FIoStoreOnDemandCoreModule::StartupModule()
 
 void FIoStoreOnDemandCoreModule::ShutdownModule()
 {
-	GOnDemandIoStore.Reset();
+	IOnDemandIoStore* ToDestroy = nullptr;
+	Swap(ToDestroy, GOnDemandIoStore);
+
+	if (ToDestroy == nullptr)
+	{
+		return;
+	}
+
+	IModularFeatures& Features = IModularFeatures::Get();
+	if (Features.GetModularFeatureImplementationCount(IOnDemandIoStoreFactory::FeatureName) > 0)
+	{
+		IModularFeature* Feature			= Features.GetModularFeatureImplementation(IOnDemandIoStoreFactory::FeatureName, 0);
+		IOnDemandIoStoreFactory* Factory	= static_cast<IOnDemandIoStoreFactory*>(Feature);
+
+		Factory->DestroyInstance(ToDestroy);
+	}
 	IModularFeatures::Get().OnModularFeatureRegistered().RemoveAll(this);
 }
 
 ////////////////////////////////////////////////////////////////////////////////
 IOnDemandIoStore* TryGetOnDemandIoStore()
 {
-	return GOnDemandIoStore.Get();
+	return GOnDemandIoStore;
 }
 
 ////////////////////////////////////////////////////////////////////////////////
 IOnDemandIoStore& GetOnDemandIoStore()
 {
-	check(GOnDemandIoStore.IsValid());
-	return *GOnDemandIoStore.Get();
+	check(GOnDemandIoStore != nullptr);
+	return *GOnDemandIoStore;
 }
 
 FName IOnDemandIoStoreFactory::FeatureName = FName(TEXT("OnDemandIoStoreFactory"));
diff --git a/Engine/Source/Runtime/Experimental/IoStore/OnDemandCore/Public/IO/IoStoreOnDemand.h b/Engine/Source/Runtime/Experimental/IoStore/OnDemandCore/Public/IO/IoStoreOnDemand.h
--- a/Engine/Source/Runtime/Experimental/IoStore/OnDemandCore/Public/IO/IoStoreOnDemand.h
+++ b/Engine/Source/Runtime/Experimental/IoStore/OnDemandCore/Public/IO/IoStoreOnDemand.h
@@ -481,7 +481,9 @@ public:
 	/** Feature name. */
 	UE_API static FName FeatureName;
 	/** Create a new instance of the I/O store.  Called once the feature has been registered. */
-	virtual FSharedOnDemandIoStore CreateInstance() = 0;
+	virtual IOnDemandIoStore* CreateInstance() = 0;
+	/** Destroy the instance. */
+	virtual void DestroyInstance(IOnDemandIoStore* Instance) = 0;
 };
 
 /** Returns the on-demand I/O store if available. */
