FOnDemandThread no longer invokes request callbacks on the task system.

#rnx

- Instead of invoking callbacks on the task system they will now be invoked immediately on the thread being used to process http requests.
-- It is expected that the caller decide if the work in the callback should be processed immediately and risk delaying subsequent http requests or moved to the task system.
- FOnDemandIoBackend retains it's existing behavior. Requests that succeeded which might take some time to process are pushed to the task system but failed requests which should be relatively light to process will be processed immediately.
- Combined some duplicate code for issuing requests to a common FOnDemandIoBackend::IssueHttpRequest method.
- Add some code documentation to FOnDemandThread.

diff --git a/Engine/Source/Runtime/Experimental/IoStore/OnDemand/Private/OnDemandHttpThread.cpp b/Engine/Source/Runtime/Experimental/IoStore/OnDemand/Private/OnDemandHttpThread.cpp
--- a/Engine/Source/Runtime/Experimental/IoStore/OnDemand/Private/OnDemandHttpThread.cpp
+++ b/Engine/Source/Runtime/Experimental/IoStore/OnDemand/Private/OnDemandHttpThread.cpp
@@ -14,7 +14,6 @@
 #include "OnDemandHttpClient.h"
 #include "OnDemandIoStore.h"
 #include "Statistics.h"
-#include "Tasks/Task.h"
 
 // TODO - Maybe do LexToString/LexFromString (dupe code in UnrealEngine.cpp ThreadPriorityToString)
 static const TCHAR* LexToString(EThreadPriority Priority)
@@ -61,8 +60,6 @@ extern int32 GIasHttpConcurrentRequests;
 extern int32 GIasHttpRateLimitKiBPerSecond;
 extern int32 GIasHttpPollTimeoutMs;
 
-extern UE::Tasks::ETaskPriority GetRequestCompletionTaskPriority();
-
 const int32 GREPLACEOnDemandBackendThreadPriorities[5] =
 {
 	EThreadPriority::TPri_Lowest,
@@ -302,11 +299,8 @@ void FOnDemandThread::TickRequests()
 
 								FOnDemandIoBackendStats::Get()->OnHttpGet(HttpResponse.Body.DataSize(), HttpResponse.DurationMilliseconds);
 
-								UE::Tasks::Launch(UE_SOURCE_LOCATION, [HttpRequest, StatusCode = HttpResponse.StatusCode, ErrorReason = MoveTemp(HttpResponse.Reason), Data = MoveTemp(HttpResponse.Body)]() mutable
-									{
-										HttpRequest->OnRequestCompleted(StatusCode, ErrorReason, MoveTemp(Data));
-										FOnDemandThread::Get().DestroyRequest(HttpRequest); // Don't pass the this pointer into the lambda to keep it below the small task allocation size
-									}, GetRequestCompletionTaskPriority());
+								HttpRequest->OnRequestCompleted(HttpResponse.StatusCode, HttpResponse.Reason, MoveTemp(HttpResponse.Body));
+								DestroyRequest(HttpRequest);
 							}
 							else
 							{
diff --git a/Engine/Source/Runtime/Experimental/IoStore/OnDemand/Private/OnDemandHttpThread.h b/Engine/Source/Runtime/Experimental/IoStore/OnDemand/Private/OnDemandHttpThread.h
--- a/Engine/Source/Runtime/Experimental/IoStore/OnDemand/Private/OnDemandHttpThread.h
+++ b/Engine/Source/Runtime/Experimental/IoStore/OnDemand/Private/OnDemandHttpThread.h
@@ -42,11 +42,33 @@ public:
 
 	static FOnDemandThread& Get();
 	
+	/**
+	 * Issue a http request to read from a chunk.
+	 * 
+	 * @param ChunkInfo		Info about the chunk to be read
+	 * @param ReadRange		The range from within the chunk to read. A default FIoOffsetAndLength will read the entire chunk.
+	 * @CompletionCallback	A callback that will be invoked when the http request has completed. This will be invoked on the
+	 *						thread that the http requests are being processed on and will block new requests being queued so
+	 *						care should be taken to do as little work in the callback as possible.
+	 * @return				A handle representing the request which can be used to modify it while the request is still in
+	 *						flight. The handle will no longer be valid after the CompletionCallback has been invoked.
+	 */
 	FRequestHandle IssueRequest(const FOnDemandChunkInfo& ChunkInfo, const FIoOffsetAndLength& ReadRange, FCompletionCallback&& CompletionCallback);
+	
+	/**
+	 * Issue a http request to read from a chunk.
+	 *
+	 * @param ChunkInfo		Info about the chunk to be read
+	 * @param ReadRange		The range from within the chunk to read. A default FIoOffsetAndLength will read the entire chunk.
+	 * @CompletionCallback	A callback that will be invoked when the http request has completed. This will be invoked on the
+	 *						thread that the http requests are being processed on and will block new requests being queued so
+	 *						care should be taken to do as little work in the callback as possible.
+	 * @return				A handle representing the request which can be used to modify it while the request is still in
+	 *						flight. The handle will no longer be valid after the CompletionCallback has been invoked.
+	 */
 	FRequestHandle IssueRequest(FOnDemandChunkInfo&& ChunkInfo, const FIoOffsetAndLength& ReadRange, FCompletionCallback&& CompletionCallback);
 
 	void ReprioritizeRequest(FRequestHandle Request);
-
 	void CancelRequest(FRequestHandle Request);
 
 	FOnTickIdle& OnTickIdle()
diff --git a/Engine/Source/Runtime/Experimental/IoStore/OnDemand/Private/OnDemandIoDispatcherBackend.cpp b/Engine/Source/Runtime/Experimental/IoStore/OnDemand/Private/OnDemandIoDispatcherBackend.cpp
--- a/Engine/Source/Runtime/Experimental/IoStore/OnDemand/Private/OnDemandIoDispatcherBackend.cpp
+++ b/Engine/Source/Runtime/Experimental/IoStore/OnDemand/Private/OnDemandIoDispatcherBackend.cpp
@@ -722,6 +722,7 @@ private:
 	bool Resolve(FIoRequestImpl* Request);
 
 	void ConditionallyStartBackendThread();
+	void IssueHttpRequest(FChunkRequest* ChunkRequest);
 	void CompleteRequest(FChunkRequest* ChunkRequest);
 	void CompleteCacheRead(FChunkRequest* ChunkRequest);
 	bool ResolveDefaultHostGroup();
@@ -940,6 +941,35 @@ void FOnDemandIoBackend::ConditionallyStartBackendThread()
 #endif // UE_ENABLE_REQUEST_THREAD
 }
 
+void FOnDemandIoBackend::IssueHttpRequest(FChunkRequest* ChunkRequest)
+{
+#if UE_ENABLE_REQUEST_THREAD
+
+	ChunkRequest->HttpRequest = FOnDemandThread::Get().IssueRequest(ChunkRequest->Params.ChunkInfo, ChunkRequest->Params.ChunkRange,
+		[this, ChunkRequest](uint32 StatusCode, FStringView ErrorReason, FIoBuffer&& Data)
+		{
+			ChunkRequest->Chunk = MoveTemp(Data);
+
+			if (IsHttpStatusOk(StatusCode))
+			{
+				UE::Tasks::Launch(UE_SOURCE_LOCATION, [this, ChunkRequest]()
+					{
+						CompleteRequest(ChunkRequest);
+					}, GetRequestCompletionTaskPriority());
+			}
+			else
+			{
+				CompleteRequest(ChunkRequest);
+			}
+		});
+
+	check(ChunkRequest->HttpRequest != nullptr);
+#else
+	Stats.OnHttpEnqueue();
+	HttpRequests.EnqueueByPriority(ChunkRequest);
+	TickBackendEvent->Trigger();
+#endif //UE_ENABLE_REQUEST_THREAD
+}
 
 void FOnDemandIoBackend::CompleteRequest(FChunkRequest* ChunkRequest)
 {
@@ -1081,21 +1111,7 @@ void FOnDemandIoBackend::CompleteCacheRead(FChunkRequest* ChunkRequest)
 		return;
 	}
 
-#if UE_ENABLE_REQUEST_THREAD
-	ChunkRequest->HttpRequest = FOnDemandThread::Get().IssueRequest(ChunkRequest->Params.ChunkInfo, ChunkRequest->Params.ChunkRange,
-		[this, ChunkRequest](uint32 StatusCode, FStringView ErrorReason, FIoBuffer&& Data)
-		{
-			ChunkRequest->Chunk = MoveTemp(Data);
-			CompleteRequest(ChunkRequest);
-		});
-
-	check(ChunkRequest->HttpRequest != nullptr);
-
-#else
-	Stats.OnHttpEnqueue();
-	HttpRequests.EnqueueByPriority(ChunkRequest);
-	TickBackendEvent->Trigger();
-#endif //UE_ENABLE_REQUEST_THREAD
+	IssueHttpRequest(ChunkRequest);
 }
 
 bool FOnDemandIoBackend::Resolve(FIoRequestImpl* Request)
@@ -1181,20 +1197,7 @@ bool FOnDemandIoBackend::Resolve(FIoRequestImpl* Request)
 
 	ChunkRequest->HostGroup = ChunkInfo.HostGroup();
 
-#if UE_ENABLE_REQUEST_THREAD
-	ChunkRequest->HttpRequest = FOnDemandThread::Get().IssueRequest(ChunkRequest->Params.ChunkInfo, ChunkRequest->Params.ChunkRange,
-		[this, ChunkRequest](uint32 StatusCode, FStringView ErrorReason, FIoBuffer&& Data)
-		{
-			ChunkRequest->Chunk = MoveTemp(Data);
-			CompleteRequest(ChunkRequest);
-		});
-
-	check(ChunkRequest->HttpRequest != nullptr);
-#else
-	Stats.OnHttpEnqueue();
-	HttpRequests.EnqueueByPriority(ChunkRequest);
-	TickBackendEvent->Trigger();
-#endif // UE_ENABLE_REQUEST_THREAD
+	IssueHttpRequest(ChunkRequest);
 
 	return true;
 }
