Add experimental version of IAD that uses FOnDemandThread to make requests.

#rnx

- Added two new macros UE_ENABLE_IAS_REQUEST_THREAD and UE_ENABLE_IAD_REQUEST_THREAD so that we can selectively enable the use of UE_ENABLE_IAS_REQUEST_THREAD for each system separately.
-- Both are disabled by default.
- Added a code path to IAD that will use FOnDemandThread to make requests. Currently this maintains the same behaviour as the previous version except that we have to wait for the http request thread to finish processing all requests rather than ticking it manually.
-- The code is not production ready and is only proof of concept.

diff --git a/Engine/Source/Runtime/Experimental/IoStore/OnDemand/Private/OnDemandHttpThread.h b/Engine/Source/Runtime/Experimental/IoStore/OnDemand/Private/OnDemandHttpThread.h
--- a/Engine/Source/Runtime/Experimental/IoStore/OnDemand/Private/OnDemandHttpThread.h
+++ b/Engine/Source/Runtime/Experimental/IoStore/OnDemand/Private/OnDemandHttpThread.h
@@ -6,6 +6,9 @@
 	#define UE_ENABLE_REQUEST_THREAD 0
 #endif
 
+#define UE_ENABLE_IAS_REQUEST_THREAD (0 && UE_ENABLE_REQUEST_THREAD)
+#define UE_ENABLE_IAD_REQUEST_THREAD (0 && UE_ENABLE_REQUEST_THREAD)
+
 #if UE_ENABLE_REQUEST_THREAD
 
 #include "Async/Mutex.h"
diff --git a/Engine/Source/Runtime/Experimental/IoStore/OnDemand/Private/OnDemandIoDispatcherBackend.cpp b/Engine/Source/Runtime/Experimental/IoStore/OnDemand/Private/OnDemandIoDispatcherBackend.cpp
--- a/Engine/Source/Runtime/Experimental/IoStore/OnDemand/Private/OnDemandIoDispatcherBackend.cpp
+++ b/Engine/Source/Runtime/Experimental/IoStore/OnDemand/Private/OnDemandIoDispatcherBackend.cpp
@@ -438,10 +438,10 @@ struct FChunkRequest
 	FChunkRequestParams Params;
 	FIoRequestImpl* RequestHead;
 	FIoRequestImpl* RequestTail;
-#if UE_ENABLE_REQUEST_THREAD
+#if UE_ENABLE_IAS_REQUEST_THREAD
 	FOnDemandThread::FRequestHandle HttpRequest = 0;
-#endif //UE_ENABLE_REQUEST_THREAD
-	FIASHostGroup HostGroup; // Still used in a few places even when UE_ENABLE_REQUEST_THREAD is enabled
+#endif //UE_ENABLE_IAS_REQUEST_THREAD
+	FIASHostGroup HostGroup; // Still used in a few places even when UE_ENABLE_IAS_REQUEST_THREAD is enabled
 	FIoBuffer Chunk;
 	uint64 StartTime;
 	int32 Priority;
@@ -631,9 +631,9 @@ class FOnDemandIoBackend final
 					check(ChunkRequest.RequestTail == nullptr);
 
 					ChunkRequest.bCancelled = true;
-#if UE_ENABLE_REQUEST_THREAD
+#if UE_ENABLE_IAS_REQUEST_THREAD
 					FOnDemandThread::Get().CancelRequest(ChunkRequest.HttpRequest);
-#endif // UE_ENABLE_REQUEST_THREAD
+#endif // UE_ENABLE_IAS_REQUEST_THREAD
 
 					if (TheCache != nullptr)
 					{
@@ -757,9 +757,9 @@ private:
 
 	EThreadPriority CurrentThreadPriority;
 
-#if UE_ENABLE_REQUEST_THREAD
+#if UE_ENABLE_IAS_REQUEST_THREAD
 	FDelegateHandle OnThreadTickIdleHandle;
-#endif // UE_ENABLE_REQUEST_THREAD
+#endif // UE_ENABLE_IAS_REQUEST_THREAD
 
 	mutable FRWLock Lock;
 	std::atomic_uint32_t InflightCacheRequestCount{0};
@@ -865,13 +865,13 @@ void FOnDemandIoBackend::Shutdown()
 
 	UE_LOG(LogIas, Log, TEXT("Shutting down on demand I/O dispatcher backend"));
 
-#if UE_ENABLE_REQUEST_THREAD
+#if UE_ENABLE_IAS_REQUEST_THREAD
 	if (OnThreadTickIdleHandle.IsValid())
 	{
 		FOnDemandThread::Get().OnTickIdle().Remove(OnThreadTickIdleHandle);
 		OnThreadTickIdleHandle.Reset();
 	}
-#endif // UE_ENABLE_REQUEST_THREAD
+#endif // UE_ENABLE_IAS_REQUEST_THREAD
 
 	// Stop and wait for our backend thread to finish.
 	// Note, the IoDispatcher typically waits for all its pending io requests before shutting down its backends.
@@ -916,7 +916,7 @@ void FOnDemandIoBackend::OnThreadTickIdle()
 
 void FOnDemandIoBackend::ConditionallyStartBackendThread()
 {
-#if !UE_ENABLE_REQUEST_THREAD
+#if !UE_ENABLE_IAS_REQUEST_THREAD
 	FWriteScopeLock _(Lock);
 	if (BackendThread.IsValid() == false)
 	{
@@ -938,13 +938,12 @@ void FOnDemandIoBackend::ConditionallyStartBackendThread()
 		});
 
 	OnThreadTickIdleHandle = FOnDemandThread::Get().OnTickIdle().AddRaw(this, &FOnDemandIoBackend::OnThreadTickIdle);
-#endif // UE_ENABLE_REQUEST_THREAD
+#endif // UE_ENABLE_IAS_REQUEST_THREAD
 }
 
 void FOnDemandIoBackend::IssueHttpRequest(FChunkRequest* ChunkRequest)
 {
-#if UE_ENABLE_REQUEST_THREAD
-
+#if UE_ENABLE_IAS_REQUEST_THREAD
 	ChunkRequest->HttpRequest = FOnDemandThread::Get().IssueRequest(ChunkRequest->Params.ChunkInfo, ChunkRequest->Params.ChunkRange,
 		[this, ChunkRequest](uint32 StatusCode, FStringView ErrorReason, FIoBuffer&& Data)
 		{
@@ -968,7 +967,7 @@ void FOnDemandIoBackend::IssueHttpRequest(FChunkRequest* ChunkRequest)
 	Stats.OnHttpEnqueue();
 	HttpRequests.EnqueueByPriority(ChunkRequest);
 	TickBackendEvent->Trigger();
-#endif //UE_ENABLE_REQUEST_THREAD
+#endif //UE_ENABLE_IAS_REQUEST_THREAD
 }
 
 void FOnDemandIoBackend::CompleteRequest(FChunkRequest* ChunkRequest)
@@ -1146,12 +1145,12 @@ bool FOnDemandIoBackend::Resolve(FIoRequestImpl* Request)
 	{
 		if (bUpdatePriority)
 		{
-#if UE_ENABLE_REQUEST_THREAD
+#if UE_ENABLE_IAS_REQUEST_THREAD
 			FOnDemandThread::Get().ReprioritizeRequest(ChunkRequest->HttpRequest);
 #else
 			TRACE_CPUPROFILER_EVENT_SCOPE(IasBackend::UpdatePriorityForIoRequest);
 			HttpRequests.Reprioritize(ChunkRequest);
-#endif
+#endif // UE_ENABLE_IAS_REQUEST_THREAD
 		}
 		// The chunk for the request is already inflight 
 		return true;
@@ -1227,11 +1226,11 @@ void FOnDemandIoBackend::UpdatePriorityForIoRequest(FIoRequestImpl* Request)
 	TRACE_CPUPROFILER_EVENT_SCOPE(IasBackend::UpdatePriorityForIoRequest);
 	if (FChunkRequest* ChunkRequest = ChunkRequests.TryUpdatePriority(Request))
 	{
-#if UE_ENABLE_REQUEST_THREAD
+#if UE_ENABLE_IAS_REQUEST_THREAD
 		FOnDemandThread::Get().ReprioritizeRequest(ChunkRequest->HttpRequest);
 #else
 		HttpRequests.Reprioritize(ChunkRequest);
-#endif // UE_ENABLE_REQUEST_THREAD
+#endif // UE_ENABLE_IAS_REQUEST_THREAD
 	}
 }
 
diff --git a/Engine/Source/Runtime/Experimental/IoStore/OnDemand/Private/OnDemandIoStore.cpp b/Engine/Source/Runtime/Experimental/IoStore/OnDemand/Private/OnDemandIoStore.cpp
--- a/Engine/Source/Runtime/Experimental/IoStore/OnDemand/Private/OnDemandIoStore.cpp
+++ b/Engine/Source/Runtime/Experimental/IoStore/OnDemand/Private/OnDemandIoStore.cpp
@@ -1601,8 +1601,12 @@ FOnDemandInstallResult FOnDemandIoStore::TickInstallRequest(FOnDemandInternalIns
 	}
 
 	// Download all chunks
+	
+#if UE_ENABLE_IAD_REQUEST_THREAD
+	std::atomic<int32> ConcurrentRequests = 0;
+#else
 	const int32 MaxConcurrentRequests = 16;
-	int32		ConcurrentRequests = 0;
+	int32 ConcurrentRequests = 0;
 
 	FMultiEndpointHttpClientConfig HttpConfig;
 	HttpConfig.MaxConnectionCount = 8;
@@ -1615,6 +1619,7 @@ FOnDemandInstallResult FOnDemandIoStore::TickInstallRequest(FOnDemandInternalIns
 		Status = FIoStatusBuilder(EIoErrorCode::InvalidCode) << TEXT("Failed to initialize HTTP client");
 		return OutResult;
 	}
+#endif // UE_ENABLE_IAD_REQUEST_THREAD
 
 	uint64 TotalChunkCount		= 0;
 	uint64 TotalBytes			= 0;
@@ -1655,86 +1660,134 @@ FOnDemandInstallResult FOnDemandIoStore::TickInstallRequest(FOnDemandInternalIns
 
 	ProgressInstallRequest(InstallRequest, Progress);
 
-	for (const Private::FInstallData::ElementType& Kv : InstallData)
 	{
-		const FOnDemandContainer& Container			= *Kv.Key;
-		const Private::FContainerInstallData& Data	= Kv.Value;
-
-		for (const TPair<uint32, Private::EChunkCacheState>& ChunkPair : Data.ResolvedChunks)
+		TRACE_CPUPROFILER_EVENT_SCOPE(FOnDemandIoStore::InstallContent::HttpRequests);
+		for (const Private::FInstallData::ElementType& Kv : InstallData)
 		{
-			const int32 EntryIndex							= ChunkPair.Key;
-			const Private::EChunkCacheState ChunkCacheState	= ChunkPair.Value;
-
-			const FIoChunkId& ChunkId						= Container.ChunkIds[EntryIndex];
-			const FOnDemandChunkEntry& ChunkEntry			= Container.ChunkEntries[EntryIndex];
+			const FOnDemandContainer& Container			= *Kv.Key;
+			const Private::FContainerInstallData& Data	= Kv.Value;
 
-			if (ChunkCacheState == Private::EChunkCacheState::Cached)
+			for (const TPair<uint32, Private::EChunkCacheState>& ChunkPair : Data.ResolvedChunks)
 			{
-				continue;
-			}
+				const int32 EntryIndex							= ChunkPair.Key;
+				const Private::EChunkCacheState ChunkCacheState	= ChunkPair.Value;
 
-			ConcurrentRequests++;
-			HttpClient->Get(
-				Container.HostGroup,
-				Private::GetChunkUrl(Container, ChunkEntry, ChunkUrl).ToView(),
-				[this, &InstallRequest, &OutResult, ChunkId, ChunkEntry, &ConcurrentRequests]
-				(FMultiEndpointHttpClientResponse&& HttpResponse)
+				const FIoChunkId& ChunkId						= Container.ChunkIds[EntryIndex];
+				const FOnDemandChunkEntry& ChunkEntry			= Container.ChunkEntries[EntryIndex];
+
+				if (ChunkCacheState == Private::EChunkCacheState::Cached)
 				{
-					ConcurrentRequests--;
+					continue;
+				}
 
-					FIoStatus& Status					= OutResult.Status;
-					FOnDemandInstallProgress& Progress	= OutResult.Progress;
+#if UE_ENABLE_IAD_REQUEST_THREAD
+				++ConcurrentRequests;
 
-					if (!Status.IsOk())
-					{
-						// a different request already failed, don't stomp the status if this one succeeded
-						return;
-					}
-					
-					if (!HttpResponse.IsOk())
+				FOnDemandThread::Get().IssueRequest(FOnDemandChunkInfo(Kv.Key, ChunkEntry), FIoOffsetAndLength(),
+					[this, &ConcurrentRequests, ChunkEntry, &OutResult, &InstallRequest](uint32 StatusCode, FStringView ErrorReason, FIoBuffer&& Data)
 					{
-						Status = FIoStatus(EIoErrorCode::ReadError) << 
-							FString::Printf(TEXT("HTTP Status %u, "), HttpResponse.StatusCode) << 
-							HttpResponse.Reason;
-						return;
-					}
+						--ConcurrentRequests;
 
-					FIoBuffer& Chunk		= HttpResponse.Body;
-					const FIoHash ChunkHash	= FIoHash::HashBuffer(Chunk.GetView());
-					if (ChunkHash != ChunkEntry.Hash)
-					{
-						Status = FIoStatus(EIoErrorCode::ReadError, TEXTVIEW("Hash mismatch"));
-						return;
-					}
-					Status = InstallCache->PutChunk(MoveTemp(Chunk), ChunkHash);
-					if (!Status.IsOk())
+						if (!OutResult.Status.IsOk())
+						{
+							// A different request already failed, don't stomp the status if this one succeeded
+							return;
+						}
+
+						if (!IsHttpStatusOk(StatusCode))
+						{
+							OutResult.Status = FIoStatus(EIoErrorCode::ReadError) << FString::Printf(TEXT("HTTP Status %u, "), StatusCode) << ErrorReason;
+							return;
+						}
+
+						const FIoHash ChunkHash = FIoHash::HashBuffer(Data.GetView());
+						if (ChunkHash != ChunkEntry.Hash)
+						{
+							OutResult.Status = FIoStatus(EIoErrorCode::ReadError, TEXTVIEW("Hash mismatch"));
+							return;
+						}
+						
+						OutResult.Status = InstallCache->PutChunk(MoveTemp(Data), ChunkHash);
+						if (!OutResult.Status.IsOk())
+						{
+							return;
+						}
+
+						// TODO: Is this good enough progress or should it be fine grained bytes from the HttpClient?
+						OutResult.Progress.CurrentInstallSize += ChunkEntry.EncodedSize;
+						ProgressInstallRequest(InstallRequest, OutResult.Progress);
+					});
+#else
+				HttpClient->Get(
+					Container.HostGroup,
+					Private::GetChunkUrl(Container, ChunkEntry, ChunkUrl).ToView(),
+					[this, &InstallRequest, &OutResult, ChunkId, ChunkEntry, &ConcurrentRequests]
+					(FMultiEndpointHttpClientResponse&& HttpResponse)
 					{
-						return;
-					}
+						ConcurrentRequests--;
 
-					// TODO: Is this good enough progress or should it be fine grained bytes from the HttpClient?
-					Progress.CurrentInstallSize += ChunkEntry.EncodedSize;
-					ProgressInstallRequest(InstallRequest, Progress);
-				});
+						FIoStatus& Status					= OutResult.Status;
+						FOnDemandInstallProgress& Progress	= OutResult.Progress;
 
-			while (ConcurrentRequests >= MaxConcurrentRequests)
-			{
-				HttpClient->Tick();
-			}
+						if (!Status.IsOk())
+						{
+							// a different request already failed, don't stomp the status if this one succeeded
+							return;
+						}
+						
+						if (!HttpResponse.IsOk())
+						{
+							Status = FIoStatus(EIoErrorCode::ReadError) << 
+								FString::Printf(TEXT("HTTP Status %u, "), HttpResponse.StatusCode) << 
+								HttpResponse.Reason;
+							return;
+						}
 
-			if ((Status.IsOk() == false) || CheckAndSetCancelled())
-			{
-				// TODO: Should there be a way to flush/cancel the HTTPClient?
-				while (HttpClient->Tick())
-					;
+						FIoBuffer& Chunk		= HttpResponse.Body;
+						const FIoHash ChunkHash	= FIoHash::HashBuffer(Chunk.GetView());
+						if (ChunkHash != ChunkEntry.Hash)
+						{
+							Status = FIoStatus(EIoErrorCode::ReadError, TEXTVIEW("Hash mismatch"));
+							return;
+						}
+						Status = InstallCache->PutChunk(MoveTemp(Chunk), ChunkHash);
+						if (!Status.IsOk())
+						{
+							return;
+						}
+
+						// TODO: Is this good enough progress or should it be fine grained bytes from the HttpClient?
+						Progress.CurrentInstallSize += ChunkEntry.EncodedSize;
+						ProgressInstallRequest(InstallRequest, Progress);
+					});
+
+				while (ConcurrentRequests >= MaxConcurrentRequests)
+				{
+					HttpClient->Tick();
+				}
+
+				if ((Status.IsOk() == false) || CheckAndSetCancelled())
+				{
+					// TODO: Should there be a way to flush/cancel the HTTPClient?
+					while (HttpClient->Tick())
+						;
 
-				return OutResult;
+					return OutResult;
+				}
+#endif //UE_ENABLE_IAD_REQUEST_THREAD
 			}
 		}
-	}
 
-	while (HttpClient->Tick())
-		;
+#if UE_ENABLE_IAD_REQUEST_THREAD
+		while (ConcurrentRequests > 0)
+		{
+			FPlatformProcess::SleepNoStats(0.0f);
+		}
+#else
+		while (HttpClient->Tick())
+			;
+#endif //UE_ENABLE_IAD_REQUEST_THREAD
+	}
 
 	if (Status.IsOk() == false)
 	{
