Dataflow : ForEach loop support
- Subgraph can now ietarte through arrays when using the GetCurrentIndex node
- Subgraph context now have a current iteration index that can be queried in the subgraph via a node
- Subgraph icon will change to a loop when detecting the use of GetCurrentindex node
- Also fix issues with EdGraphNode added to the EdGraph before they had their DataflowNodeGuid set causing listeners to graph notification to not have all the info needed

diff --git a/Engine/Plugins/Experimental/Dataflow/Source/DataflowEditor/Private/Dataflow/DataflowAssetEditUtils.cpp b/Engine/Plugins/Experimental/Dataflow/Source/DataflowEditor/Private/Dataflow/DataflowAssetEditUtils.cpp
--- a/Engine/Plugins/Experimental/Dataflow/Source/DataflowEditor/Private/Dataflow/DataflowAssetEditUtils.cpp
+++ b/Engine/Plugins/Experimental/Dataflow/Source/DataflowEditor/Private/Dataflow/DataflowAssetEditUtils.cpp
@@ -144,13 +144,14 @@ namespace UE::Dataflow
 					Dataflow->Modify();
 					EdGraph->Modify();
 
+					// make sure we set the guid before adding to graph so that the listener to the graph notification have all the info needed
+					EdNode->SetDataflowGraph(Dataflow->GetDataflow());
+					EdNode->SetDataflowNodeGuid(DataflowNode->GetGuid());
+
 					EdGraph->AddNode(EdNode, /*bUserAction*/true, /*bSelectNewNode*/false);
 					
 					EdNode->CreateNewGuid();
 					EdNode->PostPlacedNewNode();
-
-					EdNode->SetDataflowGraph(Dataflow->GetDataflow());
-					EdNode->SetDataflowNodeGuid(DataflowNode->GetGuid());
 					EdNode->AllocateDefaultPins();
 
 					if (FromPin)
diff --git a/Engine/Plugins/Experimental/Dataflow/Source/DataflowEditor/Private/Dataflow/DataflowEditorCommands.cpp b/Engine/Plugins/Experimental/Dataflow/Source/DataflowEditor/Private/Dataflow/DataflowEditorCommands.cpp
--- a/Engine/Plugins/Experimental/Dataflow/Source/DataflowEditor/Private/Dataflow/DataflowEditorCommands.cpp
+++ b/Engine/Plugins/Experimental/Dataflow/Source/DataflowEditor/Private/Dataflow/DataflowEditorCommands.cpp
@@ -81,7 +81,10 @@ void FDataflowEditorCommandsImpl::RegisterCommands()
 	UI_COMMAND(StepSimulationScene, "Step Simulation", "Step the simulation scene", EUserInterfaceActionType::Button, FInputChord());
 
 	UI_COMMAND(AddNewVariable, "Variable", "Adds a new variable to this dataflow graph.", EUserInterfaceActionType::Button, FInputChord());
-	UI_COMMAND(AddNewSubGraph, "SubGraph", "Adds a new sub-graph to this dataflow graph.", EUserInterfaceActionType::Button, FInputChord());
+	UI_COMMAND(AddNewSubGraph, "SubGraph", "Adds a new subgraph to this dataflow graph.", EUserInterfaceActionType::Button, FInputChord());
+	UI_COMMAND(ConvertToBasicSubGraph, "Convert to Basic Subgraph", "Convert the subgraph to a basic one (no loop).", EUserInterfaceActionType::Button, FInputChord());
+	UI_COMMAND(ConvertToForEachSubGraph, "Convert to For Each Subgraph", "Convert the subgraph to be able to iterate through an array.", EUserInterfaceActionType::Button, FInputChord());
+	
 
 	for (const TPair<FName, TUniquePtr<UE::Dataflow::IDataflowConstructionViewMode>>& NameAndMode : UE::Dataflow::FRenderingViewModeFactory::GetInstance().GetViewModes())
 	{
diff --git a/Engine/Plugins/Experimental/Dataflow/Source/DataflowEditor/Private/Dataflow/DataflowGraphSchemaAction.cpp b/Engine/Plugins/Experimental/Dataflow/Source/DataflowEditor/Private/Dataflow/DataflowGraphSchemaAction.cpp
--- a/Engine/Plugins/Experimental/Dataflow/Source/DataflowEditor/Private/Dataflow/DataflowGraphSchemaAction.cpp
+++ b/Engine/Plugins/Experimental/Dataflow/Source/DataflowEditor/Private/Dataflow/DataflowGraphSchemaAction.cpp
@@ -259,11 +259,33 @@ const FName FEdGraphSchemaAction_DataflowSubGraph::GetSubGraphName() const
 	return {};
 }
 
+bool FEdGraphSchemaAction_DataflowSubGraph::IsForEachSubGraph() const
+{
+	if (TStrongObjectPtr<const UDataflow> DataflowAsset = DataflowAssetWeakPtr.Pin())
+	{
+		if (const UDataflowSubGraph* SubGraph = DataflowAsset->FindSubGraphByGuid(SubGraphGuid))
+		{
+			return SubGraph->IsForEachSubGraph();
+		}
+	}
+	return false;
+}
+
+void FEdGraphSchemaAction_DataflowSubGraph::SetForEachSubGraph(bool bValue)
+{
+	if (TStrongObjectPtr<UDataflow> DataflowAsset = DataflowAssetWeakPtr.Pin())
+	{
+		if (UDataflowSubGraph* SubGraph = DataflowAsset->FindSubGraphByGuid(SubGraphGuid))
+		{
+			return SubGraph->SetForEachSubGraph(bValue);
+		}
+	}
+}
+
 UEdGraphNode* FEdGraphSchemaAction_DataflowSubGraph::PerformAction(class UEdGraph* ParentGraph, UEdGraphPin* FromPin, const FVector2D Location, bool bSelectNewNode)
 {
-	// TODO : move this code to the FEditableDataflow eventually 
 	UEdGraphNode* NewEdNode = nullptr;
-	if (TSharedPtr<FAssetSchemaAction_Dataflow_CreateNode_DataflowEdNode>  SubGraphNodeAction = FAssetSchemaAction_Dataflow_CreateNode_DataflowEdNode::CreateAction(ParentGraph, "FDataflowCallSubGraphNode", GetSubGraphName()))
+	if (TSharedPtr<FAssetSchemaAction_Dataflow_CreateNode_DataflowEdNode>  SubGraphNodeAction = FAssetSchemaAction_Dataflow_CreateNode_DataflowEdNode::CreateAction(ParentGraph, FDataflowCallSubGraphNode::StaticType(), GetSubGraphName()))
 	{
 		UDataflowEdNode* DataflowEdNode = Cast<UDataflowEdNode>(SubGraphNodeAction->PerformAction(ParentGraph, nullptr, Location, false));
 		NewEdNode = DataflowEdNode;
diff --git a/Engine/Plugins/Experimental/Dataflow/Source/DataflowEditor/Private/Dataflow/DataflowMembersWidget.cpp b/Engine/Plugins/Experimental/Dataflow/Source/DataflowEditor/Private/Dataflow/DataflowMembersWidget.cpp
--- a/Engine/Plugins/Experimental/Dataflow/Source/DataflowEditor/Private/Dataflow/DataflowMembersWidget.cpp
+++ b/Engine/Plugins/Experimental/Dataflow/Source/DataflowEditor/Private/Dataflow/DataflowMembersWidget.cpp
@@ -416,6 +416,14 @@ void SDataflowMembersWidget::InitializeCommands()
 		FExecuteAction::CreateSP(this, &SDataflowMembersWidget::OnDelete),
 		FCanExecuteAction(), FIsActionChecked(),
 		FIsActionButtonVisible::CreateSP(this, &SDataflowMembersWidget::CanDelete));
+
+	CommandList->MapAction(FDataflowEditorCommands::Get().ConvertToBasicSubGraph,
+		FExecuteAction::CreateSP(this, &SDataflowMembersWidget::SetForEachSubGraphOnSelection, false),
+		FCanExecuteAction(), FIsActionChecked());
+
+	CommandList->MapAction(FDataflowEditorCommands::Get().ConvertToForEachSubGraph,
+		FExecuteAction::CreateSP(this, &SDataflowMembersWidget::SetForEachSubGraphOnSelection, true),
+		FCanExecuteAction(), FIsActionChecked());
 }
 
 void SDataflowMembersWidget::InitializeSections()
@@ -605,6 +613,7 @@ void SDataflowMembersWidget::OnSubGraphsChanged(const UDataflow* InDataflowAsset
 					break;
 
 				case UE::Dataflow::ESubGraphChangedReason::Deleted:
+				case UE::Dataflow::ESubGraphChangedReason::ChangedType:
 					// nothing to do 
 					break;
 				}
@@ -710,12 +719,31 @@ void SDataflowMembersWidget::OnActionDoubleClicked(const TArray< TSharedPtr<FEdG
 
 TSharedPtr<SWidget> SDataflowMembersWidget::OnContextMenuOpening()
 {
+	const TAttribute<FText> EmptyTextAttribute;
+
 	const bool bShouldCloseWindowAfterMenuSelection = true;
 	FMenuBuilder MenuBuilder(bShouldCloseWindowAfterMenuSelection, CommandList);
 
 	if (IsAnyActionSelected())
 	{
+		const bool bSubGraphActions = IsOnlySubgraphActionsSelected();
+
 		MenuBuilder.AddMenuEntry(FGenericCommands::Get().Rename);
+		if (bSubGraphActions)
+		{
+			MenuBuilder.AddMenuSeparator();
+
+			if (IsSelectionForEachSubGraph())
+			{
+				const FSlateIcon FunctionIcon(FAppStyle::GetAppStyleSetName(), "GraphEditor.Function_16x");
+				MenuBuilder.AddMenuEntry(FDataflowEditorCommands::Get().ConvertToBasicSubGraph, NAME_None, EmptyTextAttribute, EmptyTextAttribute, FunctionIcon);
+			}
+			else
+			{
+				const FSlateIcon LoopIcon(FAppStyle::GetAppStyleSetName(), "GraphEditor.Macro.Loop_16x");
+				MenuBuilder.AddMenuEntry(FDataflowEditorCommands::Get().ConvertToForEachSubGraph, NAME_None, EmptyTextAttribute, EmptyTextAttribute, LoopIcon);
+			}
+		}
 		MenuBuilder.AddMenuSeparator();
 		MenuBuilder.AddMenuEntry(FGenericCommands::Get().Cut);
 		MenuBuilder.AddMenuEntry(FGenericCommands::Get().Copy);
@@ -824,6 +852,24 @@ bool SDataflowMembersWidget::IsAnyActionSelected() const
 	return (SelectedActions.Num() > 0);
 }
 
+bool SDataflowMembersWidget::IsOnlySubgraphActionsSelected() const
+{
+	TArray<TSharedPtr<FEdGraphSchemaAction> > SelectedActions;
+	GraphActionMenu->GetSelectedActions(SelectedActions);
+	if (SelectedActions.IsEmpty())
+	{
+		return false;
+	}
+	for (const TSharedPtr<FEdGraphSchemaAction>& Action : SelectedActions)
+	{
+		if (Action->GetTypeId() != FEdGraphSchemaAction_DataflowSubGraph::StaticGetTypeId())
+		{
+			return false;
+		}
+	}
+	return true;
+}
+
 TSharedPtr<FEdGraphSchemaAction> SDataflowMembersWidget::GetFirstSelectedAction() const
 {
 	TArray<TSharedPtr<FEdGraphSchemaAction>> SelectedActions;
@@ -972,4 +1018,49 @@ bool SDataflowMembersWidget::CanDelete() const
 	return false;
 }
 
+
+void SDataflowMembersWidget::SetForEachSubGraphOnSelection(bool bValue)
+{
+	TArray<TSharedPtr<FEdGraphSchemaAction> > SelectedActions;
+	GraphActionMenu->GetSelectedActions(SelectedActions);
+
+	for (TSharedPtr<FEdGraphSchemaAction>& Action : SelectedActions)
+	{
+		if (Action->GetTypeId() == FEdGraphSchemaAction_DataflowSubGraph::StaticGetTypeId())
+		{
+			if (FEdGraphSchemaAction_DataflowSubGraph* SubGraphAction = static_cast<FEdGraphSchemaAction_DataflowSubGraph*>(Action.Get()))
+			{
+				SubGraphAction->SetForEachSubGraph(bValue);
+			}
+		}
+	}
+}
+
+bool SDataflowMembersWidget::IsSelectionForEachSubGraph() const
+{
+	TArray<TSharedPtr<FEdGraphSchemaAction> > SelectedActions;
+	GraphActionMenu->GetSelectedActions(SelectedActions);
+
+	if (SelectedActions.IsEmpty())
+	{
+		return false;
+	}
+
+	int32 NumForEachSubGraph = 0;
+	for (const TSharedPtr<FEdGraphSchemaAction>& Action : SelectedActions)
+	{
+		if (Action->GetTypeId() == FEdGraphSchemaAction_DataflowSubGraph::StaticGetTypeId())
+		{
+			if (const FEdGraphSchemaAction_DataflowSubGraph* SubGraphAction = static_cast<FEdGraphSchemaAction_DataflowSubGraph*>(Action.Get()))
+			{
+				if (SubGraphAction->IsForEachSubGraph())
+				{
+					++NumForEachSubGraph;
+				}
+			}
+		}
+	}
+	return (NumForEachSubGraph > (SelectedActions.Num() / 2));
+}
+
 #undef LOCTEXT_NAMESPACE
\ No newline at end of file
diff --git a/Engine/Plugins/Experimental/Dataflow/Source/DataflowEditor/Private/Dataflow/DataflowSubGraphPaletteItem.cpp b/Engine/Plugins/Experimental/Dataflow/Source/DataflowEditor/Private/Dataflow/DataflowSubGraphPaletteItem.cpp
--- a/Engine/Plugins/Experimental/Dataflow/Source/DataflowEditor/Private/Dataflow/DataflowSubGraphPaletteItem.cpp
+++ b/Engine/Plugins/Experimental/Dataflow/Source/DataflowEditor/Private/Dataflow/DataflowSubGraphPaletteItem.cpp
@@ -36,9 +36,10 @@ TSharedRef<SWidget> SDataflowSubGraphPaletteItem::CreateTextSlotWidget(FCreateWi
 		+SHorizontalBox::Slot()
 		.AutoWidth()
 		.VAlign(VAlign_Center)
+		.Padding(FMargin(4,0))
 		[
 			SNew(SImage)
-				.Image(FAppStyle::GetBrush(TEXT("GraphEditor.Function_16x")))
+				.Image(this,&SDataflowSubGraphPaletteItem::GetSubGraphIcon)
 		]
 		+ SHorizontalBox::Slot()
 		.AutoWidth()
@@ -69,3 +70,12 @@ void SDataflowSubGraphPaletteItem::OnNameTextCommitted(const FText& InNewText, E
 		SubGraphAction->RenameItem(InNewText);
 	}
 }
+
+const FSlateBrush* SDataflowSubGraphPaletteItem::GetSubGraphIcon() const
+{
+	if (SubGraphAction && SubGraphAction->IsForEachSubGraph())
+	{
+		return FAppStyle::GetBrush(TEXT("GraphEditor.Macro.Loop_16x"));
+	}
+	return FAppStyle::GetBrush(TEXT("GraphEditor.Function_16x")); // default
+}
\ No newline at end of file
diff --git a/Engine/Plugins/Experimental/Dataflow/Source/DataflowEditor/Private/Dataflow/DataflowSubGraphPaletteItem.h b/Engine/Plugins/Experimental/Dataflow/Source/DataflowEditor/Private/Dataflow/DataflowSubGraphPaletteItem.h
--- a/Engine/Plugins/Experimental/Dataflow/Source/DataflowEditor/Private/Dataflow/DataflowSubGraphPaletteItem.h
+++ b/Engine/Plugins/Experimental/Dataflow/Source/DataflowEditor/Private/Dataflow/DataflowSubGraphPaletteItem.h
@@ -25,6 +25,8 @@ private:
 	virtual void OnNameTextCommitted(const FText& NewText, ETextCommit::Type InTextCommit) override;
 	// End of SGraphPaletteItem Interface
 
+	const FSlateBrush* GetSubGraphIcon() const;
+
 private:
 	TSharedPtr<FEdGraphSchemaAction_DataflowSubGraph> SubGraphAction;
 
diff --git a/Engine/Plugins/Experimental/Dataflow/Source/DataflowEditor/Public/Dataflow/DataflowEditorCommands.h b/Engine/Plugins/Experimental/Dataflow/Source/DataflowEditor/Public/Dataflow/DataflowEditorCommands.h
--- a/Engine/Plugins/Experimental/Dataflow/Source/DataflowEditor/Public/Dataflow/DataflowEditorCommands.h
+++ b/Engine/Plugins/Experimental/Dataflow/Source/DataflowEditor/Public/Dataflow/DataflowEditorCommands.h
@@ -73,6 +73,8 @@ public:
 	TSharedPtr<FUICommandInfo> ZoomToFitGraph;
 	TSharedPtr<FUICommandInfo> AddNewVariable;
 	TSharedPtr<FUICommandInfo> AddNewSubGraph;
+	TSharedPtr<FUICommandInfo> ConvertToBasicSubGraph;
+	TSharedPtr<FUICommandInfo> ConvertToForEachSubGraph;
 
 	TMap<FName, TSharedPtr<FUICommandInfo>> SetConstructionViewModeCommands;
 
diff --git a/Engine/Plugins/Experimental/Dataflow/Source/DataflowEditor/Public/Dataflow/DataflowGraphSchemaAction.h b/Engine/Plugins/Experimental/Dataflow/Source/DataflowEditor/Public/Dataflow/DataflowGraphSchemaAction.h
--- a/Engine/Plugins/Experimental/Dataflow/Source/DataflowEditor/Public/Dataflow/DataflowGraphSchemaAction.h
+++ b/Engine/Plugins/Experimental/Dataflow/Source/DataflowEditor/Public/Dataflow/DataflowGraphSchemaAction.h
@@ -167,7 +167,9 @@ public:
 	static FName StaticGetTypeId() { static FName Type("FEdGraphSchemaAction_DataflowSubGraph"); return Type; }
 	virtual FName GetTypeId() const override { return StaticGetTypeId(); }
 
-	const FName GetSubGraphName() const;
+	DATAFLOWEDITOR_API const FName GetSubGraphName() const;
+	DATAFLOWEDITOR_API bool IsForEachSubGraph() const;
+	DATAFLOWEDITOR_API void SetForEachSubGraph(bool bValue);
 
 	// FEdGraphSchemaAction interface
 	virtual bool IsAVariable() const { return false; }
diff --git a/Engine/Plugins/Experimental/Dataflow/Source/DataflowEditor/Public/Dataflow/DataflowMembersWidget.h b/Engine/Plugins/Experimental/Dataflow/Source/DataflowEditor/Public/Dataflow/DataflowMembersWidget.h
--- a/Engine/Plugins/Experimental/Dataflow/Source/DataflowEditor/Public/Dataflow/DataflowMembersWidget.h
+++ b/Engine/Plugins/Experimental/Dataflow/Source/DataflowEditor/Public/Dataflow/DataflowMembersWidget.h
@@ -100,6 +100,7 @@ private:
 
 	void SelectItemByName(const FName& ItemName, ESelectInfo::Type SelectInfo = ESelectInfo::Direct, int32 SectionId = INDEX_NONE, bool bIsCategory = false);
 	bool IsAnyActionSelected() const;
+	bool IsOnlySubgraphActionsSelected() const;
 	TSharedPtr<FEdGraphSchemaAction> GetFirstSelectedAction() const;
 
 	const TSharedPtr<ISection> GetSectionById(int32 SectionId) const;
@@ -121,6 +122,8 @@ private:
 	bool CanDuplicate() const;
 	void OnDelete();
 	bool CanDelete() const;
+	bool IsSelectionForEachSubGraph() const;
+	void SetForEachSubGraphOnSelection(bool bValue);
 
 	// UI refresh 
 	void Refresh();
diff --git a/Engine/Source/Runtime/Experimental/Dataflow/Core/Private/Dataflow/DataflowNodeParameters.cpp b/Engine/Source/Runtime/Experimental/Dataflow/Core/Private/Dataflow/DataflowNodeParameters.cpp
--- a/Engine/Source/Runtime/Experimental/Dataflow/Core/Private/Dataflow/DataflowNodeParameters.cpp
+++ b/Engine/Source/Runtime/Experimental/Dataflow/Core/Private/Dataflow/DataflowNodeParameters.cpp
@@ -413,13 +413,17 @@ namespace UE::Dataflow
 		SetNullData(InElementKey, InProperty, InNodeGuid, InNodeHash, InTimestamp);
 	}
 
-	bool FContext::CopyDataToAnotherContext(const FContextCacheKey InSourceKey, FContext& TargetContext, const FContextCacheKey InTargetKey) const
+	bool FContext::CopyDataToAnotherContext(const FContextCacheKey InSourceKey, FContext& TargetContext, const FContextCacheKey InTargetKey, const FProperty* InProperty, const FGuid& InNodeGuid, uint32 InNodeHash, const FTimestamp& InTimestamp) const
 	{
 		const TUniquePtr<UE::Dataflow::FContextCacheElementBase>* CacheEntryToClone = GetDataImpl(InSourceKey);
 		if (CacheEntryToClone && CacheEntryToClone->IsValid())
 		{
 			TUniquePtr<UE::Dataflow::FContextCacheElementBase> ClonedCacheEntry = (*CacheEntryToClone)->Clone(*this);
-			TargetContext.SetDataImpl(InTargetKey, MoveTemp(ClonedCacheEntry));
+			if (ClonedCacheEntry)
+			{
+				ClonedCacheEntry->UpdatePropertyAndNodeData(InProperty, InNodeGuid, InNodeHash, InTimestamp);
+				TargetContext.SetDataImpl(InTargetKey, MoveTemp(ClonedCacheEntry));
+			}
 			return true;
 		}
 		return false;
diff --git a/Engine/Source/Runtime/Experimental/Dataflow/Core/Public/Dataflow/DataflowContextCache.h b/Engine/Source/Runtime/Experimental/Dataflow/Core/Public/Dataflow/DataflowContextCache.h
--- a/Engine/Source/Runtime/Experimental/Dataflow/Core/Public/Dataflow/DataflowContextCache.h
+++ b/Engine/Source/Runtime/Experimental/Dataflow/Core/Public/Dataflow/DataflowContextCache.h
@@ -104,6 +104,15 @@ namespace UE::Dataflow
 		// use this with caution: setting the property of a wrong type may cause problems
 		void SetProperty(const FProperty* NewProperty) { Property = NewProperty; }
 
+		// use this with caution: setting the property of a wrong type may cause problems
+		void UpdatePropertyAndNodeData(const FProperty* InProperty, const FGuid& InNodeGuid, uint32 InNodeHash, const FTimestamp& InTimestamp)
+		{ 
+			Property = InProperty;
+			NodeGuid = InNodeGuid;
+			NodeHash = InNodeHash;
+			Timestamp = InTimestamp;
+		}
+
 	private:
 		friend struct FContextCache;
 
diff --git a/Engine/Source/Runtime/Experimental/Dataflow/Core/Public/Dataflow/DataflowNodeParameters.h b/Engine/Source/Runtime/Experimental/Dataflow/Core/Public/Dataflow/DataflowNodeParameters.h
--- a/Engine/Source/Runtime/Experimental/Dataflow/Core/Public/Dataflow/DataflowNodeParameters.h
+++ b/Engine/Source/Runtime/Experimental/Dataflow/Core/Public/Dataflow/DataflowNodeParameters.h
@@ -259,7 +259,7 @@ namespace UE::Dataflow
 
 		DATAFLOWCORE_API void SetArrayElementFromData(const FContextCacheKey InArrayKey, int32 Index, const FContextCacheKey InElementKey, const FProperty* InProperty, const FGuid& InNodeGuid, uint32 InNodeHash, const FTimestamp& InTimestamp);
 
-		DATAFLOWCORE_API bool CopyDataToAnotherContext(const FContextCacheKey InSourceKey, FContext& TargetContext, const FContextCacheKey InTargetKey) const;
+		DATAFLOWCORE_API bool CopyDataToAnotherContext(const FContextCacheKey InSourceKey, FContext& TargetContext, const FContextCacheKey InTargetKey, const FProperty* InProperty, const FGuid& InNodeGuid, uint32 InNodeHash, const FTimestamp& InTimestamp) const;
 
 		template<class T>
 		const T& GetData(FContextCacheKey Key, const FProperty* InProperty, const T& Default = T()) const
diff --git a/Engine/Source/Runtime/Experimental/Dataflow/Engine/Private/Dataflow/DataflowSubGraph.cpp b/Engine/Source/Runtime/Experimental/Dataflow/Engine/Private/Dataflow/DataflowSubGraph.cpp
--- a/Engine/Source/Runtime/Experimental/Dataflow/Engine/Private/Dataflow/DataflowSubGraph.cpp
+++ b/Engine/Source/Runtime/Experimental/Dataflow/Engine/Private/Dataflow/DataflowSubGraph.cpp
@@ -2,7 +2,9 @@
 
 #include "Dataflow/DataflowSubGraph.h"
 #include "Dataflow/DataflowEdNode.h"
+#include "Dataflow/DataflowObject.h"
 #include "Dataflow/DataflowSubGraphNodes.h"
+#include "GraphEditAction.h"
 
 namespace UE::Dataflow::DataflowSubGraph::Private
 {
@@ -50,14 +52,34 @@ FDataflowSubGraphOutputNode* UDataflowSubGraph::GetOutputNode() const
 	return UE::Dataflow::DataflowSubGraph::Private::GetNodeByType<FDataflowSubGraphOutputNode>(*this);
 }
 
+bool UDataflowSubGraph::IsForEachSubGraph() const
+{
+	return bIsForEach;
+}
+
+void UDataflowSubGraph::SetForEachSubGraph(bool bValue)
+{
+	if (bValue != bIsForEach)
+	{
+		bIsForEach = bValue;
+
+		Modify();
+		if (UDataflow* DataflowAsset = UDataflow::GetDataflowAssetFromEdGraph(this))
+		{
+			DataflowAsset->Modify();
+			FDataflowAssetDelegates::OnSubGraphsChanged.Broadcast(DataflowAsset, SubGraphGuid, UE::Dataflow::ESubGraphChangedReason::ChangedType);
+		}
+	}
+}
+
 void UDataflowSubGraph::PostLoad()
 {
 	Super::PostLoad();
 	
 	bIsLoaded = true;
+
 	if (FDataflowSubGraphDelegates::OnSubGraphLoaded.IsBound())
 	{
 		FDataflowSubGraphDelegates::OnSubGraphLoaded.Broadcast(*this);
 	}
 }
-
diff --git a/Engine/Source/Runtime/Experimental/Dataflow/Engine/Private/Dataflow/DataflowSubGraphNodes.cpp b/Engine/Source/Runtime/Experimental/Dataflow/Engine/Private/Dataflow/DataflowSubGraphNodes.cpp
--- a/Engine/Source/Runtime/Experimental/Dataflow/Engine/Private/Dataflow/DataflowSubGraphNodes.cpp
+++ b/Engine/Source/Runtime/Experimental/Dataflow/Engine/Private/Dataflow/DataflowSubGraphNodes.cpp
@@ -87,13 +87,9 @@ namespace UE::Dataflow::SubGraphNodes::Private
 						if (Evaluate(*OutputToEvaluate))
 						{
 							// copy the cache value from this cachestore to the parent one 
-							const TUniquePtr<FContextCacheElementBase>* CacheEntryToClone = GetDataImpl(OutputToEvaluate->CacheKey());
-							if (CacheEntryToClone && CacheEntryToClone->IsValid())
-							{
-								TUniquePtr<FContextCacheElementBase> ClonedCacheEntry = (*CacheEntryToClone)->Clone(*this);
-								ParentContext.SetDataImpl(CallerOutput.CacheKey(), MoveTemp(ClonedCacheEntry));
-								bEvaluationSuccess = true;
-							}
+							bEvaluationSuccess = CopyDataToAnotherContext(OutputToEvaluate->CacheKey(), ParentContext, CallerOutput.CacheKey(), 
+								CallerOutput.GetProperty(), CallerOutput.GetOwningNodeGuid(), CallerOutput.GetOwningNodeValueHash(), CallerOutput.GetOwningNodeTimestamp()
+							);
 						}
 
 						if (bHasError)
@@ -121,7 +117,11 @@ namespace UE::Dataflow::SubGraphNodes::Private
 			SubGraphCallback.EvaluateParentContext(ParentContext, *this, SubGraphOutput);
 		}
 
+		void SetCurrentArrayIndex(int32 Index) { CurrentArrayIndex = Index; }
+		int32 GetCurrentArrayIndex() const { return CurrentArrayIndex; }
+
 	private:
+		int32 CurrentArrayIndex = 0;
 		FContext& ParentContext;
 		const UE::Dataflow::ISubGraphContextCallback& SubGraphCallback;
 	};
@@ -235,6 +235,7 @@ namespace UE::Dataflow
 	{
 		DATAFLOW_NODE_REGISTER_CREATION_FACTORY(FDataflowSubGraphInputNode);
 		DATAFLOW_NODE_REGISTER_CREATION_FACTORY(FDataflowSubGraphOutputNode);
+		DATAFLOW_NODE_REGISTER_CREATION_FACTORY(FDataflowSubGraphGetCurrentIndexNode);
 		DATAFLOW_NODE_REGISTER_CREATION_FACTORY(FDataflowCallSubGraphNode);
 	}
 }
@@ -340,6 +341,27 @@ const FDataflowConnection* FDataflowSubGraphOutputNode::OnDropConnectionOnNode(c
 	return nullptr;
 }
 
+//////////////////////////////////////////////////////////////////////////////////////////////////////////
+FDataflowSubGraphGetCurrentIndexNode::FDataflowSubGraphGetCurrentIndexNode(const UE::Dataflow::FNodeParameters& Param, FGuid InGuid)
+	: FDataflowNode(Param, InGuid)
+{
+	RegisterOutputConnection(&Index);
+}
+
+void FDataflowSubGraphGetCurrentIndexNode::Evaluate(UE::Dataflow::FContext& Context, const FDataflowOutput* Out) const
+{
+	using namespace UE::Dataflow::SubGraphNodes::Private;
+
+	if (FContextSubGraph* SubGraphContext = Context.AsType<FContextSubGraph>())
+	{
+		SetValue(Context, SubGraphContext->GetCurrentArrayIndex(), &Index);
+	}
+	else
+	{
+		SetValue(Context, 0, &Index);
+	}
+}
+
 //////////////////////////////////////////////////////////////////////////////////////////////////////////
 
 FDataflowCallSubGraphNode::FDataflowCallSubGraphNode(const UE::Dataflow::FNodeParameters& InParam, FGuid InGuid)
@@ -464,6 +486,9 @@ void FDataflowCallSubGraphNode::OnSomeSubGraphsChanged(const UDataflow* InDatafl
 
 		case UE::Dataflow::ESubGraphChangedReason::Deleted:
 			break;
+
+		case UE::Dataflow::ESubGraphChangedReason::ChangedType:
+			break;
 		}
 	}
 }
@@ -527,22 +552,49 @@ const FInstancedPropertyBag& FDataflowCallSubGraphNode::GetPropertyBag(const FDa
 
 void FDataflowCallSubGraphNode::EvaluateParentContext(UE::Dataflow::FContext& Context, UE::Dataflow::FContext& SubGraphContext, const FDataflowOutput& SubGraphOutput) const
 {
-	if (const FDataflowInput* Input = FindInput(SubGraphOutput.GetName()))
+	using namespace UE::Dataflow::SubGraphNodes::Private;
+
+	bool bIsFirstIteration = true;
+	if (FContextSubGraph* SubGraphContextPtr = SubGraphContext.AsType<FContextSubGraph>())
+	{
+		bIsFirstIteration = (SubGraphContextPtr->GetCurrentArrayIndex() == 0);
+	}
+
+	// Do we have a matching internal output?
+	const FDataflowOutput* InternalOutput = FindOutput(SubGraphOutput.GetName());
+
+	const bool bGetDataFromExternalContext = bIsFirstIteration || (InternalOutput == nullptr);
+
+	bool bDataCopied = false;
+	// First iteration : we pull from the external/caller graph
+	if (bGetDataFromExternalContext)
 	{
-		Input->PullValue(Context);
-		if (const FDataflowOutput* Output = Input->GetConnection())
+		if (const FDataflowInput* Input = FindInput(SubGraphOutput.GetName()))
 		{
-			const TUniquePtr<UE::Dataflow::FContextCacheElementBase>* CacheEntryToClone = Context.GetDataImpl(Output->CacheKey());
-			if (CacheEntryToClone && CacheEntryToClone->IsValid())
+			Input->PullValue(Context);
+			if (const FDataflowOutput* Output = Input->GetConnection())
 			{
-				TUniquePtr<UE::Dataflow::FContextCacheElementBase> ClonedCacheEntry = (*CacheEntryToClone)->Clone(Context);
-				SubGraphContext.SetDataImpl(SubGraphOutput.CacheKey(), MoveTemp(ClonedCacheEntry));
-				return;
+				bDataCopied = Context.CopyDataToAnotherContext(Output->CacheKey(), SubGraphContext, SubGraphOutput.CacheKey(),
+					SubGraphOutput.GetProperty(), SubGraphOutput.GetOwningNodeGuid(), SubGraphOutput.GetOwningNodeValueHash(), SubGraphOutput.GetOwningNodeTimestamp()
+				);
 			}
 		}
 	}
-	// default pass when everything else failed
-	SubGraphOutput.SetNullValue(SubGraphContext);
+	else if (InternalOutput) // Second and subsequent iterations : copy from the previous outputs
+	{
+		// No need to pull the data since the output has already been evaluated
+		// todo(dataflow) : this is a fair assumption but feels there could be special cases
+		//					but also pulling the data here could result in weird infinite loops
+		bDataCopied = Context.CopyDataToAnotherContext(InternalOutput->CacheKey(), SubGraphContext, SubGraphOutput.CacheKey(),
+			SubGraphOutput.GetProperty(), SubGraphOutput.GetOwningNodeGuid(), SubGraphOutput.GetOwningNodeValueHash(), SubGraphOutput.GetOwningNodeTimestamp()
+		);
+	}
+
+	if (!bDataCopied)
+	{
+		// Default value when everything else failed
+		SubGraphOutput.SetNullValue(SubGraphContext);
+	}
 }
 
 void FDataflowCallSubGraphNode::Evaluate(UE::Dataflow::FContext& Context, const FDataflowOutput* Out) const
@@ -557,9 +609,45 @@ void FDataflowCallSubGraphNode::Evaluate(UE::Dataflow::FContext& Context, const
 				{
 					// pulling the outputs and when the graph evaluate the remote left end of its graph OnInputValueRequested will be called
 					UE::Dataflow::SubGraphNodes::Private::FContextSubGraph SubGraphContext(Context, *this);
-					SubGraphContext.EvaluateSubGraph(*SubGraphOutputNode, *this, *Out);
+					int32 NumIterations = 1;
+					if (SubGraphToCall->IsForEachSubGraph())
+					{
+						FDataflowInput* FirstArrayInput = nullptr;
+						for (FDataflowInput* Input : GetInputs())
+						{
+							if (FDataflowArrayTypePolicy::SupportsTypeStatic(Input->GetType()))
+							{
+								FirstArrayInput = Input;
+								break;
+							}
+						}
+						if (FirstArrayInput)
+						{
+							if (!FirstArrayInput->IsConnected())
+							{
+								NumIterations = 0;
+							}
+							else
+							{
+								FirstArrayInput->PullValue(Context);
+								check(FirstArrayInput->GetConnection());
+								const TUniquePtr<UE::Dataflow::FContextCacheElementBase>* CacheEntryToClone = Context.GetDataImpl(FirstArrayInput->GetConnection()->CacheKey());
+								if (CacheEntryToClone && CacheEntryToClone->IsValid())
+								{
+									NumIterations = (*CacheEntryToClone)->GetNumArrayElements(Context);
+								}
+							}
+						}
+					}
+					for (int32 Iteration = 0; Iteration < NumIterations; ++Iteration)
+					{
+						SubGraphContext.SetCurrentArrayIndex(Iteration);
+						SubGraphContext.EvaluateSubGraph(*SubGraphOutputNode, *this, *Out);
+						SubGraphContext.ClearAllData();
+					}
 				}
 			}
 		}
 	}
 }
+
diff --git a/Engine/Source/Runtime/Experimental/Dataflow/Engine/Public/Dataflow/DataflowObject.h b/Engine/Source/Runtime/Experimental/Dataflow/Engine/Public/Dataflow/DataflowObject.h
--- a/Engine/Source/Runtime/Experimental/Dataflow/Engine/Public/Dataflow/DataflowObject.h
+++ b/Engine/Source/Runtime/Experimental/Dataflow/Engine/Public/Dataflow/DataflowObject.h
@@ -37,6 +37,7 @@ namespace UE::Dataflow
 		Renamed,
 		Deleting,
 		Deleted,
+		ChangedType,
 	};
 }
 
diff --git a/Engine/Source/Runtime/Experimental/Dataflow/Engine/Public/Dataflow/DataflowSubGraph.h b/Engine/Source/Runtime/Experimental/Dataflow/Engine/Public/Dataflow/DataflowSubGraph.h
--- a/Engine/Source/Runtime/Experimental/Dataflow/Engine/Public/Dataflow/DataflowSubGraph.h
+++ b/Engine/Source/Runtime/Experimental/Dataflow/Engine/Public/Dataflow/DataflowSubGraph.h
@@ -27,6 +27,9 @@ public:
 	FDataflowSubGraphInputNode* GetInputNode() const;
 	FDataflowSubGraphOutputNode* GetOutputNode() const;
 
+	DATAFLOWENGINE_API bool IsForEachSubGraph() const;
+	DATAFLOWENGINE_API void SetForEachSubGraph(bool bValue);
+
 	bool IsLoaded() const { return bIsLoaded; }
 	const FGuid& GetSubGraphGuid() const { return SubGraphGuid; }
 
@@ -40,5 +43,8 @@ private:
 	UPROPERTY()
 	FGuid SubGraphGuid; 
 
+	UPROPERTY()
+	bool bIsForEach = false;
+
 	bool bIsLoaded : 1 = false;
 };
\ No newline at end of file
diff --git a/Engine/Source/Runtime/Experimental/Dataflow/Engine/Public/Dataflow/DataflowSubGraphNodes.h b/Engine/Source/Runtime/Experimental/Dataflow/Engine/Public/Dataflow/DataflowSubGraphNodes.h
--- a/Engine/Source/Runtime/Experimental/Dataflow/Engine/Public/Dataflow/DataflowSubGraphNodes.h
+++ b/Engine/Source/Runtime/Experimental/Dataflow/Engine/Public/Dataflow/DataflowSubGraphNodes.h
@@ -103,6 +103,31 @@ private:
 };
 
 
+/**
+* Get the current index in a subgraph
+* This is to be used in subgraph when iterating over an array
+*/
+USTRUCT(meta = (Icon = "GraphEditor.Function_16x"))
+struct FDataflowSubGraphGetCurrentIndexNode
+	: public FDataflowNode
+{
+	GENERATED_USTRUCT_BODY()
+
+	DATAFLOW_NODE_DEFINE_INTERNAL(FDataflowSubGraphGetCurrentIndexNode, "GetCurrentIndex", "SubGraph", "")
+
+public:
+	FDataflowSubGraphGetCurrentIndexNode(const UE::Dataflow::FNodeParameters& Param, FGuid InGuid = FGuid::NewGuid());
+
+private:
+	virtual void Evaluate(UE::Dataflow::FContext& Context, const FDataflowOutput* Out) const override;
+
+	UPROPERTY(Meta = (DataflowOutput))
+	int32 Index = 0;
+};
+
+/**
+* Call a subgraph
+*/
 USTRUCT(meta = (Icon = "GraphEditor.Function_16x"))
 struct FDataflowCallSubGraphNode
 	: public FDataflowNode
@@ -111,7 +136,7 @@ struct FDataflowCallSubGraphNode
 {
 	GENERATED_USTRUCT_BODY()
 
-	DATAFLOW_NODE_DEFINE_INTERNAL(FDataflowCallSubGraphNode, "CallSubGraph", "SubGraph", "")
+	DATAFLOW_NODE_DEFINE_INTERNAL(FDataflowCallSubGraphNode, "SubGraphCall", "SubGraph", "")
 
 public:
 	DATAFLOWENGINE_API FDataflowCallSubGraphNode(const UE::Dataflow::FNodeParameters& Param, FGuid InGuid = FGuid::NewGuid());
@@ -133,6 +158,7 @@ private:
 	virtual const FInstancedPropertyBag& GetPropertyBag(const FDataflowDynamicConnections* Caller) override;
 	/* End IOwnerInterface */
 
+
 	void UnregisterHandlers(UDataflowSubGraph* SubGraph);
 	void RegisterHandlers(UDataflowSubGraph* SubGraph);
 	void OnSubGraphChanged(const struct FEdGraphEditAction& InAction);
@@ -142,6 +168,8 @@ private:
 	void SyncInputsPropertyBagWithSubGraph();
 	void SyncOutputsPropertyBagWithSubGraph();
 
+	const FDataflowSubGraphOutputNode* GetSubGraphOutputNode() const;
+
 	bool IsValid() const;
 
 	UPROPERTY()
