Update Dataflow box rendering to properly support rotation.

diff --git a/Engine/Source/Runtime/Experimental/Dataflow/Engine/Private/Dataflow/DataflowDebugDraw.cpp b/Engine/Source/Runtime/Experimental/Dataflow/Engine/Private/Dataflow/DataflowDebugDraw.cpp
--- a/Engine/Source/Runtime/Experimental/Dataflow/Engine/Private/Dataflow/DataflowDebugDraw.cpp
+++ b/Engine/Source/Runtime/Experimental/Dataflow/Engine/Private/Dataflow/DataflowDebugDraw.cpp
@@ -201,7 +201,7 @@ void FDataflowDebugDraw::DrawMesh(const IDebugDrawMesh& Mesh) const
 
 void FDataflowDebugDraw::DrawBox(const FVector& Extents, const FQuat& Rotation, const FVector& Center, double UniformScale) const
 {
-	DebugRenderSceneProxy->Boxes.Add(FDebugRenderSceneProxy::FDebugBox(FBox(Center - UniformScale * Extents, Center + UniformScale * Extents), ColorWithTranslucency.ToFColor(true), DrawType, LineWidth));
+	DebugRenderSceneProxy->Boxes.Add(FDebugRenderSceneProxy::FDebugBox(FBox(-Extents, Extents), ColorWithTranslucency.ToFColor(true), FTransform(Rotation, Center, FVector(UniformScale)), DrawType, LineWidth));
 }
 
 void FDataflowDebugDraw::DrawSphere(const FVector& Center, double Radius) const
