[DSTM] Adding support to also migrate the game Thread State in full migrations.

Currently for full migrations we were only sending the PT state and applying it in both, the GT and PT state representation. The problem with doing that is that if we do a full migration back and forth between servers every frame, and some system writes to the GT state of that object, that GT data was getting lost in the subsequent migration.

diff --git a/Engine/Source/Runtime/Engine/Private/Components/PrimitiveComponent.cpp b/Engine/Source/Runtime/Engine/Private/Components/PrimitiveComponent.cpp
--- a/Engine/Source/Runtime/Engine/Private/Components/PrimitiveComponent.cpp
+++ b/Engine/Source/Runtime/Engine/Private/Components/PrimitiveComponent.cpp
@@ -1350,7 +1350,7 @@ void UPrimitiveComponent::MigratePhysicsState(FArchive& Ar)
 
 		if (CanSendMigratedPhysicsState(MigrationContext, ObjMigrationSendType))
 		{
-			const Chaos::ESerializedDataContext PhysicsDataContext = ObjMigrationSendType == EObjectMigrationSendType::ReassignOwnership ? Chaos::ESerializedDataContext::Internal : Chaos::ESerializedDataContext::External;
+			const Chaos::ESerializedDataContext PhysicsDataContext = ObjMigrationSendType == EObjectMigrationSendType::ReassignOwnership ? Chaos::ESerializedDataContext::Both : Chaos::ESerializedDataContext::External;
 			Chaos::FSerializedMultiPhysicsState MigratedState(PhysicsDataContext);
 
 			MigratedState.SetNumOfBodies(1);
@@ -1409,17 +1409,11 @@ void UPrimitiveComponent::ApplyMigratedPhysicsState(Chaos::FSerializedDataBuffer
 		
 		const Chaos::FSerializedMultiPhysicsState::FHeader& MigratedStateHeader = MigratedState.GetHeader();
 
-		if (MigratedStateHeader.IsValid())
+		if (ensure(MigratedStateHeader.IsValid()))
 		{
-			UE_LOG(LogSkeletalMesh, Verbose, TEXT("[%hs] Applying migrated physics state | Data context [%s]"), __func__, Chaos::LexToString(MigratedStateHeader.DataContext));
+			UE_LOG(LogPrimitiveComponent, Verbose, TEXT("[%hs] Applying migrated physics state | Data context [%s]"), __func__, Chaos::LexToString(MigratedStateHeader.DataContext));
 
-			Chaos::FSerializedDataBufferPtr DataBuffer = MakeUnique<Chaos::FSerializedDataBuffer>();
-			MigratedState.ReadElementDataIntoBuffer(*DataBuffer);
-
-			// If we received an Internal physics state, we want to apply it to both, Internal and External so it is available immediatly in the game thread API
-			Chaos::ESerializedDataContext SerializationDataContext = MigratedStateHeader.DataContext == Chaos::ESerializedDataContext::Internal ?  Chaos::ESerializedDataContext::Both : MigratedStateHeader.DataContext;
-
-			BodyInstance.ApplySerializedState(MoveTemp(DataBuffer), SerializationDataContext);
+			MigratedState.ApplyDataToElement(BodyInstance);
 		}	
 	}
 }
diff --git a/Engine/Source/Runtime/Engine/Private/Components/SkeletalMeshComponent.cpp b/Engine/Source/Runtime/Engine/Private/Components/SkeletalMeshComponent.cpp
--- a/Engine/Source/Runtime/Engine/Private/Components/SkeletalMeshComponent.cpp
+++ b/Engine/Source/Runtime/Engine/Private/Components/SkeletalMeshComponent.cpp
@@ -3900,7 +3900,7 @@ void USkeletalMeshComponent::MigratePhysicsState(FArchive& Ar)
 
 		if (CanSendMigratedPhysicsState(MigrationContext, ObjMigrationSendType))
 		{
-			Chaos::FSerializedMultiPhysicsState MigratedState(ObjMigrationSendType == EObjectMigrationSendType::ReassignOwnership ? Chaos::ESerializedDataContext::Internal : Chaos::ESerializedDataContext::External);
+			Chaos::FSerializedMultiPhysicsState MigratedState(ObjMigrationSendType == EObjectMigrationSendType::ReassignOwnership ? Chaos::ESerializedDataContext::Both : Chaos::ESerializedDataContext::External);
 			MigratedState.SetNumOfBodies(Bodies.Num());
 			MigratedState.SetNumOfConstraints(Constraints.Num());
 
@@ -3969,29 +3969,17 @@ void USkeletalMeshComponent::ApplyMigratedPhysicsState(Chaos::FSerializedDataBuf
 
 		UE_LOG(LogSkeletalMesh, Verbose, TEXT("[%hs] Applying migrated physics state | Data context [%s]"), __func__, Chaos::LexToString( MigratedState.GetHeader().DataContext));
 
-		// If we received an Internal physics state, we want to apply it to both, Internal and External so it is available immediatly in the game thread API
-		Chaos::ESerializedDataContext SerializationDataContext = MigratedStateHeader.DataContext == Chaos::ESerializedDataContext::Internal ? Chaos::ESerializedDataContext::Both : MigratedStateHeader.DataContext;
-
 		for (FBodyInstance* BodyInstanceIt : Bodies)
 		{
-			Chaos::FSerializedDataBufferPtr DataBuffer = MakeUnique<Chaos::FSerializedDataBuffer>();
-			MigratedState.ReadElementDataIntoBuffer(*DataBuffer);
-
-			if (ensure(BodyInstanceIt))
-			{
-				BodyInstanceIt->ApplySerializedState(MoveTemp(DataBuffer), SerializationDataContext);
-			}
+			check(BodyInstanceIt);
+			MigratedState.ApplyDataToElement(*BodyInstanceIt);
 		}
 
 		for (FConstraintInstance* ConstraintInstance : Constraints)
 		{
-			Chaos::FSerializedDataBufferPtr DataBuffer = MakeUnique<Chaos::FSerializedDataBuffer>();
-			MigratedState.ReadElementDataIntoBuffer(*DataBuffer);
+			check(ConstraintInstance);
 
-			if (ensure(ConstraintInstance))
-			{
-				ConstraintInstance->ApplySerializedState(MoveTemp(DataBuffer), SerializationDataContext);
-			}
+			MigratedState.ApplyDataToElement(*ConstraintInstance);
 		}
 	}
 }
diff --git a/Engine/Source/Runtime/Experimental/Chaos/Public/Chaos/Serialization/SerializedMultiPhysicsState.h b/Engine/Source/Runtime/Experimental/Chaos/Public/Chaos/Serialization/SerializedMultiPhysicsState.h
--- a/Engine/Source/Runtime/Experimental/Chaos/Public/Chaos/Serialization/SerializedMultiPhysicsState.h
+++ b/Engine/Source/Runtime/Experimental/Chaos/Public/Chaos/Serialization/SerializedMultiPhysicsState.h
@@ -22,17 +22,54 @@ namespace Chaos
 
 		template<typename ElementType>
 		void AppendDataFromElement(ElementType& InElement)
+		{
+			if (Header.DataContext == ESerializedDataContext::Both)
+			{
+				AppendDataFromElement_Internal(InElement, ESerializedDataContext::Internal);
+				AppendDataFromElement_Internal(InElement, ESerializedDataContext::External);
+				return;
+			}
+			
+			AppendDataFromElement_Internal(InElement, Header.DataContext);
+		}
+
+		template<typename ElementType>
+		void ApplyDataToElement(ElementType& InElement)
+		{
+			if (Header.DataContext == ESerializedDataContext::Both)
+			{
+				ApplyDataToElement_Internal(InElement, ESerializedDataContext::Internal);
+				ApplyDataToElement_Internal(InElement, ESerializedDataContext::External);
+				return;
+			}
+			
+			ApplyDataToElement_Internal(InElement, Header.DataContext);
+		}
+
+	private:
+		template<typename ElementType>
+		void AppendDataFromElement_Internal(ElementType& InElement, ESerializedDataContext DataContext)
 		{
 			FDataEntryTag DataTag;
 			DataTag.DataOffset = GetSize();
-				
-			InElement.SerializeState(MigratedStateAsBytes, Header.DataContext);
+			
+			InElement.SerializeState(MigratedStateAsBytes, DataContext);
 
 			DataTag.DataSize = GetSize() - DataTag.DataOffset;
 
 			Header.DataTagPerElementIndex.Emplace(DataTag);
 		}
 
+		template<typename ElementType>
+		void ApplyDataToElement_Internal(ElementType& InElement, ESerializedDataContext DataContext)
+		{
+			FSerializedDataBufferPtr DataBuffer = MakeUnique<FSerializedDataBuffer>();
+			ReadElementDataIntoBuffer(*DataBuffer);
+			InElement.ApplySerializedState(MoveTemp(DataBuffer), DataContext);
+		}
+
+	public:
+
 		CHAOS_API void Serialize(FArchive& Ar);
 
 		int32 GetSize() const
@@ -82,7 +119,10 @@ namespace Chaos
 
 			bool IsValid() const
 			{
-				return DataContext != ESerializedDataContext::Invalid && NumBodies + NumConstraints == DataTagPerElementIndex.Num();
+				const int32 CurrentDataElementsNum = DataTagPerElementIndex.Num();
+				const int32 CurrentBodiesAndConstraintsNum = NumBodies + NumConstraints;
+				const int32 ExpectedDataElementsNum = DataContext == ESerializedDataContext::Both ? CurrentBodiesAndConstraintsNum * 2 : CurrentBodiesAndConstraintsNum;
+				return DataContext != ESerializedDataContext::Invalid && CurrentDataElementsNum == ExpectedDataElementsNum;
 			}
 		};
 
