FOnDemandThread now respects request priorities. IAS will use the IoDispatcher request priority value and for now IAD will use EIoDispatcherPriority::IoDispatcherPriority_High.

#rnx

- FOnDemandThread::IssueRequest now requires an initial priority to be provided for the request.
- FOnDemandThread::ReprioritizeRequest also requires a new priority value to be provided.
- In the old system we would set the priority on the FChunkRequest before updating the request queue, however this could potentially cause problems if another thread added a new request between the priority being updated and the queue being updated as the priority ordering would be incorrect for that brief time. By assigning the priority as part of the call to TThreadSafeIntrusiveQueue::Reprioritize we can ensure that this doesn't happen.

diff --git a/Engine/Source/Runtime/Experimental/IoStore/OnDemand/Private/OnDemandHttpThread.cpp b/Engine/Source/Runtime/Experimental/IoStore/OnDemand/Private/OnDemandHttpThread.cpp
--- a/Engine/Source/Runtime/Experimental/IoStore/OnDemand/Private/OnDemandHttpThread.cpp
+++ b/Engine/Source/Runtime/Experimental/IoStore/OnDemand/Private/OnDemandHttpThread.cpp
@@ -142,12 +142,12 @@ FOnDemandThread& FOnDemandThread::Get()
 	return *SingletonInstance;
 }
 
-FOnDemandThread::FRequestHandle FOnDemandThread::IssueRequest(const FOnDemandChunkInfo& ChunkInfo, const FIoOffsetAndLength& ReadRange, FCompletionCallback&& CompletionCallback)
+FOnDemandThread::FRequestHandle FOnDemandThread::IssueRequest(const FOnDemandChunkInfo& ChunkInfo, const FIoOffsetAndLength& ReadRange, int32 Priority, FCompletionCallback&& CompletionCallback)
 {
-	return IssueRequest(FOnDemandChunkInfo(ChunkInfo), ReadRange, MoveTemp(CompletionCallback));
+	return IssueRequest(FOnDemandChunkInfo(ChunkInfo), ReadRange, Priority, MoveTemp(CompletionCallback));
 }
 
-FOnDemandThread::FRequestHandle FOnDemandThread::IssueRequest(FOnDemandChunkInfo&& ChunkInfo, const FIoOffsetAndLength& ReadRange, FCompletionCallback&& CompletionCallback)
+FOnDemandThread::FRequestHandle FOnDemandThread::IssueRequest(FOnDemandChunkInfo&& ChunkInfo, const FIoOffsetAndLength& ReadRange, int32 Priority, FCompletionCallback&& CompletionCallback)
 {
 	TRACE_CPUPROFILER_EVENT_SCOPE(FOnDemandThread::IssueRequest);
 
@@ -155,6 +155,7 @@ FOnDemandThread::FRequestHandle FOnDemandThread::IssueRequest(FOnDemandChunkInfo
 
 	Request->ChunkRange = ReadRange;
 	Request->CompletionCallback = MoveTemp(CompletionCallback);
+	Request->Priority = Priority;
 
 	FOnDemandIoBackendStats::Get()->OnHttpEnqueue();
 
@@ -164,12 +165,12 @@ FOnDemandThread::FRequestHandle FOnDemandThread::IssueRequest(FOnDemandChunkInfo
 	return Request;
 }
 
-void FOnDemandThread::ReprioritizeRequest(FRequestHandle Request)
+void FOnDemandThread::ReprioritizeRequest(FRequestHandle Request, int32 NewPriority)
 {
 	TRACE_CPUPROFILER_EVENT_SCOPE(FOnDemandThread::ReprioritizeRequest);
 	if (Request != 0)
 	{
-		HttpRequests.Reprioritize(static_cast<FHttpRequest*>(Request));
+		HttpRequests.Reprioritize(static_cast<FHttpRequest*>(Request), NewPriority);
 	}
 }
 
diff --git a/Engine/Source/Runtime/Experimental/IoStore/OnDemand/Private/OnDemandHttpThread.h b/Engine/Source/Runtime/Experimental/IoStore/OnDemand/Private/OnDemandHttpThread.h
--- a/Engine/Source/Runtime/Experimental/IoStore/OnDemand/Private/OnDemandHttpThread.h
+++ b/Engine/Source/Runtime/Experimental/IoStore/OnDemand/Private/OnDemandHttpThread.h
@@ -50,28 +50,30 @@ public:
 	 * 
 	 * @param ChunkInfo		Info about the chunk to be read
 	 * @param ReadRange		The range from within the chunk to read. A default FIoOffsetAndLength will read the entire chunk.
+	 * @param Priority		The priority of the request. Currently maps to EIoDispatcherPriority but may be expanded on in the future.
 	 * @CompletionCallback	A callback that will be invoked when the http request has completed. This will be invoked on the
 	 *						thread that the http requests are being processed on and will block new requests being queued so
 	 *						care should be taken to do as little work in the callback as possible.
 	 * @return				A handle representing the request which can be used to modify it while the request is still in
 	 *						flight. The handle will no longer be valid after the CompletionCallback has been invoked.
 	 */
-	FRequestHandle IssueRequest(const FOnDemandChunkInfo& ChunkInfo, const FIoOffsetAndLength& ReadRange, FCompletionCallback&& CompletionCallback);
+	FRequestHandle IssueRequest(const FOnDemandChunkInfo& ChunkInfo, const FIoOffsetAndLength& ReadRange, int32 Priority, FCompletionCallback&& CompletionCallback);
 	
 	/**
 	 * Issue a http request to read from a chunk.
 	 *
 	 * @param ChunkInfo		Info about the chunk to be read
 	 * @param ReadRange		The range from within the chunk to read. A default FIoOffsetAndLength will read the entire chunk.
+	 * @param Priority		The priority of the request. Currently maps to EIoDispatcherPriority but may be expanded on in the future.
 	 * @CompletionCallback	A callback that will be invoked when the http request has completed. This will be invoked on the
 	 *						thread that the http requests are being processed on and will block new requests being queued so
 	 *						care should be taken to do as little work in the callback as possible.
 	 * @return				A handle representing the request which can be used to modify it while the request is still in
 	 *						flight. The handle will no longer be valid after the CompletionCallback has been invoked.
 	 */
-	FRequestHandle IssueRequest(FOnDemandChunkInfo&& ChunkInfo, const FIoOffsetAndLength& ReadRange, FCompletionCallback&& CompletionCallback);
+	FRequestHandle IssueRequest(FOnDemandChunkInfo&& ChunkInfo, const FIoOffsetAndLength& ReadRange, int32 Priority, FCompletionCallback&& CompletionCallback);
 
-	void ReprioritizeRequest(FRequestHandle Request);
+	void ReprioritizeRequest(FRequestHandle Request, int32 NewPriority);
 	void CancelRequest(FRequestHandle Request);
 
 	FOnTickIdle& OnTickIdle()
diff --git a/Engine/Source/Runtime/Experimental/IoStore/OnDemand/Private/OnDemandIoDispatcherBackend.cpp b/Engine/Source/Runtime/Experimental/IoStore/OnDemand/Private/OnDemandIoDispatcherBackend.cpp
--- a/Engine/Source/Runtime/Experimental/IoStore/OnDemand/Private/OnDemandIoDispatcherBackend.cpp
+++ b/Engine/Source/Runtime/Experimental/IoStore/OnDemand/Private/OnDemandIoDispatcherBackend.cpp
@@ -944,7 +944,7 @@ void FOnDemandIoBackend::ConditionallyStartBackendThread()
 void FOnDemandIoBackend::IssueHttpRequest(FChunkRequest* ChunkRequest)
 {
 #if UE_ENABLE_IAS_REQUEST_THREAD
-	ChunkRequest->HttpRequest = FOnDemandThread::Get().IssueRequest(ChunkRequest->Params.ChunkInfo, ChunkRequest->Params.ChunkRange,
+	ChunkRequest->HttpRequest = FOnDemandThread::Get().IssueRequest(ChunkRequest->Params.ChunkInfo, ChunkRequest->Params.ChunkRange, ChunkRequest->Priority,
 		[this, ChunkRequest](uint32 StatusCode, FStringView ErrorReason, FIoBuffer&& Data)
 		{
 			ChunkRequest->Chunk = MoveTemp(Data);
@@ -1146,10 +1146,10 @@ bool FOnDemandIoBackend::Resolve(FIoRequestImpl* Request)
 		if (bUpdatePriority)
 		{
 #if UE_ENABLE_IAS_REQUEST_THREAD
-			FOnDemandThread::Get().ReprioritizeRequest(ChunkRequest->HttpRequest);
+			FOnDemandThread::Get().ReprioritizeRequest(ChunkRequest->HttpRequest, ChunkRequest->Priority);
 #else
 			TRACE_CPUPROFILER_EVENT_SCOPE(IasBackend::UpdatePriorityForIoRequest);
-			HttpRequests.Reprioritize(ChunkRequest);
+			HttpRequests.Reprioritize(ChunkRequest, ChunkRequest->Priority);
 #endif // UE_ENABLE_IAS_REQUEST_THREAD
 		}
 		// The chunk for the request is already inflight 
@@ -1227,9 +1227,9 @@ void FOnDemandIoBackend::UpdatePriorityForIoRequest(FIoRequestImpl* Request)
 	if (FChunkRequest* ChunkRequest = ChunkRequests.TryUpdatePriority(Request))
 	{
 #if UE_ENABLE_IAS_REQUEST_THREAD
-		FOnDemandThread::Get().ReprioritizeRequest(ChunkRequest->HttpRequest);
+		FOnDemandThread::Get().ReprioritizeRequest(ChunkRequest->HttpRequest, ChunkRequest->Priority);
 #else
-		HttpRequests.Reprioritize(ChunkRequest);
+		HttpRequests.Reprioritize(ChunkRequest, ChunkRequest->Priority);
 #endif // UE_ENABLE_IAS_REQUEST_THREAD
 	}
 }
diff --git a/Engine/Source/Runtime/Experimental/IoStore/OnDemand/Private/OnDemandIoStore.cpp b/Engine/Source/Runtime/Experimental/IoStore/OnDemand/Private/OnDemandIoStore.cpp
--- a/Engine/Source/Runtime/Experimental/IoStore/OnDemand/Private/OnDemandIoStore.cpp
+++ b/Engine/Source/Runtime/Experimental/IoStore/OnDemand/Private/OnDemandIoStore.cpp
@@ -1683,7 +1683,7 @@ FOnDemandInstallResult FOnDemandIoStore::TickInstallRequest(FOnDemandInternalIns
 				++ConcurrentRequests;
 
 #if UE_ENABLE_IAD_REQUEST_THREAD
-				FOnDemandThread::Get().IssueRequest(FOnDemandChunkInfo(Kv.Key, ChunkEntry), FIoOffsetAndLength(),
+				FOnDemandThread::Get().IssueRequest(FOnDemandChunkInfo(Kv.Key, ChunkEntry), FIoOffsetAndLength(), EIoDispatcherPriority::IoDispatcherPriority_High,
 					[this, &ConcurrentRequests, ChunkEntry, &OutResult, &InstallRequest](uint32 StatusCode, FStringView ErrorReason, FIoBuffer&& Data)
 					{
 						--ConcurrentRequests;
diff --git a/Engine/Source/Runtime/Experimental/IoStore/OnDemand/Private/ThreadSafeIntrusiveQueue.h b/Engine/Source/Runtime/Experimental/IoStore/OnDemand/Private/ThreadSafeIntrusiveQueue.h
--- a/Engine/Source/Runtime/Experimental/IoStore/OnDemand/Private/ThreadSafeIntrusiveQueue.h
+++ b/Engine/Source/Runtime/Experimental/IoStore/OnDemand/Private/ThreadSafeIntrusiveQueue.h
@@ -46,10 +46,14 @@ public:
 		return Requests;
 	}
 
-	void Reprioritize(T* Request)
+	void Reprioritize(T* Request, int32 NewPriority)
 	{
+		check(Request != nullptr);
+
 		// Switch to double linked list/array if this gets too expensive
 		FScopeLock _(&CriticalSection);
+		
+		Request->Priority = NewPriority;
 		if (RemoveInternal(Request))
 		{
 			EnqueueByPriorityInternal(Request);
