IAx introspection, capture with --trace=iax

#rnx

diff --git a/Engine/Source/Runtime/Experimental/IoStore/HttpClient/Private/Activity.inl b/Engine/Source/Runtime/Experimental/IoStore/HttpClient/Private/Activity.inl
--- a/Engine/Source/Runtime/Experimental/IoStore/HttpClient/Private/Activity.inl
+++ b/Engine/Source/Runtime/Experimental/IoStore/HttpClient/Private/Activity.inl
@@ -47,6 +47,15 @@ void FStopwatch::Impl(uint32 Index)
 
 #endif // IAS_HTTP_WITH_PERF
 
+////////////////////////////////////////////////////////////////////////////////
+static FLaneEstate* GActivityTraceEstate = LaneEstate_New({
+	.Name = "Iax/Activity",
+	.Group = "Iax",
+	.Channel = GetIaxTraceChannel(),
+	.Weight = 11,
+});
+
+
 ////////////////////////////////////////////////////////////////////////////////
 struct FResponseInternal
 {
@@ -98,22 +107,10 @@ struct alignas(16) FActivity
 	FBuffer				Buffer;
 };
 
-////////////////////////////////////////////////////////////////////////////////
-static void Activity_TraceStateNames()
-{
-	FAnsiStringView StateNames[] = {
-		"None", "Build", "Send", "RecvMessage", "RecvStream", "RecvContent",
-		"RecvDone", "Completed", "Cancelled", "Failed", "$",
-	};
-	static_assert(UE_ARRAY_COUNT(StateNames) == int32(FActivity::EState::_Num) + 1);
-
-	TraceEnum(StateNames);
-}
-
 ////////////////////////////////////////////////////////////////////////////////
 static void Activity_ChangeState(FActivity* Activity, FActivity::EState InState, uint32 Param=0)
 {
-	Trace(Activity, ETrace::StateChange, InState);
+	Trace(Activity, ETrace::StateChange, uint32(InState));
 
 	check(Activity->State != InState);
 	Activity->State = InState;
@@ -155,6 +152,10 @@ static uint32 Activity_RemainingKiB(FActivity* Activity)
 ////////////////////////////////////////////////////////////////////////////////
 static void Activity_CallSink(FActivity* Activity)
 {
+	static uint32 Scope = LaneTrace_NewScope("Iax/Sink");
+	FLaneTrace* Lane = LaneEstate_Lookup(GActivityTraceEstate, Activity);
+	FLaneTraceScope _(Lane, Scope);
+
 	FTicketStatus& SinkArg = *(FTicketStatus*)Activity;
 	Activity->Sink(SinkArg);
 }
@@ -173,15 +174,13 @@ static FActivity* Activity_Alloc(uint32 BufferSize)
 	uint32 ScratchSize = BufferSize;
 	Activity->Buffer = FBuffer(Scratch, ScratchSize);
 
-	Activity_ChangeState(Activity, FActivity::EState::Build);
-
 	return Activity;
 }
 
 ////////////////////////////////////////////////////////////////////////////////
 static void Activity_Free(FActivity* Activity)
 {
-	Trace(Activity, ETrace::ActivityDestroy);
+	Trace(Activity, ETrace::ActivityDestroy, 0);
 
 	Activity->~FActivity();
 	FMemory::Free(Activity);
@@ -219,6 +218,77 @@ static void Activity_SetScore(FActivity* Activity, uint32 ContentSizeEst)
 	Activity->LengthScore = uint8(Pow2);
 }
 
+
+
+////////////////////////////////////////////////////////////////////////////////
+static void Trace(const struct FActivity* Activity, ETrace Action, uint32 Param)
+{
+	if (Action == ETrace::ActivityCreate)
+	{
+		static uint32 ActScopes[8] = {};
+		if (ActScopes[0] == 0)
+		{
+			ActScopes[0] = LaneTrace_NewScope("Iax/Activity");
+			ActScopes[1] = LaneTrace_NewScope("Iax/Activity_2");
+			ActScopes[2] = LaneTrace_NewScope("Iax/Activity_4");
+			ActScopes[3] = LaneTrace_NewScope("Iax/Activity_8");
+			ActScopes[4] = LaneTrace_NewScope("Iax/Activity_16");
+			ActScopes[5] = LaneTrace_NewScope("Iax/Activity_32");
+			ActScopes[6] = LaneTrace_NewScope("Iax/Activity_128");
+			ActScopes[7] = LaneTrace_NewScope("Iax/Activity_256");
+		}
+
+		FLaneTrace* Lane = LaneEstate_Build(GActivityTraceEstate, Activity);
+		LaneTrace_Enter(Lane, ActScopes[Activity->LengthScore]);
+		return;
+	}
+
+	if (Action == ETrace::ActivityDestroy)
+	{
+		LaneEstate_Demolish(GActivityTraceEstate, Activity);
+		return;
+	}
+
+	FLaneTrace* Lane = LaneEstate_Lookup(GActivityTraceEstate, Activity);
+
+	if (Action == ETrace::StateChange)
+	{
+		static constexpr FAnsiStringView StateNames[] = {
+			"Iax/None",
+			"Iax/Build",
+			"Iax/WaitForSocket",
+			"Iax/WaitResponse",
+			"Iax/RecvStream",
+			"Iax/RecvContent",
+			"Iax/RecvDone",
+			"Iax/Completed",
+			"Iax/Cancelled",
+			"Iax/Failed",
+		};
+		static_assert(UE_ARRAY_COUNT(StateNames) == uint32(FActivity::EState::_Num));
+		static uint32 StateScopes[UE_ARRAY_COUNT(StateNames)] = {};
+		if (StateScopes[0] == 0)
+		{
+			for (int32 i = 0; FAnsiStringView Name : StateNames)
+			{
+				StateScopes[i++] = LaneTrace_NewScope(Name);
+			}
+		}
+
+		uint32 Scope = StateScopes[Param];
+		if (Param == uint32(FActivity::EState::Build))
+		{
+			LaneTrace_Enter(Lane, Scope);
+		}
+		else
+		{
+			LaneTrace_Change(Lane, Scope);
+		}
+
+		return;
+	}
+}
+
 // }}}
 
 } // namespace UE::IoStore::HTTP
diff --git a/Engine/Source/Runtime/Experimental/IoStore/HttpClient/Private/Http.cpp b/Engine/Source/Runtime/Experimental/IoStore/HttpClient/Private/Http.cpp
--- a/Engine/Source/Runtime/Experimental/IoStore/HttpClient/Private/Http.cpp
+++ b/Engine/Source/Runtime/Experimental/IoStore/HttpClient/Private/Http.cpp
@@ -7,6 +7,7 @@
 #include <HAL/CriticalSection.h>
 #include <HAL/IConsoleManager.h>
 #include <IO/IoBuffer.h>
+#include <IO/Http/LaneTrace.h>
 #include <Math/UnrealMathUtility.h>
 #include <Memory/MemoryView.h>
 #include <Misc/ByteSwap.h>
diff --git a/Engine/Source/Runtime/Experimental/IoStore/HttpClient/Private/Loop.inl b/Engine/Source/Runtime/Experimental/IoStore/HttpClient/Private/Loop.inl
--- a/Engine/Source/Runtime/Experimental/IoStore/HttpClient/Private/Loop.inl
+++ b/Engine/Source/Runtime/Experimental/IoStore/HttpClient/Private/Loop.inl
@@ -10,7 +10,6 @@ namespace UE::IoStore::HTTP
 ////////////////////////////////////////////////////////////////////////////////
 static FOutcome DoSend(FActivity* Activity, FHttpPeer& Peer)
 {
-	Trace(Activity, ETrace::StateChange, Activity->State);
 
 #if IAS_HTTP_WITH_PERF
 	Activity->Stopwatch.SendStart();
@@ -67,7 +66,6 @@ static FOutcome DoRecvPeer(FActivity* Activity, FHttpPeer& Peer, int32& MaxRecvS
 		return FOutcome::Waiting();
 	}
 
-	Trace(Activity, ETrace::StateChange, Activity->State);
 
 	FMutableMemoryView DestView = Activity->Dest->GetMutableView();
 	char* Cursor = (char*)(DestView.GetData()) + Activity->StateParam;
@@ -107,7 +105,7 @@ static FOutcome DoRecvMessage(FActivity* Activity, FHttpPeer& Peer)
 	const char* MessageRight;
 	while (true)
 	{
-		Trace(Activity, ETrace::StateChange, Activity->State);
+		Trace(Activity, ETrace::StateChange, uint32(Activity->State));
 
 #if IAS_HTTP_WITH_PERF
 		Activity->Stopwatch.RecvStart();
@@ -1098,8 +1096,8 @@ int32 FHostGroup::Wait(const FTickState& State)
 		return 0;
 	}
 
-	Trace(0, ETrace::Wait);
-	ON_SCOPE_EXIT { Trace(0, ETrace::Unwait); };
+	Trace(ETrace::Wait);
+	ON_SCOPE_EXIT { Trace(ETrace::Unwait); };
 
 	// If the poll timeout is negative then treat that as a fatal timeout
 	check(State.FailTimeoutMs);
@@ -1242,7 +1240,8 @@ FRequest FEventLoop::FImpl::Request(
 	FAnsiStringView Path,
 	FActivity* Activity)
 {
-	Trace(Activity, ETrace::ActivityCreate, this);
+	Trace(Activity, ETrace::ActivityCreate, 0);
+	Activity_ChangeState(Activity, FActivity::EState::Build);
 
 	if (Path.Len() == 0)
 	{
diff --git a/Engine/Source/Runtime/Experimental/IoStore/HttpClient/Private/Misc.inl b/Engine/Source/Runtime/Experimental/IoStore/HttpClient/Private/Misc.inl
--- a/Engine/Source/Runtime/Experimental/IoStore/HttpClient/Private/Misc.inl
+++ b/Engine/Source/Runtime/Experimental/IoStore/HttpClient/Private/Misc.inl
@@ -8,92 +8,47 @@ namespace UE::IoStore::HTTP
 // {{{1 trace ..................................................................
 
 ////////////////////////////////////////////////////////////////////////////////
-#define MAKE_TRACE_ENUM(x) \
-	x(LoopCreate) \
-	x(LoopTick) \
-	x(LoopDestroy) \
-	x(ActivityCreate) \
-	x(ActivityDestroy) \
-	x(SocketCreate) \
-	x(SocketDestroy) \
-	x(RequestBegin) \
-	x(StateChange) \
-	x(Wait) \
-	x(Unwait) \
-	x(Connect) \
-	x(Send) \
-	x(Recv) \
-	x(StartWork) \
-	x($)
-
 enum class ETrace
 {
-#define TRACE_ENUM_VALUES(x) x,
-	MAKE_TRACE_ENUM(TRACE_ENUM_VALUES)
-#undef TRACE_ENUM_VALUES
-};
-
-static const FAnsiStringView TraceEnumNames[] = {
-#define TRACE_ENUM_VALUES(x) #x ,
-	MAKE_TRACE_ENUM(TRACE_ENUM_VALUES)
-#undef TRACE_ENUM_VALUES
+	LoopCreate,
+	LoopTick,
+	LoopDestroy,
+	ActivityCreate,
+	ActivityDestroy,
+	SocketCreate,
+	SocketDestroy,
+	RequestBegin,
+	StateChange,
+	Wait,
+	Unwait,
+	Connect,
+	Send,
+	Recv,
+	StartWork,
 };
 
-#undef MAKE_TRACE_ENUM
-
-UE_TRACE_CHANNEL(IasHttpChannel);
-
-UE_TRACE_EVENT_BEGIN(IasHttp, Enum, NoSync)
-	UE_TRACE_EVENT_FIELD(UE::Trace::AnsiString, Name)
-UE_TRACE_EVENT_END()
+////////////////////////////////////////////////////////////////////////////////
+#if UE_TRACE_ENABLED
 
-UE_TRACE_EVENT_BEGIN(IasHttp, Event, NoSync)
-	UE_TRACE_EVENT_FIELD(uint64, Cycle)
-	UE_TRACE_EVENT_FIELD(uint32, Id)
-	UE_TRACE_EVENT_FIELD(uint32, Param)
-	UE_TRACE_EVENT_FIELD(uint8, Action)
-UE_TRACE_EVENT_END()
+static void Trace(const struct FActivity*, ETrace, uint32);
+static void Trace(UPTRINT, ETrace Action, const class FOutcome* =nullptr);
 
-static void Activity_TraceStateNames();
+static void Trace(ETrace Action)					{}
+static void Trace(const void*, ETrace Action, ...)	{}
 
 ////////////////////////////////////////////////////////////////////////////////
-static void TraceEnum(const FAnsiStringView* Names)
+IOSTOREHTTPCLIENT_API const void* GetIaxTraceChannel()
 {
-	for (;; ++Names)
-	{
-		UE_TRACE_LOG(IasHttp, Enum, IasHttpChannel)
-			<< Enum.Name(Names[0].GetData(), Names[0].Len());
-
-		if (*Names == "$")
-		{
-			break;
-		}
-	}
+	UE_TRACE_CHANNEL(Iax);
+	return &Iax;
 }
 
-////////////////////////////////////////////////////////////////////////////////
-static void TraceInner(const UPTRINT Id, ETrace Action, UPTRINT Param)
-{
-#if UE_TRACE_ENABLED
-	static bool Once = [] {
-		TraceEnum(TraceEnumNames);
-		Activity_TraceStateNames();
-		return true;
-	}();
-
-	UE_TRACE_LOG(IasHttp, Event, IasHttpChannel)
-		<< Event.Cycle(FPlatformTime::Cycles64())
-		<< Event.Id(uint32(Id))
-		<< Event.Param(uint32(Param))
-		<< Event.Action(uint8(Action));
-#endif // UE_TRACE_ENABLED
-}
+#else
 
-template <typename T, typename U=UPTRINT>
-static void Trace(T Id, ETrace Action, U Param=0)
-{
-	TraceInner(UPTRINT(Id), Action, UPTRINT(Param));
-}
+static void							Trace(...)				{}
+IOSTOREHTTPCLIENT_API const void*	GetIaxTraceChannel()	{ return nullptr; }
+
+#endif // UE_TRACE_ENABLED
 
 
 
diff --git a/Engine/Source/Runtime/Experimental/IoStore/HttpClient/Private/Socket.inl b/Engine/Source/Runtime/Experimental/IoStore/HttpClient/Private/Socket.inl
--- a/Engine/Source/Runtime/Experimental/IoStore/HttpClient/Private/Socket.inl
+++ b/Engine/Source/Runtime/Experimental/IoStore/HttpClient/Private/Socket.inl
@@ -434,7 +434,11 @@ FOutcome FSocket::Connect(uint32 IpAddress, uint32 Port)
 {
 	check(IsValid());
 
-	Trace(Socket, ETrace::Connect, IpAddress);
+	Trace(Socket, ETrace::Connect);
+	auto Return = [this] (FOutcome Outcome) {
+		Trace(Socket, ETrace::Connect, &Outcome);
+		return Outcome;
+	};
 
 	IpAddress = Socket_HtoNl(IpAddress);
 
@@ -447,15 +451,15 @@ FOutcome FSocket::Connect(uint32 IpAddress, uint32 Port)
 
 	if (IsSocketResult(EWOULDBLOCK) | IsSocketResult(EINPROGRESS))
 	{
-		return FOutcome::Waiting();
+		return Return(FOutcome::Waiting());
 	}
 
 	if (Result < 0)
 	{
-		return FOutcome::Error("Socket connect failed", Result);
+		return Return(FOutcome::Error("Socket connect failed", Result));
 	}
 
-	return FOutcome::Ok();
+	return Return(FOutcome::Ok());
 }
 
 ////////////////////////////////////////////////////////////////////////////////
@@ -468,13 +472,17 @@ void FSocket::Disconnect()
 ////////////////////////////////////////////////////////////////////////////////
 FOutcome FSocket::Send(const char* Data, uint32 Size)
 {
-	Trace(Socket, ETrace::Send, -1);
+	Trace(Socket, ETrace::Send);
 	int32 Result = Socket_Send(Socket, Data, Size, MsgFlagType(MSG_NOSIGNAL));
-	Trace(Socket, ETrace::Send, FMath::Max(Result, 0));
 
-	if (Result > 0)						return FOutcome::Ok(Result);
-	if (Result == 0)					return FOutcome::Error("Send ATH0");
-	if (IsSocketResult(EWOULDBLOCK))	return FOutcome::Waiting();
+	auto Return = [this] (FOutcome Outcome) {
+		Trace(Socket, ETrace::Send, &Outcome);
+		return Outcome;
+	};
+
+	if (Result > 0)						return Return(FOutcome::Ok(Result));
+	if (Result == 0)					return Return(FOutcome::Error("Send ATH0"));
+	if (IsSocketResult(EWOULDBLOCK))	return Return(FOutcome::Waiting());
 
 	if (IsSocketResult(ENOTCONN))
 	{
@@ -484,29 +492,33 @@ FOutcome FSocket::Send(const char* Data, uint32 Size)
 		if (Result < 0 || Error != 0)
 		{
 			Result = (Result < 0) ? LastSocketResult() : Error;
-			return FOutcome::Error("Error while connecting", Result);
+			return Return(FOutcome::Error("Error while connecting", Result));
 		}
 
-		return FOutcome::Waiting();
+		return Return(FOutcome::Waiting());
 	}
 
 	Result = LastSocketResult();
-	return FOutcome::Error("Send", Result);
+	return Return(FOutcome::Error("Send", Result));
 }
 
 ////////////////////////////////////////////////////////////////////////////////
 FOutcome FSocket::Recv(char* Dest, uint32 Size)
 {
-	Trace(Socket, ETrace::Recv, -1);
+	Trace(Socket, ETrace::Recv);
 	int32 Result = Socket_Recv(Socket, Dest, Size, MsgFlagType(0));
-	Trace(Socket, ETrace::Recv, FMath::Max(0, Result));
 
-	if (Result > 0)						return FOutcome::Ok(Result);
-	if (Result == 0)					return FOutcome::Error("Recv ATH0");
-	if (IsSocketResult(EWOULDBLOCK))	return FOutcome::Waiting();
+	auto Return = [this] (FOutcome Outcome) {
+		Trace(Socket, ETrace::Recv, &Outcome);
+		return Outcome;
+	};
+
+	if (Result > 0)						return Return(FOutcome::Ok(Result));
+	if (Result == 0)					return Return(FOutcome::Error("Recv ATH0"));
+	if (IsSocketResult(EWOULDBLOCK))	return Return(FOutcome::Waiting());
 
 	Result = LastSocketResult();
-	return FOutcome::Error("Recv", Result);
+	return Return(FOutcome::Error("Recv", Result));
 }
 
 ////////////////////////////////////////////////////////////////////////////////
@@ -543,6 +555,79 @@ bool FSocket::SetRecvBufSize(int32 Size)
 	return 0 == Socket_SetSockOpt(Socket, SOL_SOCKET, SO_RCVBUF, &(char&)Size, int32(sizeof(Size)));
 }
 
+
+
+////////////////////////////////////////////////////////////////////////////////
+static FLaneEstate* GSocketLaneEstate = LaneEstate_New({
+	.Name = "Iax/Socket",
+	.Group = "Iax",
+	.Channel = GetIaxTraceChannel(),
+	.Weight = 12,
+});
+
+////////////////////////////////////////////////////////////////////////////////
+static void Trace(UPTRINT Socket, ETrace Action, const FOutcome* Outcome)
+{
+	static uint32 SockScopeId = LaneTrace_NewScope("Iax/Socket");
+	static uint32 ConnScopeId = LaneTrace_NewScope("Iax/Connect");
+	static uint32 SendScopeId = LaneTrace_NewScope("Iax/Send");
+	static uint32 RecvScopeId = LaneTrace_NewScope("Iax/Recv");
+	static uint32 WaitScopeId = LaneTrace_NewScope("Iax/Wait");
+	static uint32 IdleScopeId = LaneTrace_NewScope("Iax/Idle");
+
+	if (Action == ETrace::SocketCreate)
+	{
+		FLaneTrace* Lane = LaneEstate_Build(GSocketLaneEstate, Socket);
+		LaneTrace_Enter(Lane, SockScopeId);
+		LaneTrace_Enter(Lane, IdleScopeId);
+		return;
+	}
+
+	if (Action == ETrace::SocketDestroy)
+	{
+		LaneEstate_Demolish(GSocketLaneEstate, Socket);
+		return;
+	}
+
+	FLaneTrace* Lane = LaneEstate_Lookup(GSocketLaneEstate, Socket);
+
+	if (Action == ETrace::Send)
+	{
+		(Outcome == nullptr)
+			? LaneTrace_Enter(Lane, SendScopeId)
+			: LaneTrace_Leave(Lane);
+		return;
+	}
+
+	if (Action == ETrace::Connect)
+	{
+		if (Outcome == nullptr)
+		{
+			LaneTrace_Change(Lane, ConnScopeId);
+		}
+		else if (!Outcome->IsWaiting())
+		{
+			LaneTrace_Change(Lane, IdleScopeId);
+		}
+		return;
+	}
+
+	if (Outcome == nullptr)
+	{
+		uint32 ScopeId = (Action == ETrace::Send) ? SendScopeId : RecvScopeId;
+		LaneTrace_Change(Lane, RecvScopeId);
+		return;
+	}
+
+	if (Outcome->IsWaiting())
+	{
+		LaneTrace_Change(Lane, WaitScopeId);
+		return;
+	}
+
+	LaneTrace_Change(Lane, IdleScopeId);
+}
+
 // }}}
 
 } // namespace UE::IoStore::HTTP
diff --git a/Engine/Source/Runtime/Experimental/IoStore/OnDemand/Private/OnDemandIoDispatcherBackend.cpp b/Engine/Source/Runtime/Experimental/IoStore/OnDemand/Private/OnDemandIoDispatcherBackend.cpp
--- a/Engine/Source/Runtime/Experimental/IoStore/OnDemand/Private/OnDemandIoDispatcherBackend.cpp
+++ b/Engine/Source/Runtime/Experimental/IoStore/OnDemand/Private/OnDemandIoDispatcherBackend.cpp
@@ -12,6 +12,8 @@
 #include "HAL/PreprocessorHelpers.h"
 #include "HAL/Runnable.h"
 #include "HAL/RunnableThread.h"
+#include "IO/Http/LaneTrace.h"
+#include "IO/IoAllocators.h"
 #include "IO/IoAllocators.h"
 #include "IO/IoChunkEncoding.h"
 #include "IO/IoDispatcher.h"
@@ -60,6 +62,33 @@
 namespace UE::IoStore
 {
 
+////////////////////////////////////////////////////////////////////////////////
+namespace HTTP {
+
+IOSTOREHTTPCLIENT_API const void* GetIaxTraceChannel();
+
+}
+
+FLaneEstate* GRequestLaneEstate = LaneEstate_New({
+	.Name = "Iax/Backend",
+	.Group = "Iax",
+	.Channel = HTTP::GetIaxTraceChannel(),
+	.Weight = 10,
+});
+
+#if UE_TRACE_ENABLED
+
+static void Trace(
+	bool bIsPiggyback,
+	const FIoChunkId& ChunkId,
+	const struct FChunkRequestParams* Params);
+
+#else
+
+static void Trace(...) {}
+
+#endif // UE_TRACE_ENABLED
+
 ///////////////////////////////////////////////////////////////////////////////
 
 /** Note that GIasHttpPrimaryEndpoint has no effect after initial start up */
@@ -348,10 +377,26 @@ struct FChunkRequest
 		, bCached(false)
 	{
 		check(Request && NextRequest == nullptr);
+
+		FLaneTrace* Lane = LaneEstate_Build(GRequestLaneEstate, this);
+		static uint32 LaneScopeId = LaneTrace_NewScope("Iax/Request");
+		LaneTrace_Enter(Lane, LaneScopeId);
+	}
+
+	~FChunkRequest()
+	{
+		LaneScope = FLaneTraceScope();
+		LaneEstate_Demolish(GRequestLaneEstate, this);
 	}
 
 	bool AddDispatcherRequest(FIoRequestImpl* Request)
 	{
+		/* disabled for the moment as closing of these scopes is little off
+		FLaneTrace* Lane = LaneEstate_Lookup(GRequestLaneEstate, this);
+		static uint32 LaneScopeId = LaneTrace_NewScope("Iax/Piggyback");
+		LaneTrace_Enter(Lane, LaneScopeId);
+		*/
+
 		check(RequestHead && RequestTail);
 		check(Request && !Request->NextRequest);
 
@@ -449,6 +494,7 @@ struct FChunkRequest
 	bool bCached;
 	bool bCancelled = false;
 	EIoErrorCode CacheGetStatus;
+	FLaneTraceScope LaneScope;
 };
 
 //////////////////////////////////////////////////////////////////////////////
@@ -577,6 +623,8 @@ class FOnDemandIoBackend final
 				bOutPending = true;
 				bOutUpdatePriority = ChunkRequest->AddDispatcherRequest(Request);
 
+				Trace(true, Request->ChunkId, &Params);
+
 				return ChunkRequest;
 			}
 
@@ -585,6 +633,8 @@ class FOnDemandIoBackend final
 			ChunkRequestCount++;
 			Inflight.Add(Params.ChunkKey, ChunkRequest);
 
+			Trace(false, Request->ChunkId, &Params);
+
 			// Paranoid check to make sure that no host group has currently been assigned
 			check(ChunkRequest->HostGroup.GetHostUrls().IsEmpty());
 
@@ -972,6 +1022,8 @@ void FOnDemandIoBackend::IssueHttpRequest(FChunkRequest* ChunkRequest)
 
 void FOnDemandIoBackend::CompleteRequest(FChunkRequest* ChunkRequest)
 {
+	ChunkRequest->LaneScope = FLaneTraceScope();
+
 	LLM_SCOPE_BYTAG(Ias);
 	TRACE_CPUPROFILER_EVENT_SCOPE(IasBackend::CompleteRequest);
 	check(ChunkRequest != nullptr);
@@ -1110,6 +1162,9 @@ void FOnDemandIoBackend::CompleteCacheRead(FChunkRequest* ChunkRequest)
 		return;
 	}
 
+	static uint32 ScopeId = LaneTrace_NewScope("Iax/HttpGetAgain");
+	ChunkRequest->LaneScope.Change(ScopeId);
+
 	IssueHttpRequest(ChunkRequest);
 }
 
@@ -1177,6 +1232,10 @@ bool FOnDemandIoBackend::Resolve(FIoRequestImpl* Request)
 
 		if (GetStatus == EIoErrorCode::FileNotOpen)
 		{
+			static uint32 ScopeId = LaneTrace_NewScope("Iax/CacheRead");
+			FLaneTrace* Lane = LaneEstate_Lookup(GRequestLaneEstate, ChunkRequest);
+			ChunkRequest->LaneScope = FLaneTraceScope(Lane, ScopeId);
+
 			InflightCacheRequestCount.fetch_add(1, std::memory_order_relaxed);
 
 			FTaskEvent OnReadyEvent(TEXT("IasCacheMaterializeDone"));
@@ -1196,6 +1255,10 @@ bool FOnDemandIoBackend::Resolve(FIoRequestImpl* Request)
 
 	ChunkRequest->HostGroup = ChunkInfo.HostGroup();
 
+	FLaneTrace* Lane = LaneEstate_Lookup(GRequestLaneEstate, ChunkRequest);
+	static uint32 ScopeId = LaneTrace_NewScope("Iax/HttpGet");
+	ChunkRequest->LaneScope = FLaneTraceScope(Lane, ScopeId);
+
 	IssueHttpRequest(ChunkRequest);
 
 	return true;
@@ -1940,6 +2003,62 @@ TEST_CASE("IoStore::OnDemand::Ias::PriorityQueue", "[IoStoreOnDemand][Ias]")
 
 #endif // WITH_IOSTORE_ONDEMAND_TESTS
 
+
+
+#if UE_TRACE_ENABLED
+
+////////////////////////////////////////////////////////////////////////////////
+UE_TRACE_EVENT_BEGIN(Ias, ChunkRequest, NoSync)
+	UE_TRACE_EVENT_FIELD(uint64, Timestamp)
+	UE_TRACE_EVENT_FIELD(uint32, Offset)
+	UE_TRACE_EVENT_FIELD(uint32, Length)
+	UE_TRACE_EVENT_FIELD(uint64, Hash_A)
+	UE_TRACE_EVENT_FIELD(uint64, Hash_B)
+	UE_TRACE_EVENT_FIELD(uint32, Hash_C)
+	UE_TRACE_EVENT_FIELD(uint64, Id_A)
+	UE_TRACE_EVENT_FIELD(uint32, Id_B)
+	UE_TRACE_EVENT_FIELD(bool, IsPiggyback)
+UE_TRACE_EVENT_END()
+
+////////////////////////////////////////////////////////////////////////////////
+static void Trace(
+	bool bIsPiggyback,
+	const FIoChunkId& ChunkId,
+	const FChunkRequestParams* Params)
+{
+	struct FChunkIdDecomp
+	{
+		uint64 A;
+		uint32 B;
+		uint32 Pad;
+	};
+	static_assert(sizeof(FChunkIdDecomp) - sizeof(FChunkIdDecomp::Pad) == sizeof(FIoChunkId));
+	const auto& IdDecomp = *(FChunkIdDecomp*)(ChunkId.GetData());
+
+	struct FIoHashDecomp
+	{
+		uint64 A;
+		uint64 B;
+		uint32 C;
+		uint32 Pad;
+	};
+	static_assert(sizeof(FIoHashDecomp) - sizeof(FIoHashDecomp::Pad) == sizeof(FIoHash::ByteArray));
+	const auto& HashDecomp = *(FIoHashDecomp*)(Params->ChunkInfo.Hash().GetBytes());
+
+	UE_TRACE_LOG(Ias, ChunkRequest, HTTP::GetIaxTraceChannel())
+		<< ChunkRequest.Timestamp(FPlatformTime::Cycles64())
+		<< ChunkRequest.Offset(uint32(Params->ChunkRange.GetOffset()))
+		<< ChunkRequest.Length(uint32(Params->ChunkRange.GetLength()))
+		<< ChunkRequest.Hash_A(HashDecomp.A)
+		<< ChunkRequest.Hash_B(HashDecomp.B)
+		<< ChunkRequest.Hash_C(HashDecomp.C)
+		<< ChunkRequest.Id_A(IdDecomp.A)
+		<< ChunkRequest.Id_B(IdDecomp.B)
+		<< ChunkRequest.IsPiggyback(bIsPiggyback);
+}
+
+#endif // UE_TRACE_ENABLED
+
 } // namespace UE::IoStore
 
 #undef UE_ALLOW_DISABLE_CANCELLING
