Enable FOnDemandThread for IAS by default.

#rnx

diff --git a/Engine/Source/Runtime/Experimental/IoStore/OnDemand/Private/OnDemandHttpThread.h b/Engine/Source/Runtime/Experimental/IoStore/OnDemand/Private/OnDemandHttpThread.h
--- a/Engine/Source/Runtime/Experimental/IoStore/OnDemand/Private/OnDemandHttpThread.h
+++ b/Engine/Source/Runtime/Experimental/IoStore/OnDemand/Private/OnDemandHttpThread.h
@@ -3,10 +3,10 @@
 #pragma once
 
 #ifndef UE_ENABLE_REQUEST_THREAD
-	#define UE_ENABLE_REQUEST_THREAD 0
+	#define UE_ENABLE_REQUEST_THREAD 1
 #endif
 
-#define UE_ENABLE_IAS_REQUEST_THREAD (0 && UE_ENABLE_REQUEST_THREAD)
+#define UE_ENABLE_IAS_REQUEST_THREAD (1 && UE_ENABLE_REQUEST_THREAD)
 #define UE_ENABLE_IAD_REQUEST_THREAD (0 && UE_ENABLE_REQUEST_THREAD)
 
 #if UE_ENABLE_REQUEST_THREAD
