Read-type items' parameters were getting incorrectly considered as actual work items and the odd read here and there would get spuriously cancelled.

#rnx
#ushell-unshelve of 39205013

diff --git a/Engine/Source/Runtime/Experimental/IoStore/OnDemand/Private/JournaledCache.cpp b/Engine/Source/Runtime/Experimental/IoStore/OnDemand/Private/JournaledCache.cpp
--- a/Engine/Source/Runtime/Experimental/IoStore/OnDemand/Private/JournaledCache.cpp
+++ b/Engine/Source/Runtime/Experimental/IoStore/OnDemand/Private/JournaledCache.cpp
@@ -2103,23 +2103,26 @@ void FServiceThread::ReceiveWork()
 	PendingPrev = PendingLoad;
 
 	// Unregisters first
-	for (const FWork& Work : InboundWork)
+	for (uint32 i = 0, n = InboundWork.Num(); i < n; ++i)
 	{
+		const FWork& Work = InboundWork[i];
+
 		if (Work.What != FWork::Work_Unregister)
 		{
+			i += (Work.What == FWork::Work_Read); // Consume parameter work item too.
 			continue;
 		}
 
 		FCache* CachePtr = Work.Cache;
 		bool bFound = false;
-		for (int32 i = 0, n = Caches.Num(); i < n; ++i)
+		for (int32 j = 0, m = Caches.Num(); j < m; ++j)
 		{
-			if (Caches[i].Get() != CachePtr)
+			if (Caches[j].Get() != CachePtr)
 			{
 				continue;
 			}
 
-			Caches.RemoveAtSwap(i);
+			Caches.RemoveAtSwap(j);
 			bFound = true;
 			break;
 		}
@@ -2174,19 +2177,22 @@ void FServiceThread::ReceiveWork()
 	}
 
 	// ...and finally the cancels
-	for (const FWork& Work : InboundWork)
+	for (uint32 i = 0, n = InboundWork.Num(); i < n; ++i)
 	{
+		const FWork& Work = InboundWork[i];
+
 		if (Work.What != FWork::Work_Cancel)
 		{
+			i += (Work.What == FWork::Work_Read); // Consume parameter work item too.
 			continue;
 		}
 
 		const void* ToCancel = Work.Dest;
-		for (uint32 i = 1, n = ActiveReads.Num(); i < n; i += 2)
+		for (uint32 j = 1, m = ActiveReads.Num(); j < m; j += 2)
 		{
-			if (ActiveReads[i].Dest == ToCancel)
+			if (ActiveReads[j].Dest == ToCancel)
 			{
-				ActiveReads[i].Dest = nullptr;
+				ActiveReads[j].Dest = nullptr;
 				break;
 			}
 		}
