Fix a bug that was causing the  'IoStoreOnDemand.Http' thread to wake up even when there was no work to be done.

#rnx

- Debugging code was causing the thread to only sleep for one second after all requests have been processed instead of going to sleep until the thread has work to do.

diff --git a/Engine/Source/Runtime/Experimental/IoStore/OnDemand/Private/OnDemandHttpThread.cpp b/Engine/Source/Runtime/Experimental/IoStore/OnDemand/Private/OnDemandHttpThread.cpp
--- a/Engine/Source/Runtime/Experimental/IoStore/OnDemand/Private/OnDemandHttpThread.cpp
+++ b/Engine/Source/Runtime/Experimental/IoStore/OnDemand/Private/OnDemandHttpThread.cpp
@@ -215,7 +215,7 @@ uint32 FOnDemandThread::Run()
 
 			OnTickIdleDelegate.Broadcast();
 
-			TickThreadEvent->Wait(FHostGroupManager::Get().GetNumDisconnctedHosts() > 0 ? GIasHttpHealthCheckWaitTime : 1000/*MAX_uint32*/);
+			TickThreadEvent->Wait(FHostGroupManager::Get().GetNumDisconnctedHosts() > 0 ? GIasHttpHealthCheckWaitTime : MAX_uint32);
 		}
 	}
 
