Remove most cases of double locking between OnDemand iostore and package store.

diff --git a/Engine/Source/Runtime/Experimental/IoStore/OnDemand/Private/OnDemandIoStore.cpp b/Engine/Source/Runtime/Experimental/IoStore/OnDemand/Private/OnDemandIoStore.cpp
--- a/Engine/Source/Runtime/Experimental/IoStore/OnDemand/Private/OnDemandIoStore.cpp
+++ b/Engine/Source/Runtime/Experimental/IoStore/OnDemand/Private/OnDemandIoStore.cpp
@@ -70,29 +70,6 @@ namespace UE::IoStore
 namespace Private
 {
 
-///////////////////////////////////////////////////////////////////////////////
-// Enforce lock order
-struct FPackageStoreUpdateScope : public FOnDemandPackageStoreUpdateScope
-{
-	FPackageStoreUpdateScope(
-		UE::FMutex& InMutex,
-		IOnDemandPackageStoreBackend& InPackageStore,
-		TUniqueFunction<EOnDemandPackageStoreUpdateFlags()> InGetFlags = nullptr)
-		: FOnDemandPackageStoreUpdateScope(InPackageStore, MoveTemp(InGetFlags))
-		, Mutex(InMutex)
-	{
-		Mutex.Lock();
-	}
-
-	~FPackageStoreUpdateScope()
-	{
-		Mutex.Unlock();
-	}
-
-private:
-	UE::FMutex& Mutex;
-};
-
 ///////////////////////////////////////////////////////////////////////////////
 static FAnsiStringBuilderBase& GetChunkUrl(
 	const FAnsiStringView& Host,
@@ -809,14 +786,21 @@ FIoStatus FOnDemandIoStore::Unmount(FStringView MountId)
 	}
 
 	{
-		Private::FPackageStoreUpdateScope(ContainerMutex, *PackageStoreBackend);
+		TUniqueLock Lock(ContainerMutex);
 
-		Containers.SetNum(Algo::RemoveIf(Containers, [this, &MountId](const FSharedOnDemandContainer& Container)
+		bool bRemoved = false;
+		Containers.SetNum(Algo::RemoveIf(Containers, [this, &MountId, &bRemoved](const FSharedOnDemandContainer& Container)
 		{
 			if (Container->MountId == MountId)
 			{
+				ensureMsgf(Container->ChunkEntryReferences.IsEmpty(), 
+					TEXT("Container is still referenced when unmounting, ContainerName='%s', MountId='%s'"), 
+					*Container->Name, *Container->MountId);
+
 				UE_LOG(LogIoStoreOnDemand, Log, TEXT("Unmounting container, ContainerName='%s', MountId='%s'"),
-					*WriteToString<128>(Container->Name), *WriteToString<128>(Container->MountId));
+					*Container->Name, *Container->MountId);
+
+				bRemoved = true;
 
 				const FString ContainerName = Container->UniqueName();
 				PendingContainerHeaders.Remove(ContainerName);
@@ -826,6 +810,11 @@ FIoStatus FOnDemandIoStore::Unmount(FStringView MountId)
 
 			return false;
 		}));
+
+		if (bRemoved)
+		{
+			PackageStoreBackend->NeedsUpdate();
+		}
 	}
 
 	return EIoErrorCode::Ok;
@@ -1436,12 +1425,13 @@ FIoStatus FOnDemandIoStore::TickMountRequest(FMountRequest& MountRequest)
 	}
 
 	{
-		Private::FPackageStoreUpdateScope(ContainerMutex, *PackageStoreBackend);
-
+		UE::TUniqueLock Lock(ContainerMutex);
 		Containers.Append(MoveTemp(RequestedContainers));
 		PendingContainerHeaders.Append(MoveTemp(RequestedContainerHeaders));
 	}
 
+	PackageStoreBackend->NeedsUpdate();
+
 	return bAnyPendingEncryptionKey ? EIoErrorCode::PendingEncryptionKey : EIoErrorCode::Ok;
 }
 
@@ -1811,32 +1801,28 @@ FOnDemandInstallResult FOnDemandIoStore::TickInstallRequest(FOnDemandInternalIns
 		ContentHandle->IoStore = AsWeak(); 
 	}
 
-
-	if (InstallData.IsEmpty() == false)
+	for (const auto& Kv : InstallData)
 	{
-		Private::FPackageStoreUpdateScope(ContainerMutex, *PackageStoreBackend,
-			[]
-			{
-				// Referencers were added so the package store must be refreshed
-				return EOnDemandPackageStoreUpdateFlags::ReferencedPackages;
-			});
+		FOnDemandContainer& Container				= *Kv.Key;
+		const Private::FContainerInstallData& Data	= Kv.Value;
 
-		for (const auto& Kv : InstallData)
 		{
-			FOnDemandContainer& Container = *Kv.Key;
-			const Private::FContainerInstallData& Data = Kv.Value;
-
+			UE::TUniqueLock Lock(ContainerMutex);
+			FOnDemandChunkEntryReferences& References = Container.FindOrAddChunkEntryReferences(*ContentHandle);
+			for (const TPair<uint32, Private::EChunkCacheState>& ChunkPair : Data.ResolvedChunks)
 			{
-				FOnDemandChunkEntryReferences& References = Container.FindOrAddChunkEntryReferences(*ContentHandle);
-				for (const TPair<uint32, Private::EChunkCacheState>& ChunkPair : Data.ResolvedChunks)
-				{
-					const int32 EntryIndex = ChunkPair.Key;
-					References.Indices[EntryIndex] = true;
-				}
+				const int32 EntryIndex = ChunkPair.Key;
+				References.Indices[EntryIndex] = true;
 			}
 		}
 	}
 
+	if (InstallData.IsEmpty() == false)
+	{
+		// Referencers were added so the package store must be refreshed
+		PackageStoreBackend->NeedsUpdate(EOnDemandPackageStoreUpdateMode::ReferencedPackages);
+	}
+
 	double DurationInSeconds = FPlatformTime::Seconds() - StartTime;
 	#if !UE_BUILD_SHIPPING
 	if (GIoStoreOnDemandMinInstallDurationSeconds > 0.0f)
@@ -1942,12 +1928,9 @@ void FOnDemandIoStore::OnEncryptionKeyAdded(const FGuid& Id, const FAES::FAESKey
 {
 	TRACE_CPUPROFILER_EVENT_SCOPE(FOnDemandIoStore::OnEncryptionKeyAdded);
 
+	TUniqueLock Lock(ContainerMutex);
+
 	bool bAddedContainerHeaders = false;
-	Private::FPackageStoreUpdateScope(ContainerMutex, *PackageStoreBackend,
-		[&bAddedContainerHeaders]
-		{
-			return bAddedContainerHeaders ? EOnDemandPackageStoreUpdateFlags::Full : EOnDemandPackageStoreUpdateFlags::None;
-		});
 
 	for (FSharedOnDemandContainer& Container : Containers)
 	{
@@ -2001,6 +1984,11 @@ void FOnDemandIoStore::OnEncryptionKeyAdded(const FGuid& Id, const FAES::FAESKey
 			PendingContainerHeaders.RemoveByHash(ContainerNameHash, ContainerName);
 		}
 	}
+
+	if (bAddedContainerHeaders)
+	{
+		PackageStoreBackend->NeedsUpdate();
+	}
 }
 
 void FOnDemandIoStore::CreateContainersFromToc(
@@ -2241,12 +2229,7 @@ void FOnDemandIoStore::ReleaseContent(FOnDemandInternalContentHandle& ContentHan
 	LLM_SCOPE_BYNAME(TEXT("FileSystem/OnDemandIoStore"));
 	UE_LOG(LogIoStoreOnDemand, Log, TEXT("Releasing content handle, %s"), *LexToString(ContentHandle));
 
-	Private::FPackageStoreUpdateScope(ContainerMutex, *PackageStoreBackend, 
-		[]
-		{
-			// Package store needs to be aware of install state
-			return EOnDemandPackageStoreUpdateFlags::ReferencedPackages; 
-		});
+	UE::TUniqueLock Lock(ContainerMutex);
 
 	for (FSharedOnDemandContainer& Container : Containers)
 	{
@@ -2263,6 +2246,9 @@ void FOnDemandIoStore::ReleaseContent(FOnDemandInternalContentHandle& ContentHan
 					return Refs.ContentHandleId == ContentHandleId;
 				}));
 	}
+
+	// Package store needs to be aware of install state
+	PackageStoreBackend->NeedsUpdate(EOnDemandPackageStoreUpdateMode::ReferencedPackages);
 }
 
 void FOnDemandIoStore::GetReferencedContent(TArray<FSharedOnDemandContainer>& OutContainers, TArray<TBitArray<>>& OutChunkEntryIndices, bool bPackageStore)
diff --git a/Engine/Source/Runtime/Experimental/IoStore/OnDemand/Private/OnDemandPackageStoreBackend.cpp b/Engine/Source/Runtime/Experimental/IoStore/OnDemand/Private/OnDemandPackageStoreBackend.cpp
--- a/Engine/Source/Runtime/Experimental/IoStore/OnDemand/Private/OnDemandPackageStoreBackend.cpp
+++ b/Engine/Source/Runtime/Experimental/IoStore/OnDemand/Private/OnDemandPackageStoreBackend.cpp
@@ -49,12 +49,11 @@ public:
 						FOnDemandPackageStoreBackend(TWeakPtr<FOnDemandIoStore> OnDemandIoStore);
 						virtual ~FOnDemandPackageStoreBackend ();
 
-	virtual void		BeginNeedsUpdate() override;
-	virtual void		EndNeedsUpdate(EOnDemandPackageStoreUpdateFlags Flags) override;
+	virtual void		NeedsUpdate(EOnDemandPackageStoreUpdateMode Mode) override;
 
 	virtual void		OnMounted(TSharedRef<const FPackageStoreBackendContext> Context) override;
-	virtual void 		BeginRead() override;
-	virtual void 		EndRead() override;
+	virtual void		BeginRead() override;
+	virtual void		EndRead() override;
 
 	virtual bool		GetPackageRedirectInfo(
 							FPackageId PackageId,
@@ -69,8 +68,8 @@ public:
 	virtual TConstArrayView<uint32> GetSoftReferences(FPackageId PackageId, TConstArrayView<FPackageId>& OutPackageIds) override;
 
 private:
-	void				UpdateLookupTables();
-	void				UpdateReferencedPackages();
+	void				UpdateLookupTables(const TArray<FSharedOnDemandContainer>& AllContainers, const TArray<TBitArray<>>& ReferencedChunkEntryIndices);
+	void				UpdateReferencedPackages(const TArray<FSharedOnDemandContainer>& AllContainers, const TArray<TBitArray<>>& ReferencedChunkEntryIndices);
 
 	TWeakPtr<FOnDemandIoStore>			OnDemandIoStore;
 	TArray<FContainer>					Containers;
@@ -78,7 +77,7 @@ private:
 	FLocalizedMap						LocalizedMap;
 	FRedirectMap						RedirectMap;
 	UE::FMutex							Mutex;
-	EOnDemandPackageStoreUpdateFlags	NeedsUpdateFlags = EOnDemandPackageStoreUpdateFlags::Full;
+	std::atomic<EOnDemandPackageStoreUpdateMode> NeedsUpdateMode{ EOnDemandPackageStoreUpdateMode::Full };
 };
 
 ///////////////////////////////////////////////////////////////////////////////
@@ -91,15 +90,17 @@ FOnDemandPackageStoreBackend::~FOnDemandPackageStoreBackend()
 {
 }
 
-void FOnDemandPackageStoreBackend::BeginNeedsUpdate()
+void FOnDemandPackageStoreBackend::NeedsUpdate(EOnDemandPackageStoreUpdateMode Mode)
 {
-	Mutex.Lock();
-}
-
-void FOnDemandPackageStoreBackend::EndNeedsUpdate(EOnDemandPackageStoreUpdateFlags Flags)
-{
-	NeedsUpdateFlags |= Flags;
-	Mutex.Unlock();
+	if (Mode == EOnDemandPackageStoreUpdateMode::Full)
+	{
+		NeedsUpdateMode.store(EOnDemandPackageStoreUpdateMode::Full);
+	}
+	else if (Mode == EOnDemandPackageStoreUpdateMode::ReferencedPackages)
+	{
+		EOnDemandPackageStoreUpdateMode Expected = EOnDemandPackageStoreUpdateMode::None;
+		NeedsUpdateMode.compare_exchange_strong(Expected, EOnDemandPackageStoreUpdateMode::ReferencedPackages);
+	}
 }
 
 void FOnDemandPackageStoreBackend::OnMounted(TSharedRef<const FPackageStoreBackendContext> Context)
@@ -108,8 +109,33 @@ void FOnDemandPackageStoreBackend::OnMounted(TSharedRef<const FPackageStoreBacke
 
 void FOnDemandPackageStoreBackend::BeginRead()
 {
+	EOnDemandPackageStoreUpdateMode LocalNeedsUpdate = NeedsUpdateMode.exchange(EOnDemandPackageStoreUpdateMode::None);
+
+	if (LocalNeedsUpdate == EOnDemandPackageStoreUpdateMode::None)
+	{
+		Mutex.Lock();
+		return;
+	}
+
+	TArray<FSharedOnDemandContainer> AllContainers;
+	TArray<TBitArray<>> ReferencedChunkEntryIndices;
+
+	if (TSharedPtr<FOnDemandIoStore> PinOnDemandIoStore = OnDemandIoStore.Pin())
+	{
+		const bool bPackageStore = true;
+		PinOnDemandIoStore->GetReferencedContent(AllContainers, ReferencedChunkEntryIndices, bPackageStore);
+	}
+
 	Mutex.Lock();
-	UpdateLookupTables();
+
+	if (LocalNeedsUpdate == EOnDemandPackageStoreUpdateMode::Full)
+	{
+		UpdateLookupTables(AllContainers, ReferencedChunkEntryIndices);
+	}
+	else if (LocalNeedsUpdate == EOnDemandPackageStoreUpdateMode::ReferencedPackages)
+	{
+		UpdateReferencedPackages(AllContainers, ReferencedChunkEntryIndices);
+	}
 }
 
 void FOnDemandPackageStoreBackend::EndRead()
@@ -184,50 +210,18 @@ TConstArrayView<uint32> FOnDemandPackageStoreBackend::GetSoftReferences(FPackage
 	return TConstArrayView<uint32>();
 }
 
-void FOnDemandPackageStoreBackend::UpdateLookupTables()
+void FOnDemandPackageStoreBackend::UpdateLookupTables(const TArray<FSharedOnDemandContainer>& AllContainers, const TArray<TBitArray<>>& ReferencedChunkEntryIndices)
 {
-	// This must be locked before FOnDemandIoStore::ContainerMutex is locked in the call to GetReferencedContent below
-	// We need to maintain this lock order to not deadlock
 	check(Mutex.IsLocked());
 
-	EOnDemandPackageStoreUpdateFlags LocalNeedsUpdate = NeedsUpdateFlags;
-	NeedsUpdateFlags = EOnDemandPackageStoreUpdateFlags::None;
-
-	if (LocalNeedsUpdate == EOnDemandPackageStoreUpdateFlags::None)
-	{
-		return;
-	}
-
-	if (!EnumHasAnyFlags(LocalNeedsUpdate, EOnDemandPackageStoreUpdateFlags::Full))
+	int32 PackageCount = 0;
+	for (const FSharedOnDemandContainer& C : AllContainers)
 	{
-		if (EnumHasAnyFlags(LocalNeedsUpdate, EOnDemandPackageStoreUpdateFlags::ReferencedPackages))
-		{
-			UpdateReferencedPackages();
-		}
-
-		return;
+		PackageCount += C->Header->PackageIds.Num();
 	}
 
-	// Full Update
-
 	LocalizedMap.Empty();
 	RedirectMap.Empty();
-
-	TArray<FSharedOnDemandContainer> AllContainers;
-	TArray<TBitArray<>> ReferencedChunkEntryIndices;
-	int32 PackageCount = 0;
-
-	if (TSharedPtr<FOnDemandIoStore> PinOnDemandIoStore = OnDemandIoStore.Pin())
-	{
-		const bool bPackageStore = true;
-		PinOnDemandIoStore->GetReferencedContent(AllContainers, ReferencedChunkEntryIndices, bPackageStore);
-
-		for (const FSharedOnDemandContainer& C : AllContainers)
-		{
-			PackageCount += C->Header->PackageIds.Num();
-		}
-	}
-
 	Containers.Empty(AllContainers.Num());
 	EntryMap.Empty(PackageCount);
 
@@ -300,55 +294,33 @@ void FOnDemandPackageStoreBackend::UpdateLookupTables()
 	RedirectMap.Shrink();
 }
 
-void FOnDemandPackageStoreBackend::UpdateReferencedPackages()
+void FOnDemandPackageStoreBackend::UpdateReferencedPackages(const TArray<FSharedOnDemandContainer>& AllContainers, const TArray<TBitArray<>>& ReferencedChunkEntryIndices)
 {
-	TArray<FSharedOnDemandContainer> AllContainers;
-	TArray<TBitArray<>> ReferencedChunkEntryIndices;
-
-	if (TSharedPtr<FOnDemandIoStore> PinOnDemandIoStore = OnDemandIoStore.Pin())
-	{
-		const bool bPackageStore = true;
-		PinOnDemandIoStore->GetReferencedContent(AllContainers, ReferencedChunkEntryIndices, bPackageStore);
-	}
+	check(Mutex.IsLocked());
 
 	for (int32 ContainerIdx = 0; const FSharedOnDemandContainer& SharedContainer : AllContainers)
 	{
 		const FIoContainerHeader& Hdr = *SharedContainer->Header;
-		TConstArrayView<FFilePackageStoreEntry> Entries(
-			reinterpret_cast<const FFilePackageStoreEntry*>(Hdr.StoreEntries.GetData()),
-			Hdr.PackageIds.Num());
 
-		for (int32 EntryIdx = 0; const FFilePackageStoreEntry & Entry : Entries)
+		for (const FPackageId PkgId : Hdr.PackageIds)
 		{
-			const FPackageId PkgId = Hdr.PackageIds[EntryIdx];
-			const FIoChunkId PkgChunkId = CreatePackageDataChunkId(PkgId);
-			const int32 PkgChunkIdx = SharedContainer->FindChunkEntryIndex(PkgChunkId);
-			check(PkgChunkIdx != INDEX_NONE);
-			const bool IsReferenced = ReferencedChunkEntryIndices[ContainerIdx][PkgChunkIdx];
+			// Entry may not be present if this function is racing with a container unmount
+			// It will be removed on the next full update after the unmount completes
+			if (FOnDemandPackageStoreEntry* OnDemandEntry = EntryMap.Find(PkgId))
+			{
+				const FIoChunkId PkgChunkId = CreatePackageDataChunkId(PkgId);
+				const int32 PkgChunkIdx = SharedContainer->FindChunkEntryIndex(PkgChunkId);
+				check(PkgChunkIdx != INDEX_NONE);
+				const bool IsReferenced = ReferencedChunkEntryIndices[ContainerIdx][PkgChunkIdx];
 
-			EntryMap.FindChecked(PkgId).bIsInstalled = IsReferenced;
-			++EntryIdx;
+				OnDemandEntry->bIsInstalled = IsReferenced;
+			}
 		}
 
 		++ContainerIdx;
 	}
 }
 
-///////////////////////////////////////////////////////////////////////////////
-FOnDemandPackageStoreUpdateScope::FOnDemandPackageStoreUpdateScope(
-	IOnDemandPackageStoreBackend& InPackageStore, TUniqueFunction<EOnDemandPackageStoreUpdateFlags()> InGetFlags)
-	: PackageStore(InPackageStore)
-	, GetFlags(MoveTemp(InGetFlags))
-{
-	PackageStore.BeginNeedsUpdate();
-}
-
-FOnDemandPackageStoreUpdateScope::~FOnDemandPackageStoreUpdateScope()
-{
-	EOnDemandPackageStoreUpdateFlags Flags = GetFlags ? GetFlags() : EOnDemandPackageStoreUpdateFlags::Full;
-	PackageStore.EndNeedsUpdate(Flags);
-}
-
 ///////////////////////////////////////////////////////////////////////////////
 TSharedPtr<IOnDemandPackageStoreBackend> MakeOnDemandPackageStoreBackend(TWeakPtr<FOnDemandIoStore> OnDemandIoStore)
 {
diff --git a/Engine/Source/Runtime/Experimental/IoStore/OnDemand/Private/OnDemandPackageStoreBackend.h b/Engine/Source/Runtime/Experimental/IoStore/OnDemand/Private/OnDemandPackageStoreBackend.h
--- a/Engine/Source/Runtime/Experimental/IoStore/OnDemand/Private/OnDemandPackageStoreBackend.h
+++ b/Engine/Source/Runtime/Experimental/IoStore/OnDemand/Private/OnDemandPackageStoreBackend.h
@@ -13,13 +13,12 @@ using FSharedContainerHeader = TSharedPtr<FIoContainerHeader>;
 namespace UE::IoStore
 {
 
-enum class EOnDemandPackageStoreUpdateFlags : uint8
+enum class EOnDemandPackageStoreUpdateMode : uint8
 {
-	None				= 0,
-	Full				= (1 << 0),
-	ReferencedPackages	= (1 << 1),
+	None,
+	Full,
+	ReferencedPackages,
 };
-ENUM_CLASS_FLAGS(EOnDemandPackageStoreUpdateFlags);
 
 class IOnDemandPackageStoreBackend
 	: public IPackageStoreBackend
@@ -27,22 +26,7 @@ class IOnDemandPackageStoreBackend
 public:
 	virtual ~IOnDemandPackageStoreBackend() = default;
 
-	virtual void BeginNeedsUpdate() = 0;
-
-	virtual void EndNeedsUpdate(EOnDemandPackageStoreUpdateFlags Flags) = 0;
-};
-
-// Update wrapper to help enforce mutex locking order between this and FOnDemandIoStore
-struct FOnDemandPackageStoreUpdateScope
-{
-	FOnDemandPackageStoreUpdateScope(
-		IOnDemandPackageStoreBackend& InPackageStore, 
-		TUniqueFunction<EOnDemandPackageStoreUpdateFlags()> InGetFlags = nullptr);
-	~FOnDemandPackageStoreUpdateScope();
-	
-private:
-	IOnDemandPackageStoreBackend& PackageStore;
-	TUniqueFunction<EOnDemandPackageStoreUpdateFlags()> GetFlags;
+	virtual void NeedsUpdate(EOnDemandPackageStoreUpdateMode Mode = EOnDemandPackageStoreUpdateMode::Full) = 0;
 };
 
 TSharedPtr<IOnDemandPackageStoreBackend> MakeOnDemandPackageStoreBackend(TWeakPtr<FOnDemandIoStore> OnDemandIoStore);
