[Physics] Remove deprecated functions from 5.3 and the FChaos::FOverlapInfo type

#pf 67f81f27fb42c6412f51492c

diff --git a/Engine/Source/Runtime/Experimental/Chaos/Private/Chaos/PhysicsObjectInterface.cpp b/Engine/Source/Runtime/Experimental/Chaos/Private/Chaos/PhysicsObjectInterface.cpp
--- a/Engine/Source/Runtime/Experimental/Chaos/Private/Chaos/PhysicsObjectInterface.cpp
+++ b/Engine/Source/Runtime/Experimental/Chaos/Private/Chaos/PhysicsObjectInterface.cpp
@@ -564,46 +564,6 @@ namespace Chaos
 		}
 	}
 
-	template<EThreadContext Id>
-	bool FReadPhysicsObjectInterface<Id>::GetPhysicsObjectOverlap(const FConstPhysicsObjectHandle ObjectA, const FConstPhysicsObjectHandle ObjectB, bool bTraceComplex, Chaos::FOverlapInfo& OutOverlap)
-	{
-		TRACE_CPUPROFILER_EVENT_SCOPE(FReadPhysicsObjectInterface<Id>::GetPhysicsObjectOverlap);
-		FPhysicsObjectCollisionInterface Interface{ *this };
-		// This is slow and inefficient and hence deprecated.
-		bool bRetOverlap = false;
-		if (OutOverlap.MTD)
-		{
-			bRetOverlap |= Interface.PhysicsObjectOverlapWithMTD(ObjectA, FTransform::Identity, ObjectB, FTransform::Identity, bTraceComplex, *OutOverlap.MTD);
-		}
-
-		if (OutOverlap.AxisOverlap)
-		{
-			bRetOverlap |= Interface.PhysicsObjectOverlapWithAABB(ObjectA, FTransform::Identity, ObjectB, FTransform::Identity, bTraceComplex, FVector::Zero(), *OutOverlap.AxisOverlap);
-		}
-
-		return bRetOverlap;
-	}
-
-	template<EThreadContext Id>
-	bool FReadPhysicsObjectInterface<Id>::GetPhysicsObjectOverlapWithTransform(const FConstPhysicsObjectHandle ObjectA, const FTransform& InTransformA, const FConstPhysicsObjectHandle ObjectB, const FTransform& InTransformB, bool bTraceComplex, Chaos::FOverlapInfo& OutOverlap)
-	{
-		TRACE_CPUPROFILER_EVENT_SCOPE(FReadPhysicsObjectInterface<Id>::GetPhysicsObjectOverlapWithTransform);
-		FPhysicsObjectCollisionInterface Interface{ *this };
-		// This is slow and inefficient and hence deprecated.
-		bool bRetOverlap = false;
-		if (OutOverlap.MTD)
-		{
-			bRetOverlap |= Interface.PhysicsObjectOverlapWithMTD(ObjectA, InTransformA, ObjectB, InTransformB, bTraceComplex, *OutOverlap.MTD);
-		}
-
-		if (OutOverlap.AxisOverlap)
-		{
-			bRetOverlap |= Interface.PhysicsObjectOverlapWithAABB(ObjectA, InTransformA, ObjectB, InTransformB, bTraceComplex, FVector::Zero(), *OutOverlap.AxisOverlap);
-		}
-
-		return bRetOverlap;
-	}
-
 	template<EThreadContext Id>
 	bool FReadPhysicsObjectInterface<Id>::AreAllValid(TArrayView<const FConstPhysicsObjectHandle> InObjects)
 	{
diff --git a/Engine/Source/Runtime/Experimental/Chaos/Public/Chaos/PhysicsObjectInterface.h b/Engine/Source/Runtime/Experimental/Chaos/Public/Chaos/PhysicsObjectInterface.h
--- a/Engine/Source/Runtime/Experimental/Chaos/Public/Chaos/PhysicsObjectInterface.h
+++ b/Engine/Source/Runtime/Experimental/Chaos/Public/Chaos/PhysicsObjectInterface.h
@@ -47,12 +47,6 @@ namespace Chaos
 	class FPerShapeData;
 	struct FMTDInfo;
 
-	struct FOverlapInfo
-	{
-		FMTDInfo* MTD = nullptr;
-		FBox* AxisOverlap = nullptr;
-	};
-
 	/**
 	 * FReadPhysicsObjectInterface will assume that these operations are safe to call (i.e. the relevant scenes have been read locked on the game thread).
 	 */
@@ -90,12 +84,6 @@ namespace Chaos
 		// Returns true if a shape is found and we can stop iterating.
 		CHAOS_API void VisitEveryShape(TArrayView<const FConstPhysicsObjectHandle> InObjects, TFunctionRef<bool(const FConstPhysicsObjectHandle, TThreadShapeInstance<Id>*)> Lambda);
 
-		UE_DEPRECATED(5.3, "GetPhysicsObjectOverlap has been deprecated. Please use the function for the specific overlap metric you wish to compute instead in the FPhysicsObjectCollisionInterface.")
-		CHAOS_API bool GetPhysicsObjectOverlap(const FConstPhysicsObjectHandle ObjectA, const FConstPhysicsObjectHandle ObjectB, bool bTraceComplex, Chaos::FOverlapInfo& OutOverlap);
-
-		UE_DEPRECATED(5.3, "GetPhysicsObjectOverlapWithTransform has been deprecated. Please use the function for the specific overlap metric you wish to compute instead in the FPhysicsObjectCollisionInterface.")
-		CHAOS_API bool GetPhysicsObjectOverlapWithTransform(const FConstPhysicsObjectHandle ObjectA, const FTransform& InTransformA, const FConstPhysicsObjectHandle ObjectB, const FTransform& InTransformB, bool bTraceComplex, Chaos::FOverlapInfo& OutOverlap);
-
 		CHAOS_API bool AreAllValid(TArrayView<const FConstPhysicsObjectHandle> InObjects);
 		CHAOS_API bool AreAllKinematic(TArrayView<const FConstPhysicsObjectHandle> InObjects);
 		CHAOS_API bool AreAllSleeping(TArrayView<const FConstPhysicsObjectHandle> InObjects);
