Geometry collection component : fix issue with custom proxy collision transforms causing issues with character collisions

diff --git a/Engine/Source/Runtime/Experimental/GeometryCollectionEngine/Private/GeometryCollection/GeometryCollectionComponent.cpp b/Engine/Source/Runtime/Experimental/GeometryCollectionEngine/Private/GeometryCollection/GeometryCollectionComponent.cpp
--- a/Engine/Source/Runtime/Experimental/GeometryCollectionEngine/Private/GeometryCollection/GeometryCollectionComponent.cpp
+++ b/Engine/Source/Runtime/Experimental/GeometryCollectionEngine/Private/GeometryCollection/GeometryCollectionComponent.cpp
@@ -3995,10 +3995,10 @@ PRAGMA_ENABLE_DEPRECATION_WARNINGS
 					AddParams.bDoubleSided = BodySetup->bDoubleSidedGeometry;
 					AddParams.CollisionData = BodyCollisionData;
 					AddParams.CollisionTraceType = BodySetup->GetCollisionTraceFlag();
-					AddParams.WorldTransform = MeshLocalTransform * PhysicsProxy->GetWorldTransform_External();
+					AddParams.WorldTransform = PhysicsProxy->GetWorldTransform_External();
 					AddParams.Scale = AddParams.WorldTransform.GetScale3D();
 					AddParams.SimpleMaterial = GetPhysicalMaterial();
-					AddParams.LocalTransform = InToLocal;
+					AddParams.LocalTransform = InToLocal * MeshLocalTransform;
 
 					AddParams.Geometry = &BodySetup->AggGeom;
 					ChaosInterface::CreateGeometry(AddParams, OutGeoms, OutShapes);
