Fix IoStoreOnDemandTests to work with recent API changes.

#rnx

diff --git a/Engine/Source/Runtime/Experimental/IoStore/OnDemand/Private/OnDemandIoDispatcherBackend.cpp b/Engine/Source/Runtime/Experimental/IoStore/OnDemand/Private/OnDemandIoDispatcherBackend.cpp
--- a/Engine/Source/Runtime/Experimental/IoStore/OnDemand/Private/OnDemandIoDispatcherBackend.cpp
+++ b/Engine/Source/Runtime/Experimental/IoStore/OnDemand/Private/OnDemandIoDispatcherBackend.cpp
@@ -1939,8 +1939,7 @@ TEST_CASE("IoStore::OnDemand::Ias::PriorityQueue", "[IoStoreOnDemand][Ias]")
 		Queue.EnqueueByPriority(Request);
 		Queue.EnqueueByPriority(Alloc.Construct(2, 2));
 		Queue.EnqueueByPriority(Alloc.Construct(3, 3));
-		Request->Priority = 4;
-		Queue.Reprioritize(Request);
+		Queue.Reprioritize(Request, 4);
 
 		TArray<FTestRequest*> Requests = Dequeue(Queue);
 
@@ -1963,8 +1962,7 @@ TEST_CASE("IoStore::OnDemand::Ias::PriorityQueue", "[IoStoreOnDemand][Ias]")
 		Queue.EnqueueByPriority(Request);
 		Queue.EnqueueByPriority(Alloc.Construct(2, 2));
 		Queue.EnqueueByPriority(Alloc.Construct(3, 4));
-		Request->Priority = 3;
-		Queue.Reprioritize(Request);
+		Queue.Reprioritize(Request, 3);
 
 		TArray<FTestRequest*> Requests = Dequeue(Queue);
 
@@ -1987,8 +1985,7 @@ TEST_CASE("IoStore::OnDemand::Ias::PriorityQueue", "[IoStoreOnDemand][Ias]")
 		Queue.EnqueueByPriority(Alloc.Construct(2, 2));
 		FTestRequest* Request = Alloc.Construct(3, 3);
 		Queue.EnqueueByPriority(Request);
-		Request->Priority = 0;
-		Queue.Reprioritize(Request);
+		Queue.Reprioritize(Request, 0);
 
 		TArray<FTestRequest*> Requests = Dequeue(Queue);
 
