Extend TThreadSafeIntrusiveQueue::Dequeue functionality

#rnx

- TThreadSafeIntrusiveQueue can now accept the number of requests to dequeue rather than always returning the entire list.
- Added a new set of unit tests "IoStore::OnDemand::Ias::BasicQueue" to cover the basic functionality of TThreadSafeIntrusiveQueue
-- Moved FTestRequest out of one test so that any test may use it.

diff --git a/Engine/Source/Runtime/Experimental/IoStore/OnDemand/Private/OnDemandIoDispatcherBackend.cpp b/Engine/Source/Runtime/Experimental/IoStore/OnDemand/Private/OnDemandIoDispatcherBackend.cpp
--- a/Engine/Source/Runtime/Experimental/IoStore/OnDemand/Private/OnDemandIoDispatcherBackend.cpp
+++ b/Engine/Source/Runtime/Experimental/IoStore/OnDemand/Private/OnDemandIoDispatcherBackend.cpp
@@ -1838,24 +1838,199 @@ TEST_CASE("IoStore::OnDemand::Ias::Misc", "[IoStoreOnDemand][Ias]")
 	}
 }
 
-TEST_CASE("IoStore::OnDemand::Ias::PriorityQueue", "[IoStoreOnDemand][Ias]")
+struct FTestRequest
 {
-	struct FTestRequest
-	{
-		explicit FTestRequest(int32 InValue)
-			: Value(InValue)
-		{ }
-		explicit FTestRequest(int32 InValue, int32 InPriority)
-			: Value(InValue)
-			, Priority(InPriority)
-		{ }
-		int32			Value;
-		int32			Priority	= 0;
-		FTestRequest*	NextRequest	= nullptr;
-	};
+	FTestRequest() = default;
+	explicit FTestRequest(int32 InValue)
+		: Value(InValue)
+	{ }
+	explicit FTestRequest(int32 InValue, int32 InPriority)
+		: Value(InValue)
+		, Priority(InPriority)
+	{ }
+
+	int32			Value = 0;
+	int32			Priority = 0;
+	FTestRequest*	NextRequest = nullptr;
+};
+
+TEST_CASE("IoStore::OnDemand::Ias::BasicQueue", "[IoStoreOnDemand][Ias]")
+{
+	const int32 NumRequests = 10;
+	FTestRequest Requests[NumRequests];
+
+	for (int32 Index = 0; Index < NumRequests; ++Index)
+	{
+		Requests[Index].Value = Index;
+	}
+
+	SECTION("Basic")
+	{
+		TThreadSafeIntrusiveQueue<FTestRequest> Queue;
+
+		for (int32 Index = 0; Index < NumRequests; ++Index)
+		{
+			Queue.Enqueue(&Requests[Index]);
+		}
+
+		REQUIRE(Queue.Num() == NumRequests);
+
+		FTestRequest* ReturnedRequests = Queue.Dequeue();
+
+		REQUIRE(ReturnedRequests != nullptr);
+		REQUIRE(Queue.Num() == 0);
+
+		for (int32 Index = 0; Index < NumRequests; ++Index)
+		{
+			REQUIRE(ReturnedRequests->Value == Index);
+			ReturnedRequests = ReturnedRequests->NextRequest;
+		}
+
+		REQUIRE(ReturnedRequests == nullptr);
+
+		// Try removing from the empty list
+		FTestRequest* EmptyRequests = Queue.Dequeue();
+		REQUIRE(EmptyRequests == nullptr);
+		REQUIRE(Queue.Num() == 0);
+	}
+
+	SECTION("Advanced Dequeue")
+	{
+		TThreadSafeIntrusiveQueue<FTestRequest> Queue;
+
+		for (int32 Index = 0; Index < NumRequests; ++Index)
+		{
+			Queue.Enqueue(&Requests[Index]);
+		}
+
+		int32 NumRemaining = NumRequests;
+
+		// Remove a number of items from the list
+		{
+			const int32 NumToRemove = 3;
+			const int32 ValueOffset = NumRequests - NumRemaining;
+
+			FTestRequest* ReturnedRequests = Queue.Dequeue(NumToRemove);
+			NumRemaining -= NumToRemove;
+
+			REQUIRE(ReturnedRequests != nullptr);
+			REQUIRE(Queue.Num() == NumRemaining);
 
-	using FTestRequestAlloc	= TSingleThreadedSlabAllocator<FTestRequest>;
-	using FTestRequestQueue	= TThreadSafeIntrusiveQueue<FTestRequest>;
+			for (int32 Index = 0; Index < NumToRemove; ++Index)
+			{
+				REQUIRE(ReturnedRequests->Value == ValueOffset + Index);
+				ReturnedRequests = ReturnedRequests->NextRequest;
+			}
+
+			REQUIRE(ReturnedRequests == nullptr);
+		}
+
+		// Remove one more item from the list
+		{
+			const int32 NumToRemove = 1;
+			const int32 ValueOffset = NumRequests - NumRemaining;
+
+			FTestRequest* ReturnedRequests = Queue.Dequeue(NumToRemove);
+			NumRemaining -= NumToRemove;
+
+			REQUIRE(ReturnedRequests != nullptr);
+			REQUIRE(Queue.Num() == NumRemaining);
+
+			for (int32 Index = 0; Index < NumToRemove; ++Index)
+			{
+				REQUIRE(ReturnedRequests->Value == ValueOffset + Index);
+				ReturnedRequests = ReturnedRequests->NextRequest;
+			}
+
+			REQUIRE(ReturnedRequests == nullptr);
+		}
+
+		// Remove nothing from the list
+		{
+			const int32 OriginalSize = Queue.Num();
+			FTestRequest* ReturnedRequests = Queue.Dequeue(0);
+			REQUIRE(ReturnedRequests == nullptr);
+			REQUIRE(Queue.Num() == OriginalSize);
+		}
+
+		// Remove invalid value from the list (should do nothing)
+		{
+			const int32 OriginalSize = Queue.Num();
+			FTestRequest * ReturnedRequests = Queue.Dequeue(MIN_int32);
+			REQUIRE(ReturnedRequests == nullptr);
+			REQUIRE(Queue.Num() == OriginalSize);
+		}
+
+		// Now remove the remaining items in the list
+		{
+			const int32 ValueOffset = NumRequests - NumRemaining;
+
+			FTestRequest* ReturnedRequests = Queue.Dequeue();
+			REQUIRE(ReturnedRequests != nullptr);
+			REQUIRE(Queue.Num() == 0);
+
+			for (int32 Index = 0; Index < NumRemaining; ++Index)
+			{
+				REQUIRE(ReturnedRequests->Value == ValueOffset + Index);
+				ReturnedRequests = ReturnedRequests->NextRequest;
+			}
+
+			REQUIRE(ReturnedRequests == nullptr);
+		}
+
+		// Try removing from the empty list
+		FTestRequest* EmptyRequests = Queue.Dequeue();
+		REQUIRE(EmptyRequests == nullptr);
+		REQUIRE(Queue.Num() == 0);
+	}
+
+	SECTION("Precise Dequeue")
+	{
+		TThreadSafeIntrusiveQueue<FTestRequest> Queue;
+
+		Queue.Enqueue(&Requests[0]);
+		Queue.Enqueue(&Requests[1]);
+
+		// Dequeue the exact number of items that we have in the list
+		FTestRequest* ReturnedRequests = Queue.Dequeue(5);
+		REQUIRE(ReturnedRequests != nullptr);
+		REQUIRE(Queue.Num() == 0);
+
+		for (int32 Index = 0; Index < 2; ++Index)
+		{
+			REQUIRE(ReturnedRequests->Value == Index);
+			ReturnedRequests = ReturnedRequests->NextRequest;
+		}
+
+		REQUIRE(ReturnedRequests == nullptr);
+	}
+
+	SECTION("Greedy Dequeue")
+	{
+		TThreadSafeIntrusiveQueue<FTestRequest> Queue;
+
+		Queue.Enqueue(&Requests[0]);
+		Queue.Enqueue(&Requests[1]);
+
+		// Attempt to dequeue more items than we have in the list
+		FTestRequest* ReturnedRequests = Queue.Dequeue(5);
+		REQUIRE(ReturnedRequests != nullptr);
+		REQUIRE(Queue.Num() == 0);
+
+		for (int32 Index = 0; Index < 2; ++Index)
+		{
+			REQUIRE(ReturnedRequests->Value == Index);
+			ReturnedRequests = ReturnedRequests->NextRequest;
+		}
+
+		REQUIRE(ReturnedRequests == nullptr);
+	}
+}
+
+TEST_CASE("IoStore::OnDemand::Ias::PriorityQueue", "[IoStoreOnDemand][Ias]")
+{
+	using FTestRequestAlloc = TSingleThreadedSlabAllocator<FTestRequest>;
+	using FTestRequestQueue = TThreadSafeIntrusiveQueue<FTestRequest>;
 
 	auto Dequeue = [](FTestRequestQueue& Queue) -> TArray<FTestRequest*>
 	{
diff --git a/Engine/Source/Runtime/Experimental/IoStore/OnDemand/Private/ThreadSafeIntrusiveQueue.h b/Engine/Source/Runtime/Experimental/IoStore/OnDemand/Private/ThreadSafeIntrusiveQueue.h
--- a/Engine/Source/Runtime/Experimental/IoStore/OnDemand/Private/ThreadSafeIntrusiveQueue.h
+++ b/Engine/Source/Runtime/Experimental/IoStore/OnDemand/Private/ThreadSafeIntrusiveQueue.h
@@ -4,10 +4,16 @@
 
 #include "HAL/CriticalSection.h"
 #include "Misc/ScopeLock.h"
+#include "Templates/UnrealTemplate.h"
 
 namespace UE::IoStore
 {
 
+/**
+ * Note that the type used requires the following members:
+ *     NextRequest: A pointer to the same type, this is used to maintain the list.
+ *     Priority: An int32 value used to store the priority of the value
+ */
 template<typename T>
 class TThreadSafeIntrusiveQueue
 {
@@ -28,22 +34,64 @@ public:
 		}
 
 		Tail = Request;
+
+		++ListNum;
 	}
 
 	void EnqueueByPriority(T* Request)
 	{
 		FScopeLock _(&CriticalSection);
 		EnqueueByPriorityInternal(Request);
+
+		++ListNum;
 	}
 
-	T* Dequeue()
+	T* Dequeue(int32 NumToRemove = MAX_int32)
 	{
 		FScopeLock _(&CriticalSection);
 
-		T* Requests = Head;
-		Head = Tail = nullptr;
+		if (NumToRemove <= 0 || ListNum == 0)
+		{
+			return nullptr;
+		}
+		else if (NumToRemove == MAX_int32)
+		{
+			T* Requests = Head;
+			Head = Tail = nullptr;
+
+			ListNum = 0;
+
+			return Requests;
+		}
+		else
+		{
+			int32 NumRemoved = 1;
+			T* Requests = Head;
+			while(NumRemoved < NumToRemove && Requests->NextRequest != nullptr)
+			{
+				Requests = Requests->NextRequest;
+				NumRemoved++;
+			}
+
+			// Store the new head of the list for later
+			T* NewHead = Requests->NextRequest;
+
+			// Break off the subsection we are returning
+			Requests->NextRequest = nullptr;
 
-		return Requests;
+			// Clear the tail if the entire list was dequeued
+			if (Tail == Requests)
+			{
+				Tail = nullptr;
+			}
+
+			Swap(Head, NewHead);
+			ListNum -= NumRemoved;
+
+			return NewHead;
+		}
+
+		
 	}
 
 	void Reprioritize(T* Request, int32 NewPriority)
@@ -60,6 +108,12 @@ public:
 		}
 	}
 
+	int32 Num() const
+	{
+		FScopeLock _(&CriticalSection);
+		return ListNum;
+	}
+
 private:
 	void EnqueueByPriorityInternal(T* Request)
 	{
@@ -141,9 +195,12 @@ private:
 		return false;
 	}
 
-	FCriticalSection CriticalSection;
+	mutable FCriticalSection CriticalSection;
+	
 	T* Head = nullptr;
 	T* Tail = nullptr;
+
+	int32 ListNum = 0;
 };
 
 } //namespace UE::IoStore
