Asynchronous I/O store on-demand installs.

- moved processing of install request from FOnDemandIoStore to FOnDemandContentInstaller
- HTTP requests are now processed via HTTP thread
- updated stats with basic IAD metrics

#rnx

diff --git a/Engine/Source/Runtime/Experimental/IoStore/OnDemand/Private/OnDemandContentInstaller.cpp b/Engine/Source/Runtime/Experimental/IoStore/OnDemand/Private/OnDemandContentInstaller.cpp
new file mode 100644
--- /dev/null
+++ b/Engine/Source/Runtime/Experimental/IoStore/OnDemand/Private/OnDemandContentInstaller.cpp
@@ -0,0 +1,655 @@
+// Copyright Epic Games, Inc. All Rights Reserved.
+
+#include "OnDemandContentInstaller.h"
+
+#include "Async/UniqueLock.h"
+#include "Containers/RingBuffer.h"
+#include "HAL/Platform.h"
+#include "IO/IoBuffer.h"
+#include "IO/IoChunkId.h"
+#include "IO/IoDispatcher.h"
+#include "IO/IoStoreOnDemand.h"
+#include "IO/PackageStore.h"
+#include "IO/IoStatus.h"
+#include "OnDemandHttpThread.h"
+#include "OnDemandInstallCache.h"
+#include "OnDemandIoStore.h"
+#include "OnDemandPackageStoreBackend.h"
+#include "Statistics.h"
+
+namespace UE::IoStore
+{
+
+namespace Private
+{
+
+////////////////////////////////////////////////////////////////////////////////
+void ResolvePackageDependencies(
+	const TSet<FPackageId>& PackageIds,
+	bool bIncludeSoftReferences,
+	TSet<FPackageId>& OutResolved,
+	TSet<FPackageId>& OutMissing)
+{
+	TRACE_CPUPROFILER_EVENT_SCOPE(FOnDemandContentInstaller::ResolvePackageDependencies);
+
+	TRingBuffer<FPackageId> Queue;
+	TSet<FPackageId>		Visitied;
+
+	Visitied.Reserve(PackageIds.Num());
+	Queue.Reserve(PackageIds.Num());
+	for (const FPackageId& PackageId : PackageIds)
+	{
+		Queue.Add(PackageId);
+	}
+
+	FPackageStore& PackageStore = FPackageStore::Get();
+	FPackageStoreReadScope _(PackageStore);
+
+	while (!Queue.IsEmpty())
+	{
+		FPackageId PackageId = Queue.PopFrontValue();
+		{
+			FName		SourcePackageName;
+			FPackageId	RedirectedToPackageId;
+			if (PackageStore.GetPackageRedirectInfo(PackageId, SourcePackageName, RedirectedToPackageId))
+			{
+				PackageId = RedirectedToPackageId;
+			}
+		}
+
+		bool bIsAlreadyInSet = false;
+		Visitied.Add(PackageId, &bIsAlreadyInSet);
+		if (bIsAlreadyInSet)
+		{
+			continue;
+		}
+
+		FPackageStoreEntry PackageStoreEntry;
+		const EPackageStoreEntryStatus EntryStatus = PackageStore.GetPackageStoreEntry(PackageId, NAME_None, PackageStoreEntry);
+		if (EntryStatus != EPackageStoreEntryStatus::Missing)
+		{
+			OutResolved.Add(PackageId);
+			for (const FPackageId& ImportedPackageId : PackageStoreEntry.ImportedPackageIds)
+			{
+				if (!Visitied.Contains(ImportedPackageId))
+				{
+					Queue.Add(ImportedPackageId);
+				}
+			}
+
+			if (bIncludeSoftReferences)
+			{
+				TConstArrayView<FPackageId> SoftReferences;
+				TConstArrayView<uint32> Indices = PackageStore.GetSoftReferences(PackageId, SoftReferences);
+				for (uint32 Idx : Indices)
+				{
+					const FPackageId& SoftRef = SoftReferences[Idx];
+					if (!Visitied.Contains(SoftRef))
+					{
+						Queue.Add(SoftRef);
+					}
+				}
+			}
+		}
+		else
+		{
+			OutMissing.Add(PackageId);
+		}
+	}
+}
+
+////////////////////////////////////////////////////////////////////////////////
+void ResolveChunksToInstall(
+	const TSet<FSharedOnDemandContainer>& Containers,
+	const TSet<FPackageId>& PackageIds,
+	bool bIncludeSoftReferences,
+	TArray<FResolvedContainerChunks>& OutResolvedContainerChunks,
+	TSet<FPackageId>& OutMissing)
+{
+	TRACE_CPUPROFILER_EVENT_SCOPE(FOnDemandContentInstaller::ResolveChunksToInstall);
+
+	// For now we always download these required chunks
+	for (const FSharedOnDemandContainer& Container : Containers)
+	{
+		FResolvedContainerChunks& ResolvedChunks = OutResolvedContainerChunks.AddDefaulted_GetRef();
+		ResolvedChunks.Container = Container;
+
+		for (int32 EntryIndex = 0; const FIoChunkId& ChunkId : Container->ChunkIds)
+		{
+			switch(ChunkId.GetChunkType())
+			{
+				case EIoChunkType::ExternalFile:
+				case EIoChunkType::ShaderCodeLibrary:
+				case EIoChunkType::ShaderCode:
+				{
+					ResolvedChunks.EntryIndices.Emplace(EntryIndex); 
+					ResolvedChunks.TotalSize += Container->ChunkEntries[EntryIndex].EncodedSize;
+				}
+				default:
+					break;
+			}
+			++EntryIndex;
+		}
+	}
+
+	auto FindChunkEntry = [&OutResolvedContainerChunks](const FIoChunkId& ChunkId, int32& OutIndex) -> FResolvedContainerChunks*
+	{
+		for (FResolvedContainerChunks& ContainerChunks : OutResolvedContainerChunks)
+		{
+			if (OutIndex = ContainerChunks.Container->FindChunkEntryIndex(ChunkId); OutIndex != INDEX_NONE)
+			{
+				return &ContainerChunks; 
+			}
+		}
+		return nullptr; 
+	};
+
+	TSet<FPackageId> ResolvedPackageIds;
+	ResolvePackageDependencies(
+		PackageIds,
+		bIncludeSoftReferences,
+		ResolvedPackageIds,
+		OutMissing);
+
+	// Resolve all chunk entries from the resolved package ID's
+	for (const FPackageId& PackageId : ResolvedPackageIds)
+	{
+		const FIoChunkId PackageChunkId				= CreatePackageDataChunkId(PackageId);
+		int32 EntryIndex							= INDEX_NONE; 
+		FResolvedContainerChunks* ResolvedChunks	= FindChunkEntry(PackageChunkId, EntryIndex);
+
+		if (ResolvedChunks == nullptr) 
+		{
+			// The chunk resides in a base game container
+			continue;
+		}
+
+		check(EntryIndex != INDEX_NONE);
+		FOnDemandContainer& Container	= *ResolvedChunks->Container; 
+
+		ResolvedChunks->EntryIndices.Emplace(EntryIndex); 
+		ResolvedChunks->TotalSize += Container.ChunkEntries[EntryIndex].EncodedSize;
+
+		// TODO: Installing optional bulkdata should probably be an install argument
+		const EIoChunkType AdditionalPackageChunkTypes[] =
+		{
+			EIoChunkType::BulkData,
+			EIoChunkType::OptionalBulkData,
+			EIoChunkType::MemoryMappedBulkData 
+		};
+
+		for (EIoChunkType ChunkType : AdditionalPackageChunkTypes)
+		{
+			// TODO: For Mutable we need to traverse all possible bulk data chunk indices?
+			const FIoChunkId ChunkId = CreateBulkDataIoChunkId(PackageId.Value(), 0, 0, ChunkType);
+			if (ResolvedChunks = FindChunkEntry(ChunkId, EntryIndex); ResolvedChunks != nullptr)
+			{
+				check(EntryIndex != INDEX_NONE);
+				FOnDemandContainer& OtherContainer = *ResolvedChunks->Container;
+				ResolvedChunks->EntryIndices.Emplace(EntryIndex);
+				ResolvedChunks->TotalSize += OtherContainer.ChunkEntries[EntryIndex].EncodedSize;
+			}
+		}
+	}
+}
+
+} // namespace Private
+
+////////////////////////////////////////////////////////////////////////////////
+FOnDemandContentInstaller::FOnDemandContentInstaller(FOnDemandIoStore& InIoStore, FOnDemandThread& InHttpClient)
+	: IoStore(InIoStore)
+	, HttpClient(InHttpClient)
+	, ProcessPipe(TEXT("IoStoreOnDemandProcessPipe"))
+	, HttpCompletionPipe(TEXT("IoStoreOnDemandHttpCompletionPipe"))
+{
+}
+
+FOnDemandContentInstaller::~FOnDemandContentInstaller()
+{
+}
+
+FSharedInternalInstallRequest FOnDemandContentInstaller::Enqueue(
+	FOnDemandInstallArgs&& Args,
+	FOnDemandInstallCompleted&& OnCompleted,
+	FOnDemandInstallProgressed&& OnProgress)
+{
+	FSharedInternalInstallRequest InstallRequest = MakeShared<FOnDemandInternalInstallRequest, ESPMode::ThreadSafe>(
+		MoveTemp(Args),
+		MoveTemp(OnCompleted),
+		MoveTemp(OnProgress));
+	
+	FInstallRequestState* RequestState = nullptr;
+	{
+		UE::TUniqueLock Lock(Mutex);
+		RequestState = RequestStateAllocator.Construct();
+		RequestState->Request = InstallRequest;
+		RequestStateByHandle.Add(InstallRequest, RequestState);
+	}
+
+	FOnDemandContentInstallerStats::OnRequestEnqueued();
+
+	ProcessPipe.Launch(
+		TEXT("ProcessIoStoreOnDemandInstallRequest"),
+		[this, RequestState] { ProcessRequest(*RequestState); },
+		UE::Tasks::ETaskPriority::BackgroundLow);
+
+	return InstallRequest;
+}
+
+void FOnDemandContentInstaller::Cancel(FSharedInternalInstallRequest InstallRequest)
+{
+	EIoErrorCode Expected = EIoErrorCode::Unknown;
+	if (InstallRequest->ErrorCode.compare_exchange_strong(Expected, EIoErrorCode::Cancelled))
+	{
+		UE_LOG(LogIoStoreOnDemand, Log, TEXT("Cancelling install request, ContentHandle=(%s)"),
+			*LexToString(InstallRequest->Args.ContentHandle));
+	}
+}
+
+void FOnDemandContentInstaller::UpdatePriority(FSharedInternalInstallRequest InstallRequest, int32 NewPriority)
+{
+	// Use the HTTP completion pipeline to synchronize access to the pending HTTP request(s)
+	HttpCompletionPipe.Launch(
+		TEXT("UpdateIoStoreOnDemandInstallRequestPriority"),
+		[this, InstallRequest, NewPriority]
+		{
+			UE::TUniqueLock Lock(Mutex);
+			if (FInstallRequestState* RequestState = RequestStateByHandle.FindRef(InstallRequest))
+			{
+				if (!RequestState->bHttpRequestsIssued.load(std::memory_order_relaxed) || NewPriority <= RequestState->Priority)
+				{
+					return;
+				}
+
+				UE_LOG(LogIoStoreOnDemand, Log, TEXT("Updating install request priority, ContentHandle=(%s), CurrentPriority=%d, NewPriority=%d"),
+					*LexToString(InstallRequest->Args.ContentHandle), RequestState->Priority, NewPriority);
+
+				RequestState->Priority = NewPriority;
+				for (FChunkHttpRequestHandle& PendingHttpRequest : RequestState->HttpRequestHandles)
+				{
+					if (PendingHttpRequest.Handle != nullptr)
+					{
+						HttpClient.ReprioritizeRequest(PendingHttpRequest.Handle, RequestState->Priority);
+					}
+				}
+			}
+		});
+}
+
+void FOnDemandContentInstaller::ProcessRequest(FInstallRequestState& RequestState)
+{
+	using namespace UE::IoStore::Private;
+
+	TRACE_CPUPROFILER_EVENT_SCOPE(FOnDemandContentInstaller::ProcessRequest);
+
+	FOnDemandInternalInstallRequest& Request = *RequestState.Request;
+
+	UE_LOG(LogIoStoreOnDemand, Log, TEXT("Processing install request, ContentHandle=(%s)"),
+		*LexToString(Request.Args.ContentHandle));
+
+	if (Request.Args.ContentHandle.IsValid() == false)
+	{
+		Request.ErrorCode			= EIoErrorCode::InvalidParameter;
+		RequestState.ErrorReason	= TEXT("Invalid content handle");
+
+		return CompleteRequest(RequestState);
+	}
+
+	const TSharedPtr<FOnDemandInternalContentHandle>& ContentHandle = Request.Args.ContentHandle.Handle;
+	if (ContentHandle->IoStore.IsValid() == false)
+	{
+		// First time this content handle is used
+		ContentHandle->IoStore = IoStore.AsWeak();
+	}
+
+	TSet<FSharedOnDemandContainer> ContainersForInstallation;
+	TSet<FPackageId> PackageIdsToInstall;
+	if (FIoStatus Status = IoStore.GetContainersAndPackagesForInstall(
+		Request.Args.MountId,
+		Request.Args.TagSets,
+		Request.Args.PackageIds,
+		ContainersForInstallation,
+		PackageIdsToInstall); !Status.IsOk())
+	{
+		Request.ErrorCode			= Status.GetErrorCode();
+		RequestState.ErrorReason	= Status.ToString(); 
+
+		return CompleteRequest(RequestState);
+	}
+
+	if (Request.ErrorCode.load(std::memory_order_relaxed) == EIoErrorCode::Cancelled)
+	{
+		return CompleteRequest(RequestState);
+	}
+
+	TSet<FPackageId> Missing;
+	const bool bIncludeSoftReferences = EnumHasAnyFlags(Request.Args.Options, EOnDemandInstallOptions::InstallSoftReferences);
+	Private::ResolveChunksToInstall(
+		ContainersForInstallation,
+		PackageIdsToInstall,
+		bIncludeSoftReferences,
+		RequestState.ResolvedChunks,
+		Missing);
+
+	// Check the other I/O backends for missing package chunks
+	if (Missing.IsEmpty() == false)
+	{
+		FIoDispatcher& IoDispatcher = FIoDispatcher::Get();
+		uint32 MissingCount = 0;
+
+		for (const FPackageId& PackageId : Missing)
+		{
+			const FIoChunkId ChunkId = CreatePackageDataChunkId(PackageId);
+			if (IoDispatcher.DoesChunkExist(ChunkId) == false)
+			{
+				UE_CLOG(MissingCount == 0, LogIoStoreOnDemand, Warning, TEXT("Failed to resolve the following chunk(s) for content handle '%s':"),
+					*LexToString(Request.Args.ContentHandle));
+
+				UE_LOG(LogIoStoreOnDemand, Warning, TEXT("\tChunkId='%s'"), *LexToString(ChunkId));
+				MissingCount++;
+			}
+		}
+
+		if (MissingCount > 0) 
+		{
+			Request.ErrorCode			= EIoErrorCode::UnknownChunkID;
+			RequestState.ErrorReason	= FString::Printf(TEXT("Missing chunk(s), Count=%u, ContentHandle='%s'"),
+				MissingCount, *LexToString(Request.Args.ContentHandle));
+
+			return CompleteRequest(RequestState);
+		}
+	}
+
+	if (Request.ErrorCode.load(std::memory_order_relaxed) == EIoErrorCode::Cancelled)
+	{
+		return CompleteRequest(RequestState);
+	}
+
+	uint64 TotalContentSize	= 0;	
+	uint64 TotalInstallSize	= 0;
+
+	{
+		TMap<FIoHash, uint64> ChunksToInstall;
+		for (int32 ContainerIndex = 0; FResolvedContainerChunks& ResolvedChunks : RequestState.ResolvedChunks)
+		{
+			for (int32 EntryIndex : ResolvedChunks .EntryIndices)
+			{
+				const FOnDemandChunkEntry& Entry = ResolvedChunks.Container->ChunkEntries[EntryIndex];
+				const bool bInstalled = IoStore.InstallCache->IsChunkCached(Entry.Hash);
+				if (bInstalled == false)
+				{
+					ChunksToInstall.Add(Entry.Hash, Entry.EncodedSize);
+					RequestState.HttpRequestHandles.Add(FChunkHttpRequestHandle
+					{
+						.Handle			= nullptr,
+						.ContainerIndex = ContainerIndex,
+						.EntryIndex		= EntryIndex
+					});
+					TotalInstallSize += Entry.EncodedSize;
+				}
+				++RequestState.ResolvedChunkCount;
+				TotalContentSize += Entry.EncodedSize;
+			}
+			ContainerIndex++;
+		}
+
+		if (ChunksToInstall.IsEmpty() == false)
+		{
+			if (FIoStatus Status = IoStore.InstallCache->Purge(MoveTemp(ChunksToInstall)); !Status.IsOk())
+			{
+				Request.ErrorCode			= Status.GetErrorCode();
+				RequestState.ErrorReason	= Status.ToString(); 
+
+				return CompleteRequest(RequestState);
+			}
+		}
+	}
+
+	RequestState.Progress.TotalContentSize		= TotalContentSize;
+	RequestState.Progress.TotalInstallSize		= TotalInstallSize;
+	RequestState.Progress.CurrentInstallSize	= 0;
+	RequestState.Priority						= Request.Args.Priority;
+
+	if (RequestState.HttpRequestHandles.IsEmpty())
+	{
+		Request.ErrorCode = EIoErrorCode::Ok;
+		return CompleteRequest(RequestState); 
+	}
+
+	if (Request.ErrorCode.load(std::memory_order_relaxed) == EIoErrorCode::Cancelled)
+	{
+		return CompleteRequest(RequestState);
+	}
+
+	NotifyProgress(RequestState);
+
+	for (FChunkHttpRequestHandle& HttpRequest : RequestState.HttpRequestHandles)
+	{
+		FSharedOnDemandContainer& Container		= RequestState.ResolvedChunks[HttpRequest.ContainerIndex].Container;
+		const FOnDemandChunkEntry& ChunkEntry	= Container->ChunkEntries[HttpRequest.EntryIndex];
+
+		HttpRequest.Handle = HttpClient.IssueRequest(
+			FOnDemandChunkInfo(Container, ChunkEntry),
+			FIoOffsetAndLength(),
+			RequestState.Priority,
+			[this, &RequestState, &HttpRequest](uint32 HttpStatusCode, FStringView ErrorReason, FIoBuffer&& Chunk)
+			{
+				HttpCompletionPipe.Launch(
+					TEXT("ProcessIoStoreOnDemandDownloadedChunk"),
+					[this, &RequestState, &HttpRequest, HttpStatusCode, ErrorReason = FString(ErrorReason), Chunk = MoveTemp(Chunk)]() mutable
+					{
+						ProcessDownloadedChunk(
+							RequestState,
+							HttpRequest,
+							HttpStatusCode,
+							MoveTemp(ErrorReason),
+							MoveTemp(Chunk));
+					},
+					UE::Tasks::ETaskPriority::BackgroundLow);
+			});
+	}
+	
+	RequestState.bHttpRequestsIssued = true;
+}
+
+void FOnDemandContentInstaller::ProcessDownloadedChunk(
+	FInstallRequestState& RequestState,
+	FChunkHttpRequestHandle& HttpRequest,
+	uint32 HttpStatusCode,
+	FString&& ErrorReason,
+	FIoBuffer&& Chunk)
+{
+	FOnDemandInternalInstallRequest& Request	= *RequestState.Request;
+	HttpRequest.Handle							= nullptr;
+	FSharedOnDemandContainer& Container			= RequestState.ResolvedChunks[HttpRequest.ContainerIndex].Container;
+	const FOnDemandChunkEntry& ChunkEntry		= Container->ChunkEntries[HttpRequest.EntryIndex];
+	const FIoChunkId& ChunkId					= Container->ChunkIds[HttpRequest.EntryIndex];
+
+	RequestState.Progress.CurrentInstallSize += ChunkEntry.EncodedSize;
+
+	UE_LOG(LogIoStoreOnDemand, Log, TEXT("Install progress %.2lf/%.2lf MiB, ContentHandle=(%s), ChunkId='%s', ChunkSize=%.2lf KiB, HttpStatus=%u"),
+		double(RequestState.Progress.CurrentInstallSize) / 1024.0 / 1024.0, double(RequestState.Progress.TotalInstallSize) / 1024.0 / 1024.0,
+		*LexToString(Request.Args.ContentHandle), *LexToString(ChunkId), double(ChunkEntry.EncodedSize) / 1024.0, HttpStatusCode);
+
+	if (Request.ErrorCode == EIoErrorCode::Unknown)
+	{
+		const bool bHttpOk = HttpStatusCode > 199 && HttpStatusCode < 300 && Chunk.GetSize() > 0;
+		if (bHttpOk)
+		{
+			const FIoHash ChunkHash = FIoHash::HashBuffer(Chunk.GetView());
+
+			if (ChunkEntry.Hash == ChunkHash)
+			{
+				if (FIoStatus Status = IoStore.InstallCache->PutChunk(MoveTemp(Chunk), ChunkHash); !Status.IsOk())
+				{
+					Request.ErrorCode			= Status.GetErrorCode();
+					RequestState.ErrorReason	= Status.ToString();
+				}
+			}
+			else
+			{
+				Request.ErrorCode			= EIoErrorCode::ReadError;
+				RequestState.ErrorReason	= FString::Printf(TEXT("Hash mismatch, ChunkId='%s', ExpectedHash='%s', ActualHash='%s'"),
+					*LexToString(ChunkId), *LexToString(ChunkEntry.Hash), *LexToString(ChunkHash));
+			}
+		}
+		else
+		{
+			Request.ErrorCode			= EIoErrorCode::ReadError;
+			RequestState.ErrorReason	= FString::Printf(TEXT("Http failure, StatusCode=%u, Reason=%s"), HttpStatusCode, *ErrorReason);
+		}
+	}
+
+	NotifyProgress(RequestState);
+
+	const bool bCompleted = ++RequestState.DownloadedChunkCount >= RequestState.HttpRequestHandles.Num();
+	if (bCompleted)
+	{
+		if (Request.ErrorCode == EIoErrorCode::Unknown)
+		{
+			Request.ErrorCode = EIoErrorCode::Ok;
+		}
+		CompleteRequest(RequestState);
+	}
+	else if (Request.ErrorCode != EIoErrorCode::Unknown)
+	{
+		int32 NumCancelled = 0;
+		for (FChunkHttpRequestHandle& PendingHttpRequest : RequestState.HttpRequestHandles)
+		{
+			if (PendingHttpRequest.Handle != nullptr)
+			{
+				HttpClient.CancelRequest(PendingHttpRequest.Handle);
+				++NumCancelled;
+			}
+		}
+
+		UE_CLOG(NumCancelled == 0, LogIoStoreOnDemand, Log, TEXT("Cancelled %d HTTP request(s) due to install error"), NumCancelled);
+	}
+}
+
+void FOnDemandContentInstaller::NotifyProgress(FInstallRequestState& RequestState)
+{
+	FOnDemandInternalInstallRequest& Request = *RequestState.Request;
+
+	if (!Request.OnProgress)
+	{
+		return;
+	}
+
+	const double SecondsSinceLastProgress = FPlatformTime::Seconds() - RequestState.LastProgressTimeStamp;
+	if (SecondsSinceLastProgress < .25)
+	{
+		return;
+	}
+	RequestState.LastProgressTimeStamp = FPlatformTime::Seconds();
+
+	FOnDemandInstallProgress Progress = RequestState.Progress;
+	if (EnumHasAnyFlags(Request.Args.Options, EOnDemandInstallOptions::CallbackOnGameThread))
+	{
+		ExecuteOnGameThread(
+			UE_SOURCE_LOCATION,
+			[Request = RequestState.Request, Progress]()
+			{
+				Request->OnProgress(Progress);
+			});
+	}
+	else
+	{
+		Request.OnProgress(Progress);
+	}
+}
+
+void FOnDemandContentInstaller::CompleteRequest(FInstallRequestState& RequestState)
+{
+	using namespace UE::IoStore::Private;
+
+	FSharedInternalInstallRequest Request = MoveTemp(RequestState.Request);
+	ensure(Request->ErrorCode != EIoErrorCode::Unknown);
+
+	if (FIoStatus Status = IoStore.InstallCache->Flush(); !Status.IsOk())
+	{
+		Request->ErrorCode			= Status.GetErrorCode();
+		RequestState.ErrorReason	= Status.ToString();
+	}
+
+	// Mark all resolved chunk(s) as referenced by the content handle and notify the package store to update
+	if (RequestState.ResolvedChunkCount > 0)
+	{
+		{
+			FOnDemandContentHandle& ContentHandle = Request->Args.ContentHandle;
+
+			TUniqueLock Lock(IoStore.ContainerMutex);
+			for (FResolvedContainerChunks& ResolvedChunks : RequestState.ResolvedChunks)
+			{
+				const FSharedOnDemandContainer& Container = ResolvedChunks.Container;
+				FOnDemandChunkEntryReferences& References = Container->FindOrAddChunkEntryReferences(*ContentHandle.Handle);
+				for (int32 EntryIndex : ResolvedChunks.EntryIndices)
+				{
+					References.Indices[EntryIndex] = true;
+				}
+			}
+		}
+
+		IoStore.PackageStoreBackend->NeedsUpdate(EOnDemandPackageStoreUpdateMode::ReferencedPackages);
+	}
+
+	FOnDemandInstallResult InstallResult;
+	InstallResult.Status				= Request->ErrorCode == EIoErrorCode::Ok ? FIoStatus::Ok : FIoStatus(Request->ErrorCode, RequestState.ErrorReason);
+	InstallResult.Progress				= RequestState.Progress;
+	InstallResult.DurationInSeconds		= FPlatformTime::Seconds() - Request->StartTime;
+
+	FOnDemandContentInstallerStats::OnRequestCompleted(
+		RequestState.ResolvedChunkCount,
+		InstallResult.Progress.TotalContentSize,
+		static_cast<uint64>(RequestState.HttpRequestHandles.Num()),
+		InstallResult.Progress.TotalInstallSize,
+		InstallResult.DurationInSeconds);
+
+	UE_LOG(LogIoStoreOnDemand, Log, TEXT("Install request completed, Result='%s', ContentHandle=(%s), ContentSize=%.2lf MiB, InstallSize=%.2lf MiB, Duration=%.2lfs"),
+		GetIoErrorText(InstallResult.Status.GetErrorCode()), *LexToString(Request->Args.ContentHandle), double(InstallResult.Progress.TotalContentSize) / 1024.0 / 1024.0,
+		double(InstallResult.Progress.TotalInstallSize) / 1024.0 / 1024.0, InstallResult.DurationInSeconds);
+
+	{
+		UE::TUniqueLock Lock(Mutex);
+		RequestStateByHandle.Remove(Request); 
+		RequestStateAllocator.Destroy(&RequestState);
+	}
+
+	FOnDemandInstallRequest::EStatus RequestStatus = FOnDemandInstallRequest::EStatus::None;
+	switch (InstallResult.Status.GetErrorCode())
+	{
+		case EIoErrorCode::Ok:
+			RequestStatus = FOnDemandInstallRequest::Ok;
+			break;
+		case EIoErrorCode::Cancelled:
+			RequestStatus = FOnDemandInstallRequest::Cancelled;
+			break;
+		default:
+			RequestStatus = FOnDemandInstallRequest::Error;
+			break;
+	}
+
+	if (!Request->OnCompleted)
+	{
+		Request->Status.store(RequestStatus);
+	}
+	else if (EnumHasAnyFlags(Request->Args.Options, EOnDemandInstallOptions::CallbackOnGameThread))
+	{
+		ExecuteOnGameThread(
+			UE_SOURCE_LOCATION,
+			[Request, InstallResult = MoveTemp(InstallResult), RequestStatus]() mutable
+			{
+				FOnDemandInstallCompleted OnCompleted = MoveTemp(Request->OnCompleted);
+				OnCompleted(MoveTemp(InstallResult));
+				Request->Status.store(RequestStatus);
+			});
+	}
+	else
+	{
+		FOnDemandInstallCompleted OnCompleted = MoveTemp(Request->OnCompleted);
+		OnCompleted(MoveTemp(InstallResult));
+		Request->Status.store(RequestStatus);
+	}
+}
+
+} // namespace UE::IoStore
diff --git a/Engine/Source/Runtime/Experimental/IoStore/OnDemand/Private/OnDemandContentInstaller.h b/Engine/Source/Runtime/Experimental/IoStore/OnDemand/Private/OnDemandContentInstaller.h
new file mode 100644
--- /dev/null
+++ b/Engine/Source/Runtime/Experimental/IoStore/OnDemand/Private/OnDemandContentInstaller.h
@@ -0,0 +1,113 @@
+// Copyright Epic Games, Inc. All Rights Reserved.
+
+#pragma once
+
+#include "Async/Mutex.h"
+#include "Containers/Set.h"
+#include "IO/IoStoreOnDemand.h"
+#include "IO/PackageId.h"
+#include "IO/IoStoreOnDemandInternals.h"
+#include "IO/IoAllocators.h"
+#include "OnDemandIoStore.h"
+#include "Tasks/Pipe.h"
+
+#include <atomic>
+
+class FIoBuffer;
+
+namespace UE::IoStore
+{
+
+class FOnDemandIoStore;
+class FOnDemandThread;
+
+using FSharedOnDemandContainer = TSharedPtr<struct FOnDemandContainer, ESPMode::ThreadSafe>;
+
+namespace Private
+{
+
+////////////////////////////////////////////////////////////////////////////////
+struct FResolvedContainerChunks
+{
+	FSharedOnDemandContainer	Container;
+	TArray<int32>				EntryIndices;
+	uint64						TotalSize = 0;
+};
+
+////////////////////////////////////////////////////////////////////////////////
+void ResolvePackageDependencies(
+	const TSet<FPackageId>& PackageIds,
+	bool bIncludeSoftReferences,
+	TSet<FPackageId>& OutResolved,
+	TSet<FPackageId>& OutMissing);
+
+////////////////////////////////////////////////////////////////////////////////
+void ResolveChunksToInstall(
+	const TSet<FSharedOnDemandContainer>& Containers,
+	const TSet<FPackageId>& PackageIds,
+	bool bIncludeSoftReferences,
+	TArray<FResolvedContainerChunks>& OutResolvedContainerChunks,
+	TSet<FPackageId>& OutMissing);
+
+} // namespace Private
+
+////////////////////////////////////////////////////////////////////////////////
+class FOnDemandContentInstaller
+{
+	struct FChunkHttpRequestHandle
+	{
+		void*	Handle = nullptr;
+		int32	ContainerIndex = INDEX_NONE;
+		int32	EntryIndex = INDEX_NONE;
+	};
+
+	struct FInstallRequestState
+	{
+		FSharedInternalInstallRequest				Request;
+		TArray<Private::FResolvedContainerChunks>	ResolvedChunks;
+		TArray<FChunkHttpRequestHandle>				HttpRequestHandles;
+		FString										ErrorReason;
+		FOnDemandInstallProgress					Progress;
+		uint64										ResolvedChunkCount = 0;
+		uint64										DownloadedChunkCount = 0;
+		double										LastProgressTimeStamp = 0.0;
+		int32										Priority = 0; 
+		std::atomic_bool							bHttpRequestsIssued{false};
+	};
+
+public:
+									FOnDemandContentInstaller(FOnDemandIoStore& IoStore, FOnDemandThread& HttpClient);
+									~FOnDemandContentInstaller();
+
+	FSharedInternalInstallRequest	Enqueue(
+										FOnDemandInstallArgs&& Args,
+										FOnDemandInstallCompleted&& OnCompleted,
+										FOnDemandInstallProgressed&& OnProgress);
+	void							Cancel(FSharedInternalInstallRequest InstallRequest);
+	void							UpdatePriority(FSharedInternalInstallRequest InstallRequest, int32 NewPriority);
+
+private:
+	void							ProcessRequest(FInstallRequestState& RequestState);
+	void							ProcessDownloadedChunk(
+										FInstallRequestState& RequestState,
+										FChunkHttpRequestHandle& HttpRequest,
+										uint32 HttpStatusCode,
+										FString&& ErrorReason,
+										FIoBuffer&& Chunk);
+	void							NotifyProgress(FInstallRequestState& RequestState);
+	void							CompleteRequest(FInstallRequestState& RequestState);
+
+	using FInstallRequestStateMap	= TMap<FSharedInternalInstallRequest, FInstallRequestState*>;
+	using FRequestStateAllocator	= TSingleThreadedSlabAllocator<FInstallRequestState>;
+
+	FOnDemandIoStore&			IoStore;
+	FOnDemandThread&			HttpClient;
+	UE::Tasks::FPipe			ProcessPipe;
+	UE::Tasks::FPipe			HttpCompletionPipe;
+
+	FMutex						Mutex;
+	FInstallRequestStateMap		RequestStateByHandle;
+	FRequestStateAllocator		RequestStateAllocator;
+};
+
+} // namespace UE::IoStore
diff --git a/Engine/Source/Runtime/Experimental/IoStore/OnDemand/Private/OnDemandHttpThread.h b/Engine/Source/Runtime/Experimental/IoStore/OnDemand/Private/OnDemandHttpThread.h
--- a/Engine/Source/Runtime/Experimental/IoStore/OnDemand/Private/OnDemandHttpThread.h
+++ b/Engine/Source/Runtime/Experimental/IoStore/OnDemand/Private/OnDemandHttpThread.h
@@ -7,7 +7,7 @@
 #endif
 
 #define UE_ENABLE_IAS_REQUEST_THREAD (1 && UE_ENABLE_REQUEST_THREAD)
-#define UE_ENABLE_IAD_REQUEST_THREAD (0 && UE_ENABLE_REQUEST_THREAD)
+#define UE_ENABLE_IAD_REQUEST_THREAD (1 && UE_ENABLE_REQUEST_THREAD)
 
 #if UE_ENABLE_REQUEST_THREAD
 
diff --git a/Engine/Source/Runtime/Experimental/IoStore/OnDemand/Private/OnDemandIoStore.cpp b/Engine/Source/Runtime/Experimental/IoStore/OnDemand/Private/OnDemandIoStore.cpp
--- a/Engine/Source/Runtime/Experimental/IoStore/OnDemand/Private/OnDemandIoStore.cpp
+++ b/Engine/Source/Runtime/Experimental/IoStore/OnDemand/Private/OnDemandIoStore.cpp
@@ -35,6 +35,7 @@
 #include "OnDemandHttpClient.h"
 #include "OnDemandHttpThread.h"
 #include "OnDemandInstallCache.h"
+#include "OnDemandContentInstaller.h"
 #include "OnDemandIoDispatcherBackend.h"
 #include "OnDemandPackageStoreBackend.h"
 #include "ProfilingDebugging/AssetMetadataTrace.h"
@@ -394,6 +395,9 @@ FOnDemandIoStore::FOnDemandIoStore()
 {
 #if UE_ENABLE_REQUEST_THREAD
 	OnDemandHttpThread = MakeUnique<FOnDemandThread>();
+#if UE_ENABLE_IAD_REQUEST_THREAD
+	Installer = MakeUnique<FOnDemandContentInstaller>(*this, FOnDemandThread::Get());
+#endif
 #endif // UE_ENABLE_REQUEST_THREAD
 
 	FEncryptionKeyManager::Get().OnKeyAdded().AddRaw(this, &FOnDemandIoStore::OnEncryptionKeyAdded);
@@ -658,11 +662,14 @@ FIoStatus FOnDemandIoStore::InitializePostHotfix()
 #if !UE_BUILD_SHIPPING
 		if (InstallArgs)
 		{
-			Install(MoveTemp(InstallArgs.GetValue()), [](FOnDemandInstallResult InstallResult)
+			FEventRef InstallDone;
+			Install(MoveTemp(InstallArgs.GetValue()), [&InstallDone](FOnDemandInstallResult InstallResult)
 			{
 				UE_CLOG(!InstallResult.Status.IsOk(), LogIoStoreOnDemand, Error,
 					TEXT("Failed to install content, reason '%s'"), *InstallResult.Status.ToString());
+				InstallDone->Trigger();
 			});
+			InstallDone->Wait();
 		}
 #endif
 	}
@@ -737,11 +744,19 @@ FOnDemandInstallRequest FOnDemandIoStore::Install(
 {
 	LLM_SCOPE_BYNAME(TEXT("FileSystem/OnDemandIoStore"));
 	
-	FSharedInternalInstallRequest InstallRequest = MakeShared<FOnDemandInternalInstallRequest, ESPMode::ThreadSafe>(
+#if UE_ENABLE_IAD_REQUEST_THREAD
+	FSharedInternalInstallRequest InstallRequest = Installer->Enqueue(
 		MoveTemp(Args),
 		MoveTemp(OnCompleted),
 		MoveTemp(OnProgress));
 
+	return FOnDemandInstallRequest(AsWeak(), InstallRequest);
+#else
+	FSharedInternalInstallRequest InstallRequest = MakeShared<FOnDemandInternalInstallRequest, ESPMode::ThreadSafe>(
+			MoveTemp(Args),
+			MoveTemp(OnCompleted),
+			MoveTemp(OnProgress));
+
 	{
 		UE::TUniqueLock Lock(RequestMutex);
 		InstallRequests.Add(InstallRequest);
@@ -750,6 +765,7 @@ FOnDemandInstallRequest FOnDemandIoStore::Install(
 	TryEnterTickLoop();
 
 	return FOnDemandInstallRequest(AsWeak(), InstallRequest);
+#endif
 }
 
 void FOnDemandIoStore::Purge(FOnDemandPurgeArgs&& Args, FOnDemandPurgeCompleted&& OnCompleted)
@@ -1337,6 +1353,17 @@ FIoStatus FOnDemandIoStore::TickMountRequest(FMountRequest& MountRequest)
 			{
 				EnumAddFlags(Container->Flags, EOnDemandContainerFlags::WithSoftReferences);
 			}
+
+			FIASHostGroup HostGroup = FHostGroupManager::Get().Find(NAME_None);
+			if (HostGroup.IsResolved() == false)
+			{
+				TIoStatusOr<FIASHostGroup> NewGroup = FHostGroupManager::Get().Register(NAME_None, Container->HostGroup.Hosts());
+				if (NewGroup.IsOk())
+				{
+					HostGroup = NewGroup.ConsumeValueOrDie();
+				}
+			}
+			Container->IasHostGroup = HostGroup; 
 		}
 
 		if (EnumHasAnyFlags(Container->Flags, EOnDemandContainerFlags::Encrypted) &&
@@ -1591,10 +1618,6 @@ FOnDemandInstallResult FOnDemandIoStore::TickInstallRequest(FOnDemandInternalIns
 	}
 
 	// Download all chunks
-	
-#if UE_ENABLE_IAD_REQUEST_THREAD
-	std::atomic<int32> ConcurrentRequests = 0;
-#else
 	const int32 MaxConcurrentRequests = 16;
 	int32 ConcurrentRequests = 0;
 
@@ -1609,7 +1632,6 @@ FOnDemandInstallResult FOnDemandIoStore::TickInstallRequest(FOnDemandInternalIns
 		Status = FIoStatusBuilder(EIoErrorCode::InvalidCode) << TEXT("Failed to initialize HTTP client");
 		return OutResult;
 	}
-#endif // UE_ENABLE_IAD_REQUEST_THREAD
 
 	uint64 TotalChunkCount		= 0;
 	uint64 TotalBytes			= 0;
@@ -1672,42 +1694,6 @@ FOnDemandInstallResult FOnDemandIoStore::TickInstallRequest(FOnDemandInternalIns
 
 				++ConcurrentRequests;
 
-#if UE_ENABLE_IAD_REQUEST_THREAD
-				FOnDemandThread::Get().IssueRequest(FOnDemandChunkInfo(Kv.Key, ChunkEntry), FIoOffsetAndLength(), EIoDispatcherPriority::IoDispatcherPriority_High,
-					[this, &ConcurrentRequests, ChunkEntry, &OutResult, &InstallRequest](uint32 StatusCode, FStringView ErrorReason, FIoBuffer&& Data)
-					{
-						--ConcurrentRequests;
-
-						if (!OutResult.Status.IsOk())
-						{
-							// A different request already failed, don't stomp the status if this one succeeded
-							return;
-						}
-
-						if (!IsHttpStatusOk(StatusCode))
-						{
-							OutResult.Status = FIoStatus(EIoErrorCode::ReadError) << FString::Printf(TEXT("HTTP Status %u, "), StatusCode) << ErrorReason;
-							return;
-						}
-
-						const FIoHash ChunkHash = FIoHash::HashBuffer(Data.GetView());
-						if (ChunkHash != ChunkEntry.Hash)
-						{
-							OutResult.Status = FIoStatus(EIoErrorCode::ReadError, TEXTVIEW("Hash mismatch"));
-							return;
-						}
-						
-						OutResult.Status = InstallCache->PutChunk(MoveTemp(Data), ChunkHash);
-						if (!OutResult.Status.IsOk())
-						{
-							return;
-						}
-
-						// TODO: Is this good enough progress or should it be fine grained bytes from the HttpClient?
-						OutResult.Progress.CurrentInstallSize += ChunkEntry.EncodedSize;
-						ProgressInstallRequest(InstallRequest, OutResult.Progress);
-					});
-#else
 				HttpClient->Get(
 					Container.HostGroup,
 					Private::GetChunkUrl(Container, ChunkEntry, ChunkUrl).ToView(),
@@ -1764,19 +1750,11 @@ FOnDemandInstallResult FOnDemandIoStore::TickInstallRequest(FOnDemandInternalIns
 
 					return OutResult;
 				}
-#endif //UE_ENABLE_IAD_REQUEST_THREAD
 			}
 		}
 
-#if UE_ENABLE_IAD_REQUEST_THREAD
-		while (ConcurrentRequests > 0)
-		{
-			FPlatformProcess::SleepNoStats(0.0f);
-		}
-#else
 		while (HttpClient->Tick())
 			;
-#endif //UE_ENABLE_IAD_REQUEST_THREAD
 	}
 
 	if (Status.IsOk() == false)
@@ -2291,6 +2269,9 @@ void FOnDemandIoStore::CancelInstallRequest(FSharedInternalInstallRequest Instal
 {
 	check(InstallRequest.IsValid());
 
+#if UE_ENABLE_IAD_REQUEST_THREAD
+	Installer->Cancel(InstallRequest);
+#else
 	// Cancel the request if the error code is unknown otherwise the request will be completed in the tick loop 
 	EIoErrorCode Expected = EIoErrorCode::Unknown;
 	if (InstallRequest->ErrorCode.compare_exchange_strong(Expected, EIoErrorCode::Cancelled))
@@ -2303,14 +2284,19 @@ void FOnDemandIoStore::CancelInstallRequest(FSharedInternalInstallRequest Instal
 		};
 		CompleteInstallRequest(InstallRequest, MoveTemp(InstallResult));
 	}
+#endif
 }
 
 void FOnDemandIoStore::UpdateInstallRequestPriority(FSharedInternalInstallRequest InstallRequest, int32 NewPriority)
 {
+#if UE_ENABLE_IAD_REQUEST_THREAD
+	Installer->UpdatePriority(InstallRequest, NewPriority);
+#else
 	{
 		UE::TUniqueLock Lock(RequestMutex);
 		InstallRequest->Args.Priority = NewPriority;
 	}
+#endif
 }
 
 FOnDemandCacheUsage FOnDemandIoStore::GetCacheUsage() const
diff --git a/Engine/Source/Runtime/Experimental/IoStore/OnDemand/Private/OnDemandIoStore.h b/Engine/Source/Runtime/Experimental/IoStore/OnDemand/Private/OnDemandIoStore.h
--- a/Engine/Source/Runtime/Experimental/IoStore/OnDemand/Private/OnDemandIoStore.h
+++ b/Engine/Source/Runtime/Experimental/IoStore/OnDemand/Private/OnDemandIoStore.h
@@ -26,6 +26,7 @@ using FSharedContainerHeader		= TSharedPtr<FIoContainerHeader>;
 namespace UE::IoStore
 {
 
+class FOnDemandContentInstaller;
 class FOnDemandThread;
 struct FOnDemandEndpointConfig;
 enum class ELogOnDemandCacheUsage : uint8;
@@ -223,6 +224,7 @@ struct FOnDemandChunkInfo
 
 private:
 	friend class FOnDemandIoStore;
+	friend class FOnDemandContentInstaller;
 
 	FOnDemandChunkInfo(FSharedOnDemandContainer InContainer, const FOnDemandChunkEntry& InEntry)
 		: SharedContainer(InContainer)
@@ -332,6 +334,8 @@ public:
 	virtual void						ReleaseContent(FOnDemandInternalContentHandle& ContentHandle) override;
 
 private:
+	friend class FOnDemandContentInstaller;
+
 	FIoStatus							GetContainersAndPackagesForInstall(
 											FStringView MountId,
 											const TArray<FString>& TagSets,
@@ -357,6 +361,7 @@ private:
 											const TArray<FString>& ExistingContainers,
 											TArray<FSharedOnDemandContainer>& Out);
 
+	TUniquePtr<FOnDemandContentInstaller>	Installer;
 	TUniquePtr<FOnDemandThread>				OnDemandHttpThread;
 
 	FSharedInstallCache						InstallCache;
diff --git a/Engine/Source/Runtime/Experimental/IoStore/OnDemand/Private/Statistics.cpp b/Engine/Source/Runtime/Experimental/IoStore/OnDemand/Private/Statistics.cpp
--- a/Engine/Source/Runtime/Experimental/IoStore/OnDemand/Private/Statistics.cpp
+++ b/Engine/Source/Runtime/Experimental/IoStore/OnDemand/Private/Statistics.cpp
@@ -329,6 +329,17 @@ CSV_DEFINE_STAT(Ias, HttpBandwidthMpbs);
 CSV_DEFINE_STAT(Ias, HttpDurationMsAvg);
 CSV_DEFINE_STAT(Ias, HttpDurationMsMax);
 
+////////////////////////////////////////////////////////////////////////////////
+// TRACE STATS IAD
+FIncrementalVariance	GInstallCacheHitRatio;
+FCounterAtomicInt		GTotalInstallCount(TEXT("Iad/InstallCount"), TraceCounterDisplayHint_None);
+FCounterAtomicInt		GInflightInstallCount(TEXT("Iad/InflightInstallCount"), TraceCounterDisplayHint_None);
+FCounterAtomicInt		GTotalDownloadedChunkCount(TEXT("Iad/DownloadedChunks"), TraceCounterDisplayHint_None);
+FCounterAtomicInt		GTotalDownloadedChunkBytes(TEXT("Iad/DownloadedBytes"), TraceCounterDisplayHint_Memory);
+//FCounterAtomicFloat	GTotalInstallDurationSeconds(TEXT("Iad/TotalInstallDurationSeconds"), TraceCounterDisplayHint_None);
+//FCounterAtomicFloat	GAvgInstallCacheHitRatio(TEXT("Iad/AvgInstallCacheHitRatio"), TraceCounterDisplayHint_None);
+
+////////////////////////////////////////////////////////////////////////////////
 static FOnDemandIoBackendStats* GStatistics = nullptr;
 static FDelegateHandle GStatisticsEndFrameDelegateHandle;
 static FDelegateHandle GStatisticsOnScreenDelegateHandle;
@@ -940,6 +951,28 @@ void FOnDemandIoBackendStats::OnHttpCdnCacheReply(int32 Reply)
 	}
 }
 
+////////////////////////////////////////////////////////////////////////////////
+void FOnDemandContentInstallerStats::OnRequestEnqueued()
+{
+	GInflightInstallCount.Add(1);
+}
+
+void FOnDemandContentInstallerStats::OnRequestCompleted(
+		uint64 RequestedChunkCount,
+		uint64 RequestedBytes,
+		uint64 DownloadedChunkCount,
+		uint64 DownloadedBytes,
+		double DurationSeconds) 
+{
+	GInflightInstallCount.Add(-1);
+	GTotalInstallCount.Add(1);
+	GTotalDownloadedChunkCount.Add(DownloadedChunkCount);
+	GTotalDownloadedChunkBytes.Add(DownloadedBytes);
+	//GTotalInstallDurationSeconds.Add(DurationSeconds);
+	GInstallCacheHitRatio.Increment(double(DownloadedBytes) / double(RequestedBytes));
+	//GAvgInstallCacheHitRatio.Set(GInstallCacheHitRatio.GetMean());
+}
+
 } // namespace UE::IoStore
 
 #undef UE_ENABLE_ONSCREEN_STATISTICS
diff --git a/Engine/Source/Runtime/Experimental/IoStore/OnDemand/Private/Statistics.h b/Engine/Source/Runtime/Experimental/IoStore/OnDemand/Private/Statistics.h
--- a/Engine/Source/Runtime/Experimental/IoStore/OnDemand/Private/Statistics.h
+++ b/Engine/Source/Runtime/Experimental/IoStore/OnDemand/Private/Statistics.h
@@ -29,6 +29,7 @@ struct FAvailableEndPoints;
 #	define IAS_STATISTICS_IMPL(...) { return __VA_ARGS__; }
 #endif
 
+////////////////////////////////////////////////////////////////////////////////
 class FOnDemandIoBackendStats
 {
 public:
@@ -80,6 +81,19 @@ private:
 #endif //IAS_WITH_STATISTICS
 };
 
+////////////////////////////////////////////////////////////////////////////////
+class FOnDemandContentInstallerStats
+{
+public:
+	static void OnRequestEnqueued() IAS_STATISTICS_IMPL()
+	static void OnRequestCompleted(
+		uint64 RequestedChunkCount,
+		uint64 RequestedBytes,
+		uint64 DownloadedChunkCount,
+		uint64 DownloadedBytes,
+		double DurationSeconds) IAS_STATISTICS_IMPL()
+};
+
 #undef IAS_STATISTICS_IMPL
 
 } // namespace UE::IoStore
diff --git a/Engine/Source/Runtime/Experimental/IoStore/OnDemandCore/Public/IO/IoStoreOnDemand.h b/Engine/Source/Runtime/Experimental/IoStore/OnDemandCore/Public/IO/IoStoreOnDemand.h
--- a/Engine/Source/Runtime/Experimental/IoStore/OnDemandCore/Public/IO/IoStoreOnDemand.h
+++ b/Engine/Source/Runtime/Experimental/IoStore/OnDemandCore/Public/IO/IoStoreOnDemand.h
@@ -105,6 +105,7 @@ public:
 
 private:
 	friend class FOnDemandIoStore;
+	friend class FOnDemandContentInstaller;
 	FSharedInternalContentHandle Handle;
 };
 
