Constraints: new evaluation scheme that offers much better performance and also fixes several glitches such as motion blur or mixing control rig physics and constraints

- enabled with Constraints.NewEvaluation true (false by default)

- the prerequisits system is reversed: we now make sure that the constraint ticks before the child does, to ensure that the data needed for a correct evaluation is prepared.
- added FAnimationEvaluationCache & FAnimationEvaluator to allow evaluating the animation of a skeletal mesh and store the resulting transforms without having to tick it directly. (see AnimationEvaluation.h for more info)
- handles now notify the skeletal mesh compenents they act on that their animation data are dirty and need to be re-evaluated.

[RN] Constraints: new evaluation scheme that offers much better performance and also fixes several glitches

diff --git a/Engine/Plugins/Animation/ControlRig/Source/ControlRig/Private/Constraints/ControlRigTransformableHandle.cpp b/Engine/Plugins/Animation/ControlRig/Source/ControlRig/Private/Constraints/ControlRigTransformableHandle.cpp
--- a/Engine/Plugins/Animation/ControlRig/Source/ControlRig/Private/Constraints/ControlRigTransformableHandle.cpp
+++ b/Engine/Plugins/Animation/ControlRig/Source/ControlRig/Private/Constraints/ControlRigTransformableHandle.cpp
@@ -18,7 +18,7 @@
 
 #include UE_INLINE_GENERATED_CPP_BY_NAME(ControlRigTransformableHandle)
 
-namespace ControlHandleLocals
+namespace UE::Private
 {
 	using RigGuard = TGuardValue_Bitfield_Cleanup<TFunction<void()>>;
 	
@@ -28,6 +28,40 @@ namespace ControlHandleLocals
 	{
 		return InControlRig ? NotifyingRigs.Contains(InControlRig) : false;
 	}
+
+	struct FControlPoseChangedNotifier
+	{
+		FControlPoseChangedNotifier(const URigHierarchy* InHierarchy, const FRigControlElement* InControl, const USceneComponent* InComponent)
+			: Hierarchy(InHierarchy)
+			, Control(InControl)
+			, Component(InComponent)
+		{
+			if (IsValid())
+			{
+				PoseVersion = Hierarchy->GetPoseVersion(Control); 
+			}
+		}
+		
+		~FControlPoseChangedNotifier()
+		{
+			if (IsValid() && Hierarchy->GetPoseVersion(Control) != PoseVersion)
+			{
+				TransformableHandleUtils::MarkComponentForEvaluation(Component);
+			}
+		}
+
+	private:
+
+		bool IsValid() const
+		{
+			return Hierarchy && Control && Component && TransformableHandleUtils::SkipTicking();
+		}
+		
+		const URigHierarchy* Hierarchy = nullptr;
+		const FRigControlElement* Control = nullptr;
+		const USceneComponent* Component = nullptr;
+		int32 PoseVersion = INDEX_NONE;
+	};
 }
 
 /**
@@ -79,11 +113,11 @@ void UTransformableControlHandle::PreEvaluate(const bool bTick) const
 
 	if (ControlRig->IsAdditive())
 	{
-		if (ControlHandleLocals::IsRigNotifying(ControlRig.Get()))
+		if (UE::Private::IsRigNotifying(ControlRig.Get()))
 		{
 			return;
 		}
-		
+
 		if (const USkeletalMeshComponent* SkeletalMeshComponent = GetSkeletalMesh())
 		{
 			if (!SkeletalMeshComponent->PoseTickedThisFrame())
@@ -92,8 +126,21 @@ void UTransformableControlHandle::PreEvaluate(const bool bTick) const
 			}
 		}
 	}
+
+	if (!bTick)
+	{
+		return ControlRig->Evaluate_AnyThread();
+	}
 	
-	return bTick ? TickTarget() : ControlRig->Evaluate_AnyThread();
+	if (TransformableHandleUtils::SkipTicking())
+	{
+		// TODO test UControlRigComponent with FAnimationEvaluationCache 
+		TransformableHandleUtils::EvaluateComponent( GetSkeletalMesh() );
+		return;
+	}
+
+	// else full tick
+	TickTarget();
 }
 
 void UTransformableControlHandle::TickTarget() const
@@ -103,18 +150,25 @@ void UTransformableControlHandle::TickTarget() const
 		return;
 	}
 	
-	if (ControlRig->IsAdditive() && ControlHandleLocals::IsRigNotifying(ControlRig.Get()))
+	if (ControlRig->IsAdditive() && UE::Private::IsRigNotifying(ControlRig.Get()))
 	{
 		return;
 	}
 	
 	if (USkeletalMeshComponent* SkeletalMeshComponent = GetSkeletalMesh())
 	{
+		if (TransformableHandleUtils::SkipTicking())
+		{
+			TransformableHandleUtils::EvaluateComponent( GetSkeletalMesh() );
+			return;
+		}
+		
 		return TransformableHandleUtils::TickDependantComponents(SkeletalMeshComponent);
 	}
 
 	if (UControlRigComponent* ControlRigComponent = GetControlRigComponent())
 	{
+		// TODO test an equivalent to EvaluateSkeletalMeshComponent
 		ControlRigComponent->Update();
 	}
 }
@@ -141,9 +195,13 @@ void UTransformableControlHandle::SetGlobalTransform(const FTransform& InGlobal)
 	static const FRigControlModifiedContext Context(EControlRigSetKey::Never);
 	static constexpr bool bNotify = false, bSetupUndo = false, bPrintPython = false, bFixEulerFlips = false;
 
-	//use this function so we don't set the preferred angles
-	ControlRig->SetControlGlobalTransform(ControlKey.Name, InGlobal.GetRelativeTransform(ComponentTransform),
-		bNotify, Context, bSetupUndo, bPrintPython, bFixEulerFlips);
+	{
+		UE::Private::FControlPoseChangedNotifier Notifier(ControlRig->GetHierarchy(), ControlElement, BoundComponent);
+	
+		//use this function so we don't set the preferred angles
+		ControlRig->SetControlGlobalTransform(ControlKey.Name, InGlobal.GetRelativeTransform(ComponentTransform),
+			bNotify, Context, bSetupUndo, bPrintPython, bFixEulerFlips);
+	}
 }
 
 void UTransformableControlHandle::SetLocalTransform(const FTransform& InLocal) const
@@ -158,7 +216,10 @@ void UTransformableControlHandle::SetLocalTransform(const FTransform& InLocal) c
 	const FRigElementKey& ControlKey = ControlElement->GetKey();
 	const int32 CtrlIndex = Hierarchy->GetIndex(ControlKey);
 	
-	Hierarchy->SetLocalTransform(CtrlIndex, InLocal);
+	{
+		UE::Private::FControlPoseChangedNotifier Notifier(ControlRig->GetHierarchy(), ControlElement, GetSkeletalMesh());
+		Hierarchy->SetLocalTransform(CtrlIndex, InLocal);
+	}
 }
 
 // NOTE should we cache the skeletal mesh and the CtrlIndex to avoid looking for if every time
@@ -389,11 +450,14 @@ bool UTransformableControlHandle::HasDirectDependencyWith(const UTransformableHa
 }
 
 // if there's no skeletal mesh bound then the handle is not valid so no need to do anything else
-FTickPrerequisite UTransformableControlHandle::GetPrimaryPrerequisite() const
+FTickPrerequisite UTransformableControlHandle::GetPrimaryPrerequisite(const bool bAllowThis) const
 {
-	if (FTickFunction* TickFunction = GetTickFunction())
+	if (bAllowThis)
 	{
-		return FTickPrerequisite( GetBoundComponent(), *TickFunction); 
+		if (FTickFunction* TickFunction = GetTickFunction())
+		{
+			return FTickPrerequisite(GetBoundComponent(), *TickFunction); 
+		}
 	}
 	
 	static const FTickPrerequisite DummyPrerex;
@@ -541,13 +605,18 @@ void UTransformableControlHandle::OnControlModified(
 			}
 			
 			// guard from re-entrant notification
-			const ControlHandleLocals::RigGuard NotificationGuard([ControlRig = ControlRig.Get()]()
+			const UE::Private::RigGuard NotificationGuard([ControlRig = ControlRig.Get()]()
 			{
-				ControlHandleLocals::NotifyingRigs.Remove(ControlRig);
+				UE::Private::NotifyingRigs.Remove(ControlRig);
 			});
-			ControlHandleLocals::NotifyingRigs.Add(ControlRig.Get());
+			UE::Private::NotifyingRigs.Add(ControlRig.Get());
 			
 			Notify(Event);
+
+			if (TransformableHandleUtils::SkipTicking())
+			{
+				TransformableHandleUtils::MarkComponentForEvaluation(GetSkeletalMesh());
+			}
 		}
 		else if (Event == EHandleEvent::GlobalTransformUpdated)
 		{
@@ -560,16 +629,21 @@ void UTransformableControlHandle::OnControlModified(
 				}
 
 				// guard from re-entrant notification 
-				const ControlHandleLocals::RigGuard NotificationGuard([ControlRig = ControlRig.Get()]()
+				const UE::Private::RigGuard NotificationGuard([ControlRig = ControlRig.Get()]()
 				{
-					ControlHandleLocals::NotifyingRigs.Remove(ControlRig);
+					UE::Private::NotifyingRigs.Remove(ControlRig);
 				});
-				ControlHandleLocals::NotifyingRigs.Add(ControlRig.Get());
+				UE::Private::NotifyingRigs.Add(ControlRig.Get());
 
 				const UTickableConstraint* TickableConstraint = GetTypedOuter<UTickableConstraint>();
 				const bool bIsConstraintActive = TickableConstraint && TickableConstraint->IsFullyActive();
 				const bool bPreTick = !ControlRig->IsAdditive() && bIsConstraintActive;
 				Notify(EHandleEvent::UpperDependencyUpdated, bPreTick);
+
+				if (TransformableHandleUtils::SkipTicking())
+				{
+					TransformableHandleUtils::MarkComponentForEvaluation(GetSkeletalMesh());
+				}
 			}
 		}
 	}
diff --git a/Engine/Plugins/Animation/ControlRig/Source/ControlRig/Public/Constraints/ControlRigTransformableHandle.h b/Engine/Plugins/Animation/ControlRig/Source/ControlRig/Public/Constraints/ControlRigTransformableHandle.h
--- a/Engine/Plugins/Animation/ControlRig/Source/ControlRig/Public/Constraints/ControlRigTransformableHandle.h
+++ b/Engine/Plugins/Animation/ControlRig/Source/ControlRig/Public/Constraints/ControlRigTransformableHandle.h
@@ -101,7 +101,7 @@ public:
 	virtual bool HasDirectDependencyWith(const UTransformableHandle& InOther) const override;
 
 	/** Look for a possible tick function that can be used as a prerequisite. */
-	virtual FTickPrerequisite GetPrimaryPrerequisite() const override;
+	virtual FTickPrerequisite GetPrimaryPrerequisite(const bool bAllowThis = true) const override;
 
 #if WITH_EDITOR
 	/** Returns labels used for UI. */
diff --git a/Engine/Source/Runtime/Experimental/Animation/Constraints/Private/Transform/AnimationEvaluation.cpp b/Engine/Source/Runtime/Experimental/Animation/Constraints/Private/Transform/AnimationEvaluation.cpp
new file mode 100644
--- /dev/null
+++ b/Engine/Source/Runtime/Experimental/Animation/Constraints/Private/Transform/AnimationEvaluation.cpp
@@ -0,0 +1,361 @@
+﻿// Copyright Epic Games, Inc. All Rights Reserved.
+
+#include "Transform/AnimationEvaluation.h"
+
+#include "Animation/AnimInstance.h"
+#include "Components/SceneComponent.h"
+#include "CoreGlobals.h"
+#include "ConstraintsManager.h"
+#include "Engine/SkeletalMesh.h"
+#include "GameFramework/Actor.h"
+#include "Logging/LogMacros.h"
+
+namespace UE::Anim
+{
+
+static bool	bLogNewEvaluation = false;
+static FAutoConsoleVariableRef CVarLogNewEvaluation(
+	TEXT("Constraints.LogNewEvaluation"),
+	bLogNewEvaluation,
+	TEXT("Log new constraints' evaluation scheme.")
+);
+
+	
+FString GetComponentName(const USceneComponent* InSceneComponent)
+{
+#if WITH_EDITOR
+	const AActor* Actor = InSceneComponent->GetOwner();
+	return Actor ? Actor->GetActorLabel() : InSceneComponent->GetName();
+#else
+	return InSceneComponent->GetName();
+#endif
+}
+
+FAnimationEvaluator::FAnimationEvaluator(USkeletalMeshComponent* InSkeletalMeshComponent)
+	: SkeletalMeshComponent(InSkeletalMeshComponent)
+{
+	if (InSkeletalMeshComponent)
+	{
+		const FOnBoneTransformsFinalizedMultiCast::FDelegate OnBoneTransformsFinalizedDelegate =
+			FOnBoneTransformsFinalizedMultiCast::FDelegate::CreateRaw(this, &FAnimationEvaluator::BoneTransformsFinalized);
+		OnBoneTransformsFinalizedHandle = InSkeletalMeshComponent->RegisterOnBoneTransformsFinalizedDelegate(OnBoneTransformsFinalizedDelegate);
+		
+		if (bLogNewEvaluation)
+		{
+			UE_LOG(LogTemp, Warning, TEXT("[%p] FAnimationEvaluator register finalize bone from %s"), this, *GetComponentName(InSkeletalMeshComponent));
+		}
+	}
+}
+
+FAnimationEvaluator::~FAnimationEvaluator()
+{
+	if (OnBoneTransformsFinalizedHandle.IsValid())
+	{
+		constexpr bool bEvenIfPendingKill = true;
+		if (USkeletalMeshComponent* RawSkeletalMeshComponent = SkeletalMeshComponent.Get(bEvenIfPendingKill))
+		{
+			RawSkeletalMeshComponent->UnregisterOnBoneTransformsFinalizedDelegate(OnBoneTransformsFinalizedHandle);
+			if (bLogNewEvaluation)
+			{
+				UE_LOG(LogTemp, Warning, TEXT("[%p]  FAnimationEvaluator unregister finalize bone from %s"), this, *GetComponentName(RawSkeletalMeshComponent));
+			}
+		}
+		OnBoneTransformsFinalizedHandle.Reset();
+	}
+}
+	
+void FAnimationEvaluator::Update(const bool bRefreshBoneTransforms)
+{
+	Context.Clear();
+	
+	if (bRefreshBoneTransforms)
+	{
+		RefreshBoneTransforms();
+	}
+}
+
+bool FAnimationEvaluator::IsValid() const
+{
+	return SkeletalMeshComponent.IsValid() && Context.SkeletalMesh && !Context.ComponentSpaceTransforms.IsEmpty();
+}
+
+FTransform FAnimationEvaluator::GetGlobalTransform(const FName InSocketName) const
+{
+	USkeletalMeshComponent* RawSkeletalMeshComponent = SkeletalMeshComponent.Get();
+	
+	if (IsValid() && InSocketName != NAME_None)
+	{
+		// TODO socket
+		const int32 BoneIndex = Context.SkeletalMesh->GetRefSkeleton().FindBoneIndex(InSocketName);
+		if (ensure(BoneIndex != INDEX_NONE && Context.ComponentSpaceTransforms.IsValidIndex(BoneIndex)))
+		{
+			return Context.ComponentSpaceTransforms[BoneIndex] * RawSkeletalMeshComponent->GetComponentTransform();
+		}
+	}
+	
+	return RawSkeletalMeshComponent ? RawSkeletalMeshComponent->GetSocketTransform(InSocketName) : FTransform::Identity;
+}
+
+void FAnimationEvaluator::UpdateContext()
+{
+	if (!SkeletalMeshComponent.IsValid())
+	{
+		Context.Clear();
+		return;
+	}
+
+	const bool bEvaluatePostProcessInstance = SkeletalMeshComponent->ShouldEvaluatePostProcessInstance();
+	
+	Context.SkeletalMesh = SkeletalMeshComponent->GetSkeletalMeshAsset();
+	Context.AnimInstance = SkeletalMeshComponent->AnimScriptInstance;
+	Context.PostProcessAnimInstance = bEvaluatePostProcessInstance? ToRawPtr(SkeletalMeshComponent->PostProcessAnimInstance): nullptr;
+	Context.bDoEvaluation = true;
+	Context.bDoInterpolation = false;
+	Context.bDuplicateToCacheBones = false;
+	Context.bDuplicateToCacheCurve = false;
+	Context.bDuplicateToCachedAttributes = false;
+	Context.bForceRefPose = false;
+}
+	
+void FAnimationEvaluator::RefreshBoneTransforms()
+{
+	// see. USkeletalMeshComponent::RefreshBoneTransforms
+	if (!SkeletalMeshComponent.IsValid())
+	{
+		return;
+	}
+
+	// avoid re-entrant animation evaluation
+	if (SkeletalMeshComponent->IsPostEvaluatingAnimation())
+	{
+		return;
+	}
+
+	constexpr bool bBlockOnTask = true, bPerformPostAnimEvaluation = true;
+	if (SkeletalMeshComponent->HandleExistingParallelEvaluationTask(bBlockOnTask, bPerformPostAnimEvaluation))
+	{
+		return;
+	}
+
+	// disable UpdateRateOptimizations
+	using UROGuard = TGuardValue_Bitfield_Cleanup<TFunction<void()>>;
+	UROGuard Guard([URORefValue = SkeletalMeshComponent->bEnableUpdateRateOptimizations, this]()
+	{
+		SkeletalMeshComponent->bEnableUpdateRateOptimizations = URORefValue;
+	});
+	SkeletalMeshComponent->bEnableUpdateRateOptimizations = false;
+
+	// update context & evaluate
+	UpdateContext();
+	EvaluateAnimation();
+}
+
+void FAnimationEvaluator::EvaluateAnimation()
+{	
+	// see USkeletalMeshComponent:DoInstancePreEvaluation()
+	{
+		if (Context.AnimInstance)
+		{
+			Context.AnimInstance->PreEvaluateAnimation();
+		}
+
+		if (Context.PostProcessAnimInstance)
+		{
+			Context.PostProcessAnimInstance->PreEvaluateAnimation();
+		}
+	}
+	
+	// call USkeletalMeshComponent::TickAnimation() if needed
+	{
+		bool bShouldTickAnimation = false;
+		
+		if (Context.AnimInstance && !Context.AnimInstance->NeedsUpdate())
+		{
+			bShouldTickAnimation = !Context.AnimInstance->GetUpdateCounter().HasEverBeenUpdated();
+		}
+
+		const bool bBPEnabled = !SkeletalMeshComponent->GetDisablePostProcessBlueprint() && SkeletalMeshComponent->ShouldEvaluatePostProcessAnimBP();
+		if (bBPEnabled && Context.PostProcessAnimInstance && !Context.PostProcessAnimInstance->NeedsUpdate())
+		{
+			bShouldTickAnimation |= !Context.PostProcessAnimInstance->GetUpdateCounter().HasEverBeenUpdated();
+		}
+
+		if (bShouldTickAnimation)
+		{
+			// We bypass TickPose() and call TickAnimation directly, so URO doesn't intercept us.
+			static constexpr float DeltaTime = 0.0f;
+			static constexpr bool bNeedsValidRootMotion = false;
+			SkeletalMeshComponent->TickAnimation(DeltaTime, bNeedsValidRootMotion);
+		}
+	}
+
+	// see USkeletalMeshComponent::DoParallelEvaluationTasks_OnGameThread()
+	{
+		Context.ComponentSpaceTransforms = SkeletalMeshComponent->GetEditableComponentSpaceTransforms();
+		Context.CachedComponentSpaceTransforms = SkeletalMeshComponent->GetCachedComponentSpaceTransforms();
+		Context.BoneSpaceTransforms = SkeletalMeshComponent->GetBoneSpaceTransforms();
+		Context.Curve = SkeletalMeshComponent->GetAnimCurves();
+		Context.RootBoneTranslation = SkeletalMeshComponent->RootBoneTranslation;
+
+		// note that don't use curves and custom attributes here 
+	}
+		
+	// see USkeletalMeshComponent::ParallelAnimationEvaluation()
+	{
+		SkeletalMeshComponent->PerformAnimationProcessing(
+		   Context.SkeletalMesh,
+		   Context.AnimInstance,
+		   Context.bDoEvaluation,
+		   Context.bForceRefPose,
+		   Context.ComponentSpaceTransforms,
+		   Context.BoneSpaceTransforms,
+		   Context.RootBoneTranslation,
+		   Context.Curve,
+		   Context.CustomAttributes);
+	}
+}
+
+void FAnimationEvaluator::BoneTransformsFinalized()
+{
+	if (!IsValid())
+	{
+		return;
+	}
+
+	bool bFinalized = false;
+
+	// NOTE FAnimationEvaluator could store a TBitArray of bones being requested and only check them instead of checking all of them 
+	const TArrayView<const FTransform> BonesTransforms = SkeletalMeshComponent->GetComponentSpaceTransforms();
+	if (Context.ComponentSpaceTransforms.Num() == BonesTransforms.Num())
+	{
+		for (int32 Index = 0; Index < BonesTransforms.Num(); ++Index)
+		{
+			FTransform& ContextBoneTransform = Context.ComponentSpaceTransforms[Index];
+			if (!ContextBoneTransform.Equals(BonesTransforms[Index]))
+			{
+				ContextBoneTransform = BonesTransforms[Index];
+				bFinalized = true;
+			}
+		}
+	}
+
+	if (bLogNewEvaluation && bFinalized)
+	{
+		UE_LOG(LogTemp, Warning, TEXT("Finalized %s's evaluator."), *GetComponentName(SkeletalMeshComponent.Get()));
+	}
+}
+	
+FAnimationEvaluationCache& FAnimationEvaluationCache::Get()
+{
+	static FAnimationEvaluationCache AnimationEvaluationCache;
+	if (!AnimationEvaluationCache.ConstraintsNotificationHandle.IsValid())
+	{
+		AnimationEvaluationCache.RegisterNotifications();
+	}
+	return AnimationEvaluationCache;
+}
+
+FAnimationEvaluationCache::~FAnimationEvaluationCache()
+{
+	UnregisterNotifications();
+	PerSkeletalMeshEvaluator.Reset();
+}
+	
+void FAnimationEvaluationCache::MarkForEvaluation(const USkeletalMeshComponent* InSkeletalMeshComponent)
+{
+	if (!InSkeletalMeshComponent)
+	{
+		return;
+	}
+
+	if (FAnimationEvaluator* AnimationEvaluator = PerSkeletalMeshEvaluator.Find(InSkeletalMeshComponent))
+	{
+		if (AnimationEvaluator->IsValid())
+		{
+			if (bLogNewEvaluation)
+			{
+				UE_LOG(LogTemp, Warning, TEXT("Marked %s's evaluator for evaluation."), *GetComponentName(InSkeletalMeshComponent));
+			}
+			constexpr bool bDoNotRefreshBones = false;
+			AnimationEvaluator->Update(bDoNotRefreshBones);
+		}
+	}
+}
+	
+const FAnimationEvaluator& FAnimationEvaluationCache::GetEvaluator(USkeletalMeshComponent* InSkeletalMeshComponent)
+{
+	if (!InSkeletalMeshComponent)
+	{
+		static const FAnimationEvaluator Invalid(nullptr);
+		return Invalid;
+	}
+	
+	constexpr bool bRefreshBones = true;
+
+	if (FAnimationEvaluator* Found = PerSkeletalMeshEvaluator.Find(InSkeletalMeshComponent))
+	{
+		if (!Found->IsValid())
+		{
+			if (bLogNewEvaluation)
+			{
+				UE_LOG(LogTemp, Warning, TEXT("Update %s's evaluator for evaluation."), *GetComponentName(InSkeletalMeshComponent));
+			}
+			Found->Update(bRefreshBones);
+		}
+		return *Found;
+	}
+	
+	FAnimationEvaluator& NewEvaluator = PerSkeletalMeshEvaluator.Emplace(InSkeletalMeshComponent, InSkeletalMeshComponent);
+	NewEvaluator.Update(bRefreshBones);
+
+	if (bLogNewEvaluation)
+	{
+		UE_LOG(LogTemp, Warning, TEXT("Create new evaluator for %s."), *GetComponentName(InSkeletalMeshComponent));
+	}
+	
+	return NewEvaluator;
+}
+
+void FAnimationEvaluationCache::RegisterNotifications()
+{
+	FConstraintsManagerController& Controller = FConstraintsManagerController::Get(nullptr);
+	ConstraintsNotificationHandle =
+		Controller.GetNotifyDelegate().AddLambda([this](EConstraintsManagerNotifyType InNotifyType, UObject* InObject)
+	{
+		if (InNotifyType == EConstraintsManagerNotifyType::GraphUpdated)
+		{
+			PerSkeletalMeshEvaluator.Reset();
+		}
+	});
+}
+
+void FAnimationEvaluationCache::UnregisterNotifications()
+{
+	if (ConstraintsNotificationHandle.IsValid())
+	{
+		FConstraintsManagerController& Controller = FConstraintsManagerController::Get(nullptr);
+		Controller.GetNotifyDelegate().Remove(ConstraintsNotificationHandle);
+		ConstraintsNotificationHandle.Reset();
+	}
+}
+
+void MarkComponentForEvaluation(const USceneComponent* InSceneComponent)
+{
+	if (const USkeletalMeshComponent* SkeletalMeshComponent = Cast<USkeletalMeshComponent>(InSceneComponent))
+	{
+		FAnimationEvaluationCache::Get().MarkForEvaluation(SkeletalMeshComponent);
+	}
+}
+
+const FAnimationEvaluator& EvaluateComponent(USceneComponent* InSceneComponent)
+{
+	if (USkeletalMeshComponent* SkeletalMeshComponent = Cast<USkeletalMeshComponent>(InSceneComponent))
+	{
+		return FAnimationEvaluationCache::Get().GetEvaluator(SkeletalMeshComponent);
+	}
+
+	static const FAnimationEvaluator Invalid(nullptr);
+	return Invalid;
+}
+
+}
\ No newline at end of file
diff --git a/Engine/Source/Runtime/Experimental/Animation/Constraints/Private/Transform/DependencyBuilder.cpp b/Engine/Source/Runtime/Experimental/Animation/Constraints/Private/Transform/DependencyBuilder.cpp
--- a/Engine/Source/Runtime/Experimental/Animation/Constraints/Private/Transform/DependencyBuilder.cpp
+++ b/Engine/Source/Runtime/Experimental/Animation/Constraints/Private/Transform/DependencyBuilder.cpp
@@ -147,6 +147,11 @@ bool HasConstraintDependencyWith(UWorld* InWorld, const UTransformableHandle* In
 	
 }
 
+bool FDependencyBuilder::LogDependencies()
+{
+	return DependencyLocals::bDebugDependencies;
+}
+
 bool FDependencyBuilder::BuildDependencies(UWorld* InWorld, UTickableTransformConstraint* InConstraint)
 {
 	using namespace DependencyLocals;
diff --git a/Engine/Source/Runtime/Experimental/Animation/Constraints/Private/Transform/TransformConstraint.cpp b/Engine/Source/Runtime/Experimental/Animation/Constraints/Private/Transform/TransformConstraint.cpp
--- a/Engine/Source/Runtime/Experimental/Animation/Constraints/Private/Transform/TransformConstraint.cpp
+++ b/Engine/Source/Runtime/Experimental/Animation/Constraints/Private/Transform/TransformConstraint.cpp
@@ -10,40 +10,50 @@
 #include "Transform/TransformableHandle.h"
 #include "Transform/TransformableRegistry.h"
 #include "Transform/DependencyBuilder.h"
+#include "Transform/TransformableHandleUtils.h"
 #include "Transform/TransformConstraintUtil.h"
 #include "UObject/Package.h"
 
 #include UE_INLINE_GENERATED_CPP_BY_NAME(TransformConstraint)
 
-namespace ConstraintLocals
+namespace UE::Private
 {
 
-static bool	bPreEvaluateChild = true;
-static FAutoConsoleVariableRef CVarPreEvaluateChild(
-	TEXT("Constraints.PreEvaluateChild"),
-	bPreEvaluateChild,
-	TEXT("Force child evaluation before constraint computation.")
-	);
+bool ShouldForceChildDependency()
+{
+	return TransformableHandleUtils::SkipTicking();
+}
+
+bool ShouldForceParentDependency()
+{
+	return !TransformableHandleUtils::SkipTicking();
+}
 	
-static bool	bPreTickChild = false;
-static FAutoConsoleVariableRef CVarPreTickChild(
-	TEXT("Constraints.PreTickChild"),
-	bPreTickChild,
-	TEXT("Force child ticking before constraint computation.")
-	);
-
-static bool	bForceChildDependency = false;
-static FAutoConsoleVariableRef CVarForceChildDependency(
-	TEXT("Constraints.ForceChildDependency"),
-	bForceChildDependency,
-	TEXT("Force child dependency on the current constraint.")
-	);
+bool ShouldPreEvaluateChild()
+{
+	return !TransformableHandleUtils::SkipTicking();
+}
+
+bool ShouldPreEvaluateParent()
+{
+	return TransformableHandleUtils::SkipTicking();
+}
 	
-void PreEvaluateHandle(const TObjectPtr<UTransformableHandle>& InHandle) 
+void PreEvaluateChild(const TObjectPtr<UTransformableHandle>& InHandle) 
 {
-	if ((bPreEvaluateChild || bPreTickChild) && ::IsValid(InHandle))
+	if (ShouldPreEvaluateChild() && ::IsValid(InHandle))
 	{
-		InHandle->PreEvaluate(bPreTickChild);
+		constexpr bool bTickChild = false;
+		InHandle->PreEvaluate(bTickChild);
+	}
+}
+
+void PreEvaluateParent(const TObjectPtr<UTransformableHandle>& InHandle) 
+{
+	if (ShouldPreEvaluateParent() && ::IsValid(InHandle))
+	{
+		constexpr bool bTickParent = false;
+		InHandle->PreEvaluate(bTickParent);
 	}
 }
 
@@ -76,6 +86,15 @@ uint32 GetConstrainableHash(const UObject* InObject)
 	return 0;
 }
 
+FString GetConstraintLabel(const UTickableConstraint* InConstraint)
+{
+#if WITH_EDITOR
+	return InConstraint->GetFullLabel();
+#else
+	return InConstraint->GetName();
+#endif		
+}
+	
 }
 
 /** 
@@ -103,7 +122,7 @@ EMovieSceneTransformChannel UTickableTransformConstraint::GetChannelsToKey() con
 	const ETransformConstraintType ConstType = static_cast<ETransformConstraintType>(GetType());
 	if (const EMovieSceneTransformChannel* Channel = ConstraintToChannels.Find(ConstType))
 	{
-		return *Channel;;
+		return *Channel;
 	}
 
 	return EMovieSceneTransformChannel::AllTransform;
@@ -285,7 +304,7 @@ void UTickableTransformConstraint::SetupDependencies(const UWorld* InWorld)
 	FTickFunction* ParentTickFunction = GetParentHandleTickFunction();
 	FTickFunction* ChildTickFunction = GetChildHandleTickFunction();
 	
-	if (ParentTickFunction && (ChildTickFunction != ParentTickFunction))
+	if (UE::Private::ShouldForceParentDependency() && ParentTickFunction && (ChildTickFunction != ParentTickFunction))
 	{
 		// manage dependencies
 		// force ConstraintTickFunction to tick after InParent does.
@@ -298,11 +317,11 @@ void UTickableTransformConstraint::SetupDependencies(const UWorld* InWorld)
 	if (ChildTickFunction)
 	{
 		USkeletalMeshComponent* ChildOwner = Cast<USkeletalMeshComponent>(ChildTRSHandle->GetTarget().Get());
-		if (ChildOwner == nullptr || ConstraintLocals::bForceChildDependency)
+		if (ChildOwner == nullptr || UE::Private::ShouldForceChildDependency())
 		{
 			// force InChild to tick after ConstraintTickFunction does.
 			// Note that this might not register anything if the child can't tick (static meshes for instance)
-			// Also don't do if it's a skel mesh will cause cycle.
+			// Also don't do if it's a skeletal mesh will cause cycle. (to be updated)
 			FConstraintTickFunction& ConstraintTick = ConstraintTicks.FindOrAdd(InWorld->GetCurrentLevel());
 			ChildTickFunction->AddPrerequisite(this, ConstraintTick);
 		}
@@ -327,9 +346,16 @@ void UTickableTransformConstraint::EnsurePrimaryDependency(const UWorld* InWorld
 			// if the constraint has no prerex at this stage, this means that the parent tick function
 			// is not registered or can't tick (static meshes for instance) so look for the first parent tick function if any.
 			// In a context of adding several constraints, we want to make sure that the evaluation order is the right one
-			FTickPrerequisite PrimaryPrerex = ParentTRSHandle->GetPrimaryPrerequisite();
+			FTickPrerequisite PrimaryPrerex = ParentTRSHandle->GetPrimaryPrerequisite(UE::Private::ShouldForceParentDependency());
 			if (FTickFunction* PotentialFunction = PrimaryPrerex.Get())
 			{
+				if (FDependencyBuilder::LogDependencies())
+				{
+					UE_LOG(LogTemp, Warning, TEXT("EnsurePrimaryDependency: '%s' must tick before '%s'"),
+						*PotentialFunction->DiagnosticMessage(),
+						*ConstraintTick.DiagnosticMessage() );
+				}
+				
 				UObject* Target = PrimaryPrerex.PrerequisiteObject.Get();
 				ConstraintTick.AddPrerequisite(Target, *PotentialFunction);
 			}
@@ -408,7 +434,7 @@ void UTickableTransformConstraint::Evaluate(bool bTickHandlesAlso) const
 {
 	if (IsFullyActive())
 	{
-		if (bTickHandlesAlso)
+		if (bTickHandlesAlso && !TransformableHandleUtils::SkipTicking())
 		{
 			ParentTRSHandle->TickTarget();
 		}
@@ -504,7 +530,7 @@ void UTickableTransformConstraint::OnHandleModified(UTransformableHandle* InHand
 	{
 		if (InNotification == EHandleEvent::UpperDependencyUpdated)
 		{
-			const UObject* ParentTarget = ConstraintLocals::GetHandleTarget(ParentTRSHandle);
+			const UObject* ParentTarget = UE::Private::GetHandleTarget(ParentTRSHandle);
 			if (ParentTarget && ParentTarget != Target)
 			{
 				return MarkForEvaluation();
@@ -566,6 +592,25 @@ FTickFunction* UTickableTransformConstraint::GetHandleTickFunction(const TObject
 	return InHandle->GetTickFunction();
 }
 
+void UTickableTransformConstraint::PreEvaluate() const
+{
+	UE::Private::PreEvaluateParent(ParentTRSHandle);
+	UE::Private::PreEvaluateChild(ChildTRSHandle);
+}
+
+void UTickableTransformConstraint::PostEvaluate() const
+{
+	if (TransformableHandleUtils::SkipTicking())
+	{
+		return;	
+	}
+	
+	if (ChildTRSHandle)
+	{
+		ChildTRSHandle->TickTarget();
+	}
+}
+
 void UTickableTransformConstraint::InitConstraint(UWorld *InWorld)
 {
 	if (!ensure(InWorld))
@@ -596,10 +641,17 @@ void UTickableTransformConstraint::TeardownConstraint(UWorld* InWorld)
 	FConstraintTickFunction& ConstraintTick = ConstraintTicks.FindOrAdd(InWorld->GetCurrentLevel());
 	ConstraintTick.SetTickFunctionEnable(false);
 	ConstraintTick.UnRegisterTickFunction();
+
 	if (FTickFunction* ChildTickFunction = GetChildHandleTickFunction())
 	{
 		ChildTickFunction->RemovePrerequisite(this, ConstraintTick);
 	}
+	
+	if (FTickFunction* ParentTickFunction = GetParentHandleTickFunction())
+	{
+		ConstraintTick.RemovePrerequisite(ParentTRSHandle->GetPrerequisiteObject(), *ParentTickFunction);	
+	}
+	
 	ConstraintTicks.Remove(InWorld->GetCurrentLevel());
 
 	// unregister delegates. should handles delegates be unregistered as well?
@@ -681,7 +733,7 @@ FConstraintTickFunction::ConstraintFunction UTickableTranslationConstraint::GetF
 			return;
 		}
 
-		ConstraintLocals::PreEvaluateHandle(ChildTRSHandle);
+		PreEvaluate();
 		
 		const FVector ParentTranslation = GetParentGlobalTransform().GetLocation();
 		FTransform Transform = GetChildGlobalTransform();
@@ -698,10 +750,8 @@ FConstraintTickFunction::ConstraintFunction UTickableTranslationConstraint::GetF
 		Transform.SetLocation(NewTranslation);
 			
 		SetChildGlobalTransform(Transform);
-		if (ChildTRSHandle)
-		{
-			ChildTRSHandle->TickTarget();
-		}
+
+		PostEvaluate();
 	};
 }
 
@@ -829,7 +879,7 @@ FConstraintTickFunction::ConstraintFunction UTickableRotationConstraint::GetFunc
 			return;
 		}
 
-		ConstraintLocals::PreEvaluateHandle(ChildTRSHandle);
+		PreEvaluate();
 		
 		const FQuat ParentRotation = GetParentGlobalTransform().GetRotation();
 		FTransform Transform = GetChildGlobalTransform();
@@ -848,10 +898,7 @@ FConstraintTickFunction::ConstraintFunction UTickableRotationConstraint::GetFunc
 		
 		SetChildGlobalTransform(Transform);
 
-		if (ChildTRSHandle)
-		{
-			ChildTRSHandle->TickTarget();
-		}
+		PostEvaluate();
 	};
 }
 
@@ -983,7 +1030,7 @@ FConstraintTickFunction::ConstraintFunction UTickableScaleConstraint::GetFunctio
 			return;
 		}
 
-		ConstraintLocals::PreEvaluateHandle(ChildTRSHandle);
+		PreEvaluate();
 		
 		const FVector ParentScale = GetParentGlobalTransform().GetScale3D();
 		FTransform Transform = GetChildGlobalTransform();
@@ -1000,10 +1047,8 @@ FConstraintTickFunction::ConstraintFunction UTickableScaleConstraint::GetFunctio
 		Transform.SetScale3D(NewScale);
 		
 		SetChildGlobalTransform(Transform);
-		if (ChildTRSHandle)
-		{
-			ChildTRSHandle->TickTarget();
-		}
+
+		PostEvaluate();
 	};
 }
 
@@ -1150,7 +1195,7 @@ FConstraintTickFunction::ConstraintFunction UTickableParentConstraint::GetFuncti
 			OutTransform.SetScale3D(NewScale);
 		};
 
-		ConstraintLocals::PreEvaluateHandle(ChildTRSHandle);
+		PreEvaluate();
 		
 		const FTransform ParentTransform = GetParentGlobalTransform();
 		
@@ -1166,10 +1211,8 @@ FConstraintTickFunction::ConstraintFunction UTickableParentConstraint::GetFuncti
 		}
 
 		SetChildGlobalTransform(TargetTransform);
-		if (ChildTRSHandle)
-		{
-			ChildTRSHandle->TickTarget();
-		}
+
+		PostEvaluate();
 	};
 }
 
@@ -1305,7 +1348,7 @@ FConstraintTickFunction::ConstraintFunction UTickableLookAtConstraint::GetFuncti
 			return;
 		}
 
-		ConstraintLocals::PreEvaluateHandle(ChildTRSHandle);
+		PreEvaluate();
 		
 		const FTransform ParentTransform = GetParentGlobalTransform();
 		const FTransform ChildTransform = GetChildGlobalTransform();
@@ -1325,10 +1368,8 @@ FConstraintTickFunction::ConstraintFunction UTickableLookAtConstraint::GetFuncti
 				FTransform Transform = ChildTransform;
 				Transform.SetRotation(Rotation.GetNormalized());
 				SetChildGlobalTransform(Transform);
-				if (ChildTRSHandle)
-				{
-					ChildTRSHandle->TickTarget();
-				}
+
+				PostEvaluate();
 			}
 		}
 	};
diff --git a/Engine/Source/Runtime/Experimental/Animation/Constraints/Private/Transform/TransformableHandle.cpp b/Engine/Source/Runtime/Experimental/Animation/Constraints/Private/Transform/TransformableHandle.cpp
--- a/Engine/Source/Runtime/Experimental/Animation/Constraints/Private/Transform/TransformableHandle.cpp
+++ b/Engine/Source/Runtime/Experimental/Animation/Constraints/Private/Transform/TransformableHandle.cpp
@@ -99,7 +99,34 @@ void UTransformableComponentHandle::TickTarget() const
 	{
 		return;
 	}
-	TransformableHandleUtils::TickDependantComponents(Component.Get());
+	
+	if (TransformableHandleUtils::SkipTicking())
+	{
+		if (SocketName != NAME_None)
+		{
+			TransformableHandleUtils::EvaluateComponent(Component.Get()); 
+		}
+	}
+	else
+	{
+		TransformableHandleUtils::TickDependantComponents(Component.Get());	
+	}
+}
+
+void UTransformableComponentHandle::PreEvaluate(const bool bTickSkeletalMeshes) const
+{
+	if (!Component.IsValid())
+	{
+		return;
+	}
+	
+	if (TransformableHandleUtils::SkipTicking())
+	{
+		if (SocketName != NAME_None)
+		{
+			TransformableHandleUtils::EvaluateComponent(Component.Get());
+		}
+	}
 }
 
 void UTransformableComponentHandle::SetGlobalTransform(const FTransform& InGlobal) const
@@ -150,6 +177,21 @@ FTransform UTransformableComponentHandle::GetLocalTransform() const
 
 FTransform UTransformableComponentHandle::GetGlobalTransform() const
 {
+	if (!Component.IsValid())
+	{
+		return FTransform::Identity;
+	}
+	
+	using namespace TransformableHandleUtils;
+	if (SkipTicking())
+	{
+		const UE::Anim::FAnimationEvaluator& AnimationEvaluator = EvaluateComponent(Component.Get());
+		if (AnimationEvaluator.IsValid())
+		{
+			return AnimationEvaluator.GetGlobalTransform(SocketName);
+		}
+	}
+	
 	return Component.IsValid() ? Component->GetSocketTransform(SocketName) : FTransform::Identity;
 }
 
@@ -200,7 +242,7 @@ bool UTransformableComponentHandle::HasDirectDependencyWith(const UTransformable
 namespace
 {
 	
-FTickPrerequisite LookForPrimaryPrerequisite(USceneComponent* Component)
+FTickPrerequisite LookForPrimaryPrerequisite(USceneComponent* Component, const bool bAllowThis)
 {
 	static const FTickPrerequisite DummyPrerex;
 	
@@ -220,7 +262,7 @@ FTickPrerequisite LookForPrimaryPrerequisite(USceneComponent* Component)
 	};
 
 	// check if this can tick
-	if ( IsValidTickFunction(&Component->PrimaryComponentTick) )
+	if (IsValidTickFunction(&Component->PrimaryComponentTick) && bAllowThis)
 	{
 		return FTickPrerequisite(Component, Component->PrimaryComponentTick);
 	}
@@ -256,7 +298,8 @@ FTickPrerequisite LookForPrimaryPrerequisite(USceneComponent* Component)
 		}
 
 		// get up the hierarchy if not found
-		return LookForPrimaryPrerequisite(Parent);
+		constexpr bool bAllowTickFromParent = true;
+		return LookForPrimaryPrerequisite(Parent, bAllowTickFromParent);
 	}
 
 	return DummyPrerex;
@@ -264,11 +307,11 @@ FTickPrerequisite LookForPrimaryPrerequisite(USceneComponent* Component)
 	
 }
 
-FTickPrerequisite UTransformableComponentHandle::GetPrimaryPrerequisite() const
+FTickPrerequisite UTransformableComponentHandle::GetPrimaryPrerequisite(const bool bAllowThis) const
 {
 	if (Component.IsValid())
 	{
-		return LookForPrimaryPrerequisite(Component.Get());
+		return LookForPrimaryPrerequisite(Component.Get(), bAllowThis);
 	}
 
 	static const FTickPrerequisite DummyPrerex;
diff --git a/Engine/Source/Runtime/Experimental/Animation/Constraints/Private/Transform/TransformableHandleUtils.cpp b/Engine/Source/Runtime/Experimental/Animation/Constraints/Private/Transform/TransformableHandleUtils.cpp
--- a/Engine/Source/Runtime/Experimental/Animation/Constraints/Private/Transform/TransformableHandleUtils.cpp
+++ b/Engine/Source/Runtime/Experimental/Animation/Constraints/Private/Transform/TransformableHandleUtils.cpp
@@ -10,6 +10,8 @@ namespace TransformableHandleUtils
 {
 
 static int32 SkeletalMeshTickingMode = 2;
+static bool	bNewEvaluation = false;
+
 static FAutoConsoleVariableRef CVarSkeletalMeshTickingMode(
 	TEXT("Constraints.SkeletalMesh.TickingMode"),
 	SkeletalMeshTickingMode,
@@ -22,10 +24,26 @@ static FAutoConsoleVariableRef CVarSkeletalMeshTickingMode(
 		SkeletalMeshTickingMode = FMath::Clamp(SkeletalMeshTickingMode, 0, 2);
 	})
 );
-	
+
+static FAutoConsoleVariableRef CVarNewEvaluation(
+	TEXT("Constraints.NewEvaluation"),
+	bNewEvaluation,
+	TEXT("Set constraints new evaluation scheme."),
+	FConsoleVariableDelegate::CreateLambda([](IConsoleVariable*)
+	{
+		UE_LOG(LogTemp, Warning, TEXT("Constraints' new evaluation scheme %s"), bNewEvaluation ? TEXT("enabled.") : TEXT("disabled."));
+		SkeletalMeshTickingMode = bNewEvaluation ? 0 : 2;
+	})
+);
+
+bool SkipTicking()
+{
+	return bNewEvaluation;
+}
+
 void TickDependantComponents(USceneComponent* InComponent)
 {
-	if (SkeletalMeshTickingMode <= 0)
+	if (SkeletalMeshTickingMode <= 0 || bNewEvaluation)
 	{
 		return;
 	}
@@ -52,7 +70,7 @@ void TickDependantComponents(USceneComponent* InComponent)
 
 void TickSkeletalMeshComponent(USkeletalMeshComponent* InSkeletalMeshComponent)
 {
-	if (!InSkeletalMeshComponent)
+	if (!InSkeletalMeshComponent || bNewEvaluation)
 	{
 		return;
 	}
@@ -75,4 +93,23 @@ void TickSkeletalMeshComponent(USkeletalMeshComponent* InSkeletalMeshComponent)
 	InSkeletalMeshComponent->MarkRenderDynamicDataDirty();
 }
 	
+void MarkComponentForEvaluation(const USceneComponent* InSceneComponent)
+{
+	if (const USkeletalMeshComponent* SkeletalMeshComponent = Cast<USkeletalMeshComponent>(InSceneComponent))
+	{
+		UE::Anim::FAnimationEvaluationCache::Get().MarkForEvaluation(SkeletalMeshComponent);
+	}
+}
+
+const UE::Anim::FAnimationEvaluator& EvaluateComponent(USceneComponent* InSceneComponent)
+{
+	if (USkeletalMeshComponent* SkeletalMeshComponent = Cast<USkeletalMeshComponent>(InSceneComponent))
+	{
+		return UE::Anim::FAnimationEvaluationCache::Get().GetEvaluator(SkeletalMeshComponent);
+	}
+
+	static const UE::Anim::FAnimationEvaluator Invalid(nullptr);
+	return Invalid;
+}
+	
 }
\ No newline at end of file
diff --git a/Engine/Source/Runtime/Experimental/Animation/Constraints/Public/Transform/AnimationEvaluation.h b/Engine/Source/Runtime/Experimental/Animation/Constraints/Public/Transform/AnimationEvaluation.h
new file mode 100644
--- /dev/null
+++ b/Engine/Source/Runtime/Experimental/Animation/Constraints/Public/Transform/AnimationEvaluation.h
@@ -0,0 +1,94 @@
+﻿// Copyright Epic Games, Inc. All Rights Reserved.
+
+#pragma once
+
+#include "Components/SkeletalMeshComponent.h"
+#include "Containers/Map.h"
+#include "Delegates/IDelegateInstance.h"
+#include "Math/Transform.h"
+#include "UObject/WeakObjectPtrTemplates.h"
+
+namespace UE::Anim
+{
+
+/**
+ * FAnimationEvaluator allows evaluating the animation of a skeletal mesh without having to tick it directly.
+ * USkeletalMeshComponent::TickAnimation & USkeletalMeshComponent::RefreshBoneTransforms update internal data, in particular bones transforms,
+ * which leads to several side effects (motion blur for example) as well as poor performance.
+ * This structure provides a cached equivalent of USkeletalMeshComponent::RefreshBoneTransforms, which is only updated on demand.
+ * It's also bound to USkeletalMeshComponent::OnBoneTransformsFinalizedMC so that it's updated with the real values once the skeletal mesh is 
+ * finally up to date.
+*/
+	
+struct FAnimationEvaluator
+{
+	FAnimationEvaluator(USkeletalMeshComponent* InSkeletalMeshComponent);
+	~FAnimationEvaluator();
+
+	/** Clears the context and refresh the bones transforms if bRefreshBoneTransforms is true. */ 
+	void Update(const bool bRefreshBoneTransforms);
+	
+	/** Returns true if the skeletal mesh component, its skeletal mesh and the transforms are valid for use. */
+	bool IsValid() const;
+	
+	/** Returns InSocketName component space transform composed with the skeletal mesh component's global transform. */
+	FTransform GetGlobalTransform(const FName InSocketName) const;
+
+private:
+
+	FAnimationEvaluator() = default;
+
+	/** Prepares the animation context for evaluation. */
+	void UpdateContext();
+	
+	/** Evaluates the animation using the Context. */
+	void EvaluateAnimation();
+
+	/** Evaluates the animation and store the bones transforms. */
+	void RefreshBoneTransforms();
+
+	/** Updates the bones transforms once the skeletal mesh is finally up to date. */
+	void BoneTransformsFinalized();
+	
+	TWeakObjectPtr<USkeletalMeshComponent> SkeletalMeshComponent = nullptr;
+	FAnimationEvaluationContext Context;
+	FDelegateHandle OnBoneTransformsFinalizedHandle;
+};
+
+/**
+ * FAnimationEvaluationCache provides a cache mechanism connecting skeletal mesh component and its animation evaluator to avoid multiplying
+ * FAnimationEvaluator instances when they all refer to the same skeletal mesh component. The animation evaluator is built lazily and dirtied
+ * manually on demand  that the next call to GetEvaluator actually evaluates the animation if needed.  
+*/
+	
+struct FAnimationEvaluationCache
+{
+public:
+	static FAnimationEvaluationCache& Get();
+	~FAnimationEvaluationCache();
+
+	/** Marks the evaluator related to InSkeletalMeshComponent for evaluation. */
+	void MarkForEvaluation(const USkeletalMeshComponent* InSkeletalMeshComponent);
+	
+	/**
+	 * Returns and up-to-date evaluator to be queried.
+	 * Note that calling this function will actually evaluate the animation if the evaluator has been previously dirtied.
+	 */
+	const FAnimationEvaluator& GetEvaluator(USkeletalMeshComponent* InSkeletalMeshComponent);
+
+private:
+	FAnimationEvaluationCache() = default;
+
+	/** Listen to the constraint manager notification to avoid keeping useless or invalid data. */
+	void RegisterNotifications();
+
+	/** Stop listening to the constraint manager notification system. */
+	void UnregisterNotifications();
+
+	/** List of connected skeletal mesh components and animation evaluators. */
+	TMap<USkeletalMeshComponent*, FAnimationEvaluator> PerSkeletalMeshEvaluator;
+
+	FDelegateHandle ConstraintsNotificationHandle;
+};
+
+}
\ No newline at end of file
diff --git a/Engine/Source/Runtime/Experimental/Animation/Constraints/Public/Transform/DependencyBuilder.h b/Engine/Source/Runtime/Experimental/Animation/Constraints/Public/Transform/DependencyBuilder.h
--- a/Engine/Source/Runtime/Experimental/Animation/Constraints/Public/Transform/DependencyBuilder.h
+++ b/Engine/Source/Runtime/Experimental/Animation/Constraints/Public/Transform/DependencyBuilder.h
@@ -20,6 +20,9 @@ struct FDependencyBuilder
 public: 
 	/** Ensure default dependencies between constraints. */
 	static bool BuildDependencies(UWorld* InWorld, UTickableTransformConstraint* Constraint);
+
+	/** Returns true if dependencies should be logged. */
+	static bool LogDependencies();
 	
 private:
 	
diff --git a/Engine/Source/Runtime/Experimental/Animation/Constraints/Public/Transform/TransformConstraint.h b/Engine/Source/Runtime/Experimental/Animation/Constraints/Public/Transform/TransformConstraint.h
--- a/Engine/Source/Runtime/Experimental/Animation/Constraints/Public/Transform/TransformConstraint.h
+++ b/Engine/Source/Runtime/Experimental/Animation/Constraints/Public/Transform/TransformConstraint.h
@@ -155,6 +155,12 @@ protected:
 	/** Returns the handle's tick function (ensuring it lives in the same world). */
 	CONSTRAINTS_API FTickFunction* GetHandleTickFunction(const TObjectPtr<UTransformableHandle>& InHandle) const;
 
+	/** Perform necessary pre-evaluation steps before calculating the constraint's result. */
+	void PreEvaluate() const;
+
+	/** Perform necessary post-evaluation steps after calculating the constraint's result. */
+    void PostEvaluate() const;
+
 public:
 	/** (Re-)Registers the constraint function and (re-)binds the required delegates*/
 	CONSTRAINTS_API virtual void InitConstraint(UWorld * InWorld) override;
diff --git a/Engine/Source/Runtime/Experimental/Animation/Constraints/Public/Transform/TransformableHandle.h b/Engine/Source/Runtime/Experimental/Animation/Constraints/Public/Transform/TransformableHandle.h
--- a/Engine/Source/Runtime/Experimental/Animation/Constraints/Public/Transform/TransformableHandle.h
+++ b/Engine/Source/Runtime/Experimental/Animation/Constraints/Public/Transform/TransformableHandle.h
@@ -123,7 +123,7 @@ public:
 	CONSTRAINTS_API virtual bool HasDirectDependencyWith(const UTransformableHandle& InOther) const PURE_VIRTUAL(HasDirectDependencyWith, return false;);
 
 	/** Look for a possible tick function that can be used as a prerequisite. */
-	CONSTRAINTS_API virtual FTickPrerequisite GetPrimaryPrerequisite() const PURE_VIRTUAL(GetPrimaryPrerequisite, return FTickPrerequisite(););
+	CONSTRAINTS_API virtual FTickPrerequisite GetPrimaryPrerequisite(const bool bAllowThis = true) const PURE_VIRTUAL(GetPrimaryPrerequisite, return FTickPrerequisite(););
 	
 	CONSTRAINTS_API FHandleModifiedEvent& HandleModified();
 	CONSTRAINTS_API void Notify(EHandleEvent InEvent, const bool bPreTickTarget = false) const;
@@ -176,6 +176,14 @@ public:
 	CONSTRAINTS_API virtual FTransform GetLocalTransform() const override;
 	/** Tick any skeletal mesh related to the component. */
 	CONSTRAINTS_API virtual void TickTarget() const override;
+
+	/**
+	 * Perform any pre-evaluation of the handle to ensure that the transform data is up to date.
+	 * @param bTickSkeletalMeshes to force any pre-evaluation ticking. The rig will still be pre-evaluated even
+	 * if bTick is false (it just won't tick the bound skeletal meshes) Default is false.
+	*/
+	CONSTRAINTS_API virtual void PreEvaluate(const bool bTickSkeletalMeshes = false) const override;
+	
 	/** Returns the target object containing the tick function (e.i. Component). */
 	CONSTRAINTS_API virtual UObject* GetPrerequisiteObject() const override;
 	/** Returns Component's tick function. */
@@ -190,8 +198,11 @@ public:
 	/** Check for direct dependencies (ie hierarchy) with InOther. */
 	CONSTRAINTS_API virtual bool HasDirectDependencyWith(const UTransformableHandle& InOther) const override;
 
-	/** Look for a possible tick function that can be used as a prerequisite. */
-	CONSTRAINTS_API virtual FTickPrerequisite GetPrimaryPrerequisite() const override;
+	/** 
+	 * Look for a possible tick function that can be used as a prerequisite.
+	 * If bAllowThis is false then the component's tick function is not allowed as a primary prerequisite.
+	 */
+	CONSTRAINTS_API virtual FTickPrerequisite GetPrimaryPrerequisite(const bool bAllowThis = true) const override;
 
 	/** Get the array of float channels for the specified section*/
 	CONSTRAINTS_API virtual TArrayView<FMovieSceneFloatChannel*>  GetFloatChannels(const UMovieSceneSection* InSection) const override;
diff --git a/Engine/Source/Runtime/Experimental/Animation/Constraints/Public/Transform/TransformableHandleUtils.h b/Engine/Source/Runtime/Experimental/Animation/Constraints/Public/Transform/TransformableHandleUtils.h
--- a/Engine/Source/Runtime/Experimental/Animation/Constraints/Public/Transform/TransformableHandleUtils.h
+++ b/Engine/Source/Runtime/Experimental/Animation/Constraints/Public/Transform/TransformableHandleUtils.h
@@ -2,6 +2,7 @@
 
 #pragma once
 
+#include "AnimationEvaluation.h"
 #include "HAL/Platform.h"
 
 class USceneComponent;
@@ -9,6 +10,18 @@ class USkeletalMeshComponent;
 
 namespace TransformableHandleUtils
 {
+	/** Returns true if the constraints' evaluation scheme shoud skip ticking skeletal meshes. */
+	CONSTRAINTS_API bool SkipTicking();
+
+	/** Force ticking all the skeletal meshes related to this component. */
 	CONSTRAINTS_API void TickDependantComponents(USceneComponent* InComponent);
+	
+	/** Force ticking InSkeletalMeshComponent. */
 	CONSTRAINTS_API void TickSkeletalMeshComponent(USkeletalMeshComponent* InSkeletalMeshComponent);
+
+	/** Mark InSceneComponent for animation evaluation. */
+	CONSTRAINTS_API void MarkComponentForEvaluation(const USceneComponent* InSceneComponent);
+	
+	/** Returns an updated version of InSceneComponent's animation evaluator. */
+	CONSTRAINTS_API const UE::Anim::FAnimationEvaluator& EvaluateComponent(USceneComponent* InSceneComponent);
 }
