Constraints: fixed static analysis issue

#rnx

diff --git a/Engine/Source/Runtime/Experimental/Animation/Constraints/Private/Transform/AnimationEvaluation.cpp b/Engine/Source/Runtime/Experimental/Animation/Constraints/Private/Transform/AnimationEvaluation.cpp
--- a/Engine/Source/Runtime/Experimental/Animation/Constraints/Private/Transform/AnimationEvaluation.cpp
+++ b/Engine/Source/Runtime/Experimental/Animation/Constraints/Private/Transform/AnimationEvaluation.cpp
@@ -82,6 +82,10 @@ bool FAnimationEvaluator::IsValid() const
 FTransform FAnimationEvaluator::GetGlobalTransform(const FName InSocketName) const
 {
 	USkeletalMeshComponent* RawSkeletalMeshComponent = SkeletalMeshComponent.Get();
+	if (!RawSkeletalMeshComponent)
+	{
+		return FTransform::Identity;
+	}
 	
 	if (IsValid() && InSocketName != NAME_None)
 	{
@@ -93,7 +97,7 @@ FTransform FAnimationEvaluator::GetGlobalTransform(const FName InSocketName) con
 		}
 	}
 	
-	return RawSkeletalMeshComponent ? RawSkeletalMeshComponent->GetSocketTransform(InSocketName) : FTransform::Identity;
+	return RawSkeletalMeshComponent->GetSocketTransform(InSocketName);
 }
 
 void FAnimationEvaluator::UpdateContext()
