Updated I/O store on-demand statistics

#rnx

diff --git a/Engine/Source/Runtime/Experimental/IoStore/OnDemand/Private/OnDemandContentInstaller.cpp b/Engine/Source/Runtime/Experimental/IoStore/OnDemand/Private/OnDemandContentInstaller.cpp
--- a/Engine/Source/Runtime/Experimental/IoStore/OnDemand/Private/OnDemandContentInstaller.cpp
+++ b/Engine/Source/Runtime/Experimental/IoStore/OnDemand/Private/OnDemandContentInstaller.cpp
@@ -467,7 +467,7 @@ void FOnDemandContentInstaller::ProcessDownloadedChunk(
 
 	RequestState.Progress.CurrentInstallSize += ChunkEntry.EncodedSize;
 
-	UE_LOG(LogIoStoreOnDemand, Log, TEXT("Install progress %.2lf/%.2lf MiB, ContentHandle=(%s), ChunkId='%s', ChunkSize=%.2lf KiB, HttpStatus=%u"),
+	UE_LOG(LogIoStoreOnDemand, Verbose, TEXT("Install progress %.2lf/%.2lf MiB, ContentHandle=(%s), ChunkId='%s', ChunkSize=%.2lf KiB, HttpStatus=%u"),
 		double(RequestState.Progress.CurrentInstallSize) / 1024.0 / 1024.0, double(RequestState.Progress.TotalInstallSize) / 1024.0 / 1024.0,
 		*LexToString(Request.Args.ContentHandle), *LexToString(ChunkId), double(ChunkEntry.EncodedSize) / 1024.0, HttpStatusCode);
 
@@ -536,12 +536,13 @@ void FOnDemandContentInstaller::NotifyProgress(FInstallRequestState& RequestStat
 		return;
 	}
 
-	const double SecondsSinceLastProgress = FPlatformTime::Seconds() - RequestState.LastProgressTimeStamp;
+	const uint64 Cycles = FPlatformTime::Cycles64();
+	const double SecondsSinceLastProgress = FPlatformTime::ToSeconds64(Cycles - RequestState.LastProgressCycles);
 	if (SecondsSinceLastProgress < .25)
 	{
 		return;
 	}
-	RequestState.LastProgressTimeStamp = FPlatformTime::Seconds();
+	RequestState.LastProgressCycles = Cycles; 
 
 	FOnDemandInstallProgress Progress = RequestState.Progress;
 	if (EnumHasAnyFlags(Request.Args.Options, EOnDemandInstallOptions::CallbackOnGameThread))
@@ -593,21 +594,27 @@ void FOnDemandContentInstaller::CompleteRequest(FInstallRequestState& RequestSta
 		IoStore.PackageStoreBackend->NeedsUpdate(EOnDemandPackageStoreUpdateMode::ReferencedPackages);
 	}
 
+	const uint64 DurationCycles = FPlatformTime::Cycles64() - RequestState.StartTimeCycles; 
+	const double CacheHitRatio	= RequestState.Progress.TotalContentSize > 0
+		? double(RequestState.Progress.TotalContentSize - RequestState.Progress.TotalInstallSize) / double(RequestState.Progress.TotalContentSize)
+		: 0.0;
+
 	FOnDemandInstallResult InstallResult;
 	InstallResult.Status				= Request->ErrorCode == EIoErrorCode::Ok ? FIoStatus::Ok : FIoStatus(Request->ErrorCode, RequestState.ErrorReason);
 	InstallResult.Progress				= RequestState.Progress;
-	InstallResult.DurationInSeconds		= FPlatformTime::Seconds() - Request->StartTime;
+	InstallResult.DurationInSeconds		= FPlatformTime::ToSeconds64(DurationCycles);
 
 	FOnDemandContentInstallerStats::OnRequestCompleted(
 		RequestState.ResolvedChunkCount,
 		InstallResult.Progress.TotalContentSize,
 		static_cast<uint64>(RequestState.HttpRequestHandles.Num()),
 		InstallResult.Progress.TotalInstallSize,
-		InstallResult.DurationInSeconds);
+		CacheHitRatio,
+		DurationCycles);
 
-	UE_LOG(LogIoStoreOnDemand, Log, TEXT("Install request completed, Result='%s', ContentHandle=(%s), ContentSize=%.2lf MiB, InstallSize=%.2lf MiB, Duration=%.2lfs"),
+	UE_LOG(LogIoStoreOnDemand, Log, TEXT("Install request completed, Result='%s', ContentHandle=(%s), ContentSize=%.2lf MiB, InstallSize=%.2lf MiB, CacheHitRatio=%d %%, Duration=%d ms"),
 		GetIoErrorText(InstallResult.Status.GetErrorCode()), *LexToString(Request->Args.ContentHandle), double(InstallResult.Progress.TotalContentSize) / 1024.0 / 1024.0,
-		double(InstallResult.Progress.TotalInstallSize) / 1024.0 / 1024.0, InstallResult.DurationInSeconds);
+		double(InstallResult.Progress.TotalInstallSize) / 1024.0 / 1024.0, int32(CacheHitRatio * 100), int32(InstallResult.DurationInSeconds * 1000));
 
 	{
 		UE::TUniqueLock Lock(Mutex);
diff --git a/Engine/Source/Runtime/Experimental/IoStore/OnDemand/Private/OnDemandContentInstaller.h b/Engine/Source/Runtime/Experimental/IoStore/OnDemand/Private/OnDemandContentInstaller.h
--- a/Engine/Source/Runtime/Experimental/IoStore/OnDemand/Private/OnDemandContentInstaller.h
+++ b/Engine/Source/Runtime/Experimental/IoStore/OnDemand/Private/OnDemandContentInstaller.h
@@ -70,7 +70,8 @@ class FOnDemandContentInstaller
 		FOnDemandInstallProgress					Progress;
 		uint64										ResolvedChunkCount = 0;
 		uint64										DownloadedChunkCount = 0;
-		double										LastProgressTimeStamp = 0.0;
+		uint64										LastProgressCycles = 0;
+		uint64										StartTimeCycles = FPlatformTime::Cycles64();
 		int32										Priority = 0; 
 		std::atomic_bool							bHttpRequestsIssued{false};
 	};
diff --git a/Engine/Source/Runtime/Experimental/IoStore/OnDemand/Private/Statistics.cpp b/Engine/Source/Runtime/Experimental/IoStore/OnDemand/Private/Statistics.cpp
--- a/Engine/Source/Runtime/Experimental/IoStore/OnDemand/Private/Statistics.cpp
+++ b/Engine/Source/Runtime/Experimental/IoStore/OnDemand/Private/Statistics.cpp
@@ -336,7 +336,7 @@ FCounterAtomicInt		GTotalInstallCount(TEXT("Iad/InstallCount"), TraceCounterDisp
 FCounterAtomicInt		GInflightInstallCount(TEXT("Iad/InflightInstallCount"), TraceCounterDisplayHint_None);
 FCounterAtomicInt		GTotalDownloadedChunkCount(TEXT("Iad/DownloadedChunks"), TraceCounterDisplayHint_None);
 FCounterAtomicInt		GTotalDownloadedChunkBytes(TEXT("Iad/DownloadedBytes"), TraceCounterDisplayHint_Memory);
-//FCounterAtomicFloat	GTotalInstallDurationSeconds(TEXT("Iad/TotalInstallDurationSeconds"), TraceCounterDisplayHint_None);
+FCounterAtomicInt		GTotalInstallDurationMs(TEXT("Iad/TotalInstallDurationMs"), TraceCounterDisplayHint_None);
 //FCounterAtomicFloat	GAvgInstallCacheHitRatio(TEXT("Iad/AvgInstallCacheHitRatio"), TraceCounterDisplayHint_None);
 
 ////////////////////////////////////////////////////////////////////////////////
@@ -409,6 +409,23 @@ FOnDemandIoBackendStats::FOnDemandIoBackendStats(FBackendStatus& InStatus)
 
 				UE_LOG(LogIas, Log, TEXT("HttpStats: DownloadedKiB=%d, Get=%d, Retry=%d, Cancel=%d, Error=%d, CurPending=%d, CurDurationMsAvg=%d, CurDurationMsMax=%d"),
 					(int32)HDownloadedKiB, HGetCount, HRetryCount, HCancelCount, HErrorCount, HPendingCount, HDurationMsAvg, HDurationMsMax);
+
+				{
+					const uint64 TotalInstallCount		= GTotalInstallCount.Get();
+					const uint64 TotalInstallDurationMs	= GTotalInstallDurationMs.Get();
+					const uint64 TotalDownloadedBytes	= GTotalDownloadedChunkBytes.Get();
+					uint64 AvgInstallDurationMs			= 0;
+					double ThroughputKbitPerSec			= 0.0;
+
+					if (TotalInstallCount > 0)
+					{
+						AvgInstallDurationMs	= TotalInstallDurationMs / TotalInstallCount;
+						ThroughputKbitPerSec	= (double(TotalDownloadedBytes * 8)  / (double(TotalInstallDurationMs) / 1000.0)) / 1000.0;
+					}
+
+					UE_LOG(LogIoStoreOnDemand, Log, TEXT("IadStats: InstallCount=%llu, Downloaded=%d KiB, Throughput=%d kbit/s, AvgInstallDuration=%d ms, AvgCacheHitRatio=%d %%"),
+						TotalInstallCount, int32(TotalDownloadedBytes / 1024), int32(ThroughputKbitPerSec), AvgInstallDurationMs, int32(GInstallCacheHitRatio.GetMean() * 100.0));
+				}
 				LastLogTime = Time;
 			}
 		}
@@ -962,14 +979,20 @@ void FOnDemandContentInstallerStats::OnRequestCompleted(
 		uint64 RequestedBytes,
 		uint64 DownloadedChunkCount,
 		uint64 DownloadedBytes,
-		double DurationSeconds) 
+		double CacheHitRatio,
+		uint64 DurationCycles) 
 {
+	if (RequestedChunkCount == 0)
+	{
+		return;
+	}
+
 	GInflightInstallCount.Add(-1);
 	GTotalInstallCount.Add(1);
 	GTotalDownloadedChunkCount.Add(DownloadedChunkCount);
 	GTotalDownloadedChunkBytes.Add(DownloadedBytes);
-	//GTotalInstallDurationSeconds.Add(DurationSeconds);
-	GInstallCacheHitRatio.Increment(double(DownloadedBytes) / double(RequestedBytes));
+	GTotalInstallDurationMs.Add(uint64(FPlatformTime::ToMilliseconds64(DurationCycles)));
+	GInstallCacheHitRatio.Increment(CacheHitRatio);
 	//GAvgInstallCacheHitRatio.Set(GInstallCacheHitRatio.GetMean());
 }
 
diff --git a/Engine/Source/Runtime/Experimental/IoStore/OnDemand/Private/Statistics.h b/Engine/Source/Runtime/Experimental/IoStore/OnDemand/Private/Statistics.h
--- a/Engine/Source/Runtime/Experimental/IoStore/OnDemand/Private/Statistics.h
+++ b/Engine/Source/Runtime/Experimental/IoStore/OnDemand/Private/Statistics.h
@@ -91,7 +91,8 @@ public:
 		uint64 RequestedBytes,
 		uint64 DownloadedChunkCount,
 		uint64 DownloadedBytes,
-		double DurationSeconds) IAS_STATISTICS_IMPL()
+		double CacheHitRatio,
+		uint64 DurationCycles) IAS_STATISTICS_IMPL()
 };
 
 #undef IAS_STATISTICS_IMPL
diff --git a/Engine/Source/Runtime/Experimental/IoStore/OnDemandCore/Internal/IO/IoStoreOnDemandInternals.h b/Engine/Source/Runtime/Experimental/IoStore/OnDemandCore/Internal/IO/IoStoreOnDemandInternals.h
--- a/Engine/Source/Runtime/Experimental/IoStore/OnDemandCore/Internal/IO/IoStoreOnDemandInternals.h
+++ b/Engine/Source/Runtime/Experimental/IoStore/OnDemandCore/Internal/IO/IoStoreOnDemandInternals.h
@@ -61,7 +61,6 @@ public:
 	FOnDemandInstallArgs		Args;
 	FOnDemandInstallCompleted	OnCompleted;
 	FOnDemandInstallProgressed	OnProgress;
-	double						StartTime = FPlatformTime::Seconds();
 	uint32						SeqNo = NextSeqNo++;
 	std::atomic<EStatus>		Status{EStatus::Pending};
 	std::atomic<EIoErrorCode>	ErrorCode{EIoErrorCode::Unknown};
