Rename FOnDemandThread to FOnDemandHttpThread to match the code file and to show the thread's purpose.

#rnx

- I always meant to name the class FOnDemandHttpThread but didn't get around to changing it before the original submit.

diff --git a/Engine/Source/Runtime/Experimental/IoStore/OnDemand/Private/OnDemandContentInstaller.cpp b/Engine/Source/Runtime/Experimental/IoStore/OnDemand/Private/OnDemandContentInstaller.cpp
--- a/Engine/Source/Runtime/Experimental/IoStore/OnDemand/Private/OnDemandContentInstaller.cpp
+++ b/Engine/Source/Runtime/Experimental/IoStore/OnDemand/Private/OnDemandContentInstaller.cpp
@@ -196,7 +196,7 @@ void ResolveChunksToInstall(
 } // namespace Private
 
 ////////////////////////////////////////////////////////////////////////////////
-FOnDemandContentInstaller::FOnDemandContentInstaller(FOnDemandIoStore& InIoStore, FOnDemandThread& InHttpClient)
+FOnDemandContentInstaller::FOnDemandContentInstaller(FOnDemandIoStore& InIoStore, FOnDemandHttpThread& InHttpClient)
 	: IoStore(InIoStore)
 	, HttpClient(InHttpClient)
 	, ProcessPipe(TEXT("IoStoreOnDemandProcessPipe"))
@@ -425,6 +425,7 @@ void FOnDemandContentInstaller::ProcessRequest(FInstallRequestState& RequestStat
 
 	for (FChunkHttpRequestHandle& HttpRequest : RequestState.HttpRequestHandles)
 	{
+		TRACE_CPUPROFILER_EVENT_SCOPE(FOnDemandContentInstaller::IssueRequest);
 		FSharedOnDemandContainer& Container		= RequestState.ResolvedChunks[HttpRequest.ContainerIndex].Container;
 		const FOnDemandChunkEntry& ChunkEntry	= Container->ChunkEntries[HttpRequest.EntryIndex];
 
@@ -434,6 +435,7 @@ void FOnDemandContentInstaller::ProcessRequest(FInstallRequestState& RequestStat
 			RequestState.Priority,
 			[this, &RequestState, &HttpRequest](uint32 HttpStatusCode, FStringView ErrorReason, FIoBuffer&& Chunk)
 			{
+				TRACE_CPUPROFILER_EVENT_SCOPE(FOnDemandContentInstaller::Callback);
 				HttpCompletionPipe.Launch(
 					TEXT("ProcessIoStoreOnDemandDownloadedChunk"),
 					[this, &RequestState, &HttpRequest, HttpStatusCode, ErrorReason = FString(ErrorReason), Chunk = MoveTemp(Chunk)]() mutable
@@ -459,6 +461,8 @@ void FOnDemandContentInstaller::ProcessDownloadedChunk(
 	FString&& ErrorReason,
 	FIoBuffer&& Chunk)
 {
+	TRACE_CPUPROFILER_EVENT_SCOPE(FOnDemandContentInstaller::ProcessDownloadedChunk);
+
 	FOnDemandInternalInstallRequest& Request	= *RequestState.Request;
 	HttpRequest.Handle							= nullptr;
 	FSharedOnDemandContainer& Container			= RequestState.ResolvedChunks[HttpRequest.ContainerIndex].Container;
diff --git a/Engine/Source/Runtime/Experimental/IoStore/OnDemand/Private/OnDemandContentInstaller.h b/Engine/Source/Runtime/Experimental/IoStore/OnDemand/Private/OnDemandContentInstaller.h
--- a/Engine/Source/Runtime/Experimental/IoStore/OnDemand/Private/OnDemandContentInstaller.h
+++ b/Engine/Source/Runtime/Experimental/IoStore/OnDemand/Private/OnDemandContentInstaller.h
@@ -18,8 +18,8 @@ class FIoBuffer;
 namespace UE::IoStore
 {
 
+class FOnDemandHttpThread;
 class FOnDemandIoStore;
-class FOnDemandThread;
 
 using FSharedOnDemandContainer = TSharedPtr<struct FOnDemandContainer, ESPMode::ThreadSafe>;
 
@@ -77,7 +77,7 @@ class FOnDemandContentInstaller
 	};
 
 public:
-									FOnDemandContentInstaller(FOnDemandIoStore& IoStore, FOnDemandThread& HttpClient);
+									FOnDemandContentInstaller(FOnDemandIoStore& IoStore, FOnDemandHttpThread& HttpClient);
 									~FOnDemandContentInstaller();
 
 	FSharedInternalInstallRequest	Enqueue(
@@ -102,7 +102,7 @@ private:
 	using FRequestStateAllocator	= TSingleThreadedSlabAllocator<FInstallRequestState>;
 
 	FOnDemandIoStore&			IoStore;
-	FOnDemandThread&			HttpClient;
+	FOnDemandHttpThread&		HttpClient;
 	UE::Tasks::FPipe			ProcessPipe;
 	UE::Tasks::FPipe			HttpCompletionPipe;
 
diff --git a/Engine/Source/Runtime/Experimental/IoStore/OnDemand/Private/OnDemandHttpThread.cpp b/Engine/Source/Runtime/Experimental/IoStore/OnDemand/Private/OnDemandHttpThread.cpp
--- a/Engine/Source/Runtime/Experimental/IoStore/OnDemand/Private/OnDemandHttpThread.cpp
+++ b/Engine/Source/Runtime/Experimental/IoStore/OnDemand/Private/OnDemandHttpThread.cpp
@@ -100,20 +100,20 @@ struct FHttpRequest
 	}
 
 	FHttpRequest* NextRequest = nullptr;
-	FOnDemandThread::FCompletionCallback CompletionCallback;
+	FOnDemandHttpThread::FCompletionCallback CompletionCallback;
 	FOnDemandChunkInfo ChunkInfo;
 	FIoOffsetAndLength ChunkRange;
 	bool bCancelled = false;
 	int32 Priority = 0;
 };
 
-FOnDemandThread* FOnDemandThread::SingletonInstance = nullptr;
+FOnDemandHttpThread* FOnDemandHttpThread::SingletonInstance = nullptr;
 
-FOnDemandThread::FOnDemandThread()
+FOnDemandHttpThread::FOnDemandHttpThread()
 {
 	LLM_SCOPE_BYTAG(Ias);
 
-	checkf(SingletonInstance == nullptr, TEXT("Attempting to create multiple instances of FOnDemandThread"));
+	checkf(SingletonInstance == nullptr, TEXT("Attempting to create multiple instances of FOnDemandHttpThread"));
 	SingletonInstance = this;
 
 	const int32 ThreadPriorityIndex = FMath::Clamp(GOnDemandBackendThreadPriorityIndex, 0, (int32)UE_ARRAY_COUNT(GREPLACEOnDemandBackendThreadPriorities) - 1);
@@ -125,7 +125,7 @@ FOnDemandThread::FOnDemandThread()
 	Thread.Reset(FRunnableThread::Create(this, TEXT("IoStoreOnDemand.Http"), StackSize, ThreadPriority));
 }
 
-FOnDemandThread::~FOnDemandThread()
+FOnDemandHttpThread::~FOnDemandHttpThread()
 {
 	checkf(SingletonInstance == this, TEXT("SingletonInstance was reassigned"));
 	SingletonInstance = nullptr;
@@ -136,20 +136,20 @@ FOnDemandThread::~FOnDemandThread()
 	DrainHttpRequests();
 }
 
-FOnDemandThread& FOnDemandThread::Get()
+FOnDemandHttpThread& FOnDemandHttpThread::Get()
 {
 	check(SingletonInstance != nullptr);
 	return *SingletonInstance;
 }
 
-FOnDemandThread::FRequestHandle FOnDemandThread::IssueRequest(const FOnDemandChunkInfo& ChunkInfo, const FIoOffsetAndLength& ReadRange, int32 Priority, FCompletionCallback&& CompletionCallback)
+FOnDemandHttpThread::FRequestHandle FOnDemandHttpThread::IssueRequest(const FOnDemandChunkInfo& ChunkInfo, const FIoOffsetAndLength& ReadRange, int32 Priority, FCompletionCallback&& CompletionCallback)
 {
 	return IssueRequest(FOnDemandChunkInfo(ChunkInfo), ReadRange, Priority, MoveTemp(CompletionCallback));
 }
 
-FOnDemandThread::FRequestHandle FOnDemandThread::IssueRequest(FOnDemandChunkInfo&& ChunkInfo, const FIoOffsetAndLength& ReadRange, int32 Priority, FCompletionCallback&& CompletionCallback)
+FOnDemandHttpThread::FRequestHandle FOnDemandHttpThread::IssueRequest(FOnDemandChunkInfo&& ChunkInfo, const FIoOffsetAndLength& ReadRange, int32 Priority, FCompletionCallback&& CompletionCallback)
 {
-	TRACE_CPUPROFILER_EVENT_SCOPE(FOnDemandThread::IssueRequest);
+	TRACE_CPUPROFILER_EVENT_SCOPE(FOnDemandHttpThread::IssueRequest);
 
 	FHttpRequest* Request = AllocateRequest(MoveTemp(ChunkInfo));
 
@@ -165,16 +165,16 @@ FOnDemandThread::FRequestHandle FOnDemandThread::IssueRequest(FOnDemandChunkInfo
 	return Request;
 }
 
-void FOnDemandThread::ReprioritizeRequest(FRequestHandle Request, int32 NewPriority)
+void FOnDemandHttpThread::ReprioritizeRequest(FRequestHandle Request, int32 NewPriority)
 {
-	TRACE_CPUPROFILER_EVENT_SCOPE(FOnDemandThread::ReprioritizeRequest);
+	TRACE_CPUPROFILER_EVENT_SCOPE(FOnDemandHttpThread::ReprioritizeRequest);
 	if (Request != 0)
 	{
 		HttpRequests.Reprioritize(static_cast<FHttpRequest*>(Request), NewPriority);
 	}
 }
 
-void FOnDemandThread::CancelRequest(FRequestHandle Request)
+void FOnDemandHttpThread::CancelRequest(FRequestHandle Request)
 {
 	if (Request != 0)
 	{
@@ -182,7 +182,7 @@ void FOnDemandThread::CancelRequest(FRequestHandle Request)
 	}
 }
 
-bool FOnDemandThread::Init()
+bool FOnDemandHttpThread::Init()
 {
 	LLM_SCOPE_BYTAG(Ias);
 	if (!TryCreateHttpClient())
@@ -193,7 +193,7 @@ bool FOnDemandThread::Init()
 	return true;
 }
 
-uint32 FOnDemandThread::Run()
+uint32 FOnDemandHttpThread::Run()
 {
 	LLM_SCOPE_BYTAG(Ias);
 
@@ -222,18 +222,18 @@ uint32 FOnDemandThread::Run()
 	return 0;
 }
 
-void FOnDemandThread::Stop()
+void FOnDemandHttpThread::Stop()
 {
 	bStopRequested = true;
 	TickThreadEvent->Trigger();
 }
 
-void FOnDemandThread::Exit()
+void FOnDemandHttpThread::Exit()
 {
 	HttpClient.Reset();
 }
 
-void FOnDemandThread::TickRequests()
+void FOnDemandHttpThread::TickRequests()
 {
 	// We can't limit cvars so add a hard coded cap to prevent concurrent requests from getting out of hand.
 	const int32 MaxConcurrentRequests = FMath::Min(GIasHttpConcurrentRequests, 32);
@@ -246,7 +246,7 @@ void FOnDemandThread::TickRequests()
 		while (NextHttpRequest)
 		{
 			{
-				TRACE_CPUPROFILER_EVENT_SCOPE(FOnDemandThread::IssueHttpGet);
+				TRACE_CPUPROFILER_EVENT_SCOPE(FOnDemandHttpThread::IssueHttpGet);
 				FHttpRequest* HttpRequest = NextHttpRequest;
 				NextHttpRequest = HttpRequest->NextRequest;
 				HttpRequest->NextRequest = nullptr;
@@ -283,7 +283,7 @@ void FOnDemandThread::TickRequests()
 						[this, HttpRequest, &NumConcurrentRequests]
 						(FMultiEndpointHttpClientResponse&& HttpResponse)
 						{
-							TRACE_CPUPROFILER_EVENT_SCOPE(FOnDemandThread::RequestCallback);
+							TRACE_CPUPROFILER_EVENT_SCOPE(FOnDemandHttpThread::RequestCallback);
 
 							NumConcurrentRequests--;
 
@@ -322,7 +322,7 @@ void FOnDemandThread::TickRequests()
 
 			if (NumConcurrentRequests >= MaxConcurrentRequests)
 			{
-				TRACE_CPUPROFILER_EVENT_SCOPE(FOnDemandThread::TickHttpSaturated);
+				TRACE_CPUPROFILER_EVENT_SCOPE(FOnDemandHttpThread::TickHttpSaturated);
 				while (NumConcurrentRequests >= MaxConcurrentRequests) //-V654
 				{
 					HttpClient->Tick(MAX_uint32, GIasHttpRateLimitKiBPerSecond);
@@ -337,7 +337,7 @@ void FOnDemandThread::TickRequests()
 
 		{
 			// Keep processing pending connections until all requests are completed or a new one is issued
-			TRACE_CPUPROFILER_EVENT_SCOPE(FOnDemandThread::TickHttp);
+			TRACE_CPUPROFILER_EVENT_SCOPE(FOnDemandHttpThread::TickHttp);
 			while (HttpClient->Tick(GIasHttpPollTimeoutMs, GIasHttpRateLimitKiBPerSecond))
 			{
 				if (!NextHttpRequest)
@@ -354,7 +354,7 @@ void FOnDemandThread::TickRequests()
 	}
 }
 
-void FOnDemandThread::DrainHttpRequests()
+void FOnDemandHttpThread::DrainHttpRequests()
 {
 	FHttpRequest* Iterator = HttpRequests.Dequeue();
 	while (Iterator != nullptr)
@@ -369,7 +369,7 @@ void FOnDemandThread::DrainHttpRequests()
 	}
 }
 
-bool FOnDemandThread::TryCreateHttpClient()
+bool FOnDemandHttpThread::TryCreateHttpClient()
 {
 	HttpClient = FMultiEndpointHttpClient::Create(FMultiEndpointHttpClientConfig
 		{
@@ -387,7 +387,7 @@ bool FOnDemandThread::TryCreateHttpClient()
 	return HttpClient.IsValid();
 }
 
-void FOnDemandThread::UpdateThreadPriorityIfNeeded()
+void FOnDemandHttpThread::UpdateThreadPriorityIfNeeded()
 {
 	// Read the thread priority from the cvar
 	int32 ThreadPriorityIndex = FMath::Clamp(GOnDemandBackendThreadPriorityIndex, 0, (int32)UE_ARRAY_COUNT(GREPLACEOnDemandBackendThreadPriorities) - 1);
@@ -401,13 +401,13 @@ void FOnDemandThread::UpdateThreadPriorityIfNeeded()
 	}
 }
 
-FHttpRequest* FOnDemandThread::AllocateRequest(FOnDemandChunkInfo&& ChunkInfo)
+FHttpRequest* FOnDemandHttpThread::AllocateRequest(FOnDemandChunkInfo&& ChunkInfo)
 {
 	UE::TUniqueLock _(AllocatorMutex);
 	return RequestAllocator.Construct(MoveTemp(ChunkInfo));
 }
 
-void FOnDemandThread::DestroyRequest(FHttpRequest* Request)
+void FOnDemandHttpThread::DestroyRequest(FHttpRequest* Request)
 {
 	UE::TUniqueLock _(AllocatorMutex);
 	RequestAllocator.Destroy(Request);
diff --git a/Engine/Source/Runtime/Experimental/IoStore/OnDemand/Private/OnDemandHttpThread.h b/Engine/Source/Runtime/Experimental/IoStore/OnDemand/Private/OnDemandHttpThread.h
--- a/Engine/Source/Runtime/Experimental/IoStore/OnDemand/Private/OnDemandHttpThread.h
+++ b/Engine/Source/Runtime/Experimental/IoStore/OnDemand/Private/OnDemandHttpThread.h
@@ -32,7 +32,7 @@ class FMultiEndpointHttpClient;
 struct FHttpRequest;
 struct FOnDemandChunkInfo;
 
-class FOnDemandThread final
+class FOnDemandHttpThread final
 	: public FRunnable
 {
 public:
@@ -40,10 +40,10 @@ public:
 	using FRequestHandle = void*;
 	DECLARE_MULTICAST_DELEGATE(FOnTickIdle);
 
-	FOnDemandThread();
-	virtual ~FOnDemandThread();
+	FOnDemandHttpThread();
+	virtual ~FOnDemandHttpThread();
 
-	static FOnDemandThread& Get();
+	static FOnDemandHttpThread& Get();
 	
 	/**
 	 * Issue a http request to read from a chunk.
@@ -99,7 +99,7 @@ private:
 	void DestroyRequest(FHttpRequest* Request);
 
 private:
-	static FOnDemandThread* SingletonInstance;
+	static FOnDemandHttpThread* SingletonInstance;
 	TUniquePtr<FRunnableThread> Thread;
 
 	FEventRef TickThreadEvent;
@@ -127,7 +127,7 @@ namespace UE::IoStore
 
 // Dummy class for FOnDemandIoStore::OnDemandHttpThread
 // Remove when we delete UE_ENABLE_REQUEST_THREAD
-class FOnDemandThread final
+class FOnDemandHttpThread final
 {
 };
 
diff --git a/Engine/Source/Runtime/Experimental/IoStore/OnDemand/Private/OnDemandIoDispatcherBackend.cpp b/Engine/Source/Runtime/Experimental/IoStore/OnDemand/Private/OnDemandIoDispatcherBackend.cpp
--- a/Engine/Source/Runtime/Experimental/IoStore/OnDemand/Private/OnDemandIoDispatcherBackend.cpp
+++ b/Engine/Source/Runtime/Experimental/IoStore/OnDemand/Private/OnDemandIoDispatcherBackend.cpp
@@ -484,7 +484,7 @@ struct FChunkRequest
 	FIoRequestImpl* RequestHead;
 	FIoRequestImpl* RequestTail;
 #if UE_ENABLE_IAS_REQUEST_THREAD
-	FOnDemandThread::FRequestHandle HttpRequest = 0;
+	FOnDemandHttpThread::FRequestHandle HttpRequest = 0;
 #endif //UE_ENABLE_IAS_REQUEST_THREAD
 	FIASHostGroup HostGroup; // Still used in a few places even when UE_ENABLE_IAS_REQUEST_THREAD is enabled
 	FIoBuffer Chunk;
@@ -682,7 +682,7 @@ class FOnDemandIoBackend final
 
 					ChunkRequest.bCancelled = true;
 #if UE_ENABLE_IAS_REQUEST_THREAD
-					FOnDemandThread::Get().CancelRequest(ChunkRequest.HttpRequest);
+					FOnDemandHttpThread::Get().CancelRequest(ChunkRequest.HttpRequest);
 #endif // UE_ENABLE_IAS_REQUEST_THREAD
 
 					if (TheCache != nullptr)
@@ -918,7 +918,7 @@ void FOnDemandIoBackend::Shutdown()
 #if UE_ENABLE_IAS_REQUEST_THREAD
 	if (OnThreadTickIdleHandle.IsValid())
 	{
-		FOnDemandThread::Get().OnTickIdle().Remove(OnThreadTickIdleHandle);
+		FOnDemandHttpThread::Get().OnTickIdle().Remove(OnThreadTickIdleHandle);
 		OnThreadTickIdleHandle.Reset();
 	}
 #endif // UE_ENABLE_IAS_REQUEST_THREAD
@@ -987,14 +987,14 @@ void FOnDemandIoBackend::ConditionallyStartBackendThread()
 			}
 		});
 
-	OnThreadTickIdleHandle = FOnDemandThread::Get().OnTickIdle().AddRaw(this, &FOnDemandIoBackend::OnThreadTickIdle);
+	OnThreadTickIdleHandle = FOnDemandHttpThread::Get().OnTickIdle().AddRaw(this, &FOnDemandIoBackend::OnThreadTickIdle);
 #endif // UE_ENABLE_IAS_REQUEST_THREAD
 }
 
 void FOnDemandIoBackend::IssueHttpRequest(FChunkRequest* ChunkRequest)
 {
 #if UE_ENABLE_IAS_REQUEST_THREAD
-	ChunkRequest->HttpRequest = FOnDemandThread::Get().IssueRequest(ChunkRequest->Params.ChunkInfo, ChunkRequest->Params.ChunkRange, ChunkRequest->Priority,
+	ChunkRequest->HttpRequest = FOnDemandHttpThread::Get().IssueRequest(ChunkRequest->Params.ChunkInfo, ChunkRequest->Params.ChunkRange, ChunkRequest->Priority,
 		[this, ChunkRequest](uint32 StatusCode, FStringView ErrorReason, FIoBuffer&& Data)
 		{
 			ChunkRequest->Chunk = MoveTemp(Data);
@@ -1201,7 +1201,7 @@ bool FOnDemandIoBackend::Resolve(FIoRequestImpl* Request)
 		if (bUpdatePriority)
 		{
 #if UE_ENABLE_IAS_REQUEST_THREAD
-			FOnDemandThread::Get().ReprioritizeRequest(ChunkRequest->HttpRequest, ChunkRequest->Priority);
+			FOnDemandHttpThread::Get().ReprioritizeRequest(ChunkRequest->HttpRequest, ChunkRequest->Priority);
 #else
 			TRACE_CPUPROFILER_EVENT_SCOPE(IasBackend::UpdatePriorityForIoRequest);
 			HttpRequests.Reprioritize(ChunkRequest, ChunkRequest->Priority);
@@ -1290,7 +1290,7 @@ void FOnDemandIoBackend::UpdatePriorityForIoRequest(FIoRequestImpl* Request)
 	if (FChunkRequest* ChunkRequest = ChunkRequests.TryUpdatePriority(Request))
 	{
 #if UE_ENABLE_IAS_REQUEST_THREAD
-		FOnDemandThread::Get().ReprioritizeRequest(ChunkRequest->HttpRequest, ChunkRequest->Priority);
+		FOnDemandHttpThread::Get().ReprioritizeRequest(ChunkRequest->HttpRequest, ChunkRequest->Priority);
 #else
 		HttpRequests.Reprioritize(ChunkRequest, ChunkRequest->Priority);
 #endif // UE_ENABLE_IAS_REQUEST_THREAD
diff --git a/Engine/Source/Runtime/Experimental/IoStore/OnDemand/Private/OnDemandIoStore.cpp b/Engine/Source/Runtime/Experimental/IoStore/OnDemand/Private/OnDemandIoStore.cpp
--- a/Engine/Source/Runtime/Experimental/IoStore/OnDemand/Private/OnDemandIoStore.cpp
+++ b/Engine/Source/Runtime/Experimental/IoStore/OnDemand/Private/OnDemandIoStore.cpp
@@ -394,9 +394,9 @@ FString FOnDemandContainer::UniqueName() const
 FOnDemandIoStore::FOnDemandIoStore()
 {
 #if UE_ENABLE_REQUEST_THREAD
-	OnDemandHttpThread = MakeUnique<FOnDemandThread>();
+	OnDemandHttpThread = MakeUnique<FOnDemandHttpThread>();
 #if UE_ENABLE_IAD_REQUEST_THREAD
-	Installer = MakeUnique<FOnDemandContentInstaller>(*this, FOnDemandThread::Get());
+	Installer = MakeUnique<FOnDemandContentInstaller>(*this, FOnDemandHttpThread::Get());
 #endif
 #endif // UE_ENABLE_REQUEST_THREAD
 
diff --git a/Engine/Source/Runtime/Experimental/IoStore/OnDemand/Private/OnDemandIoStore.h b/Engine/Source/Runtime/Experimental/IoStore/OnDemand/Private/OnDemandIoStore.h
--- a/Engine/Source/Runtime/Experimental/IoStore/OnDemand/Private/OnDemandIoStore.h
+++ b/Engine/Source/Runtime/Experimental/IoStore/OnDemand/Private/OnDemandIoStore.h
@@ -27,7 +27,7 @@ namespace UE::IoStore
 {
 
 class FOnDemandContentInstaller;
-class FOnDemandThread;
+class FOnDemandHttpThread;
 struct FOnDemandEndpointConfig;
 enum class ELogOnDemandCacheUsage : uint8;
 
@@ -362,7 +362,7 @@ private:
 											TArray<FSharedOnDemandContainer>& Out);
 
 	TUniquePtr<FOnDemandContentInstaller>	Installer;
-	TUniquePtr<FOnDemandThread>				OnDemandHttpThread;
+	TUniquePtr<FOnDemandHttpThread>				OnDemandHttpThread;
 
 	FSharedInstallCache						InstallCache;
 	FSharedPackageStoreBackend				PackageStoreBackend;
