Clean up some older IAS statistics values that are not needed anymore.

#rnx

diff --git a/Engine/Source/Runtime/Experimental/IoStore/OnDemand/Private/Statistics.cpp b/Engine/Source/Runtime/Experimental/IoStore/OnDemand/Private/Statistics.cpp
--- a/Engine/Source/Runtime/Experimental/IoStore/OnDemand/Private/Statistics.cpp
+++ b/Engine/Source/Runtime/Experimental/IoStore/OnDemand/Private/Statistics.cpp
@@ -294,7 +294,6 @@ int64					GHttpHistoryBytes[GHttpHistoryCount] = {};
 int64					GHttpHistoryTotalDuration = 0;
 int64					GHttpHistoryTotalBytes = 0;
 int64 					GHttpHistoryIndex = 0;
-FIncrementalVariance	GHttpAvgDuration; // Duration of the http requests, in milliseconds
 static uint32			GHttpCdnCacheHit = 0;
 static uint32			GHttpCdnCacheMiss = 0;
 static uint32			GHttpCdnCacheUnknown = 0;
@@ -567,16 +566,18 @@ void FOnDemandIoBackendStats::ReportEndPointAnalytics(TArray<FAnalyticsEventAttr
 	if (GIasReportHttpAnalyticsEnabled)
 	{
 		const int64 ByteCount = GDeltaTracking.Get(TEXTVIEW("IasHttpDownloadedBytes"), GHttpDownloadedBytes.Get());
+		const int64 GetCount = GDeltaTracking.Get(TEXTVIEW("IasHttpGetCount"), GHttpGetCount.Get());
+
 		const double DataRateBPS = GHttpDurationMsSum > 0 ? static_cast<double>(ByteCount) / (static_cast<double>(GHttpDurationMsSum) / 1000.0) : 0.0;
+		const double DurationMean = GetCount ? static_cast<double>(GHttpDurationMsSum) / static_cast<double>(GetCount) : 0.0;
 
 		AppendAnalyticsEventAttributeArray(OutAnalyticsArray
 			,TRACK_DELTA("IasHttpErrorCount", GHttpErrorCount.Get())
 			,TRACK_DELTA("IasHttpDecodeErrors", GHttpDecodeErrorCount.Get())
 			,TRACK_DELTA("IasHttpRetryCount", GHttpRetryCount.Get())
-			,TRACK_DELTA("IasHttpGetCount", GHttpGetCount.Get())
-			//,TRACK_DELTA("IasHttpDownloadedBytes", GHttpDownloadedBytes.Get()) // Disabled for now as we need ByteCount for DataRateBPS above.
-			, TEXT("IasHttpDownloadedBytes"), ByteCount
-			,TEXT("IasHttpDurationMean"), GHttpAvgDuration.GetMean()
+			,TEXT("IasHttpGetCount"), GetCount
+			,TEXT("IasHttpDownloadedBytes"), ByteCount
+			,TEXT("IasHttpDurationMean"), DurationMean
 			,TEXT("IasHttpDurationSum"), GHttpDurationMsSum
 			,TEXT("IasHttpDataRateMean"), DataRateBPS
 		);
@@ -607,7 +608,6 @@ void FOnDemandIoBackendStats::ReportEndPointAnalytics(TArray<FAnalyticsEventAttr
 		);
 
 		// These values we can just reset as they are only being used with analytics
-		GHttpAvgDuration.Reset();
 		GHttpDurationMsSum = 0;
 	}
 
@@ -930,8 +930,6 @@ void FOnDemandIoBackendStats::OnHttpGet(uint64 SizeBytes, uint64 DurationMs)
 	}
 
 	GHttpHistoryIndex = (GHttpHistoryIndex + 1) % GHttpHistoryCount;
-
-	GHttpAvgDuration.Increment(static_cast<double>(DurationMs));
 }
 
 void FOnDemandIoBackendStats::OnHttpCancel()
