Iris
* Immediately prune garbage collected pointers from within PostGarbageCollect delegate instead of waiting for the next TickFlush to occur
* Tentative fix for a crash where we were executing RPC's right after CollectingGarbage and calling RPCs on destroyed objects

#rnx

diff --git a/Engine/Source/Runtime/Experimental/Iris/Core/Private/Iris/ReplicationSystem/ReplicationSystem.cpp b/Engine/Source/Runtime/Experimental/Iris/Core/Private/Iris/ReplicationSystem/ReplicationSystem.cpp
--- a/Engine/Source/Runtime/Experimental/Iris/Core/Private/Iris/ReplicationSystem/ReplicationSystem.cpp
+++ b/Engine/Source/Runtime/Experimental/Iris/Core/Private/Iris/ReplicationSystem/ReplicationSystem.cpp
@@ -787,7 +787,6 @@ UReplicationSystem::UReplicationSystem()
 , PIEInstanceID(INDEX_NONE)
 , bIsServer(0)
 , bAllowObjectReplication(0)
-, bDoCollectGarbage(0)
 {
 }
 
@@ -863,7 +862,7 @@ void UReplicationSystem::PreSendUpdate(const FSendUpdateParams& Params)
 		InternalSys.GetTickReplicationStats() = {};
 	
 		// Force a integrity check of all replicated instances
-		if (bDoCollectGarbage || ReplicationSystemCVars::bForcePruneBeforeUpdate)
+		if (ReplicationSystemCVars::bForcePruneBeforeUpdate)
 		{
 			CollectGarbage();
 		}
@@ -1042,7 +1041,7 @@ void UReplicationSystem::PostSendUpdate()
 
 void UReplicationSystem::PostGarbageCollection()
 {
-	bDoCollectGarbage = 1U;
+	CollectGarbage();
 }
 
 void UReplicationSystem::CollectGarbage()
@@ -1052,8 +1051,6 @@ void UReplicationSystem::CollectGarbage()
 	// Prune stale object instances before descriptors and protocols are pruned
 	Impl->ReplicationSystemInternal.GetReplicationBridge()->CallPruneStaleObjects();
 	Impl->ReplicationSystemInternal.GetReplicationStateDescriptorRegistry().PruneStaleDescriptors();
-
-	bDoCollectGarbage = 0;
 }
 
 void UReplicationSystem::ResetGameWorldState()
diff --git a/Engine/Source/Runtime/Experimental/Iris/Core/Public/Iris/ReplicationSystem/ReplicationSystem.h b/Engine/Source/Runtime/Experimental/Iris/Core/Public/Iris/ReplicationSystem/ReplicationSystem.h
--- a/Engine/Source/Runtime/Experimental/Iris/Core/Public/Iris/ReplicationSystem/ReplicationSystem.h
+++ b/Engine/Source/Runtime/Experimental/Iris/Core/Public/Iris/ReplicationSystem/ReplicationSystem.h
@@ -765,7 +765,6 @@ private:
 	int32 PIEInstanceID;
 	uint32 bIsServer : 1;
 	uint32 bAllowObjectReplication : 1;
-	uint32 bDoCollectGarbage : 1;
 };
 
 namespace UE::Net
