ChaosFlesh: -fixed NaN bug in deformable solver
-refined Hessian conditioning for matrix inversion

diff --git a/Engine/Source/Runtime/Experimental/Chaos/Public/Chaos/Deformable/GaussSeidelMainConstraint.h b/Engine/Source/Runtime/Experimental/Chaos/Public/Chaos/Deformable/GaussSeidelMainConstraint.h
--- a/Engine/Source/Runtime/Experimental/Chaos/Public/Chaos/Deformable/GaussSeidelMainConstraint.h
+++ b/Engine/Source/Runtime/Experimental/Chaos/Public/Chaos/Deformable/GaussSeidelMainConstraint.h
@@ -339,23 +339,50 @@ namespace Chaos::Softs
 				AddPerNodeHessian[i](p, Dt, ParticleHessian);
 			}
 
-			const T DtSquaredInverse =  1 / (Dt * Dt); // scale the hessian so that the determinant (~dt^6) does not fall below epsilon.
-			ParticleHessian *= DtSquaredInverse;
+			T HessianScale = (T)1.;
 			T HessianDet = ParticleHessian.Determinant();
-			if (HessianDet > TMathUtilConstants<T>::Epsilon)
+			auto IsIllConditioned = [](T HessianDet)
 			{
-				Chaos::PMatrix<T, 3, 3> HessianInv = ParticleHessian.SymmetricCofactorMatrix();
-				HessianInv *= T(1) / HessianDet;
-				Chaos::TVector<T, 3> Dx = HessianInv.GetTransposed() * (-ParticleResidual) * DtSquaredInverse; // add back DtSquared
-
-				if (Dx.Size() < MaxDxSize)
+				return FMath::Abs(HessianDet) < TMathUtilConstants<T>::Epsilon ||
+					FMath::Abs(HessianDet) > TMathUtilConstants<T>::MaxReal;
+			};
+			if (IsIllConditioned(HessianDet))
+			{
+				// scale the hessian so that the determinant (~HessianScale^3) falls into the normal range.
+				HessianScale = (T)0.;
+				for (int32 RowIdx = 0; RowIdx < 3; ++RowIdx)
 				{
-					Particles.P(p) += Dx;
+					for (int32 ColIdx = 0; ColIdx < 3; ++ColIdx)
+					{
+						HessianScale = FMath::Max(FMath::Abs(ParticleHessian.GetAt(RowIdx, ColIdx)), HessianScale);
+					}
 				}
-				else
+				if (IsIllConditioned(HessianScale))
 				{
 					return false;
 				}
+				else
+				{
+					ParticleHessian *= 1 / HessianScale;
+					HessianDet = ParticleHessian.Determinant();
+					if (IsIllConditioned(HessianDet))
+					{
+						return false;
+					}
+				}
+			}
+
+			Chaos::PMatrix<T, 3, 3> HessianInv = ParticleHessian.SymmetricCofactorMatrix();
+			HessianInv *= T(1) / HessianDet;
+			Chaos::TVector<T, 3> Dx = HessianInv.GetTransposed() * (-ParticleResidual) / HessianScale; // add back HessianScale
+
+			if (Dx.Size() < MaxDxSize)
+			{
+				Particles.P(p) += Dx;
+			}
+			else
+			{
+				return false;
 			}
 			return true;
 		}
diff --git a/Engine/Source/Runtime/Experimental/ChaosCore/Public/Chaos/ImplicitQRSVD.h b/Engine/Source/Runtime/Experimental/ChaosCore/Public/Chaos/ImplicitQRSVD.h
--- a/Engine/Source/Runtime/Experimental/ChaosCore/Public/Chaos/ImplicitQRSVD.h
+++ b/Engine/Source/Runtime/Experimental/ChaosCore/Public/Chaos/ImplicitQRSVD.h
@@ -4,11 +4,9 @@
 // HEADER_UNIT_SKIP - Bad include VectorTypes (living in GeometryCore)
 
 #include "Chaos/Matrix.h"
-#include "VectorTypes.h"
 #include "Chaos/Vector.h"
 #include "GenericPlatform/GenericPlatformMath.h"
 #include "Math/NumericLimits.h"
-#include "MathUtil.h"
 
 
 namespace Chaos {
@@ -367,7 +365,7 @@ namespace Chaos {
 	*/
 	template < class T>
 	inline void SingularValueDecomposition(const PMatrix<T, 2, 2>& A, GivensRotation<T>& U, const TVector<T, 2>& Sigma, GivensRotation<T>& V,
-		const T tol = 64 * TMathUtilConstants<T>::Epsilon) {
+		const T tol = std::numeric_limits<T>::epsilon()) {
 
 		TVector<T, 2>& sigma = const_cast<TVector<T, 2>&>(Sigma);
 
@@ -377,7 +375,7 @@ namespace Chaos {
 		T x = S_Sym.M[0*2+0];
 		T y = S_Sym.M[1*2+0];
 		T z = S_Sym.M[3];
-		if (y == 0) {
+		if (y < tol) {
 			// S is already diagonal
 			cosine = (T)1.;
 			sine = (T)0.;
@@ -434,7 +432,7 @@ namespace Chaos {
 	template <class T>
 	inline void SingularValueDecomposition(
 		const PMatrix<T, 2, 2>& A, const PMatrix<T, 2, 2>& U, const TVector<T, 2>& Sigma, const PMatrix<T, 2, 2>& V,
-		const T tol = 64 * TMathUtilConstants<T>::Epsilon) {
+		const T tol = std::numeric_limits<T>::epsilon()) {
 		//using T = ScalarType<TA>;
 		GivensRotation<T> gv(0, 1);
 		GivensRotation<T> gu(0, 1);
@@ -597,7 +595,7 @@ namespace Chaos {
 	  */
 	template <class T>
 	inline int SingularValueDecomposition(const PMatrix<T, 3, 3>& A, PMatrix<T, 3, 3>& U, TVector<T, 3>& sigma, PMatrix<T, 3, 3>& V,
-		T tol = 1024 * TMathUtilConstants<T>::Epsilon) 
+		T tol = std::numeric_limits<T>::epsilon()) 
 	{
 		sigma[0] = T(0.);
 		sigma[1] = T(0.);
