Added simulation of low data rates to IAD.

diff --git a/Engine/Source/Runtime/Experimental/IoStore/OnDemand/Private/OnDemandIoStore.cpp b/Engine/Source/Runtime/Experimental/IoStore/OnDemand/Private/OnDemandIoStore.cpp
--- a/Engine/Source/Runtime/Experimental/IoStore/OnDemand/Private/OnDemandIoStore.cpp
+++ b/Engine/Source/Runtime/Experimental/IoStore/OnDemand/Private/OnDemandIoStore.cpp
@@ -61,6 +61,20 @@ static FAutoConsoleVariableRef CVar_IoStoreOnDemandMinInstallDurationSeconds(
 	GIoStoreOnDemandMinInstallDurationSeconds,
 	TEXT("Sets the minimum install duration in seconds.")
 );
+
+float GIoStoreOnDemandRestrictKibsPerSecond = 0.0f;
+static FAutoConsoleVariableRef CVar_IoStoreOnDemandRestrictKibsPerSecond(
+	TEXT("iostore.InstallRestrictedToKibs"),
+	GIoStoreOnDemandRestrictKibsPerSecond,
+	TEXT("After install delay to simulate a restricted throughput rate.")
+);
+
+
+FString GIoStoreBreakOnRequest = "";
+static FAutoConsoleVariableRef CVar_IoStoreDebugBreakOnRequest(
+		TEXT("iostore.DebugBreakOnRequest"),
+		GIoStoreBreakOnRequest,
+		TEXT("Debug break when we detect this package name."));
 #endif
 
 ///////////////////////////////////////////////////////////////////////////////
@@ -1531,6 +1545,16 @@ FOnDemandInstallResult FOnDemandIoStore::TickInstallRequest(FOnDemandInternalIns
 	TSet<FSharedOnDemandContainer> ContainersForInstallation;
 	TSet<FPackageId> PackageIdsToInstall;
 
+#if !UE_BUILD_SHIPPING
+	if (!GIoStoreBreakOnRequest.IsEmpty())
+	{
+		if (FCString::Strstr(*InstallRequest.Args.DebugName, *GIoStoreBreakOnRequest))
+		{
+			UE_DEBUG_BREAK();
+		}
+	}
+#endif
+
 	if (Status = GetContainersAndPackagesForInstall(
 		InstallRequest.Args.MountId,
 		InstallRequest.Args.TagSets,
@@ -1803,6 +1827,17 @@ FOnDemandInstallResult FOnDemandIoStore::TickInstallRequest(FOnDemandInternalIns
 
 	double DurationInSeconds = FPlatformTime::Seconds() - StartTime;
 	#if !UE_BUILD_SHIPPING
+	if (GIoStoreOnDemandRestrictKibsPerSecond > 0.0f)
+	{
+		
+		float MinSleepTime = ((float)(TotalBytes)) / (GIoStoreOnDemandRestrictKibsPerSecond * 1024.0f);
+		float SleepTime = MinSleepTime - (float)(DurationInSeconds);
+		if (SleepTime > 0.0f)
+		{
+			FPlatformProcess::SleepNoStats(SleepTime);
+			DurationInSeconds = FPlatformTime::Seconds() - StartTime;
+		}
+	}
 	if (GIoStoreOnDemandMinInstallDurationSeconds > 0.0f)
 	{
 		if (float SleepTime = GIoStoreOnDemandMinInstallDurationSeconds - float(DurationInSeconds); SleepTime > 0.0f)
