Dataflow : fix issue with duplicating assets that are using a dataflow for generation

diff --git a/Engine/Source/Runtime/Experimental/Dataflow/Engine/Private/Dataflow/DataflowInstance.cpp b/Engine/Source/Runtime/Experimental/Dataflow/Engine/Private/Dataflow/DataflowInstance.cpp
--- a/Engine/Source/Runtime/Experimental/Dataflow/Engine/Private/Dataflow/DataflowInstance.cpp
+++ b/Engine/Source/Runtime/Experimental/Dataflow/Engine/Private/Dataflow/DataflowInstance.cpp
@@ -118,6 +118,28 @@ FDataflowVariableOverrides::FDataflowVariableOverrides(FDataflowInstance* InOwne
 {
 }
 
+FDataflowVariableOverrides& FDataflowVariableOverrides::operator=(const FDataflowVariableOverrides& Other)
+{
+	if (this != &Other)
+	{
+		Variables = Other.Variables;
+
+		// transfer the overriden GUIDs that may have changed with the copy 
+		OverriddenVariableGuids.Empty();
+		for (const FGuid OtherOverrideGuid : Other.OverriddenVariableGuids)
+		{
+			if (const FPropertyBagPropertyDesc* Desc = Other.Variables.FindPropertyDescByID(OtherOverrideGuid))
+			{
+				OverriddenVariableGuids.Add(Desc->ID);
+			}
+		}
+
+		// IMPORTANT : do not copy the owner 
+	}
+
+	return *this;
+}
+
 const FInstancedPropertyBag* FDataflowVariableOverrides::GetDefaultVariablesFromAsset() const 
 {
 	if (Owner)
@@ -404,4 +426,4 @@ bool FDataflowInstance::UpdateOwnerAsset(bool bUpdateDependentAssets) const
 		}
 	}
 	return bSuccess;
-}
\ No newline at end of file
+}
diff --git a/Engine/Source/Runtime/Experimental/Dataflow/Engine/Public/Dataflow/DataflowInstance.h b/Engine/Source/Runtime/Experimental/Dataflow/Engine/Public/Dataflow/DataflowInstance.h
--- a/Engine/Source/Runtime/Experimental/Dataflow/Engine/Public/Dataflow/DataflowInstance.h
+++ b/Engine/Source/Runtime/Experimental/Dataflow/Engine/Public/Dataflow/DataflowInstance.h
@@ -28,6 +28,8 @@ struct FDataflowVariableOverrides
 
 	DATAFLOWENGINE_API FDataflowVariableOverrides(FDataflowInstance* InOwner = nullptr);
 
+	DATAFLOWENGINE_API FDataflowVariableOverrides& operator=(const FDataflowVariableOverrides& Other);
+
 	/** Remove all variables */
 	DATAFLOWENGINE_API void RemoveAllVariables();
 
@@ -75,7 +77,7 @@ private:
 	UPROPERTY()
 	TArray<FGuid> OverriddenVariableGuids;
 
-	FDataflowInstance* Owner = nullptr;;
+	FDataflowInstance* Owner = nullptr;
 };
 
 /**
