Changed I/O store on-demand HTTP thread to use forkable thread to support server builds

- the HTTP thread now implements FSingleThreadRunnable to support fake threads
- added command line option -Ias.EnableHttpThreadPreFok to support real threads pre fork on server builds
- added define to completely disable threading and process HTTP requests synchronously

diff --git a/Engine/Source/Runtime/Experimental/IoStore/OnDemand/Private/OnDemandHttpClient.cpp b/Engine/Source/Runtime/Experimental/IoStore/OnDemand/Private/OnDemandHttpClient.cpp
--- a/Engine/Source/Runtime/Experimental/IoStore/OnDemand/Private/OnDemandHttpClient.cpp
+++ b/Engine/Source/Runtime/Experimental/IoStore/OnDemand/Private/OnDemandHttpClient.cpp
@@ -3,6 +3,7 @@
 #include "OnDemandHttpClient.h"
 
 #include "Containers/StringConv.h"
+#include "GenericPlatform/GenericPlatformProcess.h"
 #include "HAL/FileManager.h"
 #include "HAL/IConsoleManager.h"
 #include "HAL/PlatformTime.h"
@@ -11,7 +12,7 @@
 #include "Misc/Paths.h"
 
 #if UE_HTTPCLIENT_THREADSAFE_CHECKS
-	#define UE_CHECK_THREADSAFETY(MethodName) checkf(OwningThread == 0 || OwningThread == FPlatformTLS::GetCurrentThreadId(), #MethodName TEXT("was called on the on the wrong thread!"))
+	#define UE_CHECK_THREADSAFETY(MethodName) checkf(OwningThread == 0 || OwningThread == FPlatformTLS::GetCurrentThreadId() || !FGenericPlatformProcess::SupportsMultithreading(), #MethodName TEXT("was called on the on the wrong thread!"))
 #else
 	#define UE_CHECK_THREADSAFETY(MethodName)
 #endif //UE_HTTPCLIENT_THREADSAFE_CHECKS
diff --git a/Engine/Source/Runtime/Experimental/IoStore/OnDemand/Private/OnDemandHttpThread.cpp b/Engine/Source/Runtime/Experimental/IoStore/OnDemand/Private/OnDemandHttpThread.cpp
--- a/Engine/Source/Runtime/Experimental/IoStore/OnDemand/Private/OnDemandHttpThread.cpp
+++ b/Engine/Source/Runtime/Experimental/IoStore/OnDemand/Private/OnDemandHttpThread.cpp
@@ -11,10 +11,16 @@
 #include "IO/IoStoreOnDemand.h"
 #include "IasHostGroup.h"
 #include "Logging/StructuredLog.h"
+#include "Misc/CommandLine.h"
+#include "Misc/Fork.h"
 #include "OnDemandHttpClient.h"
 #include "OnDemandIoStore.h"
 #include "Statistics.h"
 
+#ifndef UE_IOSTORE_ONDEMAND_NO_HTTP_THREAD
+#define UE_IOSTORE_ONDEMAND_NO_HTTP_THREAD 0
+#endif
+
 // TODO - Maybe do LexToString/LexFromString (dupe code in UnrealEngine.cpp ThreadPriorityToString)
 static const TCHAR* LexToString(EThreadPriority Priority)
 {
@@ -116,13 +122,19 @@ FOnDemandHttpThread::FOnDemandHttpThread()
 	checkf(SingletonInstance == nullptr, TEXT("Attempting to create multiple instances of FOnDemandHttpThread"));
 	SingletonInstance = this;
 
-	const int32 ThreadPriorityIndex = FMath::Clamp(GOnDemandBackendThreadPriorityIndex, 0, (int32)UE_ARRAY_COUNT(GREPLACEOnDemandBackendThreadPriorities) - 1);
-	EThreadPriority DesiredThreadPriority = (EThreadPriority)GREPLACEOnDemandBackendThreadPriorities[ThreadPriorityIndex];
-	ThreadPriority = DesiredThreadPriority;
-
-	const uint32 StackSize = 0; // Use default stack size
-
-	Thread.Reset(FRunnableThread::Create(this, TEXT("IoStoreOnDemand.Http"), StackSize, ThreadPriority));
+#if UE_IOSTORE_ONDEMAND_NO_HTTP_THREAD
+	ensure(TryCreateHttpClient());
+#else
+	const int32 ThreadPriorityIndex			= FMath::Clamp(GOnDemandBackendThreadPriorityIndex, 0, (int32)UE_ARRAY_COUNT(GREPLACEOnDemandBackendThreadPriorities) - 1);
+	EThreadPriority DesiredThreadPriority	= (EThreadPriority)GREPLACEOnDemandBackendThreadPriorities[ThreadPriorityIndex];
+	ThreadPriority							= DesiredThreadPriority;
+	const uint32 StackSize					= 0; // Use default stack size
+	const uint64 ThreadAffinityMask			= FGenericPlatformAffinity::GetNoAffinityMask();
+	const EThreadCreateFlags CreateFlags	= EThreadCreateFlags::None;
+	const bool bAllowPreFork				= FParse::Param(FCommandLine::IsInitialized() ? FCommandLine::Get() : TEXT(""), TEXT("-Ias.EnableHttpThreadPreFork"));
+
+	Thread.Reset(FForkProcessHelper::CreateForkableThread(this, TEXT("IoStoreOnDemand.Http"), StackSize, ThreadPriority, ThreadAffinityMask, CreateFlags, bAllowPreFork));
+#endif
 }
 
 FOnDemandHttpThread::~FOnDemandHttpThread()
@@ -160,8 +172,12 @@ FOnDemandHttpThread::FRequestHandle FOnDemandHttpThread::IssueRequest(FOnDemandC
 	FOnDemandIoBackendStats::Get()->OnHttpEnqueue();
 
 	HttpRequests.EnqueueByPriority(Request);
+#if UE_IOSTORE_ONDEMAND_NO_HTTP_THREAD 
+	// When threading is disabled the completion callback has already been triggered and the handle is destroyed before returned to the caller.
+	Tick();
+#else
 	TickThreadEvent->Trigger();
-
+#endif
 	return Request;
 }
 
@@ -205,16 +221,7 @@ uint32 FOnDemandHttpThread::Run()
 
 		if(!bStopRequested)
 		{
-			FHostGroupManager::Get().Tick(GIasHttpTimeOutMs, bStopRequested);
-			
-			// TODO: It would be better to only update connections as they need it, consider doing this on
-			// hostgroup connect/disconnect events.
-			HttpClient->UpdateConnections();
-
-			TickRequests();
-
-			OnTickIdleDelegate.Broadcast();
-
+			Tick();
 			TickThreadEvent->Wait(FHostGroupManager::Get().GetNumDisconnctedHosts() > 0 ? GIasHttpHealthCheckWaitTime : MAX_uint32);
 		}
 	}
@@ -233,6 +240,19 @@ void FOnDemandHttpThread::Exit()
 	HttpClient.Reset();
 }
 
+void FOnDemandHttpThread::Tick()
+{
+	FHostGroupManager::Get().Tick(GIasHttpTimeOutMs, bStopRequested);
+	
+	// TODO: It would be better to only update connections as they need it, consider doing this on
+	// hostgroup connect/disconnect events.
+	HttpClient->UpdateConnections();
+
+	TickRequests();
+
+	OnTickIdleDelegate.Broadcast();
+}
+
 void FOnDemandHttpThread::TickRequests()
 {
 	// We can't limit cvars so add a hard coded cap to prevent concurrent requests from getting out of hand.
diff --git a/Engine/Source/Runtime/Experimental/IoStore/OnDemand/Private/OnDemandHttpThread.h b/Engine/Source/Runtime/Experimental/IoStore/OnDemand/Private/OnDemandHttpThread.h
--- a/Engine/Source/Runtime/Experimental/IoStore/OnDemand/Private/OnDemandHttpThread.h
+++ b/Engine/Source/Runtime/Experimental/IoStore/OnDemand/Private/OnDemandHttpThread.h
@@ -20,6 +20,7 @@
 #include "IO/IoAllocators.h"
 #include "IO/IoBuffer.h"
 #include "IO/IoOffsetLength.h"
+#include "Misc/SingleThreadRunnable.h"
 #include "Templates/UniquePtr.h"
 #include "ThreadSafeIntrusiveQueue.h"
 
@@ -34,6 +35,7 @@ struct FOnDemandChunkInfo;
 
 class FOnDemandHttpThread final
 	: public FRunnable
+	, public FSingleThreadRunnable
 {
 public:
 	using FCompletionCallback = TUniqueFunction<void(uint32 /*StatusCode*/ , FStringView/*ErrorReason*/, FIoBuffer&& /*Data*/)>;
@@ -81,12 +83,15 @@ public:
 		return OnTickIdleDelegate;
 	}
 
+	virtual class FSingleThreadRunnable* GetSingleThreadInterface() override { return this; }
+
 private:
 
 	virtual bool Init() override;
 	virtual uint32 Run() override;
 	virtual void Stop() override;
 	virtual void Exit() override;
+	virtual void Tick() override;
 
 	void TickRequests();
 	void DrainHttpRequests();
