I/O store on-demand detailed install result stats.

#rnx

diff --git a/Engine/Source/Runtime/Experimental/IoStore/OnDemand/Private/OnDemandContentInstaller.cpp b/Engine/Source/Runtime/Experimental/IoStore/OnDemand/Private/OnDemandContentInstaller.cpp
--- a/Engine/Source/Runtime/Experimental/IoStore/OnDemand/Private/OnDemandContentInstaller.cpp
+++ b/Engine/Source/Runtime/Experimental/IoStore/OnDemand/Private/OnDemandContentInstaller.cpp
@@ -614,7 +614,8 @@ void FOnDemandContentInstaller::CompleteRequest(FInstallRequestState& RequestSta
 		static_cast<uint64>(RequestState.HttpRequestHandles.Num()),
 		InstallResult.Progress.TotalInstallSize,
 		CacheHitRatio,
-		DurationCycles);
+		DurationCycles,
+		Request->ErrorCode);
 
 	UE_LOG(LogIoStoreOnDemand, Log, TEXT("Install request completed, Result='%s', ContentHandle=(%s), ContentSize=%.2lf MiB, InstallSize=%.2lf MiB, CacheHitRatio=%d %%, Duration=%d ms"),
 		GetIoErrorText(InstallResult.Status.GetErrorCode()), *LexToString(Request->Args.ContentHandle), double(InstallResult.Progress.TotalContentSize) / 1024.0 / 1024.0,
diff --git a/Engine/Source/Runtime/Experimental/IoStore/OnDemand/Private/Statistics.cpp b/Engine/Source/Runtime/Experimental/IoStore/OnDemand/Private/Statistics.cpp
--- a/Engine/Source/Runtime/Experimental/IoStore/OnDemand/Private/Statistics.cpp
+++ b/Engine/Source/Runtime/Experimental/IoStore/OnDemand/Private/Statistics.cpp
@@ -6,6 +6,7 @@
 #include "AnalyticsEventAttribute.h"
 #include "HAL/IConsoleManager.h"
 #include "HAL/PlatformTime.h"
+#include "IO/IoStatus.h"
 #include "IO/IoStoreOnDemand.h"
 #include "IasHostGroup.h"
 #include "Internationalization/Internationalization.h"
@@ -332,6 +333,9 @@ CSV_DEFINE_STAT(Ias, HttpDurationMsMax);
 // TRACE STATS IAD
 FIncrementalVariance	GInstallCacheHitRatio;
 FCounterAtomicInt		GTotalInstallCount(TEXT("Iad/InstallCount"), TraceCounterDisplayHint_None);
+FCounterAtomicInt		GTotalInstallErrorCount(TEXT("Iad/InstallErrorCount"), TraceCounterDisplayHint_None);
+FCounterAtomicInt		GTotalMissingChunkInstallErrorCount(TEXT("Iad/MissingChunkInstallErrorCount"), TraceCounterDisplayHint_None);
+FCounterAtomicInt		GTotalWriteErrorInstallErrorCount(TEXT("Iad/WriteErrorInstallErrorCount"), TraceCounterDisplayHint_None);
 FCounterAtomicInt		GInflightInstallCount(TEXT("Iad/InflightInstallCount"), TraceCounterDisplayHint_None);
 FCounterAtomicInt		GTotalDownloadedChunkCount(TEXT("Iad/DownloadedChunks"), TraceCounterDisplayHint_None);
 FCounterAtomicInt		GTotalDownloadedChunkBytes(TEXT("Iad/DownloadedBytes"), TraceCounterDisplayHint_Memory);
@@ -978,7 +982,8 @@ void FOnDemandContentInstallerStats::OnRequestCompleted(
 		uint64 DownloadedChunkCount,
 		uint64 DownloadedBytes,
 		double CacheHitRatio,
-		uint64 DurationCycles) 
+		uint64 DurationCycles,
+		EIoErrorCode ErrorCode) 
 {
 	if (RequestedChunkCount == 0)
 	{
@@ -986,6 +991,18 @@ void FOnDemandContentInstallerStats::OnRequestCompleted(
 	}
 
 	GInflightInstallCount.Add(-1);
+	if (ErrorCode != EIoErrorCode::Ok)
+	{
+		if (ErrorCode == EIoErrorCode::UnknownChunkID || ErrorCode == EIoErrorCode::NotFound)
+		{
+			GTotalMissingChunkInstallErrorCount.Add(1);
+		}
+		else if (ErrorCode == EIoErrorCode::WriteError)
+		{
+			GTotalWriteErrorInstallErrorCount.Add(1);
+		}
+		GTotalInstallErrorCount.Add(1);
+	}
 	GTotalInstallCount.Add(1);
 	GTotalDownloadedChunkCount.Add(DownloadedChunkCount);
 	GTotalDownloadedChunkBytes.Add(DownloadedBytes);
diff --git a/Engine/Source/Runtime/Experimental/IoStore/OnDemand/Private/Statistics.h b/Engine/Source/Runtime/Experimental/IoStore/OnDemand/Private/Statistics.h
--- a/Engine/Source/Runtime/Experimental/IoStore/OnDemand/Private/Statistics.h
+++ b/Engine/Source/Runtime/Experimental/IoStore/OnDemand/Private/Statistics.h
@@ -14,6 +14,7 @@
 #endif
 
 struct FAnalyticsEventAttribute;
+enum class EIoErrorCode;
 
 LLM_DECLARE_TAG(Ias);
 
@@ -92,7 +93,8 @@ public:
 		uint64 DownloadedChunkCount,
 		uint64 DownloadedBytes,
 		double CacheHitRatio,
-		uint64 DurationCycles) IAS_STATISTICS_IMPL()
+		uint64 DurationCycles,
+		EIoErrorCode ErrorCode) IAS_STATISTICS_IMPL()
 };
 
 #undef IAS_STATISTICS_IMPL
