EngineTest
* Added Net.Engine.ForceStopReplication that tests clients destroying actors within RPC's or OnReps

Iris
* Removed ensure when clients destroy actors with bNetForce=true

#rnx

diff --git a/Engine/Source/Runtime/Experimental/Iris/Core/Private/Iris/ReplicationSystem/NetRefHandleManager.cpp b/Engine/Source/Runtime/Experimental/Iris/Core/Private/Iris/ReplicationSystem/NetRefHandleManager.cpp
--- a/Engine/Source/Runtime/Experimental/Iris/Core/Private/Iris/ReplicationSystem/NetRefHandleManager.cpp
+++ b/Engine/Source/Runtime/Experimental/Iris/Core/Private/Iris/ReplicationSystem/NetRefHandleManager.cpp
@@ -287,15 +287,7 @@ const FReplicationInstanceProtocol* FNetRefHandleManager::DetachInstanceProtocol
 	{
 		FReplicatedObjectData& Data = ReplicatedObjectData[InternalIndex];
 		const FReplicationInstanceProtocol* InstanceProtocol = Data.InstanceProtocol;
-
-		// Detect if someone unexpectedly tries to detach replicated instance
-		const bool bDetachedInstancedIsPendingEndReplication = ReplicatedObjectData[InternalIndex].bPendingEndReplication == 1U;
-		if (!bDetachedInstancedIsPendingEndReplication)
-		{
-			UE_LOG(LogIris, Warning, TEXT("FNetRefHandleManager::DetachInstanceProtocol - DetachInstanceProtocol %s %s (%s) from client"), *PrintObjectFromIndex(InternalIndex),*Data.RefHandle.ToString(), Data.Protocol->DebugName->Name);
-			ensureMsgf(false, TEXT("DetachInstanceProtocol for %s %s (%s) from client"), *PrintObjectFromIndex(InternalIndex),*Data.RefHandle.ToString(), Data.Protocol->DebugName->Name);
-		}
-		
+	
 		Data.InstanceProtocol = nullptr;
 		ReplicatedInstances[InternalIndex] = nullptr;
 
