Dataflow : Add APi to override basic type variables programmatically

diff --git a/Engine/Source/Runtime/Experimental/Dataflow/Engine/Private/Dataflow/DataflowInstance.cpp b/Engine/Source/Runtime/Experimental/Dataflow/Engine/Private/Dataflow/DataflowInstance.cpp
--- a/Engine/Source/Runtime/Experimental/Dataflow/Engine/Private/Dataflow/DataflowInstance.cpp
+++ b/Engine/Source/Runtime/Experimental/Dataflow/Engine/Private/Dataflow/DataflowInstance.cpp
@@ -140,6 +140,58 @@ FDataflowVariableOverrides& FDataflowVariableOverrides::operator=(const FDataflo
 	return *this;
 }
 
+bool FDataflowVariableOverrides::OverrideVariableBool(FName VariableName, bool Value)
+{
+	if (const FPropertyBagPropertyDesc* Desc = Variables.FindPropertyDescByName(VariableName))
+	{
+		if (Variables.SetValueBool(*Desc, Value) == EPropertyBagResult::Success)
+		{
+			OverriddenVariableGuids.AddUnique(Desc->ID);
+			return true;
+		}
+	}
+	return false;
+}
+
+bool FDataflowVariableOverrides::OverrideVariableInt(FName VariableName, int64 Value)
+{
+	if (const FPropertyBagPropertyDesc* Desc = Variables.FindPropertyDescByName(VariableName))
+	{
+		if (Variables.SetValueInt64(*Desc, Value) == EPropertyBagResult::Success)
+		{
+			OverriddenVariableGuids.AddUnique(Desc->ID);
+			return true;
+		}
+	}
+	return false;
+}
+
+bool FDataflowVariableOverrides::OverrideVariableFloat(FName VariableName, float Value)
+{
+	if (const FPropertyBagPropertyDesc* Desc = Variables.FindPropertyDescByName(VariableName))
+	{
+		if (Variables.SetValueFloat(*Desc, Value) == EPropertyBagResult::Success)
+		{
+			OverriddenVariableGuids.AddUnique(Desc->ID);
+			return true;
+		}
+	}
+	return false;
+}
+
+bool FDataflowVariableOverrides::OverrideVariableObject(FName VariableName, const UObject* Value)
+{
+	if (const FPropertyBagPropertyDesc* Desc = Variables.FindPropertyDescByName(VariableName))
+	{
+		if (Variables.SetValueObject(*Desc, Value) == EPropertyBagResult::Success)
+		{
+			OverriddenVariableGuids.AddUnique(Desc->ID);
+			return true;
+		}
+	}
+	return false;
+}
+
 const FInstancedPropertyBag* FDataflowVariableOverrides::GetDefaultVariablesFromAsset() const 
 {
 	if (Owner)
diff --git a/Engine/Source/Runtime/Experimental/Dataflow/Engine/Public/Dataflow/DataflowInstance.h b/Engine/Source/Runtime/Experimental/Dataflow/Engine/Public/Dataflow/DataflowInstance.h
--- a/Engine/Source/Runtime/Experimental/Dataflow/Engine/Public/Dataflow/DataflowInstance.h
+++ b/Engine/Source/Runtime/Experimental/Dataflow/Engine/Public/Dataflow/DataflowInstance.h
@@ -47,6 +47,18 @@ struct FDataflowVariableOverrides
 
 	DATAFLOWENGINE_API static FName GetVariablePropertyName();
 
+	/** override dataflow boolean variables for this asset */
+	DATAFLOWENGINE_API bool OverrideVariableBool(FName VariableName, bool Value);
+
+	/** override dataflow integer variables for this asset */
+	DATAFLOWENGINE_API bool OverrideVariableInt(FName VariableName, int64 Value);
+
+	/** override dataflow float variables for this asset */
+	DATAFLOWENGINE_API bool OverrideVariableFloat(FName VariableName, float Value);
+
+	/** override dataflow UObject variables for this asset */
+	DATAFLOWENGINE_API bool OverrideVariableObject(FName VariableName, const UObject* Value);
+
 #if WITH_EDITOR
 
 	void OnOwnerPostEditChangeProperty(struct FPropertyChangedEvent& PropertyChangedEvent);
