[Backout] - CL41652191
The tentative fix didn't solve the reported crashes. Going back to original behavior instead

#rnx
Original CL Desc
-----------------------------------------------------------------
Iris
* Immediately prune garbage collected pointers from within PostGarbageCollect delegate instead of waiting for the next TickFlush to occur
* Tentative fix for a crash where we were executing RPC's right after CollectingGarbage and calling RPCs on destroyed objects

#rnx

diff --git a/Engine/Source/Runtime/Experimental/Iris/Core/Private/Iris/ReplicationSystem/ReplicationSystem.cpp b/Engine/Source/Runtime/Experimental/Iris/Core/Private/Iris/ReplicationSystem/ReplicationSystem.cpp
--- a/Engine/Source/Runtime/Experimental/Iris/Core/Private/Iris/ReplicationSystem/ReplicationSystem.cpp
+++ b/Engine/Source/Runtime/Experimental/Iris/Core/Private/Iris/ReplicationSystem/ReplicationSystem.cpp
@@ -787,6 +787,7 @@ UReplicationSystem::UReplicationSystem()
 , PIEInstanceID(INDEX_NONE)
 , bIsServer(0)
 , bAllowObjectReplication(0)
+, bDoCollectGarbage(0)
 {
 }
 
@@ -862,7 +863,7 @@ void UReplicationSystem::PreSendUpdate(const FSendUpdateParams& Params)
 		InternalSys.GetTickReplicationStats() = {};
 	
 		// Force a integrity check of all replicated instances
-		if (ReplicationSystemCVars::bForcePruneBeforeUpdate)
+		if (bDoCollectGarbage || ReplicationSystemCVars::bForcePruneBeforeUpdate)
 		{
 			CollectGarbage();
 		}
@@ -1041,7 +1042,7 @@ void UReplicationSystem::PostSendUpdate()
 
 void UReplicationSystem::PostGarbageCollection()
 {
-	CollectGarbage();
+	bDoCollectGarbage = 1U;
 }
 
 void UReplicationSystem::CollectGarbage()
@@ -1051,6 +1052,8 @@ void UReplicationSystem::CollectGarbage()
 	// Prune stale object instances before descriptors and protocols are pruned
 	Impl->ReplicationSystemInternal.GetReplicationBridge()->CallPruneStaleObjects();
 	Impl->ReplicationSystemInternal.GetReplicationStateDescriptorRegistry().PruneStaleDescriptors();
+
+	bDoCollectGarbage = 0;
 }
 
 void UReplicationSystem::ResetGameWorldState()
diff --git a/Engine/Source/Runtime/Experimental/Iris/Core/Public/Iris/ReplicationSystem/ReplicationSystem.h b/Engine/Source/Runtime/Experimental/Iris/Core/Public/Iris/ReplicationSystem/ReplicationSystem.h
--- a/Engine/Source/Runtime/Experimental/Iris/Core/Public/Iris/ReplicationSystem/ReplicationSystem.h
+++ b/Engine/Source/Runtime/Experimental/Iris/Core/Public/Iris/ReplicationSystem/ReplicationSystem.h
@@ -765,6 +765,7 @@ private:
 	int32 PIEInstanceID;
 	uint32 bIsServer : 1;
 	uint32 bAllowObjectReplication : 1;
+	uint32 bDoCollectGarbage : 1;
 };
 
 namespace UE::Net
