Dataflow : add array support to setting variable overrides programmatically

diff --git a/Engine/Source/Runtime/Experimental/Dataflow/Engine/Private/Dataflow/DataflowInstance.cpp b/Engine/Source/Runtime/Experimental/Dataflow/Engine/Private/Dataflow/DataflowInstance.cpp
--- a/Engine/Source/Runtime/Experimental/Dataflow/Engine/Private/Dataflow/DataflowInstance.cpp
+++ b/Engine/Source/Runtime/Experimental/Dataflow/Engine/Private/Dataflow/DataflowInstance.cpp
@@ -107,6 +107,33 @@ namespace UE::Dataflow::InstanceUtils
 		}
 		return TerminalNodeNames;
 	}
+
+	namespace Private
+	{
+		template<typename T, typename TWriteElementFunction>
+		bool OverrideVariableArray(FInstancedPropertyBag& Variables, FName VariableName, const TArray<T>& Values, const TWriteElementFunction& WriteElementFunction, TArray<FGuid>& InOutOverrideGuids)
+		{
+			if (const FPropertyBagPropertyDesc* Desc = Variables.FindPropertyDescByName(VariableName))
+			{
+				if (FPropertyBagArrayRef* ArrayRef = Variables.GetMutableArrayRef(*Desc).TryGetValue())
+				{
+					bool bSuccess = true;
+					ArrayRef->EmptyValues();
+					ArrayRef->AddValues(Values.Num());
+					for (int32 Idx = 0; Idx < Values.Num(); ++Idx)
+					{
+						bSuccess &= (WriteElementFunction(*ArrayRef, Idx, Values[Idx]) == EPropertyBagResult::Success);
+					}
+					if (bSuccess)
+					{
+						InOutOverrideGuids.AddUnique(Desc->ID);
+					}
+					return bSuccess;
+				}
+			}
+			return false;
+		}
+	}
 }
 
 
@@ -153,6 +180,18 @@ bool FDataflowVariableOverrides::OverrideVariableBool(FName VariableName, bool V
 	return false;
 }
 
+bool FDataflowVariableOverrides::OverrideVariableBoolArray(FName VariableName, const TArray<bool>& Values)
+{
+	using namespace UE::Dataflow::InstanceUtils::Private;
+
+	auto WriteBoolValue = [](FPropertyBagArrayRef& ArrayRef, int32 Idx, bool Value)
+		{
+			return ArrayRef.SetValueBool(Idx, Value);
+		};
+
+	return OverrideVariableArray(Variables, VariableName, Values, WriteBoolValue, OverriddenVariableGuids);
+}
+
 bool FDataflowVariableOverrides::OverrideVariableInt(FName VariableName, int64 Value)
 {
 	if (const FPropertyBagPropertyDesc* Desc = Variables.FindPropertyDescByName(VariableName))
@@ -166,6 +205,30 @@ bool FDataflowVariableOverrides::OverrideVariableInt(FName VariableName, int64 V
 	return false;
 }
 
+bool FDataflowVariableOverrides::OverrideVariableInt32Array(FName VariableName, const TArray<int32>& Values)
+{
+	using namespace UE::Dataflow::InstanceUtils::Private;
+
+	auto WriteInt32Value = [](FPropertyBagArrayRef& ArrayRef, int32 Idx, int32 Value)
+		{
+			return ArrayRef.SetValueInt32(Idx, Value);
+		};
+
+	return OverrideVariableArray(Variables, VariableName, Values, WriteInt32Value, OverriddenVariableGuids);
+}
+
+bool FDataflowVariableOverrides::OverrideVariableInt64Array(FName VariableName, const TArray<int64>& Values)
+{
+	using namespace UE::Dataflow::InstanceUtils::Private;
+
+	auto WriteInt64Value = [](FPropertyBagArrayRef& ArrayRef, int32 Idx, int64 Value)
+		{
+			return ArrayRef.SetValueInt64(Idx, Value);
+		};
+
+	return OverrideVariableArray(Variables, VariableName, Values, WriteInt64Value, OverriddenVariableGuids);
+}
+
 bool FDataflowVariableOverrides::OverrideVariableFloat(FName VariableName, float Value)
 {
 	if (const FPropertyBagPropertyDesc* Desc = Variables.FindPropertyDescByName(VariableName))
@@ -179,6 +242,18 @@ bool FDataflowVariableOverrides::OverrideVariableFloat(FName VariableName, float
 	return false;
 }
 
+bool FDataflowVariableOverrides::OverrideVariableFloatArray(FName VariableName, const TArray<float>& Values)
+{
+	using namespace UE::Dataflow::InstanceUtils::Private;
+
+	auto WriteFloatValue = [](FPropertyBagArrayRef& ArrayRef, int32 Idx, float Value)
+		{
+			return ArrayRef.SetValueFloat(Idx, Value);
+		};
+
+	return OverrideVariableArray(Variables, VariableName, Values, WriteFloatValue, OverriddenVariableGuids);
+}
+
 bool FDataflowVariableOverrides::OverrideVariableObject(FName VariableName, const UObject* Value)
 {
 	if (const FPropertyBagPropertyDesc* Desc = Variables.FindPropertyDescByName(VariableName))
@@ -192,6 +267,80 @@ bool FDataflowVariableOverrides::OverrideVariableObject(FName VariableName, cons
 	return false;
 }
 
+bool FDataflowVariableOverrides::OverrideVariableObjectArray(FName VariableName, const TArray<TObjectPtr<UObject>>& Values)
+{
+	using namespace UE::Dataflow::InstanceUtils::Private;
+
+	auto WriteObjectValue = [](FPropertyBagArrayRef& ArrayRef, int32 Idx, const TObjectPtr<UObject>& Value)
+		{
+			return ArrayRef.SetValueObject(Idx, Value);
+		};
+
+	return OverrideVariableArray(Variables, VariableName, Values, WriteObjectValue, OverriddenVariableGuids);
+}
+
+bool FDataflowVariableOverrides::OverrideVariableObjectArray(FName VariableName, const TArray<UObject*>& Values)
+{
+	using namespace UE::Dataflow::InstanceUtils::Private;
+
+	auto WriteObjectValue = [](FPropertyBagArrayRef& ArrayRef, int32 Idx, UObject* Value)
+		{
+			return ArrayRef.SetValueObject(Idx, Value);
+		};
+
+	return OverrideVariableArray(Variables, VariableName, Values, WriteObjectValue, OverriddenVariableGuids);
+}
+
+bool FDataflowVariableOverrides::OverrideVariableName(FName VariableName, FName Value)
+{
+	if (const FPropertyBagPropertyDesc* Desc = Variables.FindPropertyDescByName(VariableName))
+	{
+		if (Variables.SetValueName(*Desc, Value) == EPropertyBagResult::Success)
+		{
+			OverriddenVariableGuids.AddUnique(Desc->ID);
+			return true;
+		}
+	}
+	return false;
+}
+
+bool FDataflowVariableOverrides::OverrideVariableName(FName VariableName, const TArray<FName>& Values)
+{
+	using namespace UE::Dataflow::InstanceUtils::Private;
+
+	auto WriteNameValue = [](FPropertyBagArrayRef& ArrayRef, int32 Idx, const FName& Value)
+		{
+			return ArrayRef.SetValueName(Idx, Value);
+		};
+
+	return OverrideVariableArray(Variables, VariableName, Values, WriteNameValue, OverriddenVariableGuids);
+}
+
+bool FDataflowVariableOverrides::OverrideVariableString(FName VariableName, FString Value)
+{
+	if (const FPropertyBagPropertyDesc* Desc = Variables.FindPropertyDescByName(VariableName))
+	{
+		if (Variables.SetValueString(*Desc, Value) == EPropertyBagResult::Success)
+		{
+			OverriddenVariableGuids.AddUnique(Desc->ID);
+			return true;
+		}
+	}
+	return false;
+}
+
+bool FDataflowVariableOverrides::OverrideVariableString(FName VariableName, const TArray<FString>& Values)
+{
+	using namespace UE::Dataflow::InstanceUtils::Private;
+
+	auto WriteStringValue = [](FPropertyBagArrayRef& ArrayRef, int32 Idx, const FString& Value)
+		{
+			return ArrayRef.SetValueString(Idx, Value);
+		};
+
+	return OverrideVariableArray(Variables, VariableName, Values, WriteStringValue, OverriddenVariableGuids);
+}
+
 const FInstancedPropertyBag* FDataflowVariableOverrides::GetDefaultVariablesFromAsset() const 
 {
 	if (Owner)
diff --git a/Engine/Source/Runtime/Experimental/Dataflow/Engine/Public/Dataflow/DataflowInstance.h b/Engine/Source/Runtime/Experimental/Dataflow/Engine/Public/Dataflow/DataflowInstance.h
--- a/Engine/Source/Runtime/Experimental/Dataflow/Engine/Public/Dataflow/DataflowInstance.h
+++ b/Engine/Source/Runtime/Experimental/Dataflow/Engine/Public/Dataflow/DataflowInstance.h
@@ -47,18 +47,48 @@ struct FDataflowVariableOverrides
 
 	DATAFLOWENGINE_API static FName GetVariablePropertyName();
 
-	/** override dataflow boolean variables for this asset */
+	/** override a dataflow boolean variable for this asset */
 	DATAFLOWENGINE_API bool OverrideVariableBool(FName VariableName, bool Value);
 
-	/** override dataflow integer variables for this asset */
+	/** override a dataflow boolean array variable for this asset */
+	DATAFLOWENGINE_API bool OverrideVariableBoolArray(FName VariableName, const TArray<bool>& Values);
+
+	/** override a dataflow integer variable for this asset */
 	DATAFLOWENGINE_API bool OverrideVariableInt(FName VariableName, int64 Value);
 
-	/** override dataflow float variables for this asset */
+	/** override a dataflow integer array variable for this asset */
+	DATAFLOWENGINE_API bool OverrideVariableInt32Array(FName VariableName, const TArray<int32>& Values);
+
+	/** override a dataflow integer array variable for this asset */
+	DATAFLOWENGINE_API bool OverrideVariableInt64Array(FName VariableName, const TArray<int64>& Values);
+
+	/** override a dataflow float variable for this asset */
 	DATAFLOWENGINE_API bool OverrideVariableFloat(FName VariableName, float Value);
 
-	/** override dataflow UObject variables for this asset */
+	/** override a dataflow float array variable for this asset */
+	DATAFLOWENGINE_API bool OverrideVariableFloatArray(FName VariableName, const TArray<float>& Values);
+
+	/** override a dataflow UObject variable for this asset */
 	DATAFLOWENGINE_API bool OverrideVariableObject(FName VariableName, const UObject* Value);
 
+	/** override a dataflow UObject array variable for this asset */
+	DATAFLOWENGINE_API bool OverrideVariableObjectArray(FName VariableName, const TArray<TObjectPtr<UObject>>& Values);
+
+	/** override a dataflow UObject array variable for this asset */
+	DATAFLOWENGINE_API bool OverrideVariableObjectArray(FName VariableName, const TArray<UObject*>& Values);
+
+	/** override a dataflow FName variable for this asset */
+	DATAFLOWENGINE_API bool OverrideVariableName(FName VariableName, FName Value);
+
+	/** override a dataflow FName variable for this asset */
+	DATAFLOWENGINE_API bool OverrideVariableName(FName VariableName, const TArray<FName>& Values);
+
+	/** override a dataflow String variable for this asset */
+	DATAFLOWENGINE_API bool OverrideVariableString(FName VariableName, FString Value);
+
+	/** override a dataflow String variable for this asset */
+	DATAFLOWENGINE_API bool OverrideVariableString(FName VariableName, const TArray<FString>& Values);
+
 #if WITH_EDITOR
 
 	void OnOwnerPostEditChangeProperty(struct FPropertyChangedEvent& PropertyChangedEvent);
