Dataflow blueprint library : extend the API to support override for variables and an easier way to generate asset from dataflow
- Variable overrides supports : bool, integer, float, string, name , UObject and arrays of those types

diff --git a/Engine/Source/Runtime/Experimental/Dataflow/Engine/Private/Dataflow/DataflowBlueprintLibrary.cpp b/Engine/Source/Runtime/Experimental/Dataflow/Engine/Private/Dataflow/DataflowBlueprintLibrary.cpp
--- a/Engine/Source/Runtime/Experimental/Dataflow/Engine/Private/Dataflow/DataflowBlueprintLibrary.cpp
+++ b/Engine/Source/Runtime/Experimental/Dataflow/Engine/Private/Dataflow/DataflowBlueprintLibrary.cpp
@@ -3,6 +3,8 @@
 #include "Dataflow/DataflowBlueprintLibrary.h"
 
 #include "Dataflow/DataflowObjectInterface.h"
+#include "Dataflow/DataflowInstance.h"
+#include "StructUtils/PropertyBag.h"
 
 #include UE_INLINE_GENERATED_CPP_BY_NAME(DataflowBlueprintLibrary)
 
@@ -34,9 +36,87 @@ void UDataflowBlueprintLibrary::EvaluateTerminalNodeByName(UDataflow* Dataflow,
 	}
 }
 
+bool UDataflowBlueprintLibrary::RegenerateAssetFromDataflow(UObject* AssetToRegenerate, bool bRegenerateDependentAssets)
+{
+	if (IDataflowInstanceInterface* DataflowInterface = Cast<IDataflowInstanceInterface>(AssetToRegenerate))
+	{
+		FDataflowInstance& DataflowInstance = DataflowInterface->GetDataflowInstance();
+		return DataflowInstance.UpdateOwnerAsset(bRegenerateDependentAssets);
+	}
+	return false;
+}
 
+bool UDataflowBlueprintLibrary::OverrideDataflowVariableBool(UObject* Asset, FName VariableName, bool VariableValue)
+{
+	if (IDataflowInstanceInterface* DataflowInterface = Cast<IDataflowInstanceInterface>(Asset))
+	{
+		return DataflowInterface->GetDataflowInstance().GetVariableOverrides().OverrideVariableBool(VariableName, VariableValue);
+	}
+	return false;
+}
 
+bool UDataflowBlueprintLibrary::OverrideDataflowVariableBoolArray(UObject* Asset, FName VariableName, const TArray<bool>& VariableArrayValue)
+{
+	if (IDataflowInstanceInterface* DataflowInterface = Cast<IDataflowInstanceInterface>(Asset))
+	{
+		return DataflowInterface->GetDataflowInstance().GetVariableOverrides().OverrideVariableBoolArray(VariableName, VariableArrayValue);
+	}
+	return false;
+}
+
+bool UDataflowBlueprintLibrary::OverrideDataflowVariableInt(UObject* Asset, FName VariableName, int64 VariableValue)
+{
+	if (IDataflowInstanceInterface* DataflowInterface = Cast<IDataflowInstanceInterface>(Asset))
+	{
+		return DataflowInterface->GetDataflowInstance().GetVariableOverrides().OverrideVariableInt(VariableName, VariableValue);
+	}
+	return false;
+}
 
+bool UDataflowBlueprintLibrary::OverrideDataflowVariableIntArray(UObject* Asset, FName VariableName, const TArray<int32>& VariableArrayValue)
+{
+	if (IDataflowInstanceInterface* DataflowInterface = Cast<IDataflowInstanceInterface>(Asset))
+	{
+		return DataflowInterface->GetDataflowInstance().GetVariableOverrides().OverrideVariableInt32Array(VariableName, VariableArrayValue);
+	}
+	return false;
+}
+
+bool UDataflowBlueprintLibrary::OverrideDataflowVariableFloat(UObject* Asset, FName VariableName, float VariableValue)
+{
+	if (IDataflowInstanceInterface* DataflowInterface = Cast<IDataflowInstanceInterface>(Asset))
+	{
+		return DataflowInterface->GetDataflowInstance().GetVariableOverrides().OverrideVariableFloat(VariableName, VariableValue);
+	}
+	return false;
+}
+
+bool UDataflowBlueprintLibrary::OverrideDataflowVariableFloatArray(UObject* Asset, FName VariableName, const TArray<float>& VariableArrayValue)
+{
+	if (IDataflowInstanceInterface* DataflowInterface = Cast<IDataflowInstanceInterface>(Asset))
+	{
+		return DataflowInterface->GetDataflowInstance().GetVariableOverrides().OverrideVariableFloatArray(VariableName, VariableArrayValue);
+	}
+	return false;
+}
+
+bool UDataflowBlueprintLibrary::OverrideDataflowVariableObject(UObject* Asset, FName VariableName, UObject* VariableValue)
+{
+	if (IDataflowInstanceInterface* DataflowInterface = Cast<IDataflowInstanceInterface>(Asset))
+	{
+		return DataflowInterface->GetDataflowInstance().GetVariableOverrides().OverrideVariableObject(VariableName, VariableValue);
+	}
+	return false;
+}
+
+bool UDataflowBlueprintLibrary::OverrideDataflowVariableObjectArray(UObject* Asset, FName VariableName, const TArray<UObject*>& VariableArrayValue)
+{
+	if (IDataflowInstanceInterface* DataflowInterface = Cast<IDataflowInstanceInterface>(Asset))
+	{
+		return DataflowInterface->GetDataflowInstance().GetVariableOverrides().OverrideVariableObjectArray(VariableName, VariableArrayValue);
+	}
+	return false;
+}
 
 
 
diff --git a/Engine/Source/Runtime/Experimental/Dataflow/Engine/Public/Dataflow/DataflowBlueprintLibrary.h b/Engine/Source/Runtime/Experimental/Dataflow/Engine/Public/Dataflow/DataflowBlueprintLibrary.h
--- a/Engine/Source/Runtime/Experimental/Dataflow/Engine/Public/Dataflow/DataflowBlueprintLibrary.h
+++ b/Engine/Source/Runtime/Experimental/Dataflow/Engine/Public/Dataflow/DataflowBlueprintLibrary.h
@@ -18,4 +18,67 @@ public:
 	*/
 	UFUNCTION(BlueprintCallable, Category = "Dataflow", meta = (Keywords = "Dataflow graph"))
 	static DATAFLOWENGINE_API void EvaluateTerminalNodeByName(UDataflow* Dataflow, FName TerminalNodeName, UObject* ResultAsset);
+
+	/**
+	* Regenerate an asset using its corresponding dataflow
+	* @return true if the asset was regenerated
+	*/
+	UFUNCTION(BlueprintCallable, Category = "Dataflow", meta = (Keywords = "Dataflow graph"))
+	static DATAFLOWENGINE_API bool RegenerateAssetFromDataflow(UObject* AssetToRegenerate, bool bRegenerateDependentAssets = false);
+
+	/**
+	* Override a Boolean dataflow variable for a specific asset
+	* @return true if the override was successful
+	*/
+	UFUNCTION(BlueprintCallable, Category = "Dataflow", meta = (Keywords = "Dataflow graph"))
+	static DATAFLOWENGINE_API bool OverrideDataflowVariableBool(UObject* Asset, FName VariableName, bool VariableValue);
+
+	/**
+	* Override an Boolean Array dataflow variable for a specific asset
+	* @return true if the override was successful
+	*/
+	UFUNCTION(BlueprintCallable, Category = "Dataflow", meta = (Keywords = "Dataflow graph"))
+	static DATAFLOWENGINE_API bool OverrideDataflowVariableBoolArray(UObject* Asset, FName VariableName, const TArray<bool>& VariableArrayValue);
+
+	/**
+	* Override an Integer dataflow variable for a specific asset
+	* @return true if the override was successful
+	*/
+	UFUNCTION(BlueprintCallable, Category = "Dataflow", meta = (Keywords = "Dataflow graph"))
+	static DATAFLOWENGINE_API bool OverrideDataflowVariableInt(UObject* Asset, FName VariableName, int64 VariableValue);
+
+	/**
+	* Override an Integer Array dataflow variable for a specific asset
+	* @return true if the override was successful
+	*/
+	UFUNCTION(BlueprintCallable, Category = "Dataflow", meta = (Keywords = "Dataflow graph"))
+	static DATAFLOWENGINE_API bool OverrideDataflowVariableIntArray(UObject* Asset, FName VariableName, const TArray<int32>& VariableArrayValue);
+
+	/**
+	* Override a Float dataflow variable for a specific asset
+	* @return true if the override was successful
+	*/
+	UFUNCTION(BlueprintCallable, Category = "Dataflow", meta = (Keywords = "Dataflow graph"))
+	static DATAFLOWENGINE_API bool OverrideDataflowVariableFloat(UObject* Asset, FName VariableName, float VariableValue);
+
+	/**
+	* Override a Float Array dataflow variable for a specific asset
+	* @return true if the override was successful
+	*/
+	UFUNCTION(BlueprintCallable, Category = "Dataflow", meta = (Keywords = "Dataflow graph"))
+	static DATAFLOWENGINE_API bool OverrideDataflowVariableFloatArray(UObject* Asset, FName VariableName, const TArray<float>& VariableArrayValue);
+
+	/**
+	* Override an Unreal Object dataflow variable for a specific asset
+	* @return true if the override was successful
+	*/
+	UFUNCTION(BlueprintCallable, Category = "Dataflow", meta = (Keywords = "Dataflow graph"))
+	static DATAFLOWENGINE_API bool OverrideDataflowVariableObject(UObject* Asset, FName VariableName, UObject* VariableValue);
+
+	/**
+	* Override an Unreal Object Array dataflow variable for a specific asset
+	* @return true if the override was successful
+	*/
+	UFUNCTION(BlueprintCallable, Category = "Dataflow", meta = (Keywords = "Dataflow graph"))
+	static DATAFLOWENGINE_API bool OverrideDataflowVariableObjectArray(UObject* Asset, FName VariableName, const TArray<UObject*>& VariableArrayValue);
 };
