Geometry collection : expose a project setting to optimize geometry collection for runtime memory
- this is equivalent to enabling a series of cvars in the geometry collection physics proxy

diff --git a/Engine/Source/Runtime/Experimental/Chaos/Private/Chaos/PBDRigidClustering.cpp b/Engine/Source/Runtime/Experimental/Chaos/Private/Chaos/PBDRigidClustering.cpp
--- a/Engine/Source/Runtime/Experimental/Chaos/Private/Chaos/PBDRigidClustering.cpp
+++ b/Engine/Source/Runtime/Experimental/Chaos/Private/Chaos/PBDRigidClustering.cpp
@@ -26,7 +26,10 @@
 #include "PhysicsProxy/ClusterUnionPhysicsProxy.h"
 #include "CoreMinimal.h"
 
-extern CHAOS_API bool bBuildGeometryForChildrenOnPT;
+namespace GeometryCollectionPhysicsProxy::Private
+{
+	extern CHAOS_API bool BuildGeometryForChildrenOnPT(const FSimulationParameters& Params);
+}
 
 namespace Chaos
 {
@@ -1135,7 +1138,7 @@ namespace Chaos
 			BreakDamagePropagationFactor = SimParams.BreakDamagePropagationFactor;
 			ShockDamagePropagationFactor = SimParams.ShockDamagePropagationFactor;
 			
-			if (bBuildGeometryForChildrenOnPT == false)
+			if (GeometryCollectionPhysicsProxy::Private::BuildGeometryForChildrenOnPT(SimParams) == false)
 			{
 				ConcreteGCProxy->CreateChildrenGeometry_Internal();
 			}
diff --git a/Engine/Source/Runtime/Experimental/Chaos/Private/PhysicsProxy/GeometryCollectionPhysicsProxy.cpp b/Engine/Source/Runtime/Experimental/Chaos/Private/PhysicsProxy/GeometryCollectionPhysicsProxy.cpp
--- a/Engine/Source/Runtime/Experimental/Chaos/Private/PhysicsProxy/GeometryCollectionPhysicsProxy.cpp
+++ b/Engine/Source/Runtime/Experimental/Chaos/Private/PhysicsProxy/GeometryCollectionPhysicsProxy.cpp
@@ -537,23 +537,66 @@ FGeometryCollectionPhysicsProxy::~FGeometryCollectionPhysicsProxy()
 float ReportHighParticleFraction = -1.f;
 FAutoConsoleVariableRef CVarReportHighParticleFraction(TEXT("p.gc.ReportHighParticleFraction"), ReportHighParticleFraction, TEXT("Report any objects with particle fraction above this threshold"));
 
-CHAOS_API bool bBuildGeometryForChildrenOnPT = true;
-FAutoConsoleVariableRef CVarbBuildGeometryForChildrenOnPT(TEXT("p.gc.BuildGeometryForChildrenOnPT"), bBuildGeometryForChildrenOnPT, TEXT("If true build all children geometry on Physics Thread at initilaization time, otherwise wait until destruction occurs."));
+bool bBuildGeometryForChildrenOnPT = true;
+FAutoConsoleVariableRef CVarbBuildGeometryForChildrenOnPT(
+	TEXT("p.gc.BuildGeometryForChildrenOnPT"), 
+	bBuildGeometryForChildrenOnPT, 
+	TEXT("DEPRECATED : use the 'Optimize For Memory' option from project settings instead - "
+		 "If true build all children geometry on Physics Thread at initilaization time, otherwise wait until destruction occurs.")
+);
 
 bool bBuildGeometryForChildrenOnGT = true;
-FAutoConsoleVariableRef CVarbBuildGeometryForChildrenOnGT(TEXT("p.gc.BuildGeometryForChildrenOnGT"), bBuildGeometryForChildrenOnGT, TEXT("If true build all children geometry  on Game Thread at initilaization time, otherwise wait until destruction occurs."));
+FAutoConsoleVariableRef CVarbBuildGeometryForChildrenOnGT(
+	TEXT("p.gc.BuildGeometryForChildrenOnGT"), 
+	bBuildGeometryForChildrenOnGT, 
+	TEXT("DEPRECATED : use the 'Optimize For Memory' option from project settings instead - "
+		 "If true build all children geometry  on Game Thread at initilaization time, otherwise wait until destruction occurs.")
+);
 
 bool bCreateGTParticleForChildren = true;
-FAutoConsoleVariableRef CVarCreateGTParticleForChildren(TEXT("p.gc.CreateGTParticlesForChildren"), bCreateGTParticleForChildren, TEXT("If true create all children particles at initilaization time, otherwise wait until destruction occurs."));
+FAutoConsoleVariableRef CVarCreateGTParticleForChildren(
+	TEXT("p.gc.CreateGTParticlesForChildren"), 
+	bCreateGTParticleForChildren, 
+	TEXT("DEPRECATED : use the 'Optimize For Memory' option from project settings instead - "
+		 "If true create all children particles at initilaization time, otherwise wait until destruction occurs.")
+);
 
 bool bRemoveImplicitsInDynamicCollections = false;
-FAutoConsoleVariableRef CVarbRemoveImplicitsInDynamicCollections(TEXT("p.gc.RemoveImplicitsInDynamicCollections"), 
-	bRemoveImplicitsInDynamicCollections, TEXT("This cvar has an impact only if geometry are not added for children. It removes implicits from the Dynamic Collections, and recreate then from the rest collection. \
-										Using this cvar could have an impact if geometry are updated from the dynamic collection on the GT, then those changes won't be ported to the PT."));
+FAutoConsoleVariableRef CVarbRemoveImplicitsInDynamicCollections(
+	TEXT("p.gc.RemoveImplicitsInDynamicCollections"), 
+	bRemoveImplicitsInDynamicCollections, 
+	TEXT("DEPRECATED : use the 'Optimize For Memory' option from project settings instead - "
+		 "This cvar has an impact only if geometry are not added for children. It removes implicits from the Dynamic Collections, and recreate then from the rest collection. "
+	  	 "Using this cvar could have an impact if geometry are updated from the dynamic collection on the GT, then those changes won't be ported to the PT.")
+);
 
+namespace GeometryCollectionPhysicsProxy::Private
+{
+	CHAOS_API bool BuildGeometryForChildrenOnPT(const FSimulationParameters& Params)
+	{
+		return bBuildGeometryForChildrenOnPT && !Params.bOptimizeForRuntimeMemory;
+	}
+
+	bool BuildGeometryForChildrenOnGT(const FSimulationParameters& Params)
+	{
+		return bBuildGeometryForChildrenOnGT && !Params.bOptimizeForRuntimeMemory;
+	}
+
+	bool CreateGTParticleForChildren(const FSimulationParameters& Params)
+	{
+		return bCreateGTParticleForChildren && !Params.bOptimizeForRuntimeMemory;
+	}
+
+	bool RemoveImplicitsInDynamicCollections(const FSimulationParameters& Params)
+	{
+		return bRemoveImplicitsInDynamicCollections || Params.bOptimizeForRuntimeMemory;
+	}
+}
 
 void FGeometryCollectionPhysicsProxy::Initialize(Chaos::FPBDRigidsEvolutionBase *Evolution)
 {
+	using namespace GeometryCollectionPhysicsProxy::Private;
+
 	check(IsInGameThread());
 	//
 	// Game thread initialization. 
@@ -636,7 +679,7 @@ void FGeometryCollectionPhysicsProxy::Initialize(Chaos::FPBDRigidsEvolutionBase
 	if (ensure(NumTransforms == Implicits.Num() && NumEffectiveParticles == GTParticles.Num())) // Implicits are in the transform group so this invariant should always hold
 	{
 		constexpr bool bInitializationTime = true;
-		bHasBuiltGeometryOnGT = bBuildGeometryForChildrenOnGT;
+		bHasBuiltGeometryOnGT = BuildGeometryForChildrenOnGT(Parameters);
 		CreateGTParticles(Implicits, Evolution, bInitializationTime);
 
 		// Skip simplicials, as they're owned by unique pointers.
@@ -672,13 +715,13 @@ void FGeometryCollectionPhysicsProxy::Initialize(Chaos::FPBDRigidsEvolutionBase
 		}
 	}
 
-	if (bHasBuiltGeometryOnGT || bRemoveImplicitsInDynamicCollections)
+	if (bHasBuiltGeometryOnGT || RemoveImplicitsInDynamicCollections(Parameters))
 	{
 		// The Implicits attributes from the Dynamic Collection are just used for initialization, after they can be removed and so free some memory.
 		GameThreadCollection.RemoveAttribute(FGeometryDynamicCollection::ImplicitsAttribute, FTransformCollection::TransformGroup);
 	}
 	// If GT particle have been created, we can use directly the UniqueIdx from the GTParticle, then free up some space.
-	if ((bBuildGeometryForChildrenOnGT || bCreateGTParticleForChildren) && bEmptyUniqueIndicesOnGT)
+	if ((BuildGeometryForChildrenOnGT(Parameters) || CreateGTParticleForChildren(Parameters)) && bEmptyUniqueIndicesOnGT)
 	{
 		UniqueIdxs.Empty();
 	}
@@ -686,6 +729,8 @@ void FGeometryCollectionPhysicsProxy::Initialize(Chaos::FPBDRigidsEvolutionBase
 
 void FGeometryCollectionPhysicsProxy::CreateGTParticles(TManagedArray<Chaos::FImplicitObjectPtr>& Implicits, Chaos::FPBDRigidsEvolutionBase* Evolution, bool bInitializationTime)
 {
+	using namespace GeometryCollectionPhysicsProxy::Private;
+
 	const Chaos::Facades::FCollectionAnchoringFacade DynamicCollectionAnchoringFacade(GameThreadCollection);
 	const FVector Scale = Parameters.WorldTransform.GetScale3D();
 	const TManagedArray<float>& Mass = Parameters.RestCollectionShared->GetAttribute<float>(MassAttributeName, FTransformCollection::TransformGroup);
@@ -693,7 +738,7 @@ void FGeometryCollectionPhysicsProxy::CreateGTParticles(TManagedArray<Chaos::FIm
 	const TManagedArray<int32>& Level = Parameters.RestCollectionShared->GetAttribute<int32>(LevelAttributeName, FTransformCollection::TransformGroup);
 
 	TArray<int32> ChildrenToCheckForParentFix;
-	if (!bInitializationTime && !bCreateGTParticleForChildren && !bBuildGeometryForChildrenOnGT)
+	if (!bInitializationTime && !CreateGTParticleForChildren(Parameters) && !BuildGeometryForChildrenOnGT(Parameters))
 	{
 		GTParticleToIndices.Reserve(NumEffectiveParticles);
 	}
@@ -710,8 +755,8 @@ void FGeometryCollectionPhysicsProxy::CreateGTParticles(TManagedArray<Chaos::FIm
 		}
 
 		if ((bInitializationTime && TransformIndex == Parameters.InitialRootIndex) || // When initializing always create particle for the root
-			((bInitializationTime && (bCreateGTParticleForChildren || bBuildGeometryForChildrenOnGT)) || // When initializing create all particles if one of the flag is true 
-				(!bInitializationTime && TransformIndex != Parameters.InitialRootIndex && !bCreateGTParticleForChildren && !bBuildGeometryForChildrenOnGT))) // When not initializing create other particles if flag was set to not create
+			((bInitializationTime && (CreateGTParticleForChildren(Parameters) || BuildGeometryForChildrenOnGT(Parameters))) || // When initializing create all particles if one of the flag is true 
+				(!bInitializationTime && TransformIndex != Parameters.InitialRootIndex && !CreateGTParticleForChildren(Parameters) && !BuildGeometryForChildrenOnGT(Parameters)))) // When not initializing create other particles if flag was set to not create
 		{
 			GTParticles[ParticleIndex] = FParticle::CreateParticle();
 			P = GTParticles[ParticleIndex].Get();
@@ -744,7 +789,7 @@ void FGeometryCollectionPhysicsProxy::CreateGTParticles(TManagedArray<Chaos::FIm
 			Chaos::FImplicitObjectPtr ImplicitGeometry = MakeImplicitObjectPtr<Chaos::FImplicitObjectUnion>(MoveTemp(Geoms));
 			P->SetGeometry(ImplicitGeometry);
 		}
-		else if (bInitializationTime == (TransformIndex == Parameters.InitialRootIndex) || bBuildGeometryForChildrenOnGT)
+		else if (bInitializationTime == (TransformIndex == Parameters.InitialRootIndex) || BuildGeometryForChildrenOnGT(Parameters))
 		{
 			Chaos::FImplicitObjectPtr ImplicitGeometry = Implicits[TransformIndex];
 			if (ImplicitGeometry && !Scale.Equals(FVector::OneVector))
@@ -759,7 +804,7 @@ void FGeometryCollectionPhysicsProxy::CreateGTParticles(TManagedArray<Chaos::FIm
 			Implicits[TransformIndex] = P->GetGeometry();
 		}
 
-		if (bInitializationTime == (TransformIndex == Parameters.InitialRootIndex) || bBuildGeometryForChildrenOnGT)
+		if (bInitializationTime == (TransformIndex == Parameters.InitialRootIndex) || BuildGeometryForChildrenOnGT(Parameters))
 		{
 			if (DynamicCollectionAnchoringFacade.IsAnchored(TransformIndex))
 			{
@@ -897,11 +942,13 @@ void FGeometryCollectionPhysicsProxy::CreateGTParticles(TManagedArray<Chaos::FIm
 
 void FGeometryCollectionPhysicsProxy::CreateChildrenGeometry_External()
 {
+	using namespace GeometryCollectionPhysicsProxy::Private;
+
 	if (!bHasBuiltGeometryOnGT)
 	{
 		if (Chaos::FPhysicsSolver* RBDSolver = GetSolver<Chaos::FPhysicsSolver>())
 		{
-			if (bRemoveImplicitsInDynamicCollections)
+			if (RemoveImplicitsInDynamicCollections(Parameters))
 			{
 				GameThreadCollection.AddAttribute<Chaos::FImplicitObjectPtr>(FGeometryDynamicCollection::ImplicitsAttribute, FTransformCollection::TransformGroup);
 				GameThreadCollection.CopyAttribute(*Parameters.RestCollectionShared, FGeometryDynamicCollection::ImplicitsAttribute, FTransformCollection::TransformGroup);
@@ -1159,12 +1206,14 @@ void FGeometryCollectionPhysicsProxy::UpdateDamageThreshold_Internal()
 
 void FGeometryCollectionPhysicsProxy::InitializeBodiesPT(Chaos::FPBDRigidsSolver* RigidsSolver, typename Chaos::FPBDRigidsSolver::FParticlesType& Particles)
 {
+	using namespace GeometryCollectionPhysicsProxy::Private;
+
 	const FGeometryCollection* RestCollection = Parameters.RestCollectionShared.Get();
 	const FGeometryDynamicCollection& DynamicCollection = PhysicsThreadCollection;
 
 	bIsInitializedOnPhysicsThread = true;
 	// Building geometry according to the cVar
-	bHasBuiltGeometryOnPT = bBuildGeometryForChildrenOnPT;
+	bHasBuiltGeometryOnPT = BuildGeometryForChildrenOnPT(Parameters);
 
 	if (Parameters.Simulating && RestCollection)
 	{
@@ -1239,7 +1288,7 @@ void FGeometryCollectionPhysicsProxy::InitializeBodiesPT(Chaos::FPBDRigidsSolver
 					Handle,
 					Parameters.Shared,
 					Simplicials? (*Simplicials)[TransformGroupIndex].Get(): nullptr,
-					(bBuildGeometryForChildrenOnPT || (TransformGroupIndex == Parameters.InitialRootIndex)) ? Implicits[TransformGroupIndex] : nullptr,
+					(BuildGeometryForChildrenOnPT(Parameters) || (TransformGroupIndex == Parameters.InitialRootIndex)) ? Implicits[TransformGroupIndex] : nullptr,
 					SimFilter,
 					QueryFilter,
 					ScaledMass,
@@ -1627,7 +1676,7 @@ void FGeometryCollectionPhysicsProxy::InitializeBodiesPT(Chaos::FPBDRigidsSolver
 		// call DirtyParticle to make sure the acceleration structure is up to date with all the changes happening here
 		DirtyAllParticles(*RigidsSolver);
 
-		if (bHasBuiltGeometryOnPT || bRemoveImplicitsInDynamicCollections)
+		if (bHasBuiltGeometryOnPT || RemoveImplicitsInDynamicCollections(Parameters))
 		{
 			// The Implicits attributes from the Dynamic Collection are just used for geometry initialization, after they can be removed and so free some memory.
 			PhysicsThreadCollection.RemoveAttribute(FGeometryDynamicCollection::ImplicitsAttribute, FTransformCollection::TransformGroup);
@@ -1638,13 +1687,15 @@ void FGeometryCollectionPhysicsProxy::InitializeBodiesPT(Chaos::FPBDRigidsSolver
 
 void FGeometryCollectionPhysicsProxy::CreateChildrenGeometry_Internal()
 {
+	using namespace GeometryCollectionPhysicsProxy::Private;
+
 	if (!bHasBuiltGeometryOnPT)
 	{
 		const FGeometryCollection* RestCollection = Parameters.RestCollectionShared.Get();
 
 		if (Parameters.Simulating && ensure(RestCollection))
 		{
-			if (bRemoveImplicitsInDynamicCollections)
+			if (RemoveImplicitsInDynamicCollections(Parameters))
 			{
 				PhysicsThreadCollection.AddAttribute<Chaos::FImplicitObjectPtr>(FGeometryDynamicCollection::ImplicitsAttribute, FTransformCollection::TransformGroup);
 				PhysicsThreadCollection.CopyAttribute(*RestCollection, FGeometryDynamicCollection::ImplicitsAttribute, FTransformCollection::TransformGroup);
@@ -1856,6 +1907,8 @@ Chaos::FVec3f FGeometryCollectionPhysicsProxy::AdjustInertiaForScale(const Chaos
 
 Chaos::TPBDGeometryCollectionParticleHandle<Chaos::FReal, 3>* FGeometryCollectionPhysicsProxy::BuildNonClusters_Internal(const uint32 CollectionClusterIndex, Chaos::FPBDRigidsSolver* RigidsSolver, float Mass, Chaos::FVec3f Inertia, const Chaos::FUniqueIdx* ExistingIndex)
 {
+	using namespace GeometryCollectionPhysicsProxy::Private;
+
 	FGeometryDynamicCollection& DynamicCollection = PhysicsThreadCollection;
 	TManagedArray<uint8>& DynamicState = DynamicCollection.DynamicState;
 	const TManagedArray<FTransform>& MassToLocal = Parameters.RestCollectionShared->GetAttribute<FTransform>(MassToLocalAttributeName, FTransformCollection::TransformGroup);
@@ -1916,7 +1969,7 @@ Chaos::TPBDGeometryCollectionParticleHandle<Chaos::FReal, 3>* FGeometryCollectio
 		Handle,
 		Parameters.Shared,
 		Simplicials? (*Simplicials)[CollectionClusterIndex].Get(): nullptr,
-		(bBuildGeometryForChildrenOnPT || (CollectionClusterIndex == Parameters.InitialRootIndex)) ? Implicits[CollectionClusterIndex] : nullptr,
+		(BuildGeometryForChildrenOnPT(Parameters) || (CollectionClusterIndex == Parameters.InitialRootIndex)) ? Implicits[CollectionClusterIndex] : nullptr,
 		SimFilter,
 		QueryFilter,
 		Mass,
diff --git a/Engine/Source/Runtime/Experimental/Chaos/Public/ChaosSolverConfiguration.h b/Engine/Source/Runtime/Experimental/Chaos/Public/ChaosSolverConfiguration.h
--- a/Engine/Source/Runtime/Experimental/Chaos/Public/ChaosSolverConfiguration.h
+++ b/Engine/Source/Runtime/Experimental/Chaos/Public/ChaosSolverConfiguration.h
@@ -37,6 +37,14 @@ struct FChaosSolverDestructionSettings
 	/** Maximum number of active geometry collection to reach before all released clustering disable all released particle instantly */
 	UPROPERTY(EditAnywhere, Category = "SolverConfiguration|Destruction", meta = (ClampMin = "-1"))
 	int32 ClusteringParticleReleaseThrottlingMaxCount = INDEX_NONE;
+
+	/** 
+	* When enabled, this will avoid creating non-necessary geometry collection physics data until the root breaks 
+	* This will save runtime memory for unbroken geometry collection at the cost of allocating those particle when the first break happens
+	* It is false by default 
+	*/
+	UPROPERTY(EditAnywhere, Category = "SolverConfiguration|Destruction", meta = (ClampMin = "-1"))
+	bool bOptimizeForRuntimeMemory = false;
 };
 
 USTRUCT()
diff --git a/Engine/Source/Runtime/Experimental/Chaos/Public/GeometryCollection/GeometryCollectionSimulationCoreTypes.h b/Engine/Source/Runtime/Experimental/Chaos/Public/GeometryCollection/GeometryCollectionSimulationCoreTypes.h
--- a/Engine/Source/Runtime/Experimental/Chaos/Public/GeometryCollection/GeometryCollectionSimulationCoreTypes.h
+++ b/Engine/Source/Runtime/Experimental/Chaos/Public/GeometryCollection/GeometryCollectionSimulationCoreTypes.h
@@ -217,6 +217,7 @@ PRAGMA_ENABLE_DEPRECATION_WARNINGS
 		, bOptimizeConvexes(true)
 		, bUseSimplicialsWhenAvailable(false)
 		, bUseDamagePropagation(false)
+		, bOptimizeForRuntimeMemory(false)
 		, PositionSolverIterations(8)
 		, VelocitySolverIterations(1)
 		, ProjectionSolverIterations(1)
@@ -294,6 +295,7 @@ PRAGMA_ENABLE_DEPRECATION_WARNINGS
 		, bOptimizeConvexes(Other.bOptimizeConvexes)
 		, bUseSimplicialsWhenAvailable(Other.bUseSimplicialsWhenAvailable)
 		, bUseDamagePropagation(Other.bUseDamagePropagation)
+		, bOptimizeForRuntimeMemory(Other.bOptimizeForRuntimeMemory)
 		, PositionSolverIterations(Other.PositionSolverIterations)
 		, VelocitySolverIterations(Other.VelocitySolverIterations)
 		, ProjectionSolverIterations(Other.ProjectionSolverIterations)
@@ -406,6 +408,7 @@ PRAGMA_ENABLE_DEPRECATION_WARNINGS
 	bool bUseSimplicialsWhenAvailable : 1;
 
 	bool bUseDamagePropagation : 1;
+	bool bOptimizeForRuntimeMemory : 1;
 
 	uint8 PositionSolverIterations;
 	uint8 VelocitySolverIterations;
diff --git a/Engine/Source/Runtime/Experimental/GeometryCollectionEngine/GeometryCollectionEngine.Build.cs b/Engine/Source/Runtime/Experimental/GeometryCollectionEngine/GeometryCollectionEngine.Build.cs
--- a/Engine/Source/Runtime/Experimental/GeometryCollectionEngine/GeometryCollectionEngine.Build.cs
+++ b/Engine/Source/Runtime/Experimental/GeometryCollectionEngine/GeometryCollectionEngine.Build.cs
@@ -35,6 +35,7 @@ namespace UnrealBuildTool.Rules
 				{
 					"MeshConversion",
 					"GeometryCore",
+					"PhysicsCore",
 				}
 				);
 
