ChaosFlesh:
-Supported volume (sliding) constraint including FCreateAirVolumeConstraintDataflowNode (experimental), constraint in DeformableSolver, Gauss-Seidel residual and hessian computation
-Supported fixed length spring constraint (isotropic and anisotropic), added residual computation
-Fixed solver crash when using XPBD and logging cvar.

diff --git a/Engine/Plugins/Experimental/ChaosFlesh/Source/ChaosFlesh/Private/ChaosFlesh/TetrahedralCollection.cpp b/Engine/Plugins/Experimental/ChaosFlesh/Source/ChaosFlesh/Private/ChaosFlesh/TetrahedralCollection.cpp
--- a/Engine/Plugins/Experimental/ChaosFlesh/Source/ChaosFlesh/Private/ChaosFlesh/TetrahedralCollection.cpp
+++ b/Engine/Plugins/Experimental/ChaosFlesh/Source/ChaosFlesh/Private/ChaosFlesh/TetrahedralCollection.cpp
@@ -2,6 +2,7 @@
 #include "ChaosFlesh/TetrahedralCollection.h"
 #include "GeometryCollection/Facades/CollectionPositionTargetFacade.h"
 #include "GeometryCollection/Facades/CollectionMuscleActivationFacade.h"
+#include "GeometryCollection/Facades/CollectionVolumeConstraintFacade.h"
 
 DEFINE_LOG_CATEGORY_STATIC(FTetrahedralCollectionLogging, Log, All);
 
@@ -372,6 +373,10 @@ void FTetrahedralCollection::RemoveVertices(const TArray<int32>& SortedVertexInd
 	GeometryCollection::Facades::FPositionTargetFacade PositionTargetFacade(*this);
 	PositionTargetFacade.RemoveInvalidPositionTarget();
 
+	// --- Air Tetrahedral Constraints (vertex group dependency)--
+	GeometryCollection::Facades::FVolumeConstraintFacade VolumeConstraintFacade(*this);
+	VolumeConstraintFacade.RemoveInvalidVolumeConstraint();
+
 	// --- Muscle Activation (Geometry group dependency)--
 	GeometryCollection::Facades::FMuscleActivationFacade MuscleActivationFacade(*this);
 	MuscleActivationFacade.RemoveInvalidMuscles();
diff --git a/Engine/Plugins/Experimental/ChaosFlesh/Source/ChaosFleshNodes/Private/Dataflow/ChaosFleshNodesPlugin.cpp b/Engine/Plugins/Experimental/ChaosFlesh/Source/ChaosFleshNodes/Private/Dataflow/ChaosFleshNodesPlugin.cpp
--- a/Engine/Plugins/Experimental/ChaosFlesh/Source/ChaosFleshNodes/Private/Dataflow/ChaosFleshNodesPlugin.cpp
+++ b/Engine/Plugins/Experimental/ChaosFlesh/Source/ChaosFleshNodes/Private/Dataflow/ChaosFleshNodesPlugin.cpp
@@ -50,6 +50,7 @@ void IChaosFleshNodesPlugin::StartupModule()
 	DATAFLOW_NODE_REGISTER_CREATION_FACTORY(FComputeMuscleActivationDataNode);
 	DATAFLOW_NODE_REGISTER_CREATION_FACTORY(FComputeMuscleActivationDataNode_v2);
 	DATAFLOW_NODE_REGISTER_CREATION_FACTORY(FCreateAirTetrahedralConstraintDataflowNode);
+	DATAFLOW_NODE_REGISTER_CREATION_FACTORY(FCreateAirVolumeConstraintDataflowNode);
 	DATAFLOW_NODE_REGISTER_CREATION_FACTORY(FCreateTetrahedronDataflowNode);
 	DATAFLOW_NODE_REGISTER_CREATION_FACTORY(FCurveSamplingAnimationAssetTerminalNode);
 	DATAFLOW_NODE_REGISTER_CREATION_FACTORY(FDeleteFleshVerticesDataflowNode);
diff --git a/Engine/Plugins/Experimental/ChaosFlesh/Source/ChaosFleshNodes/Private/Dataflow/ChaosFleshSetVertexTrianglePositionTargetBindingNode.cpp b/Engine/Plugins/Experimental/ChaosFlesh/Source/ChaosFleshNodes/Private/Dataflow/ChaosFleshSetVertexTrianglePositionTargetBindingNode.cpp
--- a/Engine/Plugins/Experimental/ChaosFlesh/Source/ChaosFleshNodes/Private/Dataflow/ChaosFleshSetVertexTrianglePositionTargetBindingNode.cpp
+++ b/Engine/Plugins/Experimental/ChaosFlesh/Source/ChaosFleshNodes/Private/Dataflow/ChaosFleshSetVertexTrianglePositionTargetBindingNode.cpp
@@ -2,22 +2,73 @@
 
 #include "Dataflow/ChaosFleshSetVertexTrianglePositionTargetBindingNode.h"
 
-#include "Chaos/Utilities.h"
 #include "Chaos/HierarchicalSpatialHash.h"
 #include "Chaos/TriangleCollisionPoint.h"
+#include "Chaos/Utilities.h"
 #include "ChaosFlesh/ChaosFlesh.h"
 #include "ChaosFlesh/ChaosFleshCollectionFacade.h"
 #include "ChaosFlesh/TetrahedralCollection.h"
+#include "Dataflow/ChaosFleshNodesUtility.h"
+#include "DynamicMesh/InfoTypes.h"
+#include "DynamicMesh/MeshNormals.h"
+#include "GeometryCollection/Facades/CollectionCollisionFacade.h"
 #include "GeometryCollection/Facades/CollectionMeshFacade.h"
 #include "GeometryCollection/Facades/CollectionPositionTargetFacade.h"
-#include "GeometryCollection/Facades/CollectionCollisionFacade.h"
+#include "GeometryCollection/Facades/CollectionVolumeConstraintFacade.h"
+
 #include "UDynamicMesh.h"
-#include "DynamicMesh/InfoTypes.h"
-#include "DynamicMesh/MeshNormals.h"
-#include "Dataflow/ChaosFleshNodesUtility.h"
 
 #include UE_INLINE_GENERATED_CPP_BY_NAME(ChaosFleshSetVertexTrianglePositionTargetBindingNode)
 
+static TArray<Chaos::TVector<int32, 3>> RemoveInvalidIndices(const TManagedArray<FIntVector>* Indices)
+{
+	TArray<Chaos::TVector<int32, 3>> IndicesArray;
+	for (int32 FaceIdx = 0; FaceIdx < Indices->Num(); ++FaceIdx)
+	{
+		Chaos::TVector<int32, 3> CurrentIndices(0);
+		for (int32 LocalIdx = 0; LocalIdx < 3; ++LocalIdx)
+		{
+			CurrentIndices[LocalIdx] = (*Indices)[FaceIdx][LocalIdx];
+		}
+		if (CurrentIndices[0] != INDEX_NONE
+			&& CurrentIndices[1] != INDEX_NONE
+			&& CurrentIndices[2] != INDEX_NONE)
+		{
+			IndicesArray.Add(CurrentIndices);
+		}
+	}
+	return IndicesArray;
+}
+
+static TArray<int32> ComputeBoundaryVertices(const TArray<Chaos::TVector<int32, 3>>& IndicesArray)
+{
+	TArray<TArray<int32>> LocalIndex;
+	TArray<TArray<int32>>* LocalIndexPtr = &LocalIndex;
+	TArray<TArray<int>> IncidentElements = Chaos::Utilities::ComputeIncidentElements(IndicesArray, LocalIndexPtr);
+	TArray<int32> BoundaryVertices;
+	for (int32 VertIdx = 0; VertIdx < IncidentElements.Num(); ++VertIdx)
+	{
+		if (IncidentElements[VertIdx].Num() > 0)
+		{
+			BoundaryVertices.Add(VertIdx);
+		}
+	}
+	return BoundaryVertices;
+}
+
+static void FilterBoundaryVertices(TArray<int32>& OutBoundaryVertices, const FDataflowVertexSelection& InDataflowVertexSelection)
+{
+	TArray<int32> BoundaryVerticesNew;
+	for (int32 VertIdx : OutBoundaryVertices)
+	{
+		if (InDataflowVertexSelection.IsSelected(VertIdx))
+		{
+			BoundaryVerticesNew.Add(VertIdx);
+		}
+	}
+	OutBoundaryVertices = BoundaryVerticesNew;
+}
+
 void FSetVertexTrianglePositionTargetBindingDataflowNode::Evaluate(UE::Dataflow::FContext& Context, const FDataflowOutput* Out) const
 {
 	if (Out->IsA(&Collection))
@@ -36,74 +87,22 @@ void FSetVertexTrianglePositionTargetBindingDataflowNode::Evaluate(UE::Dataflow:
 					TArray<FVector3f> Vertex = TetCollection.Vertex.Get().GetConstArray();
 					TetCollection.ComponentSpaceVertices(Vertex);
 					GeometryCollection::Facades::FCollectionMeshFacade MeshFacade(InCollection);
-					TArray<int32> ComponentIndex = MeshFacade.GetGeometryGroupIndexArray();
-					TArray<Chaos::TVector<float, 3>> IndicesPositions; 
-					TArray<Chaos::TVector<int32, 3>> IndicesArray;
-					for (int32 i = 0; i < Indices->Num(); i++)
-					{
-						Chaos::TVector<int32, 3> CurrentIndices(0);
-						for (int32 j = 0; j < 3; j++) 
-						{
-							CurrentIndices[j] = (*Indices)[i][j];
-						}
-						if (CurrentIndices[0] != INDEX_NONE
-							&& CurrentIndices[1] != INDEX_NONE
-							&& CurrentIndices[2] != INDEX_NONE)
-						{
-							IndicesArray.Emplace(CurrentIndices);
-						}
-					}
-					TArray<TArray<int32>> LocalIndex;
-					TArray<TArray<int32>>* LocalIndexPtr = &LocalIndex;
-					TArray<TArray<int>> GlobalIndex = Chaos::Utilities::ComputeIncidentElements(IndicesArray, LocalIndexPtr);
-					int32 ActualParticleCount = 0;
-					for (int32 l = 0; l < GlobalIndex.Num(); l++)
-					{
-						if (GlobalIndex[l].Num() > 0)
-						{
-							ActualParticleCount += 1;
-						}
-					}
-					
-					IndicesPositions.SetNum(ActualParticleCount);
-					TArray<int32> IndicesMap;
-					IndicesMap.SetNum(ActualParticleCount);
-					int32 CurrentParticleIndex = 0;
-					for (int32 i = 0; i < GlobalIndex.Num(); i++)
-					{
-						if (GlobalIndex[i].Num() > 0)
-						{
-							IndicesPositions[CurrentParticleIndex] = (*Vertices)[(*Indices)[GlobalIndex[i][0]][LocalIndex[i][0]]];
-							IndicesMap[CurrentParticleIndex] = (*Indices)[GlobalIndex[i][0]][LocalIndex[i][0]];
-							CurrentParticleIndex += 1;
-						}
-					}
-
-					//Compute max bound and determine detection radius
-					Chaos::TVec3<float> CoordMaxs(-FLT_MAX);
-					Chaos::TVec3<float> CoordMins(FLT_MAX);
-					for (int32 i = 0; i < Vertex.Num(); i++)
-					{
-						for (int32 j = 0; j < 3; j++)
-						{
-							if (Vertex[i][j] > CoordMaxs[j])
-							{
-								CoordMaxs[j] = Vertex[i][j];
-							}
-							if (Vertex[i][j] < CoordMins[j])
-							{
-								CoordMins[j] = Vertex[i][j];
-							}
-						}
-					}
-					Chaos::TVec3<float> CoordDiff = (CoordMaxs - CoordMins) * VertexRadiusRatio;
-					Chaos::FReal SphereRadius = Chaos::FReal(FGenericPlatformMath::Max(CoordDiff[0], FGenericPlatformMath::Max(CoordDiff[1], CoordDiff[2])));
-					
+					const TArray<int32> ComponentIndex = MeshFacade.GetGeometryGroupIndexArray();
+					const TArray<Chaos::TVector<int32, 3>> IndicesArray = RemoveInvalidIndices(Indices);
+					TArray<int32> BoundaryVertices = ComputeBoundaryVertices(IndicesArray);		
+					// Only keep boundary vertices within VertexSelection
 					if (IsConnected(&VertexSelection))
 					{
 						FDataflowVertexSelection InDataflowVertexSelection = GetValue<FDataflowVertexSelection>(Context, &VertexSelection);
-						IndicesMap = InDataflowVertexSelection.AsArray();
-
+						if (InDataflowVertexSelection.Num() != Vertices->Num())
+						{
+							Context.Error(FString::Printf(
+								TEXT("VertexSelection size [%d] is not equal to the collection's vertex count [%d]"),
+								InDataflowVertexSelection.Num(), Vertices->Num()),
+								this, Out);
+							return;
+						}
+						FilterBoundaryVertices(BoundaryVertices, InDataflowVertexSelection);
 					}
 
 					GeometryCollection::Facades::FPositionTargetFacade PositionTargets(InCollection);
@@ -121,12 +120,12 @@ void FSetVertexTrianglePositionTargetBindingDataflowNode::Evaluate(UE::Dataflow:
 					TriangleMesh.Init(IndicesArray);
 					Chaos::FTriangleMesh::TSpatialHashType<Chaos::FReal> SpatialHash;
 					TConstArrayView<Chaos::TVec3<Chaos::FReal>> ConstArrayViewVertex(VertexTVec3);
-					TriangleMesh.BuildSpatialHash(ConstArrayViewVertex, SpatialHash, SphereRadius);
-					for (int32 PointIndex: IndicesMap)
+					TriangleMesh.BuildSpatialHash(ConstArrayViewVertex, SpatialHash, Chaos::FReal(SearchRadius));
+					for (int32 PointIndex: BoundaryVertices)
 					{
 						TArray<Chaos::TTriangleCollisionPoint<Chaos::FReal>> Result;
 						if (TriangleMesh.PointClosestTriangleQuery(SpatialHash, ConstArrayViewVertex,
-							PointIndex, VertexTVec3[PointIndex], SphereRadius / 2.f, SphereRadius / 2.f,
+							PointIndex, VertexTVec3[PointIndex], Chaos::FReal(SearchRadius) / 2.f, Chaos::FReal(SearchRadius) / 2.f,
 							[this, &ComponentIndex, &IndicesArray](const int32 PointIndex, const int32 TriangleIndex)->bool
 							{
 								return ComponentIndex[PointIndex] != ComponentIndex[IndicesArray[TriangleIndex][0]];
@@ -166,6 +165,8 @@ void FSetVertexTrianglePositionTargetBindingDataflowNode::Evaluate(UE::Dataflow:
 								{
 									DataPackage.Stiffness = PositionTargetStiffness;
 								}
+								DataPackage.bIsAnisotropic = bAllowSliding;
+								DataPackage.bIsZeroRestLength = bUseZeroRestLengthSprings;
 								PositionTargets.AddPositionTarget(DataPackage);
 							}
 						}
@@ -174,7 +175,7 @@ void FSetVertexTrianglePositionTargetBindingDataflowNode::Evaluate(UE::Dataflow:
 					MassRatioArray.Sort();
 					if (MassRatioArray.Num())
 					{
-						UE_LOG(LogChaosFlesh, Warning, TEXT("SetVertexTrianglePositionTargetBinding: Max mass ratio = %f, median mass ratio = %f")
+						UE_LOG(LogChaosFlesh, Display, TEXT("SetVertexTrianglePositionTargetBinding: Max mass ratio = %f, median mass ratio = %f")
 							, MassRatioArray[MassRatioArray.Num()-1], MassRatioArray[MassRatioArray.Num()/2]);
 					}
 #endif
@@ -201,6 +202,7 @@ void FDeleteVertexTrianglePositionTargetBindingDataflowNode::Evaluate(UE::Datafl
 					TEXT("VertexSelection1 size (%d) is not equal to the collection's vertex count (%d)"),
 					InVertexSelection1.Num(), VertexCount),
 					this, Out);
+				return;
 			}
 			if (InVertexSelection2.Num() != VertexCount)
 			{
@@ -208,9 +210,10 @@ void FDeleteVertexTrianglePositionTargetBindingDataflowNode::Evaluate(UE::Datafl
 					TEXT("VertexSelection2 size (%d) is not equal to the collection's vertex count (%d)"),
 					InVertexSelection2.Num(), VertexCount),
 					this, Out);
+				return;
 			}
 			GeometryCollection::Facades::FPositionTargetFacade PositionTargets(InCollection);
-			const int32 NumRemoved = PositionTargets.RemovePositionTarget(
+			const int32 NumRemoved = PositionTargets.RemovePositionTargetBetween(
 				[&InVertexSelection1](const int32 VertexIdx) { return InVertexSelection1.IsSelected(VertexIdx); },
 				[&InVertexSelection2](const int32 VertexIdx) { return InVertexSelection2.IsSelected(VertexIdx); });
 			Context.Info(FString::Printf(
@@ -252,52 +255,156 @@ void FCreateAirTetrahedralConstraintDataflowNode::Evaluate(UE::Dataflow::FContex
 				TArray<FVector3f> Vertex = Vertices->GetConstArray();
 				TetCollection.ComponentSpaceVertices(Vertex);
 				GeometryCollection::Facades::FCollectionMeshFacade MeshFacade(*InFleshCollection);
-				TArray<int32> ComponentIndex = MeshFacade.GetGeometryGroupIndexArray();
-				TArray<Chaos::TVector<int32, 3>> IndicesArray;
-				for (int32 i = 0; i < Indices->Num(); i++)
+				const TArray<int32> ComponentIndex = MeshFacade.GetGeometryGroupIndexArray();
+				const TArray<Chaos::TVector<int32, 3>> IndicesArray = RemoveInvalidIndices(Indices);
+				TArray<int32> BoundaryVertices = ComputeBoundaryVertices(IndicesArray);
+				// Only keep boundary vertices within VertexSelection
+				if (IsConnected(&VertexSelection))
 				{
-					Chaos::TVector<int32, 3> CurrentIndices(0);
-					for (int32 j = 0; j < 3; j++)
-					{
-						CurrentIndices[j] = (*Indices)[i][j];
-					}
-					if (CurrentIndices[0] != INDEX_NONE
-						&& CurrentIndices[1] != INDEX_NONE
-						&& CurrentIndices[2] != INDEX_NONE)
+					FDataflowVertexSelection InDataflowVertexSelection = GetValue<FDataflowVertexSelection>(Context, &VertexSelection);
+					if (InDataflowVertexSelection.Num() != Vertices->Num())
 					{
-						IndicesArray.Emplace(CurrentIndices);
+						Context.Error(FString::Printf(
+							TEXT("VertexSelection size [%d] is not equal to the collection's vertex count [%d]"),
+							InDataflowVertexSelection.Num(), Vertices->Num()),
+							this, Out);
+						return;
 					}
+					FilterBoundaryVertices(BoundaryVertices, InDataflowVertexSelection);
 				}
 
-				TArray<int32> IndicesMap;
-				if (IsConnected(&VertexSelection))
+				TArray<Chaos::TVec3<Chaos::FReal>> VertexTVec3;
+				VertexTVec3.SetNum(Vertex.Num());
+				for (int32 VertexIdx = 0; VertexIdx < Vertex.Num(); VertexIdx++)
 				{
-					FDataflowVertexSelection InDataflowVertexSelection = GetValue<FDataflowVertexSelection>(Context, &VertexSelection);
-					IndicesMap = InDataflowVertexSelection.AsArray();
+					VertexTVec3[VertexIdx] = Chaos::TVec3<Chaos::FReal>(Vertex[VertexIdx]);
 				}
-				else
+				Chaos::FTriangleMesh TriangleMesh;
+				TriangleMesh.Init(IndicesArray);
+				Chaos::FTriangleMesh::TSpatialHashType<Chaos::FReal> SpatialHash;
+				TConstArrayView<Chaos::TVec3<Chaos::FReal>> ConstArrayViewVertex(VertexTVec3);
+				TriangleMesh.BuildSpatialHash(ConstArrayViewVertex, SpatialHash, Chaos::FReal(SearchRadius));
+				TArray<FIntVector4> NewTetrahedra;
+				auto ComputeDihedralAngle = [](const FVector3f& A, const FVector3f& B, const FVector3f& C, const FVector3f& D)
+					{
+						FVector3f Normal1 = FVector3f::CrossProduct(B - A, C - A).GetSafeNormal();
+						FVector3f Normal2 = FVector3f::CrossProduct(B - A, D - A).GetSafeNormal();
+
+						float CosTheta = FVector3f::DotProduct(Normal1, Normal2);
+						return FMath::RadiansToDegrees(FMath::Acos(CosTheta));
+					};
+				auto IfTetInverted = [](const FVector3f& A, const FVector3f& B, const FVector3f& C, const FVector3f& D)
+					{
+						return (D - A).Dot(FVector3f::CrossProduct(B - A, C - A)) < 0;
+					};
+				int32 NumSkippedDihedral = 0;
+				for (int32 PointIndex : BoundaryVertices)
 				{
-					TArray<TArray<int32>> LocalIndex;
-					TArray<TArray<int32>>* LocalIndexPtr = &LocalIndex;
-					TArray<TArray<int>> GlobalIndex = Chaos::Utilities::ComputeIncidentElements(IndicesArray, LocalIndexPtr);
-					int32 ActualParticleCount = 0;
-					for (int32 l = 0; l < GlobalIndex.Num(); l++)
+					TArray<Chaos::TTriangleCollisionPoint<Chaos::FReal>> Result;
+					if (TriangleMesh.PointClosestTriangleQuery(SpatialHash, ConstArrayViewVertex,
+						PointIndex, VertexTVec3[PointIndex], Chaos::FReal(SearchRadius) / 2.f, Chaos::FReal(SearchRadius) / 2.f,
+						[this, &ComponentIndex, &IndicesArray](const int32 PointIndex, const int32 TriangleIndex)->bool
+						{
+							return ComponentIndex[PointIndex] != ComponentIndex[IndicesArray[TriangleIndex][0]];
+						},
+						Result))
 					{
-						if (GlobalIndex[l].Num() > 0)
+						for (const Chaos::TTriangleCollisionPoint<Chaos::FReal>& CollisionPoint : Result)
 						{
-							ActualParticleCount += 1;
+							int32 Tri0 = IndicesArray[CollisionPoint.Indices[1]][0];
+							int32 Tri1 = IndicesArray[CollisionPoint.Indices[1]][1];
+							int32 Tri2 = IndicesArray[CollisionPoint.Indices[1]][2];
+							float DihedralAngle = ComputeDihedralAngle(Vertex[PointIndex], Vertex[Tri0], Vertex[Tri1], Vertex[Tri2]);
+							if (DihedralAngle < 10 || DihedralAngle > 170)
+							{
+								NumSkippedDihedral++;
+								continue;
+							}
+							if (IfTetInverted(Vertex[PointIndex], Vertex[Tri0], Vertex[Tri1], Vertex[Tri2]))
+							{
+								Swap(Tri0, Tri1);
+							}
+							NewTetrahedra.Add(FIntVector4(PointIndex,
+								Tri0,
+								Tri1,
+								Tri2));
+							break;					
 						}
 					}
-					IndicesMap.SetNum(ActualParticleCount);
-					int32 CurrentParticleIndex = 0;
-					for (int32 i = 0; i < GlobalIndex.Num(); i++)
+				}
+				int32 StartSize = InFleshCollection->AddElements(NewTetrahedra.Num(), FTetrahedralCollection::TetrahedralGroup);
+				for (int32 Idx = 0; Idx < NewTetrahedra.Num(); ++Idx)
+				{
+					Tetrahedron[StartSize + Idx] = NewTetrahedra[Idx];
+				}
+				InFleshCollection->InitIncidentElements(); //recompute incident elements after appending new tet constraints
+				UE_LOG(LogChaosFlesh, Display, TEXT("CreateAirTetrahedralConstraint: Added %d volumetric constraints.")
+					, NewTetrahedra.Num());
+				UE_LOG(LogChaosFlesh, Display, TEXT("CreateAirTetrahedralConstraint: Skipped %d volumetric constraints due to extreme dihedral angles.")
+					, NumSkippedDihedral);
+				//Draw dynamic mesh of new tet boundary
+				auto FloatVert = [](FVector3d V) { return FVector3f(V.X, V.Y, V.Z); };
+				auto DoubleVert = [](FVector3f V) { return FVector3d(V.X, V.Y, V.Z); };
+				OutDynamicMesh->Reset();
+				FDynamicMesh3& OutDynamicMesh3 = OutDynamicMesh->GetMeshRef();
+				OutDynamicMesh3.EnableAttributes();
+
+				// Compute boundary triangle mesh
+				const TArray<FIntVector3> SurfaceElements = UE::Dataflow::GetSurfaceTriangles(NewTetrahedra, /*bKeepInterior = */ false);
+				for (int32 TriangleIndex = 0; TriangleIndex < SurfaceElements.Num(); ++TriangleIndex)
+				{
+					int v0 = OutDynamicMesh3.AppendVertex(DoubleVert(Vertex[SurfaceElements[TriangleIndex][0]]));
+					int v1 = OutDynamicMesh3.AppendVertex(DoubleVert(Vertex[SurfaceElements[TriangleIndex][1]]));
+					int v2 = OutDynamicMesh3.AppendVertex(DoubleVert(Vertex[SurfaceElements[TriangleIndex][2]]));
+					OutDynamicMesh3.AppendTriangle(FIndex3i(v0, v1, v2));
+				}
+				
+				// Compute normals
+				OutDynamicMesh3.EnableVertexNormals(FVector3f(1, 0, 0));
+				FMeshNormals MeshNormals(&OutDynamicMesh3);
+				MeshNormals.ComputeVertexNormals();
+				for (int32 VertexIndex = 0; VertexIndex < OutDynamicMesh3.VertexCount(); ++VertexIndex)
+				{
+					OutDynamicMesh3.SetVertexNormal(VertexIndex, FloatVert(MeshNormals[VertexIndex]));
+				}
+			} // if Indices
+		} // if Vertices
+		SetValue(Context, MoveTemp(static_cast<FManagedArrayCollection&>(*InFleshCollection.Release())), &Collection);
+		SetValue(Context, OutDynamicMesh, &DynamicMesh);
+	}
+}
+
+void FCreateAirVolumeConstraintDataflowNode::Evaluate(UE::Dataflow::FContext& Context, const FDataflowOutput* Out) const
+{
+	using namespace UE::Geometry;
+	if (Out->IsA(&Collection))
+	{
+		TUniquePtr<FFleshCollection> InFleshCollection(GetValue(Context, &Collection).NewCopy<FFleshCollection>());
+		Chaos::FFleshCollectionFacade TetCollection(*InFleshCollection);
+		TObjectPtr<UDynamicMesh> OutDynamicMesh = NewObject<UDynamicMesh>();
+		if (const TManagedArray<FVector3f>* Vertices = InFleshCollection->FindAttribute<FVector3f>("Vertex", FGeometryCollection::VerticesGroup))
+		{
+			if (const TManagedArray<FIntVector>* Indices = InFleshCollection->FindAttribute<FIntVector>("Indices", FGeometryCollection::FacesGroup))
+			{
+				TArray<FVector3f> Vertex = Vertices->GetConstArray();
+				TetCollection.ComponentSpaceVertices(Vertex);
+				GeometryCollection::Facades::FCollectionMeshFacade MeshFacade(*InFleshCollection);
+				const TArray<int32> ComponentIndex = MeshFacade.GetGeometryGroupIndexArray();
+				const TArray<Chaos::TVector<int32, 3>> IndicesArray = RemoveInvalidIndices(Indices);
+				TArray<int32> BoundaryVertices = ComputeBoundaryVertices(IndicesArray);
+				// Only keep boundary vertices within VertexSelection
+				if (IsConnected(&VertexSelection))
+				{
+					FDataflowVertexSelection InDataflowVertexSelection = GetValue<FDataflowVertexSelection>(Context, &VertexSelection);
+					if (InDataflowVertexSelection.Num() != Vertices->Num())
 					{
-						if (GlobalIndex[i].Num() > 0)
-						{
-							IndicesMap[CurrentParticleIndex] = (*Indices)[GlobalIndex[i][0]][LocalIndex[i][0]];
-							CurrentParticleIndex += 1;
-						}
+						Context.Error(FString::Printf(
+							TEXT("VertexSelection size [%d] is not equal to the collection's vertex count [%d]"),
+							InDataflowVertexSelection.Num(), Vertices->Num()),
+							this, Out);
+						return;
 					}
+					FilterBoundaryVertices(BoundaryVertices, InDataflowVertexSelection);
 				}
 
 				TArray<Chaos::TVec3<Chaos::FReal>> VertexTVec3;
@@ -325,7 +432,7 @@ void FCreateAirTetrahedralConstraintDataflowNode::Evaluate(UE::Dataflow::FContex
 						return (D - A).Dot(FVector3f::CrossProduct(B - A, C - A)) < 0;
 					};
 				int32 NumSkippedDihedral = 0;
-				for (int32 PointIndex : IndicesMap)
+				for (int32 PointIndex : BoundaryVertices)
 				{
 					TArray<Chaos::TTriangleCollisionPoint<Chaos::FReal>> Result;
 					if (TriangleMesh.PointClosestTriangleQuery(SpatialHash, ConstArrayViewVertex,
@@ -355,19 +462,20 @@ void FCreateAirTetrahedralConstraintDataflowNode::Evaluate(UE::Dataflow::FContex
 								Tri0,
 								Tri1,
 								Tri2));
-							break;					
+							break;
 						}
 					}
 				}
-				int32 StartSize = InFleshCollection->AddElements(NewTetrahedra.Num(), FTetrahedralCollection::TetrahedralGroup);
-				for (int32 Idx = 0; Idx < NewTetrahedra.Num(); ++Idx)
+				// Add volume constraints
+				GeometryCollection::Facades::FVolumeConstraintFacade VolumeConstraint(*InFleshCollection);
+				for (const FIntVector4& NewTetrahedron : NewTetrahedra)
 				{
-					Tetrahedron[StartSize + Idx] = NewTetrahedra[Idx];
+					VolumeConstraint.AddVolumeConstraint(NewTetrahedron, Stiffness);
 				}
-				InFleshCollection->InitIncidentElements(); //recompute incident elements after appending new tet constraints
-				UE_LOG(LogChaosFlesh, Warning, TEXT("CreateAirTetrahedralConstraint: Added %d volumetric constraints.")
+
+				UE_LOG(LogChaosFlesh, Display, TEXT("CreateAirVolumeConstraint: Added %d volumetric constraints.")
 					, NewTetrahedra.Num());
-				UE_LOG(LogChaosFlesh, Warning, TEXT("CreateAirTetrahedralConstraint: Skipped %d volumetric constraints due to extreme dihedral angles.")
+				UE_LOG(LogChaosFlesh, Display, TEXT("CreateAirVolumeConstraint: Skipped %d volumetric constraints due to extreme dihedral angles.")
 					, NumSkippedDihedral);
 				//Draw dynamic mesh of new tet boundary
 				auto FloatVert = [](FVector3d V) { return FVector3f(V.X, V.Y, V.Z); };
@@ -377,7 +485,7 @@ void FCreateAirTetrahedralConstraintDataflowNode::Evaluate(UE::Dataflow::FContex
 				OutDynamicMesh3.EnableAttributes();
 
 				// Compute boundary triangle mesh
-				TArray<FIntVector3> SurfaceElements = UE::Dataflow::GetSurfaceTriangles(NewTetrahedra, /*bKeepInterior = */ false);
+				const TArray<FIntVector3> SurfaceElements = UE::Dataflow::GetSurfaceTriangles(NewTetrahedra, /*bKeepInterior = */ false);
 				for (int32 TriangleIndex = 0; TriangleIndex < SurfaceElements.Num(); ++TriangleIndex)
 				{
 					int v0 = OutDynamicMesh3.AppendVertex(DoubleVert(Vertex[SurfaceElements[TriangleIndex][0]]));
@@ -385,7 +493,7 @@ void FCreateAirTetrahedralConstraintDataflowNode::Evaluate(UE::Dataflow::FContex
 					int v2 = OutDynamicMesh3.AppendVertex(DoubleVert(Vertex[SurfaceElements[TriangleIndex][2]]));
 					OutDynamicMesh3.AppendTriangle(FIndex3i(v0, v1, v2));
 				}
-				
+
 				// Compute normals
 				OutDynamicMesh3.EnableVertexNormals(FVector3f(1, 0, 0));
 				FMeshNormals MeshNormals(&OutDynamicMesh3);
diff --git a/Engine/Plugins/Experimental/ChaosFlesh/Source/ChaosFleshNodes/Public/Dataflow/ChaosFleshSetVertexTrianglePositionTargetBindingNode.h b/Engine/Plugins/Experimental/ChaosFlesh/Source/ChaosFleshNodes/Public/Dataflow/ChaosFleshSetVertexTrianglePositionTargetBindingNode.h
--- a/Engine/Plugins/Experimental/ChaosFlesh/Source/ChaosFleshNodes/Public/Dataflow/ChaosFleshSetVertexTrianglePositionTargetBindingNode.h
+++ b/Engine/Plugins/Experimental/ChaosFlesh/Source/ChaosFleshNodes/Public/Dataflow/ChaosFleshSetVertexTrianglePositionTargetBindingNode.h
@@ -12,7 +12,7 @@
 class UDynamicMesh;
 
 /**
- * Create point-triangle weak constraints (springs) between different geometries based on search radius.
+ * Create point-triangle weak constraints (springs) between surface meshes of different geometries based on search radius.
  */
 USTRUCT(meta = (DataflowFlesh))
 struct FSetVertexTrianglePositionTargetBindingDataflowNode : public FDataflowNode
@@ -26,15 +26,27 @@ public:
 	FManagedArrayCollection Collection;
 
 	UPROPERTY(EditAnywhere, Category = "Dataflow")
-	float PositionTargetStiffness = 10000.f;
+	float PositionTargetStiffness = 1.f;
 
-	/* (optional) only create weak constraints from vertices in VertexSelection to triangles in other geometries */
+	/** (optional) only create weak constraints from surface vertices in VertexSelection to triangles in other geometries */
 	UPROPERTY(meta = (DataflowInput, DisplayName = "(Optional) VertexSelection"))
 	FDataflowVertexSelection VertexSelection;
 	
-	/* weak constraint search radius = Largest side of the collection's bounding box * VertexRadiusRatio */
-	UPROPERTY(EditAnywhere, Category = "Dataflow", meta = (ClampMin = "0.0", ClampMax = "2.0", UIMin = "0.0", UIMax = "2.0"))
-	float VertexRadiusRatio = .001f;
+	/** Search radius for point-triangle pairs between geometry surfaces. */
+	UPROPERTY(EditAnywhere, Category = "Dataflow", meta = (ClampMin = "0.0"))
+	float SearchRadius = 0.f;
+
+	/** if point-triangle weak constraints created are anisotropic and allow sliding along the triangle plane */
+	UPROPERTY(EditAnywhere, Category = "Dataflow", meta = (DisplayName = "Allow Sliding"))
+	bool bAllowSliding = false;
+
+	UPROPERTY(EditAnywhere, Category = "Dataflow", meta = (InlineEditConditionToggle))
+	bool bZeroRestLengthEditable = false;
+
+	/** if point-triangle weak constraints created are zero rest-length.
+	* if true, this will cause point triangle pair to stick together, as opposed to separated by their rest state distance.*/
+	UPROPERTY(EditAnywhere, Category = "Dataflow", meta = (DisplayName = "Use zero rest length springs", EditCondition = "bZeroRestLengthEditable"))
+	bool bUseZeroRestLengthSprings = false;
 
 	FSetVertexTrianglePositionTargetBindingDataflowNode(const UE::Dataflow::FNodeParameters& InParam, FGuid InGuid = FGuid::NewGuid())
 		: FDataflowNode(InParam, InGuid)
@@ -71,11 +83,11 @@ private:
 	UPROPERTY(meta = (DataflowInput, DataflowIntrinsic, DataflowOutput, DisplayName = "Collection", DataflowPassthrough = "Collection"))
 	FManagedArrayCollection Collection;
 
-	/* This node deletes springs between VertexSelection1 and VertexSelection2. */
+	/** This node deletes springs between VertexSelection1 and VertexSelection2. */
 	UPROPERTY(meta = (DataflowInput, DataflowIntrinsic, DisplayName = "VertexSelection1"))
 	FDataflowVertexSelection VertexSelection1;
 
-	/* This node deletes springs between VertexSelection1 and VertexSelection2. */
+	/** This node deletes springs between VertexSelection1 and VertexSelection2. */
 	UPROPERTY(meta = (DataflowInput, DataflowIntrinsic, DisplayName = "VertexSelection2"))
 	FDataflowVertexSelection VertexSelection2;
 
@@ -106,7 +118,7 @@ private:
 	UPROPERTY(meta = (DataflowInput, DataflowIntrinsic, DataflowOutput, DisplayName = "Collection", DataflowPassthrough = "Collection"))
 	FManagedArrayCollection Collection;
 
-	/* Vertices selected to be able to collide with others. Unselected vertices will not collide with unselected vertices*/
+	/** Vertices selected to be able to collide with others. Unselected vertices will not collide with unselected vertices*/
 	UPROPERTY(meta = (DataflowInput, DataflowIntrinsic, DisplayName = "Collision Vertex Selection"))
 	FDataflowVertexSelection VertexSelection;
 
@@ -114,7 +126,7 @@ private:
 };
 
 /**
- * Create air tetrahedral constraint between point-triangle pair from different geometries based on search radius. 
+ * Create air tetrahedral constraint between point-triangle pair from surface meshes of different geometries based on search radius. 
  * The added tetrahedra help to maintain distance between geometries.
  * This node renders the boundary of the added tetrahedral mesh.
  */
@@ -141,14 +153,14 @@ private:
 	UPROPERTY(meta = (DataflowInput, DataflowOutput, DisplayName = "Collection", DataflowPassthrough = "Collection"))
 	FManagedArrayCollection Collection;
 
-	/* (optional) only create tetrahedral constraints from vertices in VertexSelection to triangles in other geometries.
+	/** (optional) only create tetrahedral constraints from surface vertices in VertexSelection to triangles in other geometries.
 	* For example, if the VertexSelection contains only one geometry, only this geometry will bind to other geometries.
 	* No constraints will be created between two geometries that are not in the VertexSelection.
 	*/
 	UPROPERTY(meta = (DataflowInput, DisplayName = "(Optional) VertexSelection"))
 	FDataflowVertexSelection VertexSelection;
 
-	/* tetrahedral constraint search radius*/
+	/** tetrahedral constraint search radius*/
 	UPROPERTY(EditAnywhere, Category = "Dataflow", meta = (ClampMin = "0.0"))
 	float SearchRadius = 0.f;
 
@@ -156,5 +168,53 @@ private:
 	UPROPERTY(meta = (DataflowOutput))
 	TObjectPtr<UDynamicMesh> DynamicMesh;
 
+	virtual void Evaluate(UE::Dataflow::FContext& Context, const FDataflowOutput* Out) const override;
+};
+
+/**
+ * Creates volume constraint (defined by point-triangle tetrahedron volume) between surface meshes of different geometries.
+ * This constraint allow sliding of the point along the triangle plane.
+ */
+USTRUCT(meta = (Experimental, DataflowFlesh))
+struct FCreateAirVolumeConstraintDataflowNode : public FDataflowNode
+{
+	GENERATED_USTRUCT_BODY()
+	DATAFLOW_NODE_DEFINE_INTERNAL(FCreateAirVolumeConstraintDataflowNode, "CreateAirVolumeConstraint", "Flesh", "")
+	DATAFLOW_NODE_RENDER_TYPE("SurfaceRender", FName("FDynamicMesh3"), "DynamicMesh")
+
+public:
+	FCreateAirVolumeConstraintDataflowNode(const UE::Dataflow::FNodeParameters& InParam, FGuid InGuid = FGuid::NewGuid())
+		: FDataflowNode(InParam, InGuid)
+	{
+		RegisterInputConnection(&Collection);
+		RegisterInputConnection(&VertexSelection);
+		RegisterOutputConnection(&Collection, &Collection);
+		RegisterOutputConnection(&DynamicMesh)
+			.SetCanHidePin(true)
+			.SetPinIsHidden(true);
+	}
+
+private:
+	UPROPERTY(meta = (DataflowInput, DataflowOutput, DisplayName = "Collection", DataflowPassthrough = "Collection"))
+	FManagedArrayCollection Collection;
+
+	/** (optional) only create volume constraints from surface vertices in VertexSelection to triangles in other geometries.
+	* For example, if the VertexSelection contains only one geometry, only this geometry will bind to other geometries.
+	* No constraints will be created between two geometries that are not in the VertexSelection. */
+	UPROPERTY(meta = (DataflowInput, DisplayName = "(Optional) VertexSelection"))
+	FDataflowVertexSelection VertexSelection;
+
+	/** search radius for point-triangle pairs */
+	UPROPERTY(EditAnywhere, Category = "Dataflow", meta = (ClampMin = "0.0"))
+	float SearchRadius = 0.f;
+
+	/** Stiffness of the volume constraint. This should be around the same magnitude as Young's modulus. */
+	UPROPERTY(EditAnywhere, Category = "Dataflow", meta = (ClampMin = "0.0"))
+	float Stiffness = 1.f;
+
+	/** Render dynamic mesh of the boundary mesh of added volume */
+	UPROPERTY(meta = (DataflowOutput))
+	TObjectPtr<UDynamicMesh> DynamicMesh;
+
 	virtual void Evaluate(UE::Dataflow::FContext& Context, const FDataflowOutput* Out) const override;
 };
\ No newline at end of file
diff --git a/Engine/Source/Runtime/Experimental/Chaos/Private/Chaos/Deformable/ChaosDeformableSolver.cpp b/Engine/Source/Runtime/Experimental/Chaos/Private/Chaos/Deformable/ChaosDeformableSolver.cpp
--- a/Engine/Source/Runtime/Experimental/Chaos/Private/Chaos/Deformable/ChaosDeformableSolver.cpp
+++ b/Engine/Source/Runtime/Experimental/Chaos/Private/Chaos/Deformable/ChaosDeformableSolver.cpp
@@ -40,6 +40,7 @@
 #include "GeometryCollection/Facades/CollectionMeshFacade.h"
 #include "GeometryCollection/Facades/CollectionMuscleActivationFacade.h"
 #include "GeometryCollection/Facades/CollectionTetrahedralFacade.h"
+#include "GeometryCollection/Facades/CollectionVolumeConstraintFacade.h"
 #include "Misc/FileHelper.h"
 #include "HAL/PlatformFileManager.h"
 #include "HAL/IConsoleManager.h"
@@ -96,7 +97,7 @@ namespace Chaos::Softs
 	FDeformableXPBDCorotatedParams GDeformableXPBDCorotatedParams;
 	FAutoConsoleVariableRef CVarDeformableXPBDCorotatedBatchSize(TEXT("p.Chaos.Deformable.XPBDBatchSize"), GDeformableXPBDCorotatedParams.XPBDCorotatedBatchSize, TEXT("Batch size for physics parallel for. [def: 5]"));
 	FAutoConsoleVariableRef CVarDeformableXPBDCorotatedBatchThreshold(TEXT("p.Chaos.Deformable.XPBDBatchThreshold"), GDeformableXPBDCorotatedParams.XPBDCorotatedBatchThreshold, TEXT("Batch threshold for physics parallel for. [def: 5]"));
-	FAutoConsoleVariableRef CVarDeformableXPBDCorotatedBatchNumLogExtremeParticle (TEXT("p.Chaos.Deformable.NumLogExtremeParticle "), GDeformableXPBDCorotatedParams.NumLogExtremeParticle, TEXT("Number of most deformed particles logged. [def: 0]"));
+	FAutoConsoleVariableRef CVarDeformableXPBDCorotatedBatchNumLogExtremeParticle(TEXT("p.Chaos.Deformable.NumLogExtremeParticle"), GDeformableXPBDCorotatedParams.NumLogExtremeParticle, TEXT("Number of most deformed particles logged. [def: 0]"));
 
 	FDeformableXPBDWeakConstraintParams GDeformableXPBDWeakConstraintParams;
 	FAutoConsoleVariableRef CVarDeformableXPBDWeakConstraintLineWidth(TEXT("p.Chaos.Deformable.XPBDWeakConstraintLineWidth"), GDeformableXPBDWeakConstraintParams.DebugLineWidth, TEXT("Line width for visualizing the double bindings in XPBD weak constraints. [def: 5]"));
@@ -185,9 +186,7 @@ namespace Chaos::Softs
 			GSWeakConstraints.Reset(new FGaussSeidelWeakConstraints<FSolverReal, FSolverParticles>({}, {}, {}, {}, {}, GDeformableXPBDWeakConstraintParams));
 			GSDynamicWeakConstraints.Reset(new FGaussSeidelDynamicWeakConstraints<FSolverReal, FSolverParticles>(GDeformableXPBDWeakConstraintParams));
 			GSSphereRepulsionConstraints.Reset(new FGaussSeidelSphereRepulsionConstraints<FSolverReal, FSolverParticles>(Property.SphereRepulsionRadius, Property.SphereRepulsionStiffness, Evolution->Particles(), GDeformableXPBDWeakConstraintParams));
-		}
-		if (Property.bDoLengthBasedMuscleActivation || Property.bOverrideMuscleActivationWithAnimatedCurves)
-		{
+			GSVolumeConstraints.Reset(new FGaussSeidelUnilateralTetConstraints<Softs::FSolverReal, Softs::FSolverParticles>(Evolution->Particles(), {}, {}));
 			MuscleIndexOffset.Empty();
 			MuscleActivationConstraints.Reset(new FMuscleActivationConstraints<FSolverReal, FSolverParticles>());
 		}
@@ -424,8 +423,9 @@ namespace Chaos::Softs
 					if (Property.bUseGaussSeidelConstraints)
 					{
 						InitializeGaussSeidelConstraintVariables();
+						// TODO(flesh): Support muscle activation with XPBD, use an array for GSConstraints
+						InitializeMuscleActivationVariables();
 					}
-					InitializeMuscleActivationVariables();
 				}
 				UninitializedProxys_Internal.SetNum(0, EAllowShrinking::Yes);
 			}
@@ -673,22 +673,24 @@ namespace Chaos::Softs
 			typedef GeometryCollection::Facades::FPositionTargetFacade FPositionTargets;
 			FPositionTargets PositionTargets(Rest);
 
-			TSet<int32> ConstrainedVertices;
-
 			TArray<TArray<int32>> PositionTargetIndices;
 			TArray<TArray<FSolverReal>> PositionTargetWeights;
 			TArray<TArray<int32>> PositionTargetSecondIndices;
 			TArray<TArray<FSolverReal>> PositionTargetSecondWeights;
 			TArray<FSolverReal> PositionTargetStiffness;
-
-			PositionTargetIndices.SetNum(PositionTargets.NumPositionTargets());
-			PositionTargetWeights.SetNum(PositionTargets.NumPositionTargets());
-			PositionTargetSecondIndices.SetNum(PositionTargets.NumPositionTargets());
-			PositionTargetSecondWeights.SetNum(PositionTargets.NumPositionTargets());
-			PositionTargetStiffness.SetNum(PositionTargets.NumPositionTargets());
+			TArray<bool> PositionTargetIsAnisotrpic;
+			TArray<bool> PositionTargetIsZeroRestLength;
+			int32 NumPositionTargets = PositionTargets.NumPositionTargets();
+			PositionTargetIndices.SetNum(NumPositionTargets);
+			PositionTargetWeights.SetNum(NumPositionTargets);
+			PositionTargetSecondIndices.SetNum(NumPositionTargets);
+			PositionTargetSecondWeights.SetNum(NumPositionTargets);
+			PositionTargetStiffness.SetNum(NumPositionTargets);
+			PositionTargetIsAnisotrpic.SetNum(NumPositionTargets);
+			PositionTargetIsZeroRestLength.SetNum(NumPositionTargets);
 
 			// Read in position target info
-			for (int i = PositionTargets.NumPositionTargets() - 1; i >= 0; i--)
+			for (int i = NumPositionTargets - 1; i >= 0; i--)
 			{
 				GeometryCollection::Facades::FPositionTargetsData DataPackage = PositionTargets.GetPositionTarget(i);
 				PositionTargetIndices[i] = DataPackage.SourceIndex;
@@ -696,6 +698,8 @@ namespace Chaos::Softs
 				PositionTargetSecondIndices[i] = DataPackage.TargetIndex;
 				PositionTargetSecondWeights[i] = DataPackage.TargetWeights;
 				PositionTargetStiffness[i] = DataPackage.Stiffness;
+				PositionTargetIsAnisotrpic[i] = DataPackage.bIsAnisotropic;
+				PositionTargetIsZeroRestLength[i] = DataPackage.bIsZeroRestLength;
 			}
 
 			if (Property.bUseGaussSeidelConstraints)
@@ -711,7 +715,9 @@ namespace Chaos::Softs
 						PositionTargetSecondIndices[i][j] += Range.Start;
 					}
 				}
-				GSWeakConstraints->AddExtraConstraints(PositionTargetIndices, PositionTargetWeights, PositionTargetStiffness, PositionTargetSecondIndices, PositionTargetSecondWeights);
+				GSWeakConstraints->AddExtraConstraints(PositionTargetIndices, PositionTargetWeights, 
+					PositionTargetStiffness, PositionTargetSecondIndices, PositionTargetSecondWeights, 
+					PositionTargetIsAnisotrpic, PositionTargetIsZeroRestLength);
 			} 
 			else
 			{
@@ -737,6 +743,28 @@ namespace Chaos::Softs
 				WeakConstraints.Add(TUniquePtr<FXPBDWeakConstraints<FSolverReal, FSolverParticles>>(WeakConstraint));
 			}
 		}
+
+		//initialize volume constraint
+		GeometryCollection::Facades::FVolumeConstraintFacade VolumeConstraint(Rest);
+		int32 NumConstraints = VolumeConstraint.NumVolumeConstraints();
+		TArray<TVector<int32, 4>> InConstraints;
+		TArray<FSolverReal> InVolumes;
+		TArray<FSolverReal> InStiffnessArray;
+		InConstraints.SetNumUninitialized(NumConstraints);
+		InStiffnessArray.SetNumUninitialized(NumConstraints);
+		for (int32 ConstraintIdx = 0; ConstraintIdx < NumConstraints; ++ConstraintIdx)
+		{
+			for (int32 LocalIdx = 0; LocalIdx < 4; ++LocalIdx)
+			{
+				InConstraints[ConstraintIdx][LocalIdx] = VolumeConstraint.GetVolumeIndex(ConstraintIdx)[LocalIdx];
+			}
+			InStiffnessArray[ConstraintIdx] = VolumeConstraint.GetStiffness(ConstraintIdx);
+		}
+		GSVolumeConstraints.Reset(new FGaussSeidelUnilateralTetConstraints<Softs::FSolverReal, Softs::FSolverParticles>(Evolution->Particles(), MoveTemp(InConstraints), MoveTemp(InStiffnessArray)));
+		if (GSVolumeConstraints->NumConstraints() && !Property.bUseGaussSeidelConstraints)
+		{
+			UE_LOG(LogChaosDeformableSolver, Error, TEXT("Error: must check [Use Gauss Seidel constraints] for volume constraints."));
+		}
 	}
 
 
@@ -1707,29 +1735,6 @@ namespace Chaos::Softs
 			{
 				this->GSNeohookeanConstraints->AddHyperelasticResidualAndHessian(Particles, ElementIndex, ElementIndexLocal, Dt, ParticleResidual, ParticleHessian);
 			};
-			
-			
-			if (Property.bEnablePositionTargets)
-			{
-				TArray<TArray<int32>> ParticlesPerColor;
-				GSWeakConstraints->ComputeInitialWCData(Evolution->Particles());
-
-				TArray<TArray<int32>> StaticIncidentElements, StaticIncidentElementsLocal;
-				TArray<TArray<int32>> StaticConstraints = GSWeakConstraints->GetStaticConstraintArrays(StaticIncidentElements, StaticIncidentElementsLocal);
-				GSMainConstraint->AddStaticConstraints(StaticConstraints, StaticIncidentElements, StaticIncidentElementsLocal);
-
-				int32 StaticIndex1 = GSMainConstraint->AddStaticConstraintResidualAndHessianRange(1);
-				GSMainConstraint->StaticConstraintResidualAndHessian()[StaticIndex1] = [this](const FSolverParticles& Particles, const int32 ConstraintIndex, const int32 ConstraintIndexLocal, const FSolverReal Dt, TVec3<FSolverReal>& ParticleResidual, Chaos::PMatrix<FSolverReal, 3, 3>& ParticleHessian)
-				{
-					this->GSWeakConstraints->AddWCResidualAndHessian(Particles, ConstraintIndex, ConstraintIndexLocal, Dt, ParticleResidual, ParticleHessian);
-				};
-
-				int32 PerNodeIndex = GSMainConstraint->AddPerNodeHessianRange(1);
-				GSMainConstraint->PerNodeHessian()[PerNodeIndex] = [this](const int32 p, const FSolverReal Dt, Chaos::PMatrix<FSolverReal, 3, 3>& ParticleHessian)
-				{
-					this->GSWeakConstraints->AddWCHessian(p, Dt, ParticleHessian);
-				};
-			}
 		} 
 		else
 		{
@@ -1760,30 +1765,42 @@ namespace Chaos::Softs
 				{
 					this->GSCorotatedConstraints->AddHyperelasticResidualAndHessian(Particles, ElementIndex, ElementIndexLocal, Dt, ParticleResidual, ParticleHessian);
 				};
+		}
 
+		if (Property.bEnablePositionTargets)
+		{
+			GSWeakConstraints->ComputeInitialWCData(Evolution->Particles());
 
-			if (Property.bEnablePositionTargets)
-			{
-				TArray<TArray<int32>> ParticlesPerColor;
-				GSWeakConstraints->ComputeInitialWCData(Evolution->Particles());
+			TArray<TArray<int32>> StaticIncidentElements, StaticIncidentElementsLocal;
+			const TArray<TArray<int32>>& StaticConstraints = GSWeakConstraints->GetStaticConstraintArrays(StaticIncidentElements, StaticIncidentElementsLocal);
+			GSMainConstraint->AddStaticConstraints(StaticConstraints, StaticIncidentElements, StaticIncidentElementsLocal);
 
-				TArray<TArray<int32>> StaticIncidentElements, StaticIncidentElementsLocal;
-				TArray<TArray<int32>> StaticConstraints = GSWeakConstraints->GetStaticConstraintArrays(StaticIncidentElements, StaticIncidentElementsLocal);
-				GSMainConstraint->AddStaticConstraints(StaticConstraints, StaticIncidentElements, StaticIncidentElementsLocal);
+			const int32 StaticIndex1 = GSMainConstraint->AddStaticConstraintResidualAndHessianRange(1);
+			GSMainConstraint->StaticConstraintResidualAndHessian()[StaticIndex1] = [this](const FSolverParticles& Particles, const int32 ConstraintIndex, const int32 ConstraintIndexLocal, const FSolverReal Dt, TVec3<FSolverReal>& ParticleResidual, Chaos::PMatrix<FSolverReal, 3, 3>& ParticleHessian)
+				{
+					this->GSWeakConstraints->AddWCResidual(Particles, ConstraintIndex, ConstraintIndexLocal, Dt, ParticleResidual, ParticleHessian);
+				};
 
-				int32 StaticIndex1 = GSMainConstraint->AddStaticConstraintResidualAndHessianRange(1);
-				GSMainConstraint->StaticConstraintResidualAndHessian()[StaticIndex1] = [this](const FSolverParticles& Particles, const int32 ConstraintIndex, const int32 ConstraintIndexLocal, const FSolverReal Dt, TVec3<FSolverReal>& ParticleResidual, Chaos::PMatrix<FSolverReal, 3, 3>& ParticleHessian)
-					{
-						this->GSWeakConstraints->AddWCResidualAndHessian(Particles, ConstraintIndex, ConstraintIndexLocal, Dt, ParticleResidual, ParticleHessian);
-					};
+			int32 PerNodeIndex = GSMainConstraint->AddPerNodeHessianRange(1);
+			GSMainConstraint->PerNodeHessian()[PerNodeIndex] = [this](const int32 p, const FSolverReal Dt, Chaos::PMatrix<FSolverReal, 3, 3>& ParticleHessian)
+				{
+					this->GSWeakConstraints->AddWCHessian(p, Dt, ParticleHessian);
+				};
+		}
 
-				int32 PerNodeIndex = GSMainConstraint->AddPerNodeHessianRange(1);
-				GSMainConstraint->PerNodeHessian()[PerNodeIndex] = [this](const int32 p, const FSolverReal Dt, Chaos::PMatrix<FSolverReal, 3, 3>& ParticleHessian)
-					{
-						this->GSWeakConstraints->AddWCHessian(p, Dt, ParticleHessian);
-					};
+		if (GSVolumeConstraints->NumConstraints())
+		{
+			TArray<TArray<int32>> StaticIncidentElements, StaticIncidentElementsLocal;
+			const TArray<TArray<int32>> StaticConstraints = GSVolumeConstraints->GetStaticConstraintArrays(StaticIncidentElements, StaticIncidentElementsLocal);
+			GSMainConstraint->AddStaticConstraints(StaticConstraints, StaticIncidentElements, StaticIncidentElementsLocal);
 
-			}
+			const int32 StaticIndex1 = GSMainConstraint->AddStaticConstraintResidualAndHessianRange(1);
+			GSMainConstraint->StaticConstraintResidualAndHessian()[StaticIndex1] = [this](const FSolverParticles& Particles, const int32 ConstraintIndex, const int32 ConstraintIndexLocal, const FSolverReal Dt, TVec3<FSolverReal>& ParticleResidual, Chaos::PMatrix<FSolverReal, 3, 3>& ParticleHessian)
+				{
+					this->GSVolumeConstraints->AddResidualAndHessian(
+						Evolution->Particles(),
+						ConstraintIndex, ConstraintIndexLocal, Dt, ParticleResidual, ParticleHessian);
+				};
 		}
 
 		GSMainConstraint->InitStaticColor(Evolution->Particles(), &(Evolution->ParticlesActiveView()));
@@ -1803,7 +1820,7 @@ namespace Chaos::Softs
 			int32 DynamicIndex = GSMainConstraint->AddDynamicConstraintResidualAndHessianRange(1);
 			GSMainConstraint->DynamicConstraintResidualAndHessian()[DynamicIndex] = [this](const FSolverParticles& Particles, const int32 ConstraintIndex, const int32 ConstraintIndexLocal, const FSolverReal Dt, TVec3<FSolverReal>& ParticleResidual, Chaos::PMatrix<FSolverReal, 3, 3>& ParticleHessian)
 			{
-				this->GSDynamicWeakConstraints->AddWCResidualAndHessian(Particles, ConstraintIndex, ConstraintIndexLocal, Dt, ParticleResidual, ParticleHessian);
+				this->GSDynamicWeakConstraints->AddWCResidual(Particles, ConstraintIndex, ConstraintIndexLocal, Dt, ParticleResidual, ParticleHessian);
 			};
 
 			int32 PerNodeIndex = GSMainConstraint->AddPerNodeHessianRange(1);
@@ -1823,7 +1840,7 @@ namespace Chaos::Softs
 					if (this->bDynamicConstraintIsUpdated)
 					{
 						TArray<TArray<int32>> WCDynamicIncidentElements,WCDynamicIncidentElementsLocal;
-						TArray<TArray<int32>> DynamicConstraints = GSDynamicWeakConstraints->GetStaticConstraintArrays(WCDynamicIncidentElements, WCDynamicIncidentElementsLocal);
+						const TArray<TArray<int32>>& DynamicConstraints = GSDynamicWeakConstraints->GetStaticConstraintArrays(WCDynamicIncidentElements, WCDynamicIncidentElementsLocal);
 						GSMainConstraint->ResetDynamicConstraints();
 						GSMainConstraint->AddDynamicConstraints(DynamicConstraints, WCDynamicIncidentElements, WCDynamicIncidentElementsLocal, true);
 						this->GSMainConstraint->InitDynamicColor(InParticles);
@@ -1836,7 +1853,7 @@ namespace Chaos::Softs
 			int32 TransientIndex = GSMainConstraint->AddTransientConstraintResidualAndHessianRange(1);
 			GSMainConstraint->TransientConstraintResidualAndHessian()[TransientIndex] = [this](const FSolverParticles& Particles, const int32 ConstraintIndex, const int32 ConstraintIndexLocal, const FSolverReal Dt, TVec3<FSolverReal>& ParticleResidual, Chaos::PMatrix<FSolverReal, 3, 3>& ParticleHessian)
 				{
-					this->GSWeakConstraints->AddWCResidualAndHessian(Particles, ConstraintIndex + this->GSWeakConstraints->InitialWCSize, ConstraintIndexLocal, Dt, ParticleResidual, ParticleHessian);
+					this->GSWeakConstraints->AddWCResidual(Particles, ConstraintIndex + this->GSWeakConstraints->InitialWCSize, ConstraintIndexLocal, Dt, ParticleResidual, ParticleHessian);
 				};
 
 			if (!Property.bEnablePositionTargets)
@@ -1915,50 +1932,47 @@ namespace Chaos::Softs
 		int32 InitIndex = Evolution->AddConstraintInitRange(1, true);
 		Evolution->ConstraintInits()[InitIndex] = [this](FSolverParticles& InParticles, const FSolverReal Dt)
 		{
-			if (Property.bUseGaussSeidelConstraints)
+			if (Property.bDoLengthBasedMuscleActivation)
 			{
-				if (Property.bDoLengthBasedMuscleActivation)
-				{
-					this->MuscleActivationConstraints->UpdateLengthBasedMuscleActivation(InParticles);
-				}
-				if (Property.bOverrideMuscleActivationWithAnimatedCurves)
+				this->MuscleActivationConstraints->UpdateLengthBasedMuscleActivation(InParticles);
+			}
+			if (Property.bOverrideMuscleActivationWithAnimatedCurves)
+			{
+				//Override muscle activation with curve data
+				for (const TPair<FThreadingProxy::FKey, TUniquePtr<FThreadingProxy>>& ProxyElem : Proxies)
 				{
-					//Override muscle activation with curve data
-					for (const TPair<FThreadingProxy::FKey, TUniquePtr<FThreadingProxy>>& ProxyElem : Proxies)
+					if (const FFleshThreadingProxy* Proxy = ProxyElem.Value->As<FFleshThreadingProxy>())
 					{
-						if (const FFleshThreadingProxy* Proxy = ProxyElem.Value->As<FFleshThreadingProxy>())
+						if (const UObject* Owner = Proxy->GetOwner())
 						{
-							if (const UObject* Owner = Proxy->GetOwner())
+							if (!Proxy->GetIsCached() && MuscleIndexOffset.Contains(Owner))
 							{
-								if (!Proxy->GetIsCached() && MuscleIndexOffset.Contains(Owner))
+								if (this->CurrentInputPackage && this->CurrentInputPackage->ObjectMap.Contains(Owner))
 								{
-									if (this->CurrentInputPackage && this->CurrentInputPackage->ObjectMap.Contains(Owner))
+									if (FFleshThreadingProxy::FFleshInputBuffer* FleshInputBuffer = this->CurrentInputPackage->ObjectMap[Owner]->As<FFleshThreadingProxy::FFleshInputBuffer>())
 									{
-										if (FFleshThreadingProxy::FFleshInputBuffer* FleshInputBuffer = this->CurrentInputPackage->ObjectMap[Owner]->As<FFleshThreadingProxy::FFleshInputBuffer>())
+										const TArray<int32>& MuscleIndices = FleshInputBuffer->MuscleIndices;
+										const TArray<float>& MuscleActivations = FleshInputBuffer->MuscleActivation;
+										if (ensure(MuscleIndices.Num() == MuscleActivations.Num()))
 										{
-											const TArray<int32>& MuscleIndices = FleshInputBuffer->MuscleIndices;
-											const TArray<float>& MuscleActivations = FleshInputBuffer->MuscleActivation;
-											if (ensure(MuscleIndices.Num() == MuscleActivations.Num()))
+											const int32 Offset = MuscleIndexOffset[Owner];
+											for (int32 Idx = 0; Idx < MuscleIndices.Num(); ++Idx)
 											{
-												const int32 Offset = MuscleIndexOffset[Owner];
-												for (int32 Idx = 0; Idx < MuscleIndices.Num(); ++Idx)
-												{
-													this->MuscleActivationConstraints->SetMuscleActivation(Offset + MuscleIndices[Idx], MuscleActivations[Idx]);
-												}
-												if (this->PreviousInputPackage && this->PreviousInputPackage->ObjectMap.Contains(Owner))
+												this->MuscleActivationConstraints->SetMuscleActivation(Offset + MuscleIndices[Idx], MuscleActivations[Idx]);
+											}
+											if (this->PreviousInputPackage && this->PreviousInputPackage->ObjectMap.Contains(Owner))
+											{
+												if (FFleshThreadingProxy::FFleshInputBuffer* PreviousFleshBuffer = this->PreviousInputPackage->ObjectMap[Owner]->As<FFleshThreadingProxy::FFleshInputBuffer>())
 												{
-													if (FFleshThreadingProxy::FFleshInputBuffer* PreviousFleshBuffer = this->PreviousInputPackage->ObjectMap[Owner]->As<FFleshThreadingProxy::FFleshInputBuffer>())
+													FSolverReal CurrentRatio = FSolverReal(this->Iteration) / FSolverReal(this->Property.NumSolverSubSteps);
+													const TArray<int32>& PrevMuscleIndices = PreviousFleshBuffer->MuscleIndices;
+													const TArray<float>& PrevMuscleActivations = PreviousFleshBuffer->MuscleActivation;
+													for (int32 Idx = 0; Idx < PrevMuscleIndices.Num(); ++Idx)
 													{
-														FSolverReal CurrentRatio = FSolverReal(this->Iteration) / FSolverReal(this->Property.NumSolverSubSteps);
-														const TArray<int32>& PrevMuscleIndices = PreviousFleshBuffer->MuscleIndices;
-														const TArray<float>& PrevMuscleActivations = PreviousFleshBuffer->MuscleActivation;
-														for (int32 Idx = 0; Idx < PrevMuscleIndices.Num(); ++Idx)
-														{
-															const int32 MuscleGlobalIdx = Offset + PrevMuscleIndices[Idx];
-															const float CurrentMuscleActivation = this->MuscleActivationConstraints->GetMuscleActivation(MuscleGlobalIdx);
-															this->MuscleActivationConstraints->SetMuscleActivation(MuscleGlobalIdx,
-																CurrentRatio * CurrentMuscleActivation + (1 - CurrentRatio) * PrevMuscleActivations[Idx]);
-														}
+														const int32 MuscleGlobalIdx = Offset + PrevMuscleIndices[Idx];
+														const float CurrentMuscleActivation = this->MuscleActivationConstraints->GetMuscleActivation(MuscleGlobalIdx);
+														this->MuscleActivationConstraints->SetMuscleActivation(MuscleGlobalIdx,
+															CurrentRatio * CurrentMuscleActivation + (1 - CurrentRatio) * PrevMuscleActivations[Idx]);
 													}
 												}
 											}
@@ -1969,42 +1983,28 @@ namespace Chaos::Softs
 						}
 					}
 				}
+			}
 
-				if (Property.bUseGSNeohookean)
-				{
-					this->MuscleActivationConstraints->ApplyMuscleActivation(*this->GSNeohookeanConstraints);
-				}
-				else
-				{
-					this->MuscleActivationConstraints->ApplyMuscleActivation(*this->GSCorotatedConstraints);
-				}
+			if (Property.bUseGSNeohookean)
+			{
+				this->MuscleActivationConstraints->ApplyMuscleActivation(*this->GSNeohookeanConstraints);
 			}
 			else
 			{
-				// TODO(flesh): Support muscle activation with XPBD, use an array for GSConstraints
-				//this->MuscleActivationConstraints->ApplyMuscleActivation(*this->CorotatedConstraints[0]);
-				UE_LOG(LogChaosDeformableSolver, Warning, TEXT("Muscle activation is not supported for XPBD now, please switch to Gauss-Seidel."));
+				this->MuscleActivationConstraints->ApplyMuscleActivation(*this->GSCorotatedConstraints);
 			}
 		};
 
 		//Adjust muscle rest volume
-		if (MuscleActivationConstraints->NumMuscles())
+		if (MuscleActivationConstraints && MuscleActivationConstraints->NumMuscles())
 		{
-			if (Property.bUseGaussSeidelConstraints)
+			if (Property.bUseGSNeohookean)
 			{
-				if (Property.bUseGSNeohookean)
-				{
-					MuscleActivationConstraints->ApplyInflationVolumeScale(*GSNeohookeanConstraints);
-				}
-				else
-				{
-					MuscleActivationConstraints->ApplyInflationVolumeScale(*GSCorotatedConstraints);
-				}
+				MuscleActivationConstraints->ApplyInflationVolumeScale(*GSNeohookeanConstraints);
 			}
 			else
 			{
-				// TODO(flesh): Support muscle activation with XPBD, use an array for GSConstraints
-				UE_LOG(LogChaosDeformableSolver, Warning, TEXT("Inflation volume scale is only supported for the Gauss-Seidel solver."));
+				MuscleActivationConstraints->ApplyInflationVolumeScale(*GSCorotatedConstraints);
 			}
 		}
 	}
@@ -2443,7 +2443,7 @@ namespace Chaos::Softs
 		}
 
 		// p.Chaos.Deformable.NumLogExtremeParticle 
-		if (GDeformableXPBDCorotatedParams.NumLogExtremeParticle > 0)
+		if (Property.bUseGaussSeidelConstraints && GDeformableXPBDCorotatedParams.NumLogExtremeParticle > 0)
 		{
 			auto RestPosition = [this](int32 VertIdx, FVector3f& OutPosition)
 				{
@@ -2497,6 +2497,7 @@ namespace Chaos::Softs
 					{
 						continue;
 					}
+					//Particles().X is already updated after AdvanceDt
 					float Dist = (Evolution->Particles().GetX(VertIdx) - Evolution->Particles().GetX(NeighborIdx)).Size();
 					TotalDist += Dist;
 					float Rest = (VertPosition - NeighborPosition).Size();
diff --git a/Engine/Source/Runtime/Experimental/Chaos/Private/GeometryCollection/Facades/CollectionPositionTargetFacade.cpp b/Engine/Source/Runtime/Experimental/Chaos/Private/GeometryCollection/Facades/CollectionPositionTargetFacade.cpp
--- a/Engine/Source/Runtime/Experimental/Chaos/Private/GeometryCollection/Facades/CollectionPositionTargetFacade.cpp
+++ b/Engine/Source/Runtime/Experimental/Chaos/Private/GeometryCollection/Facades/CollectionPositionTargetFacade.cpp
@@ -18,6 +18,8 @@ namespace GeometryCollection::Facades
 	const FName FPositionTargetFacade::TargetName("TargetName");
 	const FName FPositionTargetFacade::TargetWeights("TargetWeights");
 	const FName FPositionTargetFacade::SourceWeights("SourceWeights");
+	const FName FPositionTargetFacade::IsAnisotropic("IsAnisotropic");
+	const FName FPositionTargetFacade::IsZeroRestLength("IsZeroRestLength");
 
 	FPositionTargetFacade::FPositionTargetFacade(FManagedArrayCollection& InCollection, const FName& InVerticesGroup)
 		: ConstCollection(InCollection)
@@ -29,6 +31,8 @@ namespace GeometryCollection::Facades
 		, DampingAttribute(InCollection, Damping, GroupName)
 		, TargetWeightsAttribute(InCollection, TargetWeights, GroupName)
 		, SourceWeightsAttribute(InCollection, SourceWeights, GroupName)
+		, IsAnisotropicAttribute(InCollection, IsAnisotropic, GroupName)
+		, IsZeroRestLengthAttribute(InCollection, IsZeroRestLength, GroupName)
 	{
 		DefineSchema();
 	}
@@ -43,12 +47,15 @@ namespace GeometryCollection::Facades
 		, DampingAttribute(InCollection, Damping, GroupName)
 		, TargetWeightsAttribute(InCollection, TargetWeights, GroupName)
 		, SourceWeightsAttribute(InCollection, SourceWeights, GroupName)
+		, IsAnisotropicAttribute(InCollection, IsAnisotropic, GroupName)
+		, IsZeroRestLengthAttribute(InCollection, IsZeroRestLength, GroupName)
 	{
 		//DefineSchema();
 	}
 
 	bool FPositionTargetFacade::IsValid() const
 	{
+		// Not checking IsAnisotropicAttribute/IsZeroRestLengthAttribute for backward compatibility in 5.5 and before
 		return TargetIndexAttribute.IsValid() && TargetIndexAttribute.GetGroupDependency() ==  VerticesGroup 
 			&& SourceIndexAttribute.IsValid() && SourceIndexAttribute.GetGroupDependency() == VerticesGroup 
 			&& StiffnessAttribute.IsValid() && DampingAttribute.IsValid() && TargetWeightsAttribute.IsValid() && SourceWeightsAttribute.IsValid();
@@ -63,6 +70,8 @@ namespace GeometryCollection::Facades
 		DampingAttribute.Add();
 		TargetWeightsAttribute.Add();
 		SourceWeightsAttribute.Add();
+		IsAnisotropicAttribute.Add();
+		IsZeroRestLengthAttribute.Add();
 	}
 
 
@@ -78,7 +87,8 @@ namespace GeometryCollection::Facades
 			DampingAttribute.Modify()[NewIndex] = InputData.Damping;
 			TargetWeightsAttribute.Modify()[NewIndex] = InputData.TargetWeights;
 			SourceWeightsAttribute.Modify()[NewIndex] = InputData.SourceWeights;
-
+			IsAnisotropicAttribute.Modify()[NewIndex] = InputData.bIsAnisotropic;
+			IsZeroRestLengthAttribute.Modify()[NewIndex] = InputData.bIsZeroRestLength;
 			return NewIndex;
 		}
 		return INDEX_NONE;
@@ -89,36 +99,53 @@ namespace GeometryCollection::Facades
 		FPositionTargetsData ReturnData;
 		if (IsValid())
 		{
-			if (StiffnessAttribute.Num() > DataIndex && DataIndex > -1) 
+			if (StiffnessAttribute.IsValidIndex(DataIndex))
 			{
 				ReturnData.Stiffness = StiffnessAttribute.Get()[DataIndex];
 			}
 
-			if (DampingAttribute.Num() > DataIndex && DataIndex > -1)
+			if (DampingAttribute.IsValidIndex(DataIndex))
 			{
 				ReturnData.Damping = DampingAttribute.Get()[DataIndex];
 			}
 
-			if (SourceIndexAttribute.Num() > DataIndex && DataIndex > -1)
+			if (SourceIndexAttribute.IsValidIndex(DataIndex))
 			{
 				ReturnData.SourceIndex = SourceIndexAttribute[DataIndex];
 			}
 
-			if (TargetIndexAttribute.Num() > DataIndex && DataIndex > -1)
+			if (TargetIndexAttribute.IsValidIndex(DataIndex))
 			{
 				ReturnData.TargetIndex = TargetIndexAttribute[DataIndex];
 			}
 
-			if (SourceWeightsAttribute.Num() > DataIndex && DataIndex > -1)
+			if (SourceWeightsAttribute.IsValidIndex(DataIndex))
 			{
 				ReturnData.SourceWeights = SourceWeightsAttribute.Get()[DataIndex];
 			}
 
-			if (TargetWeightsAttribute.Num() > DataIndex && DataIndex > -1)
+			if (TargetWeightsAttribute.IsValidIndex(DataIndex))
 			{
 				ReturnData.TargetWeights = TargetWeightsAttribute.Get()[DataIndex];
 			}
 
+			if (IsAnisotropicAttribute.IsValid() && IsAnisotropicAttribute.IsValidIndex(DataIndex))
+			{
+				ReturnData.bIsAnisotropic = IsAnisotropicAttribute.Get()[DataIndex];
+			}
+			else
+			{
+				ReturnData.bIsAnisotropic = false; //backward compatible for 5.5 and before
+			}
+
+			if (IsZeroRestLengthAttribute.IsValid() && IsZeroRestLengthAttribute.IsValidIndex(DataIndex))
+			{
+				ReturnData.bIsZeroRestLength = IsZeroRestLengthAttribute.Get()[DataIndex];
+			}
+			else
+			{
+				ReturnData.bIsZeroRestLength = true; //backward compatible for 5.5 and before
+			}
 		}
 		return ReturnData;
 	
@@ -159,10 +186,10 @@ namespace GeometryCollection::Facades
 		return InvalidConstraintIndices.Num();
 	}
 
-	int32 FPositionTargetFacade::RemovePositionTarget(TFunctionRef<bool(const int32 VertexIdx)> IsVertexGroup1, TFunctionRef<bool(const int32 VertexIdx)> IsVertexGroup2)
+	int32 FPositionTargetFacade::RemovePositionTargetBetween(TFunctionRef<bool(const int32 VertexIdx)> IsVertexGroup1, TFunctionRef<bool(const int32 VertexIdx)> IsVertexGroup2)
 	{
 		check(!IsConst());
-		TArray<int32> InvalidConstraintIndices;
+		TArray<int32> ConstraintIndicesToRemove;
 		if (IsValid())
 		{
 			for (int32 Index = 0; Index < Collection->NumElements(GroupName); ++Index)
@@ -183,11 +210,11 @@ namespace GeometryCollection::Facades
 				}
 				if ((IsTargetGroup1 && IsSourceGroup2) || (IsTargetGroup2 && IsSourceGroup1))
 				{
-					InvalidConstraintIndices.Add(Index);
+					ConstraintIndicesToRemove.Add(Index);
 				}
 			}
-			Collection->RemoveElements(GroupName, InvalidConstraintIndices);
+			Collection->RemoveElements(GroupName, ConstraintIndicesToRemove);
 		}
-		return InvalidConstraintIndices.Num();
+		return ConstraintIndicesToRemove.Num();
 	}
 }
diff --git a/Engine/Source/Runtime/Experimental/Chaos/Private/GeometryCollection/Facades/CollectionVolumeConstraintFacade.cpp b/Engine/Source/Runtime/Experimental/Chaos/Private/GeometryCollection/Facades/CollectionVolumeConstraintFacade.cpp
new file mode 100644
--- /dev/null
+++ b/Engine/Source/Runtime/Experimental/Chaos/Private/GeometryCollection/Facades/CollectionVolumeConstraintFacade.cpp
@@ -0,0 +1,124 @@
+// Copyright Epic Games, Inc. All Rights Reserved.
+
+/*=============================================================================
+	FleshCollection.cpp: FFleshCollection methods.
+=============================================================================*/
+
+#include "GeometryCollection/Facades/CollectionVolumeConstraintFacade.h"
+
+namespace GeometryCollection::Facades
+{
+	// Attributes
+	const FName FVolumeConstraintFacade::GroupName("VolumeConstraints");
+	const FName FVolumeConstraintFacade::VolumeIndex("VolumeIndex");
+	const FName FVolumeConstraintFacade::Stiffness("Stiffness");
+
+	FVolumeConstraintFacade::FVolumeConstraintFacade(FManagedArrayCollection& InCollection)
+		: ConstCollection(InCollection)
+		, Collection(&InCollection)
+		, VolumeIndexAttribute(InCollection, VolumeIndex, GroupName, FGeometryCollection::VerticesGroup)
+		, StiffnessAttribute(InCollection, Stiffness, GroupName)
+	{
+		DefineSchema();
+	}
+
+	FVolumeConstraintFacade::FVolumeConstraintFacade(const FManagedArrayCollection& InCollection)
+		: ConstCollection(InCollection)
+		, Collection(nullptr)
+		, VolumeIndexAttribute(InCollection, VolumeIndex, GroupName, FGeometryCollection::VerticesGroup)
+		, StiffnessAttribute(InCollection, Stiffness, GroupName)
+	{}
+
+	bool FVolumeConstraintFacade::IsValid() const
+	{
+		return VolumeIndexAttribute.IsValid() && StiffnessAttribute.IsValid();
+	}
+
+	void FVolumeConstraintFacade::DefineSchema()
+	{
+		check(!IsConst());
+		VolumeIndexAttribute.Add();
+		StiffnessAttribute.Add();
+	}
+
+
+	int32 FVolumeConstraintFacade::AddVolumeConstraint(const FIntVector4& NewVolumeIndex, float NewStiffness)
+	{
+		check(!IsConst());
+		if (IsValid())
+		{
+			int32 NewIndex = VolumeIndexAttribute.AddElements(1);
+			VolumeIndexAttribute.Modify()[NewIndex] = NewVolumeIndex;
+			StiffnessAttribute.Modify()[NewIndex] = NewStiffness;
+			return NewIndex;
+		}
+		return INDEX_NONE;
+	}
+
+	FIntVector4 FVolumeConstraintFacade::GetVolumeIndex(int32 AttributeIndex) const
+	{
+		if (IsValid() && VolumeIndexAttribute.IsValidIndex(AttributeIndex))
+		{
+			return VolumeIndexAttribute[AttributeIndex];
+		}
+		return FIntVector4(INDEX_NONE);
+	}
+
+	float FVolumeConstraintFacade::GetStiffness(int32 AttributeIndex) const
+	{
+		if (IsValid() && StiffnessAttribute.IsValidIndex(AttributeIndex))
+		{
+			return StiffnessAttribute[AttributeIndex];
+		}
+		return 0.f;
+	}
+
+	int32 FVolumeConstraintFacade::RemoveInvalidVolumeConstraint()
+	{
+		check(!IsConst());
+		TArray<int32> InvalidConstraintIndices;
+		if (IsValid())
+		{
+			const int32 NumElemDependency = Collection->NumElements(VolumeIndexAttribute.GetGroupDependency());
+			for (int32 Index = 0; Index < VolumeIndexAttribute.Num(); ++Index)
+			{
+				for (int32 LocalIdx = 0; LocalIdx < 4; ++LocalIdx)
+				{
+					if (VolumeIndexAttribute[Index][LocalIdx] < 0 || VolumeIndexAttribute[Index][LocalIdx] > NumElemDependency)
+					{
+						InvalidConstraintIndices.Add(Index);
+						break;
+					}
+				}
+			}
+			Collection->RemoveElements(GroupName, InvalidConstraintIndices);
+		}
+		return InvalidConstraintIndices.Num();
+	}
+
+	int32 FVolumeConstraintFacade::RemoveVolumeConstraintBetween(TFunctionRef<bool(const int32 VertexIdx)> IsVertexGroup1, TFunctionRef<bool(const int32 VertexIdx)> IsVertexGroup2)
+	{
+		check(!IsConst());
+		TArray<int32> ConstraintIndicesToRemove;
+		if (IsValid())
+		{
+			for (int32 Index = 0; Index < Collection->NumElements(GroupName); ++Index)
+			{
+				bool HasGroup1 = false;
+				bool HasGroup2 = false;
+				for (int32 LocalIdx = 0; LocalIdx < 4; ++LocalIdx)
+				{
+					int32 VertIdx = VolumeIndexAttribute[Index][LocalIdx];
+					HasGroup1 = HasGroup1 || IsVertexGroup1(VertIdx);
+					HasGroup2 = HasGroup2 || IsVertexGroup2(VertIdx);	
+				}
+				if (HasGroup1 && HasGroup2)
+				{
+					ConstraintIndicesToRemove.Add(Index);
+				}
+			}
+			Collection->RemoveElements(GroupName, ConstraintIndicesToRemove);
+		}
+		return ConstraintIndicesToRemove.Num();
+	}
+}
diff --git a/Engine/Source/Runtime/Experimental/Chaos/Public/Chaos/Deformable/ChaosDeformableSolver.h b/Engine/Source/Runtime/Experimental/Chaos/Public/Chaos/Deformable/ChaosDeformableSolver.h
--- a/Engine/Source/Runtime/Experimental/Chaos/Public/Chaos/Deformable/ChaosDeformableSolver.h
+++ b/Engine/Source/Runtime/Experimental/Chaos/Public/Chaos/Deformable/ChaosDeformableSolver.h
@@ -16,6 +16,7 @@
 #include "Chaos/Deformable/GaussSeidelWeakConstraints.h"
 #include "Chaos/Deformable/GaussSeidelDynamicWeakConstraints.h"
 #include "Chaos/Deformable/GaussSeidelSphereRepulsionConstraints.h"
+#include "Chaos/Deformable/GaussSeidelUnilateralTetConstraints.h"
 #include "Chaos/Deformable/GaussSeidelMainConstraint.h"
 #include "Chaos/Deformable/MuscleActivationConstraints.h"
 #include "Chaos/Deformable/GaussSeidelCorotatedCodimensionalConstraints.h"
@@ -209,6 +210,7 @@ namespace Chaos::Softs
 		TUniquePtr<Softs::FGaussSeidelWeakConstraints<Softs::FSolverReal, Softs::FSolverParticles>> GSWeakConstraints;
 		TUniquePtr<Softs::FGaussSeidelDynamicWeakConstraints<Softs::FSolverReal, Softs::FSolverParticles>> GSDynamicWeakConstraints;
 		TUniquePtr<Softs::FGaussSeidelSphereRepulsionConstraints<Softs::FSolverReal, Softs::FSolverParticles>> GSSphereRepulsionConstraints;
+		TUniquePtr<Softs::FGaussSeidelUnilateralTetConstraints<Softs::FSolverReal, Softs::FSolverParticles>> GSVolumeConstraints;
 		TArray<TUniquePtr<Softs::FXPBDWeakConstraints<Softs::FSolverReal, Softs::FSolverParticles>>> WeakConstraints;
 		TArray<TUniquePtr<Softs::FBlendedXPBDCorotatedConstraints<Softs::FSolverReal, Softs::FSolverParticles>>> BlendedCorotatedConstraints;
 		TUniquePtr<Softs::FXPBDGridBasedCorotatedConstraints<Softs::FSolverReal, Softs::FSolverParticles>> GridBasedCorotatedConstraint;
diff --git a/Engine/Source/Runtime/Experimental/Chaos/Public/Chaos/Deformable/GaussSeidelUnilateralTetConstraints.h b/Engine/Source/Runtime/Experimental/Chaos/Public/Chaos/Deformable/GaussSeidelUnilateralTetConstraints.h
new file mode 100644
--- /dev/null
+++ b/Engine/Source/Runtime/Experimental/Chaos/Public/Chaos/Deformable/GaussSeidelUnilateralTetConstraints.h
@@ -0,0 +1,109 @@
+// Copyright Epic Games, Inc. All Rights Reserved.
+#pragma once
+
+#include "Chaos/Utilities.h"
+
+namespace Chaos::Softs
+{
+	template <typename T, typename ParticleType>
+	class FGaussSeidelUnilateralTetConstraints
+	{
+	public: 
+		FGaussSeidelUnilateralTetConstraints(
+			const ParticleType& Particles,
+			TArray<TVector<int32, 4>>&& InConstraints,
+			TArray<T> InStiffnessArray)
+			: Constraints(InConstraints), 
+			StiffnessArray(InStiffnessArray)
+		{
+			IncidentElements = Chaos::Utilities::ComputeIncidentElements(Constraints, &IncidentElementsLocal);
+			Volumes.SetNumUninitialized(Constraints.Num());
+
+			for (int32 Idx = 0; Idx < Constraints.Num(); ++Idx)
+			{
+				const TVec4<int32>& Constraint = Constraints[Idx];
+				Volumes[Idx] = ComputeVolume(
+					Particles.X(Constraint[0]),
+					Particles.X(Constraint[1]),
+					Particles.X(Constraint[2]),
+					Particles.X(Constraint[3]));
+				StiffnessArray[Idx] /= Volumes[Idx];
+			}
+		};
+
+		static T ComputeVolume(const TVec3<T>& P1, const TVec3<T>& P2, const TVec3<T>& P3, const TVec3<T>& P4)
+		{
+			const TVec3<T> P2P1 = P2 - P1;
+			const TVec3<T> P4P1 = P4 - P1;
+			const TVec3<T> P3P1 = P3 - P1;
+			return TVec3<T>::DotProduct(TVec3<T>::CrossProduct(P2P1, P3P1), P4P1) / (T)6.;
+		}
+
+		void AddEnergy(const ParticleType& InParticles, const int32 ConstraintIndex, const T Dt, T& Energy) const
+		{
+			const TVec4<int32>& Constraint = Constraints[ConstraintIndex];
+			T CurrentVol = ComputeVolume(
+				InParticles.P(Constraint[0]),
+				InParticles.P(Constraint[1]),
+				InParticles.P(Constraint[2]),
+				InParticles.P(Constraint[3]));
+			Energy += Dt * Dt * StiffnessArray[ConstraintIndex] / (T)2. * (CurrentVol - Volumes[ConstraintIndex]) * (CurrentVol - Volumes[ConstraintIndex]);
+		}
+
+		void AddResidualAndHessian(const ParticleType& InParticles, const int32 ConstraintIndex, const int32 LocalIndex, const T Dt, TVec3<T>& ParticleResidual, Chaos::PMatrix<T, 3, 3>& ParticleHessian) const
+		{
+			const TVec4<int32>& Constraint = Constraints[ConstraintIndex];
+			const int32 Index1 = Constraint[0];
+			const int32 Index2 = Constraint[1];
+			const int32 Index3 = Constraint[2];
+			const int32 Index4 = Constraint[3];
+			
+			const TVec3<T>& P1 = InParticles.P(Index1);
+			const TVec3<T>& P2 = InParticles.P(Index2);
+			const TVec3<T>& P3 = InParticles.P(Index3);
+			const TVec3<T>& P4 = InParticles.P(Index4);
+
+			TVec4<TVec3<T>> Grads;
+			const TVec3<T> P2P1 = P2 - P1;
+			const TVec3<T> P4P1 = P4 - P1;
+			const TVec3<T> P3P1 = P3 - P1;
+			Grads[1] = TVec3<T>::CrossProduct(P3P1, P4P1) / (T)6.;
+			Grads[2] = TVec3<T>::CrossProduct(P4P1, P2P1) / (T)6.;
+			Grads[3] = TVec3<T>::CrossProduct(P2P1, P3P1) / (T)6.;
+			Grads[0] = -(Grads[1] + Grads[2] + Grads[3]);
+
+			const T Volume = ComputeVolume(P1, P2, P3, P4);
+			const T C_Hessian = StiffnessArray[ConstraintIndex] * Dt * Dt;
+			const T C_Residual = (Volume - Volumes[ConstraintIndex]) * C_Hessian;
+			ParticleResidual += C_Residual * Grads[LocalIndex];
+			ParticleHessian += C_Hessian * Chaos::PMatrix<T, 3, 3>::OuterProduct(Grads[LocalIndex], Grads[LocalIndex]);
+		}
+
+		TArray<TArray<int32>> GetStaticConstraintArrays(TArray<TArray<int32>>& InIncidentElements, TArray<TArray<int32>>& InIncidentElementsLocal) const
+		{
+			InIncidentElements = IncidentElements;
+			InIncidentElementsLocal = IncidentElementsLocal;
+			TArray<TArray<int32>> NestedConstraints;
+			NestedConstraints.SetNum(Constraints.Num());
+			for (int32 Idx = 0; Idx < NestedConstraints.Num(); ++Idx)
+			{
+				NestedConstraints[Idx] = {Constraints[Idx][0], Constraints[Idx][1], Constraints[Idx][2], Constraints[Idx][3]};
+			}
+			return NestedConstraints;
+		}
+
+		int32 NumConstraints() const
+		{
+			return Constraints.Num();
+		}
+
+	private:
+		TArray<TVector<int32, 4>> Constraints;
+		TArray<T> Volumes;
+		TArray<TArray<int32>> IncidentElements;
+		TArray<TArray<int32>> IncidentElementsLocal;
+		TArray<T> StiffnessArray;
+	};
+
+
+}// End namespace Chaos::Softs
diff --git a/Engine/Source/Runtime/Experimental/Chaos/Public/Chaos/Deformable/GaussSeidelWeakConstraints.h b/Engine/Source/Runtime/Experimental/Chaos/Public/Chaos/Deformable/GaussSeidelWeakConstraints.h
--- a/Engine/Source/Runtime/Experimental/Chaos/Public/Chaos/Deformable/GaussSeidelWeakConstraints.h
+++ b/Engine/Source/Runtime/Experimental/Chaos/Public/Chaos/Deformable/GaussSeidelWeakConstraints.h
@@ -30,6 +30,8 @@ namespace Chaos::Softs
 		TArray<T> SingleSecondWeights = {};
 		bool bIsAnisotropic = false;
 		TVec3<T> SingleNormal = TVec3<T>((T)0.);
+		bool bIsZeroRestLength = false;
+		T RestLength = (T)0.;
 	};
 
 	template<class T>
@@ -45,6 +47,8 @@ namespace Chaos::Softs
 			AddArray(&MStiffness);
 			AddArray(&MIsAnisotropic);
 			AddArray(&MNormals);
+			AddArray(&MIsZeroRestLength);
+			AddArray(&MRestLength);
 		}
 		TGaussSeidelWeakConstraintData(const TGaussSeidelWeakConstraintData<T>& Other) = delete;
 		TGaussSeidelWeakConstraintData(TGaussSeidelWeakConstraintData<T>&& Other)
@@ -63,6 +67,8 @@ namespace Chaos::Softs
 			AddArray(&MStiffness);
 			AddArray(&MIsAnisotropic);
 			AddArray(&MNormals);
+			AddArray(&MIsZeroRestLength);
+			AddArray(&MRestLength);
 			Other.MSize = 0;
 		}
 
@@ -88,6 +94,7 @@ namespace Chaos::Softs
 			MSecondWeights[ConstraintIndex] = SingleData.SingleSecondWeights;
 			MNormals[ConstraintIndex] = SingleData.SingleNormal;
 			MIsAnisotropic[ConstraintIndex] = SingleData.bIsAnisotropic;
+			MIsZeroRestLength[ConstraintIndex] = SingleData.bIsZeroRestLength;
 		}
 
 		void AddSingleConstraint(const FGaussSeidelWeakConstraintSingleData<T>& SingleData)
@@ -115,6 +122,7 @@ namespace Chaos::Softs
 			MStiffness = MoveTemp(Other.MStiffness);
 			MIsAnisotropic = MoveTemp(Other.MIsAnisotropic);
 			MNormals = MoveTemp(Other.MNormals);
+			MIsZeroRestLength = MoveTemp(Other.MIsZeroRestLength);
 			ResizeHelper(Other.Size());
 			Other.MSize = 0;
 			return *this;
@@ -180,11 +188,6 @@ namespace Chaos::Softs
 			MSecondWeights[Index] = InWeights;
 		}
 
-		inline const TArrayCollectionArray<bool>& IsAnisotropic() const
-		{
-			return MIsAnisotropic;
-		}
-
 		const bool GetIsAnisotropic(const int32 Index) const
 		{
 			return MIsAnisotropic[Index];
@@ -225,6 +228,21 @@ namespace Chaos::Softs
 			MStiffness[Index] = InStiffness;
 		}
 
+		const bool GetIsZeroRestLength(const int32 Index) const
+		{
+			return MIsZeroRestLength[Index];
+		}
+
+		void SetIsZeroRestLength(const int32 Index, const bool InIsZeroRestLength)
+		{
+			MIsZeroRestLength[Index] = InIsZeroRestLength;
+		}
+
+		void SetRestLength(const int32 Index, const T InRestLength)
+		{
+			MRestLength[Index] = InRestLength;
+		}
+
 		const FGaussSeidelWeakConstraintSingleData<T> GetSingleConstraintData(const int32 ConstraintIndex) const 
 		{
 			FGaussSeidelWeakConstraintSingleData<T> SingleConstraintData;
@@ -238,7 +256,8 @@ namespace Chaos::Softs
 				SingleConstraintData.SingleSecondWeights = MSecondWeights[ConstraintIndex];
 				SingleConstraintData.bIsAnisotropic = MIsAnisotropic[ConstraintIndex];
 				SingleConstraintData.SingleNormal = MNormals[ConstraintIndex];
-
+				SingleConstraintData.bIsZeroRestLength = MIsZeroRestLength[ConstraintIndex];
+				SingleConstraintData.RestLength = MRestLength[ConstraintIndex];
 			}
 			return SingleConstraintData;
 		}
@@ -251,6 +270,8 @@ namespace Chaos::Softs
 		TArrayCollectionArray<T> MStiffness;
 		TArrayCollectionArray<bool> MIsAnisotropic;
 		TArrayCollectionArray<TVector<T, 3>> MNormals;
+		TArrayCollectionArray<bool> MIsZeroRestLength;
+		TArrayCollectionArray<T> MRestLength;
 	};
 
 
@@ -307,130 +328,102 @@ namespace Chaos::Softs
 		{
 			TArray<TArray<int32>> ExtraConstraints;
 			ExtraConstraints.Init(TArray<int32>(), ConstraintsData.Size());
-			for (int32 i = 0; i < ExtraConstraints.Num(); i++)
+			for (int32 ConstraintIdx = 0; ConstraintIdx < ExtraConstraints.Num(); ++ConstraintIdx)
 			{
-				ExtraConstraints[i].SetNum(ConstraintsData.GetIndices(i).Num() + ConstraintsData.GetSecondIndices(i).Num());
-				for (int32 j = 0; j < ConstraintsData.GetIndices(i).Num(); j++)
+				ExtraConstraints[ConstraintIdx].SetNum(ConstraintsData.GetIndices(ConstraintIdx).Num() + ConstraintsData.GetSecondIndices(ConstraintIdx).Num());
+				for (int32 LocalIdx = 0; LocalIdx < ConstraintsData.GetIndices(ConstraintIdx).Num(); ++LocalIdx)
 				{
-					ExtraConstraints[i][j] = ConstraintsData.GetIndices(i)[j];
+					ExtraConstraints[ConstraintIdx][LocalIdx] = ConstraintsData.GetIndices(ConstraintIdx)[LocalIdx];
 				}
-				for (int32 j = 0; j < ConstraintsData.GetSecondIndices(i).Num(); j++)
+				for (int32 LocalSecondIdx = 0; LocalSecondIdx < ConstraintsData.GetSecondIndices(ConstraintIdx).Num(); ++LocalSecondIdx)
 				{
-					ExtraConstraints[i][j+ConstraintsData.GetIndices(i).Num()] = ConstraintsData.GetSecondIndices(i)[j];
+					ExtraConstraints[ConstraintIdx][LocalSecondIdx+ConstraintsData.GetIndices(ConstraintIdx).Num()] = ConstraintsData.GetSecondIndices(ConstraintIdx)[LocalSecondIdx];
 				}
 			}
 			WCIncidentElements = Chaos::Utilities::ComputeIncidentElements(ExtraConstraints, &WCIncidentElementsLocal);
 
-			NodalWeights.Init({}, InParticles.Size());
+			//Update rest state normal and nodal weights
+			UpdateTriangleNormalAndNodalWeight(InParticles, /*bUseParticleX = */true);
+			
+			NoCollisionNodalWeights = NodalWeights;
+			NoCollisionConstraints = ExtraConstraints;
+			InitialWCSize = ConstraintsData.Size();
+
+			NoCollisionWCIncidentElements = WCIncidentElements;
+			NoCollisionWCIncidentElementsLocal = WCIncidentElementsLocal;
 
-			for (int32 i = 0; i < WCIncidentElements.Num(); i++) 
+			// Compute rest length
+			for (int32 ConstraintIdx = 0; ConstraintIdx < ConstraintsData.Size(); ++ConstraintIdx)
 			{
-				if (WCIncidentElements[i].Num() > 0)
+				if (ConstraintsData.GetIsZeroRestLength(ConstraintIdx))
 				{
-					int32 p = ExtraConstraints[WCIncidentElements[i][0]][WCIncidentElementsLocal[i][0]];
-					NodalWeights[p].Init(T(0), 6);
-					for (int32 j = 0; j < WCIncidentElements[i].Num(); j++)
+					ConstraintsData.SetRestLength(ConstraintIdx, 0);
+				}
+				else
+				{
+					const TArray<int32>& Indices = ConstraintsData.GetIndices(ConstraintIdx);
+					const TArray<int32>& SecondIndices = ConstraintsData.GetSecondIndices(ConstraintIdx);
+					const TArray<T>& Weights = ConstraintsData.GetWeights(ConstraintIdx);
+					const TArray<T>& SecondWeights = ConstraintsData.GetSecondWeights(ConstraintIdx);
+					const TVec3<T> RestSpringEdge = ComputeSpringEdge(InParticles, Indices, SecondIndices, Weights, SecondWeights, /*bUseParticleX =*/true);
+					if (ConstraintsData.GetIsAnisotropic(ConstraintIdx))
 					{
-						int32 ConstraintIndex = WCIncidentElements[i][j];
-						int32 LocalIndex = WCIncidentElementsLocal[i][j];
-
-						T Weight = T(0);
-						if (LocalIndex >= ConstraintsData.GetIndices(ConstraintIndex).Num())
-						{
-							Weight = ConstraintsData.GetSecondWeights(ConstraintIndex)[LocalIndex - ConstraintsData.GetIndices(ConstraintIndex).Num()];
-						}
-						else
-						{
-							Weight = ConstraintsData.GetWeights(ConstraintIndex)[LocalIndex];
-						}
-
-						if (ConstraintsData.GetIsAnisotropic(ConstraintIndex))
-						{
-							for (int32 alpha = 0; alpha < 3; alpha++) 
-							{
-								NodalWeights[p][alpha] += ConstraintsData.GetNormal(ConstraintIndex)[alpha] * ConstraintsData.GetNormal(ConstraintIndex)[alpha] * Weight * Weight * ConstraintsData.GetStiffness(ConstraintIndex);
-							}
-
-							NodalWeights[p][3] += ConstraintsData.GetNormal(ConstraintIndex)[0] * ConstraintsData.GetNormal(ConstraintIndex)[1] * Weight * Weight * ConstraintsData.GetStiffness(ConstraintIndex);
-							NodalWeights[p][4] += ConstraintsData.GetNormal(ConstraintIndex)[0] * ConstraintsData.GetNormal(ConstraintIndex)[2] * Weight * Weight * ConstraintsData.GetStiffness(ConstraintIndex);
-							NodalWeights[p][5] += ConstraintsData.GetNormal(ConstraintIndex)[1] * ConstraintsData.GetNormal(ConstraintIndex)[2] * Weight * Weight * ConstraintsData.GetStiffness(ConstraintIndex);
-						}
-						else
-						{
-							for (int32 alpha = 0; alpha < 3; alpha++)
-							{
-								NodalWeights[p][alpha] += Weight * Weight * ConstraintsData.GetStiffness(ConstraintIndex);
-							}
-						}
+						//if the spring is anisotropic, rest length could be negative depending on the normal direction
+						ConstraintsData.SetRestLength(ConstraintIdx, TVec3<T>::DotProduct(ConstraintsData.GetNormal(ConstraintIdx), RestSpringEdge));
+					}
+					else
+					{
+						ConstraintsData.SetRestLength(ConstraintIdx, RestSpringEdge.Size());
 					}
 				}
 			}
-			NoCollisionNodalWeights = NodalWeights;
-			NoCollisionConstraints = ExtraConstraints;
-			InitialWCSize = ConstraintsData.Size();
-
-			NoCollisionWCIncidentElements = WCIncidentElements;
-			NoCollisionWCIncidentElementsLocal = WCIncidentElementsLocal;
 		}
 
-		void AddWCResidual(const ParticleType& InParticles, const int32 p, const T Dt, TVec3<T>& res)
+		TVec3<T> ComputeSpringEdge(const ParticleType& InParticles, const TArray<int32>& LocalIndices,
+			const TArray<int32>& LocalSecondIndices, const TArray<T>& Weight, const TArray<T>& SecondWeight, bool bUseParticleX) const
 		{
-			int32 WCIndex = Particle2WCIndices[p];
-			if (WCIndex != INDEX_NONE)
+			TVec3<T> SpringEdge((T)0.);
+			if (ensure(LocalIndices.Num() == Weight.Num() && LocalSecondIndices.Num() == SecondWeight.Num()))
 			{
-				for (int32 j = 0; j < WCIncidentElements[WCIndex].Num(); j++)
-				{
-					int32 ConstraintIndex = WCIncidentElements[WCIndex][j];
-					int32 LocalIndex = WCIncidentElementsLocal[WCIndex][j];
-					TVec3<T> SpringEdge((T)0.);
-					const TArray<T>& Weight = ConstraintsData.GetWeights(ConstraintIndex);
-					const TArray<T>& SecondWeight = ConstraintsData.GetSecondWeights(ConstraintIndex);
-					const TArray<int32>& LocalIndices = ConstraintsData.GetIndices(ConstraintIndex);
-					const TArray<int32>& LocalSecondIndices = ConstraintsData.GetSecondIndices(ConstraintIndex);
-					for (int32 l = 0; l < Weight.Num(); l++)
+				if (bUseParticleX)
+				{
+					for (int32 Idx = 0; Idx < LocalIndices.Num(); ++Idx)
 					{
 						for (int32 beta = 0; beta < 3; beta++)
 						{
-							SpringEdge[beta] += Weight[l] * InParticles.P(LocalIndices[l])[beta];
+							SpringEdge[beta] += Weight[Idx] * InParticles.X(LocalIndices[Idx])[beta];
 						}
 					}
-					for (int32 l = 0; l < LocalSecondIndices.Num(); l++)
+					for (int32 SecondIdx = 0; SecondIdx < LocalSecondIndices.Num(); ++SecondIdx)
 					{
 						for (int32 beta = 0; beta < 3; beta++)
 						{
-							SpringEdge[beta] -= SecondWeight[l] * InParticles.P(LocalSecondIndices[l])[beta];
+							SpringEdge[beta] -= SecondWeight[SecondIdx] * InParticles.X(LocalSecondIndices[SecondIdx])[beta];
 						}
 					}
-					T weight = T(0);
-					if (LocalIndex >= LocalIndices.Num())
-					{
-						weight = -SecondWeight[LocalIndex - LocalIndices.Num()];
-					}
-					else
-					{
-						weight = Weight[LocalIndex];
-					}
-					if (ConstraintsData.GetIsAnisotropic(ConstraintIndex))
+				}
+				else
+				{
+					for (int32 Idx = 0; Idx < LocalIndices.Num(); ++Idx)
 					{
-						T comp = TVec3<T>::DotProduct(SpringEdge, ConstraintsData.GetNormal(ConstraintIndex));
-						TVec3<T> proj = ConstraintsData.GetNormal(ConstraintIndex) * comp;
-						for (int32 alpha = 0; alpha < 3; alpha++)
+						for (int32 beta = 0; beta < 3; beta++)
 						{
-							res[alpha] += Dt * Dt * ConstraintsData.GetStiffness(ConstraintIndex) * proj[alpha] * weight;
+							SpringEdge[beta] += Weight[Idx] * InParticles.P(LocalIndices[Idx])[beta];
 						}
 					}
-					else
+					for (int32 SecondIdx = 0; SecondIdx < LocalSecondIndices.Num(); ++SecondIdx)
 					{
-						for (int32 alpha = 0; alpha < 3; alpha++)
+						for (int32 beta = 0; beta < 3; beta++)
 						{
-							res[alpha] += Dt * Dt * ConstraintsData.GetStiffness(ConstraintIndex) * SpringEdge[alpha] * weight;
+							SpringEdge[beta] -= SecondWeight[SecondIdx] * InParticles.P(LocalSecondIndices[SecondIdx])[beta];
 						}
 					}
 				}
 			}
+			return SpringEdge;
 		}
 
-		void AddWCHessian(const int32 p, const T Dt, Chaos::PMatrix<T, 3, 3>& ParticleHessian)
+		void AddWCHessian(const int32 p, const T Dt, Chaos::PMatrix<T, 3, 3>& ParticleHessian) const
 		{
 			if (NodalWeights[p].Num() > 0)
 			{
@@ -458,7 +451,9 @@ namespace Chaos::Softs
 								const TArray<TArray<T>>& InWeights,
 								const TArray<T>& InStiffness,
 								const TArray<TArray<int32>>& InSecondIndices,
-								const TArray<TArray<T>>& InSecondWeights)
+								const TArray<TArray<T>>& InSecondWeights,
+								const TArray<bool>& InIsAnisotrpic,
+								const TArray<bool>& InIsZeroRestLength)
 		{
 			const int32 Offset = ConstraintsData.Size();
 
@@ -472,6 +467,8 @@ namespace Chaos::Softs
 				SingleConstraintData.SingleWeights = InWeights[i];
 				SingleConstraintData.SingleSecondWeights = InSecondWeights[i];
 				SingleConstraintData.SingleStiffness = InStiffness[i];
+				SingleConstraintData.bIsAnisotropic = InIsAnisotrpic[i];
+				SingleConstraintData.bIsZeroRestLength = InIsZeroRestLength[i];
 				ConstraintsData.SetSingleConstraint(SingleConstraintData, i + Offset);
 			}
 		}
@@ -494,7 +491,7 @@ namespace Chaos::Softs
 				ensureMsgf(OriginalConstraintsData.GetIndices(i).Num() == 3, TEXT("Collision format is not point-triangle"));
 				ensureMsgf(OriginalConstraintsData.GetSecondIndices(i).Num() == 1, TEXT("Collision format is not point-triangle"));
 				Chaos::TVector<float, 3> TriPos0(Particles.P(IndicesTemp[0])), TriPos1(Particles.P(IndicesTemp[1])), TriPos2(Particles.P(IndicesTemp[2])), ParticlePos(Particles.P(SecondIndicesTemp[0]));
-				Chaos::TVector<T, 3> Normal = FVector3f::CrossProduct(TriPos1 - TriPos0, TriPos2 - TriPos0);
+				Chaos::TVector<T, 3> Normal = FVector3f::CrossProduct(TriPos2 - TriPos0, TriPos1 - TriPos0); //triangle normal convention (see FTriangleMesh::GetFaceNormals())
 				if (FVector3f::DotProduct(ParticlePos - TriPos0, Normal) < 0.f) //not resolved, keep the spring
 				{
 					ConstraintsData.AddConstraints(OriginalConstraintsData.GetSingleConstraintData(i));
@@ -548,13 +545,82 @@ namespace Chaos::Softs
 #endif
 		}
 
-		void Init(const FSolverParticles& InParticles, const T Dt) const
+		void Init(const FSolverParticles& InParticles, const T Dt)
 		{
+			UpdateTriangleNormalAndNodalWeight(InParticles, /*bUseParticleX =*/ false);
 			if (DebugDrawParams.bVisualizeBindings)
 			{
 				VisualizeAllBindings(InParticles, Dt);
 			}
+		}
+
+		void UpdateTriangleNormalAndNodalWeight(const FSolverParticles& InParticles, bool bUseParticleX)
+		{
+			for (int32 i = 0; i < ConstraintsData.Size(); i++)
+			{
+				if (ConstraintsData.GetIsAnisotropic(i))
+				{
+					const TArray<int32>& IndicesTemp = ConstraintsData.GetIndices(i);
+					const TArray<int32>& SecondIndicesTemp = ConstraintsData.GetSecondIndices(i);
+					ensureMsgf(ConstraintsData.GetIndices(i).Num() == 3, TEXT("Collision format is not point-triangle"));
+					ensureMsgf(ConstraintsData.GetSecondIndices(i).Num() == 1, TEXT("Collision format is not point-triangle"));
+					Chaos::TVector<T, 3> Normal;
+					if (bUseParticleX)
+					{
+						Chaos::TVector<float, 3> TriPos0(InParticles.X(IndicesTemp[0])), TriPos1(InParticles.X(IndicesTemp[1])), TriPos2(InParticles.X(IndicesTemp[2]));
+						Normal = FVector3f::CrossProduct(TriPos2 - TriPos0, TriPos1 - TriPos0).GetSafeNormal(); //triangle normal convention (see FTriangleMesh::GetFaceNormals())
+					}
+					else
+					{
+						Chaos::TVector<float, 3> TriPos0(InParticles.P(IndicesTemp[0])), TriPos1(InParticles.P(IndicesTemp[1])), TriPos2(InParticles.P(IndicesTemp[2]));
+						Normal = FVector3f::CrossProduct(TriPos2 - TriPos0, TriPos1 - TriPos0).GetSafeNormal(); //triangle normal convention (see FTriangleMesh::GetFaceNormals())
+					}
+					ConstraintsData.SetNormal(i, Normal);
+				}
+			}
+
+			NodalWeights.Init({}, InParticles.Size());
+			for (int32 p = 0; p < WCIncidentElements.Num(); ++p)
+			{
+				if (WCIncidentElements[p].Num() > 0)
+				{
+					NodalWeights[p].Init(T(0), 6);
+					for (int32 j = 0; j < WCIncidentElements[p].Num(); j++)
+					{
+						int32 ConstraintIndex = WCIncidentElements[p][j];
+						int32 LocalIndex = WCIncidentElementsLocal[p][j];
 
+						T Weight = T(0);
+						if (LocalIndex >= ConstraintsData.GetIndices(ConstraintIndex).Num())
+						{
+							Weight = ConstraintsData.GetSecondWeights(ConstraintIndex)[LocalIndex - ConstraintsData.GetIndices(ConstraintIndex).Num()];
+						}
+						else
+						{
+							Weight = ConstraintsData.GetWeights(ConstraintIndex)[LocalIndex];
+						}
+
+						if (ConstraintsData.GetIsAnisotropic(ConstraintIndex))
+						{
+							for (int32 alpha = 0; alpha < 3; alpha++)
+							{
+								NodalWeights[p][alpha] += ConstraintsData.GetNormal(ConstraintIndex)[alpha] * ConstraintsData.GetNormal(ConstraintIndex)[alpha] * Weight * Weight * ConstraintsData.GetStiffness(ConstraintIndex);
+							}
+
+							NodalWeights[p][3] += ConstraintsData.GetNormal(ConstraintIndex)[0] * ConstraintsData.GetNormal(ConstraintIndex)[1] * Weight * Weight * ConstraintsData.GetStiffness(ConstraintIndex);
+							NodalWeights[p][4] += ConstraintsData.GetNormal(ConstraintIndex)[0] * ConstraintsData.GetNormal(ConstraintIndex)[2] * Weight * Weight * ConstraintsData.GetStiffness(ConstraintIndex);
+							NodalWeights[p][5] += ConstraintsData.GetNormal(ConstraintIndex)[1] * ConstraintsData.GetNormal(ConstraintIndex)[2] * Weight * Weight * ConstraintsData.GetStiffness(ConstraintIndex);
+						}
+						else
+						{
+							for (int32 alpha = 0; alpha < 3; alpha++)
+							{
+								NodalWeights[p][alpha] += Weight * Weight * ConstraintsData.GetStiffness(ConstraintIndex);
+							}
+						}
+					}
+				}
+			}
 		}
 
 		//CollisionDetectionSpatialHash should be faster than CollisionDetectionBVH
@@ -681,7 +747,8 @@ namespace Chaos::Softs
 					SingleConstraintData.SecondWeights = {T(1.f)};
 					SingleConstraintData.bIsAnisotropic = UseAnisotropicSpring;
 					SingleConstraintData.SingleNormal = FaceNormal;
-					
+					SingleConstraintData.bIsZeroRestLength = true; //Push-out type collision springs should be zero rest length
+
 					float SpringStiffness = 0.f;
 					for (int32 k = 0; k < 3; k++)
 					{
@@ -730,6 +797,7 @@ namespace Chaos::Softs
 								SingleConstraintData.SingleSecondWeights = {T(1.f)};
 								SingleConstraintData.bIsAnisotropic = UseAnisotropicSpring;
 								SingleConstraintData.SingleNormal = CollisionPoint.Normal;
+								SingleConstraintData.bIsZeroRestLength = true; //Push-out type collision springs should be zero rest length
 
 								float SpringStiffness = 0.f;
 								for (int32 k = 0; k < 3; k++)
@@ -785,6 +853,7 @@ namespace Chaos::Softs
 								SingleConstraintData.SingleSecondWeights = {T(1.f)};
 								SingleConstraintData.bIsAnisotropic = UseAnisotropicSpring;
 								SingleConstraintData.SingleNormal = CollisionPoint.Normal;
+								SingleConstraintData.bIsZeroRestLength = true; //Push-out type collision springs should be zero rest length
 
 								float SpringStiffness = 0.f;
 								for (int32 k = 0; k < 3; k++)
@@ -878,14 +947,14 @@ namespace Chaos::Softs
 		}
 
 
-		TArray<TArray<int32>> GetStaticConstraintArrays(TArray<TArray<int32>>& IncidentElements, TArray<TArray<int32>>& IncidentElementsLocal)
+		const TArray<TArray<int32>>& GetStaticConstraintArrays(TArray<TArray<int32>>& IncidentElements, TArray<TArray<int32>>& IncidentElementsLocal) const
 		{
 			IncidentElements = NoCollisionWCIncidentElements;
 			IncidentElementsLocal = NoCollisionWCIncidentElementsLocal;
 			return NoCollisionConstraints;
 		}
 
-		TArray<TArray<int32>> GetDynamicConstraintArrays(TArray<TArray<int32>>& IncidentElements, TArray<TArray<int32>>& IncidentElementsLocal)
+		TArray<TArray<int32>> GetDynamicConstraintArrays(TArray<TArray<int32>>& IncidentElements, TArray<TArray<int32>>& IncidentElementsLocal) const
 		{
 			TArray<TArray<int32>> ExtraConstraints;
 			ExtraConstraints.Init(TArray<int32>(), ConstraintsData.Size());
@@ -909,25 +978,13 @@ namespace Chaos::Softs
 			return ExtraConstraints;
 		}
 
-		void AddWCResidualAndHessian(const ParticleType& InParticles, const int32 ConstraintIndex, const int32 LocalIndex, const T Dt, TVec3<T>& ParticleResidual, Chaos::PMatrix<T, 3, 3>& ParticleHessian)
+		// Deprecated, now replaced with AddWCResidual for more general cases
+		void AddZeroRestLengthWCResidualAndHessian(const ParticleType& InParticles, const int32 ConstraintIndex, const int32 LocalIndex, const T Dt, TVec3<T>& ParticleResidual, Chaos::PMatrix<T, 3, 3>& ParticleHessian) const
 		{
 			const FGaussSeidelWeakConstraintSingleData<T>& SingleData = ConstraintsData.GetSingleConstraintData(ConstraintIndex);
 
-			TVec3<T> SpringEdge((T)0.);
-			for (int32 l = 0; l < SingleData.SingleWeights.Num(); l++)
-			{
-				for (int32 beta = 0; beta < 3; beta++)
-				{
-					SpringEdge[beta] += SingleData.SingleWeights[l] * InParticles.P(SingleData.SingleIndices[l])[beta];
-				}
-			}
-			for (int32 l = 0; l < SingleData.SingleSecondWeights.Num(); l++)
-			{
-				for (int32 beta = 0; beta < 3; beta++)
-				{
-					SpringEdge[beta] -= SingleData.SingleSecondWeights[l] * InParticles.P(SingleData.SingleSecondIndices[l])[beta];
-				}
-			}
+			const TVec3<T> SpringEdge = ComputeSpringEdge(InParticles, SingleData.SingleIndices, SingleData.SingleSecondIndices,
+				SingleData.SingleWeights, SingleData.SingleSecondWeights, /*bUseParticleX =*/false);
 			T weight = T(0);
 			if (LocalIndex >= SingleData.SingleIndices.Num())
 			{
@@ -953,13 +1010,46 @@ namespace Chaos::Softs
 					ParticleResidual[alpha] += Dt * Dt * SingleData.SingleStiffness * SpringEdge[alpha] * weight;
 				}
 			}
+		}
+
+		void AddWCResidual(const ParticleType& InParticles, const int32 ConstraintIndex, const int32 LocalIndex, const T Dt, TVec3<T>& ParticleResidual, Chaos::PMatrix<T, 3, 3>& ParticleHessian) const
+		{
+			const FGaussSeidelWeakConstraintSingleData<T>& SingleData = ConstraintsData.GetSingleConstraintData(ConstraintIndex);
 
+			const TVec3<T> SpringEdge = ComputeSpringEdge(InParticles, SingleData.SingleIndices, SingleData.SingleSecondIndices,
+				SingleData.SingleWeights, SingleData.SingleSecondWeights, /*bUseParticleX =*/false);
+			T weight = T(0);
+			if (LocalIndex >= SingleData.SingleIndices.Num())
+			{
+				weight = -SingleData.SingleSecondWeights[LocalIndex - SingleData.SingleIndices.Num()];
+			}
+			else
+			{
+				weight = SingleData.SingleWeights[LocalIndex];
+			}
+			TVec3<T> Projection;
+			if (SingleData.bIsAnisotropic)
+			{
+				T LengthDiff = TVec3<T>::DotProduct(SpringEdge, SingleData.SingleNormal) - SingleData.RestLength;
+				Projection = SingleData.SingleNormal * LengthDiff;
+			}
+			else
+			{
+				Projection = SpringEdge;
+				if (!SingleData.bIsZeroRestLength) // if not zero rest-length, apply repulsion force
+				{
+					Projection -= SingleData.RestLength * SpringEdge.GetSafeNormal();
+				}	
+			}
+			for (int32 alpha = 0; alpha < 3; ++alpha)
+			{
+				ParticleResidual[alpha] += Dt * Dt * SingleData.SingleStiffness * Projection[alpha] * weight;
+			}
 		}
 
 		TGaussSeidelWeakConstraintData<T> ConstraintsData;
 
 		TArray<TArray<T>> NodalWeights;
-		TArray<int32> Particle2WCIndices;
 
 		TArray<TArray<int32>> WCIncidentElements;
 		TArray<TArray<int32>> WCIncidentElementsLocal;
diff --git a/Engine/Source/Runtime/Experimental/Chaos/Public/GeometryCollection/Facades/CollectionPositionTargetFacade.h b/Engine/Source/Runtime/Experimental/Chaos/Public/GeometryCollection/Facades/CollectionPositionTargetFacade.h
--- a/Engine/Source/Runtime/Experimental/Chaos/Public/GeometryCollection/Facades/CollectionPositionTargetFacade.h
+++ b/Engine/Source/Runtime/Experimental/Chaos/Public/GeometryCollection/Facades/CollectionPositionTargetFacade.h
@@ -20,9 +20,10 @@ namespace GeometryCollection::Facades
 		FString TargetName;
 		UE_DEPRECATED(5.6, "SourceName will be removed and is no longer populated.")
 		FString SourceName;
-		float Stiffness;
-		float Damping;
-
+		float Stiffness = 0.f;
+		float Damping = 0.f;
+		bool bIsAnisotropic = false;
+		bool bIsZeroRestLength = false;
 		PRAGMA_DISABLE_DEPRECATION_WARNINGS
 		FPositionTargetsData() = default;
 		FPositionTargetsData(const FPositionTargetsData&) = default;
@@ -53,7 +54,9 @@ namespace GeometryCollection::Facades
 		static CHAOS_API const FName TargetName;
 		static CHAOS_API const FName TargetWeights;
 		static CHAOS_API const FName SourceWeights;
-
+		static CHAOS_API const FName IsAnisotropic;
+		static CHAOS_API const FName IsZeroRestLength;
+		
 		CHAOS_API FPositionTargetFacade(FManagedArrayCollection& InCollection, const FName& InVerticesGroup = FGeometryCollection::VerticesGroup);
 		CHAOS_API FPositionTargetFacade(const FManagedArrayCollection& InCollection, const FName& InVerticesGroup = FGeometryCollection::VerticesGroup);
 
@@ -80,7 +83,7 @@ namespace GeometryCollection::Facades
 		/** Remove position targets with invalid indices*/
 		CHAOS_API int32 RemoveInvalidPositionTarget();
 		/** Remove position targets between two groups of vertices*/
-		CHAOS_API int32 RemovePositionTarget(TFunctionRef<bool(const int32 VertexIdx)> IsVertexGroup1, TFunctionRef<bool(const int32 VertexIdx)> IsVertexGroup2);
+		CHAOS_API int32 RemovePositionTargetBetween(TFunctionRef<bool(const int32 VertexIdx)> IsVertexGroup1, TFunctionRef<bool(const int32 VertexIdx)> IsVertexGroup2);
 
 	protected:
 		const FManagedArrayCollection& ConstCollection;
@@ -93,5 +96,7 @@ namespace GeometryCollection::Facades
 		TManagedArrayAccessor<float> DampingAttribute;
 		TManagedArrayAccessor<TArray<float>> TargetWeightsAttribute;
 		TManagedArrayAccessor<TArray<float>> SourceWeightsAttribute;
+		TManagedArrayAccessor<bool> IsAnisotropicAttribute;
+		TManagedArrayAccessor<bool> IsZeroRestLengthAttribute;
 	};
 }
diff --git a/Engine/Source/Runtime/Experimental/Chaos/Public/GeometryCollection/Facades/CollectionVolumeConstraintFacade.h b/Engine/Source/Runtime/Experimental/Chaos/Public/GeometryCollection/Facades/CollectionVolumeConstraintFacade.h
new file mode 100644
--- /dev/null
+++ b/Engine/Source/Runtime/Experimental/Chaos/Public/GeometryCollection/Facades/CollectionVolumeConstraintFacade.h
@@ -0,0 +1,54 @@
+// Copyright Epic Games, Inc. All Rights Reserved.
+
+#pragma once
+
+#include "GeometryCollection/ManagedArrayAccessor.h"
+#include "GeometryCollection/ManagedArrayCollection.h"
+#include "GeometryCollection/GeometryCollection.h"
+
+namespace GeometryCollection::Facades
+{
+	/** Volume Constraint Facade */
+	class FVolumeConstraintFacade
+	{
+	public:
+		// Attributes
+		static CHAOS_API const FName GroupName;
+		static CHAOS_API const FName VolumeIndex;
+		static CHAOS_API const FName Stiffness;
+
+		CHAOS_API FVolumeConstraintFacade(FManagedArrayCollection& InCollection);
+		CHAOS_API FVolumeConstraintFacade(const FManagedArrayCollection& InCollection);
+	
+		FVolumeConstraintFacade(FVolumeConstraintFacade&&) = default;
+
+		/** Create the facade attributes. */
+		CHAOS_API void DefineSchema();
+
+		/** Is the facade defined constant. */
+		bool IsConst() const { return Collection == nullptr; }
+
+		/** Is the Facade defined on the collection? */
+		CHAOS_API bool IsValid() const;
+
+		//
+		//  Accessors of constraints
+		//
+		CHAOS_API int32 AddVolumeConstraint(const FIntVector4& NewVolumeIndex, float NewStiffness);
+		CHAOS_API FIntVector4 GetVolumeIndex(const int32 AttributeIndex) const;
+		CHAOS_API float GetStiffness(const int32 AttributeIndex) const;
+		int32 NumVolumeConstraints() const { return VolumeIndexAttribute.Num(); }
+		
+		/** Remove volume constraint with invalid indices*/
+		CHAOS_API int32 RemoveInvalidVolumeConstraint();
+		/** Remove volume constraint between two groups of vertices*/
+		CHAOS_API int32 RemoveVolumeConstraintBetween(TFunctionRef<bool(const int32 VertexIdx)> IsVertexGroup1, TFunctionRef<bool(const int32 VertexIdx)> IsVertexGroup2);
+
+	protected:
+		const FManagedArrayCollection& ConstCollection;
+		FManagedArrayCollection* Collection = nullptr;
+
+		TManagedArrayAccessor<FIntVector4> VolumeIndexAttribute;
+		TManagedArrayAccessor<float> StiffnessAttribute;
+	};
+}
