CIS fix: replace TMathUtilConstants with numeric_limits

diff --git a/Engine/Source/Runtime/Experimental/Chaos/Public/Chaos/PBDAltitudeSpringConstraints.h b/Engine/Source/Runtime/Experimental/Chaos/Public/Chaos/PBDAltitudeSpringConstraints.h
--- a/Engine/Source/Runtime/Experimental/Chaos/Public/Chaos/PBDAltitudeSpringConstraints.h
+++ b/Engine/Source/Runtime/Experimental/Chaos/Public/Chaos/PBDAltitudeSpringConstraints.h
@@ -3,7 +3,6 @@
 
 #include "Chaos/PBDSpringConstraintsBase.h"
 #include "ChaosStats.h"
-#include "Chaos/ImplicitQRSVD.h"
 #include "Chaos/GraphColoring.h"
 #include "Chaos/Framework/Parallel.h"
 
@@ -243,7 +242,7 @@ namespace Chaos::Softs
 		int32 GetSingleElementShortestLength(FSolverParticles& Particles, const FSolverReal Dt, const int32 ElementIndex) const 
 		{	
 			int32 ReturnIndex = -1;
-			FSolverReal MinConstraint = TMathUtilConstants<FSolverReal>::MaxReal;
+			FSolverReal MinConstraint = TNumericLimits<FSolverReal>::Max();
 			const TVec4<int32>& Constraint = Constraints[ElementIndex];
 			for (int ie = 0; ie < 4; ++ie)
 			{
