[NetToken] Prevent compiler from optimizing out inline static that Iris uses for identifying SerializerRegistryDelegates.

diff --git a/Engine/Source/Runtime/Engine/Public/Net/Iris/ReplicationSystem/StructNetTokenDataStoreHelper.h b/Engine/Source/Runtime/Engine/Public/Net/Iris/ReplicationSystem/StructNetTokenDataStoreHelper.h
--- a/Engine/Source/Runtime/Engine/Public/Net/Iris/ReplicationSystem/StructNetTokenDataStoreHelper.h
+++ b/Engine/Source/Runtime/Engine/Public/Net/Iris/ReplicationSystem/StructNetTokenDataStoreHelper.h
@@ -44,6 +44,11 @@ public:
 		{
 			DataStoreType::NetSerializeScriptDelegate.BindStatic(&SerializeHelper);
 		}
+
+		if (!TStructNetTokenDataStore<T>::NetSerializerRegistryDelegates.HasPostFreezeBeenCalled())
+		{
+			UE_LOG(LogNetToken, Warning, TEXT("NetSerializer registries have not been initialized for %s"), *DataStoreType::GetTokenStoreName().ToString());
+		}
 		
 		if (Ar.IsSaving())
 		{
diff --git a/Engine/Source/Runtime/Experimental/Iris/Core/Public/Iris/ReplicationSystem/StructNetTokenDataStore.h b/Engine/Source/Runtime/Experimental/Iris/Core/Public/Iris/ReplicationSystem/StructNetTokenDataStore.h
--- a/Engine/Source/Runtime/Experimental/Iris/Core/Public/Iris/ReplicationSystem/StructNetTokenDataStore.h
+++ b/Engine/Source/Runtime/Experimental/Iris/Core/Public/Iris/ReplicationSystem/StructNetTokenDataStore.h
@@ -338,13 +338,21 @@ private:
 	// Helper for NetSerializerDescriptor setup
 	class FNetSerializerRegistryDelegates final : private UE::Net::FNetSerializerRegistryDelegates
 	{
+	public:
+		bool HasPostFreezeBeenCalled() const { return bPostFreezeHasBeenCalled; }
+		
 	private:
+		bool bPostFreezeHasBeenCalled = false;
 		virtual void OnPostFreezeNetSerializerRegistry() override
 		{
+			bPostFreezeHasBeenCalled = true;
+			UE_LOG(LogNetToken, VeryVerbose, TEXT("TStructNetTokenDataStore::OnPostFreezeNetSerializerRegistry %s"), *GetTokenStoreName().ToString());
 			if (!Struct)
 			{
 				Struct = StaticStruct<T>();
 			}
+			UE_LOG(LogNetToken, VeryVerbose, TEXT("TStructNetTokenDataStore::OnPostFreezeNetSerializerRegistry %s Struct: %s"), *GetTokenStoreName().ToString(), *GetNameSafe(Struct));
+			
 			if (!Descriptor.IsValid())
 			{				
 				FReplicationStateDescriptorBuilder::FParameters Params;
@@ -353,11 +361,15 @@ private:
 				Params.SkipCheckForCustomNetSerializerForStruct = 1U;
 
 				Descriptor = FReplicationStateDescriptorBuilder::CreateDescriptorForStruct(Struct, Params);
+				UE_LOG(LogNetToken, VeryVerbose, TEXT("TStructNetTokenDataStore::OnPostFreezeNetSerializerRegistry %s Descriptor: %s"), *GetTokenStoreName().ToString(), Descriptor ? TEXT("Exists") : TEXT("NULL"));
 			}
 		}
 	};
-
+public:
+	// Compiler tries to optimize out the NetSerializer since it isn't obvious that anything is referencing it during module compilation.
+	UE_DISABLE_OPTIMIZATION_SHIP
 	inline static TStructNetTokenDataStore<T>::FNetSerializerRegistryDelegates NetSerializerRegistryDelegates;
+	UE_ENABLE_OPTIMIZATION_SHIP
 };
 	
 // Helper to implement a NetSerializer for a struct that should serialize using a NetTokenStore
