Dataflow async execution
- Add an option to the evaluation menu for enabling async ( off by default, not saved in the settings yet)
- Add an async evaluator class to take care of the task collection and dispatch
- Change evaluation in the dataflow command to use an async form ( pass a post evluation lambda )
- Add an async evaluation widget with a progress bar and stats about the tasks
- Add a blinking text above the nodes currently evaluating  asynchronously
- Make sure the render callback as are not pulling directly on the graph when executing asynchronously to ensure responsive UI
- DataflowNode now inherit from TSharedFromThis
- Add virtual function for certain nodes to ask to be executed on the gamethread
- For now by default all nodes with a StaticMesh or SkeletalMesh input will be executed on the gamethread

diff --git a/Engine/Plugins/Experimental/Dataflow/Source/DataflowEditor/Private/Dataflow/DataflowEditorCommands.cpp b/Engine/Plugins/Experimental/Dataflow/Source/DataflowEditor/Private/Dataflow/DataflowEditorCommands.cpp
--- a/Engine/Plugins/Experimental/Dataflow/Source/DataflowEditor/Private/Dataflow/DataflowEditorCommands.cpp
+++ b/Engine/Plugins/Experimental/Dataflow/Source/DataflowEditor/Private/Dataflow/DataflowEditorCommands.cpp
@@ -41,10 +41,6 @@ const FString FDataflowEditorCommandsImpl::PauseSimulationSceneIdentifier = TEXT
 const FString FDataflowEditorCommandsImpl::StartSimulationSceneIdentifier = TEXT("StartSimulationScene");
 const FString FDataflowEditorCommandsImpl::StepSimulationSceneIdentifier = TEXT("StepSimulationScene");
 
-// @todo(brice) Remove Example Tools
-//const FString FDataflowEditorCommandsImpl::BeginAttributeEditorToolIdentifier = TEXT("BeginAttributeEditorTool");
-//const FString FDataflowEditorCommandsImpl::BeginMeshSelectionToolIdentifier = TEXT("BeginMeshSelectionTool");
-
 FDataflowEditorCommandsImpl::FDataflowEditorCommandsImpl()
 	: TBaseCharacterFXEditorCommands<FDataflowEditorCommandsImpl>("DataflowEditor", 
 		LOCTEXT("ContextDescription", "Dataflow Editor"), 
@@ -63,6 +59,7 @@ void FDataflowEditorCommandsImpl::RegisterCommands()
 	UI_COMMAND(EvaluateGraphManual, "Manual Graph Evaluation", "Set the evaluation mode of the graph to Manual.", EUserInterfaceActionType::RadioButton, FInputChord());
 	UI_COMMAND(ClearGraphCache, "Clear Graph Cache", "Clear the graph Data cache.", EUserInterfaceActionType::Button, FInputChord());
 	UI_COMMAND(TogglePerfData, "Performance Data", "Toggle the evaluation performance data for each node.", EUserInterfaceActionType::ToggleButton, FInputChord());
+	UI_COMMAND(ToggleAsyncEvaluation, "Asynchronous Evaluation (Experimental)", "Toggle asynchronous evaluation of the graph. This is an experimental feature.", EUserInterfaceActionType::ToggleButton, FInputChord());
 	UI_COMMAND(FreezeNodes, "FreezeNodes", "Freeze the evaluation of the selected nodes.", EUserInterfaceActionType::Button, FInputChord());
 	UI_COMMAND(UnfreezeNodes, "UnfreezeNodes", "Unfreeze the evaluation of the selected nodes.", EUserInterfaceActionType::Button, FInputChord());
 	UI_COMMAND(CreateComment, "CreateComment", "Create a Comment node.", EUserInterfaceActionType::None, FInputChord());
@@ -188,22 +185,77 @@ const FDataflowNode* FDataflowEditorCommands::EvaluateNode(UE::Dataflow::FContex
 		// Note: If the node is deactivated and has any outputs, then these outputs might still need to be forwarded.
 		//       Therefore the Evaluate method has to be called for whichever value of bActive.
 		//       This however isn't the case of SetAssetValue() for which the active state needs to be checked before the call.
-		Context.Evaluate(Node, Output);
 
-		if (const FDataflowTerminalNode* const TerminalNode = Node->AsType<const FDataflowTerminalNode>())
-		{
-			if (TerminalNode->IsActive() && Asset)
+		TWeakPtr<const FDataflowNode> WeakNode{ Node->AsWeak() };
+		TWeakObjectPtr<UObject> WeakAsset{ Asset };
+		auto OnPostEvaluation =
+			[WeakNode, WeakAsset](UE::Dataflow::FContext& Context)
 			{
-				UE_LOG(LogChaosDataflow, Verbose, TEXT("FDataflowTerminalNode::SetAssetValue(): TerminalNode [%s], Asset [%s]"), *TerminalNode->GetName().ToString(), *Asset->GetName());
-				TerminalNode->SetAssetValue(Asset, Context);
-			}
-		}
+				if (TSharedPtr<const FDataflowNode> Node = WeakNode.Pin())
+				{
+					if (const FDataflowTerminalNode* const TerminalNode = Node->AsType<const FDataflowTerminalNode>())
+					{
+						if (TerminalNode->IsActive())
+						{
+							if (TStrongObjectPtr<UObject> Asset = WeakAsset.Pin())
+							{
+								UE_LOG(LogChaosDataflow, Verbose, TEXT("FDataflowTerminalNode::SetAssetValue(): TerminalNode [%s], Asset [%s]"), *TerminalNode->GetName().ToString(), *Asset->GetName());
+								TerminalNode->SetAssetValue(Asset.Get(), Context);
+							}
+						}
+					}
+				}
+			};
+
+		Context.Evaluate(Node, Output, OnPostEvaluation);
 
 		InOutLastNodeTimestamp = Node->GetTimestamp();
 	}
 	return Node;
 }
 
+void FDataflowEditorCommands::EvaluateNode(UE::Dataflow::FContext& Context, const FDataflowNode& Node, const FDataflowOutput* Output, UObject* Asset, UE::Dataflow::FTimestamp& InOutLastNodeTimestamp, UE::Dataflow::FOnPostEvaluationFunction OnEvaluationCompleted)
+{
+	UE_LOG(LogChaosDataflow, VeryVerbose, TEXT("FDataflowEditorCommands::EvaluateNode() : Node [%s], Output [%s]"), *Node.GetName().ToString(), Output ? *Output->GetName().ToString() : TEXT("nullptr"));
+
+	if (InOutLastNodeTimestamp < Node.GetTimestamp())
+	{
+		// Note: If the node is deactivated and has any outputs, then these outputs might still need to be forwarded.
+		//       Therefore the Evaluate method has to be called for whichever value of bActive.
+		//       This however isn't the case of SetAssetValue() for which the active state needs to be checked before the call.
+
+		TWeakPtr<const FDataflowNode> WeakNode{ Node.AsWeak() };
+		TWeakObjectPtr<UObject> WeakAsset{ Asset };
+		auto OnPostEvaluation =
+			[WeakNode, WeakAsset, OnEvaluationCompleted](UE::Dataflow::FContext& Context)
+			{
+				if (TSharedPtr<const FDataflowNode> Node = WeakNode.Pin())
+				{
+					if (const FDataflowTerminalNode* const TerminalNode = Node->AsType<const FDataflowTerminalNode>())
+					{
+						if (TerminalNode->IsActive())
+						{
+							if (TStrongObjectPtr<UObject> Asset = WeakAsset.Pin())
+							{
+								UE_LOG(LogChaosDataflow, Verbose, TEXT("FDataflowTerminalNode::SetAssetValue(): TerminalNode [%s], Asset [%s]"), *TerminalNode->GetName().ToString(), *Asset->GetName());
+								TerminalNode->SetAssetValue(Asset.Get(), Context);
+							}
+						}
+					}
+				}
+				OnEvaluationCompleted(Context);
+			};
+
+		Context.Evaluate(&Node, Output, OnPostEvaluation);
+
+		InOutLastNodeTimestamp = Node.GetTimestamp();
+	}
+	else
+	{
+		OnEvaluationCompleted(Context);
+	}
+}
+
 bool FDataflowEditorCommands::OnNodeVerifyTitleCommit(const FText& NewText, UEdGraphNode* GraphNode, FText& OutErrorMessage)
 {
 	const FString NewString = NewText.ToString();
diff --git a/Engine/Plugins/Experimental/Dataflow/Source/DataflowEditor/Private/Dataflow/DataflowEditorToolkit.cpp b/Engine/Plugins/Experimental/Dataflow/Source/DataflowEditor/Private/Dataflow/DataflowEditorToolkit.cpp
--- a/Engine/Plugins/Experimental/Dataflow/Source/DataflowEditor/Private/Dataflow/DataflowEditorToolkit.cpp
+++ b/Engine/Plugins/Experimental/Dataflow/Source/DataflowEditor/Private/Dataflow/DataflowEditorToolkit.cpp
@@ -786,6 +786,33 @@ bool FDataflowEditorToolkit::IsPerfDataEnabled() const
 	return false;
 }
 
+void FDataflowEditorToolkit::ToggleAsyncEvaluation()
+{
+	if (TObjectPtr<UDataflowBaseContent>& EditorContent = GetEditorContent())
+	{
+		if (TSharedPtr<UE::Dataflow::FEngineContext> DataflowContext = EditorContent->GetDataflowContext())
+		{
+			if (DataflowContext->IsThreaded())
+			{
+				DataflowContext->CancelAsyncEvaluation();
+			}
+			DataflowContext->SetThreaded(!DataflowContext->IsThreaded());
+		}
+	}
+}
+
+bool FDataflowEditorToolkit::IsAsyncEvaluationEnabled() const
+{
+	if (const TObjectPtr<UDataflowBaseContent>& EditorContent = GetEditorContent())
+	{
+		if (TSharedPtr<UE::Dataflow::FEngineContext> DataflowContext = EditorContent->GetDataflowContext())
+		{
+			return DataflowContext->IsThreaded();
+		}
+	}
+	return false;
+}
+
 void FDataflowEditorToolkit::ClearGraphCache()
 {
 	if (TObjectPtr<UDataflowBaseContent>& EditorContent = GetEditorContent())
@@ -876,7 +903,12 @@ void FDataflowEditorToolkit::AddEvaluationWidget(FToolMenuSection& InSection)
 		FCanExecuteAction(),
 		FIsActionChecked::CreateSP(this, &FDataflowEditorToolkit::IsPerfDataEnabled)
 	);
-
+	ToolkitCommands->MapAction(FDataflowEditorCommands::Get().ToggleAsyncEvaluation,
+		FExecuteAction::CreateSP(this, &FDataflowEditorToolkit::ToggleAsyncEvaluation),
+		FCanExecuteAction(),
+		FIsActionChecked::CreateSP(this, &FDataflowEditorToolkit::IsAsyncEvaluationEnabled)
+	);
+	
 	InSection.AddEntry(FToolMenuEntry::InitToolBarButton(
 		FDataflowEditorCommands::Get().EvaluateGraph,
 		TAttribute<FText>(),
@@ -913,6 +945,8 @@ TSharedRef<SWidget> FDataflowEditorToolkit::GenerateEvaluationOptionsMenu()
 	MenuBuilder.AddMenuEntry(FDataflowEditorCommands::Get().TogglePerfData);
 	MenuBuilder.AddSeparator();
 	MenuBuilder.AddMenuEntry(FDataflowEditorCommands::Get().ClearGraphCache);
+	MenuBuilder.AddSeparator();
+	MenuBuilder.AddMenuEntry(FDataflowEditorCommands::Get().ToggleAsyncEvaluation);
 	MenuBuilder.EndSection();
 	return MenuBuilder.MakeWidget();
 }
@@ -1937,11 +1971,22 @@ void FDataflowEditorToolkit::EvaluateNode(const FDataflowNode* Node, const FData
 		{
 			if (const TObjectPtr<UDataflowBaseContent>& EditorContent = GetEditorContent())
 			{
-				if (EditorContent->GetDataflowAsset())
+				if (ensure(Node))
 				{
+					// when evaluation completes, refresh the views
+					TWeakPtr<FDataflowEditorToolkit> WeakToolkit = SharedThis(this);
+					auto OnEvaluationCompleted = [WeakToolkit](UE::Dataflow::FContext& Context)
+						{
+							// refresh views 
+							if (TSharedPtr<FDataflowEditorToolkit> ToolKit = WeakToolkit.Pin())
+							{
+								ToolKit->bViewsNeedRefresh = true;
+							}
+						};
+					
+
 					// If Node is null, the terminal node with the given name will be used instead
-					FDataflowEditorCommands::EvaluateNode(*EditorContent->GetDataflowContext().Get(), InOutTimestamp, EditorContent->GetDataflowAsset(),
-						Node, Output, EditorContent->GetDataflowTerminal(), EditorContent->GetTerminalAsset());
+					FDataflowEditorCommands::EvaluateNode(*EditorContent->GetDataflowContext().Get(), *Node, Output, EditorContent->GetTerminalAsset(), InOutTimestamp, OnEvaluationCompleted);
 				}
 			}
 		}
diff --git a/Engine/Plugins/Experimental/Dataflow/Source/DataflowEditor/Private/Dataflow/DataflowGraphEditor.cpp b/Engine/Plugins/Experimental/Dataflow/Source/DataflowEditor/Private/Dataflow/DataflowGraphEditor.cpp
--- a/Engine/Plugins/Experimental/Dataflow/Source/DataflowEditor/Private/Dataflow/DataflowGraphEditor.cpp
+++ b/Engine/Plugins/Experimental/Dataflow/Source/DataflowEditor/Private/Dataflow/DataflowGraphEditor.cpp
@@ -19,8 +19,13 @@
 #include "HAL/PlatformApplicationMisc.h"
 #include "ScopedTransaction.h"
 #include "IStructureDetailsView.h"
+
+#include "SGraphEditorActionMenu.h"
 #include "SGraphPanel.h"
 
+#include "Widgets/Input/SEditableTextBox.h"
+#include "Widgets/Notifications/SProgressBar.h"
+
 #define LOCTEXT_NAMESPACE "DataflowGraphEditor"
 
 namespace DataflowGraphEditor::Private
@@ -117,6 +122,7 @@ namespace DataflowGraphEditor::Private
 
 TSharedPtr<FDataflowGraphEditorNodeFactory> SDataflowGraphEditor::NodeFactory;
 TWeakPtr<SDataflowGraphEditor> SDataflowGraphEditor::SelectedGraphEditor;
+TWeakPtr<SDataflowGraphEditor> SDataflowGraphEditor::LastActionMenuGraphEditor;
 
 void SDataflowGraphEditor::Construct(const FArguments& InArgs, UObject* InAssetOwner)
 {
@@ -284,12 +290,16 @@ void SDataflowGraphEditor::Construct(const FArguments& InArgs, UObject* InAssetO
 	ensureMsgf(!Arguments._GraphEvents.OnSelectionChanged.IsBound(), TEXT("DataflowGraphEditor::OnSelectionChanged rebound during construction."));
 	Arguments._GraphEvents.OnSelectionChanged = SGraphEditor::FOnSelectionChanged::CreateSP(this, &SDataflowGraphEditor::OnSelectedNodesChanged);
 
+	Arguments._GraphEvents.OnCreateActionMenuAtLocation = SGraphEditor::FOnCreateActionMenuAtLocation::CreateSP(this, &SDataflowGraphEditor::OnCreateActionMenu);
+
+
 	SGraphEditor::Construct(Arguments);
 
 	SetNodeFactory( MakeShared<FDataflowGraphNodeFactory>(this) );
 	GetGraphPanel()->SetZoomLevelsContainer<DataflowGraphEditor::Private::FDataflowZoomLevelsContainer>();
 
 	InitGraphEditorMessageBar();
+	InitEvaluationProgressBar();
 
 	// Take the existing graph panel widget and add it to a new SOverlay so we can place the message bar over top of it
 	if (ChildSlot.Num() > 0)
@@ -307,6 +317,10 @@ void SDataflowGraphEditor::Construct(const FArguments& InArgs, UObject* InAssetO
 				[
 					MessageBar.ToSharedRef()
 				]
+				+ SOverlay::Slot()
+				[
+					EvaluationProgressBar.ToSharedRef()
+				]
 		];
 	}
 }
@@ -319,6 +333,22 @@ SDataflowGraphEditor::~SDataflowGraphEditor()
 	}
 }
 
+FActionMenuContent SDataflowGraphEditor::OnCreateActionMenu(UEdGraph* Graph, const FVector2f& Position, const TArray<UEdGraphPin*>& DraggedPins, bool bAutoExpandActionMenu, SGraphEditor::FActionMenuClosed OnClosed)
+{
+	LastActionMenuGraphEditor = StaticCastSharedRef<SDataflowGraphEditor>(AsShared()).ToWeakPtr();
+
+	TSharedRef<SGraphEditorActionMenu> ActionMenu =
+		SNew(SGraphEditorActionMenu)
+		.GraphObj(EdGraphObj)
+		.NewNodePosition(FVector2f{ (float)Position.X, (float)Position.Y })
+		.DraggedFromPins(DraggedPins)
+		.AutoExpandActionMenu(bAutoExpandActionMenu)
+		.OnClosedCallback(OnClosed)
+		;
+
+	return FActionMenuContent(ActionMenu, ActionMenu->GetFilterTextBox());
+}
+
 TSharedPtr<UE::Dataflow::FContext> SDataflowGraphEditor::GetDataflowContext() const
 {
 	if (DataflowEditor)
@@ -358,8 +388,6 @@ void SDataflowGraphEditor::EvaluateNode()
 
 	using namespace UE::Dataflow;
 
-	TOptional<FContextThreaded> DefaultContext(EvaluateGraphCallback ? TOptional<FContextThreaded>() : TOptional<FContextThreaded>(FContextThreaded()));
-
 	for (UObject* Node : GetSelectedNodes())
 	{
 		if (UDataflowEdNode* const EdNode = Cast<UDataflowEdNode>(Node))
@@ -378,8 +406,8 @@ void SDataflowGraphEditor::EvaluateNode()
 						}
 						else
 						{
-							check(DefaultContext);
-							DefaultContext->Evaluate(DataflowNode.Get(), nullptr);
+							FContextThreaded DefaultContext;
+							DefaultContext.Evaluate(DataflowNode.Get(), nullptr);
 						}
 					}
 				}
@@ -1018,6 +1046,121 @@ void SDataflowGraphEditor::InitGraphEditorMessageBar()
 	}
 }
 
+void SDataflowGraphEditor::InitEvaluationProgressBar()
+{
+	const FSlateBrush* OverlayBrush = FAppStyle::Get().GetBrush("EditorViewport.OverlayBrush");
+
+	auto GetBarVisibility = [this]() ->EVisibility
+		{
+			bool bVisible = false;
+			if (TSharedPtr<UE::Dataflow::FContext> Context = GetDataflowContext())
+			{
+				bVisible = Context->IsAsyncEvaluating();
+			}
+			return bVisible ? EVisibility::SelfHitTestInvisible : EVisibility::Collapsed;
+		};
+
+	auto GetProgressText = [this]() -> FText
+		{
+			FText ProgressText;
+			if (TSharedPtr<UE::Dataflow::FContext> Context = GetDataflowContext())
+			{
+				int32 NumPendingTasks = 0;
+				int32 NumRunningTasks = 0;
+				int32 NumCompletedTasks = 0;
+				Context->GetAsyncEvaluationStats(NumPendingTasks, NumRunningTasks, NumCompletedTasks);
+
+				ProgressText = FText::Format(
+					LOCTEXT("DataflowGraphEditor_ProgressText", "Pending: {1} | Running: {0} | Completed: {2}"),
+					NumPendingTasks,
+					NumRunningTasks,
+					NumCompletedTasks
+				);
+			}
+			return ProgressText;
+		};
+
+	auto GetProgressPercent = [this]() -> float
+		{
+			float Percent = 0.f;
+			if (TSharedPtr<UE::Dataflow::FContext> Context = GetDataflowContext())
+			{
+				int32 NumPendingTasks = 0;
+				int32 NumRunningTasks = 0;
+				int32 NumCompletedTasks = 0;
+				Context->GetAsyncEvaluationStats(NumPendingTasks, NumRunningTasks, NumCompletedTasks);
+
+				const int32 TotalTasks = (NumPendingTasks + NumRunningTasks + NumCompletedTasks);
+				Percent = static_cast<float>(NumCompletedTasks) / static_cast<float>(TotalTasks);
+			}
+			return Percent;
+		};
+
+	auto OnCancel = [this]() -> FReply
+		{
+			if (TSharedPtr<UE::Dataflow::FContext> Context = GetDataflowContext())
+			{
+				Context->CancelAsyncEvaluation();
+			}
+			return FReply::Handled();
+		};
+
+	SAssignNew(EvaluationProgressBar, SHorizontalBox)
+	.Visibility_Lambda(GetBarVisibility)
+
+	+SHorizontalBox::Slot()
+		.HAlign(HAlign_Center)
+		.VAlign(VAlign_Bottom)
+		.FillWidth(1.0)
+		.Padding(FMargin(0.0f, 0.0f, 0.f, 15.f))
+		[
+			SNew(SBorder)
+			.BorderImage(OverlayBrush)
+			[
+				SNew(SHorizontalBox)
+				
+				// Text + progress bar 
+				+ SHorizontalBox::Slot()
+				.AutoWidth()
+				.Padding(8)
+				[
+					SNew(SVerticalBox)
+
+					// Progress text
+					+ SVerticalBox::Slot()
+					.VAlign(VAlign_Center)
+					.Padding(4)
+					[
+						SNew(STextBlock)
+						.Text_Lambda(GetProgressText)
+					]
+
+					// Progress bar 
+					+ SVerticalBox::Slot()
+					.VAlign(VAlign_Center)
+					.Padding(4)
+					[
+						SNew(SProgressBar)
+						.Percent_Lambda(GetProgressPercent)
+					]
+				]
+
+				// Cancel button
+				+ SHorizontalBox::Slot()
+				.AutoWidth()
+				.Padding(8)
+				[
+					SNew(SButton)
+					.Text(LOCTEXT("DataflowGraphEditor_CancelText", "Cancel"))
+					.ToolTipText(LOCTEXT("DataflowGraphEditor_CancelTooltip", "Cancel current evaluation"))
+					.HAlign(HAlign_Center)
+					.VAlign(VAlign_Center)
+					.OnClicked_Lambda(OnCancel)
+				]
+			]
+		];
+}
+
 FText SDataflowGraphEditor::GetGraphEditorOverlayText() const
 {
 	return MessageBarText;
diff --git a/Engine/Plugins/Experimental/Dataflow/Source/DataflowEditor/Private/Dataflow/DataflowSNode.cpp b/Engine/Plugins/Experimental/Dataflow/Source/DataflowEditor/Private/Dataflow/DataflowSNode.cpp
--- a/Engine/Plugins/Experimental/Dataflow/Source/DataflowEditor/Private/Dataflow/DataflowSNode.cpp
+++ b/Engine/Plugins/Experimental/Dataflow/Source/DataflowEditor/Private/Dataflow/DataflowSNode.cpp
@@ -248,7 +248,18 @@ TArray<FOverlayWidgetInfo> SDataflowEdNode::GetOverlayWidgets(bool bSelected, co
 				if (TSharedPtr<const FDataflowNode> DataflowNode = DataflowGraphNode->GetDataflowNode())
 				{
 					UE::Dataflow::FContextPerfData::FData PerfData = Context->GetPerfDataForNode(*DataflowNode);
-					if (PerfData.ExclusiveTimeMs > 0 && PerfData.InclusiveTimeMs > 0)
+					if (DataflowNode->IsAsyncEvaluating())
+					{
+						constexpr double BlinkPeriod = 2000;
+						const double ElapsedMs = FPlatformTime::ToMilliseconds64(FPlatformTime::Cycles64());
+						const double Fraction = FMath::Fractional(ElapsedMs / BlinkPeriod);
+						const double LerpFactor = FMath::Sin(Fraction * 2.0 * UE_PI) * 0.5 + 0.5;
+						const FLinearColor Color = FMath::Lerp(FLinearColor::Red, FLinearColor::White, LerpFactor);
+						PerfWidget->SetColorAndOpacity(Color);
+						PerfWidget->SetVisibility(EVisibility::Visible);
+						PerfWidget->SetText(LOCTEXT("DataflowNodeEvaluatingMessage", "Evaluating..."));
+					}
+					else if (PerfData.ExclusiveTimeMs > 0 && PerfData.InclusiveTimeMs > 0)
 					{
 						constexpr double FadeTimeMs = 5000;
 						const double MsSinceUpdate = FPlatformTime::ToMilliseconds64(FPlatformTime::Cycles64() - PerfData.LastTimestamp.Value);
diff --git a/Engine/Plugins/Experimental/Dataflow/Source/DataflowEditor/Public/Dataflow/DataflowEditorCommands.h b/Engine/Plugins/Experimental/Dataflow/Source/DataflowEditor/Public/Dataflow/DataflowEditorCommands.h
--- a/Engine/Plugins/Experimental/Dataflow/Source/DataflowEditor/Public/Dataflow/DataflowEditorCommands.h
+++ b/Engine/Plugins/Experimental/Dataflow/Source/DataflowEditor/Public/Dataflow/DataflowEditorCommands.h
@@ -75,6 +75,7 @@ public:
 	TSharedPtr<FUICommandInfo> AddNewSubGraph;
 	TSharedPtr<FUICommandInfo> ConvertToBasicSubGraph;
 	TSharedPtr<FUICommandInfo> ConvertToForEachSubGraph;
+	TSharedPtr<FUICommandInfo> ToggleAsyncEvaluation;
 
 	TMap<FName, TSharedPtr<FUICommandInfo>> SetConstructionViewModeCommands;
 
@@ -135,6 +136,8 @@ public:
 		const UDataflow* Dataflow, const FDataflowNode* Node, const FDataflowOutput* Output = nullptr, 
 		const FString& NodeName = FString(), UObject* Asset = nullptr);
 
+	static void EvaluateNode(UE::Dataflow::FContext& Context, const FDataflowNode& Node, const FDataflowOutput* Output, UObject* Asset, UE::Dataflow::FTimestamp& InOutLastNodeTimestamp, UE::Dataflow::FOnPostEvaluationFunction OnEvaluationCompleted);
+
 	/*
 	*  DeleteNodes
 	*/
diff --git a/Engine/Plugins/Experimental/Dataflow/Source/DataflowEditor/Public/Dataflow/DataflowEditorToolkit.h b/Engine/Plugins/Experimental/Dataflow/Source/DataflowEditor/Public/Dataflow/DataflowEditorToolkit.h
--- a/Engine/Plugins/Experimental/Dataflow/Source/DataflowEditor/Public/Dataflow/DataflowEditorToolkit.h
+++ b/Engine/Plugins/Experimental/Dataflow/Source/DataflowEditor/Public/Dataflow/DataflowEditorToolkit.h
@@ -247,6 +247,8 @@ private:
 	bool CanClearGraphCache() const;
 	void TogglePerfData();
 	bool IsPerfDataEnabled() const;
+	void ToggleAsyncEvaluation();
+	bool IsAsyncEvaluationEnabled() const;
 
 	/** Create the simulation viewport client */
 	void CreateSimulationViewportClient();
diff --git a/Engine/Plugins/Experimental/Dataflow/Source/DataflowEditor/Public/Dataflow/DataflowGraphEditor.h b/Engine/Plugins/Experimental/Dataflow/Source/DataflowEditor/Public/Dataflow/DataflowGraphEditor.h
--- a/Engine/Plugins/Experimental/Dataflow/Source/DataflowEditor/Public/Dataflow/DataflowGraphEditor.h
+++ b/Engine/Plugins/Experimental/Dataflow/Source/DataflowEditor/Public/Dataflow/DataflowGraphEditor.h
@@ -165,6 +165,8 @@ public:
 	/** Return the currently selected editor. Only valid for the duration of the OnSelectedNodesChanged callback where the property editor is updated. */
 	static const TWeakPtr<SDataflowGraphEditor>& GetSelectedGraphEditor() { return SelectedGraphEditor; }
 
+	static const TWeakPtr<SDataflowGraphEditor>& GetLastActionMenuGraphEditor() { return LastActionMenuGraphEditor; }
+
 private:
 	/** Add an additional option pin to all selected Dataflow nodes for those that overrides the AddPin function. */
 	void OnAddOptionPin();
@@ -186,6 +188,9 @@ private:
 	/** Create a widget to display an overlaid message in the graph editor panel */
 	void InitGraphEditorMessageBar();
 
+	/** Create a widget to display teh progress of the evaluation of the graph */
+	void InitEvaluationProgressBar();
+
 	/** Text for overlay message in graph editor panel */
 	FText GetGraphEditorOverlayText() const;
 
@@ -195,6 +200,8 @@ private:
 	FDataflowEditorCommands::FOnDragDropEventCallback OnDragDropEventCallback;
 	FDataflowEditorCommands::FGraphEvaluationCallback EvaluateGraphCallback;
 
+	FActionMenuContent OnCreateActionMenu(UEdGraph* Graph, const FVector2f& Position, const TArray<UEdGraphPin*>& DraggedPins, bool bAutoExpandActionMenu, SGraphEditor::FActionMenuClosed OnClosed);
+
 	/** The asset that ownes this dataflow graph */
 	TWeakObjectPtr<UObject> AssetOwner;
 
@@ -213,6 +220,9 @@ private:
 	/** The current graph editor when the selection callback is invoked. */
 	static TWeakPtr<SDataflowGraphEditor> SelectedGraphEditor;
 
+	/** The last graph editor used when a action context  menu was brough up in the graph. */
+	static TWeakPtr<SDataflowGraphEditor> LastActionMenuGraphEditor;
+
 	/** Editor for the content */
 	UDataflowEditor* DataflowEditor = nullptr;
 
@@ -224,6 +234,7 @@ private:
 
 	FDelegateHandle CVarChangedDelegateHandle;
 	TSharedPtr<SWidget> MessageBar;
+	TSharedPtr<SWidget> EvaluationProgressBar;
 	FText MessageBarText;
 };
 
diff --git a/Engine/Plugins/Experimental/Dataflow/Source/DataflowEditor/Public/Dataflow/DataflowRenderingFactory.h b/Engine/Plugins/Experimental/Dataflow/Source/DataflowEditor/Public/Dataflow/DataflowRenderingFactory.h
--- a/Engine/Plugins/Experimental/Dataflow/Source/DataflowEditor/Public/Dataflow/DataflowRenderingFactory.h
+++ b/Engine/Plugins/Experimental/Dataflow/Source/DataflowEditor/Public/Dataflow/DataflowRenderingFactory.h
@@ -40,6 +40,10 @@ namespace UE::Dataflow
 			{
 				if (const FDataflowOutput* Output = Node->FindOutput(OutputName))
 				{
+					if (Context.IsThreaded())
+					{
+						return Output->ReadValue<T>(Context, Default);
+					}
 					return Output->GetValue<T>(Context, Default);
 				}
 			}
diff --git a/Engine/Source/Runtime/Experimental/Dataflow/Core/Private/Dataflow/DataflowContextEvaluator.cpp b/Engine/Source/Runtime/Experimental/Dataflow/Core/Private/Dataflow/DataflowContextEvaluator.cpp
new file mode 100644
--- /dev/null
+++ b/Engine/Source/Runtime/Experimental/Dataflow/Core/Private/Dataflow/DataflowContextEvaluator.cpp
@@ -0,0 +1,276 @@
+// Copyright Epic Games, Inc. All Rights Reserved.
+
+#include "Dataflow/DataflowContextEvaluator.h"
+#include "Dataflow/DataflowNode.h"
+
+namespace UE::Dataflow
+{
+	void FContextEvaluator::ScheduleNodeEvaluation(const FDataflowNode& Node, FOnPostEvaluationFunction OnPostEvaluation)
+	{
+		if (Node.NumOutputs() == 0)
+		{
+			// likely a terminal node or no output set the output to default Guid
+			FEvaluationEntry Entry
+			{
+				.WeakNode = Node.AsWeak(),
+				.Id = FNodeOutputId{ Node.GetGuid(), FGuid() },
+				.OnPostEvaluation = OnPostEvaluation,
+			};
+			ScheduleEvaluation(Entry);
+		}
+		else
+		{
+			// schedule all available outputs
+			for (const FDataflowOutput* Output : Node.GetOutputs())
+			{
+				if (Output)
+				{
+					ScheduleEvaluation(*Output, OnPostEvaluation);
+				}
+			}
+		}
+		Process();
+	}
+
+	void FContextEvaluator::ScheduleOutputEvaluation(const FDataflowOutput& Output, FOnPostEvaluationFunction OnPostEvaluation)
+	{
+		ScheduleEvaluation(Output, OnPostEvaluation);
+		Process();
+	}
+
+	void FContextEvaluator::ScheduleEvaluation(const FDataflowOutput& Output, FOnPostEvaluationFunction OnPostEvaluation)
+	{
+		if (const FDataflowNode* Node = Output.GetOwningNode())
+		{
+			FEvaluationEntry Entry
+			{
+				.WeakNode = Node->AsWeak(),
+				.Id = FNodeOutputId{ Node->GetGuid(), Output.GetGuid() },
+				.OnPostEvaluation = OnPostEvaluation,
+			};
+			ScheduleEvaluation(Entry);
+		}
+	}
+
+	void FContextEvaluator::ScheduleEvaluation(const FEvaluationEntry& Entry)
+	{
+		if (RunningTasks.Contains(Entry.Id) || PendingEvaluationEntries.Contains(Entry.Id))
+		{
+			UE_LOG(LogTemp, Warning, TEXT("FContextEvaluator::ScheduleEvaluation : skipped [%s]"), *Entry.ToString());
+			return;
+		}
+		UE_LOG(LogTemp, Warning, TEXT("FContextEvaluator::ScheduleEvaluation : [%s]"), *Entry.ToString());
+		PendingEvaluationEntries.Add(Entry.Id, Entry);
+
+		// add upstream outputs
+		TArray<const FDataflowOutput*> InvalidUpstreamOutputs;
+		if (TSharedPtr<const FDataflowNode> Node = Entry.WeakNode.Pin())
+		{
+			InvalidUpstreamOutputs.Reset();
+			FindInvalidUpstreamOutputs(*Node, InvalidUpstreamOutputs);
+			for (const FDataflowOutput* UpstreamOutput : InvalidUpstreamOutputs)
+			{
+				UE_LOG(LogTemp, Warning, TEXT("FContextEvaluator::ScheduleEvaluation :  [%s] -- Invalid Upstream output [%s]"), *Entry.ToString(), *UpstreamOutput->GetName().ToString());
+			}
+			for (const FDataflowOutput* UpstreamOutput : InvalidUpstreamOutputs)
+			{
+				ScheduleEvaluation(*UpstreamOutput, {});
+			}
+		}
+	}
+
+	void FContextEvaluator::Cancel()
+	{
+		PendingEvaluationEntries.Reset();
+		CompletedTasks.Reset();
+	}
+
+	void FContextEvaluator::FindInvalidUpstreamOutputs(const FDataflowNode& Node, TArray<const FDataflowOutput*>& OutInvalidUpstreamOutputs)
+	{
+		for (const FDataflowInput* Input : Node.GetInputs())
+		{
+			if (Input)
+			{
+				if (const FDataflowOutput* UpstreamOutput = Input->GetConnection())
+				{
+					UE_LOG(LogTemp, Warning, TEXT("FContextEvaluator::FindInvalidUpstreamOutputs :  [%s] input[%s] -> output [%s]"), 
+						*Node.GetName().ToString(), 
+						*Input->GetName().ToString(),
+						*UpstreamOutput->GetName().ToString()
+					);
+
+					if (!UpstreamOutput->HasValidData(OwningContext))
+					{
+						OutInvalidUpstreamOutputs.Add(UpstreamOutput);
+					}
+				}
+			}
+		}
+	}
+
+	bool FContextEvaluator::ShouldRunOnGameThread(const FDataflowNode& Node)
+	{
+		for (const FDataflowInput* Input : Node.GetInputs())
+		{
+			if (Input)
+			{
+				const FString InputTypeName = Input->GetType().ToString();
+				// skeletal mesh and static mesh do support asynchronous loading and do not allow for accessing 
+				// their property from elsewhere than the gamethread
+				if (InputTypeName.Contains("UStaticMesh") || InputTypeName.Contains("USkeletalMesh"))
+				{
+					return true;
+				}
+			}
+		}
+		return Node.EvaluateOnGameThreadOnly();
+	}
+
+	bool FContextEvaluator::TryScheduleTask(const FEvaluationEntry& Entry)
+	{
+		TArray<const FDataflowOutput*> InvalidUpstreamOutputs;
+		if (TSharedPtr<const FDataflowNode> Node = Entry.WeakNode.Pin())
+		{
+			InvalidUpstreamOutputs.Reset();
+			FindInvalidUpstreamOutputs(*Node, InvalidUpstreamOutputs);
+			if (InvalidUpstreamOutputs.IsEmpty())
+			{
+				ScheduleTask(Entry);
+				return true;
+			}
+		}
+		return false;
+	}
+
+	void FContextEvaluator::GetStats(int32& OutNumPendingTasks, int32& OutNumRunningTasks, int32& OutNumCompletedTasks) const
+	{
+		// ugly but we need to make sure the stats are up to date 
+		const_cast<FContextEvaluator*>(this)->ClearCompletedTasks();
+
+		OutNumPendingTasks = PendingEvaluationEntries.Num();
+		OutNumRunningTasks = RunningTasks.Num();
+		OutNumCompletedTasks = CompletedTasks.Num();
+	}
+
+	void FContextEvaluator::Process()
+	{
+		int32 NumScheduleTasks = 0;
+		for (auto Iter = PendingEvaluationEntries.CreateIterator(); Iter; ++Iter)
+		{
+			if (TryScheduleTask(Iter->Value))
+			{
+				NumScheduleTasks++;
+				Iter.RemoveCurrent();
+			}
+		}
+
+		if (NumScheduleTasks == 0)
+		{
+			UE_LOG(LogTemp, Warning, TEXT("FContextEvaluator::Process : No Task Scheduled NumPendingTasks=[%d]"), PendingEvaluationEntries.Num());
+			for (const TPair<FNodeOutputId, FEvaluationEntry>& Entry : PendingEvaluationEntries)
+			{
+				UE_LOG(LogTemp, Warning, TEXT("FContextEvaluator::Process : \t -[%s]"), *Entry.Value.ToString());
+			}
+		}
+
+		ClearCompletedTasks();
+	}
+
+	void FContextEvaluator::ClearCompletedTasks()
+	{
+		for (auto Iter = RunningTasks.CreateIterator(); Iter; ++Iter)
+		{
+			const FGraphEventRef& Task = (*Iter).Value;
+			if (Task->IsCompleted())
+			{
+				Iter.RemoveCurrent();
+				CompletedTasks.Add((*Iter).Key);
+			}
+		}
+	}
+
+	void FContextEvaluator::ScheduleTask(const FEvaluationEntry& Entry)
+	{
+		if (TSharedPtr<const FDataflowNode> Node = Entry.WeakNode.Pin())
+		{
+			FGraphEventArray ExistingTasks;
+			if (FGraphEventRef* ExistingTask = RunningTasks.Find(Entry.Id))
+			{ 
+				ExistingTasks.Add(*ExistingTask);
+			}
+
+			const bool bUseGameThread = ShouldRunOnGameThread(*Node);
+
+			UE_LOG(LogTemp, Warning, TEXT("FContextEvaluator::ScheduleTask : [%s] GameThread=[%d] previousTasks=[%d]"),
+				*Entry.ToString(),
+				(int32)bUseGameThread,
+				(int32)ExistingTasks.Num()
+			);
+
+			FContext* ContextPtr = &OwningContext;
+			FGraphEventRef NewTask = FFunctionGraphTask::CreateAndDispatchWhenReady(
+				[ContextPtr, Entry]
+				{
+					if (ContextPtr)
+					{
+						if (TSharedPtr<const FDataflowNode> Node = Entry.WeakNode.Pin())
+						{
+							Node->SetAsyncEvaluating(true);
+							if (const FDataflowOutput* Output = Node->FindOutput(Entry.Id.OutputId))
+							{
+								Node->Evaluate(*ContextPtr, Output);
+							}
+							else if (Node->NumOutputs() == 0)
+							{
+								Node->Evaluate(*ContextPtr, nullptr);
+							}
+							Node->SetAsyncEvaluating(false);
+
+							UE_LOG(LogTemp, Warning, TEXT("FContextEvaluator::EndTask : [%s]"), *Entry.ToString());
+						}
+					}
+				},
+				TStatId(),
+				&ExistingTasks, /* prerequisites - make sure we wait on the previous one if any */
+				bUseGameThread ? ENamedThreads::GameThread : ENamedThreads::AnyThread
+			);
+
+			auto OnFinishEvaluating = 
+				[Evaluator = this, ContextPtr, OnPostEvaluation = Entry.OnPostEvaluation]()
+				{
+					Evaluator->Process();
+					if (OnPostEvaluation.IsSet() && ContextPtr)
+					{
+						OnPostEvaluation(*ContextPtr);
+					}
+				};
+
+			// handle post evaluation and run it on the game thread 
+			NewTask = FFunctionGraphTask::CreateAndDispatchWhenReady(
+				OnFinishEvaluating,
+				TStatId(),
+				NewTask, /* prerequisites */
+				ENamedThreads::GameThread
+			);
+
+			RunningTasks.Add(Entry.Id, NewTask);
+		}
+	}
+
+	FString FContextEvaluator::FEvaluationEntry::ToString() const
+	{
+		static FName UnknownName("-Unknown-");
+		FName NodeName = UnknownName;
+		FName OutputName = UnknownName;
+		if (TSharedPtr<const FDataflowNode> Node = WeakNode.Pin())
+		{
+			NodeName = Node->GetName();
+			if (const FDataflowOutput* Output = Node->FindOutput(Id.OutputId))
+			{
+				OutputName = Output->GetName();
+			}
+		}
+		return FString::Format(TEXT("{0}.{1}"), { NodeName.ToString(), OutputName.ToString() });
+	}
+};
+
diff --git a/Engine/Source/Runtime/Experimental/Dataflow/Core/Private/Dataflow/DataflowNodeParameters.cpp b/Engine/Source/Runtime/Experimental/Dataflow/Core/Private/Dataflow/DataflowNodeParameters.cpp
--- a/Engine/Source/Runtime/Experimental/Dataflow/Core/Private/Dataflow/DataflowNodeParameters.cpp
+++ b/Engine/Source/Runtime/Experimental/Dataflow/Core/Private/Dataflow/DataflowNodeParameters.cpp
@@ -6,13 +6,13 @@
 #include "Dataflow/DataflowContextCachingFactory.h"
 #include "Dataflow/DataflowInputOutput.h"
 #include "Dataflow/DataflowNode.h"
+#include "Dataflow/DataflowTerminalNode.h"
 #include "Serialization/Archive.h"
 #include "GeometryCollection/ManagedArrayCollection.h"
 #include "Modules/ModuleManager.h"
 
 namespace UE::Dataflow
 {
-
 	////////////////////////////////////////////////////////////////////////////////////////////////////////////////////////////////////////////
 
 	void FContextPerfData::Reset()
@@ -106,6 +106,21 @@ namespace UE::Dataflow
 		}
 	}
 
+	bool FContext::IsAsyncEvaluating() const
+	{
+		return IsThreaded() && (AsyncEvaluator.GetNumRunningTasks() > 0);
+	}
+
+	void FContext::CancelAsyncEvaluation()
+	{
+		AsyncEvaluator.Cancel();
+	}
+
+	void FContext::GetAsyncEvaluationStats(int32& OutNumPendingTasks, int32& OutNumRunningTasks, int32& OutNumCompletedTasks) const
+	{
+		AsyncEvaluator.GetStats(OutNumPendingTasks, OutNumRunningTasks, OutNumCompletedTasks);
+	}
+
 	int32 FContext::GetKeys(TSet<FContextCacheKey>& InKeys) const
 	{ 
 		FScopedOptionalLock Lock(DataLock.Get());
@@ -115,20 +130,7 @@ namespace UE::Dataflow
 	void FContext::SetDataImpl(FContextCacheKey Key, TUniquePtr<FContextCacheElementBase>&& DataStoreEntry)
 	{
 		FScopedOptionalLock Lock(DataLock.Get());
-		if (IsThreaded())
-		{
-			// Threaded evaluation can only set an output once per context evaluation. Otherwise
-				// downstream nodes that are extracting the data will get corrupted store entries. 
-			const TUniquePtr<FContextCacheElementBase>* const CurrentData = DataStore.Find(Key);
-			if (!CurrentData || !(*CurrentData) || (*CurrentData)->GetTimestamp() < DataStoreEntry->GetTimestamp())
-			{
-				DataStore.Emplace(Key, MoveTemp(DataStoreEntry));
-			}
-		}
-		else
-		{
-			DataStore.Emplace(Key, MoveTemp(DataStoreEntry));
-		}
+		DataStore.Emplace(Key, MoveTemp(DataStoreEntry));
 	}
 
 	const TUniquePtr<FContextCacheElementBase>* FContext::GetDataImpl(FContextCacheKey Key) const
@@ -268,12 +270,14 @@ namespace UE::Dataflow
 
 	void FContext::Info(const FString& InInfo, const FDataflowNode* InNode, const FDataflowOutput* InOutput)
 	{
-		if (OnContextHasInfo.IsBound())
+		if (!IsThreaded())
 		{
-			OnContextHasInfo.Broadcast(InNode, InOutput, InInfo);
+			if (OnContextHasInfo.IsBound())
+			{
+				OnContextHasInfo.Broadcast(InNode, InOutput, InInfo);
+			}
+			NodesWithInfo.Add(InNode);
 		}
-
-		NodesWithInfo.Add(InNode);
 	}
 
 	int32 FContext::GetNumInfo() const
@@ -283,12 +287,14 @@ namespace UE::Dataflow
 
 	void FContext::Warning(const FString& InWarning, const FDataflowNode* InNode, const FDataflowOutput* InOutput)
 	{
-		if (OnContextHasWarning.IsBound())
+		if (!IsThreaded())
 		{
-			OnContextHasWarning.Broadcast(InNode, InOutput, InWarning);
+			if (OnContextHasWarning.IsBound())
+			{
+				OnContextHasWarning.Broadcast(InNode, InOutput, InWarning);
+			}
+			NodesWithWarning.Add(InNode);
 		}
-
-		NodesWithWarning.Add(InNode);
 	}
 
 	int32 FContext::GetNumWarnings() const
@@ -298,34 +304,36 @@ namespace UE::Dataflow
 
 	void FContext::Error(const FString& InError, const FDataflowNode* InNode, const FDataflowOutput* InOutput)
 	{
-		if (OnContextHasError.IsBound())
+		if (!IsThreaded())
 		{
-			OnContextHasError.Broadcast(InNode, InOutput, InError);
-		}
-	
-		NodesWithError.Add(InNode);
+			if (OnContextHasError.IsBound())
+			{
+				OnContextHasError.Broadcast(InNode, InOutput, InError);
+			}
+			NodesWithError.Add(InNode);
 
-		//
+			//
 		// Add all the connected nodes to NodesFailed set
 		//
 #if DATAFLOW_EDITOR_EVALUATION
-		const int32 NumNodesInCallStack = Callstack.Num();
-		if (NumNodesInCallStack > 1)
-		{
-			for (int32 Idx = NumNodesInCallStack - 2; Idx >= 0; --Idx)
+			const int32 NumNodesInCallStack = Callstack.Num();
+			if (NumNodesInCallStack > 1)
 			{
-				const FDataflowNode* OwningNode = Callstack[Idx]->GetOwningNode();
-				const FString WarningStr = FString(TEXT("Evaluation failed"));
-
-				if (OnContextHasWarning.IsBound())
+				for (int32 Idx = NumNodesInCallStack - 2; Idx >= 0; --Idx)
 				{
-					OnContextHasWarning.Broadcast(OwningNode, InOutput, WarningStr);
-				}
+					const FDataflowNode* OwningNode = Callstack[Idx]->GetOwningNode();
+					const FString WarningStr = FString(TEXT("Evaluation failed"));
+
+					if (OnContextHasWarning.IsBound())
+					{
+						OnContextHasWarning.Broadcast(OwningNode, InOutput, WarningStr);
+					}
 
-				NodesFailed.Add(OwningNode);
+					NodesFailed.Add(OwningNode);
+				}
 			}
-		}
 #endif
+		}
 	}
 
 	int32 FContext::GetNumErrors() const
@@ -473,34 +481,80 @@ namespace UE::Dataflow
 		}
 	}
 
+	void FContext::Evaluate(const FDataflowNode* Node, const FDataflowOutput* Output, FOnPostEvaluationFunction PostEvaluationFunction)
+	{
+		if (IsThreaded())
+		{
+			if (Output)
+			{
+				AsyncEvaluator.ScheduleOutputEvaluation(*Output, PostEvaluationFunction);
+			}
+			else if (Node)
+			{
+				AsyncEvaluator.ScheduleNodeEvaluation(*Node, PostEvaluationFunction);
+			}
+		}
+		else
+		{
+			Evaluate(Node, Output);
+			if (PostEvaluationFunction.IsSet())
+			{
+				PostEvaluationFunction(*this);
+			}
+		}
+	}
+
 	void FContext::Evaluate(const FDataflowNode* Node, const FDataflowOutput* Output)
 	{
-		BeginContextEvaluation(Node, Output);
+		ensureMsgf(!IsThreaded(), TEXT(
+			"WARNING: Trying to use no - asynchronous evaluation method in a threaded context "
+			"The node outputs will not be up to date when returning from this method"
+			"Please use the EValuation method that takes a PostEvaluation function"
+		));
+		if (IsThreaded())
+		{
+			Evaluate(Node, Output, {});
+		}
+		else
+		{
+			BeginContextEvaluation(Node, Output);
+		}
 	}
 
 	bool FContext::Evaluate(const FDataflowOutput& Connection)
 	{
 		CheckIntrinsicInputs(Connection);
 
-		if (!IsThreaded() && OnNodeBeginEvaluateMulticast.IsBound())
-		{
-			OnNodeBeginEvaluateMulticast.Broadcast(Connection.GetOwningNode(), &Connection);
-		}
-
-		UE_LOG(LogChaosDataflow, VeryVerbose, TEXT("FContext::Evaluate(): Node [%s], Output [%s]"), *Connection.GetOwningNode()->GetName().ToString(), *Connection.GetName().ToString());
-		bool RetValue = false;
+		if (IsThreaded())
 		{
-			FScopedOptionalLock Lock(IsThreaded() ? Connection.OutputLock.Get() : nullptr);
-			RetValue = Connection.EvaluateImpl(*this);
+			if (!Connection.HasCachedValue(*this))
+			{
+				Connection.SetNullValue(*this);
+			}
+			return true;
 		}
-		if (!IsThreaded() && OnNodeFinishEvaluateMulticast.IsBound())
+		else
 		{
-			OnNodeFinishEvaluateMulticast.Broadcast(Connection.GetOwningNode(), &Connection);
-		}
+			if (OnNodeBeginEvaluateMulticast.IsBound())
+			{
+				OnNodeBeginEvaluateMulticast.Broadcast(Connection.GetOwningNode(), &Connection);
+			}
 
-		const bool bHasError = NodeHasError(Connection.GetOwningNode()) || NodeFailed(Connection.GetOwningNode());
+			UE_LOG(LogChaosDataflow, VeryVerbose, TEXT("FContext::Evaluate(): Node [%s], Output [%s]"), *Connection.GetOwningNode()->GetName().ToString(), *Connection.GetName().ToString());
+			bool RetValue = false;
+			{
+				FScopedOptionalLock Lock(IsThreaded() ? Connection.OutputLock.Get() : nullptr);
+				RetValue = Connection.EvaluateImpl(*this);
+			}
+			if (OnNodeFinishEvaluateMulticast.IsBound())
+			{
+				OnNodeFinishEvaluateMulticast.Broadcast(Connection.GetOwningNode(), &Connection);
+			}
 
-		return !bHasError && RetValue;
+			const bool bHasError = NodeHasError(Connection.GetOwningNode()) || NodeFailed(Connection.GetOwningNode());
+
+			return !bHasError && RetValue;
+		}
 	}
 
 	void FContextCache::Serialize(FArchive& Ar)
diff --git a/Engine/Source/Runtime/Experimental/Dataflow/Core/Public/Dataflow/DataflowContextEvaluator.h b/Engine/Source/Runtime/Experimental/Dataflow/Core/Public/Dataflow/DataflowContextEvaluator.h
new file mode 100644
--- /dev/null
+++ b/Engine/Source/Runtime/Experimental/Dataflow/Core/Public/Dataflow/DataflowContextEvaluator.h
@@ -0,0 +1,92 @@
+// Copyright Epic Games, Inc. All Rights Reserved.
+
+#pragma once
+
+#include "Async/TaskGraphInterfaces.h"
+#include "Misc/Guid.h"
+#include "Templates/Function.h"
+#include "Templates/SharedPointer.h"
+
+struct FDataflowNode;
+struct FDataflowOutput;
+
+namespace UE::Dataflow
+{
+	class FContext;
+
+	using FOnPostEvaluationFunction = TFunction<void(FContext&)>;
+
+	/*
+	* Asynchronously evaluate a dataflow nodes 
+	* This may be slower than executing the graph in one go synchronously but this offers the following advantage:
+	* - this can be cancelled at anytime ( only the in progress node will have finish evaluating while all the pending ones will be discarded )
+	* - nodes that need to run on the game thread will do so ( see FDataflowNode EvaluateOnGameThreadOnly method )
+	*/
+	struct FContextEvaluator
+	{
+	public:
+		struct FNodeOutputId
+		{
+			FGuid NodeId;
+			FGuid OutputId;
+			bool operator == (const FNodeOutputId& Other) const
+			{
+				return NodeId == Other.NodeId && OutputId == Other.OutputId;
+			}
+			friend inline uint32 GetTypeHash(const FNodeOutputId& Id)
+			{
+				return HashCombine(GetTypeHash(Id.NodeId), GetTypeHash(Id.OutputId));
+			}
+		};
+		struct FEvaluationEntry
+		{
+			TWeakPtr<const FDataflowNode> WeakNode;
+			FNodeOutputId Id;
+			FOnPostEvaluationFunction OnPostEvaluation;
+
+			bool operator == (const FEvaluationEntry& Other) const
+			{
+				return Id == Other.Id;
+			}
+			bool operator == (const FNodeOutputId& Other) const
+			{
+				return (Id == Other);
+			}
+
+			FString ToString() const;
+		};
+
+		FContextEvaluator(FContext& InOwningContext)
+			: OwningContext(InOwningContext)
+		{}
+
+		void ScheduleNodeEvaluation(const FDataflowNode& Node, FOnPostEvaluationFunction OnPostEvaluation);
+		void ScheduleOutputEvaluation(const FDataflowOutput& Output, FOnPostEvaluationFunction OnPostEvaluation);
+
+		void Process();
+		void Cancel();
+		void GetStats(int32& OutNumPendingTasks, int32& OutNumRunningTasks, int32& OutNumCompletedTasks) const;
+
+		int32 GetNumPendingTasks() const { return PendingEvaluationEntries.Num(); }
+		int32 GetNumRunningTasks() const { return RunningTasks.Num(); }
+		int32 GetNumCompletedTasks() const { return CompletedTasks.Num(); }
+
+	private:
+		void FindInvalidUpstreamOutputs(const FDataflowNode& Node, TArray<const FDataflowOutput*>& OutInvalidUpstreamOutputs);
+		bool ShouldRunOnGameThread(const FDataflowNode& Node);
+
+		void ScheduleEvaluation(const FDataflowOutput& Output, FOnPostEvaluationFunction OnPostEvaluation);
+		void ScheduleEvaluation(const FEvaluationEntry& Entry);
+
+		bool TryScheduleTask(const FEvaluationEntry& Entry);
+		void ScheduleTask(const FEvaluationEntry& Entry);
+		void ClearCompletedTasks();
+
+		FContext& OwningContext;
+		TMap<FNodeOutputId, FEvaluationEntry> PendingEvaluationEntries;
+		TMap<FNodeOutputId, FGraphEventRef> RunningTasks;
+		TSet<FNodeOutputId> CompletedTasks;
+	};
+}
+
+
diff --git a/Engine/Source/Runtime/Experimental/Dataflow/Core/Public/Dataflow/DataflowNode.h b/Engine/Source/Runtime/Experimental/Dataflow/Core/Public/Dataflow/DataflowNode.h
--- a/Engine/Source/Runtime/Experimental/Dataflow/Core/Public/Dataflow/DataflowNode.h
+++ b/Engine/Source/Runtime/Experimental/Dataflow/Core/Public/Dataflow/DataflowNode.h
@@ -11,6 +11,7 @@
 #include "Dataflow/DataflowSettings.h"
 #include "StructUtils/PropertyBag.h"
 #include "UObject/StructOnScope.h"
+#include "Templates/SharedPointer.h"
 
 #include "DataflowNode.generated.h"
 
@@ -45,6 +46,9 @@ struct FDataflowFreezeActions final
 PRAGMA_DISABLE_DEPRECATION_WARNINGS
 USTRUCT()
 struct FDataflowNode
+#if CPP
+	: public TSharedFromThis<FDataflowNode, ESPMode::ThreadSafe>
+#endif
 {
 PRAGMA_ENABLE_DEPRECATION_WARNINGS
 	GENERATED_USTRUCT_BODY()
@@ -402,6 +406,8 @@ PRAGMA_ENABLE_DEPRECATION_WARNINGS
 	//
 	virtual void Evaluate(UE::Dataflow::FContext& Context, const FDataflowOutput*) const { ensure(false); }
 
+	virtual bool EvaluateOnGameThreadOnly() const { return false; }
+
 	/** Freeze the node's outputs. */
 	DATAFLOWCORE_API void Freeze(UE::Dataflow::FContext& Context);
 	/** Unfreeze the node's outputs. */
@@ -765,6 +771,9 @@ PRAGMA_ENABLE_DEPRECATION_WARNINGS
 	/** Check if any connection matching a specific type dependency group are connected */
 	DATAFLOWCORE_API bool IsAnytypeDependencyConnected(FName DependencyGroup, const FDataflowConnection* IgnoreConnection = nullptr) const;
 
+	DATAFLOWCORE_API bool IsAsyncEvaluating() const { return bIsAsyncEvaluating; }
+	DATAFLOWCORE_API void SetAsyncEvaluating(bool bValue) const { bIsAsyncEvaluating = bValue; }
+
 protected:
 	bool OnConnectionTypeChanged(const FDataflowConnection& Connection);
 
@@ -857,6 +866,8 @@ private:
 
 	TWeakPtr<UE::Dataflow::FGraph> WeakDataflowGraph;
 
+	mutable TAtomic<bool> bIsAsyncEvaluating = false;
+
 protected:
 	bool bPauseInvalidations = false;
 	UE::Dataflow::FTimestamp PausedModifiedTimestamp = UE::Dataflow::FTimestamp::Invalid; // When unpausing invalidations, Invalidate will be called with this timestamp.
diff --git a/Engine/Source/Runtime/Experimental/Dataflow/Core/Public/Dataflow/DataflowNodeParameters.h b/Engine/Source/Runtime/Experimental/Dataflow/Core/Public/Dataflow/DataflowNodeParameters.h
--- a/Engine/Source/Runtime/Experimental/Dataflow/Core/Public/Dataflow/DataflowNodeParameters.h
+++ b/Engine/Source/Runtime/Experimental/Dataflow/Core/Public/Dataflow/DataflowNodeParameters.h
@@ -8,6 +8,7 @@
 #include "Templates/UniquePtr.h"
 #include "HAL/CriticalSection.h"
 #include "Dataflow/DataflowContextCache.h"
+#include "Dataflow/DataflowContextEvaluator.h"
 #include "Serialization/Archive.h"
 
 class  UDataflow;
@@ -128,8 +129,8 @@ namespace UE::Dataflow
 	class FContext: public IContextCacheStore
 	{
 	protected:
-		FContext(FContext&&) = default;
-		FContext& operator=(FContext&&) = default;
+		FContext(FContext&&) = delete;
+		FContext& operator=(FContext&&) = delete;
 		
 		FContext(const FContext&) = delete;
 		FContext& operator=(const FContext&) = delete;
@@ -137,8 +138,10 @@ namespace UE::Dataflow
 		FContextCache DataStore;
 		TUniquePtr<FCriticalSection> DataLock;
 
+		FContextEvaluator AsyncEvaluator;
+
 	public:
-		FContext() = default;
+		FContext() : AsyncEvaluator(*this) {}
 		virtual ~FContext() = default;
 
 		static FName StaticType() { return FName("FContext"); }
@@ -170,14 +173,20 @@ namespace UE::Dataflow
 		DATAFLOWCORE_API bool IsThreaded() const { return DataLock.IsValid(); }
 		DATAFLOWCORE_API void SetThreaded(bool bValue);
 
+		DATAFLOWCORE_API bool IsAsyncEvaluating() const;
+		DATAFLOWCORE_API void CancelAsyncEvaluation();
+		DATAFLOWCORE_API void GetAsyncEvaluationStats(int32& OutNumPendingTasks, int32& OutNumRunningTasks, int32& OutNumCompletedTasks) const;
+
 		DATAFLOWCORE_API virtual void SetDataImpl(FContextCacheKey Key, TUniquePtr<FContextCacheElementBase>&& DataStoreEntry);
 		DATAFLOWCORE_API virtual const TUniquePtr<FContextCacheElementBase>* GetDataImpl(FContextCacheKey Key) const;
 		DATAFLOWCORE_API virtual bool HasDataImpl(FContextCacheKey Key, FTimestamp InTimestamp = FTimestamp::Invalid) const;
 		DATAFLOWCORE_API virtual bool IsEmptyImpl() const;
 		DATAFLOWCORE_API virtual int32 GetKeys(TSet<FContextCacheKey>& InKeys) const;
 
-		DATAFLOWCORE_API virtual void Evaluate(const FDataflowNode* Node, const FDataflowOutput* Output);
-		DATAFLOWCORE_API virtual bool Evaluate(const FDataflowOutput& Connection);
+		DATAFLOWCORE_API void Evaluate(const FDataflowNode* Node, const FDataflowOutput* Output, FOnPostEvaluationFunction PostEvaluationFunction);
+
+		DATAFLOWCORE_API virtual void Evaluate(const FDataflowNode* Node, const FDataflowOutput* Output );
+		DATAFLOWCORE_API virtual bool Evaluate(const FDataflowOutput& Connection );
 
 		DATAFLOWCORE_API void ClearAllData();
 
diff --git a/Engine/Source/Runtime/Experimental/Dataflow/Engine/Private/Dataflow/DataflowSubGraphNodes.cpp b/Engine/Source/Runtime/Experimental/Dataflow/Engine/Private/Dataflow/DataflowSubGraphNodes.cpp
--- a/Engine/Source/Runtime/Experimental/Dataflow/Engine/Private/Dataflow/DataflowSubGraphNodes.cpp
+++ b/Engine/Source/Runtime/Experimental/Dataflow/Engine/Private/Dataflow/DataflowSubGraphNodes.cpp
@@ -19,7 +19,10 @@ namespace UE::Dataflow::SubGraphNodes::Private
 		FContextSubGraph(FContext& InParentContext, const ISubGraphContextCallback& InSubGraphCallback)
 			: ParentContext(InParentContext)
 			, SubGraphCallback(InSubGraphCallback)
-		{}
+		{
+			// make sure to match the threaded mode
+			SetThreaded(ParentContext.IsThreaded());
+		}
 
 		virtual void SetDataImpl(FContextCacheKey Key, TUniquePtr<FContextCacheElementBase>&& DataStoreEntry) override
 		{
