Dataflow : add method to override struct based variables ( and arrays of structure )

diff --git a/Engine/Source/Runtime/Experimental/Dataflow/Engine/Private/Dataflow/DataflowInstance.cpp b/Engine/Source/Runtime/Experimental/Dataflow/Engine/Private/Dataflow/DataflowInstance.cpp
--- a/Engine/Source/Runtime/Experimental/Dataflow/Engine/Private/Dataflow/DataflowInstance.cpp
+++ b/Engine/Source/Runtime/Experimental/Dataflow/Engine/Private/Dataflow/DataflowInstance.cpp
@@ -107,33 +107,6 @@ namespace UE::Dataflow::InstanceUtils
 		}
 		return TerminalNodeNames;
 	}
-
-	namespace Private
-	{
-		template<typename T, typename TWriteElementFunction>
-		bool OverrideVariableArray(FInstancedPropertyBag& Variables, FName VariableName, const TArray<T>& Values, const TWriteElementFunction& WriteElementFunction, TArray<FGuid>& InOutOverrideGuids)
-		{
-			if (const FPropertyBagPropertyDesc* Desc = Variables.FindPropertyDescByName(VariableName))
-			{
-				if (FPropertyBagArrayRef* ArrayRef = Variables.GetMutableArrayRef(*Desc).TryGetValue())
-				{
-					bool bSuccess = true;
-					ArrayRef->EmptyValues();
-					ArrayRef->AddValues(Values.Num());
-					for (int32 Idx = 0; Idx < Values.Num(); ++Idx)
-					{
-						bSuccess &= (WriteElementFunction(*ArrayRef, Idx, Values[Idx]) == EPropertyBagResult::Success);
-					}
-					if (bSuccess)
-					{
-						InOutOverrideGuids.AddUnique(Desc->ID);
-					}
-					return bSuccess;
-				}
-			}
-			return false;
-		}
-	}
 }
 
 
diff --git a/Engine/Source/Runtime/Experimental/Dataflow/Engine/Public/Dataflow/DataflowInstance.h b/Engine/Source/Runtime/Experimental/Dataflow/Engine/Public/Dataflow/DataflowInstance.h
--- a/Engine/Source/Runtime/Experimental/Dataflow/Engine/Public/Dataflow/DataflowInstance.h
+++ b/Engine/Source/Runtime/Experimental/Dataflow/Engine/Public/Dataflow/DataflowInstance.h
@@ -13,6 +13,34 @@ namespace UE::Dataflow::Private
 	class FVariablesOverridesDetails;
 }
 
+namespace UE::Dataflow::InstanceUtils::Private
+{
+	// need to be in the header to be used by the templated structure override methods
+	template<typename T, typename TWriteElementFunction>
+	bool OverrideVariableArray(FInstancedPropertyBag& Variables, FName VariableName, const TArray<T>& Values, const TWriteElementFunction& WriteElementFunction, TArray<FGuid>& InOutOverrideGuids)
+	{
+		if (const FPropertyBagPropertyDesc* Desc = Variables.FindPropertyDescByName(VariableName))
+		{
+			if (FPropertyBagArrayRef* ArrayRef = Variables.GetMutableArrayRef(*Desc).TryGetValue())
+			{
+				bool bSuccess = true;
+				ArrayRef->EmptyValues();
+				ArrayRef->AddValues(Values.Num());
+				for (int32 Idx = 0; Idx < Values.Num(); ++Idx)
+				{
+					bSuccess &= (WriteElementFunction(*ArrayRef, Idx, Values[Idx]) == EPropertyBagResult::Success);
+				}
+				if (bSuccess)
+				{
+					InOutOverrideGuids.AddUnique(Desc->ID);
+				}
+				return bSuccess;
+			}
+		}
+		return false;
+	}
+}
+
 class UDataflow;
 struct FDataflowInstance;
 
@@ -28,6 +56,10 @@ struct FDataflowVariableOverrides
 
 	DATAFLOWENGINE_API FDataflowVariableOverrides(FDataflowInstance* InOwner = nullptr);
 
+	// prevent the copy constructor to avoid having a copy with the same owner
+	// use the assignment operator instead on a newly created FDataflowVariableOverrides with the right owner to copy its content over 
+	DATAFLOWENGINE_API FDataflowVariableOverrides(const FDataflowVariableOverrides& Other) = delete;
+
 	DATAFLOWENGINE_API FDataflowVariableOverrides& operator=(const FDataflowVariableOverrides& Other);
 
 	/** Remove all variables */
@@ -89,6 +121,33 @@ struct FDataflowVariableOverrides
 	/** override a dataflow String variable for this asset */
 	DATAFLOWENGINE_API bool OverrideVariableString(FName VariableName, const TArray<FString>& Values);
 
+	template <typename T>
+	bool OverrideVariableStruct(FName VariableName, const T& Value)
+	{
+		if (const FPropertyBagPropertyDesc* Desc = Variables.FindPropertyDescByName(VariableName))
+		{
+			if (Variables.SetValueStruct(VariableName, Value) == EPropertyBagResult::Success)
+			{
+				OverriddenVariableGuids.AddUnique(Desc->ID);
+				return true;
+			}
+		}
+		return false;
+	}
+
+	template <typename T>
+	bool OverrideVariableStructArray(FName VariableName, const TArray<T>& Values)
+	{
+		using namespace UE::Dataflow::InstanceUtils::Private;
+
+		auto WriteStructValue = [](FPropertyBagArrayRef& ArrayRef, int32 Idx, const T& Value)
+			{
+				return ArrayRef.SetValueStruct(Idx, Value);
+			};
+
+		return OverrideVariableArray(Variables, VariableName, Values, WriteStructValue, OverriddenVariableGuids);
+	}
+
 #if WITH_EDITOR
 
 	void OnOwnerPostEditChangeProperty(struct FPropertyChangedEvent& PropertyChangedEvent);
