Dataflow : Fix compilation error

diff --git a/Engine/Source/Runtime/Experimental/Dataflow/Engine/Public/Dataflow/DataflowInstance.h b/Engine/Source/Runtime/Experimental/Dataflow/Engine/Public/Dataflow/DataflowInstance.h
--- a/Engine/Source/Runtime/Experimental/Dataflow/Engine/Public/Dataflow/DataflowInstance.h
+++ b/Engine/Source/Runtime/Experimental/Dataflow/Engine/Public/Dataflow/DataflowInstance.h
@@ -58,7 +58,7 @@ struct FDataflowVariableOverrides
 
 	// prevent the copy constructor to avoid having a copy with the same owner
 	// use the assignment operator instead on a newly created FDataflowVariableOverrides with the right owner to copy its content over 
-	DATAFLOWENGINE_API FDataflowVariableOverrides(const FDataflowVariableOverrides& Other) = delete;
+	FDataflowVariableOverrides(const FDataflowVariableOverrides& Other) = delete;
 
 	DATAFLOWENGINE_API FDataflowVariableOverrides& operator=(const FDataflowVariableOverrides& Other);
 
