[Core]
* Added macros to be able to declare/define GetTlsSlot properly with dllstorage.
* Changed code to use macros instead of previous hacky implementation

diff --git a/Engine/Source/Runtime/Core/Private/Misc/CoreMisc.cpp b/Engine/Source/Runtime/Core/Private/Misc/CoreMisc.cpp
--- a/Engine/Source/Runtime/Core/Private/Misc/CoreMisc.cpp
+++ b/Engine/Source/Runtime/Core/Private/Misc/CoreMisc.cpp
@@ -360,6 +360,9 @@ FBoolConfigValueHelper::FBoolConfigValueHelper(const TCHAR* Section, const TCHAR
 /*----------------------------------------------------------------------------
 FScriptExceptionHandler
 ----------------------------------------------------------------------------*/
+
+UE_DEFINE_THREAD_SINGLETON_TLS(FScriptExceptionHandler)
+
 FScriptExceptionHandlerFunc FScriptExceptionHandler::DefaultExceptionHandler = &FScriptExceptionHandler::LoggingExceptionHandler;
 
 FScriptExceptionHandler& FScriptExceptionHandler::Get()
diff --git a/Engine/Source/Runtime/Core/Private/Misc/MemStack.cpp b/Engine/Source/Runtime/Core/Private/Misc/MemStack.cpp
--- a/Engine/Source/Runtime/Core/Private/Misc/MemStack.cpp
+++ b/Engine/Source/Runtime/Core/Private/Misc/MemStack.cpp
@@ -409,6 +409,9 @@ bool FMemStackBase::ContainsPointer(const void* Pointer) const
 	return false;
 }
 
+UE_DEFINE_THREAD_SINGLETON_TLS(FMemStack)
+
+
 FORCENOINLINE void UE::Core::Private::OnInvalidMemStackAllocatorNum(int32 NewNum, SIZE_T NumBytesPerElement)
 {
 	UE_LOG(LogCore, Fatal, TEXT("Trying to resize TMemStackAllocator to an invalid size of %d with element size %" SIZE_T_FMT), NewNum, NumBytesPerElement);
diff --git a/Engine/Source/Runtime/Core/Private/Stats/ThreadIdleStats.cpp b/Engine/Source/Runtime/Core/Private/Stats/ThreadIdleStats.cpp
--- a/Engine/Source/Runtime/Core/Private/Stats/ThreadIdleStats.cpp
+++ b/Engine/Source/Runtime/Core/Private/Stats/ThreadIdleStats.cpp
@@ -8,8 +8,11 @@ UE_TRACE_CHANNEL_DEFINE(ThreadIdleScopeChannel);
 TRACE_CPUPROFILER_EVENT_DECLARE(ThreadIdleScopeTraceEventId);
 #endif
 
+UE_DEFINE_THREAD_SINGLETON_TLS(UE::Stats::FThreadIdleStats)
+
 namespace UE::Stats
 {
+
 	CORE_API FThreadIdleStats::FScopeIdle::FScopeIdle(bool bInIgnore /* = false */)
 		: Start(FPlatformTime::Cycles())
 		, bIgnore(bInIgnore || FThreadIdleStats::Get().bInIdleScope)
diff --git a/Engine/Source/Runtime/Core/Public/HAL/ThreadSingleton.h b/Engine/Source/Runtime/Core/Public/HAL/ThreadSingleton.h
--- a/Engine/Source/Runtime/Core/Public/HAL/ThreadSingleton.h
+++ b/Engine/Source/Runtime/Core/Public/HAL/ThreadSingleton.h
@@ -32,6 +32,8 @@ public:
 	static CORE_API FTlsAutoCleanup* Inject(FTlsAutoCleanup* Instance, uint32& TlsSlot);
 };
 
+#define UE_DECLARE_THREAD_SINGLETON_TLS(Type, Api) template<> Api uint32& TThreadSingleton<Type>::GetTlsSlot();
+#define UE_DEFINE_THREAD_SINGLETON_TLS(Type) template<> uint32& TThreadSingleton<Type>::GetTlsSlot() { static uint32 TlsSlot = FPlatformTLS::InvalidTlsSlot; return TlsSlot; }
 
 /**
  * This a special version of singleton. It means that there is created only one instance for each thread.
diff --git a/Engine/Source/Runtime/Core/Public/Misc/CoreMisc.h b/Engine/Source/Runtime/Core/Public/Misc/CoreMisc.h
--- a/Engine/Source/Runtime/Core/Public/Misc/CoreMisc.h
+++ b/Engine/Source/Runtime/Core/Public/Misc/CoreMisc.h
@@ -257,41 +257,43 @@ public:
  */
 typedef TFunction<void(ELogVerbosity::Type /*Verbosity*/, const TCHAR* /*ExceptionMessage*/, const TCHAR* /*StackMessage*/)> FScriptExceptionHandlerFunc;
 
+UE_DECLARE_THREAD_SINGLETON_TLS(class FScriptExceptionHandler, CORE_API)
+
 /** 
  * Exception handler stack used for script exceptions.
  */
-class CORE_API FScriptExceptionHandler : public TThreadSingleton<FScriptExceptionHandler>
+class FScriptExceptionHandler : public TThreadSingleton<FScriptExceptionHandler>
 {
 public:
 	/**
 	 * Get the exception handler for the current thread
 	 */
-	static FScriptExceptionHandler& Get();
+	static CORE_API FScriptExceptionHandler& Get();
 
 	/**
 	 * Push an exception handler onto the stack
 	 */
-	void PushExceptionHandler(const FScriptExceptionHandlerFunc& InFunc);
+	CORE_API void PushExceptionHandler(const FScriptExceptionHandlerFunc& InFunc);
 
 	/**
 	 * Pop an exception handler from the stack
 	 */
-	void PopExceptionHandler();
+	CORE_API void PopExceptionHandler();
 
 	/**
 	 * Handle an exception using the active exception handler
 	 */
-	void HandleException(ELogVerbosity::Type Verbosity, const TCHAR* ExceptionMessage, const TCHAR* StackMessage);
+	CORE_API void HandleException(ELogVerbosity::Type Verbosity, const TCHAR* ExceptionMessage, const TCHAR* StackMessage);
 
 	/**
 	 * Handler for a script exception that emits an ensure (for warnings or errors)
 	 */
-	static void AssertionExceptionHandler(ELogVerbosity::Type Verbosity, const TCHAR* ExceptionMessage, const TCHAR* StackMessage);
+	static CORE_API void AssertionExceptionHandler(ELogVerbosity::Type Verbosity, const TCHAR* ExceptionMessage, const TCHAR* StackMessage);
 
 	/**
 	 * Handler for a script exception that emits a log message
 	 */
-	static void LoggingExceptionHandler(ELogVerbosity::Type Verbosity, const TCHAR* ExceptionMessage, const TCHAR* StackMessage);
+	static CORE_API void LoggingExceptionHandler(ELogVerbosity::Type Verbosity, const TCHAR* ExceptionMessage, const TCHAR* StackMessage);
 
 private:
 	/**
diff --git a/Engine/Source/Runtime/Core/Public/Misc/MemStack.h b/Engine/Source/Runtime/Core/Public/Misc/MemStack.h
--- a/Engine/Source/Runtime/Core/Public/Misc/MemStack.h
+++ b/Engine/Source/Runtime/Core/Public/Misc/MemStack.h
@@ -239,8 +239,9 @@ protected:
 	bool bShouldEnforceAllocMarks = false;
 };
 
+UE_DECLARE_THREAD_SINGLETON_TLS(class FMemStack, CORE_API)
 
-class CORE_API FMemStack : public TThreadSingleton<FMemStack>, public FMemStackBase
+class FMemStack : public TThreadSingleton<FMemStack>, public FMemStackBase
 {
 public:
 	FMemStack()
diff --git a/Engine/Source/Runtime/Core/Public/Stats/ThreadIdleStats.h b/Engine/Source/Runtime/Core/Public/Stats/ThreadIdleStats.h
--- a/Engine/Source/Runtime/Core/Public/Stats/ThreadIdleStats.h
+++ b/Engine/Source/Runtime/Core/Public/Stats/ThreadIdleStats.h
@@ -8,10 +8,18 @@
 
 namespace UE::Stats
 {
+	class FThreadIdleStats;
+}
+
+UE_DECLARE_THREAD_SINGLETON_TLS(UE::Stats::FThreadIdleStats, CORE_API)
+
+namespace UE::Stats
+{
+
 	/**
 	* This is thread-private information about the thread idle stats, which we always collect, even in final builds
 	*/
-	class CORE_API FThreadIdleStats : public TThreadSingleton<FThreadIdleStats>
+	class FThreadIdleStats : public TThreadSingleton<FThreadIdleStats>
 	{
 		friend class TThreadSingleton<FThreadIdleStats>;
 
diff --git a/Engine/Source/Runtime/CoreUObject/Private/UObject/SoftObjectPath.cpp b/Engine/Source/Runtime/CoreUObject/Private/UObject/SoftObjectPath.cpp
--- a/Engine/Source/Runtime/CoreUObject/Private/UObject/SoftObjectPath.cpp
+++ b/Engine/Source/Runtime/CoreUObject/Private/UObject/SoftObjectPath.cpp
@@ -1041,6 +1041,8 @@ FSoftClassPath FSoftClassPath::GetOrCreateIDForClass(const UClass *InClass)
 	return FSoftClassPath(InClass);
 }
 
+UE_DEFINE_THREAD_SINGLETON_TLS(FSoftObjectPathThreadContext)
+
 bool FSoftObjectPathThreadContext::GetSerializationOptions(FName& OutPackageName, FName& OutPropertyName, ESoftObjectPathCollectType& OutCollectType, ESoftObjectPathSerializeType& OutSerializeType, FArchive* Archive) const
 {
 	FName CurrentPackageName, CurrentPropertyName;
diff --git a/Engine/Source/Runtime/CoreUObject/Private/UObject/UObjectThreadContext.cpp b/Engine/Source/Runtime/CoreUObject/Private/UObject/UObjectThreadContext.cpp
--- a/Engine/Source/Runtime/CoreUObject/Private/UObject/UObjectThreadContext.cpp
+++ b/Engine/Source/Runtime/CoreUObject/Private/UObject/UObjectThreadContext.cpp
@@ -13,6 +13,8 @@
 
 DEFINE_LOG_CATEGORY(LogUObjectThreadContext);
 
+UE_DEFINE_THREAD_SINGLETON_TLS(FUObjectThreadContext)
+
 FUObjectThreadContext::FUObjectThreadContext()
 : IsRoutingPostLoad(false)
 , IsDeletingLinkers(false)
diff --git a/Engine/Source/Runtime/CoreUObject/Public/UObject/SoftObjectPath.h b/Engine/Source/Runtime/CoreUObject/Public/UObject/SoftObjectPath.h
--- a/Engine/Source/Runtime/CoreUObject/Public/UObject/SoftObjectPath.h
+++ b/Engine/Source/Runtime/CoreUObject/Public/UObject/SoftObjectPath.h
@@ -592,7 +592,9 @@ enum class ESoftObjectPathSerializeType : uint8
 	AlwaysSerialize,
 };
 
-class COREUOBJECT_API FSoftObjectPathThreadContext : public TThreadSingleton<FSoftObjectPathThreadContext>
+UE_DECLARE_THREAD_SINGLETON_TLS(class FSoftObjectPathThreadContext, COREUOBJECT_API)
+
+class FSoftObjectPathThreadContext : public TThreadSingleton<FSoftObjectPathThreadContext>
 {
 	friend TThreadSingleton<FSoftObjectPathThreadContext>;
 	friend struct FSoftObjectPathSerializationScope;
@@ -632,7 +634,7 @@ public:
 	 * @param OutCollectType Type of collecting that should be done
 	 * @param Archive The FArchive that is serializing this path if known. If null it will check FUObjectThreadContext
 	 */
-	bool GetSerializationOptions(FName& OutPackageName, FName& OutPropertyName, ESoftObjectPathCollectType& OutCollectType, ESoftObjectPathSerializeType& OutSerializeType, FArchive* Archive = nullptr) const;
+	COREUOBJECT_API bool GetSerializationOptions(FName& OutPackageName, FName& OutPropertyName, ESoftObjectPathCollectType& OutCollectType, ESoftObjectPathSerializeType& OutSerializeType, FArchive* Archive = nullptr) const;
 };
 
 /** Helper class to set and restore serialization options for soft object paths */
diff --git a/Engine/Source/Runtime/CoreUObject/Public/UObject/UObjectThreadContext.h b/Engine/Source/Runtime/CoreUObject/Public/UObject/UObjectThreadContext.h
--- a/Engine/Source/Runtime/CoreUObject/Public/UObject/UObjectThreadContext.h
+++ b/Engine/Source/Runtime/CoreUObject/Public/UObject/UObjectThreadContext.h
@@ -33,16 +33,18 @@ struct FUObjectSerializeContext;
 
 COREUOBJECT_API DECLARE_LOG_CATEGORY_EXTERN(LogUObjectThreadContext, Log, All);
 
-class COREUOBJECT_API FUObjectThreadContext : public TThreadSingleton<FUObjectThreadContext>
+UE_DECLARE_THREAD_SINGLETON_TLS(class FUObjectThreadContext, COREUOBJECT_API)
+
+class FUObjectThreadContext : public TThreadSingleton<FUObjectThreadContext>
 {
 	friend TThreadSingleton<FUObjectThreadContext>;
 
-	FUObjectThreadContext();
-	virtual ~FUObjectThreadContext();
+	COREUOBJECT_API FUObjectThreadContext();
+	COREUOBJECT_API virtual ~FUObjectThreadContext();
 #if WITH_EDITORONLY_DATA
 	// Remove declaration of manual constructors after removing deprecated PackagesMarkedEditorOnlyByOtherPackage
-	FUObjectThreadContext(const FUObjectThreadContext& Other);
-	FUObjectThreadContext(FUObjectThreadContext&& Other);
+	COREUOBJECT_API FUObjectThreadContext(const FUObjectThreadContext& Other);
+	COREUOBJECT_API FUObjectThreadContext(FUObjectThreadContext&& Other);
 #endif
 
 	/** Stack of currently used FObjectInitializers for this thread */
@@ -127,7 +129,7 @@ public:
 
 private:
 	/** Report that the current ObjectInitializer is null. */
-	FObjectInitializer& ReportNull();
+	COREUOBJECT_API FObjectInitializer& ReportNull();
 
 	/** Current serialization context */
 	TRefCountPtr<FUObjectSerializeContext> SerializeContext;
diff --git a/Engine/Source/Runtime/Experimental/Chaos/Private/ChaosDebugDraw/ChaosDDContext.cpp b/Engine/Source/Runtime/Experimental/Chaos/Private/ChaosDebugDraw/ChaosDDContext.cpp
--- a/Engine/Source/Runtime/Experimental/Chaos/Private/ChaosDebugDraw/ChaosDDContext.cpp
+++ b/Engine/Source/Runtime/Experimental/Chaos/Private/ChaosDebugDraw/ChaosDDContext.cpp
@@ -11,6 +11,8 @@ DEFINE_LOG_CATEGORY(LogChaosDD);
 
 #if CHAOS_DEBUG_DRAW
 
+UE_DEFINE_THREAD_SINGLETON_TLS(ChaosDD::Private::FChaosDDContext)
+
 namespace ChaosDD::Private
 {
 	bool bChaosDebugDraw_EnableGlobalQueue = true;
diff --git a/Engine/Source/Runtime/Experimental/Chaos/Private/ChaosVisualDebugger/ChaosVDContextProvider.cpp b/Engine/Source/Runtime/Experimental/Chaos/Private/ChaosVisualDebugger/ChaosVDContextProvider.cpp
--- a/Engine/Source/Runtime/Experimental/Chaos/Private/ChaosVisualDebugger/ChaosVDContextProvider.cpp
+++ b/Engine/Source/Runtime/Experimental/Chaos/Private/ChaosVisualDebugger/ChaosVDContextProvider.cpp
@@ -5,6 +5,9 @@
 #include "ChaosVisualDebugger/ChaosVisualDebuggerTrace.h"
 
 #if WITH_CHAOS_VISUAL_DEBUGGER
+
+UE_DEFINE_THREAD_SINGLETON_TLS(FChaosVDThreadContext)
+
 bool FChaosVDThreadContext::GetCurrentContext(FChaosVDContext& OutContext)
 {
 	if (LocalContextStack.Num() > 0)
diff --git a/Engine/Source/Runtime/Experimental/Chaos/Private/Framework/Threading.cpp b/Engine/Source/Runtime/Experimental/Chaos/Private/Framework/Threading.cpp
--- a/Engine/Source/Runtime/Experimental/Chaos/Private/Framework/Threading.cpp
+++ b/Engine/Source/Runtime/Experimental/Chaos/Private/Framework/Threading.cpp
@@ -3,6 +3,10 @@
 #include "Framework/Threading.h"
 #include "ChaosCheck.h"
 
+#if PHYSICS_THREAD_CONTEXT
+UE_DEFINE_THREAD_SINGLETON_TLS(Chaos::FPhysicsThreadContext)
+#endif
+
 namespace Chaos
 {
 	// uint32 pair for tracking reads and writes per thread.
diff --git a/Engine/Source/Runtime/Experimental/Chaos/Public/ChaosDebugDraw/ChaosDDContext.h b/Engine/Source/Runtime/Experimental/Chaos/Public/ChaosDebugDraw/ChaosDDContext.h
--- a/Engine/Source/Runtime/Experimental/Chaos/Public/ChaosDebugDraw/ChaosDDContext.h
+++ b/Engine/Source/Runtime/Experimental/Chaos/Public/ChaosDebugDraw/ChaosDDContext.h
@@ -9,16 +9,23 @@
 
 #if CHAOS_DEBUG_DRAW
 
+namespace ChaosDD::Private
+{
+	class FChaosDDContext;
+}
+
+UE_DECLARE_THREAD_SINGLETON_TLS(ChaosDD::Private::FChaosDDContext, CHAOS_API)
+
 namespace ChaosDD::Private
 {
 	//
 	// A thread-local debug draw context used to access the queue to draw to
 	// for any thread on which debug draw has been set up.
 	//
-	class CHAOS_API FChaosDDContext : public TThreadSingleton<FChaosDDContext>
+	class FChaosDDContext : public TThreadSingleton<FChaosDDContext>
 	{
 	public:
-		FChaosDDContext();
+		CHAOS_API FChaosDDContext();
 
 		// Whether the debug draw system is enabled
 		static bool IsDebugDrawEnabled()
@@ -43,7 +50,7 @@ namespace ChaosDD::Private
 		}
 
 		// For internal use - collect all out-of-frame debug draw commands for rendering
-		static FChaosDDFramePtr ExtractGlobalFrame();
+		static CHAOS_API FChaosDDFramePtr ExtractGlobalFrame();
 
 	private:
 		friend class FChaosDDScene;
@@ -67,7 +74,7 @@ namespace ChaosDD::Private
 		}
 
 		// Global frame management
-		static const FChaosDDFramePtr& GetGlobalFrame();
+		static CHAOS_API const FChaosDDFramePtr& GetGlobalFrame();
 		static void CreateGlobalFrame();
 		static void SetGlobalDrawRegion(const FSphere3d& InDrawRegion);
 		static void SetGlobalCommandBudget(int32 InCommandBudget);
@@ -76,7 +83,7 @@ namespace ChaosDD::Private
 		FChaosDDFramePtr Frame;
 
 		// Whether the system is enabled
-		static bool bDebugDrawEnabled;
+		static CHAOS_API bool bDebugDrawEnabled;
 
 		// Global frame: fallback for out-of-context debug draw
 		static FCriticalSection GlobalFrameCS;
@@ -91,11 +98,11 @@ namespace ChaosDD::Private
 	// The active context should be accessed via FChaosDDContext::GetWriter(). FChaosDDTimelineContext
 	// is not directly used other than to instantiate.
 	//
-	class CHAOS_API FChaosDDTimelineContext
+	class FChaosDDTimelineContext
 	{
 	public:
-		void BeginFrame(const FChaosDDTimelinePtr& InTimeline, double InTime, double InDt);
-		void EndFrame();
+		CHAOS_API void BeginFrame(const FChaosDDTimelinePtr& InTimeline, double InTime, double InDt);
+		CHAOS_API void EndFrame();
 
 	private:
 		FChaosDDTimelinePtr Timeline;
@@ -106,11 +113,11 @@ namespace ChaosDD::Private
 	//
 	// A scoped wrapper for FChaosDDTimelineContext
 	//
-	class CHAOS_API FChaosDDScopeTimelineContext
+	class FChaosDDScopeTimelineContext
 	{
 	public:
-		FChaosDDScopeTimelineContext(const FChaosDDTimelinePtr& InTimeline, double InTime, double InDt);
-		~FChaosDDScopeTimelineContext();
+		CHAOS_API FChaosDDScopeTimelineContext(const FChaosDDTimelinePtr& InTimeline, double InTime, double InDt);
+		CHAOS_API ~FChaosDDScopeTimelineContext();
 
 	private:
 		FChaosDDTimelineContext Context;
@@ -124,10 +131,10 @@ namespace ChaosDD::Private
 	//		- put FChaosDDScopeTaskContext(ParentContext) on the child thread
 	// (Search for FChaosDDScopeTaskContext for examples.)
 	//
-	class CHAOS_API FChaosDDTaskParentContext
+	class FChaosDDTaskParentContext
 	{
 	public:
-		FChaosDDTaskParentContext();
+		CHAOS_API FChaosDDTaskParentContext();
 	private:
 		friend class FChaosDDTaskContext;
 		FChaosDDFramePtr Frame;
@@ -148,11 +155,11 @@ namespace ChaosDD::Private
 	// @todo(chaos): ChaosDDFrame should track how many contexts it is referenced by and
 	// assert that it is not active when we end the frame.
 	//
-	class CHAOS_API FChaosDDTaskContext
+	class FChaosDDTaskContext
 	{
 	public:
-		void BeginThread(const FChaosDDTaskParentContext& InParentDDContext);
-		void EndThread();
+		CHAOS_API void BeginThread(const FChaosDDTaskParentContext& InParentDDContext);
+		CHAOS_API void EndThread();
 
 	private:
 		FChaosDDFramePtr PreviousFrame;
@@ -162,11 +169,11 @@ namespace ChaosDD::Private
 	//
 	// A scoped wrapper for FChaosDDTaskContext
 	//
-	class CHAOS_API FChaosDDScopeTaskContext
+	class FChaosDDScopeTaskContext
 	{
 	public:
-		FChaosDDScopeTaskContext(const FChaosDDTaskParentContext& InParentDDContext);
-		~FChaosDDScopeTaskContext();
+		CHAOS_API FChaosDDScopeTaskContext(const FChaosDDTaskParentContext& InParentDDContext);
+		CHAOS_API ~FChaosDDScopeTaskContext();
 
 	private:
 		FChaosDDTaskContext Context;
diff --git a/Engine/Source/Runtime/Experimental/Chaos/Public/ChaosVisualDebugger/ChaosVDContextProvider.h b/Engine/Source/Runtime/Experimental/Chaos/Public/ChaosVisualDebugger/ChaosVDContextProvider.h
--- a/Engine/Source/Runtime/Experimental/Chaos/Public/ChaosVisualDebugger/ChaosVDContextProvider.h
+++ b/Engine/Source/Runtime/Experimental/Chaos/Public/ChaosVisualDebugger/ChaosVDContextProvider.h
@@ -46,8 +46,10 @@ private:
 	TSharedRef<Chaos::VisualDebugger::FChaosVDOptionalDataChannel> CurrentDataChannel = CVDDC_Default;
 };
 
+UE_DECLARE_THREAD_SINGLETON_TLS(class FChaosVDThreadContext, CHAOS_API)
+
 /** Singleton class that manages the thread local storage used to store CVD Context data */
-class CHAOS_API FChaosVDThreadContext : public TThreadSingleton<FChaosVDThreadContext>
+class FChaosVDThreadContext : public TThreadSingleton<FChaosVDThreadContext>
 {
 public:
 
@@ -62,26 +64,26 @@ public:
 	/** Copies the Current CVD context data into the provided struct
 	 * @return true if the copy was successful
 	 */
-	bool GetCurrentContext(FChaosVDContext& OutContext);
+	CHAOS_API bool GetCurrentContext(FChaosVDContext& OutContext);
 	
 	/** Gets the current CVD context data -
 	 * Don't use of a function that will recursively push new context data as it might invalidate the pointer
 	 * @param Type type of the context we want to get
 	 * @return Ptr to the Current CVD context data if it matches the provided type
 	 */
-	const FChaosVDContext* GetCurrentContext(EChaosVDContextType Type);
+	CHAOS_API const FChaosVDContext* GetCurrentContext(EChaosVDContextType Type);
 	
 	/** Gets the current CVD context data -
 	 * Don't use of a function that will recursively push new context data as it might invalidate the pointer
 	 * @return Ptr to the Current CVD context data
 	 */
-	const FChaosVDContext* GetCurrentContext();
+	CHAOS_API const FChaosVDContext* GetCurrentContext();
 	
 	/** Pushed a new CVD Context data to the local cvd context stack */
-	void PushContext(const FChaosVDContext& InContext);
+	CHAOS_API void PushContext(const FChaosVDContext& InContext);
 	
 	/** Removes the CVD Context data at the top of the local cvd context stack */
-	void PopContext();
+	CHAOS_API void PopContext();
 
 protected:
 	TArray<FChaosVDContext, TInlineAllocator<16>> LocalContextStack;
diff --git a/Engine/Source/Runtime/Experimental/Chaos/Public/Framework/Threading.h b/Engine/Source/Runtime/Experimental/Chaos/Public/Framework/Threading.h
--- a/Engine/Source/Runtime/Experimental/Chaos/Public/Framework/Threading.h
+++ b/Engine/Source/Runtime/Experimental/Chaos/Public/Framework/Threading.h
@@ -61,6 +61,12 @@
 	#endif
 #endif
 
+#if PHYSICS_THREAD_CONTEXT
+namespace Chaos { class FPhysicsThreadContext; }
+UE_DECLARE_THREAD_SINGLETON_TLS(Chaos::FPhysicsThreadContext, CHAOS_API)
+#endif
+
+
 namespace Chaos
 {
 	// Not intended for external callers, provided here to allow the locks below to record depths
@@ -157,8 +163,9 @@ namespace Chaos
 #define CHAOS_RECORD_LEAVE_WRITE_LOCK Chaos::ThreadingPrivate::DecWriteDepth(this);
 
 #if PHYSICS_THREAD_CONTEXT
+
 /** Debug helper to ensure threading mistakes are caught. Do not use for ship */
-class CHAOS_API FPhysicsThreadContext : public TThreadSingleton<FPhysicsThreadContext>
+class FPhysicsThreadContext : public TThreadSingleton<FPhysicsThreadContext>
 {
 public:
 	bool IsInPhysicsSimContext() const
