[Chaos]
* Moved dll storage to methods/variables instead of on types. This is needed to get merged modules working

diff --git a/Engine/Source/Runtime/Experimental/Chaos/Private/Chaos/Collision/PBDCollisionConstraintHandle.cpp b/Engine/Source/Runtime/Experimental/Chaos/Private/Chaos/Collision/PBDCollisionConstraintHandle.cpp
--- a/Engine/Source/Runtime/Experimental/Chaos/Private/Chaos/Collision/PBDCollisionConstraintHandle.cpp
+++ b/Engine/Source/Runtime/Experimental/Chaos/Private/Chaos/Collision/PBDCollisionConstraintHandle.cpp
@@ -7,4 +7,10 @@
 
 namespace Chaos
 {
+
+	const FConstraintHandleTypeID& FPBDCollisionConstraintHandle::StaticType()
+	{
+		static FConstraintHandleTypeID STypeID(TEXT("FCollisionConstraintHandle"), &FIntrusiveConstraintHandle::StaticType());
+		return STypeID;
+	}
 }
\ No newline at end of file
diff --git a/Engine/Source/Runtime/Experimental/Chaos/Public/Chaos/Collision/PBDCollisionConstraintHandle.h b/Engine/Source/Runtime/Experimental/Chaos/Public/Chaos/Collision/PBDCollisionConstraintHandle.h
--- a/Engine/Source/Runtime/Experimental/Chaos/Public/Chaos/Collision/PBDCollisionConstraintHandle.h
+++ b/Engine/Source/Runtime/Experimental/Chaos/Public/Chaos/Collision/PBDCollisionConstraintHandle.h
@@ -58,8 +58,8 @@ namespace Chaos
 		{
 		}
 
-		CHAOS_API const FPBDCollisionConstraint& GetContact() const;
-		CHAOS_API FPBDCollisionConstraint& GetContact();
+		const FPBDCollisionConstraint& GetContact() const;
+		FPBDCollisionConstraint& GetContact();
 
 		UE_DEPRECATED(4.27, "Use GetContact()")
 		const FPBDCollisionConstraint& GetPointContact() const { return GetContact(); }
@@ -73,27 +73,23 @@ namespace Chaos
 		UE_DEPRECATED(4.27, "Use GetContact()")
 		FPBDCollisionConstraint& GetSweptPointContact() { return GetContact(); }
 
-		CHAOS_API bool GetCCDEnabled() const;
+		bool GetCCDEnabled() const;
 
-		CHAOS_API virtual void SetEnabled(bool InEnabled) override final;
+		virtual void SetEnabled(bool InEnabled) override final;
 
-		CHAOS_API virtual bool IsEnabled() const override final;
+		virtual bool IsEnabled() const override final;
 
-		CHAOS_API virtual bool IsProbe() const override final;
+		virtual bool IsProbe() const override final;
 
-		CHAOS_API FVec3 GetAccumulatedImpulse() const;
+		FVec3 GetAccumulatedImpulse() const;
 
 		// Declared final so that TPBDConstraintGraphRuleImpl::AddToGraph() does not need to hit vtable
 		CHAOS_API virtual FParticlePair GetConstrainedParticles() const override final;
 
-		CHAOS_API const FPBDCollisionConstraints* ConcreteContainer() const;
-		CHAOS_API FPBDCollisionConstraints* ConcreteContainer();
+		const FPBDCollisionConstraints* ConcreteContainer() const;
+		FPBDCollisionConstraints* ConcreteContainer();
 
-		static const FConstraintHandleTypeID& StaticType()
-		{
-			static FConstraintHandleTypeID STypeID(TEXT("FCollisionConstraintHandle"), &FIntrusiveConstraintHandle::StaticType());
-			return STypeID;
-		}
+		static CHAOS_API const FConstraintHandleTypeID& StaticType();
 	};
 
 }
diff --git a/Engine/Source/Runtime/Experimental/Chaos/Public/Chaos/ShapeInstance.h b/Engine/Source/Runtime/Experimental/Chaos/Public/Chaos/ShapeInstance.h
--- a/Engine/Source/Runtime/Experimental/Chaos/Public/Chaos/ShapeInstance.h
+++ b/Engine/Source/Runtime/Experimental/Chaos/Public/Chaos/ShapeInstance.h
@@ -80,91 +80,91 @@ namespace Chaos
 
 		CHAOS_API virtual void Serialize(FChaosArchive& Ar);
 
-		CHAOS_API void UpdateShapeBounds(const FRigidTransform3& WorldTM);
+		void UpdateShapeBounds(const FRigidTransform3& WorldTM);
 
 		UE_DEPRECATED(5.4, "Bounds no longer expanded. Use UpdateShapeBounds without BoundsExpansion")
 		CHAOS_API void UpdateShapeBounds(const FRigidTransform3& WorldTM, const FVec3& BoundsExpansion) { UpdateShapeBounds(WorldTM); }
 
-		CHAOS_API void* GetUserData() const;
-		CHAOS_API void SetUserData(void* InUserData);
+		void* GetUserData() const;
+		void SetUserData(void* InUserData);
 
-		CHAOS_API const FCollisionFilterData& GetQueryData() const;
-		CHAOS_API void SetQueryData(const FCollisionFilterData& InQueryData);
+		const FCollisionFilterData& GetQueryData() const;
+		void SetQueryData(const FCollisionFilterData& InQueryData);
 
-		CHAOS_API const FCollisionFilterData& GetSimData() const;
-		CHAOS_API void SetSimData(const FCollisionFilterData& InSimData);
+		const FCollisionFilterData& GetSimData() const;
+		void SetSimData(const FCollisionFilterData& InSimData);
 
-		CHAOS_API FImplicitObjectRef GetGeometry() const;
+		FImplicitObjectRef GetGeometry() const;
 
-		CHAOS_API const TAABB<FReal, 3>& GetWorldSpaceShapeBounds() const;
+		const TAABB<FReal, 3>& GetWorldSpaceShapeBounds() const;
 
-		CHAOS_API void UpdateWorldSpaceState(const FRigidTransform3& WorldTransform);
+		void UpdateWorldSpaceState(const FRigidTransform3& WorldTransform);
 
 		UE_DEPRECATED(5.4, "Bounds no longer expanded. Use GetWorldSpaceShapeBounds()")
-		CHAOS_API const TAABB<FReal, 3>& GetWorldSpaceInflatedShapeBounds() const { return GetWorldSpaceShapeBounds(); }
+		const TAABB<FReal, 3>& GetWorldSpaceInflatedShapeBounds() const { return GetWorldSpaceShapeBounds(); }
 		
 		UE_DEPRECATED(5.4, "Bounds no longer expanded. Use UpdateWorldSpaceState() without BoundsExpansion")
-		CHAOS_API void UpdateWorldSpaceState(const FRigidTransform3& WorldTransform, const FVec3& BoundsExpansion) { UpdateWorldSpaceState(WorldTransform); }
+		void UpdateWorldSpaceState(const FRigidTransform3& WorldTransform, const FVec3& BoundsExpansion) { UpdateWorldSpaceState(WorldTransform); }
 
 		// The leaf shape (with transformed and implicit wrapper removed).
-		CHAOS_API const FImplicitObject* GetLeafGeometry() const;
+		const FImplicitObject* GetLeafGeometry() const;
 
 		// The actor-relative transform of the leaf geometry.
-		CHAOS_API FRigidTransform3 GetLeafRelativeTransform() const;
+		FRigidTransform3 GetLeafRelativeTransform() const;
 
 		// The world-space transform of the leaf geometry.
 		// If we have non-identity leaf relative transform, is cached from the last call to UpdateWorldSpaceState.
 		// If not cahced, is constructed from arguments.
-		CHAOS_API FRigidTransform3 GetLeafWorldTransform(const FGeometryParticleHandle* Particle) const;
+		FRigidTransform3 GetLeafWorldTransform(const FGeometryParticleHandle* Particle) const;
 
-		CHAOS_API void UpdateLeafWorldTransform(FGeometryParticleHandle* Particle);
+		void UpdateLeafWorldTransform(FGeometryParticleHandle* Particle);
 
-		CHAOS_API int32 NumMaterials() const;
-		CHAOS_API const FMaterialHandle& GetMaterial(const int32 Index) const;
+		int32 NumMaterials() const;
+		const FMaterialHandle& GetMaterial(const int32 Index) const;
 
-		CHAOS_API const TArray<FMaterialHandle>& GetMaterials() const;
-		CHAOS_API void SetMaterial(FMaterialHandle InMaterial);
-		CHAOS_API void SetMaterials(const TArray<FMaterialHandle>& InMaterials);
-		CHAOS_API void SetMaterials(TArray<FMaterialHandle>&& InMaterials);
+		const TArray<FMaterialHandle>& GetMaterials() const;
+		void SetMaterial(FMaterialHandle InMaterial);
+		void SetMaterials(const TArray<FMaterialHandle>& InMaterials);
+		void SetMaterials(TArray<FMaterialHandle>&& InMaterials);
 
-		CHAOS_API const TArray<FMaterialMaskHandle>& GetMaterialMasks() const;
-		CHAOS_API void SetMaterialMasks(const TArray<FMaterialMaskHandle>& InMaterialMasks);
-		CHAOS_API void SetMaterialMasks(TArray<FMaterialMaskHandle>&& InMaterialMasks);
+		const TArray<FMaterialMaskHandle>& GetMaterialMasks() const;
+		void SetMaterialMasks(const TArray<FMaterialMaskHandle>& InMaterialMasks);
+		void SetMaterialMasks(TArray<FMaterialMaskHandle>&& InMaterialMasks);
 
-		CHAOS_API const TArray<uint32>& GetMaterialMaskMaps() const;
-		CHAOS_API void SetMaterialMaskMaps(const TArray<uint32>& InMaterialMaskMaps);
-		CHAOS_API void SetMaterialMaskMaps(TArray<uint32>&& InMaterialMaskMaps);
+		const TArray<uint32>& GetMaterialMaskMaps() const;
+		void SetMaterialMaskMaps(const TArray<uint32>& InMaterialMaskMaps);
+		void SetMaterialMaskMaps(TArray<uint32>&& InMaterialMaskMaps);
 
-		CHAOS_API const TArray<FMaterialHandle>& GetMaterialMaskMapMaterials() const;
-		CHAOS_API void SetMaterialMaskMapMaterials(const TArray<FMaterialHandle>& InMaterialMaskMapMaterials);
-		CHAOS_API void SetMaterialMaskMapMaterials(TArray<FMaterialHandle>&& InMaterialMaskMapMaterials);
+		const TArray<FMaterialHandle>& GetMaterialMaskMapMaterials() const;
+		void SetMaterialMaskMapMaterials(const TArray<FMaterialHandle>& InMaterialMaskMapMaterials);
+		void SetMaterialMaskMapMaterials(TArray<FMaterialHandle>&& InMaterialMaskMapMaterials);
 
-		CHAOS_API const FShapeDirtyFlags GetDirtyFlags() const;
+		const FShapeDirtyFlags GetDirtyFlags() const;
 
-		CHAOS_API bool GetQueryEnabled() const;
-		CHAOS_API void SetQueryEnabled(const bool bEnable);
+		bool GetQueryEnabled() const;
+		void SetQueryEnabled(const bool bEnable);
 
-		CHAOS_API bool GetSimEnabled() const;
-		CHAOS_API void SetSimEnabled(const bool bEnable);
+		bool GetSimEnabled() const;
+		void SetSimEnabled(const bool bEnable);
 
-		CHAOS_API bool GetIsProbe() const;
-		CHAOS_API void SetIsProbe(const bool bIsProbe);
+		bool GetIsProbe() const;
+		void SetIsProbe(const bool bIsProbe);
 
-		CHAOS_API EChaosCollisionTraceFlag GetCollisionTraceType() const;
-		CHAOS_API void SetCollisionTraceType(const EChaosCollisionTraceFlag InTraceFlag);
+		EChaosCollisionTraceFlag GetCollisionTraceType() const;
+		void SetCollisionTraceType(const EChaosCollisionTraceFlag InTraceFlag);
 
-		CHAOS_API const FCollisionData& GetCollisionData() const;
-		CHAOS_API void SetCollisionData(const FCollisionData& Data);
+		const FCollisionData& GetCollisionData() const;
+		void SetCollisionData(const FCollisionData& Data);
 
-		CHAOS_API const FMaterialData& GetMaterialData() const;
-		CHAOS_API void SetMaterialData(const FMaterialData& Data);
+		const FMaterialData& GetMaterialData() const;
+		void SetMaterialData(const FMaterialData& Data);
 
-		CHAOS_API void SyncRemoteData(FDirtyPropertiesManager& Manager, int32 ShapeDataIdx, FShapeDirtyData& RemoteData);
+		void SyncRemoteData(FDirtyPropertiesManager& Manager, int32 ShapeDataIdx, FShapeDirtyData& RemoteData);
 		
-		CHAOS_API void SetProxy(IPhysicsProxyBase* InProxy);
+		void SetProxy(IPhysicsProxyBase* InProxy);
 		
-		CHAOS_API int32 GetShapeIndex() const;
-		CHAOS_API void ModifyShapeIndex(int32 NewShapeIndex);
+		int32 GetShapeIndex() const;
+		void ModifyShapeIndex(int32 NewShapeIndex);
 
 		template <typename Lambda> void ModifySimData(const Lambda& LambdaFunc);
 		template <typename Lambda> void ModifyMaterials(const Lambda& LambdaFunc);
diff --git a/Engine/Source/Runtime/Experimental/Chaos/Public/Chaos/SimCallbackInput.h b/Engine/Source/Runtime/Experimental/Chaos/Public/Chaos/SimCallbackInput.h
--- a/Engine/Source/Runtime/Experimental/Chaos/Public/Chaos/SimCallbackInput.h
+++ b/Engine/Source/Runtime/Experimental/Chaos/Public/Chaos/SimCallbackInput.h
@@ -127,7 +127,7 @@ public:
 	const FSimCallbackOutput& operator*() const { return *SimCallbackOutput; }
 private:
 
-	CHAOS_API void Free_External();
+	void Free_External();
 
 	FSimCallbackOutput* SimCallbackOutput;
 	ISimCallbackObject* SimCallbackObject;
diff --git a/Engine/Source/Runtime/Experimental/Chaos/Public/ChaosDebugDraw/ChaosDDFrame.h b/Engine/Source/Runtime/Experimental/Chaos/Public/ChaosDebugDraw/ChaosDDFrame.h
--- a/Engine/Source/Runtime/Experimental/Chaos/Public/ChaosDebugDraw/ChaosDDFrame.h
+++ b/Engine/Source/Runtime/Experimental/Chaos/Public/ChaosDebugDraw/ChaosDDFrame.h
@@ -28,7 +28,7 @@ namespace ChaosDD::Private
 	//
 	// @todo(chaos): move commands to a per-thread buffer and eliminate write locks
 	//
-	class CHAOS_API FChaosDDFrame : public TSharedFromThis<FChaosDDFrame>
+	class FChaosDDFrame : public TSharedFromThis<FChaosDDFrame>
 	{
 	public:
 		FChaosDDFrame(const FChaosDDTimelineWeakPtr& InTimeline, int64 InFrameIndex, double InTime, double InDt, const FSphere3d& InDrawRegion, int32 InCommandBudget, int32 InCommandQueueLength)
@@ -261,7 +261,7 @@ namespace ChaosDD::Private
 	// 
 	// @todo(chaos): eventually all threads that want debug draw should have a valid frame and this will be redundant.
 	//
-	class CHAOS_API FChaosDDGlobalFrame : public FChaosDDFrame
+	class FChaosDDGlobalFrame : public FChaosDDFrame
 	{
 	public:
 
@@ -299,7 +299,7 @@ namespace ChaosDD::Private
 	// preventing the frame from being Ended. Eventually this will be a
 	// per-thread buffer to avoid the need for locks.
 	//
-	class CHAOS_API FChaosDDFrameWriter
+	class FChaosDDFrameWriter
 	{
 	public:
 		FChaosDDFrameWriter(const FChaosDDFramePtr& InFrame)
diff --git a/Engine/Source/Runtime/Experimental/Chaos/Public/ChaosDebugDraw/ChaosDDRenderer.h b/Engine/Source/Runtime/Experimental/Chaos/Public/ChaosDebugDraw/ChaosDDRenderer.h
--- a/Engine/Source/Runtime/Experimental/Chaos/Public/ChaosDebugDraw/ChaosDDRenderer.h
+++ b/Engine/Source/Runtime/Experimental/Chaos/Public/ChaosDebugDraw/ChaosDDRenderer.h
@@ -16,7 +16,7 @@ namespace ChaosDD::Private
 	//
 	// Primitive rendering API for use by debug draw objects
 	//
-	class CHAOS_API IChaosDDRenderer
+	class IChaosDDRenderer
 	{
 	public:
 		IChaosDDRenderer() {}
diff --git a/Engine/Source/Runtime/Experimental/Chaos/Public/ChaosVisualDebugger/ChaosVDOptionalDataChannel.h b/Engine/Source/Runtime/Experimental/Chaos/Public/ChaosVisualDebugger/ChaosVDOptionalDataChannel.h
--- a/Engine/Source/Runtime/Experimental/Chaos/Public/ChaosVisualDebugger/ChaosVDOptionalDataChannel.h
+++ b/Engine/Source/Runtime/Experimental/Chaos/Public/ChaosVisualDebugger/ChaosVDOptionalDataChannel.h
@@ -67,11 +67,11 @@ namespace Chaos::VisualDebugger
 	};
 
 	/** Manager that provides a way to iterate through all existing CVD Data channels or get a specific one */
-	class CHAOS_API FChaosVDDataChannelsManager
+	class FChaosVDDataChannelsManager
 	{
 	public:
 
-		static FChaosVDDataChannelsManager& Get();
+		static CHAOS_API FChaosVDDataChannelsManager& Get();
 
 		/** Iterates through all the available data channels and executes the provided callback passing each channel as an argument.
 		 * If the callback returns false, the iteration will be stopped.
diff --git a/Engine/Source/Runtime/Experimental/Chaos/Public/Field/FieldSystemTypes.h b/Engine/Source/Runtime/Experimental/Chaos/Public/Field/FieldSystemTypes.h
--- a/Engine/Source/Runtime/Experimental/Chaos/Public/Field/FieldSystemTypes.h
+++ b/Engine/Source/Runtime/Experimental/Chaos/Public/Field/FieldSystemTypes.h
@@ -233,7 +233,7 @@ enum EFieldOutputType : int
 };
 
 inline
-TArray<EFieldPhysicsType> CHAOS_API GetFieldTargetTypes(EFieldOutputType OutputType)
+TArray<EFieldPhysicsType> GetFieldTargetTypes(EFieldOutputType OutputType)
 {
 	TArray<EFieldPhysicsType> PhysicsTypes;
 	switch (OutputType)
@@ -276,7 +276,7 @@ TArray<EFieldPhysicsType> CHAOS_API GetFieldTargetTypes(EFieldOutputType OutputT
 }
 
 inline
-EFieldOutputType CHAOS_API GetFieldTargetIndex(const TArray<EFieldPhysicsType>& VectorTypes, 
+EFieldOutputType GetFieldTargetIndex(const TArray<EFieldPhysicsType>& VectorTypes, 
 											   const TArray<EFieldPhysicsType>& ScalarTypes, 
 											   const TArray<EFieldPhysicsType>& IntegerTypes,
 											   const EFieldPhysicsType FieldTarget, int32& TargetIndex)
@@ -308,7 +308,7 @@ EFieldOutputType CHAOS_API GetFieldTargetIndex(const TArray<EFieldPhysicsType>&
 }
 
 inline
-EFieldOutputType CHAOS_API GetFieldTargetOutput(const EFieldPhysicsType FieldTarget)
+EFieldOutputType GetFieldTargetOutput(const EFieldPhysicsType FieldTarget)
 {
 	EFieldOutputType OutputType = EFieldOutputType::Field_Output_Max;
 
@@ -321,7 +321,7 @@ EFieldOutputType CHAOS_API GetFieldTargetOutput(const EFieldPhysicsType FieldTar
 }
 
 inline
-FName CHAOS_API GetFieldOutputName(const EFieldOutputType Type)
+FName GetFieldOutputName(const EFieldOutputType Type)
 {
 	switch (Type)
 	{
@@ -335,7 +335,7 @@ FName CHAOS_API GetFieldOutputName(const EFieldOutputType Type)
 	return "None";
 }
 
-inline CHAOS_API const TArray<FName>& GetFieldPhysicsNames()
+inline const TArray<FName>& GetFieldPhysicsNames()
 {
 	static const TArray<FName> FieldPhysicsNames =
 	{ "Node", "DynamicState", "LinearForce", "ExternalClusterStrain", "Kill", "LinearVelocity", "AngularVelocity",
@@ -345,7 +345,7 @@ inline CHAOS_API const TArray<FName>& GetFieldPhysicsNames()
 	return FieldPhysicsNames;
 }
 
-inline CHAOS_API const TArray<EFieldPhysicsType>& GetFieldPhysicsTypes()
+inline const TArray<EFieldPhysicsType>& GetFieldPhysicsTypes()
 {
 	static const TArray<EFieldPhysicsType> FieldPhysicsTypes ={ EFieldPhysicsType::Field_None, 
 																EFieldPhysicsType::Field_DynamicState, 
@@ -371,7 +371,7 @@ inline CHAOS_API const TArray<EFieldPhysicsType>& GetFieldPhysicsTypes()
 	return FieldPhysicsTypes;
 }
 
-inline CHAOS_API const FName& GetFieldPhysicsName(EFieldPhysicsType PhysicsType)
+inline const FName& GetFieldPhysicsName(EFieldPhysicsType PhysicsType)
 {
 	const TArray<FName>& FieldPhysicsNames = GetFieldPhysicsNames();
 	static const FName NoneField("");
@@ -379,7 +379,7 @@ inline CHAOS_API const FName& GetFieldPhysicsName(EFieldPhysicsType PhysicsType)
 	return (PhysicsType < FieldPhysicsNames.Num()) ? FieldPhysicsNames[PhysicsType] : NoneField;
 }
 
-inline CHAOS_API EFieldPhysicsType GetFieldPhysicsType(const FName& PhysicsName)
+inline EFieldPhysicsType GetFieldPhysicsType(const FName& PhysicsName)
 {
 	const TArray<FName>& FieldPhysicsNames = GetFieldPhysicsNames();
 	const int32 PhysicsIndex = FieldPhysicsNames.Find(PhysicsName);
@@ -405,7 +405,7 @@ enum EFieldPhysicsDefaultFields : int
 };
 
 USTRUCT()
-struct CHAOS_API FFieldCollection : public FManagedArrayCollection
+struct FFieldCollection : public FManagedArrayCollection
 {
 	GENERATED_USTRUCT_BODY()
 public:
diff --git a/Engine/Source/Runtime/Experimental/Chaos/Public/GeometryCollection/GeometryCollectionSimulationTypes.h b/Engine/Source/Runtime/Experimental/Chaos/Public/GeometryCollection/GeometryCollectionSimulationTypes.h
--- a/Engine/Source/Runtime/Experimental/Chaos/Public/GeometryCollection/GeometryCollectionSimulationTypes.h
+++ b/Engine/Source/Runtime/Experimental/Chaos/Public/GeometryCollection/GeometryCollectionSimulationTypes.h
@@ -63,7 +63,7 @@ enum class EGeometryCollectionPhysicsTypeEnum : uint8
 };
 
 
-inline CHAOS_API EFieldPhysicsType GetGeometryCollectionPhysicsType(const EGeometryCollectionPhysicsTypeEnum GeoCollectionType)
+inline EFieldPhysicsType GetGeometryCollectionPhysicsType(const EGeometryCollectionPhysicsTypeEnum GeoCollectionType)
 {
 	static const TArray<EFieldPhysicsType> PhysicsTypes = {
 		EFieldPhysicsType::Field_AngularVelociy,
@@ -120,7 +120,7 @@ enum class EDamageModelTypeEnum : uint8
 };
 
 // convert user level damage model to chaos damage evaluation model
-inline CHAOS_API Chaos::EDamageEvaluationModel GetDamageEvaluationModel(const EDamageModelTypeEnum DamageModel)
+inline Chaos::EDamageEvaluationModel GetDamageEvaluationModel(const EDamageModelTypeEnum DamageModel)
 {
 	switch (DamageModel)
 	{
diff --git a/Engine/Source/Runtime/Experimental/Chaos/Public/PhysicsProxy/SingleParticlePhysicsProxy.h b/Engine/Source/Runtime/Experimental/Chaos/Public/PhysicsProxy/SingleParticlePhysicsProxy.h
--- a/Engine/Source/Runtime/Experimental/Chaos/Public/PhysicsProxy/SingleParticlePhysicsProxy.h
+++ b/Engine/Source/Runtime/Experimental/Chaos/Public/PhysicsProxy/SingleParticlePhysicsProxy.h
@@ -60,7 +60,7 @@ public:
 	using PARTICLE_TYPE = FGeometryParticle;
 	using FParticleHandle = FGeometryParticleHandle;
 
-	static CHAOS_API FSingleParticlePhysicsProxy* Create(TUniquePtr<FGeometryParticle>&& Particle, UObject* InOwner);
+	static FSingleParticlePhysicsProxy* Create(TUniquePtr<FGeometryParticle>&& Particle, UObject* InOwner);
 
 	FSingleParticlePhysicsProxy() = delete;
 	FSingleParticlePhysicsProxy(const FSingleParticlePhysicsProxy&) = delete;
diff --git a/Engine/Source/Runtime/Experimental/ChaosVisualDebugger/Public/ChaosVDRuntimeModule.h b/Engine/Source/Runtime/Experimental/ChaosVisualDebugger/Public/ChaosVDRuntimeModule.h
--- a/Engine/Source/Runtime/Experimental/ChaosVisualDebugger/Public/ChaosVDRuntimeModule.h
+++ b/Engine/Source/Runtime/Experimental/ChaosVisualDebugger/Public/ChaosVDRuntimeModule.h
@@ -12,6 +12,8 @@
 #include "Misc/ScopeRWLock.h"
 #include "Misc/TransactionallySafeRWLock.h"
 
+#define UE_API CHAOSVDRUNTIME_API
+
 struct FChaosVDRecording;
 class FText;
 
@@ -28,29 +30,29 @@ DECLARE_MULTICAST_DELEGATE(FChaosVDRecordingStateChangedDelegate)
 DECLARE_MULTICAST_DELEGATE_OneParam(FChaosVDCaptureRequestDelegate, EChaosVDFullCaptureFlags)
 DECLARE_MULTICAST_DELEGATE_OneParam(FChaosVDRecordingStartFailedDelegate, const FText&)
 
-class CHAOSVDRUNTIME_API FChaosVDRuntimeModule : public IModuleInterface
+class FChaosVDRuntimeModule : public IModuleInterface
 {
 public:
 
-	static FChaosVDRuntimeModule& Get();
-	static bool IsLoaded();
+	static UE_API FChaosVDRuntimeModule& Get();
+	static UE_API bool IsLoaded();
 	
-	virtual void StartupModule() override;
-	virtual void ShutdownModule() override;
+	UE_API virtual void StartupModule() override;
+	UE_API virtual void ShutdownModule() override;
 
 	/** Starts a CVD recording by starting a Trace session. It will stop any existing trace session
 	 * @param Args : Arguments array provided by the commandline. Used to determine if we want to record to file or a local trace server
 	 */
-	void StartRecording(TConstArrayView<FString> Args);
+	UE_API void StartRecording(TConstArrayView<FString> Args);
 	
 	/* Stops an active recording */
-	void StopRecording();
+	UE_API void StopRecording();
 
 	/** Returns true if we are currently recording a Physics simulation */
 	bool IsRecording() const { return bIsRecording; }
 
 	/** Returns a unique ID used to be used to identify CVD (Chaos Visual Debugger) data */
-	int32 GenerateUniqueID();
+	UE_API int32 GenerateUniqueID();
 
 	static FDelegateHandle RegisterRecordingStartedCallback(const FChaosVDRecordingStateChangedDelegate::FDelegate& InCallback)
 	{
@@ -104,9 +106,9 @@ public:
 	float GetAccumulatedRecordingTime() const { return AccumulatedRecordingTime; }
 
 	/** Returns the full path of the active recording file*/
-	FString GetLastRecordingFileNamePath() const;
+	UE_API FString GetLastRecordingFileNamePath() const;
 
-	FChaosVDTraceDetails GetCurrentTraceSessionDetails() const;
+	UE_API FChaosVDTraceDetails GetCurrentTraceSessionDetails() const;
 
 	EChaosVDRecordingMode GetCurrentRecordingMode() const
 	{
@@ -116,26 +118,26 @@ public:
 private:
 
 	/** Stops the current Trace session */
-	void StopTrace();
+	UE_API void StopTrace();
 	/** Finds a valid file name for a new file - Used to generate the file name for the Trace recording */
-	void GenerateRecordingFileName(FString& OutFileName);
+	UE_API void GenerateRecordingFileName(FString& OutFileName);
 
 	/** Queues a full Capture of the simulation on the next frame */
-	bool RequestFullCapture(float DeltaTime);
+	UE_API bool RequestFullCapture(float DeltaTime);
 
 	/** Queues a full Capture of the simulation on the next frame */
-	bool RecordingTimerTick(float DeltaTime);
+	UE_API bool RecordingTimerTick(float DeltaTime);
 
 	/** Used to handle stop requests to the active trace session that were not done by us
 	 * That is a possible scenario because Trace is shared by other In-Editor tools
 	 */
-	void HandleTraceStopRequest(FTraceAuxiliary::EConnectionType TraceType, const FString& TraceDestination);
+	UE_API void HandleTraceStopRequest(FTraceAuxiliary::EConnectionType TraceType, const FString& TraceDestination);
 
-	bool WaitForTraceSessionDisconnect();
+	UE_API bool WaitForTraceSessionDisconnect();
 
-	void EnableRequiredTraceChannels();
-	void SaveAndDisabledCurrentEnabledTraceChannels();
-	void RestoreTraceChannelsToPreRecordingState();
+	UE_API void EnableRequiredTraceChannels();
+	UE_API void SaveAndDisabledCurrentEnabledTraceChannels();
+	UE_API void RestoreTraceChannelsToPreRecordingState();
 
 	bool bIsRecording = false;
 	bool bRequestedStop = false;
@@ -145,10 +147,10 @@ private:
 	FTSTicker::FDelegateHandle FullCaptureRequesterHandle;
 	FTSTicker::FDelegateHandle RecordingTimerHandle;
 
-	static FChaosVDRecordingStateChangedDelegate RecordingStartedDelegate;
-	static FChaosVDRecordingStateChangedDelegate RecordingStopDelegate;
-	static FChaosVDRecordingStartFailedDelegate RecordingStartFailedDelegate;
-	static FChaosVDCaptureRequestDelegate PerformFullCaptureDelegate;
+	static UE_API FChaosVDRecordingStateChangedDelegate RecordingStartedDelegate;
+	static UE_API FChaosVDRecordingStateChangedDelegate RecordingStopDelegate;
+	static UE_API FChaosVDRecordingStartFailedDelegate RecordingStartFailedDelegate;
+	static UE_API FChaosVDCaptureRequestDelegate PerformFullCaptureDelegate;
 
 	std::atomic<int32> LastGeneratedID;
 
@@ -158,6 +160,8 @@ private:
 
 	EChaosVDRecordingMode CurrentRecordingMode = EChaosVDRecordingMode::Invalid;
 
-	static FTransactionallySafeRWLock DelegatesRWLock;
+	static UE_API FTransactionallySafeRWLock DelegatesRWLock;
 };
 #endif
+
+#undef UE_API 
\ No newline at end of file
diff --git a/Engine/Source/Runtime/Experimental/ChaosVisualDebugger/Public/DataWrappers/ChaosVDParticleDataWrapper.h b/Engine/Source/Runtime/Experimental/ChaosVisualDebugger/Public/DataWrappers/ChaosVDParticleDataWrapper.h
--- a/Engine/Source/Runtime/Experimental/ChaosVisualDebugger/Public/DataWrappers/ChaosVDParticleDataWrapper.h
+++ b/Engine/Source/Runtime/Experimental/ChaosVisualDebugger/Public/DataWrappers/ChaosVDParticleDataWrapper.h
@@ -42,10 +42,10 @@ enum class EChaosVDObjectStateType: int8
 
 /** Base struct that declares the interface to be used for any ParticleData Viewer */
 USTRUCT()
-struct CHAOSVDRUNTIME_API FChaosVDWrapperDataBase
+struct FChaosVDWrapperDataBase
 {
 	GENERATED_BODY()
-	virtual ~FChaosVDWrapperDataBase() = default;
+	CHAOSVDRUNTIME_API virtual ~FChaosVDWrapperDataBase() = default;
 
 	virtual bool HasValidData() const { return bHasValidData; }
 
