Allowing C++ to set bUseRootProxyForNavigation, which was already exposed to Blueprint

#rnx

diff --git a/Engine/Source/Runtime/Experimental/GeometryCollectionEngine/Public/GeometryCollection/GeometryCollectionComponent.h b/Engine/Source/Runtime/Experimental/GeometryCollectionEngine/Public/GeometryCollection/GeometryCollectionComponent.h
--- a/Engine/Source/Runtime/Experimental/GeometryCollectionEngine/Public/GeometryCollection/GeometryCollectionComponent.h
+++ b/Engine/Source/Runtime/Experimental/GeometryCollectionEngine/Public/GeometryCollection/GeometryCollectionComponent.h
@@ -1326,6 +1326,8 @@ public:
 
 	GEOMETRYCOLLECTIONENGINE_API void SetUpdateNavigationInTick(const bool bUpdateInTick) { bUpdateNavigationInTick = bUpdateInTick; }
 
+	GEOMETRYCOLLECTIONENGINE_API void SetUseRootProxyForNavigation(const bool bUseRootProxyForNav) { bUseRootProxyForNavigation = bUseRootProxyForNav; }
+
 	// todo(chaos): Remove this and move to a cook time approach of the SM data based on the GC property
 	UFUNCTION(BlueprintPure, BlueprintInternalUseOnly, Category = "Physics")
 	GEOMETRYCOLLECTIONENGINE_API bool GetUseStaticMeshCollisionForTraces() const { return bUseStaticMeshCollisionForTraces; }
