Iris
* Added ESubObjectInsertionOrder::InsertFirst option to make a specific subobject always the first to be created or to replicate

#rnx

diff --git a/Engine/Plugins/Runtime/ReplicationSystemTestPlugin/Source/Private/Tests/ReplicationSystem/ReplicatedTestObject.h b/Engine/Plugins/Runtime/ReplicationSystemTestPlugin/Source/Private/Tests/ReplicationSystem/ReplicatedTestObject.h
--- a/Engine/Plugins/Runtime/ReplicationSystemTestPlugin/Source/Private/Tests/ReplicationSystem/ReplicatedTestObject.h
+++ b/Engine/Plugins/Runtime/ReplicationSystemTestPlugin/Source/Private/Tests/ReplicationSystem/ReplicatedTestObject.h
@@ -36,6 +36,9 @@ public:
 	virtual void RegisterReplicationFragments(UE::Net::FFragmentRegistrationContext& Fragments, UE::Net::EFragmentRegistrationFlags RegistrationFlags) override {}
 	virtual bool IsSupportedForNetworking() const override { return true; }
 
+	virtual void OnSubObjectCreated(UObject* SubObject) {}
+	virtual void OnSubObjectDestroyed(UObject* SubObject) {}
+
 	/** Cached NetRefHandle to simplify testing. DO NOT use this for multi-system tests. */
 	UE::Net::FNetRefHandle NetRefHandle;
 
@@ -443,6 +446,8 @@ public:
 		LastRepOrderCounter = ++RepOrderCounter;
 	}
 
+	int32 CreationOrder = INDEX_NONE;
+
 public:
 	static uint32 RepOrderCounter;
 
@@ -450,6 +455,29 @@ public:
 	TArray<UE::Net::FReplicationFragment*> ReplicationFragments;
 };
 
+UCLASS()
+class UReplicatedObjectTestSubObjectCreationOrder : public UReplicatedTestObject
+{
+	GENERATED_BODY()
+
+public:
+	UReplicatedObjectTestSubObjectCreationOrder() {}
+
+	// Network interface must be part of base.
+	virtual void RegisterReplicationFragments(UE::Net::FFragmentRegistrationContext& Fragments, UE::Net::EFragmentRegistrationFlags RegistrationFlags) {};
+
+public:
+	virtual void OnSubObjectCreated(UObject* SubObject) override
+	{
+		if (UReplicatedSubObjectOrderObject* SubObjectOrder = Cast<UReplicatedSubObjectOrderObject>(SubObject))
+		{
+			SubObjectOrder->CreationOrder = ++CurrentCreationOrder;
+		}
+	}
+
+	int32 CurrentCreationOrder = 0;
+};
+
 UCLASS()
 class UReplicatedSubObjectDestroyOrderObject : public UReplicatedSubObjectOrderObject
 {
diff --git a/Engine/Plugins/Runtime/ReplicationSystemTestPlugin/Source/Private/Tests/ReplicationSystem/ReplicatedTestObjectFactory.cpp b/Engine/Plugins/Runtime/ReplicationSystemTestPlugin/Source/Private/Tests/ReplicationSystem/ReplicatedTestObjectFactory.cpp
--- a/Engine/Plugins/Runtime/ReplicationSystemTestPlugin/Source/Private/Tests/ReplicationSystem/ReplicatedTestObjectFactory.cpp
+++ b/Engine/Plugins/Runtime/ReplicationSystemTestPlugin/Source/Private/Tests/ReplicationSystem/ReplicatedTestObjectFactory.cpp
@@ -148,9 +148,11 @@ UNetObjectFactory::FInstantiateResult UReplicatedTestObjectFactory::InstantiateR
 	FStaticConstructObjectParameters ConstructObjectParameters(ArcheType->GetClass());
 	UObject* CreatedObject = StaticConstructObject_Internal(ConstructObjectParameters);
 
+	const bool bIsSubObject = Context.RootObjectOfSubObject.IsValid();
+
 	if (UReplicatedTestObject* BaseTestObject = Cast<UReplicatedTestObject>(CreatedObject))
 	{
-		BaseTestObject->bIsSubObject = Context.RootObjectOfSubObject.IsValid();
+		BaseTestObject->bIsSubObject = bIsSubObject;
 	}
 	
 	if (UTestReplicatedIrisObject* CreatedTestObject = Cast<UTestReplicatedIrisObject>(CreatedObject))
@@ -176,6 +178,12 @@ UNetObjectFactory::FInstantiateResult UReplicatedTestObjectFactory::InstantiateR
 	FInstantiateResult InstantiateResult;
 	InstantiateResult.Instance = CreatedObject;
 	InstantiateResult.Flags |= EReplicationBridgeCreateNetRefHandleResultFlags::AllowDestroyInstanceFromRemote;
+
+	if (bIsSubObject)
+	{
+		InstantiateResult.Flags |= EReplicationBridgeCreateNetRefHandleResultFlags::ShouldCallSubObjectCreatedFromReplication;
+	}
+
 	return InstantiateResult;
 }
 
@@ -215,4 +223,23 @@ void UReplicatedTestObjectFactory::GetWorldInfo(const FWorldInfoContext& Context
 	{
 		GetInstanceWorldObjectInfoFunction(Context.Handle, Context.Instance, OutData.WorldLocation, OutData.CullDistance);
 	}
+}
+
+void UReplicatedTestObjectFactory::SubObjectCreatedFromReplication(UE::Net::FNetRefHandle RootObject, UE::Net::FNetRefHandle SubObjectCreated)
+{
+	UReplicatedTestObject* RootInstance = Cast<UReplicatedTestObject>(Bridge->GetReplicatedObject(RootObject));
+	if (ensure(RootInstance))
+	{
+		UObject* SubObject = Bridge->GetReplicatedObject(SubObjectCreated);
+		RootInstance->OnSubObjectCreated(SubObject);
+	}
+}
+
+void UReplicatedTestObjectFactory::SubObjectDestroyedFromReplication(const FDestroyedContext& Context)
+{
+	UReplicatedTestObject* RootInstance = Cast<UReplicatedTestObject>(Context.RootObject);
+	if (ensure(RootInstance))
+	{
+		RootInstance->OnSubObjectDestroyed(Context.DestroyedInstance);
+	}
 }
\ No newline at end of file
diff --git a/Engine/Plugins/Runtime/ReplicationSystemTestPlugin/Source/Private/Tests/ReplicationSystem/ReplicatedTestObjectFactory.h b/Engine/Plugins/Runtime/ReplicationSystemTestPlugin/Source/Private/Tests/ReplicationSystem/ReplicatedTestObjectFactory.h
--- a/Engine/Plugins/Runtime/ReplicationSystemTestPlugin/Source/Private/Tests/ReplicationSystem/ReplicatedTestObjectFactory.h
+++ b/Engine/Plugins/Runtime/ReplicationSystemTestPlugin/Source/Private/Tests/ReplicationSystem/ReplicatedTestObjectFactory.h
@@ -43,6 +43,8 @@ public:
 	virtual void PostInit(const FPostInitContext& Context) override;
 	virtual void DestroyReplicatedObject(const FDestroyedContext& Context) override;
 	virtual void GetWorldInfo(const FWorldInfoContext& Context, FWorldInfoData& OutData) override;
+	virtual void SubObjectCreatedFromReplication(UE::Net::FNetRefHandle RootObject, UE::Net::FNetRefHandle SubObjectCreated) override;
+	virtual void SubObjectDestroyedFromReplication(const FDestroyedContext& Context) override;
 
 	//------------------------------------------------------------------------
 
diff --git a/Engine/Plugins/Runtime/ReplicationSystemTestPlugin/Source/Private/Tests/ReplicationSystem/TestReplicationSystem.cpp b/Engine/Plugins/Runtime/ReplicationSystemTestPlugin/Source/Private/Tests/ReplicationSystem/TestReplicationSystem.cpp
--- a/Engine/Plugins/Runtime/ReplicationSystemTestPlugin/Source/Private/Tests/ReplicationSystem/TestReplicationSystem.cpp
+++ b/Engine/Plugins/Runtime/ReplicationSystemTestPlugin/Source/Private/Tests/ReplicationSystem/TestReplicationSystem.cpp
@@ -555,7 +555,7 @@ UE_NET_TEST_FIXTURE(FReplicationSystemServerClientTestFixture, TestSubObjectDefa
 	UE_NET_ASSERT_GT(ClientSubObject2->LastRepOrderCounter, ClientSubObject1->LastRepOrderCounter);
 }
 
-UE_NET_TEST_FIXTURE(FReplicationSystemServerClientTestFixture, TestSubObjectSpecfiedReplicationOrder)
+UE_NET_TEST_FIXTURE(FReplicationSystemServerClientTestFixture, TestSubObjectSpecifiedReplicationOrder)
 {
 	// Add a client
 	FReplicationSystemTestClient* Client = CreateClient();
@@ -592,6 +592,49 @@ UE_NET_TEST_FIXTURE(FReplicationSystemServerClientTestFixture, TestSubObjectSpec
 	UE_NET_ASSERT_EQ(ClientSubObject2->LastRepOrderCounter, 3U);
 }
 
+
+UE_NET_TEST_FIXTURE(FReplicationSystemServerClientTestFixture, TestSubObjectInsertAtStart)
+{
+	// Add a client
+	FReplicationSystemTestClient* Client = CreateClient();
+
+	// Spawn object on server
+	UReplicatedObjectTestSubObjectCreationOrder* ServerObject = Server->CreateObject<UReplicatedObjectTestSubObjectCreationOrder>();
+
+	// Spawn a subobject
+	UReplicatedSubObjectOrderObject* ServerSubObject0 = Server->CreateSubObject<UReplicatedSubObjectOrderObject>(ServerObject->NetRefHandle);
+	
+	// Spawn a subobject and make it replicate before SubObject0
+	UReplicatedSubObjectOrderObject* ServerSubObject1 = Server->CreateSubObject<UReplicatedSubObjectOrderObject>(ServerObject->NetRefHandle, ServerSubObject0->NetRefHandle, UReplicationBridge::ESubObjectInsertionOrder::ReplicateWith);
+	
+	// Spawn a subobject and make it replicate first
+	UReplicatedSubObjectOrderObject* ServerSubObject2 = Server->CreateSubObject<UReplicatedSubObjectOrderObject>(ServerObject->NetRefHandle, FNetRefHandle::GetInvalid(), UReplicationBridge::ESubObjectInsertionOrder::InsertAtStart);
+
+	UReplicatedSubObjectOrderObject::RepOrderCounter = 0U;
+
+	// Send and deliver packet
+	Server->UpdateAndSend({ Client });
+
+	// Verify that objects have replicated
+	UReplicatedSubObjectOrderObject* ClientSubObject0 = Cast<UReplicatedSubObjectOrderObject>(Client->GetReplicationBridge()->GetReplicatedObject(ServerSubObject0->NetRefHandle));
+	UReplicatedSubObjectOrderObject* ClientSubObject1 = Cast<UReplicatedSubObjectOrderObject>(Client->GetReplicationBridge()->GetReplicatedObject(ServerSubObject1->NetRefHandle));
+	UReplicatedSubObjectOrderObject* ClientSubObject2 = Cast<UReplicatedSubObjectOrderObject>(Client->GetReplicationBridge()->GetReplicatedObject(ServerSubObject2->NetRefHandle));
+
+	UE_NET_ASSERT_NE(ClientSubObject0, nullptr);
+	UE_NET_ASSERT_NE(ClientSubObject1, nullptr);
+	UE_NET_ASSERT_NE(ClientSubObject2, nullptr);
+
+	// Verify that they were created in the expected order
+	UE_NET_ASSERT_EQ(ClientSubObject2->CreationOrder, 1);
+	UE_NET_ASSERT_EQ(ClientSubObject1->CreationOrder, 2);
+	UE_NET_ASSERT_EQ(ClientSubObject0->CreationOrder, 3);
+
+	// Verify that they have replicated in the expected order
+	UE_NET_ASSERT_EQ(ClientSubObject2->LastRepOrderCounter, 1U);
+	UE_NET_ASSERT_EQ(ClientSubObject1->LastRepOrderCounter, 2U);
+	UE_NET_ASSERT_EQ(ClientSubObject0->LastRepOrderCounter, 3U);
+}
+
 class FTestNetTokensFixture : public FReplicationSystemServerClientTestFixture
 {
 public:
diff --git a/Engine/Source/Runtime/Experimental/Iris/Core/Private/Iris/ReplicationSystem/NetRefHandleManager.cpp b/Engine/Source/Runtime/Experimental/Iris/Core/Private/Iris/ReplicationSystem/NetRefHandleManager.cpp
--- a/Engine/Source/Runtime/Experimental/Iris/Core/Private/Iris/ReplicationSystem/NetRefHandleManager.cpp
+++ b/Engine/Source/Runtime/Experimental/Iris/Core/Private/Iris/ReplicationSystem/NetRefHandleManager.cpp
@@ -415,7 +415,7 @@ void FNetRefHandleManager::InternalDestroyNetObject(FInternalNetRefIndex Interna
 		}
 	}
 	
-	// If we have subobjects remove them from our list
+	// If this is a RootObject, remove all subobjects from the list
 	if (FNetDependencyData::FInternalNetRefIndexArray* SubObjectArray = SubObjects.GetInternalIndexArray<FNetDependencyData::EArrayType::SubObjects>(InternalIndex))
 	{
 		for (FInternalNetRefIndex SubObjectInternalIndex : *SubObjectArray)
@@ -574,29 +574,29 @@ void FNetRefHandleManager::DestroyObjectsPendingDestroy()
 	}
 }
 
-bool FNetRefHandleManager::AddSubObject(FNetRefHandle OwnerHandle, FNetRefHandle SubObjectHandle, FNetRefHandle RelativeOtherSubObjectHandle, EAddSubObjectFlags Flags)
+bool FNetRefHandleManager::AddSubObject(FNetRefHandle RootObjectHandle, FNetRefHandle SubObjectHandle, FNetRefHandle RelativeOtherSubObjectHandle, EAddSubObjectFlags Flags)
 {
-	check(OwnerHandle != SubObjectHandle);
+	check(RootObjectHandle != SubObjectHandle);
 
 	// validate objects
-	const FInternalNetRefIndex OwnerInternalIndex = GetInternalIndex(OwnerHandle);
+	const FInternalNetRefIndex RootObjectInternalIndex = GetInternalIndex(RootObjectHandle);
 	const FInternalNetRefIndex SubObjectInternalIndex = GetInternalIndex(SubObjectHandle);
 
-	const bool bIsValidOwner = ensure(OwnerInternalIndex != InvalidInternalIndex);
-	const bool bIsValidSubObejct = ensure(SubObjectInternalIndex != InvalidInternalIndex);
+	const bool bIsValidOwner = ensure(RootObjectInternalIndex != InvalidInternalIndex);
+	const bool bIsValidSubObject = ensure(SubObjectInternalIndex != InvalidInternalIndex);
 
-	if (!(bIsValidOwner && bIsValidSubObejct))
+	if (!bIsValidOwner || !bIsValidSubObject)
 	{
 		return false;
 	}
 
-	UE_NET_TRACE_SUBOBJECT(OwnerHandle, SubObjectHandle);
+	UE_NET_TRACE_SUBOBJECT(RootObjectHandle, SubObjectHandle);
 
 	const FInternalNetRefIndex RelativeOtherSubObjectInternalIndex = EnumHasAnyFlags(Flags, EAddSubObjectFlags::ReplicateWithSubObject) ? GetInternalIndex(RelativeOtherSubObjectHandle) : InvalidInternalIndex;
-	return InternalAddSubObject(OwnerInternalIndex, SubObjectInternalIndex, RelativeOtherSubObjectInternalIndex, Flags);
+	return InternalAddSubObject(RootObjectInternalIndex, SubObjectInternalIndex, RelativeOtherSubObjectInternalIndex, Flags);
 }
 
-bool FNetRefHandleManager::InternalAddSubObject(FInternalNetRefIndex OwnerInternalIndex, FInternalNetRefIndex SubObjectInternalIndex, FInternalNetRefIndex RelativeOtherSubObjectInternalIndex, EAddSubObjectFlags Flags)
+bool FNetRefHandleManager::InternalAddSubObject(FInternalNetRefIndex RootObjectInternalIndex, FInternalNetRefIndex SubObjectInternalIndex, FInternalNetRefIndex RelativeOtherSubObjectInternalIndex, EAddSubObjectFlags Flags)
 {
 	using namespace UE::Net::Private;
 
@@ -606,84 +606,110 @@ bool FNetRefHandleManager::InternalAddSubObject(FInternalNetRefIndex OwnerIntern
 		return false;
 	}
 
-	FNetDependencyData::FInternalNetRefIndexArray& SubObjectArray = SubObjects.GetOrCreateInternalIndexArray<FNetDependencyData::SubObjects>(OwnerInternalIndex);
-	SubObjectArray.Add(SubObjectInternalIndex);
-	SubObjectData.SubObjectRootIndex = OwnerInternalIndex;
+	// Add the subobject to the RootObject
+	FNetDependencyData::FInternalNetRefIndexArray& SubObjectArray = SubObjects.GetOrCreateInternalIndexArray<FNetDependencyData::SubObjects>(RootObjectInternalIndex);
+
+	if (EnumHasAnyFlags(Flags, EAddSubObjectFlags::InsertAtStart))
+	{
+		// Add at the start
+		SubObjectArray.Insert(SubObjectInternalIndex, 0);
+	}
+	else
+	{
+		// At at the end
+		SubObjectArray.Add(SubObjectInternalIndex);
+	}
+	
+	SubObjectData.SubObjectRootIndex = RootObjectInternalIndex;
 	SubObjectData.bDestroySubObjectWithOwner = EnumHasAnyFlags(Flags, EAddSubObjectFlags::DestroyWithOwner);
+
 	// Mark the object as a subobject
 	SetIsSubObject(SubObjectInternalIndex, true);
 
-	if (RelativeOtherSubObjectInternalIndex != InvalidInternalIndex && ensureMsgf(SubObjectArray.Find(RelativeOtherSubObjectInternalIndex) != INDEX_NONE, TEXT("RelativeOtherSubObjectHandle %s Must be a Subobject of %s"), ToCStr(GetNetRefHandleFromInternalIndex(RelativeOtherSubObjectInternalIndex).ToString()), ToCStr(GetNetRefHandleFromInternalIndex(OwnerInternalIndex).ToString())))
+	FInternalNetRefIndex ParentOfSubObjectIndex = RootObjectInternalIndex;
+	
+	// Make sure the non-rootobject Parent is a subobject of the same RootObject
+	if (RelativeOtherSubObjectInternalIndex != InvalidInternalIndex)
 	{
-		// Add to child array of RelativeOtherSubObjectRelativeIndex for hierarchical replication order
-		FNetDependencyData::FSubObjectConditionalsArray* SubObjectConditionalsArray = nullptr;
-		FNetDependencyData::FInternalNetRefIndexArray& ChildSubObjectArray = SubObjects.GetOrCreateInternalChildSubObjectsArray(RelativeOtherSubObjectInternalIndex, SubObjectConditionalsArray);
-		ChildSubObjectArray.Add(SubObjectInternalIndex);
-		if (SubObjectConditionalsArray)
+		const bool bIsValidOuter = SubObjectArray.Contains(RelativeOtherSubObjectInternalIndex);
+		if (ensureMsgf(bIsValidOuter, TEXT("RelativeOtherSubObjectHandle %s must be a Subobject of %s"), *PrintObjectFromIndex(RelativeOtherSubObjectInternalIndex), *PrintObjectFromIndex(RootObjectInternalIndex)))
 		{
-			SubObjectConditionalsArray->Add(ELifetimeCondition::COND_None);
+			ParentOfSubObjectIndex = RelativeOtherSubObjectInternalIndex;
 		}
-		SubObjectData.SubObjectParentIndex = RelativeOtherSubObjectInternalIndex;
+	}
+
+	// Add the subobject to his Parent's list
+	FNetDependencyData::FSubObjectConditionalsArray* SubObjectConditionalsArray = nullptr;
+	FNetDependencyData::FInternalNetRefIndexArray& ChildSubObjectArray = SubObjects.GetOrCreateInternalChildSubObjectsArray(ParentOfSubObjectIndex, SubObjectConditionalsArray);
+
+	if (EnumHasAnyFlags(Flags, EAddSubObjectFlags::InsertAtStart))
+	{
+		// Add at the start
+		ChildSubObjectArray.Insert(SubObjectInternalIndex, 0);
 	}
 	else
 	{
-		// Add to child array of root
-		FNetDependencyData::FSubObjectConditionalsArray* SubObjectConditionalsArray = nullptr;
-		FNetDependencyData::FInternalNetRefIndexArray& ChildSubObjectArray = SubObjects.GetOrCreateInternalChildSubObjectsArray(OwnerInternalIndex, SubObjectConditionalsArray);
+		// At at the end
 		ChildSubObjectArray.Add(SubObjectInternalIndex);
-		if (SubObjectConditionalsArray)
-		{
-			SubObjectConditionalsArray->Add(ELifetimeCondition::COND_None);
-		}
-		SubObjectData.SubObjectParentIndex = OwnerInternalIndex;
 	}
+	
+
+	if (SubObjectConditionalsArray)
+	{
+		SubObjectConditionalsArray->Add(ELifetimeCondition::COND_None);
+	}
+
+	SubObjectData.SubObjectParentIndex = ParentOfSubObjectIndex;
 
 	return true;
 }
 
-void FNetRefHandleManager::InternalRemoveSubObject(FInternalNetRefIndex OwnerInternalIndex, FInternalNetRefIndex SubObjectInternalIndex, bool bRemoveFromSubObjectArray)
+void FNetRefHandleManager::InternalRemoveSubObject(FInternalNetRefIndex RootObjectInternalIndex, FInternalNetRefIndex SubObjectInternalIndex, bool bRemoveFromSubObjectArray)
 {
 	// both must be valid
-	if (OwnerInternalIndex != InvalidInternalIndex && SubObjectInternalIndex != InvalidInternalIndex)
+	if (RootObjectInternalIndex == InvalidInternalIndex || SubObjectInternalIndex == InvalidInternalIndex)
 	{
-		FReplicatedObjectData& SubObjectData = GetReplicatedObjectDataNoCheck(SubObjectInternalIndex);
-		check(SubObjectData.SubObjectRootIndex == OwnerInternalIndex);
+		return;
+	}
+
+	FReplicatedObjectData& SubObjectData = GetReplicatedObjectDataNoCheck(SubObjectInternalIndex);
+	check(SubObjectData.SubObjectRootIndex == RootObjectInternalIndex);
 		
-		if (bRemoveFromSubObjectArray)
+	if (bRemoveFromSubObjectArray)
+	{
+		// Remove subobject from the root object's list
+		if (FNetDependencyData::FInternalNetRefIndexArray* SubObjectArray = SubObjects.GetInternalIndexArray<FNetDependencyData::EArrayType::SubObjects>(RootObjectInternalIndex))
 		{
-			// Remove from root parent
-			if (FNetDependencyData::FInternalNetRefIndexArray* SubObjectArray = SubObjects.GetInternalIndexArray<FNetDependencyData::EArrayType::SubObjects>(OwnerInternalIndex))
-			{
-				SubObjectArray->Remove(SubObjectInternalIndex);
-			}
-			// Remove from parents ChildSubObjectArray
-			if (SubObjectData.SubObjectParentIndex != InvalidInternalIndex)
-			{
-				FNetDependencyData::FInternalNetRefIndexArray* ChildSubObjectArray;
-				FNetDependencyData::FSubObjectConditionalsArray* SubObjectConditionsArray;
+			SubObjectArray->Remove(SubObjectInternalIndex);
+		}
+
+		// Remove subobject from the parent's list
+		if (SubObjectData.SubObjectParentIndex != InvalidInternalIndex)
+		{
+			FNetDependencyData::FInternalNetRefIndexArray* ChildSubObjectArray;
+			FNetDependencyData::FSubObjectConditionalsArray* SubObjectConditionsArray;
 
-				if (SubObjects.GetInternalChildSubObjectAndConditionalArrays(SubObjectData.SubObjectParentIndex, ChildSubObjectArray, SubObjectConditionsArray))
+			if (SubObjects.GetInternalChildSubObjectAndConditionalArrays(SubObjectData.SubObjectParentIndex, ChildSubObjectArray, SubObjectConditionsArray))
+			{
+				const int32 ArrayIndex = ChildSubObjectArray->Find(SubObjectInternalIndex);
+				if (ensureMsgf(ArrayIndex != INDEX_NONE, TEXT("Subobject: %s not found in the list of his parent"), *PrintObjectFromIndex(SubObjectInternalIndex)))
 				{
-					const int32 ArrayIndex = ChildSubObjectArray->Find(SubObjectInternalIndex);
-					if (ensure(ArrayIndex != INDEX_NONE))
+					ChildSubObjectArray->RemoveAt(ArrayIndex);
+					if (SubObjectConditionsArray)
 					{
-						ChildSubObjectArray->RemoveAt(ArrayIndex);
-						if (SubObjectConditionsArray)
-						{
-							SubObjectConditionsArray->RemoveAt(ArrayIndex);
-							check(SubObjectConditionsArray->Num() == ChildSubObjectArray->Num());
-						}
+						SubObjectConditionsArray->RemoveAt(ArrayIndex);
+						check(SubObjectConditionsArray->Num() == ChildSubObjectArray->Num());
 					}
 				}
 			}
 		}
+	}
 
-		SubObjectData.SubObjectRootIndex = InvalidInternalIndex;
-		SubObjectData.SubObjectParentIndex = InvalidInternalIndex;
-		SubObjectData.bDestroySubObjectWithOwner = false;
+	SubObjectData.SubObjectRootIndex = InvalidInternalIndex;
+	SubObjectData.SubObjectParentIndex = InvalidInternalIndex;
+	SubObjectData.bDestroySubObjectWithOwner = false;
 
-		SetIsSubObject(SubObjectInternalIndex, false);
-	}
+	SetIsSubObject(SubObjectInternalIndex, false);
 }
 
 void FNetRefHandleManager::RemoveSubObject(FNetRefHandle Handle)
@@ -693,10 +719,10 @@ void FNetRefHandleManager::RemoveSubObject(FNetRefHandle Handle)
 
 	if (SubObjectInternalIndex != InvalidInternalIndex)
 	{
-		const uint32 OwnerInternalIndex = ReplicatedObjectData[SubObjectInternalIndex].SubObjectRootIndex;
-		if (OwnerInternalIndex != InvalidInternalIndex)
+		const FInternalNetRefIndex RootObjectInternalIndex = ReplicatedObjectData[SubObjectInternalIndex].SubObjectRootIndex;
+		if (RootObjectInternalIndex != InvalidInternalIndex)
 		{
-			InternalRemoveSubObject(OwnerInternalIndex, SubObjectInternalIndex);
+			InternalRemoveSubObject(RootObjectInternalIndex, SubObjectInternalIndex);
 		}
 	}
 }
diff --git a/Engine/Source/Runtime/Experimental/Iris/Core/Private/Iris/ReplicationSystem/NetRefHandleManager.h b/Engine/Source/Runtime/Experimental/Iris/Core/Private/Iris/ReplicationSystem/NetRefHandleManager.h
--- a/Engine/Source/Runtime/Experimental/Iris/Core/Private/Iris/ReplicationSystem/NetRefHandleManager.h
+++ b/Engine/Source/Runtime/Experimental/Iris/Core/Private/Iris/ReplicationSystem/NetRefHandleManager.h
@@ -41,6 +41,9 @@ enum class EAddSubObjectFlags : uint32
 	SkipIfAlreadySubObject = WarnIfAlreadySubObject << 1U,
 	DestroyWithOwner = SkipIfAlreadySubObject << 1U,
 	ReplicateWithSubObject = DestroyWithOwner << 1U,
+	// Insert the subject at the start of the list
+	InsertAtStart = ReplicateWithSubObject << 1U, 
+
 	Default = WarnIfAlreadySubObject | DestroyWithOwner,
 };
 ENUM_CLASS_FLAGS(EAddSubObjectFlags);
@@ -233,8 +236,8 @@ public:
 	const FNetBitArray& GetSubObjectInternalIndices() const { return SubObjectInternalIndices; }
 	const FNetBitArrayView GetSubObjectInternalIndicesView() const { return MakeNetBitArrayView(SubObjectInternalIndices); }
 
-	bool AddSubObject(FNetRefHandle OwnerHandle, FNetRefHandle SubObjectHandle, FNetRefHandle RelativeOtherSubObjectHandle, EAddSubObjectFlags Flags = EAddSubObjectFlags::Default);
-	bool AddSubObject(FNetRefHandle OwnerHandle, FNetRefHandle SubObjectHandle, EAddSubObjectFlags Flags = EAddSubObjectFlags::Default);
+	bool AddSubObject(FNetRefHandle RootObjectHandle, FNetRefHandle SubObjectHandle, FNetRefHandle RelativeOtherSubObjectHandle, EAddSubObjectFlags Flags = EAddSubObjectFlags::Default);
+	bool AddSubObject(FNetRefHandle RootObjectHandle, FNetRefHandle SubObjectHandle, EAddSubObjectFlags Flags = EAddSubObjectFlags::Default);
 	void RemoveSubObject(FNetRefHandle SubObjectHandle);
 	
 	FNetRefHandle GetRootObjectOfSubObject(FNetRefHandle SubObjectHandle) const;
@@ -372,10 +375,10 @@ private:
 	// Get the next free internal index, returns InvalidInternalIndex if a free one cannot be found
 	FInternalNetRefIndex GetNextFreeInternalIndex() const;
 
-	bool InternalAddSubObject(FInternalNetRefIndex OwnerInternalIndex, FInternalNetRefIndex SubObjectInternalIndex, FInternalNetRefIndex RelativeOtherSubObjectInternalIndex, EAddSubObjectFlags Flags);
+	bool InternalAddSubObject(FInternalNetRefIndex RootObjectInternalIndex, FInternalNetRefIndex SubObjectInternalIndex, FInternalNetRefIndex RelativeOtherSubObjectInternalIndex, EAddSubObjectFlags Flags);
 
 	inline void SetIsSubObject(FInternalNetRefIndex InternalIndex, bool IsSubObject) { SubObjectInternalIndices.SetBitValue(InternalIndex, IsSubObject); }
-	void InternalRemoveSubObject(FInternalNetRefIndex OwnerInternalIndex, FInternalNetRefIndex SubObjectInternalIndex, bool bRemoveFromSubObjectArray = true);
+	void InternalRemoveSubObject(FInternalNetRefIndex RootObjectInternalIndex, FInternalNetRefIndex SubObjectInternalIndex, bool bRemoveFromSubObjectArray = true);
 
 	void InternalRemoveDependentObject(FInternalNetRefIndex ParentInternalIndex, FInternalNetRefIndex DependentInternalIndex, ERemoveDependentObjectFlags Flags = ERemoveDependentObjectFlags::All);
 	void InternalRemoveDependentObject(FInternalNetRefIndex DependentInternalIndex);
@@ -575,9 +578,9 @@ uint32 FNetRefHandleManager::GetOriginalDestroyedStartupObjectIndex(FInternalNet
 	return FoundOriginalInternalIndex ? *FoundOriginalInternalIndex : 0U;
 }
 
-inline bool FNetRefHandleManager::AddSubObject(FNetRefHandle OwnerHandle, FNetRefHandle SubObjectHandle, EAddSubObjectFlags Flags)
+inline bool FNetRefHandleManager::AddSubObject(FNetRefHandle RootObjectHandle, FNetRefHandle SubObjectHandle, EAddSubObjectFlags Flags)
 {
-	return AddSubObject(OwnerHandle, SubObjectHandle, FNetRefHandle::GetInvalid(), Flags);
+	return AddSubObject(RootObjectHandle, SubObjectHandle, FNetRefHandle::GetInvalid(), Flags);
 }
 
 TArrayView<const FInternalNetRefIndex> FNetRefHandleManager::GetSubObjects(FInternalNetRefIndex OwnerIndex) const
diff --git a/Engine/Source/Runtime/Experimental/Iris/Core/Private/Iris/ReplicationSystem/ReplicationBridge.cpp b/Engine/Source/Runtime/Experimental/Iris/Core/Private/Iris/ReplicationSystem/ReplicationBridge.cpp
--- a/Engine/Source/Runtime/Experimental/Iris/Core/Private/Iris/ReplicationSystem/ReplicationBridge.cpp
+++ b/Engine/Source/Runtime/Experimental/Iris/Core/Private/Iris/ReplicationSystem/ReplicationBridge.cpp
@@ -570,7 +570,21 @@ void UReplicationBridge::InternalAddSubObject(FNetRefHandle OwnerHandle, FNetRef
 	using namespace UE::Net::Private;
 
 	EAddSubObjectFlags AddSubObjectFlags = EAddSubObjectFlags::Default;
-	AddSubObjectFlags |= InsertionOrder == ESubObjectInsertionOrder::None ? EAddSubObjectFlags::None : EAddSubObjectFlags::ReplicateWithSubObject;
+
+	switch(InsertionOrder)
+	{
+		case ESubObjectInsertionOrder::None: 
+			break;
+		case ESubObjectInsertionOrder::ReplicateWith:
+			AddSubObjectFlags |= EAddSubObjectFlags::ReplicateWithSubObject;
+			break;
+		case ESubObjectInsertionOrder::InsertAtStart:
+			AddSubObjectFlags |= EAddSubObjectFlags::InsertAtStart;
+			break;
+		default:
+			checkf(false, TEXT("Missing implementation of ESubObjectInsertionOrder enum"));
+			break;
+	}
 
 	if (NetRefHandleManager->AddSubObject(OwnerHandle, SubObjectHandle, InsertRelativeToSubObjectHandle, AddSubObjectFlags))
 	{
diff --git a/Engine/Source/Runtime/Experimental/Iris/Core/Public/Iris/ReplicationSystem/ReplicationBridge.h b/Engine/Source/Runtime/Experimental/Iris/Core/Public/Iris/ReplicationSystem/ReplicationBridge.h
--- a/Engine/Source/Runtime/Experimental/Iris/Core/Public/Iris/ReplicationSystem/ReplicationBridge.h
+++ b/Engine/Source/Runtime/Experimental/Iris/Core/Public/Iris/ReplicationSystem/ReplicationBridge.h
@@ -103,7 +103,10 @@ public:
 	enum class ESubObjectInsertionOrder : uint8
 	{
 		None,
+		/** Insert the subobject so it will replicate before the other subobject. */
 		ReplicateWith,
+		/** Insert the subobject at the start of the list so it can be created and replicated first */
+		InsertAtStart,
 	};
 
 	/**
