[IoStore]
* Moved dll storage from types to methods

diff --git a/Engine/Source/Runtime/Experimental/IoStore/HttpClient/Internal/IO/Http/Client.h b/Engine/Source/Runtime/Experimental/IoStore/HttpClient/Internal/IO/Http/Client.h
--- a/Engine/Source/Runtime/Experimental/IoStore/HttpClient/Internal/IO/Http/Client.h
+++ b/Engine/Source/Runtime/Experimental/IoStore/HttpClient/Internal/IO/Http/Client.h
@@ -50,20 +50,20 @@ using	FTicket			= uint64;
 struct	FActivity;
 
 ////////////////////////////////////////////////////////////////////////////////
-class UE_API FCertRoots
+class FCertRoots
 {
 public:
-							FCertRoots()					= default;
-							~FCertRoots();
-							FCertRoots(FMemoryView PemData);
-							FCertRoots(FCertRoots&& Rhs)	{ *this = MoveTemp(Rhs); }
-	FCertRoots&				operator = (FCertRoots&& Rhs)	{ Swap(Handle, Rhs.Handle); return *this; }
-	bool					IsValid() const					{ return Handle != 0; }
-	int32					Num() const;
-	static void				SetDefault(FCertRoots&& CertRoots);
-	static FCertRootsRef	NoTls();
-	static FCertRootsRef	Default();
-	static FCertRootsRef	Explicit(const FCertRoots& CertRoots);
+								FCertRoots()					= default;
+	UE_API 						~FCertRoots();
+	UE_API 						FCertRoots(FMemoryView PemData);
+								FCertRoots(FCertRoots&& Rhs)	{ *this = MoveTemp(Rhs); }
+	FCertRoots&					operator = (FCertRoots&& Rhs)	{ Swap(Handle, Rhs.Handle); return *this; }
+	bool						IsValid() const					{ return Handle != 0; }
+	UE_API int32				Num() const;
+	static UE_API void			SetDefault(FCertRoots&& CertRoots);
+	static UE_API FCertRootsRef	NoTls();
+	static UE_API FCertRootsRef	Default();
+	static UE_API FCertRootsRef	Explicit(const FCertRoots& CertRoots);
 
 private:
 	UPTRINT					Handle = 0;
@@ -74,12 +74,12 @@ private:
 };
 
 ////////////////////////////////////////////////////////////////////////////////
-class UE_API FConnectionPool
+class FConnectionPool
 {
 public:
-	struct UE_API FParams
+	struct FParams
 	{
-		int32				SetHostFromUrl(FAnsiStringView Url);
+		UE_API int32		SetHostFromUrl(FAnsiStringView Url);
 		FAnsiStringView		HostName;
 		FCertRootsRef		VerifyCert = FCertRoots::NoTls();
 		int32				SendBufSize = -1;
@@ -93,14 +93,14 @@ public:
 	};
 
 							FConnectionPool() = default;
-							FConnectionPool(const FParams& Params);
-							~FConnectionPool();
+	UE_API 					FConnectionPool(const FParams& Params);
+	UE_API 					~FConnectionPool();
 							FConnectionPool(FConnectionPool&& Rhs)	{ *this = MoveTemp(Rhs); }
 	FConnectionPool&		operator = (FConnectionPool&& Rhs)		{ Swap(Ptr, Rhs.Ptr); return *this; }
-	bool					Resolve();
-	void					Describe(FAnsiStringBuilderBase&) const;
+	UE_API bool				Resolve();
+	UE_API void				Describe(FAnsiStringBuilderBase&) const;
 
-	static bool				IsValidHostUrl(FAnsiStringView Url);
+	static UE_API bool		IsValidHostUrl(FAnsiStringView Url);
 
 private:
 	friend					class FEventLoop;
@@ -112,17 +112,17 @@ private:
 };
 
 ////////////////////////////////////////////////////////////////////////////////
-class UE_API FRequest
+class FRequest
 {
 public:
-						~FRequest();
-						FRequest(FRequest&& Rhs);
+	UE_API 				~FRequest();
+	UE_API 				FRequest(FRequest&& Rhs);
 	bool				IsValid() const { return Ptr != nullptr; }
-	FRequest&&			Accept(EMimeType MimeType);
-	FRequest&&			Accept(FAnsiStringView MimeType);
-	FRequest&&			Header(FAnsiStringView Key, FAnsiStringView Value);
-	void				Content(const void* Data, SIZE_T Size, EMimeType MimeType);
-	void				Content(const void* Data, SIZE_T Size, FAnsiStringView MimeType);
+	UE_API FRequest&&	Accept(EMimeType MimeType);
+	UE_API FRequest&&	Accept(FAnsiStringView MimeType);
+	UE_API FRequest&&	Header(FAnsiStringView Key, FAnsiStringView Value);
+	UE_API void			Content(const void* Data, SIZE_T Size, EMimeType MimeType);
+	UE_API void			Content(const void* Data, SIZE_T Size, FAnsiStringView MimeType);
 
 private:
 	friend				class FEventLoop;
@@ -136,31 +136,31 @@ private:
 };
 
 ////////////////////////////////////////////////////////////////////////////////
-class UE_API FResponse
+class FResponse
 {
 public:
 	using FHeaderSink = TFunction<bool (FAnsiStringView, FAnsiStringView)>;
 
-	EStatusCodeClass	GetStatus() const;
-	uint32				GetStatusCode() const;
-	FAnsiStringView 	GetStatusMessage() const;
-	int64				GetContentLength() const;
-	EMimeType			GetContentType() const;
-	void				GetContentType(FAnsiStringView& Out) const;
-	FAnsiStringView 	GetHeader(FAnsiStringView Name) const;
-	void				ReadHeaders(FHeaderSink Sink) const;
-	void				SetDestination(FIoBuffer* Buffer);
+	UE_API EStatusCodeClass	GetStatus() const;
+	UE_API uint32			GetStatusCode() const;
+	UE_API FAnsiStringView 	GetStatusMessage() const;
+	UE_API int64			GetContentLength() const;
+	UE_API EMimeType		GetContentType() const;
+	UE_API void				GetContentType(FAnsiStringView& Out) const;
+	UE_API FAnsiStringView 	GetHeader(FAnsiStringView Name) const;
+	UE_API void				ReadHeaders(FHeaderSink Sink) const;
+	UE_API void				SetDestination(FIoBuffer* Buffer);
 
 private:
-						FResponse() = delete;
-						FResponse(const FResponse&) = delete;
-						FResponse(FResponse&&) = delete;
-	FResponse&			operator = (const FResponse&) = delete;
-	FResponse&			operator = (FResponse&&) = delete;
+							FResponse() = delete;
+							FResponse(const FResponse&) = delete;
+							FResponse(FResponse&&) = delete;
+	FResponse&				operator = (const FResponse&) = delete;
+	FResponse&				operator = (FResponse&&) = delete;
 };
 
 ////////////////////////////////////////////////////////////////////////////////
-class UE_API FTicketPerf
+class FTicketPerf
 {
 public:
 	struct FSample
@@ -173,7 +173,7 @@ public:
 		uint32			GetSizeEstKiB() const	{ return (RecvKiBps * RecvMs) / 1000; }
 	};
 
-	FSample				GetSample() const;
+	UE_API FSample		GetSample() const;
 };
 
 #if !IAS_HTTP_WITH_PERF
@@ -181,67 +181,67 @@ inline FTicketPerf::FSample FTicketPerf::GetSample() const { return {}; }
 #endif
 
 ////////////////////////////////////////////////////////////////////////////////
-class UE_API FTicketStatus
+class FTicketStatus
 {
 public:
 	enum class EId : uint8 { Response, Content, Cancelled, Error };
 
 	struct FError
 	{
-		const char*		Reason;
-		uint32			Code;
+		const char*				Reason;
+		uint32					Code;
 	};
 
-	EId					GetId() const;
-	UPTRINT				GetParam() const;
-	FTicket				GetTicket() const;
-	uint32				GetIndex() const;
-	FResponse&			GetResponse() const;		// if GetId() == EId::Response
-	const FIoBuffer&	GetContent() const;			// if GetId() == EId::Content
-	uint32				GetContentLength() const;	//  |
-	const FTicketPerf&	GetPerf() const;			// _|_
-	FError				GetError() const;			// if GetId() == EId::Error
+	UE_API EId					GetId() const;
+	UE_API UPTRINT				GetParam() const;
+	UE_API FTicket				GetTicket() const;
+	UE_API uint32				GetIndex() const;
+	UE_API FResponse&			GetResponse() const;		// if GetId() == EId::Response
+	UE_API const FIoBuffer&		GetContent() const;			// if GetId() == EId::Content
+	UE_API uint32				GetContentLength() const;	//  |
+	UE_API const FTicketPerf&	GetPerf() const;			// _|_
+	UE_API FError				GetError() const;			// if GetId() == EId::Error
 
 private:
-						FTicketStatus() = delete;
-						FTicketStatus(const FTicketStatus&) = delete;
-						FTicketStatus(FTicketStatus&&) = delete;
-	FTicketStatus&		operator = (const FTicketStatus&) = delete;
-	FTicketStatus&		operator = (FTicketStatus&&) = delete;
+								FTicketStatus() = delete;
+								FTicketStatus(const FTicketStatus&) = delete;
+								FTicketStatus(FTicketStatus&&) = delete;
+	FTicketStatus&				operator = (const FTicketStatus&) = delete;
+	FTicketStatus&				operator = (FTicketStatus&&) = delete;
 };
 
 using FTicketSink = TFunction<void (const FTicketStatus&)>;
 
 ////////////////////////////////////////////////////////////////////////////////
-class UE_API FEventLoop
+class FEventLoop
 {
 	class FImpl;
 
 public:
-	static const uint32		MaxActiveTickets = 64;
+	static UE_API const uint32		MaxActiveTickets = 64;
 
 	struct FRequestParams
 	{
-		FCertRootsRef		VerifyCert		= {};
-		uint32				ContentSizeEst	= 0;
-		uint16				BufferSize		= 256;
-		bool				bAutoRedirect	= false;
-		bool				bAllowChunked	= true;
+		FCertRootsRef				VerifyCert		= {};
+		uint32						ContentSizeEst	= 0;
+		uint16						BufferSize		= 256;
+		bool						bAutoRedirect	= false;
+		bool						bAllowChunked	= true;
 	};
 
 	template <typename... T> [[nodiscard]] FRequest Get(T&&... t)  { return Request("GET",  Forward<T&&>(t)...); }
 	template <typename... T> [[nodiscard]] FRequest Post(T&&... t) { return Request("POST", Forward<T&&>(t)...); }
 
-							FEventLoop();
-							~FEventLoop();
-	uint32					Tick(int32 PollTimeoutMs=0);
-	void					Throttle(uint32 KiBPerSec);
-	void					SetFailTimeout(int32 TimeoutMs);
-	bool					IsIdle() const;
-	void					Cancel(FTicket Ticket);
-	[[nodiscard]] FRequest	Request(FAnsiStringView Method, FAnsiStringView Url, const FRequestParams* Params=nullptr);
-	[[nodiscard]] FRequest	Request(FAnsiStringView Method, FAnsiStringView Path, FConnectionPool& Pool, const FRequestParams* Params=nullptr);
-	FTicket					Send(FRequest&& Request, FTicketSink Sink, UPTRINT SinkParam=0);
+	UE_API 							FEventLoop();
+	UE_API 							~FEventLoop();
+	UE_API uint32					Tick(int32 PollTimeoutMs=0);
+	UE_API void						Throttle(uint32 KiBPerSec);
+	UE_API void						SetFailTimeout(int32 TimeoutMs);
+	UE_API bool						IsIdle() const;
+	UE_API void						Cancel(FTicket Ticket);
+	[[nodiscard]] UE_API FRequest	Request(FAnsiStringView Method, FAnsiStringView Url, const FRequestParams* Params=nullptr);
+	[[nodiscard]] UE_API FRequest	Request(FAnsiStringView Method, FAnsiStringView Path, FConnectionPool& Pool, const FRequestParams* Params=nullptr);
+	UE_API FTicket					Send(FRequest&& Request, FTicketSink Sink, UPTRINT SinkParam=0);
 
 private:
 	bool					Redirect(const FTicketStatus& Status, FTicketSink& OuterSink);
