ChaosFlesh:
-Added muscle length-activation curve to SetMuscleActivationParameter node.
-Added Chaos::FLinearCurve, managed array support and conversion from FRichCurve.

diff --git a/Engine/Plugins/Experimental/ChaosFlesh/Source/ChaosFleshNodes/Private/Dataflow/ChaosFleshComputeMuscleActivationNode.cpp b/Engine/Plugins/Experimental/ChaosFlesh/Source/ChaosFleshNodes/Private/Dataflow/ChaosFleshComputeMuscleActivationNode.cpp
--- a/Engine/Plugins/Experimental/ChaosFlesh/Source/ChaosFleshNodes/Private/Dataflow/ChaosFleshComputeMuscleActivationNode.cpp
+++ b/Engine/Plugins/Experimental/ChaosFlesh/Source/ChaosFleshNodes/Private/Dataflow/ChaosFleshComputeMuscleActivationNode.cpp
@@ -9,6 +9,7 @@
 #include "BonePose.h"
 #include "ChaosFlesh/ChaosFlesh.h"
 #include "ChaosFlesh/TetrahedralCollection.h"
+#include "Chaos/Curve.h"
 #include "Dataflow/DataflowSelection.h"
 #include "Engine/SkeletalMesh.h"
 #include "GeometryCollection/Facades/CollectionMeshFacade.h"
@@ -335,6 +336,8 @@ FSetMuscleActivationParameterNode::FSetMuscleActivationParameterNode(const UE::D
 		.SetCanHidePin(true)
 		.SetPinIsHidden(true);
 	RegisterOutputConnection(&Collection, &Collection);
+	// Set default linear curve
+	FLengthActivationUtils::SetDefaultLengthActivationCurve(GlobalLengthActivationCurve);
 }
 
 
@@ -352,6 +355,12 @@ void FSetMuscleActivationParameterNode::Evaluate(UE::Dataflow::FContext& Context
 					GlobalFiberLengthRatioAtMaxActivation,
 					GlobalMuscleLengthRatioThresholdForMaxActivation,
 					GlobalInflationVolumeScale);
+			if (bUseLengthActivationCurve)
+			{
+				Chaos::FLinearCurve ChaosCurve;
+				GlobalLengthActivationCurve.GetRichCurveConst()->ConvertToChaosCurve(ChaosCurve);
+				MuscleActivation.UpdateGlobalLengthActivationCurve(ChaosCurve);
+			}
 			if (ParameterMethod == EParameterMethod::Custom)
 			{
 				// match muscle names and override parameters
@@ -367,6 +376,12 @@ void FSetMuscleActivationParameterNode::Evaluate(UE::Dataflow::FContext& Context
 							Params.FiberLengthRatioAtMaxActivation,
 							Params.MuscleLengthRatioThresholdForMaxActivation,
 							Params.InflationVolumeScale);
+						if (Params.bUseLengthActivationCurve)
+						{
+							Chaos::FLinearCurve ChaosCurve;
+							Params.LengthActivationCurve.GetRichCurveConst()->ConvertToChaosCurve(ChaosCurve);
+							MuscleActivation.UpdateLengthActivationCurve(MuscleIndex, ChaosCurve);
+						}
 					}
 				}
 			}
diff --git a/Engine/Plugins/Experimental/ChaosFlesh/Source/ChaosFleshNodes/Public/Dataflow/ChaosFleshComputeMuscleActivationNode.h b/Engine/Plugins/Experimental/ChaosFlesh/Source/ChaosFleshNodes/Public/Dataflow/ChaosFleshComputeMuscleActivationNode.h
--- a/Engine/Plugins/Experimental/ChaosFlesh/Source/ChaosFleshNodes/Public/Dataflow/ChaosFleshComputeMuscleActivationNode.h
+++ b/Engine/Plugins/Experimental/ChaosFlesh/Source/ChaosFleshNodes/Public/Dataflow/ChaosFleshComputeMuscleActivationNode.h
@@ -6,6 +6,7 @@
 #include "Dataflow/DataflowCore.h"
 #include "Dataflow/DataflowEngine.h"
 #include "Dataflow/DataflowFunctionProperty.h"
+#include "Curves/CurveFloat.h"
 
 #include "ChaosFleshComputeMuscleActivationNode.generated.h"
 
@@ -80,6 +81,22 @@ public:
 	virtual void Evaluate(UE::Dataflow::FContext& Context, const FDataflowOutput* Out) const override;
 };
 
+class FLengthActivationUtils
+{
+public:
+	/** Sets a default linear (0,0) -> (1,1) curve */
+	static void SetDefaultLengthActivationCurve(FRuntimeFloatCurve& OutCurve)
+	{
+		FRichCurve* RichCurve = OutCurve.GetRichCurve();
+		if (RichCurve)
+		{
+			RichCurve->Reset();
+			RichCurve->AddKey(0.0f, 0.0f);
+			RichCurve->AddKey(1.0f, 1.0f);
+		}
+	}
+};
+
 /**
 * Struct data structure to store per-muscle activation parameters
 */
@@ -99,6 +116,15 @@ struct FPerMuscleParameter
 	float MuscleLengthRatioThresholdForMaxActivation = 0.75f;
 	UPROPERTY(EditAnywhere, Category = "Muscle Activation Parameters", meta = (ClampMin = "0"))
 	float InflationVolumeScale = 1.f;
+	UPROPERTY(EditAnywhere, meta = (InlineEditConditionToggle), Category = "Muscle Activation Parameters")
+	bool bUseLengthActivationCurve = false;
+	UPROPERTY(EditAnywhere, Category = "Muscle Activation Parameters", meta = (EditCondition = "bUseLengthActivationCurve"))
+	FRuntimeFloatCurve LengthActivationCurve;
+
+	FPerMuscleParameter()
+	{
+		FLengthActivationUtils::SetDefaultLengthActivationCurve(LengthActivationCurve);
+	}
 };
 
 UENUM()
@@ -145,6 +171,15 @@ public:
 	UPROPERTY(EditAnywhere, Category = "Global Parameters", meta = (ClampMin = "0"))
 	float GlobalInflationVolumeScale = 1.f;
 	
+	UPROPERTY(EditAnywhere, meta = (InlineEditConditionToggle), Category = "Global Parameters")
+	bool bUseLengthActivationCurve = false;
+
+	/** Curve editor for muscle length-activation curve. Default: linear
+	* X-axis: normalized muscle length level from 0 (rest state) to 1 (muscle length ratio = Muscle Length Ratio Threshold For Max Activation)
+	* Y-axis: muscle activation from 0 to 1 */
+	UPROPERTY(EditAnywhere, Category = "Global Parameters", meta = (EditCondition = "bUseLengthActivationCurve"))
+	FRuntimeFloatCurve GlobalLengthActivationCurve;
+
 	/** Use minimum muscle lengths across this animation asset to infer MuscleLengthRatioThresholdForMaxActivation*/
 	UPROPERTY(EditAnywhere, Category = "Animation Data", meta = (DataflowInput))
 	TObjectPtr<UAnimSequence> AnimationAsset = nullptr;
diff --git a/Engine/Source/Runtime/Engine/Classes/Curves/RichCurve.h b/Engine/Source/Runtime/Engine/Classes/Curves/RichCurve.h
--- a/Engine/Source/Runtime/Engine/Classes/Curves/RichCurve.h
+++ b/Engine/Source/Runtime/Engine/Classes/Curves/RichCurve.h
@@ -7,6 +7,7 @@
 #include "UObject/Class.h"
 #include "Curves/KeyHandle.h"
 #include "Curves/RealCurve.h"
+#include "Chaos/Curve.h"
 #include "RichCurve.generated.h"
 
 /** If using RCIM_Cubic, this enum describes how the tangents should be controlled in editor. */
@@ -335,6 +336,9 @@ public:
 	/** Allocates a duplicate of the curve */
 	virtual FIndexedCurve* Duplicate() const final { return new FRichCurve(*this); }
 
+	/** Convert a rich curve to Chaos FLinearCurve */
+	ENGINE_API void ConvertToChaosCurve(Chaos::FLinearCurve& OutCurve) const;
+
 private:
 	ENGINE_API void RemoveRedundantKeysInternal(float Tolerance, int32 InStartKeepKey, int32 InEndKeepKey, FFrameRate SampleRate);
 	ENGINE_API virtual int32 GetKeyIndex(float KeyTime, float KeyTimeTolerance) const override final;
diff --git a/Engine/Source/Runtime/Engine/Private/Curves/RichCurve.cpp b/Engine/Source/Runtime/Engine/Private/Curves/RichCurve.cpp
--- a/Engine/Source/Runtime/Engine/Private/Curves/RichCurve.cpp
+++ b/Engine/Source/Runtime/Engine/Private/Curves/RichCurve.cpp
@@ -1669,6 +1669,16 @@ void FRichCurve::CompressCurve(FCompressedRichCurve& OutCurve, float ErrorThresh
 	}
 }
 
+void FRichCurve::ConvertToChaosCurve(Chaos::FLinearCurve& OutCurve) const
+{
+	OutCurve.Keys.Reset();
+
+	for (const FRichCurveKey& Key: Keys)
+	{
+		OutCurve.Keys.Add({ Key.Time, Key.Value });
+	}
+}
+
 struct Quantized16BitKeyTimeAdapter
 {
 	static constexpr float QuantizationScale = 1.0f / 65535.0f;
diff --git a/Engine/Source/Runtime/Experimental/Chaos/Private/GeometryCollection/Facades/CollectionMuscleActivationFacade.cpp b/Engine/Source/Runtime/Experimental/Chaos/Private/GeometryCollection/Facades/CollectionMuscleActivationFacade.cpp
--- a/Engine/Source/Runtime/Experimental/Chaos/Private/GeometryCollection/Facades/CollectionMuscleActivationFacade.cpp
+++ b/Engine/Source/Runtime/Experimental/Chaos/Private/GeometryCollection/Facades/CollectionMuscleActivationFacade.cpp
@@ -29,6 +29,7 @@ namespace GeometryCollection::Facades
 	const FName FMuscleActivationFacade::FiberStreamline("FiberStreamline");
 	const FName FMuscleActivationFacade::FiberStreamlineRestLength("FiberStreamlineRestLength");
 	const FName FMuscleActivationFacade::MuscleActivationCurveName("MuscleActivationCurveName");
+	const FName FMuscleActivationFacade::LengthActivationCurve("LengthActivationCurve");
 	
 	FMuscleActivationFacade::FMuscleActivationFacade(FManagedArrayCollection& InCollection)
 		: ConstCollection(InCollection)
@@ -45,6 +46,7 @@ namespace GeometryCollection::Facades
 		, FiberStreamlineAttribute(InCollection, FiberStreamline, GroupName)
 		, FiberStreamlineRestLengthAttribute(InCollection, FiberStreamlineRestLength, GroupName)
 		, MuscleActivationCurveNameAttribute(InCollection, MuscleActivationCurveName, GroupName)
+		, LengthActivationCurveAttribute(InCollection, LengthActivationCurve, GroupName)
 	{
 		DefineSchema();
 	}
@@ -64,6 +66,7 @@ namespace GeometryCollection::Facades
 		, FiberStreamlineAttribute(InCollection, FiberStreamline, GroupName)
 		, FiberStreamlineRestLengthAttribute(InCollection, FiberStreamlineRestLength, GroupName)
 		, MuscleActivationCurveNameAttribute(InCollection, MuscleActivationCurveName, GroupName)
+		, LengthActivationCurveAttribute(InCollection, LengthActivationCurve, GroupName)
 	{
 		
 	}
@@ -91,6 +94,7 @@ namespace GeometryCollection::Facades
 		FiberStreamlineAttribute.Add();
 		FiberStreamlineRestLengthAttribute.Add();
 		MuscleActivationCurveNameAttribute.Add();
+		LengthActivationCurveAttribute.Add();
 	}
 
 	int32 FMuscleActivationFacade::AddMuscleActivationData(const FMuscleActivationData& InputData)
@@ -149,54 +153,54 @@ namespace GeometryCollection::Facades
 	int32 FMuscleActivationFacade::MuscleVertexOffset(const int32 MuscleIndex) const
 	{
 		const TManagedArray<int32>* VertexStart = ConstCollection.FindAttributeTyped<int32>("VertexStart", FGeometryCollection::GeometryGroup);
-		if (GeometryGroupIndexAttribute.IsValidIndex(MuscleIndex) && VertexStart->IsValidIndex(GeometryGroupIndexAttribute[MuscleIndex]))
+		if (IsValid() && GeometryGroupIndexAttribute.IsValidIndex(MuscleIndex) && VertexStart->IsValidIndex(GeometryGroupIndexAttribute[MuscleIndex]))
 		{
 			return (*VertexStart)[GeometryGroupIndexAttribute[MuscleIndex]];
 		}
-		else
-		{
-			return 0;
-		}
+		return 0;
 	}
 	
 	int32 FMuscleActivationFacade::NumMuscleVertices(const int32 MuscleIndex) const
 	{
 		const TManagedArray<int32>* VertexCount = ConstCollection.FindAttributeTyped<int32>("VertexCount", FGeometryCollection::GeometryGroup);
-		if (GeometryGroupIndexAttribute.IsValidIndex(MuscleIndex) && VertexCount->IsValidIndex(GeometryGroupIndexAttribute[MuscleIndex]))
+		if (IsValid() && GeometryGroupIndexAttribute.IsValidIndex(MuscleIndex) && VertexCount->IsValidIndex(GeometryGroupIndexAttribute[MuscleIndex]))
 		{
 			return (*VertexCount)[GeometryGroupIndexAttribute[MuscleIndex]];
 		}
-		else
-		{
-			return 0;
-		}
+		return 0;
 	}
 
 	FString FMuscleActivationFacade::FindMuscleName(const int32 MuscleIndex) const
 	{
 		const TManagedArray<FString>* BoneName = ConstCollection.FindAttributeTyped<FString>("BoneName", FTransformCollection::TransformGroup);
 		const TManagedArray<int32>* TransformIndex = ConstCollection.FindAttributeTyped<int32>("TransformIndex", FGeometryCollection::GeometryGroup);
-		int32 MuscleGeometryIndex = GeometryGroupIndexAttribute[MuscleIndex];
-		if (BoneName && TransformIndex && TransformIndex->IsValidIndex(MuscleGeometryIndex) && BoneName->IsValidIndex((*TransformIndex)[MuscleGeometryIndex]))
+		if (IsValid() && GeometryGroupIndexAttribute.IsValidIndex(MuscleIndex))
 		{
-			return (*BoneName)[(*TransformIndex)[MuscleGeometryIndex]];
+			int32 MuscleGeometryIndex = GeometryGroupIndexAttribute[MuscleIndex];
+			if (BoneName && TransformIndex && TransformIndex->IsValidIndex(MuscleGeometryIndex) && BoneName->IsValidIndex((*TransformIndex)[MuscleGeometryIndex]))
+			{
+				return (*BoneName)[(*TransformIndex)[MuscleGeometryIndex]];
+			}
 		}
 		return FString();
 	}
 
 	int32 FMuscleActivationFacade::FindMuscleIndexByName(const FString MuscleName) const
 	{
-		const TManagedArray<FString>* BoneName = ConstCollection.FindAttributeTyped<FString>("BoneName", FTransformCollection::TransformGroup);
-		const TManagedArray<int32>* TransformIndex = ConstCollection.FindAttributeTyped<int32>("TransformIndex", FGeometryCollection::GeometryGroup);
-		if (BoneName && TransformIndex)
+		if (IsValid())
 		{
-			int32 MuscleTransformIndex = BoneName->Find(MuscleName);
-			if (MuscleTransformIndex >= 0)
+			const TManagedArray<FString>* BoneName = ConstCollection.FindAttributeTyped<FString>("BoneName", FTransformCollection::TransformGroup);
+			const TManagedArray<int32>* TransformIndex = ConstCollection.FindAttributeTyped<int32>("TransformIndex", FGeometryCollection::GeometryGroup);
+			if (BoneName && TransformIndex)
 			{
-				int32 MuscleGeometryIndex = TransformIndex->Find(MuscleTransformIndex);
-				if (MuscleGeometryIndex >= 0)
+				int32 MuscleTransformIndex = BoneName->Find(MuscleName);
+				if (MuscleTransformIndex >= 0)
 				{
-					return GeometryGroupIndexAttribute.Get().Find(MuscleGeometryIndex);
+					int32 MuscleGeometryIndex = TransformIndex->Find(MuscleTransformIndex);
+					if (MuscleGeometryIndex >= 0)
+					{
+						return GeometryGroupIndexAttribute.Get().Find(MuscleGeometryIndex);
+					}
 				}
 			}
 		}
@@ -205,7 +209,7 @@ namespace GeometryCollection::Facades
 
 	int32 FMuscleActivationFacade::FindMuscleGeometryIndex(const int32 MuscleIndex) const
 	{
-		if (GeometryGroupIndexAttribute.IsValidIndex(MuscleIndex))
+		if (IsValid() && GeometryGroupIndexAttribute.IsValidIndex(MuscleIndex))
 		{
 			return GeometryGroupIndexAttribute[MuscleIndex];
 		}
@@ -232,6 +236,7 @@ namespace GeometryCollection::Facades
 
 	bool FMuscleActivationFacade::SetUpMuscleActivation(const TArray<int32>& InOrigin, const TArray<int32>& Insertion, float InContractionVolumeScale)
 	{
+		check(!IsConst());
 		// Vertices and fiber field
 		if (InOrigin.Num() == 0)
 		{
@@ -404,14 +409,15 @@ namespace GeometryCollection::Facades
 		float InGlobalMuscleLengthRatioThresholdForMaxActivation,
 		float InGlobalInflationVolumeScale)
 	{
+		check(!IsConst());
 		for (int32 MuscleIndex = 0; MuscleIndex < NumMuscles(); ++MuscleIndex)
 		{
-			FMuscleActivationData MuscleActivationData = GetMuscleActivationData(MuscleIndex);
-			MuscleActivationData.ContractionVolumeScale.Init(InGlobalContractionVolumeScale, MuscleActivationData.MuscleActivationElement.Num());
-			MuscleActivationData.FiberLengthRatioAtMaxActivation = InGlobalFiberLengthRatioAtMaxActivation;
-			MuscleActivationData.MuscleLengthRatioThresholdForMaxActivation = InGlobalMuscleLengthRatioThresholdForMaxActivation;
-			MuscleActivationData.InflationVolumeScale = InGlobalInflationVolumeScale;
-			UpdateMuscleActivationData(MuscleIndex, MuscleActivationData);
+			UpdateMuscleActivationParameters(
+				MuscleIndex,
+				InGlobalContractionVolumeScale,
+				InGlobalFiberLengthRatioAtMaxActivation,
+				InGlobalMuscleLengthRatioThresholdForMaxActivation,
+				InGlobalInflationVolumeScale);
 		}
 	}
 
@@ -422,6 +428,7 @@ namespace GeometryCollection::Facades
 		float InMuscleLengthRatioThresholdForMaxActivation,
 		float InInflationVolumeScale)
 	{
+		check(!IsConst());
 		if (IsValidMuscleIndex(MuscleIndex))
 		{
 			FMuscleActivationData MuscleActivationData = GetMuscleActivationData(MuscleIndex);
@@ -435,6 +442,33 @@ namespace GeometryCollection::Facades
 		return false;
 	}
 
+	void FMuscleActivationFacade::UpdateGlobalLengthActivationCurve(const Chaos::FLinearCurve& InGlobalLengthActivationCurve)
+	{
+		check(!IsConst());
+		for (int32 MuscleIndex = 0; MuscleIndex < NumMuscles(); ++MuscleIndex)
+		{
+			UpdateLengthActivationCurve(MuscleIndex, InGlobalLengthActivationCurve);
+		}
+	}
+
+	void FMuscleActivationFacade::UpdateLengthActivationCurve(int32 MuscleIndex, const Chaos::FLinearCurve& InLengthActivationCurve)
+	{
+		check(!IsConst());
+		if (LengthActivationCurveAttribute.IsValid() && IsValidMuscleIndex(MuscleIndex))
+		{
+			LengthActivationCurveAttribute.Modify()[MuscleIndex] = InLengthActivationCurve;
+		}
+	}
+
+	Chaos::FLinearCurve FMuscleActivationFacade::GetLengthActivationCurve(int32 MuscleIndex) const
+	{
+		if (LengthActivationCurveAttribute.IsValid() && LengthActivationCurveAttribute.IsValidIndex(MuscleIndex))
+		{
+			return LengthActivationCurveAttribute[MuscleIndex];
+		}
+		return Chaos::FLinearCurve();
+	}
+
 	TArray<TArray<TArray<FVector3f>>> FMuscleActivationFacade::BuildStreamlines(const TArray<int32>& Origin, const TArray<int32>& Insertion,
 		int32 NumLinesMultiplier, int32 MaxStreamlineIterations, int32 MaxPointsPerLine)
 	{
@@ -709,6 +743,7 @@ namespace GeometryCollection::Facades
 
 	int32 FMuscleActivationFacade::AssignCurveName(const FString& CurveName, const FString& MuscleName)
 	{
+		check(!IsConst());
 		const int32 MuscleIdx = FindMuscleIndexByName(MuscleName);
 		if (MuscleActivationCurveNameAttribute.IsValidIndex(MuscleIdx))
 		{
@@ -722,7 +757,7 @@ namespace GeometryCollection::Facades
 		TArray<int32> MuscleIndices;
 		if (MuscleActivationCurveNameAttribute.IsValid())
 		{
-			for (int32 MuscleIdx = 0; MuscleIdx < MuscleActivationCurveNameAttribute.Num(); ++MuscleIdx)
+			for (int32 MuscleIdx = 0; MuscleIdx < NumMuscles(); ++MuscleIdx)
 			{
 				//Animation curve name becomes lower case somehow after update, IgnoreCase here
 				if (MuscleActivationCurveNameAttribute[MuscleIdx].Equals(CurveName, ESearchCase::IgnoreCase)) 
diff --git a/Engine/Source/Runtime/Experimental/Chaos/Public/Chaos/Curve.h b/Engine/Source/Runtime/Experimental/Chaos/Public/Chaos/Curve.h
new file mode 100644
--- /dev/null
+++ b/Engine/Source/Runtime/Experimental/Chaos/Public/Chaos/Curve.h
@@ -0,0 +1,83 @@
+#pragma once
+
+#include "Chaos/ChaosArchive.h"
+#include "Containers/Array.h"
+#include "Math/UnrealMathUtility.h"
+
+namespace Chaos
+{
+	struct FCurveKey
+	{
+		float Time;
+		float Value;
+
+		friend FArchive& operator<<(FArchive& Ar, FCurveKey& Key)
+		{
+			Ar << Key.Time;
+			Ar << Key.Value;
+			return Ar;
+		}
+	};
+
+	class FLinearCurve
+	{
+	public:
+		TArray<FCurveKey> Keys;
+
+		FLinearCurve() = default;
+
+		FLinearCurve(std::initializer_list<FCurveKey> InKeys)
+			: Keys(InKeys)
+		{}
+
+		int32 GetNumKeys() const { return Keys.Num(); };
+
+		// Evaluate the curve at a given time using simple linear interpolation
+		float Eval(float InTime) const
+		{
+			const int32 NumKeys = Keys.Num();
+			if (NumKeys == 0)
+			{
+				return 0.f;
+			}
+
+			if (InTime <= Keys[0].Time)
+			{
+				return Keys[0].Value;
+			}
+
+			if (InTime >= Keys[NumKeys - 1].Time)
+			{
+				return Keys[NumKeys - 1].Value;
+			}
+
+			// Linear interpolation between nearest keys
+			for (int32 i = 1; i < NumKeys; ++i)
+			{
+				if (InTime < Keys[i].Time)
+				{
+					const FCurveKey& A = Keys[i - 1];
+					const FCurveKey& B = Keys[i];
+
+					const float Alpha = (InTime - A.Time) / (B.Time - A.Time);
+					return FMath::Lerp(A.Value, B.Value, Alpha);
+				}
+			}
+
+			return 0.f; // Should not reach here
+		}
+
+		bool Serialize(FArchive& Ar)
+		{
+			Ar << Keys;
+			return true;
+		}
+
+		friend FArchive& operator<<(FArchive& Ar, FLinearCurve& Curve)
+		{
+			Curve.Serialize(Ar);
+			return Ar;
+		}
+	};
+
+} // namespace Chaos
\ No newline at end of file
diff --git a/Engine/Source/Runtime/Experimental/Chaos/Public/Chaos/Deformable/MuscleActivationConstraints.h b/Engine/Source/Runtime/Experimental/Chaos/Public/Chaos/Deformable/MuscleActivationConstraints.h
--- a/Engine/Source/Runtime/Experimental/Chaos/Public/Chaos/Deformable/MuscleActivationConstraints.h
+++ b/Engine/Source/Runtime/Experimental/Chaos/Public/Chaos/Deformable/MuscleActivationConstraints.h
@@ -1,6 +1,7 @@
 // Copyright Epic Games, Inc. All Rights Reserved.
 #pragma once
 
+#include "Chaos/Curve.h"
 #include "Chaos/Framework/Parallel.h"
 #include "Chaos/XPBDCorotatedConstraints.h"
 #include "GeometryCollection/Facades/CollectionMuscleActivationFacade.h"
@@ -9,7 +10,6 @@ DEFINE_LOG_CATEGORY_STATIC(LogMuscleActivationConstraint, Log, All);
 
 namespace Chaos::Softs
 {
-
 	using Chaos::TVec3;
 	typedef GeometryCollection::Facades::FMuscleActivationFacade Facade;
 	typedef GeometryCollection::Facades::FMuscleActivationData Data;
@@ -65,8 +65,14 @@ namespace Chaos::Softs
 						MuscleActivation.Add(0.f);
 						MuscleVertexOffset.Add(FMuscleActivation.MuscleVertexOffset(MuscleIdx));
 						MuscleVertexCount.Add(FMuscleActivation.NumMuscleVertices(MuscleIdx));
+						LengthActivationCurves.Add(FMuscleActivation.GetLengthActivationCurve(MuscleIdx));
 					}
 				}
+				else
+				{
+					UE_LOG(LogMuscleActivationConstraint, Error, TEXT("Muscle Idx[%d] has invalid geometry index[%d]."), 
+						MuscleIdx, MuscleActivationData.GeometryGroupIndex);
+				}
 			}
 		}
 
@@ -87,8 +93,16 @@ namespace Chaos::Softs
 					}
 					else
 					{
-						// linear muscle activation model: muscle reaches max activation 1 at threshold length
-						MuscleActivation[MuscleIdx] = FMath::Clamp((1.f - MuscleLengthRatio) / (1 - MuscleLengthRatioThresholdForMaxActivation[MuscleIdx]), 0.f, 1.f);
+						float NormalizedMuscleLengthLevel = FMath::Clamp((1.f - MuscleLengthRatio) / (1 - MuscleLengthRatioThresholdForMaxActivation[MuscleIdx]), 0.f, 1.f);
+						if (LengthActivationCurves[MuscleIdx].GetNumKeys())
+						{
+							MuscleActivation[MuscleIdx] = LengthActivationCurves[MuscleIdx].Eval(NormalizedMuscleLengthLevel);
+						}
+						else
+						{
+							// no keys, default to linear muscle activation model: muscle reaches max activation 1 at threshold length
+							MuscleActivation[MuscleIdx] = NormalizedMuscleLengthLevel;
+						}
 					}
 				}
 			}
@@ -164,6 +178,7 @@ namespace Chaos::Softs
 		TArray<float> FiberLengthRatioAtMaxActivation;
 		TArray<float> MuscleLengthRatioThresholdForMaxActivation;
 		TArray<float> InflationVolumeScale;
+		TArray<Chaos::FLinearCurve> LengthActivationCurves;
 		TArray<int32> MuscleVertexOffset;
 		TArray<int32> MuscleVertexCount;
 	};
diff --git a/Engine/Source/Runtime/Experimental/Chaos/Public/GeometryCollection/Facades/CollectionMuscleActivationFacade.h b/Engine/Source/Runtime/Experimental/Chaos/Public/GeometryCollection/Facades/CollectionMuscleActivationFacade.h
--- a/Engine/Source/Runtime/Experimental/Chaos/Public/GeometryCollection/Facades/CollectionMuscleActivationFacade.h
+++ b/Engine/Source/Runtime/Experimental/Chaos/Public/GeometryCollection/Facades/CollectionMuscleActivationFacade.h
@@ -8,6 +8,11 @@
 #include "GeometryCollection/ManagedArrayCollection.h"
 #include "GeometryCollection/Facades/CollectionSelectionFacade.h"
 
+namespace Chaos
+{
+	class FLinearCurve;
+}
+
 namespace GeometryCollection::Facades
 {
 	//Activation data for each muscle
@@ -43,6 +48,7 @@ namespace GeometryCollection::Facades
 		static CHAOS_API const FName FiberStreamline;
 		static CHAOS_API const FName FiberStreamlineRestLength;
 		static CHAOS_API const FName MuscleActivationCurveName;
+		static CHAOS_API const FName LengthActivationCurve;
 
 		CHAOS_API FMuscleActivationFacade(FManagedArrayCollection& InCollection);
 		CHAOS_API FMuscleActivationFacade(const FManagedArrayCollection& InCollection);
@@ -84,6 +90,9 @@ namespace GeometryCollection::Facades
 			float InFiberLengthRatioAtMaxActivation,
 			float InMuscleLengthRatioThresholdForMaxActivation,
 			float InInflationVolumeScale);
+		CHAOS_API void UpdateGlobalLengthActivationCurve(const Chaos::FLinearCurve& InGlobalLengthActivationCurve);
+		CHAOS_API void UpdateLengthActivationCurve(int32 MuscleIndex, const Chaos::FLinearCurve& InLengthActivationCurve);
+		CHAOS_API Chaos::FLinearCurve GetLengthActivationCurve(int32 MuscleIndex) const;
 		CHAOS_API TArray<TArray<TArray<FVector3f>>> BuildStreamlines(const TArray<int32>& Origin, const TArray<int32>& Insertion,
 			 int32 NumLinesMultiplier, int32 MaxStreamlineIterations, int32 MaxPointsPerLine);
 		CHAOS_API int32 AssignCurveName(const FString& CurveName, const FString& MuscleName);
@@ -103,5 +112,6 @@ namespace GeometryCollection::Facades
 		TManagedArrayAccessor<TArray<TArray<FVector3f>>> FiberStreamlineAttribute;
 		TManagedArrayAccessor<TArray<float>> FiberStreamlineRestLengthAttribute;
 		TManagedArrayAccessor<FString> MuscleActivationCurveNameAttribute;
+		TManagedArrayAccessor<Chaos::FLinearCurve> LengthActivationCurveAttribute;
 	};
 }
diff --git a/Engine/Source/Runtime/Experimental/Chaos/Public/GeometryCollection/ManagedArrayTypeValues.inl b/Engine/Source/Runtime/Experimental/Chaos/Public/GeometryCollection/ManagedArrayTypeValues.inl
--- a/Engine/Source/Runtime/Experimental/Chaos/Public/GeometryCollection/ManagedArrayTypeValues.inl
+++ b/Engine/Source/Runtime/Experimental/Chaos/Public/GeometryCollection/ManagedArrayTypeValues.inl
@@ -61,7 +61,7 @@ MANAGED_ARRAY_TYPE(TArray<Chaos::PMatrix33d>, PMatrix33dArray)
 MANAGED_ARRAY_TYPE(TArray<TArray<FVector3f>>, FVector3fNestedArray)
 MANAGED_ARRAY_TYPE(FUintVector2, UintVector2)
 MANAGED_ARRAY_TYPE(TObjectPtr<UObject>, UObjectArray)
-
+MANAGED_ARRAY_TYPE(Chaos::FLinearCurve, LinearCurve)
 // NOTE: new types must be added at the bottom to keep serialization from breaking
 
 
diff --git a/Engine/Source/Runtime/Experimental/Chaos/Public/GeometryCollection/ManagedArrayTypes.h b/Engine/Source/Runtime/Experimental/Chaos/Public/GeometryCollection/ManagedArrayTypes.h
--- a/Engine/Source/Runtime/Experimental/Chaos/Public/GeometryCollection/ManagedArrayTypes.h
+++ b/Engine/Source/Runtime/Experimental/Chaos/Public/GeometryCollection/ManagedArrayTypes.h
@@ -18,6 +18,7 @@
 #include "Chaos/ParticleHandleFwd.h"
 #include "Chaos/Convex.h"
 #include "Chaos/Matrix.h"
+#include "Chaos/Curve.h"
 
 inline FArchive& operator<<(FArchive& Ar, TArray<FVector3f>*& ValueIn)
 {
