CIS fix: Add copyright boilerplate

diff --git a/Engine/Source/Runtime/Experimental/Chaos/Public/Chaos/Curve.h b/Engine/Source/Runtime/Experimental/Chaos/Public/Chaos/Curve.h
--- a/Engine/Source/Runtime/Experimental/Chaos/Public/Chaos/Curve.h
+++ b/Engine/Source/Runtime/Experimental/Chaos/Public/Chaos/Curve.h
@@ -1,3 +1,4 @@
+// Copyright Epic Games, Inc. All Rights Reserved.
 #pragma once
 
 #include "Chaos/ChaosArchive.h"
