OnDemand cache Defrag WIP. Checking in disabled while I continue to test/iterate.

diff --git a/Engine/Source/Runtime/Experimental/IoStore/OnDemand/Private/OnDemandContentInstaller.cpp b/Engine/Source/Runtime/Experimental/IoStore/OnDemand/Private/OnDemandContentInstaller.cpp
--- a/Engine/Source/Runtime/Experimental/IoStore/OnDemand/Private/OnDemandContentInstaller.cpp
+++ b/Engine/Source/Runtime/Experimental/IoStore/OnDemand/Private/OnDemandContentInstaller.cpp
@@ -369,16 +369,19 @@ void FOnDemandContentInstaller::ProcessRequest(FInstallRequestState& RequestStat
 	uint64 TotalInstallSize	= 0;
 
 	{
-		TMap<FIoHash, uint64> ChunksToInstall;
+		// TODO: This may double-request chunks. Chunks don't become referenced until the request completes, so two requests that
+		// share chunks can race to download the same chunk and it may be cached more than once.
+
+		TSet<FIoHash> ChunksToInstall;
 		for (int32 ContainerIndex = 0; FResolvedContainerChunks& ResolvedChunks : RequestState.ResolvedChunks)
 		{
-			for (int32 EntryIndex : ResolvedChunks .EntryIndices)
+			for (int32 EntryIndex : ResolvedChunks.EntryIndices)
 			{
 				const FOnDemandChunkEntry& Entry = ResolvedChunks.Container->ChunkEntries[EntryIndex];
 				const bool bInstalled = IoStore.InstallCache->IsChunkCached(Entry.Hash);
 				if (bInstalled == false)
 				{
-					ChunksToInstall.Add(Entry.Hash, Entry.EncodedSize);
+					ChunksToInstall.Add(Entry.Hash);
 					RequestState.HttpRequestHandles.Add(FChunkHttpRequestHandle
 					{
 						.Handle			= nullptr,
diff --git a/Engine/Source/Runtime/Experimental/IoStore/OnDemand/Private/OnDemandInstallCache.cpp b/Engine/Source/Runtime/Experimental/IoStore/OnDemand/Private/OnDemandInstallCache.cpp
--- a/Engine/Source/Runtime/Experimental/IoStore/OnDemand/Private/OnDemandInstallCache.cpp
+++ b/Engine/Source/Runtime/Experimental/IoStore/OnDemand/Private/OnDemandInstallCache.cpp
@@ -5,6 +5,7 @@
 #include "OnDemandIoStore.h"
 
 #include "Algo/Accumulate.h"
+#include "Algo/Find.h"
 #include "Async/Mutex.h"
 #include "Async/UniqueLock.h"
 #include "Async/AsyncFileHandle.h"
@@ -34,7 +35,7 @@
 #endif
 
 #ifndef UE_ONDEMANDINSTALLCACHE_EXCLUSIVE_WRITE
-#define UE_ONDEMANDINSTALLCACHE_EXCLUSIVE_WRITE (0)
+#define UE_ONDEMANDINSTALLCACHE_EXCLUSIVE_WRITE (1) // TODO: do not checkin
 #endif
 
 #if UE_ONDEMANDINSTALLCACHE_EXCLUSIVE_WRITE
@@ -47,6 +48,12 @@ namespace UE::IoStore
 ///////////////////////////////////////////////////////////////////////////////
 namespace CVars
 {
+	static bool GIoStoreOnDemandEnableDefrag = false;
+	static FAutoConsoleVariableRef CVar_IoStoreOnDemandEnableDefrag(
+		TEXT("iostore.EnableDefrag"),
+		GIoStoreOnDemandEnableDefrag,
+		TEXT("Whether to enable defrag when purging")
+	);
 }
 
 ///////////////////////////////////////////////////////////////////////////////
@@ -62,11 +69,15 @@ double ToMiB(uint64 Value)
 }
 
 ///////////////////////////////////////////////////////////////////////////////
-using FSharedAsyncFileHandle	= TSharedPtr<IAsyncReadFileHandle, ESPMode::ThreadSafe>;
-using FWeakAsyncFileHandle		= TWeakPtr<IAsyncReadFileHandle, ESPMode::ThreadSafe>;
-using FUniqueFileHandle			= TUniquePtr<IFileHandle>;
-using FSharedFileOpenResult		= TValueOrError<FSharedAsyncFileHandle, FFileSystemError>;
-using FCasAddr					= FHash96;
+using FUniqueFileHandle				= TUniquePtr<IFileHandle>;
+using FSharedFileHandle				= TSharedPtr<IFileHandle>;
+using FSharedFileOpenResult			= TValueOrError<FSharedFileHandle, FFileSystemError>;
+
+using FSharedAsyncFileHandle		= TSharedPtr<IAsyncReadFileHandle>;
+using FWeakAsyncFileHandle			= TWeakPtr<IAsyncReadFileHandle>;
+using FSharedFileOpenAsyncResult	= TValueOrError<FSharedAsyncFileHandle, FFileSystemError>;
+
+using FCasAddr						= FHash96;
 
 ///////////////////////////////////////////////////////////////////////////////
 struct FCasBlockId
@@ -135,7 +146,7 @@ struct FCasBlockInfo
 {
 	uint64	FileSize = 0;
 	int64	LastAccess = 0;
-	uint32	RefCount = 0;
+	uint64	RefSize = 0;
 };
 
 using FCasBlockInfoMap = TMap<FCasBlockId, FCasBlockInfo>;
@@ -143,31 +154,35 @@ using FCasBlockInfoMap = TMap<FCasBlockId, FCasBlockInfo>;
 ///////////////////////////////////////////////////////////////////////////////
 struct FCas
 {
-	using FLookup			= TMap<FCasAddr, FCasLocation>;
-	using FReadHandles		= TMap<FCasBlockId, FWeakAsyncFileHandle>;
-	using FLastAccess		= TMap<FCasBlockId, int64>;
-
-	FIoStatus				Initialize(FStringView Directory);
-	FCasLocation			FindChunk(const FIoHash& Hash) const;
-	FCasBlockId				CreateBlock();
-	FIoStatus				DeleteBlock(FCasBlockId BlockId, TArray<FCasAddr>& OutAddrs);
-	FString					GetBlockFilename(FCasBlockId BlockId) const;
-	FFileOpenResult 		OpenRead(FCasBlockId BlockId);
-	FSharedFileOpenResult	OpenAsyncRead(FCasBlockId  BlockId);
-	FUniqueFileHandle		OpenWrite(FCasBlockId BlockId);
-	void					TrackAccess(FCasBlockId BlockId, int64 UtcTicks);
-	void					TrackAccess(FCasBlockId BlockId) { TrackAccess(BlockId, FDateTime::UtcNow().GetTicks()); }
-	uint64					GetBlockInfo(FCasBlockInfoMap& OutBlockInfo);
-	void					Compact();
-	FIoStatus				Verify(TArray<FCasAddr>& OutAddrs);
-
-	FStringView			RootDirectory;
-	FLookup				Lookup;
-	TSet<FCasBlockId>	BlockIds;
-	FLastAccess			LastAccess;
-	FReadHandles		ReadHandles;	
-	uint32				MaxBlockSize = 32 << 20; //TODO: Make configurable	
-	mutable UE::FMutex	Mutex;
+	using FLookup				= TMap<FCasAddr, FCasLocation>;
+	using FReadHandles			= TMap<FCasBlockId, FWeakAsyncFileHandle>;
+	using FLastAccess			= TMap<FCasBlockId, int64>;
+	using FBlockIdHandleCounts	= TMap<FCasBlockId, int32>;
+
+	FIoStatus					Initialize(FStringView Directory);
+	FCasLocation				FindChunk(const FIoHash& Hash) const;
+	FCasBlockId					CreateBlock();
+	FIoStatus					DeleteBlock(FCasBlockId BlockId, TArray<FCasAddr>& OutAddrs);
+	FString						GetBlockFilename(FCasBlockId BlockId) const;
+	FSharedFileOpenResult		OpenRead(FCasBlockId BlockId);
+	FSharedFileOpenAsyncResult	OpenAsyncRead(FCasBlockId BlockId);
+	void						OnFileHandleDeleted(FCasBlockId BlockId);
+	FUniqueFileHandle			OpenWrite(FCasBlockId BlockId);
+	void						TrackAccess(FCasBlockId BlockId, int64 UtcTicks);
+	void						TrackAccess(FCasBlockId BlockId) { TrackAccess(BlockId, FDateTime::UtcNow().GetTicks()); }
+	void						TrackAccessIfNewer(FCasBlockId BlockId, int64 UtcTicks);
+	uint64						GetBlockInfo(FCasBlockInfoMap& OutBlockInfo);
+	void						Compact();
+	FIoStatus					Verify(TArray<FCasAddr>& OutAddrs);
+
+	FStringView				RootDirectory;
+	FLookup					Lookup;
+	FBlockIdHandleCounts	BlockIds;
+	FLastAccess				LastAccess;
+	FReadHandles			ReadHandles;
+	FEventRef				BlockReadsDoneEvent;
+	uint32					MaxBlockSize = 32 << 20; //TODO: Make configurable	
+	mutable UE::FMutex		Mutex;
 };
 
 ///////////////////////////////////////////////////////////////////////////////
@@ -236,7 +251,7 @@ FCasBlockId FCas::CreateBlock()
 			continue;
 		}
 
-		BlockIds.Add(BlockId);
+		BlockIds.Add(BlockId, 0);
 		LastAccess.FindOrAdd(BlockId, FDateTime::UtcNow().GetTicks());
 		Out = BlockId;
 	}
@@ -246,11 +261,29 @@ FCasBlockId FCas::CreateBlock()
 
 FIoStatus FCas::DeleteBlock(FCasBlockId BlockId, TArray<FCasAddr>& OutAddrs)
 {
-	UE::TUniqueLock Lock(Mutex);
-
-	IPlatformFile&	Ipf = FPlatformFileManager::Get().GetPlatformFile();
+	IPlatformFile& Ipf = FPlatformFileManager::Get().GetPlatformFile();
 	const FString	Filename = GetBlockFilename(BlockId);
 
+	UE::TDynamicUniqueLock Lock(Mutex, UE::FDeferLock());
+
+	// Wait for pending reads to flush before deleting block
+	for (;;)
+	{
+		Lock.Lock();
+
+		const int32 RequestCount = BlockIds.FindRef(BlockId);
+		if (RequestCount)
+		{
+			Lock.Unlock();
+			BlockReadsDoneEvent->Wait();
+		}
+		else
+		{
+			// Leave mutex locked until it goes out of scope
+			break;
+		}
+	}
+
 	UE_LOG(LogIoStoreOnDemand, Log, TEXT("Deleting CAS block '%s'"), *Filename);
 	if (Ipf.DeleteFile(*Filename) == false)
 	{
@@ -261,6 +294,7 @@ FIoStatus FCas::DeleteBlock(FCasBlockId BlockId, TArray<FCasAddr>& OutAddrs)
 	}
 
 	BlockIds.Remove(BlockId);
+	ReadHandles.Remove(BlockId);
 	for (auto It = Lookup.CreateIterator(); It; ++It)
 	{
 		if (It->Value.BlockId == BlockId)
@@ -286,15 +320,34 @@ FString FCas::GetBlockFilename(FCasBlockId BlockId) const
 	return FString(Path.ToView());
 }
 
-FFileOpenResult FCas::OpenRead(FCasBlockId BlockId)
+FSharedFileOpenResult FCas::OpenRead(FCasBlockId BlockId)
 {
 	const FString	Filename = GetBlockFilename(BlockId);
 	IPlatformFile&	Ipf = FPlatformFileManager::Get().GetPlatformFile();
 
-	return Ipf.OpenRead(*Filename, IPlatformFile::EOpenReadFlags::AllowWrite);
+	UE::TUniqueLock Lock(Mutex);
+
+	FFileOpenResult Result = Ipf.OpenRead(*Filename, IPlatformFile::EOpenReadFlags::AllowWrite);
+	if (Result.HasValue())
+	{
+		BlockIds.FindOrAdd(BlockId, 0)++;
+
+		FSharedFileHandle NewHandle(
+			Result.GetValue().Release(),
+			[this, BlockId](IFileHandle* RawHandle)
+			{
+				delete RawHandle;
+				OnFileHandleDeleted(BlockId);
+			}
+		);
+
+		return MakeValue(MoveTemp(NewHandle));
+	}
+
+	return MakeError(Result.StealError());
 }
 
-FSharedFileOpenResult FCas::OpenAsyncRead(FCasBlockId BlockId)
+FSharedFileOpenAsyncResult FCas::OpenAsyncRead(FCasBlockId BlockId)
 {
 	UE::TUniqueLock Lock(Mutex);
 
@@ -302,6 +355,8 @@ FSharedFileOpenResult FCas::OpenAsyncRead(FCasBlockId BlockId)
 	{
 		if (FSharedAsyncFileHandle Handle = MaybeHandle->Pin(); Handle.IsValid())
 		{
+			BlockIds.FindOrAdd(BlockId, 0)++;
+
 			return MakeValue(MoveTemp(Handle));
 		}
 	}
@@ -312,17 +367,35 @@ FSharedFileOpenResult FCas::OpenAsyncRead(FCasBlockId BlockId)
 
 	if (HandleResult.HasValue())
 	{
-		FSharedAsyncFileHandle NewHandle(HandleResult.GetValue().Release());
-		if (NewHandle.IsValid())
-		{
-			ReadHandles.FindOrAdd(BlockId, NewHandle);
-		}
+		BlockIds.FindOrAdd(BlockId, 0)++;
+
+		FSharedAsyncFileHandle NewHandle(
+			HandleResult.GetValue().Release(),
+			[this, BlockId](IAsyncReadFileHandle* RawHandle)
+			{
+				delete RawHandle;
+				OnFileHandleDeleted(BlockId);
+			}
+		);
+		ReadHandles.FindOrAdd(BlockId, NewHandle);
+		
 		return MakeValue(MoveTemp(NewHandle));
 	}
 
 	return MakeError(HandleResult.StealError());
 }
 
+void FCas::OnFileHandleDeleted(FCasBlockId BlockId)
+{
+	UE::TUniqueLock Lock(Mutex);
+	const int32 Count = --BlockIds.FindChecked(BlockId);
+	check(Count >= 0);
+	if (Count == 0)
+	{
+		BlockReadsDoneEvent->Trigger();
+	}
+}
+
 FUniqueFileHandle FCas::OpenWrite(FCasBlockId BlockId)
 {
 	IPlatformFile&	Ipf = FPlatformFileManager::Get().GetPlatformFile();
@@ -340,6 +413,17 @@ void FCas::TrackAccess(FCasBlockId BlockId, int64 UtcTicks)
 	LastAccess.FindOrAdd(BlockId, UtcTicks);
 }
 
+void FCas::TrackAccessIfNewer(FCasBlockId BlockId, int64 UtcTicks)
+{
+	check(BlockId.IsValid());
+	UE::TUniqueLock Lock(Mutex);
+	int64& FoundTicks = LastAccess.FindOrAdd(BlockId, FDateTime::MinValue().GetTicks());
+	if (FoundTicks < UtcTicks)
+	{
+		FoundTicks = UtcTicks;
+	}
+}
+
 uint64 FCas::GetBlockInfo(FCasBlockInfoMap& OutBlockInfo)
 {
 	TStringBuilder<256> Path;
@@ -432,7 +516,7 @@ FIoStatus FCas::Verify(TArray<FCasAddr>& OutAddrs)
 
 	for (auto BlockIt = BlockIds.CreateIterator(); BlockIt; ++BlockIt)
 	{
-		const FCasBlockId BlockId = *BlockIt;
+		const FCasBlockId BlockId = BlockIt->Key;
 		if (const FCasBlockInfo* Info = BlockInfo.Find(BlockId))
 		{
 			TotalVerifiedBytes += Info->FileSize;
@@ -1241,21 +1325,27 @@ public:
 	// IOnDemandInstallCache
 	virtual bool								IsChunkCached(const FIoHash& ChunkHash) override;
 	virtual FIoStatus							PutChunk(FIoBuffer&& Chunk, const FIoHash& ChunkHash) override;
-	virtual FIoStatus							Purge(TMap<FIoHash, uint64>&& ChunksToIntall) override;
+	virtual FIoStatus							Purge(TSet<FIoHash>&& ChunksToInstall) override;
 	virtual FIoStatus							PurgeAllUnreferenced() override;
 	virtual FIoStatus							Flush() override;
 	virtual FOnDemandInstallCacheUsage			GetCacheUsage() override;
 
 private:
-	void										AddReferencesToBlocks(
+	uint64										AddReferencesToBlocks(
 													const TArray<FSharedOnDemandContainer>& Containers, 
 													const TArray<TBitArray<>>& ChunkEntryIndices,
-													FCasBlockInfoMap& BlockInfo) const;
+													const TSet<FIoHash>& ChunksToInstall,
+													FCasBlockInfoMap& BlockInfoMap) const;
 	FIoStatus									Purge(const FCasBlockInfoMap& BlockInfo, uint64 TotalBytesToPurge, uint64& OutTotalPurgedBytes);
+	FIoStatus									Defrag(
+													const TArray<FSharedOnDemandContainer>& Containers,
+													const TArray<TBitArray<>>& ChunkEntryIndices,
+													FCasBlockInfoMap& BlockInfo, 
+													uint64 TotalBytesToFree);
 	bool										Resolve(FIoRequestImpl* Request);
 	void										CompleteRequest(FIoRequestImpl* Request, EIoErrorCode Status);
-	FIoStatus									FlushPendingChunks(FPendingChunks& Chunks);
-	FIoStatus									FlushPendingChunksImpl(FPendingChunks& Chunks);
+	FIoStatus									FlushPendingChunks(FPendingChunks& Chunks, int64 UtcAccessTicks = 0);
+	FIoStatus									FlushPendingChunksImpl(FPendingChunks& Chunks, int64 UtcAccessTicks = 0);
 	FString										GetJournalFilename() const { return CacheDirectory / TEXT("cas.jrn"); }
 	FString										GetSnapshotFilename() const { return CacheDirectory / TEXT("cas.snp"); }
 
@@ -1283,6 +1373,8 @@ FOnDemandInstallCache::FOnDemandInstallCache(const FOnDemandInstallCacheConfig&
 	, MaxJournalSize(Config.JournalMaxSize)
 {
 
+	// TODO: Reserve one block of the configured cache size be resvered for defrag overhead
+
 	UE_LOG(LogIoStoreOnDemand, Log, TEXT("Initializing install cache, MaxCacheSize=%.2lf MiB, MaxJournalSize=%.2lf KiB"),
 		ToMiB(MaxCacheSize), ToKiB(MaxJournalSize));
 
@@ -1313,7 +1405,7 @@ FOnDemandInstallCache::FOnDemandInstallCache(const FOnDemandInstallCacheConfig&
 			Cas.LastAccess.Reserve(Snapshot.Blocks.Num());
 			for (const FCasSnapshot::FBlock& Block : Snapshot.Blocks)
 			{
-				Cas.BlockIds.Add(Block.BlockId);
+				Cas.BlockIds.Add(Block.BlockId, 0);
 				Cas.LastAccess.Add(Block.BlockId, Block.LastAccess);
 			}
 		}
@@ -1341,7 +1433,7 @@ FOnDemandInstallCache::FOnDemandInstallCache(const FOnDemandInstallCacheConfig&
 		{
 			const FCasJournal::FEntry::FBlockOperation& Op = JournalEntry.BlockOperation;
 			CurrentBlock = Op.BlockId;
-			Cas.BlockIds.Add(Op.BlockId);
+			Cas.BlockIds.Add(Op.BlockId, 0);
 			break;
 		}
 		case FCasJournal::FEntry::EType::BlockDeleted:
@@ -1608,19 +1700,30 @@ bool FOnDemandInstallCache::Resolve(FIoRequestImpl* Request)
 		ExclusivePipe.Launch(UE_SOURCE_LOCATION, [this, Request, CasLoc]
 		{
 			FChunkRequest& ChunkRequest = FChunkRequest::GetRef(*Request);
-			bool bOk = false;
+			EIoErrorCode Status = EIoErrorCode::FileOpenFailed;
+
+			const FString Filename = Cas.GetBlockFilename(CasLoc.BlockId);
 
-			FFileOpenResult FileOpenResult = Cas.OpenRead(CasLoc.BlockId);
+			FSharedFileOpenResult FileOpenResult = Cas.OpenRead(CasLoc.BlockId);
 			if (FileOpenResult.IsValid())
 			{
-				TUniquePtr<IFileHandle> FileHandle = FileOpenResult.StealValue();
+				Status = EIoErrorCode::ReadError;
+
+				TSharedPtr<IFileHandle> FileHandle = FileOpenResult.StealValue();
 				const int64 CasBlockOffset = CasLoc.BlockOffset + ChunkRequest.ChunkRange.GetOffset();
-				if (Request->IsCancelled() == false && FileHandle->Seek(CasBlockOffset))
+				if (Request->IsCancelled())
 				{
-					bOk = FileHandle->Read(ChunkRequest.EncodedChunk.GetData(), ChunkRequest.EncodedChunk.GetSize());
-					if (bOk == false)
+					UE_LOG(LogIoStoreOnDemand, Verbose, TEXT("Cancelled request - skipped seek to offset %lld in CAS block '%s'"), CasBlockOffset, *Filename);
+				}
+				else if (FileHandle->Seek(CasBlockOffset))
+				{
+					const bool bOk = FileHandle->Read(ChunkRequest.EncodedChunk.GetData(), ChunkRequest.EncodedChunk.GetSize());
+					if (bOk)
+					{
+						Status = EIoErrorCode::Ok;
+					}
+					else
 					{
-						const FString Filename = Cas.GetBlockFilename(CasLoc.BlockId);
 						UE_LOG(LogIoStoreOnDemand, Error, TEXT("Failed to read %llu bytes at offset %lld in CAS block '%s'"),
 							ChunkRequest.EncodedChunk.GetSize(),
 							CasBlockOffset,
@@ -1629,20 +1732,17 @@ bool FOnDemandInstallCache::Resolve(FIoRequestImpl* Request)
 				}
 				else
 				{
-					const FString Filename = Cas.GetBlockFilename(CasLoc.BlockId);
 					UE_LOG(LogIoStoreOnDemand, Error, TEXT("Failed to seek to offset %lld in CAS block '%s'"), CasBlockOffset, *Filename);
 				}
 			}
 			else
 			{
-				const FString Filename = Cas.GetBlockFilename(CasLoc.BlockId);
 				UE_LOG(LogIoStoreOnDemand, Error, TEXT("Failed to open CAS block '%s' for reading, reason '%s'"),
 					*Filename, *FileOpenResult.GetError().GetMessage());
 			}
 
-			UE::Tasks::Launch(UE_SOURCE_LOCATION, [this, Request, bOk]
+			UE::Tasks::Launch(UE_SOURCE_LOCATION, [this, Request, Status]
 			{
-				const EIoErrorCode Status = bOk ? EIoErrorCode::Ok : EIoErrorCode::ReadError;
 				CompleteRequest(Request, Status);
 			});
 		}, UE::Tasks::ETaskPriority::BackgroundHigh);
@@ -1651,7 +1751,7 @@ bool FOnDemandInstallCache::Resolve(FIoRequestImpl* Request)
 	}
 #endif // UE_ONDEMANDINSTALLCACHE_EXCLUSIVE_WRITE
 
-	FSharedFileOpenResult FileOpenResult = Cas.OpenAsyncRead(CasLoc.BlockId);
+	FSharedFileOpenAsyncResult FileOpenResult = Cas.OpenAsyncRead(CasLoc.BlockId);
 	if (FileOpenResult.HasError())
 	{
 		const FString Filename = Cas.GetBlockFilename(CasLoc.BlockId);
@@ -1672,7 +1772,7 @@ bool FOnDemandInstallCache::Resolve(FIoRequestImpl* Request)
 		ChunkRange.ConsumeValueOrDie(),
 		RequestSize));
 
-	FAsyncFileCallBack Callback = [this, Request](bool bWasCancelled, IAsyncReadRequest* ReadRequest) 
+	FAsyncFileCallBack Callback = [this, Request](bool bWasCancelled, IAsyncReadRequest* ReadRequest)
 	{
 		UE::Tasks::Launch(UE_SOURCE_LOCATION, [this, Request, bWasCancelled]
 		{
@@ -1725,23 +1825,18 @@ FIoStatus FOnDemandInstallCache::PutChunk(FIoBuffer&& Chunk, const FIoHash& Chun
 	return FIoStatus::Ok;
 }
 
-FIoStatus FOnDemandInstallCache::Purge(TMap<FIoHash, uint64>&& ChunksToIntall) 
+FIoStatus FOnDemandInstallCache::Purge(TSet<FIoHash>&& ChunksToInstall)
 {
 	FCasBlockInfoMap	BlockInfo;
 	const uint64		TotalCachedBytes = Cas.GetBlockInfo(BlockInfo);
-	uint64				TotalUncachedBytes = 0;
 
-	for (const TPair<FIoHash, uint64>& Kv : ChunksToIntall)
-	{
-		if (FCasLocation Loc = Cas.FindChunk(Kv.Key); Loc.IsValid())
-		{
-			BlockInfo.FindOrAdd(Loc.BlockId).RefCount++;
-		}
-		else
-		{
-			TotalUncachedBytes += Kv.Value;
-		}
-	}
+	TArray<FSharedOnDemandContainer>	Containers;
+	TArray<TBitArray<>>					ChunkEntryIndices;
+
+	IoStore.GetReferencedContent(Containers, ChunkEntryIndices);
+	check(Containers.Num() == ChunkEntryIndices.Num());
+
+	const uint64 TotalUncachedBytes = AddReferencesToBlocks(Containers, ChunkEntryIndices, ChunksToInstall, BlockInfo);
 
 	const uint64 TotalRequiredBytes = TotalCachedBytes + TotalUncachedBytes;
 	if (TotalRequiredBytes <= MaxCacheSize)
@@ -1751,14 +1846,6 @@ FIoStatus FOnDemandInstallCache::Purge(TMap<FIoHash, uint64>&& ChunksToIntall)
 		return FIoStatus::Ok;
 	}
 
-	TArray<FSharedOnDemandContainer>	Containers; 
-	TArray<TBitArray<>>					ChunkEntryIndices;
-
-	IoStore.GetReferencedContent(Containers, ChunkEntryIndices);
-	check(Containers.Num() == ChunkEntryIndices.Num());
-
-	AddReferencesToBlocks(Containers, ChunkEntryIndices, BlockInfo);
-
 	//TODO: Compute fragmentation metric and redownload chunks when this number gets too high
 
 	BlockInfo.ValueSort([](const FCasBlockInfo& LHS, const FCasBlockInfo& RHS)
@@ -1767,7 +1854,7 @@ FIoStatus FOnDemandInstallCache::Purge(TMap<FIoHash, uint64>&& ChunksToIntall)
 	});
 
 	const uint64 TotalReferencedBytes = Algo::TransformAccumulate(BlockInfo,
-		[](const TPair<FCasBlockId, FCasBlockInfo>& Kv){ return (Kv.Value.RefCount > 0) ? Kv.Value.FileSize : uint64(0); },
+		[](const TPair<FCasBlockId, FCasBlockInfo>& Kv){ return (Kv.Value.RefSize > 0) ? Kv.Value.FileSize : uint64(0); },
 		uint64(0));
 
 	UE_LOG(LogIoStoreOnDemand, Log, TEXT("Purging install cache, MaxCacheSize=%.2lf MiB, CacheSize=%.2lf MiB, UncachedSize=%.2lf MiB, ReferencedBytes=%.2lf MiB"),
@@ -1791,7 +1878,20 @@ FIoStatus FOnDemandInstallCache::Purge(TMap<FIoHash, uint64>&& ChunksToIntall)
 
 	if (TotalPurgedBytes < TotalBytesToPurge)
 	{
-		return FIoStatusBuilder(EIoErrorCode::WriteError) << FString::Printf(TEXT("Failed to purge %" UINT64_FMT " from install cache"), TotalBytesToPurge);
+		if (UE::IoStore::CVars::GIoStoreOnDemandEnableDefrag)
+		{
+			// Attempt to defrag
+			FIoStatus DefragStatus = Defrag(Containers, ChunkEntryIndices, BlockInfo, TotalBytesToPurge - TotalPurgedBytes);
+			if (DefragStatus.IsOk() == false)
+			{
+				return FIoStatusBuilder(EIoErrorCode::WriteError) 
+					<< FString::Printf(TEXT("Failed to purge %" UINT64_FMT " from install cache after defrag (%s)"), TotalBytesToPurge, *DefragStatus.ToString());
+			}
+		}
+		else
+		{
+			return FIoStatusBuilder(EIoErrorCode::WriteError) << FString::Printf(TEXT("Failed to purge %" UINT64_FMT " from install cache"), TotalBytesToPurge);
+		}
 	}
 
 	return Status;
@@ -1801,7 +1901,6 @@ FIoStatus FOnDemandInstallCache::PurgeAllUnreferenced()
 {
 	FCasBlockInfoMap	BlockInfo;
 	const uint64		TotalCachedBytes = Cas.GetBlockInfo(BlockInfo);
-	uint64				TotalUncachedBytes = 0;
 
 	TArray<FSharedOnDemandContainer>	Containers;
 	TArray<TBitArray<>>					ChunkEntryIndices;
@@ -1809,20 +1908,22 @@ FIoStatus FOnDemandInstallCache::PurgeAllUnreferenced()
 	IoStore.GetReferencedContent(Containers, ChunkEntryIndices);
 	check(Containers.Num() == ChunkEntryIndices.Num());
 
-	AddReferencesToBlocks(Containers, ChunkEntryIndices, BlockInfo);
+	AddReferencesToBlocks(Containers, ChunkEntryIndices, {}, BlockInfo);
 
 	//TODO: Compute fragmentation metric and redownload chunks when this number gets too high
 
 	const uint64 TotalReferencedBytes = Algo::TransformAccumulate(BlockInfo,
-		[](const TPair<FCasBlockId, FCasBlockInfo>& Kv) { return (Kv.Value.RefCount > 0) ? Kv.Value.FileSize : uint64(0); },
+		[](const TPair<FCasBlockId, FCasBlockInfo>& Kv) { return (Kv.Value.RefSize > 0) ? Kv.Value.FileSize : uint64(0); },
 		uint64(0));
 
-	UE_LOG(LogIoStoreOnDemand, Log, TEXT("Purging install cache, MaxCacheSize=%.2lf MiB, CacheSize=%.2lf MiB, UncachedSize=%.2lf MiB, ReferencedBytes=%.2lf MiB"),
-		ToMiB(MaxCacheSize), ToMiB(TotalCachedBytes), ToMiB(TotalUncachedBytes), ToMiB(TotalReferencedBytes));
+	UE_LOG(LogIoStoreOnDemand, Log, TEXT("Purging install cache, MaxCacheSize=%.2lf MiB, CacheSize=%.2lf MiB, ReferencedBytes=%.2lf MiB"),
+		ToMiB(MaxCacheSize), ToMiB(TotalCachedBytes), ToMiB(TotalReferencedBytes));
 
 	uint64 TotalPurgedBytes = 0;
 	FIoStatus Status = Purge(BlockInfo, MaxCacheSize, TotalPurgedBytes);
 
+	// TODO: defrag as well?
+
 	if (TotalPurgedBytes > 0)
 	{
 		UE_LOG(LogIoStoreOnDemand, Log, TEXT("Purged %.2lf MiB (%.2lf%%) from install cache"),
@@ -1837,30 +1938,46 @@ FIoStatus FOnDemandInstallCache::PurgeAllUnreferenced()
 	return Status;
 }
 
-void FOnDemandInstallCache::AddReferencesToBlocks( 
-	const TArray<FSharedOnDemandContainer>& Containers, 
+uint64 FOnDemandInstallCache::AddReferencesToBlocks(
+	const TArray<FSharedOnDemandContainer>& Containers,
 	const TArray<TBitArray<>>& ChunkEntryIndices,
-	FCasBlockInfoMap& BlockInfo) const
+	const TSet<FIoHash>& ChunksToInstall,
+	FCasBlockInfoMap& BlockInfoMap) const
 {
+	uint64 TotalUncachedBytes = 0;
+
 	for (int32 Index = 0; FSharedOnDemandContainer Container : Containers)
 	{
 		const TBitArray<>& IsReferenced = ChunkEntryIndices[Index++];
 		for (int32 EntryIndex = 0; const FOnDemandChunkEntry& Entry : Container->ChunkEntries)
 		{
-			if (bool bIsReferenced = IsReferenced[EntryIndex++]; bIsReferenced == false)
+			const bool bIsReferenced = IsReferenced[EntryIndex++];
+			const bool bToInstall = ChunksToInstall.Contains(Entry.Hash);
+			
+			if (!bIsReferenced && !bToInstall)
 			{
 				continue;
 			}
-
+			
+			FCasBlockInfo* BlockInfo = nullptr;
 			if (FCasLocation Loc = Cas.FindChunk(Entry.Hash); Loc.IsValid())
 			{
-				if (FCasBlockInfo* Info = BlockInfo.Find(Loc.BlockId))
-				{
-					Info->RefCount++;
-				}
+				BlockInfo = BlockInfoMap.Find(Loc.BlockId);
+				ensure(BlockInfo);
+			}
+
+			if (BlockInfo)
+			{
+				BlockInfo->RefSize += Entry.EncodedSize;
+			}
+			else if (bToInstall)
+			{
+				TotalUncachedBytes += Entry.EncodedSize;
 			}
 		}
 	}
+
+	return TotalUncachedBytes;
 }
 
 FIoStatus FOnDemandInstallCache::Purge(const FCasBlockInfoMap& BlockInfo, const uint64 TotalBytesToPurge, uint64& OutTotalPurgedBytes)
@@ -1871,7 +1988,7 @@ FIoStatus FOnDemandInstallCache::Purge(const FCasBlockInfoMap& BlockInfo, const
 	{
 		const FCasBlockId BlockId = Kv.Key;
 		const FCasBlockInfo& Info = Kv.Value;
-		if (Info.RefCount > 0)
+		if (Info.RefSize > 0)
 		{
 			continue;
 		}
@@ -1910,6 +2027,177 @@ FIoStatus FOnDemandInstallCache::Purge(const FCasBlockInfoMap& BlockInfo, const
 	return FIoStatus::Ok;
 }
 
+FIoStatus FOnDemandInstallCache::Defrag(
+	const TArray<FSharedOnDemandContainer>& Containers,
+	const TArray<TBitArray<>>& ChunkEntryIndices,
+	FCasBlockInfoMap& BlockInfo, 
+	uint64 TotalBytesToFree)
+{
+	struct FDefragBlockReferencedChunk
+	{
+		uint32 BlockOffset = 0;
+		uint32 EncodedSize = 0;
+		FIoHash Hash;
+	};
+
+	struct FDefragBlock
+	{
+		FCasBlockId BlockId;
+		int64 LastAccess = 0;
+		TArray<FDefragBlockReferencedChunk> ReferencedChunks;
+	};
+
+	// Build the list of blocks to defrag and determine if its possible to free enough data through defragging
+	TArray<FDefragBlock> BlocksToDefrag;
+	bool bDefragPossible = false;
+	{
+		// Start with the least referenced blocks
+		BlockInfo.ValueSort([](const FCasBlockInfo& LHS, const FCasBlockInfo& RHS)
+		{
+			return LHS.RefSize < RHS.RefSize;
+		});
+
+		uint64 FreedBlockBytes = 0;
+		uint64 NewBlockBytes = 0;
+
+		for (const TPair<FCasBlockId, FCasBlockInfo>& Kv : BlockInfo)
+		{
+			const FCasBlockId BlockId = Kv.Key;
+			const FCasBlockInfo& Info = Kv.Value;
+
+			FreedBlockBytes += Info.FileSize;
+			NewBlockBytes += Info.RefSize; // For now, assume that nothing will be moved to the current block
+
+			BlocksToDefrag.Add(FDefragBlock{ .BlockId = BlockId, .LastAccess = Info.LastAccess });
+
+			if (FreedBlockBytes >= NewBlockBytes && FreedBlockBytes - NewBlockBytes > TotalBytesToFree)
+			{
+				bDefragPossible = true;
+				break;
+			}
+		}
+	}
+
+	if (!bDefragPossible)
+	{
+		return FIoStatusBuilder(EIoErrorCode::WriteError) << FString::Printf(TEXT("Defrag failed - cannot free %" UINT64_FMT), TotalBytesToFree);
+	}
+
+	// Right now, don't allow moving chunks to the current block for defrag. Its somewhat dangerous and hard to reason about.
+	// - Currently, the slack in the current block cannot be determined without opening a write handle to the block.
+	// - If we defrag the current block itself, then we would need additional tracking so we don't lose any chunks moved into it.
+	// - Additionally, this would also depend on the order blocks are defragged.
+	// This should be the only thread writing to CurrentBlock.
+	CurrentBlock = FCasBlockId::Invalid;
+
+	// Determine chunks that need to be moved for each defrag block
+	for (int32 Index = 0; FSharedOnDemandContainer Container : Containers)
+	{
+		const TBitArray<>& IsReferenced = ChunkEntryIndices[Index++];
+		for (int32 EntryIndex = 0; const FOnDemandChunkEntry& Entry : Container->ChunkEntries)
+		{
+			if (bool bIsReferenced = IsReferenced[EntryIndex++]; bIsReferenced == false)
+			{
+				continue;
+			}
+
+			if (FCasLocation Loc = Cas.FindChunk(Entry.Hash); Loc.IsValid())
+			{
+				if (FDefragBlock* DefragBlock = Algo::FindBy(BlocksToDefrag, Loc.BlockId, &FDefragBlock::BlockId))
+				{
+					DefragBlock->ReferencedChunks.Add(FDefragBlockReferencedChunk
+					{
+						.BlockOffset = Loc.BlockOffset,
+						.EncodedSize = Entry.EncodedSize,
+						.Hash = Entry.Hash,
+					});
+				}
+			}
+		}
+	}
+
+	// Move chunks to new blocks and delete old blocks
+	FPendingChunks DefragPendingChunks;
+	for (const FDefragBlock& DefragBlock : BlocksToDefrag)
+	{
+		FSharedFileOpenResult FileOpenResult = Cas.OpenRead(DefragBlock.BlockId);
+		if (!FileOpenResult.IsValid())
+		{
+			const FString Filename = Cas.GetBlockFilename(DefragBlock.BlockId);
+			const FString ErrorMsg = FileOpenResult.GetError().GetMessage();
+			UE_LOG(LogIoStoreOnDemand, Error, TEXT("Failed to open CAS block '%s' for reading, reason '%s'"), *Filename, *ErrorMsg);
+			return FIoStatusBuilder(EIoErrorCode::FileOpenFailed) << ErrorMsg;
+		}
+
+		FSharedFileHandle FileHandle = FileOpenResult.StealValue();
+
+		Algo::SortBy(DefragBlock.ReferencedChunks, &FDefragBlockReferencedChunk::BlockOffset);
+
+		for (const FDefragBlockReferencedChunk& ReffedChunk : DefragBlock.ReferencedChunks)
+		{
+			FileHandle->Seek(ReffedChunk.BlockOffset);
+
+			FIoBuffer Buffer;
+			Buffer.SetSize(ReffedChunk.EncodedSize);
+
+			const bool bOk = FileHandle->Read(Buffer.GetData(), Buffer.GetSize()); // TODO: error handling
+			if (!bOk)
+			{
+				return EIoErrorCode::ReadError;
+			}
+
+			if (DefragPendingChunks.TotalSize > FPendingChunks::MaxPendingBytes)
+			{
+				if (FIoStatus Status = FlushPendingChunks(DefragPendingChunks, DefragBlock.LastAccess); Status.IsOk() == false)
+				{
+					return Status;
+				}
+				check(DefragPendingChunks.IsEmpty());
+			}
+
+			DefragPendingChunks.Append(MoveTemp(Buffer), ReffedChunk.Hash);
+		}
+		
+		FileHandle.Reset();
+
+		if (FIoStatus Status = FlushPendingChunks(DefragPendingChunks); Status.IsOk() == false)
+		{
+			return Status;
+		}
+		check(DefragPendingChunks.IsEmpty());
+
+		// TODO: At startup, as part of the Verify, the cache needs to be check if there is enough room for the defrag chunk, and
+		// purge down to the correct size if not. Its possible something went wrong before the old block could be deleted.
+		// Alternatively we could create a transaction for every defragged block processed but its possible that would include
+		// multiple flushes. Is that OK?
+
+		// TODO: if something goes wrong before we get the chunk deletion here, we can be left in an over-allocated state.
+		// If we just keep trying to defrag in that state, we may just continue making new blocks without cleaning up old blocks.
+		// Its probably best that all subsequent defrag attempts fail if we ever leak a block here. A restart would be required to heal 
+		// the cache in that case.
+
+		FCasJournal::FTransaction Transaction = FCasJournal::Begin(GetJournalFilename());
+
+		// Flushing should overwrite the lookup info for the cas addr to point at the new block.
+		// Can now remove the old block
+		TArray<FCasAddr> DeletedChunkAddresses;
+		Cas.DeleteBlock(DefragBlock.BlockId, DeletedChunkAddresses);
+
+		for (const FCasAddr& Addr : DeletedChunkAddresses)
+		{
+			Transaction.ChunkLocation(FCasLocation::Invalid, Addr);
+		}
+		Transaction.BlockDeleted(DefragBlock.BlockId);
+
+		if (FIoStatus Status = FCasJournal::Commit(MoveTemp(Transaction)); !Status.IsOk())
+		{
+			return Status;
+		}
+	}
+
+	return FIoStatus::Ok;
+}
+
 FIoStatus FOnDemandInstallCache::Flush()
 {
 	if (PendingChunks.IsValid())
@@ -1937,19 +2225,19 @@ FOnDemandInstallCacheUsage FOnDemandInstallCache::GetCacheUsage()
 	IoStore.GetReferencedContent(Containers, ChunkEntryIndices);
 	check(Containers.Num() == ChunkEntryIndices.Num());
 
-	AddReferencesToBlocks(Containers, ChunkEntryIndices, BlockInfo);
+	AddReferencesToBlocks(Containers, ChunkEntryIndices, {}, BlockInfo);
 
 	return FOnDemandInstallCacheUsage
 	{
 		.MaxSize = MaxCacheSize,
 		.TotalSize = TotalCachedBytes,
 		.ReferencedBlockSize = Algo::TransformAccumulate(BlockInfo,
-			[](const TPair<FCasBlockId, FCasBlockInfo>& Kv) { return (Kv.Value.RefCount > 0) ? Kv.Value.FileSize : uint64(0); },
+			[](const TPair<FCasBlockId, FCasBlockInfo>& Kv) { return (Kv.Value.RefSize > 0) ? Kv.Value.FileSize : uint64(0); },
 			uint64(0))
 	};
 }
 
-FIoStatus FOnDemandInstallCache::FlushPendingChunks(FPendingChunks& Chunks)
+FIoStatus FOnDemandInstallCache::FlushPendingChunks(FPendingChunks& Chunks, int64 UtcAccessTicks)
 {
 	if (Chunks.IsEmpty())
 	{
@@ -1957,9 +2245,9 @@ FIoStatus FOnDemandInstallCache::FlushPendingChunks(FPendingChunks& Chunks)
 	}
 
 #if UE_ONDEMANDINSTALLCACHE_EXCLUSIVE_WRITE
-	UE::Tasks::TTask<FIoStatus> Task = ExclusivePipe.Launch(UE_SOURCE_LOCATION, [this, &Chunks]
+	UE::Tasks::TTask<FIoStatus> Task = ExclusivePipe.Launch(UE_SOURCE_LOCATION, [this, &Chunks, UtcAccessTicks]
 	{
-		return FlushPendingChunksImpl(Chunks);
+		return FlushPendingChunksImpl(Chunks, UtcAccessTicks);
 	}, UE::Tasks::ETaskPriority::BackgroundHigh);
 
 	Task.Wait();
@@ -1967,11 +2255,11 @@ FIoStatus FOnDemandInstallCache::FlushPendingChunks(FPendingChunks& Chunks)
 	return Task.GetResult();
 
 #else
-	return FlushPendingChunksImpl(Chunks);
+	return FlushPendingChunksImpl(Chunks, UtcAccessTicks);
 #endif // UE_ONDEMANDINSTALLCACHE_EXCLUSIVE_WRITE
 }
 
-FIoStatus FOnDemandInstallCache::FlushPendingChunksImpl(FPendingChunks& Chunks)
+FIoStatus FOnDemandInstallCache::FlushPendingChunksImpl(FPendingChunks& Chunks, int64 UtcAccessTicks)
 {
 	ON_SCOPE_EXIT { Chunks.Reset(); };
 
@@ -2026,7 +2314,15 @@ FIoStatus FOnDemandInstallCache::FlushPendingChunksImpl(FPendingChunks& Chunks)
 				return FIoStatusBuilder(EIoErrorCode::WriteError)
 					<< TEXT("Failed to serialize chunks to cache block");
 			}
-			Cas.TrackAccess(CurrentBlockId);
+
+			if (UtcAccessTicks)
+			{
+				Cas.TrackAccessIfNewer(CurrentBlockId, UtcAccessTicks);
+			}
+			else
+			{
+				Cas.TrackAccess(CurrentBlockId);
+			}
 
 			if (CasFileHandle->Flush() == false)
 			{
diff --git a/Engine/Source/Runtime/Experimental/IoStore/OnDemand/Private/OnDemandInstallCache.h b/Engine/Source/Runtime/Experimental/IoStore/OnDemand/Private/OnDemandInstallCache.h
--- a/Engine/Source/Runtime/Experimental/IoStore/OnDemand/Private/OnDemandInstallCache.h
+++ b/Engine/Source/Runtime/Experimental/IoStore/OnDemand/Private/OnDemandInstallCache.h
@@ -29,7 +29,7 @@ public:
 	virtual										~IOnDemandInstallCache() = default;
 	virtual bool								IsChunkCached(const FIoHash& ChunkHash) = 0;
 	virtual FIoStatus							PutChunk(FIoBuffer&& Chunk, const FIoHash& ChunkHash) = 0;
-	virtual FIoStatus							Purge(TMap<FIoHash, uint64>&& ChunksToInstall) = 0;
+	virtual FIoStatus							Purge(TSet<FIoHash>&& ChunksToInstall) = 0;
 	virtual FIoStatus							PurgeAllUnreferenced() = 0;
 	virtual FIoStatus							Flush() = 0;
 	virtual FOnDemandInstallCacheUsage			GetCacheUsage() = 0;
diff --git a/Engine/Source/Runtime/Experimental/IoStore/OnDemand/Private/OnDemandIoStore.cpp b/Engine/Source/Runtime/Experimental/IoStore/OnDemand/Private/OnDemandIoStore.cpp
--- a/Engine/Source/Runtime/Experimental/IoStore/OnDemand/Private/OnDemandIoStore.cpp
+++ b/Engine/Source/Runtime/Experimental/IoStore/OnDemand/Private/OnDemandIoStore.cpp
@@ -1615,7 +1615,7 @@ FOnDemandInstallResult FOnDemandIoStore::TickInstallRequest(FOnDemandInternalIns
 	{
 		TRACE_CPUPROFILER_EVENT_SCOPE(FOnDemandIoStore::PurgeInstallCache);
 
-		TMap<FIoHash, uint64> ChunksToInstall;
+		TSet<FIoHash> ChunksToInstall;
 		for (const auto& Kv : InstallData)
 		{
 			FOnDemandContainer&						Container = *Kv.Key;
@@ -1626,7 +1626,7 @@ FOnDemandInstallResult FOnDemandIoStore::TickInstallRequest(FOnDemandInternalIns
 				const int32 EntryIndex = ChunkPair.Key;
 
 				const FOnDemandChunkEntry& Entry = Container.ChunkEntries[EntryIndex];
-				ChunksToInstall.Add(Entry.Hash, Entry.EncodedSize);
+				ChunksToInstall.Add(Entry.Hash);
 			}
 		}
 
