Partial sleep & full wake

diff --git a/Engine/Source/Programs/HeadlessChaos/Private/HeadlessChaosTestGraphEvolution.cpp b/Engine/Source/Programs/HeadlessChaos/Private/HeadlessChaosTestGraphEvolution.cpp
--- a/Engine/Source/Programs/HeadlessChaos/Private/HeadlessChaosTestGraphEvolution.cpp
+++ b/Engine/Source/Programs/HeadlessChaos/Private/HeadlessChaosTestGraphEvolution.cpp
@@ -2034,7 +2034,7 @@ namespace ChaosTest
 	}
 
 	// Test validating the wake-up propagation for a rectangular block wall without lateral spacing.
-	// We add another constraint to connect two particles diagonally and wake one of the particles.
+	// We add another constraint to connect two particles diagonally, wake particle H and add an impulse to it.
 	//			  Hs - Ks - Ls
 	//			  |    |    |
 	//		      Es - Fs - Gs
@@ -2042,7 +2042,8 @@ namespace ChaosTest
 	//            Bs - Cs - Ds
 	//             \   |    /
 	//                 Ak
-	//  =>
+	//  1) For full island sleeping, the entire island will wake up.
+	// (d=dynamic, s=sleeping, k=kinematic)
 	//			  Hd - Kd - Ld
 	//			  |  \ |    |
 	//		      Ed - Fd - Gd
@@ -2050,6 +2051,15 @@ namespace ChaosTest
 	//            Bd - Cd - Dd
 	//             \   |    /
 	//                 Ak
+	//  2) For partial island sleeping, the state change in particle H will wake up its immediate neighbors and all particles at a higher level.
+	// (d=dynamic, s=sleeping, k=kinematic)
+	//			  Hd - Kd - Ls
+	//			  |  \ |    |
+	//		      Ed - Fd - Gs
+	//            |    |    |
+	//            Bs - Cs - Ds
+	//             \   |    /
+	//                 As
 	TEST_P(GraphEvolutionTests, TestConstraintGraph_WakeUpPropagation_AddToBlockWithoutSpacingAwake)
 	{
 		FGraphEvolutionTest Test(10, GetParam());
@@ -2126,34 +2136,68 @@ namespace ChaosTest
 		EXPECT_EQ(Test.IslandManager->GetNumParticles(), 10);
 		EXPECT_EQ(Test.IslandManager->GetNumConstraints(), 16);
 
-		EXPECT_EQ(Test.IslandManager->GetNumIslands(), 1);
-		EXPECT_FALSE(Test.IslandManager->GetIsland(0)->IsSleeping());
-		EXPECT_TRUE(FConstGenericParticleHandle(Test.ParticleHandles[0])->IsKinematic());
-		EXPECT_FALSE(FConstGenericParticleHandle(Test.ParticleHandles[1])->Sleeping());
-		EXPECT_FALSE(FConstGenericParticleHandle(Test.ParticleHandles[2])->Sleeping());
-		EXPECT_FALSE(FConstGenericParticleHandle(Test.ParticleHandles[3])->Sleeping());
-		EXPECT_FALSE(FConstGenericParticleHandle(Test.ParticleHandles[4])->Sleeping());
-		EXPECT_FALSE(FConstGenericParticleHandle(Test.ParticleHandles[5])->Sleeping());
-		EXPECT_FALSE(FConstGenericParticleHandle(Test.ParticleHandles[6])->Sleeping());
-		EXPECT_FALSE(FConstGenericParticleHandle(Test.ParticleHandles[7])->Sleeping());
-		EXPECT_FALSE(FConstGenericParticleHandle(Test.ParticleHandles[8])->Sleeping());
-		EXPECT_FALSE(FConstGenericParticleHandle(Test.ParticleHandles[9])->Sleeping());
-		EXPECT_FALSE(Test.ConstraintHandles[0]->IsSleeping());
-		EXPECT_FALSE(Test.ConstraintHandles[1]->IsSleeping());
-		EXPECT_FALSE(Test.ConstraintHandles[2]->IsSleeping());
-		EXPECT_FALSE(Test.ConstraintHandles[3]->IsSleeping());
-		EXPECT_FALSE(Test.ConstraintHandles[4]->IsSleeping());
-		EXPECT_FALSE(Test.ConstraintHandles[5]->IsSleeping());
-		EXPECT_FALSE(Test.ConstraintHandles[6]->IsSleeping());
-		EXPECT_FALSE(Test.ConstraintHandles[7]->IsSleeping());
-		EXPECT_FALSE(Test.ConstraintHandles[8]->IsSleeping());
-		EXPECT_FALSE(Test.ConstraintHandles[9]->IsSleeping());
-		EXPECT_FALSE(Test.ConstraintHandles[10]->IsSleeping());
-		EXPECT_FALSE(Test.ConstraintHandles[11]->IsSleeping());
-		EXPECT_FALSE(Test.ConstraintHandles[12]->IsSleeping());
-		EXPECT_FALSE(Test.ConstraintHandles[13]->IsSleeping());
-		EXPECT_FALSE(Test.ConstraintHandles[14]->IsSleeping());
-		EXPECT_FALSE(Test.ConstraintHandles[15]->IsSleeping());
+		if (!Test.CVarPartialSleeping->GetBool()) // island sleeping
+		{
+			EXPECT_EQ(Test.IslandManager->GetNumIslands(), 1);
+			EXPECT_FALSE(Test.IslandManager->GetIsland(0)->IsSleeping());
+			EXPECT_TRUE(FConstGenericParticleHandle(Test.ParticleHandles[0])->IsKinematic());
+			EXPECT_FALSE(FConstGenericParticleHandle(Test.ParticleHandles[1])->Sleeping());
+			EXPECT_FALSE(FConstGenericParticleHandle(Test.ParticleHandles[2])->Sleeping());
+			EXPECT_FALSE(FConstGenericParticleHandle(Test.ParticleHandles[3])->Sleeping());
+			EXPECT_FALSE(FConstGenericParticleHandle(Test.ParticleHandles[4])->Sleeping());
+			EXPECT_FALSE(FConstGenericParticleHandle(Test.ParticleHandles[5])->Sleeping());
+			EXPECT_FALSE(FConstGenericParticleHandle(Test.ParticleHandles[6])->Sleeping());
+			EXPECT_FALSE(FConstGenericParticleHandle(Test.ParticleHandles[7])->Sleeping());
+			EXPECT_FALSE(FConstGenericParticleHandle(Test.ParticleHandles[8])->Sleeping());
+			EXPECT_FALSE(FConstGenericParticleHandle(Test.ParticleHandles[9])->Sleeping());
+			EXPECT_FALSE(Test.ConstraintHandles[0]->IsSleeping());
+			EXPECT_FALSE(Test.ConstraintHandles[1]->IsSleeping());
+			EXPECT_FALSE(Test.ConstraintHandles[2]->IsSleeping());
+			EXPECT_FALSE(Test.ConstraintHandles[3]->IsSleeping());
+			EXPECT_FALSE(Test.ConstraintHandles[4]->IsSleeping());
+			EXPECT_FALSE(Test.ConstraintHandles[5]->IsSleeping());
+			EXPECT_FALSE(Test.ConstraintHandles[6]->IsSleeping());
+			EXPECT_FALSE(Test.ConstraintHandles[7]->IsSleeping());
+			EXPECT_FALSE(Test.ConstraintHandles[8]->IsSleeping());
+			EXPECT_FALSE(Test.ConstraintHandles[9]->IsSleeping());
+			EXPECT_FALSE(Test.ConstraintHandles[10]->IsSleeping());
+			EXPECT_FALSE(Test.ConstraintHandles[11]->IsSleeping());
+			EXPECT_FALSE(Test.ConstraintHandles[12]->IsSleeping());
+			EXPECT_FALSE(Test.ConstraintHandles[13]->IsSleeping());
+			EXPECT_FALSE(Test.ConstraintHandles[14]->IsSleeping());
+			EXPECT_FALSE(Test.ConstraintHandles[15]->IsSleeping());
+		}
+		else // partial island sleeping
+		{
+			EXPECT_EQ(Test.IslandManager->GetNumIslands(), 1);
+			EXPECT_FALSE(Test.IslandManager->GetIsland(0)->IsSleeping());
+			EXPECT_TRUE(FConstGenericParticleHandle(Test.ParticleHandles[0])->IsKinematic());
+			EXPECT_TRUE(FConstGenericParticleHandle(Test.ParticleHandles[1])->Sleeping());
+			EXPECT_TRUE(FConstGenericParticleHandle(Test.ParticleHandles[2])->Sleeping());
+			EXPECT_TRUE(FConstGenericParticleHandle(Test.ParticleHandles[3])->Sleeping());
+			EXPECT_FALSE(FConstGenericParticleHandle(Test.ParticleHandles[4])->Sleeping());
+			EXPECT_FALSE(FConstGenericParticleHandle(Test.ParticleHandles[5])->Sleeping());
+			EXPECT_TRUE(FConstGenericParticleHandle(Test.ParticleHandles[6])->Sleeping());
+			EXPECT_FALSE(FConstGenericParticleHandle(Test.ParticleHandles[7])->Sleeping());
+			EXPECT_FALSE(FConstGenericParticleHandle(Test.ParticleHandles[8])->Sleeping());
+			EXPECT_TRUE(FConstGenericParticleHandle(Test.ParticleHandles[9])->Sleeping());
+			EXPECT_TRUE(Test.ConstraintHandles[0]->IsSleeping());
+			EXPECT_TRUE(Test.ConstraintHandles[1]->IsSleeping());
+			EXPECT_TRUE(Test.ConstraintHandles[2]->IsSleeping());
+			EXPECT_TRUE(Test.ConstraintHandles[3]->IsSleeping());
+			EXPECT_TRUE(Test.ConstraintHandles[4]->IsSleeping());
+			EXPECT_FALSE(Test.ConstraintHandles[5]->IsSleeping());
+			EXPECT_FALSE(Test.ConstraintHandles[6]->IsSleeping());
+			EXPECT_TRUE(Test.ConstraintHandles[7]->IsSleeping());
+			EXPECT_FALSE(Test.ConstraintHandles[8]->IsSleeping());
+			EXPECT_FALSE(Test.ConstraintHandles[9]->IsSleeping());
+			EXPECT_FALSE(Test.ConstraintHandles[10]->IsSleeping());
+			EXPECT_FALSE(Test.ConstraintHandles[11]->IsSleeping());
+			EXPECT_TRUE(Test.ConstraintHandles[12]->IsSleeping());
+			EXPECT_FALSE(Test.ConstraintHandles[13]->IsSleeping());
+			EXPECT_FALSE(Test.ConstraintHandles[14]->IsSleeping());
+			EXPECT_FALSE(Test.ConstraintHandles[15]->IsSleeping());
+		}
 	}
 
 	INSTANTIATE_TEST_SUITE_P(, GraphEvolutionTests,	testing::Bool());
diff --git a/Engine/Source/Programs/HeadlessChaos/Private/HeadlessChaosTestSleeping.cpp b/Engine/Source/Programs/HeadlessChaos/Private/HeadlessChaosTestSleeping.cpp
--- a/Engine/Source/Programs/HeadlessChaos/Private/HeadlessChaosTestSleeping.cpp
+++ b/Engine/Source/Programs/HeadlessChaos/Private/HeadlessChaosTestSleeping.cpp
@@ -467,12 +467,12 @@ namespace ChaosTest
 		}
 		else
 		{			
-			// The entire stack will wake up
-			// @todo (Chaos): Not great, this should trigger a partial wake-up in the future.
-			for (FRigidBodyHandle_External* ParticleHandle : Test.ParticleHandles)
-			{
-				EXPECT_TRUE(ParticleHandle->ObjectState() == EObjectStateType::Dynamic);
-			}
+			// The stack wakes up only above the level of where the impulse is applied.
+			EXPECT_TRUE(Test.ParticleHandles[0]->ObjectState() == EObjectStateType::Sleeping);
+			EXPECT_TRUE(Test.ParticleHandles[1]->ObjectState() == EObjectStateType::Sleeping);
+			EXPECT_TRUE(Test.ParticleHandles[2]->ObjectState() == EObjectStateType::Dynamic);
+			EXPECT_TRUE(Test.ParticleHandles[3]->ObjectState() == EObjectStateType::Dynamic);
+			EXPECT_TRUE(Test.ParticleHandles[4]->ObjectState() == EObjectStateType::Dynamic);
 		}
 
 	}
diff --git a/Engine/Source/Runtime/Experimental/Chaos/Private/Chaos/Island/IslandManager.cpp b/Engine/Source/Runtime/Experimental/Chaos/Private/Chaos/Island/IslandManager.cpp
--- a/Engine/Source/Runtime/Experimental/Chaos/Private/Chaos/Island/IslandManager.cpp
+++ b/Engine/Source/Runtime/Experimental/Chaos/Private/Chaos/Island/IslandManager.cpp
@@ -53,15 +53,19 @@ namespace Chaos::CVars
 
 	/** Cvar to enable/disable partial island sleeping for all islands */
 	bool bChaosSolverPartialIslandSleep = false;
-	FAutoConsoleVariableRef CVarChaosSolverPartialIslandSleep(TEXT("p.Chaos.Solver.Sleep.PartialIslandSleep"), bChaosSolverPartialIslandSleep, TEXT("Partial island sleeping: Single non-moving particles in islands will be deactivated even if touching moving particles. NOTE: Wake-up only works correctly if shock propagation is enabled."));
+	FAutoConsoleVariableRef CVarChaosSolverPartialIslandSleep(TEXT("p.Chaos.Solver.Sleep.PartialIslandSleep"), bChaosSolverPartialIslandSleep, TEXT("[EXPERIMENTAL] Single non-moving particles in islands will be deactivated even if they are connected to moving particles."));
+
+	/** Cvar to enable/disable partial island waking for all islands */
+	bool bChaosSolverPartialIslandWake = true;
+	FAutoConsoleVariableRef CVarChaosSolverPartialIslandWake(TEXT("p.Chaos.Solver.Sleep.PartialIslandWake"), bChaosSolverPartialIslandWake, TEXT("[EXPERIMENTAL] New constraints and particle state changes will only wake particles farther away from the closest kinematic including the ground (partial island sleep only)."));
 
 	/** Cvar to enable/disable waking up the entire island if the state of any particle has changed (partial island sleeping only) */
-	bool bChaosSolverWakeOnParticleStateChanged = true;
-	FAutoConsoleVariableRef CVarChaosSolverWakeOnParticleStateChanged(TEXT("p.Chaos.Solver.Sleep.WakeOnParticleStateChanged"), bChaosSolverWakeOnParticleStateChanged, TEXT("State changes of a single particle will trigger the island to wake up entirely (partial island sleeping only)."));
+	bool bChaosSolverWakeOnParticleStateChanged = false;
+	FAutoConsoleVariableRef CVarChaosSolverWakeOnParticleStateChanged(TEXT("p.Chaos.Solver.Sleep.WakeOnParticleStateChanged"), bChaosSolverWakeOnParticleStateChanged, TEXT("[EXPERIMENTAL] State changes of a particle will trigger the island to wake up entirely (partial island sleep only)."));
 
 	/** Cvar to enable/disable triggering wake up for newly detected constraints with moving particles only (partial island sleeping only) */
 	bool bChaosSolverWakeOnMovingOnly = true;
-	FAutoConsoleVariableRef CVarChaosSolverWakeOnMovingOnly(TEXT("p.Chaos.Solver.Sleep.WakeOnMovingOnly"), bChaosSolverWakeOnMovingOnly, TEXT("Only newly detected constraints with moving particle will trigger wake-up event(partial island sleeping only)."));
+	FAutoConsoleVariableRef CVarChaosSolverWakeOnMovingOnly(TEXT("p.Chaos.Solver.Sleep.WakeOnMovingOnly"), bChaosSolverWakeOnMovingOnly, TEXT("[EXPERIMENTAL] Only newly detected constraints connected to a moving particle will trigger wake-up event (partial island sleep only)."));
 
 	/** Cvar to override the sleep counter threshold if necessary */
 	int32 ChaosSolverCollisionDefaultSleepCounterThreshold = 20;
@@ -82,6 +86,14 @@ namespace Chaos::CVars
 	FRealSingle ChaosSolverCollisionAngularSleepThresholdSize = 0;
 	FAutoConsoleVariableRef CVarChaosSolverCollisionAngularSleepThresholdSize(TEXT("p.Chaos.Solver.Sleep.AngularSleepThresholdSize"), ChaosSolverCollisionAngularSleepThresholdSize, TEXT("Scales the angular threshold based on size (0 to disable size based scaling)"));
 
+	/** Cvar to adjust the linear wake threshold for partial waking */
+	FRealSingle LinearPartialWakeThresholdMultiplier = 1.0f;
+	FAutoConsoleVariableRef CVarChaosSolverLinearWakeThresholdMultiplier(TEXT("p.Chaos.Solver.Sleep.Defaults.LinearWakeThresholdMultiplier"), LinearPartialWakeThresholdMultiplier, TEXT("[EXPERIMENTAL] A multiplier applied to SleepLinearThreshold for particle waking (partial island sleeping only)."));
+
+	/** Cvar to adjust the angular wake threshold for partial waking */
+	FRealSingle AngularPartialWakeThresholdMultiplier = 1.0f;
+	FAutoConsoleVariableRef CVarChaosSolverAngularWakeThresholdMultiplier(TEXT("p.Chaos.Solver.Sleep.Defaults.AngularWakeThresholdMultiplier"), AngularPartialWakeThresholdMultiplier, TEXT("[EXPERIMENTAL] A multiplier applied to SleepAngularThreshold for particle waking (partial island sleeping only)."));
+
 	/* Cvar to increase the sleep counter threshold for floating particles */
 	int32 IsolatedParticleSleepCounterThresholdMultiplier = 1;
 	FAutoConsoleVariableRef CVarChaosSolverIsolatedParticleSleepCounterThresholdMultiplier(TEXT("p.Chaos.Solver.Sleep.IsolatedParticle.CounterMultiplier"), IsolatedParticleSleepCounterThresholdMultiplier, TEXT("A multiplier applied to SleepCounterThreshold for floating particles"));
@@ -173,7 +185,9 @@ namespace Chaos::Private
 		const THandleArray<FChaosPhysicsMaterial>* SimMaterials,
 		FRealSingle& OutSleepLinearThreshold,
 		FRealSingle& OutSleepAngularThreshold,
-		int& OutSleepCounterThreshold)
+		int& OutSleepCounterThreshold,
+		FRealSingle& OutWakeLinearThreshold,
+		FRealSingle& OutWakeAngularThreshold)
 	{
 		OutSleepLinearThreshold = 0;
 		OutSleepAngularThreshold = 0;
@@ -211,6 +225,10 @@ namespace Chaos::Private
 				}
 			}
 
+			// Set the wake threshold used in partial island waking based on the particle's sleep threshold
+			OutWakeLinearThreshold = CVars::LinearPartialWakeThresholdMultiplier * OutSleepLinearThreshold;
+			OutWakeAngularThreshold = CVars::AngularPartialWakeThresholdMultiplier * OutSleepAngularThreshold;
+
 			return true;
 		}
 
@@ -226,10 +244,12 @@ namespace Chaos::Private
 		FRealSingle& OutSleepAngularThreshold,
 		int& OutSleepCounterThreshold)
 	{
-		if (GetIslandParticleSleepThresholds(Particle, PhysicsMaterials, PerParticlePhysicsMaterials, SimMaterials, OutSleepLinearThreshold, OutSleepAngularThreshold, OutSleepCounterThreshold))
+		FRealSingle OutWakeAngularThreshold, OutWakeLinearThreshold;
+		if (GetIslandParticleSleepThresholds(Particle, PhysicsMaterials, PerParticlePhysicsMaterials, SimMaterials, OutSleepLinearThreshold, OutSleepAngularThreshold, OutSleepCounterThreshold,
+			OutWakeLinearThreshold, OutWakeAngularThreshold))
 		{
 			// Sleep thresholds are tuned to hide minor collision jitter which isn't a problem for floating particles
-			// so we scale the thesholds to make sleeping harder. E.g., to avoid going to sleep at the apex of an
+			// so we scale the thresholds to make sleeping harder. E.g., to avoid going to sleep at the apex of an
 			// upwards ballistic trajectory, or when a floating oscillating body reverses rotation, etc.
 			OutSleepCounterThreshold *= FMath::Max(1, CVars::IsolatedParticleSleepCounterThresholdMultiplier);
 			OutSleepLinearThreshold *= FMath::Max(0.0f, CVars::IsolatedParticleSleepLinearThresholdMultiplier);
@@ -1236,14 +1256,18 @@ namespace Chaos::Private
 		if (Node->Flags.bIsDynamic)
 		{
 			FRealSingle SleepLinearThreshold, SleepAngularThreshold;
+			FRealSingle WakeLinearThreshold, WakeAngularThreshold;
 			FRealSingle DisableLinearThreshold, DisableAngularThreshold;
 			int32 SleepCounterThreshold;
-			GetIslandParticleSleepThresholds(Node->Particle, PhysicsMaterials, PerParticlePhysicsMaterials, SimMaterials, SleepLinearThreshold, SleepAngularThreshold, SleepCounterThreshold);
+			GetIslandParticleSleepThresholds(Node->Particle, PhysicsMaterials, PerParticlePhysicsMaterials, SimMaterials, SleepLinearThreshold, SleepAngularThreshold,
+				SleepCounterThreshold, WakeLinearThreshold, WakeAngularThreshold);
 			GetParticleDisableThresholds(Node->Particle, PhysicsMaterials, PerParticlePhysicsMaterials, SimMaterials, DisableLinearThreshold, DisableAngularThreshold);
 
 			Node->SleepLinearThresholdSq = FMath::Square(SleepLinearThreshold);
 			Node->SleepAngularThresholdSq = FMath::Square(SleepAngularThreshold);
 			Node->SleepCounterThreshold = SleepCounterThreshold;
+			Node->WakeLinearThresholdSq = FMath::Square(WakeLinearThreshold);
+			Node->WakeAngularThresholdSq = FMath::Square(WakeAngularThreshold);
 			Node->DisableLinearThresholdSq = FMath::Square(DisableLinearThreshold);
 			Node->DisableAngularThresholdSq = FMath::Square(DisableAngularThreshold);
 		}
@@ -1306,13 +1330,23 @@ namespace Chaos::Private
 				// When do we want to force the entire island to wake up for at least 1 step? (partial island sleeping only)
 				// Case 1) Switch from non-moving kinematic to non-sleeping dynamic OR
 				// Case 2) Switch dynamic particle from sleeping to not-sleeping
-				// @todo(chaos): This will wake up the entire island for partial island sleeping.
-				// Not a great solution for performance. Update this when implementing the next version.
 				const bool bSwitchStaticKinematicToAwakeDynamic = (!bWasDynamic && !bWasMoving) && (bIsDynamic && !bIsSleeping); // Case 1
 				const bool bSwitchSleepingDynamicToAwake = (bIsDynamic && bWasSleeping && !bIsSleeping); // Case 2
 				const bool bIsSleepAllowed = (bSwitchStaticKinematicToAwakeDynamic || bSwitchSleepingDynamicToAwake) ? false : true;
 				EnqueueIslandCheckSleep(Node->Island, bIsSleepAllowed);
 			}
+			else if (CVars::bChaosSolverPartialIslandSleep && !CVars::bChaosSolverWakeOnParticleStateChanged)
+			{
+				// Register wake-up events for updates in the particles (partial island sleeping only)
+				// Case 1) Switch from non-moving kinematic to non-sleeping dynamic OR
+				// Case 2) Switch dynamic particle from sleeping to not-sleeping
+				const bool bSwitchStaticKinematicToAwakeDynamic = (!bWasDynamic && !bWasMoving) && (bIsDynamic && !bIsSleeping); // Case 1
+				const bool bSwitchSleepingDynamicToAwake = (bIsDynamic && bWasSleeping && !bIsSleeping); // Case 2
+				const bool bIsSleepAllowed = true;
+				const bool bTriggerWakeUpEvent = (bSwitchStaticKinematicToAwakeDynamic || bSwitchSleepingDynamicToAwake);
+				Node->Flags.bWasUpdated = bTriggerWakeUpEvent ? true : false;
+				EnqueueIslandCheckSleep(Node->Island, bIsSleepAllowed);
+			}
 			else
 			{
 				const bool bIsSleepAllowed = true;
@@ -2122,21 +2156,21 @@ namespace Chaos::Private
 				{
 					if (!Island->Flags.bIsSleeping)
 					{
-						ProcessParticleWakes(Island);
+						ProcessPartialIslandWake(Island);
 					}
 					SyncConstraintSleepStateToParticles(Island);
 				}
 				else // For island sleeping to sync the sleep state OR if (!bIsSleepAllowed) to wake up the entire island
 				{
 					PropagateIslandSleep(Island);
-
-					if (CVars::bChaosSolverPartialIslandSleep)
-					{
-						// NOTE: This reset is needed if (CVars::bChaosSolverPartialIslandSleepEnabled && Island->Flags.bIsSleepAllowed)
-						ResetEdgeFlags(Island);
-					}
 				}
 				Island->UpdateSyncState();
+
+				if (CVars::bChaosSolverPartialIslandSleep)
+				{
+					ResetEdgeFlags(Island);
+					ResetNodeFlags(Island);
+				}
 			}
 
 			// Remove islands without constraints. NOTE: must come after PropagateIslandSleep
@@ -2590,67 +2624,120 @@ namespace Chaos::Private
 		}
 	}
 
-	void FPBDIslandManager::ProcessParticleWakes(FPBDIsland* Island)
+	void FPBDIslandManager::ProcessPartialIslandWake(FPBDIsland* Island)
 	{
-		TArray<FGeometryParticleHandle*> ParticlesToWakeUp;
+		TArray<FGeometryParticleHandle*> ParticlesToWakeUp; // Particles which will be woken up after processing all partial wake events.
+		TArray<FPBDIslandParticle*> NodeQueue;				// Queue of island nodes to be traversed by partial waking.
+		NodeQueue.Reserve(Island->Nodes.Num());
+
+		// Phase 1 - Waking triggered by new constraints
+		// Step 1a - Find trigger particles (i.e. particles connected to new constraints with significant motion)
 		for (TArray<FPBDIslandConstraint*>& IslandEdges : Island->ContainerEdges)
 		{
 			for (FPBDIslandConstraint* IslandEdge : IslandEdges)
 			{
-				// For now, only new constraints can induce partial island wakes
-				if (!!IslandEdge->Flags.bIsNew)
+				if (IslandEdge->Flags.bIsNew)
 				{
-					IslandEdge->Flags.bWasProcessed = true;
-
-					// Only propagated the wake-up event if at least one particle connected be the new constraint is woken up.
-					// This means there will not be any wake-up propagation if
-					// 1) both particles are already awake.
-					// 2) both particles are sleeping.
-					// 3) neither particle is moving if (bChaosSolverWakeUpOnConstraintWithNonMovingParticle == false).
-					if (ProcessConnectedParticles(IslandEdge, ParticlesToWakeUp, !CVars::bChaosSolverWakeOnMovingOnly))
+					FPBDIslandParticle* Node0 = IslandEdge->Nodes[0];
+					FPBDIslandParticle* Node1 = IslandEdge->Nodes[1];
+
+					if (Node0 && !Node0->Flags.bWasProcessed)
 					{
-						// Wake-up propagation: Breadth-first traversal of the constraint graph
-						// If level information is available, only propagate the wake-up away from kinematic particles (i.e. Node->Level >= Edge->Level).
-						// This will cause the top of a stack to wake up if a particles closer to the kinematic ground wakes up.
-						// Meanwhile particles closer to the ground remains sleeping.
-						TArray<FPBDIslandParticle*> NodeQueue;
-						NodeQueue.Reserve(Island->Nodes.Num());
-						NodeQueue.Add(IslandEdge->Nodes[0]);
-						NodeQueue.Add(IslandEdge->Nodes[1]);
-						for (int32 NodeIndex = 0; NodeIndex < NodeQueue.Num(); ++NodeIndex)
-						{
-							FPBDIslandParticle* Node = NodeQueue[NodeIndex];
-							for (FPBDIslandConstraint* Edge : Node->Edges)
-							{
-								if (!Edge->Flags.bWasProcessed)
-								{
-									Edge->Flags.bWasProcessed = true;
-
-									FPBDIslandParticle* OtherNode = (Node == Edge->Nodes[0]) ? Edge->Nodes[1] : Edge->Nodes[0];
-									// Wake up the OtherNode if the node level is greater than or equal to the edge level.
-									if ((OtherNode != nullptr) && OtherNode->Flags.bIsDynamic && (OtherNode->Level >= Edge->Level))
-									{
-										ProcessConnectedParticles(Edge, ParticlesToWakeUp, true);
-										NodeQueue.Add(OtherNode);
-									}
-								}
-							}
-						}
+						ProcessTriggerParticle(Node0, NodeQueue, /*bAlwaysTrigger=*/ !CVars::bChaosSolverWakeOnMovingOnly);
+						Node0->Flags.bWasProcessed = true;
+					}
+					if (Node1 && !Node1->Flags.bWasProcessed)
+					{
+						ProcessTriggerParticle(Node1, NodeQueue, /*bAlwaysTrigger=*/ !CVars::bChaosSolverWakeOnMovingOnly);
+						Node1->Flags.bWasProcessed = true;
 					}
 				}
 			}
 		}
+		const int32 NumNodesConstraintWake = NodeQueue.Num();
+		bool bResetEdgeFlags = false;
+
+		// Step 1b - Propagate wake event
+		// Partial island wake disabled, but at least one trigger particle found? Then, perform a full island wake.
+		if (!CVars::bChaosSolverPartialIslandWake && NumNodesConstraintWake > 0)
+		{
+			PropagateIslandSleep(Island);
+		}
+		else // Otherwise, propagate the partial wake event through the constraint graph.
+		{
+			bResetEdgeFlags = PropagatePartialIslandWake(NodeQueue, ParticlesToWakeUp, NumNodesConstraintWake, /*bTriggerNodesAlwaysPropagate=*/ false);
+		}
+
+		// Phase 2 - Waking triggered by particle state changes
+		// Step 2a - Find trigger particles (i.e. particles which were just woken up the user) 
+		NodeQueue.Reset();
+		for (FPBDIslandParticle* Node : Island->Nodes)
+		{
+			if (Node->Flags.bWasUpdated)
+			{
+				ProcessTriggerParticle(Node, NodeQueue, /*bAlwaysTrigger=*/ true);
+			}
+		}
+		const int32 NumNodesNodeWake = NodeQueue.Num();
+
+		// Step 2b - Propagate wake event
+		// Partial island wake disabled, but at least one trigger particle found? Then, perform a full island wake.
+		if (!CVars::bChaosSolverPartialIslandWake && NumNodesNodeWake > 0)
+		{
+			PropagateIslandSleep(Island);
+		}
+		else // Otherwise, propagate the partial wake event through the constraint graph.
+		{
+			// Reset edge flags if we previously processed partial wake events.
+			if (bResetEdgeFlags && !NodeQueue.IsEmpty())
+			{
+				ResetEdgeFlags(Island);
+			}
+			PropagatePartialIslandWake(NodeQueue, ParticlesToWakeUp, NumNodesNodeWake, /*bTriggerNodesAlwaysPropagate=*/ true);
+		}
+
+		// Phase 3 - Wake up all relevant particles and reset sleep counter
 		WakeUpParticles(ParticlesToWakeUp);
 	}
 
-	// Checks both particles connected to the edge and adds the sleeping particle to ParticlesToWake if
-	//	   (bAlwaysWakeUp == true): if both particles are dynamic and one particle is awake while the other one is sleeping
-	//     (bAlwaysWakeUp == false): if satisfying the conditions above and the awake particle is also moving
-	bool FPBDIslandManager::ProcessConnectedParticles(const FPBDIslandConstraint* Edge, TArray<FGeometryParticleHandle*>& ParticlesToWake, const bool bAlwaysWakeUp)
+	bool FPBDIslandManager::PropagatePartialIslandWake(TArray<FPBDIslandParticle*>& NodeQueue,
+		TArray<FGeometryParticleHandle*>& ParticlesToWakeUp, const int32 NumNodesInQueue, const bool bTriggerNodesAlwaysPropagate)
 	{
-		FPBDIslandParticle* Node0 = Edge->Nodes[0];
-		FPBDIslandParticle* Node1 = Edge->Nodes[1];
+		bool bResetEdgeFlags = false;
+		for (int32 NodeIndex = 0; NodeIndex < NodeQueue.Num(); ++NodeIndex)
+		{
+			bResetEdgeFlags = true;
+			FPBDIslandParticle* Node = NodeQueue[NodeIndex];
+
+			// Traverse all edges connected to a node in the wake-up queue.
+			for (FPBDIslandConstraint* Edge : Node->Edges)
+			{
+				if (!Edge->Flags.bWasProcessed)
+				{
+					Edge->Flags.bWasProcessed = true;
+
+					FPBDIslandParticle* OtherNode = (Node == Edge->Nodes[0]) ? Edge->Nodes[1] : Edge->Nodes[0];
+					if ((OtherNode != nullptr) && OtherNode->Flags.bIsDynamic)
+					{
+						// Wake and propagate always to the next constraint if
+						// 1) We process one of the trigger nodes and propagate through a new constraint (if bTriggerNodesAlwaysPropagate == false).
+						// 2) We propagate upward (away from kinematic) and bChaosSolverLevelBasedWake is enabled.
+						const bool bIsTriggerNode = NodeIndex < NumNodesInQueue;
+						const bool bPropagateThrough = bIsTriggerNode && (Edge->Flags.bIsNew || bTriggerNodesAlwaysPropagate);
+						const bool bIsUpwardPropagation = OtherNode->Level >= Edge->Level;	// Propagation upward through the levels, i.e. away from closest kinematic
+						if (bPropagateThrough || bIsUpwardPropagation)
+						{
+							ProcessParticleWake(OtherNode, NodeQueue[NodeIndex], NodeQueue, ParticlesToWakeUp, /*bAlwaysPropagate=*/ true);
+						}
+					}
+				}
+			}
+		}
+		return bResetEdgeFlags;
+	}
 
+	bool FPBDIslandManager::ProcessTriggerParticle(FPBDIslandParticle* Node, TArray<FPBDIslandParticle*>& NodeQueue, const bool bAlwaysTrigger)
+	{
 		const auto IsMoving = [](FPBDIslandParticle* Node) -> bool
 			{
 				if (Node == nullptr)
@@ -2664,26 +2751,19 @@ namespace Chaos::Private
 				}
 
 				// Use the smoothed velocity to check for motion, unless it is zero (i.e. was never set).
-				const FReal VSmoothed = Rigid->VSmooth().SizeSquared();
-				const FReal WSmoothed = Rigid->WSmooth().SizeSquared();
-				const bool bIsVelSmoothUnset = (VSmoothed == 0) && (WSmoothed == 0);
-				if (bIsVelSmoothUnset)
+				FVec3f V = Rigid->VSmooth();
+				FVec3f W = Rigid->WSmooth();
+				if (V.SizeSquared() == 0 && W.SizeSquared() == 0)
 				{
 					// Use the regular velocity from the last time step
-					if ((Rigid->GetPreV().SizeSquared() > Node->SleepLinearThresholdSq)
-						|| (Rigid->GetPreW().SizeSquared() > Node->SleepAngularThresholdSq))
-					{
-						return true;
-					}
+					V = Rigid->GetPreV();
+					W = Rigid->GetPreW();
 				}
-				else
+
+				if ((V.SizeSquared() > Node->WakeLinearThresholdSq)
+					|| (W.SizeSquared() > Node->WakeAngularThresholdSq))
 				{
-					// Use the smooth velocities from the last step
-					if ((VSmoothed > Node->SleepLinearThresholdSq)
-						|| (WSmoothed > Node->SleepAngularThresholdSq))
-					{
-						return true;
-					}
+					return true;
 				}
 
 				return false;
@@ -2699,28 +2779,81 @@ namespace Chaos::Private
 				return Node && !!Node->Flags.bIsDynamic;
 			};
 
-		// If one particle is awake and the other one is sleeping, wake up the other dynamic particle.
-		const bool bIsAwake0 = IsAwake(Node0);
-		const bool bIsAwake1 = IsAwake(Node1);
-		const bool bIsDynamic0 = IsDynamic(Node0);
-		const bool bIsDynamic1 = IsDynamic(Node1);
-		const bool bTriggerWakeUp0 = bAlwaysWakeUp || IsMoving(Node0);
-		const bool bTriggerWakeUp1 = bAlwaysWakeUp || IsMoving(Node1);
-		if (bIsDynamic1 && !bIsAwake1 && bIsAwake0 && bTriggerWakeUp0)
-		{
-			Node1->Flags.bIsSleeping = false;
-			ParticlesToWake.Add(Node1->GetParticle());
-			return true;
+		// Stop wake-up event if the node or particle is empty
+		if (Node == nullptr)
+		{
+			return false;
 		}
-		else if (bIsDynamic0 && !bIsAwake0 && bIsAwake1 && bTriggerWakeUp1)
+		FPBDRigidParticleHandle* Rigid = Node->GetParticle()->CastToRigidParticle();
+		if (Rigid == nullptr)
 		{
-			Node0->Flags.bIsSleeping = false;
-			ParticlesToWake.Add(Node0->GetParticle());
-			return true;
+			return false;
 		}
+
+		const bool bIsAwake = IsAwake(Node);
+		const bool bIsDynamic = IsDynamic(Node);
+
+		// Stop wake-up event if the particle is not dynamic
+		if (!bIsDynamic)
+		{
+			return false;
+		}
+
+		// The particle must be awake to trigger a partial wake-up event.
+		if (bIsAwake)
+		{
+			if (bAlwaysTrigger || IsMoving(Node))
+			{
+				NodeQueue.Add(Node);
+				return true;
+			}
+		}
+
 		return false;
 	}
 
+	bool FPBDIslandManager::ProcessParticleWake(FPBDIslandParticle* Node, const FPBDIslandParticle* OtherNode, TArray<FPBDIslandParticle*>& NodeQueue, TArray<FGeometryParticleHandle*>& WakeUpQueue, const bool bAlwaysPropagate)
+	{
+		const auto IsAwake = [](const FPBDIslandParticle* Node) -> bool
+			{
+				return Node && !Node->Flags.bIsSleeping && !!Node->Flags.bIsDynamic;
+			};
+
+		const auto IsDynamic = [](const FPBDIslandParticle* Node) -> bool
+			{
+				return Node && !!Node->Flags.bIsDynamic;
+			};
+
+		// Stop wake-up event if the node or particle is empty
+		if (Node == nullptr)
+		{
+			return false;
+		}
+		FPBDRigidParticleHandle* Rigid = Node->GetParticle()->CastToRigidParticle();
+		if (Rigid == nullptr)
+		{
+			return false;
+		}
+
+		const bool bIsAwake = IsAwake(Node);
+		const bool bIsDynamic = IsDynamic(Node);
+
+		// Stop wake-up event if the particle is not dynamic
+		if (!bIsDynamic)
+		{
+			return false;
+		}
+
+		// Otherwise, the particle is added to the node queue and woken up if needed.
+		NodeQueue.Add(Node);
+		if (!bIsAwake)
+		{
+			Node->Flags.bIsSleeping = false;
+			WakeUpQueue.Add(Rigid);
+		}
+		return true;
+	}
+
 	void FPBDIslandManager::SyncConstraintSleepStateToParticles(FPBDIsland* Island)
 	{
 		// Iterate over the edges
@@ -2776,15 +2909,11 @@ namespace Chaos::Private
 						Constraint->SetIsSleeping(false);
 					}
 				}
-
-				// Reset
-				IslandEdge->Flags.bIsNew = false;
-				IslandEdge->Flags.bWasProcessed = false;
 			}
 		}
 	}
 
-	void FPBDIslandManager::WakeUpParticles(TArray<FGeometryParticleHandle*>& ParticlesToWake)
+	void FPBDIslandManager::WakeUpParticles(TArray<FGeometryParticleHandle*>& ParticlesToWake, const bool bResetSleepCounter)
 	{
 		if (!ParticlesToWake.IsEmpty())
 		{
@@ -2800,11 +2929,14 @@ namespace Chaos::Private
 				};
 
 			Particles.ActivateParticles(ParticlesToWake);
-			for (FGeometryParticleHandle* Particle : ParticlesToWake)
+			if (bResetSleepCounter)
 			{
-				if (FPBDRigidParticleHandle* Rigid = Particle->CastToRigidParticle())
+				for (FGeometryParticleHandle* Particle : ParticlesToWake)
 				{
-					Rigid->SetSleepCounter(0);
+					if (FPBDRigidParticleHandle* Rigid = Particle->CastToRigidParticle())
+					{
+						Rigid->SetSleepCounter(0);
+					}
 				}
 			}
 			UpdateShapeTransforms(ParticlesToWake);
@@ -2823,6 +2955,15 @@ namespace Chaos::Private
 		}
 	}
 
+	void FPBDIslandManager::ResetNodeFlags(FPBDIsland* Island)
+	{
+		for (FPBDIslandParticle* Node : Island->Nodes)
+		{
+			Node->Flags.bWasUpdated = false;
+			Node->Flags.bWasProcessed = false;
+		}
+	}
+
 	void FPBDIslandManager::ProcessDisable(TFunctionRef<void(FPBDRigidParticleHandle*)> ParticleDisableFunctor)
 	{
 		// @todo(chaos): parallelize
diff --git a/Engine/Source/Runtime/Experimental/Chaos/Public/Chaos/Island/IslandManager.h b/Engine/Source/Runtime/Experimental/Chaos/Public/Chaos/Island/IslandManager.h
--- a/Engine/Source/Runtime/Experimental/Chaos/Public/Chaos/Island/IslandManager.h
+++ b/Engine/Source/Runtime/Experimental/Chaos/Public/Chaos/Island/IslandManager.h
@@ -85,6 +85,8 @@ namespace Chaos::Private
 				uint32 bIsSleeping : 1;
 				uint32 bIsMoving : 1;
 				uint32 bNeedsResim : 1;
+				uint32 bWasUpdated : 1;
+				uint32 bWasProcessed : 1;
 			};
 
 		private:
@@ -120,6 +122,10 @@ namespace Chaos::Private
 		FRealSingle SleepAngularThresholdSq = 0;
 		int32 SleepCounterThreshold = 0;
 
+		// Partial wake thresholds
+		FRealSingle WakeLinearThresholdSq = 0;
+		FRealSingle WakeAngularThresholdSq = 0;
+
 		// Disable thresholds
 		FRealSingle DisableLinearThresholdSq = 0;
 		FRealSingle DisableAngularThresholdSq = 0;
@@ -420,7 +426,7 @@ namespace Chaos::Private
 	* @todo(chaos): the sleep and disable logic feel like they don't belong in here. It requires
 	* that we know about materials, and the Disabled Particle concept sits at a higher level
 	* (in the Evolution, and in principle sleep does too since other system might need to know).
-	* The main reason it is here is so that we can cache the thresholds and don't have to check 
+	* The main reason it is here is so that we can cache the thresholds and don't have to check
 	* the materials for every particle every tick.
 	*/
 	class FPBDIslandManager
@@ -598,11 +604,14 @@ namespace Chaos::Private
 		CHAOS_API void PropagateParticleSleepToConstraints(FPBDIsland* Island);
 
 		// Partial Island Waking
-		CHAOS_API void ProcessParticleWakes(FPBDIsland* Island);
-		CHAOS_API bool ProcessConnectedParticles(const FPBDIslandConstraint* Edge, TArray<FGeometryParticleHandle*>& ParticlesToWake, const bool bAlwaysWakeUp);
+		CHAOS_API void ProcessPartialIslandWake(FPBDIsland* Island);
+		CHAOS_API bool PropagatePartialIslandWake(TArray<FPBDIslandParticle*>& NodeQueue, TArray<FGeometryParticleHandle*>& ParticlesToWakeUp, const int32 NumNodesConstraintWake, const bool bTriggerNodesAlwaysPropagate);
+		CHAOS_API bool ProcessTriggerParticle(FPBDIslandParticle* Node, TArray<FPBDIslandParticle*>& NodeQueue, const bool bAlwaysTrigger);
+		CHAOS_API bool ProcessParticleWake(FPBDIslandParticle* Node, const FPBDIslandParticle* OtherNode, TArray<FPBDIslandParticle*>& NodeQueue, TArray<FGeometryParticleHandle*>& WakeUpQueue, const bool bAlwaysPropagate);
 		CHAOS_API void SyncConstraintSleepStateToParticles(FPBDIsland* Island);
-		CHAOS_API void WakeUpParticles(TArray<FGeometryParticleHandle*>& ParticlesToWake);
+		CHAOS_API void WakeUpParticles(TArray<FGeometryParticleHandle*>& ParticlesToWake, const bool bResetSleepCounter = true);
 		CHAOS_API void ResetEdgeFlags(FPBDIsland* Island);
+		CHAOS_API void ResetNodeFlags(FPBDIsland* Island);
 
 		// Disabling
 		CHAOS_API void ProcessDisable(TFunctionRef<void(FPBDRigidParticleHandle*)> ParticleDisableFunctor);
@@ -694,7 +703,7 @@ namespace Chaos::Private
 	template<typename VisitorType>
 	void FPBDIslandManager::VisitAwakeConstraints(const int32 ContainerId, const VisitorType& Visitor)
 	{
-		const_cast<FPBDIslandManager*>(this)->VisitAwakeConstConstraints(ContainerId, 
+		const_cast<FPBDIslandManager*>(this)->VisitAwakeConstConstraints(ContainerId,
 			[&Visitor](const FPBDIslandConstraint* Edge)
 			{
 				Visitor(const_cast<FPBDIslandConstraint*>(Edge));
