Should not have checked in with UE_ONDEMANDINSTALLCACHE_EXCLUSIVE_WRITE enabled.

diff --git a/Engine/Source/Runtime/Experimental/IoStore/OnDemand/Private/OnDemandInstallCache.cpp b/Engine/Source/Runtime/Experimental/IoStore/OnDemand/Private/OnDemandInstallCache.cpp
--- a/Engine/Source/Runtime/Experimental/IoStore/OnDemand/Private/OnDemandInstallCache.cpp
+++ b/Engine/Source/Runtime/Experimental/IoStore/OnDemand/Private/OnDemandInstallCache.cpp
@@ -35,7 +35,7 @@
 #endif
 
 #ifndef UE_ONDEMANDINSTALLCACHE_EXCLUSIVE_WRITE
-#define UE_ONDEMANDINSTALLCACHE_EXCLUSIVE_WRITE (1) // TODO: do not checkin
+#define UE_ONDEMANDINSTALLCACHE_EXCLUSIVE_WRITE (0)
 #endif
 
 #if UE_ONDEMANDINSTALLCACHE_EXCLUSIVE_WRITE
@@ -2434,7 +2434,11 @@ TSharedPtr<IOnDemandInstallCache> MakeOnDemandInstallCache(
 }
 
 ///////////////////////////////////////////////////////////////////////////////
-#if defined(WITH_IOSTORE_ONDEMAND_TESTS)
+#ifndef WITH_IOSTORE_ONDEMAND_TESTS
+#define WITH_IOSTORE_ONDEMAND_TESTS 0
+#endif
+
+#if WITH_IOSTORE_ONDEMAND_TESTS
 
 class FTmpDirectoryScope
 {
