Fix build - WITH_IOSTORE_ONDEMAND_TESTS should always be defined per the coding standard

diff --git a/Engine/Source/Runtime/Experimental/IoStore/OnDemand/Private/OnDemandInstallCache.cpp b/Engine/Source/Runtime/Experimental/IoStore/OnDemand/Private/OnDemandInstallCache.cpp
--- a/Engine/Source/Runtime/Experimental/IoStore/OnDemand/Private/OnDemandInstallCache.cpp
+++ b/Engine/Source/Runtime/Experimental/IoStore/OnDemand/Private/OnDemandInstallCache.cpp
@@ -27,7 +27,7 @@
 #include "Tasks/Task.h"
 #include "ProfilingDebugging/IoStoreTrace.h"
 
-#if defined(WITH_IOSTORE_ONDEMAND_TESTS)
+#if WITH_IOSTORE_ONDEMAND_TESTS
 #include "Algo/Find.h"
 #include "TestHarness.h"
 #include "TestMacros/Assertions.h"
@@ -2434,10 +2434,6 @@ TSharedPtr<IOnDemandInstallCache> MakeOnDemandInstallCache(
 }
 
 ///////////////////////////////////////////////////////////////////////////////
-#ifndef WITH_IOSTORE_ONDEMAND_TESTS
-#define WITH_IOSTORE_ONDEMAND_TESTS 0
-#endif
-
 #if WITH_IOSTORE_ONDEMAND_TESTS
 
 class FTmpDirectoryScope
diff --git a/Engine/Source/Runtime/Experimental/IoStore/OnDemand/Private/OnDemandIoDispatcherBackend.cpp b/Engine/Source/Runtime/Experimental/IoStore/OnDemand/Private/OnDemandIoDispatcherBackend.cpp
--- a/Engine/Source/Runtime/Experimental/IoStore/OnDemand/Private/OnDemandIoDispatcherBackend.cpp
+++ b/Engine/Source/Runtime/Experimental/IoStore/OnDemand/Private/OnDemandIoDispatcherBackend.cpp
@@ -44,7 +44,7 @@
 
 #include <atomic>
 
-#if defined(WITH_IOSTORE_ONDEMAND_TESTS)
+#if WITH_IOSTORE_ONDEMAND_TESTS
 #include "TestHarness.h"
 #include "TestMacros/Assertions.h"
 #include <catch2/generators/catch_generators.hpp>
@@ -1790,7 +1790,7 @@ TSharedPtr<IOnDemandIoDispatcherBackend> MakeOnDemandIoDispatcherBackend(
 }
 
 ///////////////////////////////////////////////////////////////////////////////
-#if defined(WITH_IOSTORE_ONDEMAND_TESTS)
+#if WITH_IOSTORE_ONDEMAND_TESTS
 
 TEST_CASE("IoStore::OnDemand::Ias::Misc", "[IoStoreOnDemand][Ias]")
 {
diff --git a/Engine/Source/Runtime/Experimental/IoStore/OnDemandCore/Internal/IO/IoStoreOnDemandInternals.h b/Engine/Source/Runtime/Experimental/IoStore/OnDemandCore/Internal/IO/IoStoreOnDemandInternals.h
--- a/Engine/Source/Runtime/Experimental/IoStore/OnDemandCore/Internal/IO/IoStoreOnDemandInternals.h
+++ b/Engine/Source/Runtime/Experimental/IoStore/OnDemandCore/Internal/IO/IoStoreOnDemandInternals.h
@@ -5,6 +5,10 @@
 #include "IO/IoStoreOnDemand.h"
 #include "HAL/PlatformTime.h"
 
+#ifndef WITH_IOSTORE_ONDEMAND_TESTS
+#define WITH_IOSTORE_ONDEMAND_TESTS 0
+#endif
+
 #define UE_API IOSTOREONDEMANDCORE_API
 
 namespace UE::IoStore
