Extended mount options with host group name for on-demand containers.

#rnx

diff --git a/Engine/Source/Runtime/Core/Private/IO/IoDispatcher.cpp b/Engine/Source/Runtime/Core/Private/IO/IoDispatcher.cpp
--- a/Engine/Source/Runtime/Core/Private/IO/IoDispatcher.cpp
+++ b/Engine/Source/Runtime/Core/Private/IO/IoDispatcher.cpp
@@ -57,7 +57,8 @@ static const TCHAR* const GetIoErrorText_ErrorCodeTextArray[] =
 	TEXT("Pending Fork"),
 	TEXT("Pending Encryption Key"),
 	TEXT("Disabled"),
-	TEXT("Not Installed")
+	TEXT("Not Installed"),
+	TEXT("Pending Host Group")
 };
 
 const TCHAR* const* GetIoErrorText_ErrorCodeText = GetIoErrorText_ErrorCodeTextArray;
diff --git a/Engine/Source/Runtime/Core/Public/IO/IoStatus.h b/Engine/Source/Runtime/Core/Public/IO/IoStatus.h
--- a/Engine/Source/Runtime/Core/Public/IO/IoStatus.h
+++ b/Engine/Source/Runtime/Core/Public/IO/IoStatus.h
@@ -33,7 +33,8 @@ enum class EIoErrorCode
 	PendingFork,
 	PendingEncryptionKey,
 	Disabled,
-	NotInstalled
+	NotInstalled,
+	PendingHostGroup
 };
 
 /** Get I/O error code description. */
diff --git a/Engine/Source/Runtime/Experimental/IoStore/OnDemand/Private/OnDemandIoStore.cpp b/Engine/Source/Runtime/Experimental/IoStore/OnDemand/Private/OnDemandIoStore.cpp
--- a/Engine/Source/Runtime/Experimental/IoStore/OnDemand/Private/OnDemandIoStore.cpp
+++ b/Engine/Source/Runtime/Experimental/IoStore/OnDemand/Private/OnDemandIoStore.cpp
@@ -1255,7 +1255,7 @@ FIoStatus FOnDemandIoStore::TickMountRequest(FMountRequest& MountRequest)
 		return FIoStatusBuilder(EIoErrorCode::InvalidParameter) << TEXT("Invalid mount ID");
 	}
 
-	bool bAnyPendingEncryptionKey = false;
+	EIoErrorCode MountResult = EIoErrorCode::Ok;
 	TArray<FString> ExistingContainers;
 
 	// Find containers matching the mount ID
@@ -1272,7 +1272,12 @@ FIoStatus FOnDemandIoStore::TickMountRequest(FMountRequest& MountRequest)
 			
 			if (EnumHasAnyFlags(Container->Flags, EOnDemandContainerFlags::PendingEncryptionKey))
 			{
-				bAnyPendingEncryptionKey = true;
+				MountResult = EIoErrorCode::PendingEncryptionKey;
+			}
+
+			if (EnumHasAnyFlags(Container->Flags, EOnDemandContainerFlags::PendingHostGroup))
+			{
+				MountResult = EIoErrorCode::PendingHostGroup;
 			}
 		}
 	}
@@ -1283,7 +1288,7 @@ FIoStatus FOnDemandIoStore::TickMountRequest(FMountRequest& MountRequest)
 
 	if (Args.Toc)
 	{
-		CreateContainersFromToc(Args.MountId, TocPath, *Args.Toc, Args.HostGroup, ExistingContainers, RequestedContainers);
+		CreateContainersFromToc(Args.MountId, TocPath, *Args.Toc, ExistingContainers, RequestedContainers);
 	}
 	else if (Args.FilePath.IsEmpty() == false)
 	{
@@ -1299,7 +1304,7 @@ FIoStatus FOnDemandIoStore::TickMountRequest(FMountRequest& MountRequest)
 
 		Args.Toc = MakeUnique<FOnDemandToc>(TocStatus.ConsumeValueOrDie());
 
-		CreateContainersFromToc(Args.MountId, TocPath, *Args.Toc, Args.HostGroup, ExistingContainers, RequestedContainers);
+		CreateContainersFromToc(Args.MountId, TocPath, *Args.Toc, ExistingContainers, RequestedContainers);
 	}
 	else if (Args.HostGroup.IsEmpty() == false && Args.TocRelativeUrl.IsEmpty() == false)
 	{
@@ -1337,7 +1342,7 @@ FIoStatus FOnDemandIoStore::TickMountRequest(FMountRequest& MountRequest)
 			return Status;
 		}
 
-		CreateContainersFromToc(Args.MountId, TocPath, *Args.Toc, Args.HostGroup, ExistingContainers, RequestedContainers);
+		CreateContainersFromToc(Args.MountId, TocPath, *Args.Toc, ExistingContainers, RequestedContainers);
 	}
 	
 	TAnsiStringBuilder<512> ChunkUrlBuilder;
@@ -1353,7 +1358,9 @@ FIoStatus FOnDemandIoStore::TickMountRequest(FMountRequest& MountRequest)
 				EnumAddFlags(Container->Flags, EOnDemandContainerFlags::StreamOnDemand);
 			}
 
-			Container->IasHostGroup = FHostGroupManager::Get().Find(FName("DefaultOnDemand"));
+			//TODO: Remove this special case for streaming containers
+			Container->HostGroupName	= FName("DefaultOnDemand");
+			Container->HostGroup		= FHostGroupManager::Get().Find(Container->HostGroupName);
 		}
 		else if (EnumHasAnyFlags(Args.Options, EOnDemandMountOptions::InstallOnDemand))
 		{
@@ -1368,16 +1375,50 @@ FIoStatus FOnDemandIoStore::TickMountRequest(FMountRequest& MountRequest)
 				EnumAddFlags(Container->Flags, EOnDemandContainerFlags::WithSoftReferences);
 			}
 
-			FIASHostGroup HostGroup = FHostGroupManager::Get().Find(NAME_None);
-			if (HostGroup.IsResolved() == false)
+			UE_CLOG(Container->HostGroupName != Args.HostGroupName, LogIoStoreOnDemand, Log, TEXT("Overriding container host group '%s' with '%s' from mount option"),
+				*Container->HostGroupName.ToString(), *Args.HostGroupName.ToString());
+
+			Container->HostGroupName	= Args.HostGroupName;
+			Container->HostGroup		= FHostGroupManager::Get().Find(Container->HostGroupName);
+
+			if (Container->HostGroup.IsResolved() == false && Args.HostGroup.IsEmpty() == false)
 			{
-				TIoStatusOr<FIASHostGroup> NewGroup = FHostGroupManager::Get().Register(NAME_None, Container->HostGroup.Hosts());
-				if (NewGroup.IsOk())
+				TIoStatusOr<FIASHostGroup> NewHostGroup = FHostGroupManager::Get().Register(Container->HostGroupName, Args.HostGroup.Hosts());
+				if (NewHostGroup.IsOk())
+				{
+					Container->HostGroup = NewHostGroup.ConsumeValueOrDie();
+					{
+						const FOnDemandHostGroup& HostGroup = Container->HostGroup.GetUnderlyingHostGroup();
+						TConstArrayView<FAnsiString> Hosts	= HostGroup.Hosts();
+						TStringBuilder<128> Sb;
+
+						UE_LOG(LogIoStoreOnDemand, Log, TEXT("Created new host group '%s'"), *Container->HostGroupName.ToString());
+						for (int32 HostIndex = 0; HostIndex < Hosts.Num(); ++HostIndex) 
+						{
+							Sb.Reset();
+							Sb << Hosts[HostIndex];
+							if (HostIndex == HostGroup.PrimaryHostIndex())
+							{
+								Sb << TEXT(" (primary)");
+							}
+							UE_LOG(LogIoStoreOnDemand, Log, TEXT("\t%s"), Sb.ToString()); 
+						}
+					}
+				}
+				else
 				{
-					HostGroup = NewGroup.ConsumeValueOrDie();
+					return NewHostGroup.Status();
 				}
 			}
-			Container->IasHostGroup = HostGroup; 
+
+			if (Container->HostGroup.IsResolved() == false)
+			{
+				//TODO: Handle new host group being registered
+				EnumAddFlags(Container->Flags, EOnDemandContainerFlags::PendingHostGroup);
+				MountResult = EIoErrorCode::PendingHostGroup;
+				UE_LOG(LogIoStoreOnDemand, Log, TEXT("Deferring container '%s' until host group '%s' becomes available"),
+					*Container->Name, *Container->HostGroupName.ToString());
+			}
 		}
 
 		if (EnumHasAnyFlags(Container->Flags, EOnDemandContainerFlags::Encrypted) &&
@@ -1388,10 +1429,9 @@ FIoStatus FOnDemandIoStore::TickMountRequest(FMountRequest& MountRequest)
 			if (FEncryptionKeyManager::Get().TryGetKey(KeyGuid, Container->EncryptionKey) == false)
 			{
 				EnumAddFlags(Container->Flags, EOnDemandContainerFlags::PendingEncryptionKey);
+				MountResult = EIoErrorCode::PendingEncryptionKey;
 				UE_LOG(LogIoStoreOnDemand, Log, TEXT("Deferring container '%s' until encryption key '%s' becomes available"),
 					*Container->Name, *Container->EncryptionKeyGuid);
-
-				bAnyPendingEncryptionKey = true;
 			}
 		}
 
@@ -1416,7 +1456,12 @@ FIoStatus FOnDemandIoStore::TickMountRequest(FMountRequest& MountRequest)
 					}
 					else
 					{
-						if (Container->HostGroup.IsEmpty())
+						if (EnumHasAnyFlags(Container->Flags, EOnDemandContainerFlags::PendingHostGroup))
+						{
+							continue;
+						}
+
+						if (Container->HostGroup.IsResolved() == false)
 						{
 							return FIoStatusBuilder(EIoErrorCode::InvalidParameter) << TEXT("No valid endpoint(s)");
 						}
@@ -1424,9 +1469,10 @@ FIoStatus FOnDemandIoStore::TickMountRequest(FMountRequest& MountRequest)
 						UE_LOG(LogIoStoreOnDemand, Warning, TEXT("Fetching container header, ContainerName='%s', ChunkId='%s'"),
 								*Container->Name, *LexToString(ChunkId));
 
+						const FAnsiStringView PrimaryHost = Container->HostGroup.GetUnderlyingHostGroup().PrimaryHost();
 						const uint32 RetryCount = 2;
 						TIoStatusOr<FMultiEndpointHttpClientResponse> Response = FMultiEndpointHttpClient::Get(
-								Private::GetChunkUrl(Args.HostGroup.PrimaryHost(), *Container, ChunkInfo.ChunkEntry(), ChunkUrlBuilder).ToView(),
+								Private::GetChunkUrl(PrimaryHost, *Container, ChunkInfo.ChunkEntry(), ChunkUrlBuilder).ToView(),
 								FMultiEndpointHttpClientConfig{.MaxRetryCount = RetryCount});
 
 						if (Response.IsOk() == false)
@@ -1457,10 +1503,15 @@ FIoStatus FOnDemandIoStore::TickMountRequest(FMountRequest& MountRequest)
 			}
 		}
 
+		if (EnumHasAnyFlags(Container->Flags, EOnDemandContainerFlags::PendingEncryptionKey | EOnDemandContainerFlags::PendingHostGroup))
+		{
+			continue;
+		}
+
 		TStringBuilder<128> Sb;
 		Sb << Container->Flags;
-		UE_LOG(LogIoStoreOnDemand, Log, TEXT("Mounting container '%s', Entries=%d, Flags='%s'"),
-			*Container->Name, Container->ChunkEntries.Num(), Sb.ToString());
+		UE_LOG(LogIoStoreOnDemand, Log, TEXT("Mounting container '%s', Entries=%d, Flags='%s', HostGroup='%s'"),
+			*Container->Name, Container->ChunkEntries.Num(), Sb.ToString(), *Container->HostGroupName.ToString());
 		EnumAddFlags(Container->Flags, EOnDemandContainerFlags::Mounted);
 		Container->EncryptionKeyGuid.Empty();
 	}
@@ -1473,7 +1524,7 @@ FIoStatus FOnDemandIoStore::TickMountRequest(FMountRequest& MountRequest)
 
 	PackageStoreBackend->NeedsUpdate();
 
-	return bAnyPendingEncryptionKey ? EIoErrorCode::PendingEncryptionKey : EIoErrorCode::Ok;
+	return MountResult; 
 }
 
 void FOnDemandIoStore::CompleteMountRequest(FMountRequest& Request, FOnDemandMountResult&& MountResult)
@@ -1719,7 +1770,7 @@ FOnDemandInstallResult FOnDemandIoStore::TickInstallRequest(FOnDemandInternalIns
 				++ConcurrentRequests;
 
 				HttpClient->Get(
-					Container.HostGroup,
+					Container.HostGroup.GetUnderlyingHostGroup(),
 					Private::GetChunkUrl(Container, ChunkEntry, ChunkUrl).ToView(),
 					[this, &InstallRequest, &OutResult, ChunkId, ChunkEntry, &ConcurrentRequests]
 					(FMultiEndpointHttpClientResponse&& HttpResponse)
@@ -2008,7 +2059,6 @@ void FOnDemandIoStore::CreateContainersFromToc(
 	FStringView MountId,
 	FStringView TocPath,
 	FOnDemandToc& Toc,
-	FOnDemandHostGroup HostGroup,
 	const TArray<FString>& ExistingContainers,
 	TArray<FSharedOnDemandContainer>& Out)
 {
@@ -2055,7 +2105,6 @@ void FOnDemandIoStore::CreateContainersFromToc(
 		}
 
 		FSharedOnDemandContainer Container = MakeShared<FOnDemandContainer>();
-		Container->HostGroup			= HostGroup;
 		Container->Name					= MoveTemp(ContainerEntry.ContainerName);
 		Container->MountId				= MountId;
 		Container->ChunksDirectory		= StringCast<ANSICHAR>(ChunksDirectory.GetData(), ChunksDirectory.Len());
@@ -2064,6 +2113,7 @@ void FOnDemandIoStore::CreateContainersFromToc(
 		Container->BlockSizes			= MoveTemp(ContainerEntry.BlockSizes);
 		Container->BlockHashes			= MoveTemp(ContainerEntry.BlockHashes);
 		Container->ContainerId			= ContainerEntry.ContainerId;
+		Container->HostGroupName		= Toc.Header.HostGroupName.IsEmpty() ? FOnDemandHostGroup::DefaultName : FName(Toc.Header.HostGroupName);
 
 		Container->CompressionFormats.Reserve(1);
 		Container->CompressionFormats.Add(CompressionFormat);
@@ -2397,7 +2447,8 @@ FStringBuilderBase& operator<<(FStringBuilderBase& Sb, UE::IoStore::EOnDemandCon
 		TEXT("StreamOnDemand"),
 		TEXT("InstallOnDemand"),
 		TEXT("Encrypted"),
-		TEXT("WithSoftReferences")
+		TEXT("WithSoftReferences"),
+		TEXT("PendingHostGroup")
 	};
 
 	if (Flags == EOnDemandContainerFlags::None)
diff --git a/Engine/Source/Runtime/Experimental/IoStore/OnDemand/Private/OnDemandIoStore.h b/Engine/Source/Runtime/Experimental/IoStore/OnDemand/Private/OnDemandIoStore.h
--- a/Engine/Source/Runtime/Experimental/IoStore/OnDemand/Private/OnDemandIoStore.h
+++ b/Engine/Source/Runtime/Experimental/IoStore/OnDemand/Private/OnDemandIoStore.h
@@ -45,7 +45,8 @@ enum class EOnDemandContainerFlags : uint8
 	InstallOnDemand			= (1 << 3),
 	Encrypted				= (1 << 4),
 	WithSoftReferences		= (1 << 5),
-	Last = WithSoftReferences
+	PendingHostGroup		= (1 << 6),
+	Last = PendingHostGroup
 };
 ENUM_CLASS_FLAGS(EOnDemandContainerFlags);
 
@@ -91,8 +92,7 @@ struct FOnDemandContainer
 
 	FAES::FAESKey							EncryptionKey;
 	FSharedContainerHeader					Header;
-	FOnDemandHostGroup						HostGroup;
-	FIASHostGroup							IasHostGroup;
+	FIASHostGroup							HostGroup;
 	FString									EncryptionKeyGuid;
 	FString									Name;
 	FString									MountId;
@@ -107,6 +107,7 @@ struct FOnDemandContainer
 	TArray<FOnDemandChunkEntryReferences>	ChunkEntryReferences;
 	FIoContainerId							ContainerId;
 	uint32									BlockSize = 0;
+	FName									HostGroupName = FOnDemandHostGroup::DefaultName;
 	EOnDemandContainerFlags 				Flags = EOnDemandContainerFlags::None;
 };
 
@@ -249,7 +250,7 @@ TConstArrayView<FIoBlockHash> FOnDemandChunkInfo::BlockHashes() const
 
 const FIASHostGroup& FOnDemandChunkInfo::HostGroup() const
 {
-	return SharedContainer->IasHostGroup;
+	return SharedContainer->HostGroup;
 }
 
 FOnDemandChunkInfo FOnDemandChunkInfo::Find(FSharedOnDemandContainer Container, const FIoChunkId& ChunkId)
@@ -357,7 +358,6 @@ private:
 											FStringView MountId,
 											FStringView TocPath,
 											FOnDemandToc& Toc,
-											FOnDemandHostGroup HostGroup,
 											const TArray<FString>& ExistingContainers,
 											TArray<FSharedOnDemandContainer>& Out);
 
diff --git a/Engine/Source/Runtime/Experimental/IoStore/OnDemandCore/Private/OnDemandHostGroup.cpp b/Engine/Source/Runtime/Experimental/IoStore/OnDemandCore/Private/OnDemandHostGroup.cpp
--- a/Engine/Source/Runtime/Experimental/IoStore/OnDemandCore/Private/OnDemandHostGroup.cpp
+++ b/Engine/Source/Runtime/Experimental/IoStore/OnDemandCore/Private/OnDemandHostGroup.cpp
@@ -170,4 +170,6 @@ TIoStatusOr<FOnDemandHostGroup> FOnDemandHostGroup::Create(TConstArrayView<FStri
 	return FOnDemandHostGroup(MoveTemp(Impl));
 }
 
+FName FOnDemandHostGroup::DefaultName = FName("Default");
+
 } //namespace UE::IoStore
diff --git a/Engine/Source/Runtime/Experimental/IoStore/OnDemandCore/Public/IO/IoStoreOnDemand.h b/Engine/Source/Runtime/Experimental/IoStore/OnDemandCore/Public/IO/IoStoreOnDemand.h
--- a/Engine/Source/Runtime/Experimental/IoStore/OnDemandCore/Public/IO/IoStoreOnDemand.h
+++ b/Engine/Source/Runtime/Experimental/IoStore/OnDemandCore/Public/IO/IoStoreOnDemand.h
@@ -148,6 +148,8 @@ struct FOnDemandMountArgs
 	FString TocRelativeUrl;
 	/** Serialize the TOC from the specified file path. */
 	FString FilePath;
+	/** Name of a new or existing host group. */
+	FName HostGroupName = FOnDemandHostGroup::DefaultName;
 	/** List of URLs from where to download the chunks. */
 	FOnDemandHostGroup HostGroup;
 	/** Mount options. */
diff --git a/Engine/Source/Runtime/Experimental/IoStore/OnDemandCore/Public/IO/OnDemandHostGroup.h b/Engine/Source/Runtime/Experimental/IoStore/OnDemandCore/Public/IO/OnDemandHostGroup.h
--- a/Engine/Source/Runtime/Experimental/IoStore/OnDemandCore/Public/IO/OnDemandHostGroup.h
+++ b/Engine/Source/Runtime/Experimental/IoStore/OnDemandCore/Public/IO/OnDemandHostGroup.h
@@ -46,6 +46,8 @@ public:
 	UE_API static TIoStatusOr<FOnDemandHostGroup> Create(TConstArrayView<FAnsiString> Urls);
 	/** Create a new host group with the specified URLs. */
 	UE_API static TIoStatusOr<FOnDemandHostGroup> Create(TConstArrayView<FString> Urls);
+	/** Default host group name. */
+	UE_API static FName DefaultName;
 
 	/** Returns whether two endpoint instance are the same or not. */
 	friend inline bool operator==(const FOnDemandHostGroup& Lhs, const FOnDemandHostGroup& Rhs)
