Print number of referenced bytes by content handle when purging the on-demand install cache.

#rnx

diff --git a/Engine/Source/Runtime/Experimental/IoStore/OnDemand/Private/OnDemandInstallCache.cpp b/Engine/Source/Runtime/Experimental/IoStore/OnDemand/Private/OnDemandInstallCache.cpp
--- a/Engine/Source/Runtime/Experimental/IoStore/OnDemand/Private/OnDemandInstallCache.cpp
+++ b/Engine/Source/Runtime/Experimental/IoStore/OnDemand/Private/OnDemandInstallCache.cpp
@@ -1902,6 +1902,36 @@ FIoStatus FOnDemandInstallCache::PurgeAllUnreferenced()
 	FCasBlockInfoMap	BlockInfo;
 	const uint64		TotalCachedBytes = Cas.GetBlockInfo(BlockInfo);
 
+	{
+		TMap<FOnDemandWeakContentHandle, TArray<FOnDemandContainerChunkEntryReferences>> RefsByHandle;
+		IoStore.GetReferencedContentByHandle(RefsByHandle);
+
+		for (const TPair<FOnDemandWeakContentHandle, TArray<FOnDemandContainerChunkEntryReferences>>& Kv : RefsByHandle)
+		{
+			const FOnDemandWeakContentHandle& WeakHandle = Kv.Key;
+			const TArray<FOnDemandContainerChunkEntryReferences>& ContainerReferences = Kv.Value;
+
+			uint64 ReferencedBytesByHandle = 0;
+			for (const FOnDemandContainerChunkEntryReferences& Refs : ContainerReferences)
+			{
+				for (int32 EntryIndex = 0; const FOnDemandChunkEntry& ChunkEntry : Refs.Container->ChunkEntries)
+				{
+					if (Refs.Indices[EntryIndex])
+					{
+						ReferencedBytesByHandle += ChunkEntry.EncodedSize;
+					}
+					EntryIndex++;
+				}
+			}
+
+			if (ReferencedBytesByHandle > 0)
+			{
+				UE_LOG(LogIoStoreOnDemand, Log, TEXT("HandleId=0x%llX, DebugName='%s', ReferencedBytes=%.2lf KiB"),
+					WeakHandle.HandleId, *WeakHandle.DebugName, ToKiB(ReferencedBytesByHandle));
+			}
+		}
+	}
+
 	TArray<FSharedOnDemandContainer>	Containers;
 	TArray<TBitArray<>>					ChunkEntryIndices;
 
diff --git a/Engine/Source/Runtime/Experimental/IoStore/OnDemand/Private/OnDemandIoStore.cpp b/Engine/Source/Runtime/Experimental/IoStore/OnDemand/Private/OnDemandIoStore.cpp
--- a/Engine/Source/Runtime/Experimental/IoStore/OnDemand/Private/OnDemandIoStore.cpp
+++ b/Engine/Source/Runtime/Experimental/IoStore/OnDemand/Private/OnDemandIoStore.cpp
@@ -2350,6 +2350,24 @@ void FOnDemandIoStore::GetReferencedContent(TArray<FSharedOnDemandContainer>& Ou
 	}
 }
 
+void FOnDemandIoStore::GetReferencedContentByHandle(TMap<FOnDemandWeakContentHandle, TArray<FOnDemandContainerChunkEntryReferences>>& OutReferencesByHandle)
+{
+	UE::TUniqueLock Lock(ContainerMutex);
+	for (FSharedOnDemandContainer& Container : Containers)
+	{
+		for (const FOnDemandChunkEntryReferences& Refs : Container->ChunkEntryReferences)
+		{
+			FOnDemandWeakContentHandle WeakHandle = FOnDemandWeakContentHandle::FromUnsafeHandle(Refs.ContentHandleId);
+			TArray<FOnDemandContainerChunkEntryReferences>& Entries = OutReferencesByHandle.FindOrAdd(WeakHandle);
+			Entries.Add(FOnDemandContainerChunkEntryReferences
+			{
+				.Container	= Container,
+				.Indices	= Refs.Indices
+			});
+		}
+	}
+}
+
 void FOnDemandIoStore::CancelInstallRequest(FSharedInternalInstallRequest InstallRequest)
 {
 	check(InstallRequest.IsValid());
diff --git a/Engine/Source/Runtime/Experimental/IoStore/OnDemand/Private/OnDemandIoStore.h b/Engine/Source/Runtime/Experimental/IoStore/OnDemand/Private/OnDemandIoStore.h
--- a/Engine/Source/Runtime/Experimental/IoStore/OnDemand/Private/OnDemandIoStore.h
+++ b/Engine/Source/Runtime/Experimental/IoStore/OnDemand/Private/OnDemandIoStore.h
@@ -74,10 +74,19 @@ struct FOnDemandTagSet
 ///////////////////////////////////////////////////////////////////////////////
 struct FOnDemandChunkEntryReferences
 {
-	UPTRINT ContentHandleId = 0;	
+	UPTRINT		ContentHandleId = 0;	
 	TBitArray<>	Indices;
 };
 
+using FSharedOnDemandContainer = TSharedPtr<struct FOnDemandContainer, ESPMode::ThreadSafe>;
+
+///////////////////////////////////////////////////////////////////////////////
+struct FOnDemandContainerChunkEntryReferences
+{
+	FSharedOnDemandContainer	Container;
+	TBitArray<>					Indices;
+};
+
 ///////////////////////////////////////////////////////////////////////////////
 struct FOnDemandContainer
 {
@@ -111,8 +120,6 @@ struct FOnDemandContainer
 	EOnDemandContainerFlags 				Flags = EOnDemandContainerFlags::None;
 };
 
-using FSharedOnDemandContainer = TSharedPtr<FOnDemandContainer, ESPMode::ThreadSafe>;
-
 int32 FOnDemandContainer::FindChunkEntryIndex(const FIoChunkId& ChunkId) const
 {
 	if (const int32 Index = Algo::LowerBound(ChunkIds, ChunkId); Index < ChunkIds.Num())
@@ -313,6 +320,7 @@ public:
 	TArray<FIoChunkId>					DebugFindStreamingChunkIds(int32 NumToFind);
 #endif //!UE_BUILD_SHIPPING
 	void								GetReferencedContent(TArray<FSharedOnDemandContainer>& OutContainers, TArray<TBitArray<>>& OutChunkEntryIndices, bool bPackageStore = false);
+	void								GetReferencedContentByHandle(TMap<FOnDemandWeakContentHandle, TArray<FOnDemandContainerChunkEntryReferences>>& OutReferencesByHandle);
 
 	// IOnDemandIoStore interface
 	virtual FIoStatus					Initialize() override;
diff --git a/Engine/Source/Runtime/Experimental/IoStore/OnDemandCore/Internal/IO/IoStoreOnDemandInternals.h b/Engine/Source/Runtime/Experimental/IoStore/OnDemandCore/Internal/IO/IoStoreOnDemandInternals.h
--- a/Engine/Source/Runtime/Experimental/IoStore/OnDemandCore/Internal/IO/IoStoreOnDemandInternals.h
+++ b/Engine/Source/Runtime/Experimental/IoStore/OnDemandCore/Internal/IO/IoStoreOnDemandInternals.h
@@ -14,6 +14,30 @@
 namespace UE::IoStore
 {
 
+///////////////////////////////////////////////////////////////////////////////
+struct FOnDemandWeakContentHandle
+{
+	inline bool operator==(const FOnDemandWeakContentHandle& Other) const
+	{
+		return HandleId == Other.HandleId;
+	}
+
+	inline bool operator!=(const FOnDemandWeakContentHandle& Other) const
+	{
+		return HandleId != Other.HandleId;
+	}
+
+	inline friend uint32 GetTypeHash(const FOnDemandWeakContentHandle& Handle)
+	{
+		return GetTypeHash(Handle.HandleId);
+	}
+
+	UE_API static FOnDemandWeakContentHandle FromUnsafeHandle(UPTRINT HandleId);
+
+	UPTRINT			HandleId = 0;
+	FSharedString	DebugName;
+};
+
 ///////////////////////////////////////////////////////////////////////////////
 class FOnDemandInternalContentHandle
 {
diff --git a/Engine/Source/Runtime/Experimental/IoStore/OnDemandCore/Private/IoStoreOnDemandInternals.cpp b/Engine/Source/Runtime/Experimental/IoStore/OnDemandCore/Private/IoStoreOnDemandInternals.cpp
--- a/Engine/Source/Runtime/Experimental/IoStore/OnDemandCore/Private/IoStoreOnDemandInternals.cpp
+++ b/Engine/Source/Runtime/Experimental/IoStore/OnDemandCore/Private/IoStoreOnDemandInternals.cpp
@@ -5,6 +5,17 @@
 namespace UE::IoStore
 {
 
+///////////////////////////////////////////////////////////////////////////////
+FOnDemandWeakContentHandle FOnDemandWeakContentHandle::FromUnsafeHandle(UPTRINT HandleId)
+{
+	FOnDemandInternalContentHandle& ContentHandle = *reinterpret_cast<FOnDemandInternalContentHandle*>(HandleId);
+	return FOnDemandWeakContentHandle
+	{
+		.HandleId	= HandleId,
+		.DebugName	= ContentHandle.DebugName
+	};
+}
+
 ///////////////////////////////////////////////////////////////////////////////
 FOnDemandInternalContentHandle::~FOnDemandInternalContentHandle()
 {
