add sphere covering output to the leaf convex hulls dataflow node

+ some helpers to transform a sphere covering -- using the min axis scale for the radius if scale is non-uniform, matching how physics collision transforms spheres

diff --git a/Engine/Plugins/Experimental/GeometryCollectionPlugin/Source/GeometryCollectionNodes/Private/Dataflow/GeometryCollectionUtilityNodes.cpp b/Engine/Plugins/Experimental/GeometryCollectionPlugin/Source/GeometryCollectionNodes/Private/Dataflow/GeometryCollectionUtilityNodes.cpp
--- a/Engine/Plugins/Experimental/GeometryCollectionPlugin/Source/GeometryCollectionNodes/Private/Dataflow/GeometryCollectionUtilityNodes.cpp
+++ b/Engine/Plugins/Experimental/GeometryCollectionPlugin/Source/GeometryCollectionNodes/Private/Dataflow/GeometryCollectionUtilityNodes.cpp
@@ -129,6 +129,7 @@ FCreateLeafConvexHullsDataflowNode::FCreateLeafConvexHullsDataflowNode(const UE:
 	RegisterInputConnection(&SimplificationDistanceThreshold).SetCanHidePin(true).SetPinIsHidden(true);
 	RegisterInputConnection(&ConvexDecompositionSettings).SetCanHidePin(true).SetPinIsHidden(true);
 	RegisterOutputConnection(&Collection);
+	RegisterOutputConnection(&SphereCovering);
 
 	DebugDrawRenderSettings.RenderType = EDataflowDebugDrawRenderType::Wireframe;
 	DebugDrawRenderSettings.Color = FLinearColor::Green;
@@ -137,12 +138,14 @@ FCreateLeafConvexHullsDataflowNode::FCreateLeafConvexHullsDataflowNode(const UE:
 
 void FCreateLeafConvexHullsDataflowNode::Evaluate(UE::Dataflow::FContext& Context, const FDataflowOutput* Out) const
 {
-	if (Out->IsA(&Collection))
+	if (Out->IsA(&Collection) || Out->IsA(&SphereCovering))
 	{
 		const FManagedArrayCollection& InCollection = GetValue(Context, &Collection);
+		FDataflowSphereCovering CombinedSphereCovering;
 		if (InCollection.NumElements(FGeometryCollection::TransformGroup) == 0)
 		{
 			SetValue(Context, InCollection, &Collection);
+			SetValue(Context, MoveTemp(CombinedSphereCovering), &SphereCovering);
 			return;
 		}
 
@@ -177,8 +180,16 @@ void FCreateLeafConvexHullsDataflowNode::Evaluate(UE::Dataflow::FContext& Contex
 			LeafSettings.DecompositionSettings.NegativeSpaceTolerance = InDecompSettings.NegativeSpaceTolerance;
 
 			LeafSettings.bComputeIntersectionsBeforeHull = bComputeIntersectionsBeforeHull;
-			FGeometryCollectionConvexUtility::GenerateLeafConvexHulls(*GeomCollection, bRestrictToSelection, SelectedBones, LeafSettings);
+			TArray<FGeometryCollectionConvexUtility::FSphereCoveringInfo> SphereCoverings;
+			FGeometryCollectionConvexUtility::GenerateLeafConvexHulls(*GeomCollection, bRestrictToSelection, SelectedBones, LeafSettings, &SphereCoverings);
 			SetValue(Context, (const FManagedArrayCollection&)(*GeomCollection), &Collection);
+			
+			for (FGeometryCollectionConvexUtility::FSphereCoveringInfo& Info : SphereCoverings)
+			{
+				CombinedSphereCovering.Spheres.AppendTransformed(Info.SphereCovering, Info.Transform);
+			}
+			
+			SetValue(Context, MoveTemp(CombinedSphereCovering), &SphereCovering);
 		}
 	}
 }
diff --git a/Engine/Plugins/Experimental/GeometryCollectionPlugin/Source/GeometryCollectionNodes/Public/Dataflow/GeometryCollectionUtilityNodes.h b/Engine/Plugins/Experimental/GeometryCollectionPlugin/Source/GeometryCollectionNodes/Public/Dataflow/GeometryCollectionUtilityNodes.h
--- a/Engine/Plugins/Experimental/GeometryCollectionPlugin/Source/GeometryCollectionNodes/Public/Dataflow/GeometryCollectionUtilityNodes.h
+++ b/Engine/Plugins/Experimental/GeometryCollectionPlugin/Source/GeometryCollectionNodes/Public/Dataflow/GeometryCollectionUtilityNodes.h
@@ -33,6 +33,18 @@ enum class EConvexOverlapRemovalMethodEnum : uint8
 	Dataflow_Max                UMETA(Hidden)
 };
 
+//~ Simple wrapper class to make the sphere covering data possible to pass via dataflow
+// A set of spheres generated to represent empty space when creating a minimal set of convex hulls, e.g. in one of the Generate Cluster Convex Hulls nodes
+USTRUCT()
+struct FDataflowSphereCovering
+{
+	GENERATED_USTRUCT_BODY()
+
+public:
+	UE::Geometry::FSphereCovering Spheres;
+};
+
+
 USTRUCT()
 struct FDataflowConvexDecompositionSettings
 {
@@ -153,6 +165,9 @@ private:
 	UPROPERTY(meta = (DataflowInput, DataflowOutput))
 	FManagedArrayCollection Collection;
 
+	UPROPERTY(meta = (DataflowOutput))
+	FDataflowSphereCovering SphereCovering;
+
 	/** Optional transform selection to compute leaf hulls on -- if not provided, all leaf hulls will be computed. */
 	UPROPERTY(meta = (DataflowInput))
 	FDataflowTransformSelection OptionalSelectionFilter;
@@ -321,17 +336,6 @@ public:
 	FCreateNonOverlappingConvexHullsDataflowNode(const UE::Dataflow::FNodeParameters& InParam, FGuid InGuid = FGuid::NewGuid());
 };
 
-//~ Simple wrapper class to make the sphere covering data possible to pass via dataflow
-// A set of spheres generated to represent empty space when creating a minimal set of convex hulls, e.g. in one of the Generate Cluster Convex Hulls nodes
-USTRUCT()
-struct FDataflowSphereCovering
-{
-	GENERATED_USTRUCT_BODY()
-
-public:
-	UE::Geometry::FSphereCovering Spheres;
-};
-
 //~ Dataflow-specific copy of the negative space sampling method enum in ConvexDecomposition3.h,
 //~ so that it can be exposed as a UENUM
 // Method to distribute sampling spheres
@@ -601,6 +605,28 @@ public:
 };
 
 
+/** Copy convex hulls from a reference collection */
+USTRUCT(meta = (DataflowGeometryCollection))
+struct FCopyConvexHullsAcrossCollectionBonesDataflowNode : public FDataflowNode
+{
+	GENERATED_USTRUCT_BODY()
+	DATAFLOW_NODE_DEFINE_INTERNAL(FCopyConvexHullsAcrossCollectionBonesDataflowNode, "CopyConvexHullsAcrossCollectionBones", "GeometryCollection|Utilities", "")
+
+private:
+	UPROPERTY(meta = (DataflowInput))
+	FManagedArrayCollection SourceCollection;
+
+	UPROPERTY(meta = (DataflowOutput))
+	FManagedArrayCollection TargetCollection;
+
+	UPROPERTY(meta = (DataflowInput))
+	FDataflowTransformSelection SourceTransformSelection;
+
+	UPROPERTY(meta = (DataflowInput))
+	FDataflowTransformSelection TargetTransformSelection;
+};
+
+
 /** Merge convex hulls on transforms with multiple hulls */
 USTRUCT(meta = (DataflowGeometryCollection))
 struct FMergeConvexHullsDataflowNode : public FDataflowNode
diff --git a/Engine/Source/Runtime/Experimental/Chaos/Private/GeometryCollection/GeometryCollectionConvexUtility.cpp b/Engine/Source/Runtime/Experimental/Chaos/Private/GeometryCollection/GeometryCollectionConvexUtility.cpp
--- a/Engine/Source/Runtime/Experimental/Chaos/Private/GeometryCollection/GeometryCollectionConvexUtility.cpp
+++ b/Engine/Source/Runtime/Experimental/Chaos/Private/GeometryCollection/GeometryCollectionConvexUtility.cpp
@@ -12,7 +12,6 @@
 #include "CompGeom/ConvexHull3.h"
 #include "Templates/Sorting.h"
 #include "Spatial/PointHashGrid3.h"
-#include "CompGeom/ConvexDecomposition3.h"
 #include "Operations/MeshBoolean.h"
 #include "Operations/MeshSelfUnion.h"
 #include "MeshQueries.h"
@@ -1307,7 +1306,8 @@ void HullsFromGeometry(
 	TFunction<bool(int32)> SkipBoneFn = nullptr,
 	const FGeometryCollectionConvexUtility::FConvexDecompositionSettings* OptionalDecompositionSettings = nullptr,
 	const TArray<FGeometryCollectionConvexUtility::FTransformedConvex>* OptionalIntersectConvexHulls = nullptr,
-	const TArray<TSet<int32>>* OptionalTransformToIntersectHulls = nullptr
+	const TArray<TSet<int32>>* OptionalTransformToIntersectHulls = nullptr,
+	TArray<FGeometryCollectionConvexUtility::FSphereCoveringInfo>* OutComputedNavigableSpheres = nullptr
 )
 {
 	TArray<FVector> GlobalVertices;
@@ -1325,6 +1325,7 @@ void HullsFromGeometry(
 	int32 NumBones = Geometry.TransformToGeometryIndex.Num();
 	TransformToConvexIndices.SetNum(NumBones);
 	FCriticalSection HullCS;
+	FCriticalSection NavSpheresCS;
 	ParallelFor(NumBones, [&](int32 Idx)
 	{
 		if (SkipBoneFn && !SkipBoneFn(Idx))
@@ -1497,6 +1498,19 @@ void HullsFromGeometry(
 						Decomposition.ThickenAfterHullFailure = PreprocessOptions.ThickenInputAfterHullFailure;
 
 						Decomposition.InitializeNegativeSpace(Settings);
+						if (OutComputedNavigableSpheres)
+						{
+							FGeometryCollectionConvexUtility::FSphereCoveringInfo SphereInfo
+							{
+								.Transform = GlobalTransformArray[Idx],
+								.SourceBoneIndex = Idx,
+								.SphereCovering = Decomposition.GetNegativeSpace()
+							};
+							SphereInfo.SphereCovering.ApplyTransform(Decomposition.ResultTransform);
+							NavSpheresCS.Lock();
+							OutComputedNavigableSpheres->Add(MoveTemp(SphereInfo));
+							NavSpheresCS.Unlock();
+						}
 
 						constexpr int32 MaxAllowedSplits = 1000000; // more parts than any expected / reasonable decomposition
 						int32 TargetNumSplits = OptionalDecompositionSettings->MaxHullsPerGeometry > 0 ? OptionalDecompositionSettings->MaxHullsPerGeometry - 1 : MaxAllowedSplits + 1;
@@ -1713,7 +1727,8 @@ UE::GeometryCollectionConvexUtility::FConvexHulls
 FGeometryCollectionConvexUtility::ComputeLeafHulls(FGeometryCollection* GeometryCollection, const TArray<FTransform>& GlobalTransformArray, double SimplificationDistanceThreshold, double OverlapRemovalShrinkPercent,
 	TFunction<bool(int32)> SkipBoneFn, const FConvexDecompositionSettings* OptionalDecompositionSettings,
 	const TArray<FTransformedConvex>* OptionalIntersectConvexHulls,
-	const TArray<TSet<int32>>* OptionalTransformToIntersectHulls)
+	const TArray<TSet<int32>>* OptionalTransformToIntersectHulls,
+	TArray<FGeometryCollectionConvexUtility::FSphereCoveringInfo>* OutComputedNavigableSpheres)
 {
 	check(GeometryCollection);
 
@@ -1731,7 +1746,7 @@ FGeometryCollectionConvexUtility::ComputeLeafHulls(FGeometryCollection* Geometry
 
 	HullsFromGeometry(*GeometryCollection, GlobalTransformArray, HasCustomConvexFn, Hulls.Hulls, Hulls.Pivots, Hulls.TransformToHullsIndices, GeometryCollection->SimulationType,
 		FGeometryCollection::ESimulationTypes::FST_Rigid, SimplificationDistanceThreshold, Hulls.OverlapRemovalShrinkPercent, SkipBoneFn, OptionalDecompositionSettings,
-		OptionalIntersectConvexHulls, OptionalTransformToIntersectHulls
+		OptionalIntersectConvexHulls, OptionalTransformToIntersectHulls, OutComputedNavigableSpheres
 	);
 
 	return Hulls;
@@ -1788,7 +1803,12 @@ FGeometryCollectionConvexUtility::FGeometryCollectionConvexData FGeometryCollect
 	return { TransformToConvexIndices, ConvexHull };
 }
 
-void FGeometryCollectionConvexUtility::GenerateLeafConvexHulls(FGeometryCollection& Collection, bool bRestrictToSelection, const TArrayView<const int32> TransformSubset, const FLeafConvexHullSettings& Settings)
+void FGeometryCollectionConvexUtility::GenerateLeafConvexHulls(
+	FGeometryCollection& Collection, bool bRestrictToSelection, 
+	const TArrayView<const int32> TransformSubset, const FLeafConvexHullSettings& Settings,
+	TArray<FSphereCoveringInfo>* OutComputedNavigableSpheres
+
+)
 {
 	int32 NumTransforms = Collection.NumElements(FGeometryCollection::TransformGroup);
 
@@ -1951,7 +1971,9 @@ void FGeometryCollectionConvexUtility::GenerateLeafConvexHulls(FGeometryCollecti
 		}
 		ComputedHulls = ComputeLeafHulls(&Collection, IdentityTransformArray, Settings.SimplificationDistanceThreshold, 0, SkipBoneFn, &Settings.DecompositionSettings,
 			Settings.bComputeIntersectionsBeforeHull ? &ExternalHulls : nullptr, 
-			Settings.bComputeIntersectionsBeforeHull ? &TransformToExternalHullsIndices : nullptr);
+			Settings.bComputeIntersectionsBeforeHull ? &TransformToExternalHullsIndices : nullptr,
+			OutComputedNavigableSpheres
+		);
 
 		if (!bUseIntersect)
 		{
diff --git a/Engine/Source/Runtime/Experimental/Chaos/Public/GeometryCollection/GeometryCollectionConvexUtility.h b/Engine/Source/Runtime/Experimental/Chaos/Public/GeometryCollection/GeometryCollectionConvexUtility.h
--- a/Engine/Source/Runtime/Experimental/Chaos/Public/GeometryCollection/GeometryCollectionConvexUtility.h
+++ b/Engine/Source/Runtime/Experimental/Chaos/Public/GeometryCollection/GeometryCollectionConvexUtility.h
@@ -5,16 +5,12 @@
 #include "CoreMinimal.h"
 #include "GeometryCollection/ManagedArrayCollection.h"
 #include "Chaos/ImplicitFwd.h"
+#include "CompGeom/ConvexDecomposition3.h"
 
 #include "GeometryCollectionConvexUtility.generated.h"
 
 class FGeometryCollection;
 
-namespace UE::Geometry
-{
-	class FSphereCovering;
-	struct FNegativeSpaceSampleSettings;
-}
 
 UENUM()
 enum class EConvexOverlapRemoval : int32
@@ -253,8 +249,20 @@ public:
 		// Note: It seems more logical for this setting to be true, but we expose it as an option because in some special cases the results when it was false were preferred
 		bool bComputeIntersectionsBeforeHull = true;
 	};
+
+	// Computed sphere covering, and optional associated transform/bone
+	struct FSphereCoveringInfo
+	{
+		FTransform Transform = FTransform::Identity;
+		int32 SourceBoneIndex = INDEX_NONE;
+		UE::Geometry::FSphereCovering SphereCovering;
+	};
 	
-	static CHAOS_API void GenerateLeafConvexHulls(FGeometryCollection& Collection, bool bRestrictToSelection, const TArrayView<const int32> TransformSubset, const FLeafConvexHullSettings& Settings);
+	static CHAOS_API void GenerateLeafConvexHulls(
+		FGeometryCollection& Collection, bool bRestrictToSelection, 
+		const TArrayView<const int32> TransformSubset, const FLeafConvexHullSettings& Settings,
+		TArray<FSphereCoveringInfo>* OutComputedNavigableSpheres = nullptr
+	);
 
 	/** Returns the convex hull of the vertices contained in the specified geometry. */
 	static CHAOS_API Chaos::FConvexPtr GetConvexHull(const FGeometryCollection* GeometryCollection, int32 GeometryIndex);
@@ -314,7 +322,8 @@ public:
 	static CHAOS_API UE::GeometryCollectionConvexUtility::FConvexHulls ComputeLeafHulls(FGeometryCollection* GeometryCollection, const TArray<FTransform>& GlobalTransformArray, double SimplificationDistanceThreshold = 0.0, double OverlapRemovalShrinkPercent = 0.0,
 		TFunction<bool(int32)> SkipBoneFn = nullptr, const FConvexDecompositionSettings* OptionalDecompositionSettings = nullptr,
 		const TArray<FTransformedConvex>* OptionalIntersectConvexHulls = nullptr,
-		const TArray<TSet<int32>>* OptionalTransformToIntersectHulls = nullptr);
+		const TArray<TSet<int32>>* OptionalTransformToIntersectHulls = nullptr,
+		TArray<FGeometryCollectionConvexUtility::FSphereCoveringInfo>* OutComputedNavigableSpheres = nullptr);
 
 	// generate a list of convex out of a hierarchy of implciit shapes
 	// suported shapes are scaled / transformed implicits as well as Boxes, convexes, spheres and capsules
diff --git a/Engine/Source/Runtime/GeometryCore/Public/CompGeom/ConvexDecomposition3.h b/Engine/Source/Runtime/GeometryCore/Public/CompGeom/ConvexDecomposition3.h
--- a/Engine/Source/Runtime/GeometryCore/Public/CompGeom/ConvexDecomposition3.h
+++ b/Engine/Source/Runtime/GeometryCore/Public/CompGeom/ConvexDecomposition3.h
@@ -9,6 +9,7 @@
 #include "CoreMinimal.h"
 #include "DynamicMesh/DynamicMesh3.h"
 #include "IndexTypes.h"
+#include "Math/Sphere.h"
 #include "Math/UnrealMathSSE.h"
 #include "Math/Vector.h"
 #include "MathUtil.h"
@@ -243,7 +244,7 @@ public:
 		}
 	}
 
-	void AppendSpheres(TArrayView<const FSphere> Spheres)
+	void AppendSpheres(TArrayView<const UE::Math::TSphere<double>> Spheres)
 	{
 		const int32 AddCount = Spheres.Num();
 		const int32 OrigCount = Position.Num();
@@ -262,11 +263,41 @@ public:
 		Radius.Add(InRadius);
 	}
 
+	// Apply transform to spheres
+	// Note: Non-uniform scales will scale the radius by the smallest scale axis's magnitude (matching how physics collision transforms spheres)
+	void ApplyTransform(const FTransform& Transform)
+	{
+		TransformRange(0, Position.Num(), Transform);
+	}
+
+	// Append spheres from Other w/ the given Transform
+	// Note: Non-uniform scales will scale the radius by the smallest scale axis's magnitude (matching how physics collision transforms spheres)
+	void AppendTransformed(const FSphereCovering& Other, const FTransform& Transform)
+	{
+		int32 AppendStart = Position.Num();
+		Position.Append(Other.Position);
+		Radius.Append(Other.Radius);
+		TransformRange(AppendStart, Position.Num(), Transform);
+	}
+
 private:
 	// Sphere centers
 	TArray<FVector3d> Position;
 	// Sphere radii
 	TArray<double> Radius;
+
+	void TransformRange(const int32 Start, const int32 Max, const FTransform& Transform)
+	{
+		for (int32 Idx = Start; Idx < Max; ++Idx)
+		{
+			Position[Idx] = Transform.TransformPosition(Position[Idx]);
+		}
+		double RadiusScale = Transform.GetScale3D().GetAbsMin();
+		for (int32 Idx = Start; Idx < Max; ++Idx)
+		{
+			Radius[Idx] *= RadiusScale;
+		}
+	}
 };
 
 
diff --git a/Engine/Source/Runtime/GeometryCore/Public/Util/IndexUtil.h b/Engine/Source/Runtime/GeometryCore/Public/Util/IndexUtil.h
--- a/Engine/Source/Runtime/GeometryCore/Public/Util/IndexUtil.h
+++ b/Engine/Source/Runtime/GeometryCore/Public/Util/IndexUtil.h
@@ -16,7 +16,6 @@
 
 namespace IndexUtil
 {
-	using namespace UE::Math;
 	using namespace UE::Geometry;
 	using UE::Geometry::FIndex2i;
 
@@ -257,7 +256,7 @@ namespace IndexUtil
 	 * Replace Val with MapFunc[Val] using index operator
 	 */
 	template<typename T, typename Func>
-	void ApplyMap(TVector<T> & Val, Func MapFunc)
+	void ApplyMap(UE::Math::TVector<T> & Val, Func MapFunc)
 	{
 		Val[0] = MapFunc[Val[0]];
 		Val[1] = MapFunc[Val[1]];
@@ -279,9 +278,9 @@ namespace IndexUtil
 	 * @return MapFunc[Val] using index operator
 	 */
 	template<typename T, typename Func>
-	TVector<T> ApplyMap(const TVector<T> & Val, Func MapFunc)
+	UE::Math::TVector<T> ApplyMap(const UE::Math::TVector<T> & Val, Func MapFunc)
 	{
-		return TVector<T>(MapFunc[Val[0]], MapFunc[Val[1]], MapFunc[Val[2]]);
+		return UE::Math::TVector<T>(MapFunc[Val[0]], MapFunc[Val[1]], MapFunc[Val[2]]);
 	}
 	
 	/**
