Iris
- Fixed issue introduced with 41894142

diff --git a/Engine/Source/Runtime/Experimental/Iris/Core/Private/Iris/ReplicationSystem/NetRefHandleManager.cpp b/Engine/Source/Runtime/Experimental/Iris/Core/Private/Iris/ReplicationSystem/NetRefHandleManager.cpp
--- a/Engine/Source/Runtime/Experimental/Iris/Core/Private/Iris/ReplicationSystem/NetRefHandleManager.cpp
+++ b/Engine/Source/Runtime/Experimental/Iris/Core/Private/Iris/ReplicationSystem/NetRefHandleManager.cpp
@@ -742,16 +742,18 @@ bool FNetRefHandleManager::SetSubObjectNetCondition(FInternalNetRefIndex SubObje
 				const int32 SubObjectArrayIndex = SubObjectsArray->Find(SubObjectInternalIndex);
 				if (ensure(SubObjectArrayIndex != INDEX_NONE))
 				{
-					if (!SubObjectConditionals && SubObjectCondition != ELifetimeCondition::COND_None)
+					if (!SubObjectConditionals)
 					{
+						// No need to create up the conditionals array if we are not setting a condition.
+						if (SubObjectCondition == ELifetimeCondition::COND_None)
+						{
+							
+							bOutWasModified = false;
+							return true;
+
+						}
 						SubObjectConditionals = &SubObjects.GetOrCreateSubObjectConditionalsArray(SubObjectParentIndex);
 					}
-					else
-					{
-						// Valid but not modified.
-						bOutWasModified = false;
-						return true;
-					}					
 
 					check(SubObjectConditionals->Num() == SubObjectsArray->Num());
 
