Implement iostore.DefragInstallCache command.
Add FragmentedChunksSize to iostore.CacheUsage command.
Move dump of referenced handles from purge command to iostore.CacheUsage command.

diff --git a/Engine/Source/Runtime/Experimental/IoStore/OnDemand/Private/OnDemandInstallCache.cpp b/Engine/Source/Runtime/Experimental/IoStore/OnDemand/Private/OnDemandInstallCache.cpp
--- a/Engine/Source/Runtime/Experimental/IoStore/OnDemand/Private/OnDemandInstallCache.cpp
+++ b/Engine/Source/Runtime/Experimental/IoStore/OnDemand/Private/OnDemandInstallCache.cpp
@@ -1327,6 +1327,7 @@ public:
 	virtual FIoStatus							PutChunk(FIoBuffer&& Chunk, const FIoHash& ChunkHash) override;
 	virtual FIoStatus							Purge(TSet<FIoHash>&& ChunksToInstall) override;
 	virtual FIoStatus							PurgeAllUnreferenced() override;
+	virtual FIoStatus							DefragAll() override;
 	virtual FIoStatus							Flush() override;
 	virtual FOnDemandInstallCacheUsage			GetCacheUsage() override;
 
@@ -1341,7 +1342,7 @@ private:
 													const TArray<FSharedOnDemandContainer>& Containers,
 													const TArray<TBitArray<>>& ChunkEntryIndices,
 													FCasBlockInfoMap& BlockInfo, 
-													uint64 TotalBytesToFree);
+													const uint64* TotalBytesToFree = nullptr);
 	bool										Resolve(FIoRequestImpl* Request);
 	void										CompleteRequest(FIoRequestImpl* Request, EIoErrorCode Status);
 	FIoStatus									FlushPendingChunks(FPendingChunks& Chunks, int64 UtcAccessTicks = 0);
@@ -1881,7 +1882,8 @@ FIoStatus FOnDemandInstallCache::Purge(TSet<FIoHash>&& ChunksToInstall)
 		if (UE::IoStore::CVars::GIoStoreOnDemandEnableDefrag)
 		{
 			// Attempt to defrag
-			FIoStatus DefragStatus = Defrag(Containers, ChunkEntryIndices, BlockInfo, TotalBytesToPurge - TotalPurgedBytes);
+			const uint64 DefragBytesToPurge = TotalBytesToPurge - TotalPurgedBytes;
+			FIoStatus DefragStatus = Defrag(Containers, ChunkEntryIndices, BlockInfo, &DefragBytesToPurge);
 			if (DefragStatus.IsOk() == false)
 			{
 				return FIoStatusBuilder(EIoErrorCode::WriteError) 
@@ -1902,36 +1904,6 @@ FIoStatus FOnDemandInstallCache::PurgeAllUnreferenced()
 	FCasBlockInfoMap	BlockInfo;
 	const uint64		TotalCachedBytes = Cas.GetBlockInfo(BlockInfo);
 
-	{
-		TMap<FOnDemandWeakContentHandle, TArray<FOnDemandContainerChunkEntryReferences>> RefsByHandle;
-		IoStore.GetReferencedContentByHandle(RefsByHandle);
-
-		for (const TPair<FOnDemandWeakContentHandle, TArray<FOnDemandContainerChunkEntryReferences>>& Kv : RefsByHandle)
-		{
-			const FOnDemandWeakContentHandle& WeakHandle = Kv.Key;
-			const TArray<FOnDemandContainerChunkEntryReferences>& ContainerReferences = Kv.Value;
-
-			uint64 ReferencedBytesByHandle = 0;
-			for (const FOnDemandContainerChunkEntryReferences& Refs : ContainerReferences)
-			{
-				for (int32 EntryIndex = 0; const FOnDemandChunkEntry& ChunkEntry : Refs.Container->ChunkEntries)
-				{
-					if (Refs.Indices[EntryIndex])
-					{
-						ReferencedBytesByHandle += ChunkEntry.EncodedSize;
-					}
-					EntryIndex++;
-				}
-			}
-
-			if (ReferencedBytesByHandle > 0)
-			{
-				UE_LOG(LogIoStoreOnDemand, Log, TEXT("HandleId=0x%llX, DebugName='%s', ReferencedBytes=%.2lf KiB"),
-					WeakHandle.HandleId, *WeakHandle.DebugName, ToKiB(ReferencedBytesByHandle));
-			}
-		}
-	}
-
 	TArray<FSharedOnDemandContainer>	Containers;
 	TArray<TBitArray<>>					ChunkEntryIndices;
 
@@ -1968,6 +1940,29 @@ FIoStatus FOnDemandInstallCache::PurgeAllUnreferenced()
 	return Status;
 }
 
+FIoStatus FOnDemandInstallCache::DefragAll()
+{
+	FCasBlockInfoMap	BlockInfo;
+	const uint64		TotalCachedBytes = Cas.GetBlockInfo(BlockInfo);
+
+	TArray<FSharedOnDemandContainer>	Containers;
+	TArray<TBitArray<>>					ChunkEntryIndices;
+
+	IoStore.GetReferencedContent(Containers, ChunkEntryIndices);
+	check(Containers.Num() == ChunkEntryIndices.Num());
+
+	AddReferencesToBlocks(Containers, ChunkEntryIndices, {}, BlockInfo);
+
+	const uint64 TotalReferencedBytes = Algo::TransformAccumulate(BlockInfo,
+		[](const TPair<FCasBlockId, FCasBlockInfo>& Kv) { return (Kv.Value.RefSize > 0) ? Kv.Value.FileSize : uint64(0); },
+		uint64(0));
+
+	UE_LOG(LogIoStoreOnDemand, Log, TEXT("Defragmenting install cache, MaxCacheSize=%.2lf MiB, CacheSize=%.2lf MiB, ReferencedBytes=%.2lf MiB"),
+		ToMiB(MaxCacheSize), ToMiB(TotalCachedBytes), ToMiB(TotalReferencedBytes));
+
+	return Defrag(Containers, ChunkEntryIndices, BlockInfo);
+}
+
 uint64 FOnDemandInstallCache::AddReferencesToBlocks(
 	const TArray<FSharedOnDemandContainer>& Containers,
 	const TArray<TBitArray<>>& ChunkEntryIndices,
@@ -2061,7 +2056,7 @@ FIoStatus FOnDemandInstallCache::Defrag(
 	const TArray<FSharedOnDemandContainer>& Containers,
 	const TArray<TBitArray<>>& ChunkEntryIndices,
 	FCasBlockInfoMap& BlockInfo, 
-	uint64 TotalBytesToFree)
+	const uint64* TotalBytesToFree /*= nullptr*/)
 {
 	struct FDefragBlockReferencedChunk
 	{
@@ -2079,13 +2074,20 @@ FIoStatus FOnDemandInstallCache::Defrag(
 
 	// Build the list of blocks to defrag and determine if its possible to free enough data through defragging
 	TArray<FDefragBlock> BlocksToDefrag;
-	bool bDefragPossible = false;
+	
+	// Start with the least referenced blocks
+	BlockInfo.ValueSort([](const FCasBlockInfo& LHS, const FCasBlockInfo& RHS)
 	{
-		// Start with the least referenced blocks
-		BlockInfo.ValueSort([](const FCasBlockInfo& LHS, const FCasBlockInfo& RHS)
-		{
-			return LHS.RefSize < RHS.RefSize;
-		});
+		return LHS.RefSize < RHS.RefSize;
+	});
+
+	uint64 FragmentedBytes = 0;
+	uint64 TotalBlockSize = 0;
+
+	if (TotalBytesToFree)
+	{
+		// Partial defrag
+		bool bPossibleToFreeBytes = false;
 
 		uint64 FreedBlockBytes = 0;
 		uint64 NewBlockBytes = 0;
@@ -2095,24 +2097,56 @@ FIoStatus FOnDemandInstallCache::Defrag(
 			const FCasBlockId BlockId = Kv.Key;
 			const FCasBlockInfo& Info = Kv.Value;
 
-			FreedBlockBytes += Info.FileSize;
-			NewBlockBytes += Info.RefSize; // For now, assume that nothing will be moved to the current block
+			if (Info.RefSize < Info.FileSize)
+			{
+				FragmentedBytes += (Info.FileSize - Info.RefSize);
+				TotalBlockSize += Info.FileSize;
+
+				FreedBlockBytes += Info.FileSize;
+				NewBlockBytes += Info.RefSize; // For now, assume that nothing will be moved to the current block
 
-			BlocksToDefrag.Add(FDefragBlock{ .BlockId = BlockId, .LastAccess = Info.LastAccess });
+				BlocksToDefrag.Add(FDefragBlock{ .BlockId = BlockId, .LastAccess = Info.LastAccess });
 
-			if (FreedBlockBytes >= NewBlockBytes && FreedBlockBytes - NewBlockBytes > TotalBytesToFree)
-			{
-				bDefragPossible = true;
-				break;
+				if (FreedBlockBytes >= NewBlockBytes && FreedBlockBytes - NewBlockBytes > *TotalBytesToFree)
+				{
+					bPossibleToFreeBytes = true;
+					break;
+				}
 			}
 		}
-	}
 
-	if (!bDefragPossible)
+		if (!bPossibleToFreeBytes)
+		{
+			return FIoStatusBuilder(EIoErrorCode::WriteError) << FString::Printf(TEXT("Defrag failed - cannot free %" UINT64_FMT), *TotalBytesToFree);
+		}
+	}
+	else
 	{
-		return FIoStatusBuilder(EIoErrorCode::WriteError) << FString::Printf(TEXT("Defrag failed - cannot free %" UINT64_FMT), TotalBytesToFree);
+		// Full defrag
+		for (const TPair<FCasBlockId, FCasBlockInfo>& Kv : BlockInfo)
+		{
+			const FCasBlockId BlockId = Kv.Key;
+			const FCasBlockInfo& Info = Kv.Value;
+
+			if (Info.RefSize < Info.FileSize)
+			{
+				FragmentedBytes += (Info.FileSize - Info.RefSize);
+				TotalBlockSize += Info.FileSize;
+
+				BlocksToDefrag.Add(FDefragBlock{ .BlockId = BlockId, .LastAccess = Info.LastAccess });
+			}
+		}
+
+		if (BlocksToDefrag.IsEmpty())
+		{
+			// Already defragged
+			UE_LOG(LogIoStoreOnDemand, Display, TEXT("Cache not fragmented."));
+			return FIoStatus::Ok;
+		}
 	}
 
+	UE_LOG(LogIoStoreOnDemand, Display, TEXT("Defrag found %" UINT64_FMT " fragmented bytes of %" UINT64_FMT " total bytes in %i blocks."), FragmentedBytes, TotalBlockSize, BlocksToDefrag.Num());
+
 	// Right now, don't allow moving chunks to the current block for defrag. Its somewhat dangerous and hard to reason about.
 	// - Currently, the slack in the current block cannot be determined without opening a write handle to the block.
 	// - If we defrag the current block itself, then we would need additional tracking so we don't lose any chunks moved into it.
@@ -2150,51 +2184,52 @@ FIoStatus FOnDemandInstallCache::Defrag(
 	FPendingChunks DefragPendingChunks;
 	for (const FDefragBlock& DefragBlock : BlocksToDefrag)
 	{
-		FSharedFileOpenResult FileOpenResult = Cas.OpenRead(DefragBlock.BlockId);
-		if (!FileOpenResult.IsValid())
+		if (DefragBlock.ReferencedChunks.IsEmpty() == false)
 		{
-			const FString Filename = Cas.GetBlockFilename(DefragBlock.BlockId);
-			const FString ErrorMsg = FileOpenResult.GetError().GetMessage();
-			UE_LOG(LogIoStoreOnDemand, Error, TEXT("Failed to open CAS block '%s' for reading, reason '%s'"), *Filename, *ErrorMsg);
-			return FIoStatusBuilder(EIoErrorCode::FileOpenFailed) << ErrorMsg;
-		}
-
-		FSharedFileHandle FileHandle = FileOpenResult.StealValue();
-
-		Algo::SortBy(DefragBlock.ReferencedChunks, &FDefragBlockReferencedChunk::BlockOffset);
+			FSharedFileOpenResult FileOpenResult = Cas.OpenRead(DefragBlock.BlockId);
+			if (!FileOpenResult.IsValid())
+			{
+				const FString Filename = Cas.GetBlockFilename(DefragBlock.BlockId);
+				const FString ErrorMsg = FileOpenResult.GetError().GetMessage();
+				UE_LOG(LogIoStoreOnDemand, Error, TEXT("Failed to open CAS block '%s' for reading, reason '%s'"), *Filename, *ErrorMsg);
+				return FIoStatusBuilder(EIoErrorCode::FileOpenFailed) << ErrorMsg;
+			}
 
-		for (const FDefragBlockReferencedChunk& ReffedChunk : DefragBlock.ReferencedChunks)
-		{
-			FileHandle->Seek(ReffedChunk.BlockOffset);
+			FSharedFileHandle FileHandle = FileOpenResult.StealValue();
 
-			FIoBuffer Buffer;
-			Buffer.SetSize(ReffedChunk.EncodedSize);
+			Algo::SortBy(DefragBlock.ReferencedChunks, &FDefragBlockReferencedChunk::BlockOffset);
 
-			const bool bOk = FileHandle->Read(Buffer.GetData(), Buffer.GetSize()); // TODO: error handling
-			if (!bOk)
+			for (const FDefragBlockReferencedChunk& ReffedChunk : DefragBlock.ReferencedChunks)
 			{
-				return EIoErrorCode::ReadError;
-			}
+				FileHandle->Seek(ReffedChunk.BlockOffset);
 
-			if (DefragPendingChunks.TotalSize > FPendingChunks::MaxPendingBytes)
-			{
-				if (FIoStatus Status = FlushPendingChunks(DefragPendingChunks, DefragBlock.LastAccess); Status.IsOk() == false)
+				FIoBuffer Buffer(ReffedChunk.EncodedSize);
+				const bool bOk = FileHandle->Read(Buffer.GetData(), Buffer.GetSize());
+				if (!bOk)
 				{
-					return Status;
+					return EIoErrorCode::ReadError;
 				}
-				check(DefragPendingChunks.IsEmpty());
+
+				if (DefragPendingChunks.TotalSize > FPendingChunks::MaxPendingBytes)
+				{
+					if (FIoStatus Status = FlushPendingChunks(DefragPendingChunks, DefragBlock.LastAccess); Status.IsOk() == false)
+					{
+						return Status;
+					}
+					check(DefragPendingChunks.IsEmpty());
+				}
+
+				DefragPendingChunks.Append(MoveTemp(Buffer), ReffedChunk.Hash);
 			}
 
-			DefragPendingChunks.Append(MoveTemp(Buffer), ReffedChunk.Hash);
-		}
-		
-		FileHandle.Reset();
+			FileHandle.Reset();
 
-		if (FIoStatus Status = FlushPendingChunks(DefragPendingChunks); Status.IsOk() == false)
-		{
-			return Status;
+			if (FIoStatus Status = FlushPendingChunks(DefragPendingChunks); Status.IsOk() == false)
+			{
+				return Status;
+			}
+			check(DefragPendingChunks.IsEmpty());
 		}
-		check(DefragPendingChunks.IsEmpty());
 
 		// TODO: At startup, as part of the Verify, the cache needs to be check if there is enough room for the defrag chunk, and
 		// purge down to the correct size if not. Its possible something went wrong before the old block could be deleted.
@@ -2257,13 +2292,30 @@ FOnDemandInstallCacheUsage FOnDemandInstallCache::GetCacheUsage()
 
 	AddReferencesToBlocks(Containers, ChunkEntryIndices, {}, BlockInfo);
 
+	uint64 FragmentedBytes = 0;
+	uint64 ReferencedBlockBytes = 0;
+	for (const TPair<FCasBlockId, FCasBlockInfo>& Kv : BlockInfo)
+	{
+		const FCasBlockId BlockId = Kv.Key;
+		const FCasBlockInfo& Info = Kv.Value;
+
+		if (Info.RefSize < Info.FileSize)
+		{
+			FragmentedBytes += (Info.FileSize - Info.RefSize);
+		}
+
+		if (Info.RefSize > 0)
+		{
+			ReferencedBlockBytes += Info.FileSize;
+		}
+	}
+
 	return FOnDemandInstallCacheUsage
 	{
 		.MaxSize = MaxCacheSize,
 		.TotalSize = TotalCachedBytes,
-		.ReferencedBlockSize = Algo::TransformAccumulate(BlockInfo,
-			[](const TPair<FCasBlockId, FCasBlockInfo>& Kv) { return (Kv.Value.RefSize > 0) ? Kv.Value.FileSize : uint64(0); },
-			uint64(0))
+		.ReferencedBlockSize = ReferencedBlockBytes,
+		.FragmentedChunksSize = FragmentedBytes,
 	};
 }
 
diff --git a/Engine/Source/Runtime/Experimental/IoStore/OnDemand/Private/OnDemandInstallCache.h b/Engine/Source/Runtime/Experimental/IoStore/OnDemand/Private/OnDemandInstallCache.h
--- a/Engine/Source/Runtime/Experimental/IoStore/OnDemand/Private/OnDemandInstallCache.h
+++ b/Engine/Source/Runtime/Experimental/IoStore/OnDemand/Private/OnDemandInstallCache.h
@@ -31,6 +31,7 @@ public:
 	virtual FIoStatus							PutChunk(FIoBuffer&& Chunk, const FIoHash& ChunkHash) = 0;
 	virtual FIoStatus							Purge(TSet<FIoHash>&& ChunksToInstall) = 0;
 	virtual FIoStatus							PurgeAllUnreferenced() = 0;
+	virtual FIoStatus							DefragAll() = 0;
 	virtual FIoStatus							Flush() = 0;
 	virtual FOnDemandInstallCacheUsage			GetCacheUsage() = 0;
 };
diff --git a/Engine/Source/Runtime/Experimental/IoStore/OnDemand/Private/OnDemandIoStore.cpp b/Engine/Source/Runtime/Experimental/IoStore/OnDemand/Private/OnDemandIoStore.cpp
--- a/Engine/Source/Runtime/Experimental/IoStore/OnDemand/Private/OnDemandIoStore.cpp
+++ b/Engine/Source/Runtime/Experimental/IoStore/OnDemand/Private/OnDemandIoStore.cpp
@@ -784,10 +784,12 @@ FOnDemandInstallRequest FOnDemandIoStore::Install(
 
 void FOnDemandIoStore::Purge(FOnDemandPurgeArgs&& Args, FOnDemandPurgeCompleted&& OnCompleted)
 {
+	// TODO: redirect to installer?
+
 	LLM_SCOPE_BYNAME(TEXT("FileSystem/OnDemandIoStore"));
-	FSharedPurgeRequest PurgeRequest = MakeShared<FPurgeRequest>();
-	PurgeRequest->Args			= MoveTemp(Args);
-	PurgeRequest->OnCompleted	= MoveTemp(OnCompleted);
+	FSharedPurgeRequest PurgeRequest	= MakeShared<FPurgeRequest>();
+	PurgeRequest->Args					= MoveTemp(Args);
+	PurgeRequest->OnCompleted			= MoveTemp(OnCompleted);
 
 	{
 		UE::TUniqueLock Lock(RequestMutex);
@@ -797,6 +799,23 @@ void FOnDemandIoStore::Purge(FOnDemandPurgeArgs&& Args, FOnDemandPurgeCompleted&
 	TryEnterTickLoop();
 }
 
+void FOnDemandIoStore::Defrag(FOnDemandDefragArgs&& Args, FOnDemandDefragCompleted&& OnCompleted)
+{
+	// TODO: redirect to installer?
+
+	LLM_SCOPE_BYNAME(TEXT("FileSystem/OnDemandIoStore"));
+	FSharedDefragRequest DefragRequest	= MakeShared<FDefragRequest>();
+	DefragRequest->Args					= MoveTemp(Args);
+	DefragRequest->OnCompleted			= MoveTemp(OnCompleted);
+
+	{
+		UE::TUniqueLock Lock(RequestMutex);
+		DefragRequests.Add(MoveTemp(DefragRequest));
+	}
+
+	TryEnterTickLoop();
+}
+
 FIoStatus FOnDemandIoStore::Unmount(FStringView MountId)
 {
 	LLM_SCOPE_BYNAME(TEXT("FileSystem/OnDemandIoStore"));
@@ -1177,10 +1196,13 @@ bool FOnDemandIoStore::Tick()
 	}
 
 	TArray<FSharedPurgeRequest> LocalPurgeRequests;
+	TArray<FSharedDefragRequest> LocalDefragRequests;
 	{
 		UE::TUniqueLock Lock(RequestMutex);
 		LocalPurgeRequests = MoveTemp(PurgeRequests);
 		PurgeRequests.Empty();
+		LocalDefragRequests = MoveTemp(DefragRequests);
+		DefragRequests.Empty();
 	}
 	
 	// Tick purge requests
@@ -1198,6 +1220,21 @@ bool FOnDemandIoStore::Tick()
 		}
 	}
 
+	// Tick defrag requests
+	if (LocalDefragRequests.IsEmpty() == false)
+	{
+		FOnDemandDefragResult DefragResult;
+
+		const double StartTime = FPlatformTime::Seconds();
+		DefragResult.Status = InstallCache->DefragAll();
+		DefragResult.DurationInSeconds = FPlatformTime::Seconds() - StartTime;
+
+		for (FSharedDefragRequest& Request : LocalDefragRequests)
+		{
+			CompleteDefragRequest(*Request, FOnDemandDefragResult(DefragResult));
+		}
+	}
+
 	FSharedInternalInstallRequest InstallRequest;
 	{
 		UE::TUniqueLock Lock(RequestMutex);
@@ -1972,7 +2009,7 @@ void FOnDemandIoStore::CompletePurgeRequest(FPurgeRequest& Request, FOnDemandPur
 		return;
 	}
 
-	if (EnumHasAnyFlags(Request.Args.Options, EOnDemandInstallOptions::CallbackOnGameThread))
+	if (EnumHasAnyFlags(Request.Args.Options, EOnDemandPurgeOptions::CallbackOnGameThread))
 	{
 		ExecuteOnGameThread(
 			UE_SOURCE_LOCATION,
@@ -1988,6 +2025,29 @@ void FOnDemandIoStore::CompletePurgeRequest(FPurgeRequest& Request, FOnDemandPur
 	}
 }
 
+void FOnDemandIoStore::CompleteDefragRequest(FDefragRequest& Request, FOnDemandDefragResult&& Result)
+{
+	if (!Request.OnCompleted)
+	{
+		return;
+	}
+
+	if (EnumHasAnyFlags(Request.Args.Options, EOnDemandDefragOptions::CallbackOnGameThread))
+	{
+		ExecuteOnGameThread(
+			UE_SOURCE_LOCATION,
+			[OnCompleted = MoveTemp(Request.OnCompleted), Result = MoveTemp(Result)]() mutable
+			{
+				OnCompleted(MoveTemp(Result));
+			});
+	}
+	else
+	{
+		FOnDemandDefragCompleted OnCompleted = MoveTemp(Request.OnCompleted);
+		OnCompleted(MoveTemp(Result));
+	}
+}
+
 void FOnDemandIoStore::OnEncryptionKeyAdded(const FGuid& Id, const FAES::FAESKey& Key)
 {
 	TRACE_CPUPROFILER_EVENT_SCOPE(FOnDemandIoStore::OnEncryptionKeyAdded);
@@ -2350,10 +2410,10 @@ void FOnDemandIoStore::GetReferencedContent(TArray<FSharedOnDemandContainer>& Ou
 	}
 }
 
-void FOnDemandIoStore::GetReferencedContentByHandle(TMap<FOnDemandWeakContentHandle, TArray<FOnDemandContainerChunkEntryReferences>>& OutReferencesByHandle)
+void FOnDemandIoStore::GetReferencedContentByHandle(TMap<FOnDemandWeakContentHandle, TArray<FOnDemandContainerChunkEntryReferences>>& OutReferencesByHandle) const
 {
 	UE::TUniqueLock Lock(ContainerMutex);
-	for (FSharedOnDemandContainer& Container : Containers)
+	for (const FSharedOnDemandContainer& Container : Containers)
 	{
 		for (const FOnDemandChunkEntryReferences& Refs : Container->ChunkEntryReferences)
 		{
@@ -2402,8 +2462,39 @@ void FOnDemandIoStore::UpdateInstallRequestPriority(FSharedInternalInstallReques
 #endif
 }
 
-FOnDemandCacheUsage FOnDemandIoStore::GetCacheUsage() const
+FOnDemandCacheUsage FOnDemandIoStore::GetCacheUsage(const FOnDemandGetCacheUsageArgs& Args) const
 {
+	if (EnumHasAnyFlags(Args.Options, EOnDemandGetCacheUsageOptions::DumpHandlesToLog))
+	{
+		TMap<FOnDemandWeakContentHandle, TArray<FOnDemandContainerChunkEntryReferences>> RefsByHandle;
+		GetReferencedContentByHandle(RefsByHandle);
+
+		for (const TPair<FOnDemandWeakContentHandle, TArray<FOnDemandContainerChunkEntryReferences>>& Kv : RefsByHandle)
+		{
+			const FOnDemandWeakContentHandle& WeakHandle = Kv.Key;
+			const TArray<FOnDemandContainerChunkEntryReferences>& ContainerReferences = Kv.Value;
+
+			uint64 ReferencedBytesByHandle = 0;
+			for (const FOnDemandContainerChunkEntryReferences& Refs : ContainerReferences)
+			{
+				for (int32 EntryIndex = 0; const FOnDemandChunkEntry & ChunkEntry : Refs.Container->ChunkEntries)
+				{
+					if (Refs.Indices[EntryIndex])
+					{
+						ReferencedBytesByHandle += ChunkEntry.EncodedSize;
+					}
+					EntryIndex++;
+				}
+			}
+
+			if (ReferencedBytesByHandle > 0)
+			{
+				UE_LOG(LogIoStoreOnDemand, Log, TEXT("HandleId=0x%llX, DebugName='%s', ReferencedBytes=%.2lf KiB"),
+					WeakHandle.HandleId, *WeakHandle.DebugName, double(ReferencedBytesByHandle) / 1024.0);
+			}
+		}
+	}
+
 	FOnDemandCacheUsage CacheUsage;
 	if (InstallCache.IsValid())
 	{
diff --git a/Engine/Source/Runtime/Experimental/IoStore/OnDemand/Private/OnDemandIoStore.h b/Engine/Source/Runtime/Experimental/IoStore/OnDemand/Private/OnDemandIoStore.h
--- a/Engine/Source/Runtime/Experimental/IoStore/OnDemand/Private/OnDemandIoStore.h
+++ b/Engine/Source/Runtime/Experimental/IoStore/OnDemand/Private/OnDemandIoStore.h
@@ -304,6 +304,14 @@ class FOnDemandIoStore
 
 	using FSharedPurgeRequest = TSharedRef<FPurgeRequest>;
 
+	struct FDefragRequest
+	{
+		FOnDemandDefragArgs					Args;
+		FOnDemandDefragCompleted			OnCompleted;
+	};
+
+	using FSharedDefragRequest = TSharedRef<FDefragRequest>;
+
 public:
 	UE_NONCOPYABLE(FOnDemandIoStore);
 
@@ -320,7 +328,7 @@ public:
 	TArray<FIoChunkId>					DebugFindStreamingChunkIds(int32 NumToFind);
 #endif //!UE_BUILD_SHIPPING
 	void								GetReferencedContent(TArray<FSharedOnDemandContainer>& OutContainers, TArray<TBitArray<>>& OutChunkEntryIndices, bool bPackageStore = false);
-	void								GetReferencedContentByHandle(TMap<FOnDemandWeakContentHandle, TArray<FOnDemandContainerChunkEntryReferences>>& OutReferencesByHandle);
+	void								GetReferencedContentByHandle(TMap<FOnDemandWeakContentHandle, TArray<FOnDemandContainerChunkEntryReferences>>& OutReferencesByHandle) const;
 
 	// IOnDemandIoStore interface
 	virtual FIoStatus					Initialize() override;
@@ -330,10 +338,11 @@ public:
 											FOnDemandInstallCompleted&& OnCompleted,
 											FOnDemandInstallProgressed&& OnProgress = nullptr) override;
 	virtual void						Purge(FOnDemandPurgeArgs&& Args, FOnDemandPurgeCompleted&& OnCompleted) override;
+	virtual void						Defrag(FOnDemandDefragArgs&& Args, FOnDemandDefragCompleted&& OnCompleted) override;
 	virtual FIoStatus					Unmount(FStringView MountId) override;
 	virtual TIoStatusOr<uint64> 		GetInstallSize(const FOnDemandGetInstallSizeArgs& Args) const override;
 	virtual FIoStatus					GetInstallSizesByMountId(const FOnDemandGetInstallSizeArgs& Args, TMap<FString, uint64>& OutSizesByMountId) const override;
-	virtual FOnDemandCacheUsage			GetCacheUsage() const override;
+	virtual FOnDemandCacheUsage			GetCacheUsage(const FOnDemandGetCacheUsageArgs& Args) const override;
 	virtual bool						IsOnDemandStreamingEnabled() const override;
 	virtual void						SetStreamingOptions(EOnDemandStreamingOptions Options) override;
 	virtual void						ReportAnalytics(TArray<FAnalyticsEventAttribute>& OutAnalyticsArray) const override;
@@ -361,6 +370,7 @@ private:
 	void								CompleteInstallRequest(FSharedInternalInstallRequest InstallRequest, FOnDemandInstallResult&& InstallResult);
 	void								ProgressInstallRequest(const FOnDemandInternalInstallRequest& InstallRequest, const FOnDemandInstallProgress& Progress);
 	void								CompletePurgeRequest(FPurgeRequest& PurgeRequest, FOnDemandPurgeResult&& Result);
+	void								CompleteDefragRequest(FDefragRequest& DefragRequest, FOnDemandDefragResult&& Result);
 	void								OnEncryptionKeyAdded(const FGuid& Id, const FAES::FAESKey& Key);
 	static void							CreateContainersFromToc(
 											FStringView MountId,
@@ -384,6 +394,7 @@ private:
 	TArray<FSharedMountRequest>				MountRequests;
 	TArray<FSharedInternalInstallRequest>	InstallRequests;
 	TArray<FSharedPurgeRequest>				PurgeRequests;
+	TArray<FSharedDefragRequest>			DefragRequests;
 	UE::FMutex								RequestMutex;
 
 	EOnDemandStreamingOptions				StreamingOptions = EOnDemandStreamingOptions::Default;
diff --git a/Engine/Source/Runtime/Experimental/IoStore/OnDemandCore/Private/IoStoreOnDemand.cpp b/Engine/Source/Runtime/Experimental/IoStore/OnDemandCore/Private/IoStoreOnDemand.cpp
--- a/Engine/Source/Runtime/Experimental/IoStore/OnDemandCore/Private/IoStoreOnDemand.cpp
+++ b/Engine/Source/Runtime/Experimental/IoStore/OnDemandCore/Private/IoStoreOnDemand.cpp
@@ -80,12 +80,39 @@ static void PurgeInstallCache()
 	}
 }
 
+////////////////////////////////////////////////////////////////////////////////
+static void DefragInstallCache()
+{
+	if (UE::IoStore::IOnDemandIoStore* IoStore = UE::IoStore::TryGetOnDemandIoStore())
+	{
+		UE_LOG(LogIoStoreOnDemand, Display, TEXT("Defragging on demand install cache"));
+		IoStore->Defrag(FOnDemandDefragArgs(), [](const FOnDemandDefragResult& Result)
+		{
+			if (Result.Status.IsOk())
+			{
+				UE_LOG(LogIoStoreOnDemand, Display, TEXT("Defragmented on demand install cache"));
+			}
+			else
+			{
+				UE_LOG(LogIoStoreOnDemand, Error, TEXT("Failed to defragment on demand install cache: %s"), *Result.Status.ToString());
+			}
+		});
+	}
+	else
+	{
+		UE_LOG(LogIoStoreOnDemand, Error, TEXT("Defrag install cache failed, reason 'I/O store on-demand module not initialized'"));
+	}
+}
+
 ////////////////////////////////////////////////////////////////////////////////
 static void PrintCacheUsage()
 {
 	if (UE::IoStore::IOnDemandIoStore* IoStore = UE::IoStore::TryGetOnDemandIoStore())
 	{
-		TIoStatusOr<FOnDemandCacheUsage> MaybeUsage = IoStore->GetCacheUsage();
+		FOnDemandGetCacheUsageArgs Args;
+		Args.Options = EOnDemandGetCacheUsageOptions::DumpHandlesToLog;
+
+		TIoStatusOr<FOnDemandCacheUsage> MaybeUsage = IoStore->GetCacheUsage(Args);
 		if (!MaybeUsage.IsOk())
 		{
 			UE_LOG(LogIoStoreOnDemand, Error, TEXT("iostore.CacheUsage failed: %s"), *MaybeUsage.Status().ToString());
@@ -307,6 +334,12 @@ static FAutoConsoleCommand PurgeCacheCommand(
 	FConsoleCommandDelegate::CreateStatic(PurgeInstallCache),
 	ECVF_Cheat);
 
+static FAutoConsoleCommand DefragCacheCommand(
+	TEXT("iostore.DefragInstallCache"),
+	TEXT("Defragment On Demand Install Cache"),
+	FConsoleCommandDelegate::CreateStatic(DefragInstallCache),
+	ECVF_Cheat);
+
 static FAutoConsoleCommand PrintCacheUsageCommand(
 	TEXT("iostore.CacheUsage"),
 	TEXT("print cache usage"),
@@ -336,10 +369,11 @@ static const TCHAR* NotInitializedError = TEXT("I/O store on-demand not initiali
 ////////////////////////////////////////////////////////////////////////////////
 FStringBuilderBase& operator<<(FStringBuilderBase& Sb, const FOnDemandInstallCacheUsage& CacheUsage)
 {
-	Sb.Appendf(TEXT("MaxSize=%.2lf MiB, TotalSize=%.2lf MiB, ReferencedBlockSize=%.2lf MiB"),
+	Sb.Appendf(TEXT("MaxSize=%.2lf MiB, TotalSize=%.2lf MiB, ReferencedBlockSize=%.2lf MiB, FragmentedChunksSize=%.2lf MiB"),
 		double(CacheUsage.MaxSize) / 1024.0 / 1024.0,
 		double(CacheUsage.TotalSize) / 1024.0 / 1024.0,
-		double(CacheUsage.ReferencedBlockSize) / 1024.0 / 1024.0);
+		double(CacheUsage.ReferencedBlockSize) / 1024.0 / 1024.0,
+		double(CacheUsage.FragmentedChunksSize) / 1024.0 / 1024.0);
 
 	return Sb;
 }
diff --git a/Engine/Source/Runtime/Experimental/IoStore/OnDemandCore/Public/IO/IoStoreOnDemand.h b/Engine/Source/Runtime/Experimental/IoStore/OnDemandCore/Public/IO/IoStoreOnDemand.h
--- a/Engine/Source/Runtime/Experimental/IoStore/OnDemandCore/Public/IO/IoStoreOnDemand.h
+++ b/Engine/Source/Runtime/Experimental/IoStore/OnDemandCore/Public/IO/IoStoreOnDemand.h
@@ -305,11 +305,24 @@ private:
 	FSharedInternalInstallRequest Request;
 };
 
+/** Options for controlling the behavior of the purge request. */
+enum class EOnDemandPurgeOptions
+{
+	/** No additional options. */
+	None = 0,
+	/** Trigger callback on game thread. */
+	CallbackOnGameThread = 1 << 0,
+	/** Defrag after purging */
+	Defrag = 1 << 1, // TODO:
+};
+ENUM_CLASS_FLAGS(EOnDemandPurgeOptions)
+
 /** Arguments for purging on-demand content. */
 struct FOnDemandPurgeArgs
 {
 	/** Purge options. */
-	EOnDemandInstallOptions Options = EOnDemandInstallOptions::None;
+	EOnDemandPurgeOptions Options = EOnDemandPurgeOptions::None;
+	// TODO: allow passing in the amount to purge
 };
 
 /** Holds information about a purge request */
@@ -324,6 +337,36 @@ struct FOnDemandPurgeResult
 /** Purge completion callback */
 using FOnDemandPurgeCompleted = TUniqueFunction<void(FOnDemandPurgeResult)>;
 
+/** Options for controlling the behavior of the defrag request. */
+enum class EOnDemandDefragOptions
+{
+	/** No additional options. */
+	None = 0,
+	/** Trigger callback on game thread. */
+	CallbackOnGameThread = 1 << 0,
+};
+ENUM_CLASS_FLAGS(EOnDemandDefragOptions)
+
+/** Arguments for defragmenting on-demand content. */
+struct FOnDemandDefragArgs
+{
+	/** Defrag options. */
+	EOnDemandDefragOptions Options = EOnDemandDefragOptions::None;
+	// TODO: allow passing in the amount to defrag
+};
+
+/** Holds information about a defrag request */
+struct FOnDemandDefragResult
+{
+	/** The status of the defrag request. */
+	FIoStatus Status;
+	/** Duration in seconds. */
+	double DurationInSeconds = 0.0;
+};
+
+/** Purge completion callback */
+using FOnDemandDefragCompleted = TUniqueFunction<void(FOnDemandDefragResult)>;
+
 /** Options for controlling the behavior of the install size request. */
 enum class EOnDemandGetInstallSizeOptions
 {
@@ -347,12 +390,29 @@ struct FOnDemandGetInstallSizeArgs
 	EOnDemandGetInstallSizeOptions Options = EOnDemandGetInstallSizeOptions::None;
 };
 
+/** Options for controlling the behavior of the install size request. */
+enum class EOnDemandGetCacheUsageOptions
+{
+	/** No additional options. */
+	None = 0,
+	/** Dump Handle information to log */
+	DumpHandlesToLog = 1 << 0,
+};
+ENUM_CLASS_FLAGS(EOnDemandGetCacheUsageOptions);
+
+/** Arguments for getting the install cache usage. */
+struct FOnDemandGetCacheUsageArgs
+{
+	EOnDemandGetCacheUsageOptions Options = EOnDemandGetCacheUsageOptions::None;
+};
+
 /** Holds information about the install cache usage */
 struct FOnDemandInstallCacheUsage
 {
 	uint64 MaxSize = 0;
 	uint64 TotalSize = 0;
 	uint64 ReferencedBlockSize = 0;
+	uint64 FragmentedChunksSize = 0;
 
 	UE_API friend FStringBuilderBase& operator<<(FStringBuilderBase& Sb, const FOnDemandInstallCacheUsage& CacheUsage);
 };
@@ -433,6 +493,12 @@ public:
 	 * @param	OnComplete	Completion callback.
 	 **/
 	virtual void Purge(FOnDemandPurgeArgs&& Args, FOnDemandPurgeCompleted&& OnCompleted) = 0;
+	/**
+	 * Defrag the cache.
+	 * @param	Args		The defrag arguments.
+	 * @param	OnComplete	Completion callback.
+	 **/
+	virtual void Defrag(FOnDemandDefragArgs&& Args, FOnDemandDefragCompleted&& OnCompleted) = 0;
 	/**
 	 * Get the total space in bytes needed to install the specified content.
 	 * @param Args	Arguments for computing required install size.
@@ -445,7 +511,7 @@ public:
 	 **/
 	virtual FIoStatus GetInstallSizesByMountId(const FOnDemandGetInstallSizeArgs& Args, TMap<FString, uint64>& OutSizesByMountId) const = 0;
 	/** Returns the total cache size in bytes. */
-	virtual FOnDemandCacheUsage GetCacheUsage() const = 0;
+	virtual FOnDemandCacheUsage GetCacheUsage(const FOnDemandGetCacheUsageArgs& Args) const = 0;
 	/** Return if the OnDemand streaming system is enabled. */
 	virtual bool IsOnDemandStreamingEnabled() const = 0;
 	/** Set streaming options.*/
