Constraints: ensure that a constraint acting on an object is evaluated before constraints whose parent is that object

#rnx

diff --git a/Engine/Source/Runtime/Experimental/Animation/Constraints/Private/Transform/DependencyBuilder.cpp b/Engine/Source/Runtime/Experimental/Animation/Constraints/Private/Transform/DependencyBuilder.cpp
--- a/Engine/Source/Runtime/Experimental/Animation/Constraints/Private/Transform/DependencyBuilder.cpp
+++ b/Engine/Source/Runtime/Experimental/Animation/Constraints/Private/Transform/DependencyBuilder.cpp
@@ -29,13 +29,6 @@ UObject* GetHandleTarget(const TObjectPtr<UTransformableHandle>& InHandle)
 	return IsValid(InHandle) ? InHandle->GetTarget().Get() : nullptr; 
 }
 
-static bool	bIncludeTarget = true;
-static FAutoConsoleVariableRef CVarIncludeTarget(
-	TEXT("Constraints.IncludeTarget"),
-	bIncludeTarget,
-	TEXT("Include target when getting child's existing constraints.")
-	);
-
 static bool	bDebugDependencies = false;
 static FAutoConsoleVariableRef CVarDebugDependencies(
 	TEXT("Constraints.DebugDependencies"),
@@ -152,286 +145,374 @@ bool FDependencyBuilder::LogDependencies()
 	return DependencyLocals::bDebugDependencies;
 }
 
-bool FDependencyBuilder::BuildDependencies(UWorld* InWorld, UTickableTransformConstraint* InConstraint)
+void FDependencyBuilder::BuildSelfDependencies(UWorld* InWorld, UTickableTransformConstraint* InConstraint)
 {
 	using namespace DependencyLocals;
 	using ConstraintWeakPtr = TWeakObjectPtr<UTickableConstraint>;
+
+	static const TCHAR* SelfDependencyDesc = TEXT("Self Dependency");
 	
-	if (!ensure(InWorld))
+	if (!InConstraint || !InConstraint->IsValid())
 	{
-		return false;
+		return;
 	}
-
-	if (!InConstraint || !InConstraint->IsValid())
+	
+	const UTransformableHandle* ParentHandle = InConstraint->ParentTRSHandle.Get();
+	const UTransformableHandle* ChildHandle = InConstraint->ChildTRSHandle.Get();
+		
+	const UObject* ParentTarget = GetHandleTarget(InConstraint->ParentTRSHandle);
+	const UObject* ChildTarget = GetHandleTarget(InConstraint->ChildTRSHandle);
+	
+	const bool bSelf = ParentTarget && ParentTarget == ChildTarget;
+	if (!bSelf)
 	{
-		return false;
+		return;
 	}
 
-	if (bDebugDependencies)
+	const UObject* SelfTarget = ChildTarget;
+
+	const FConstraintsManagerController& Controller = FConstraintsManagerController::Get(InWorld);
+	auto SelfTargetPredicate = [InConstraint, SelfTarget](const ConstraintWeakPtr& Constraint)
 	{
-		UE_LOG(LogTemp, Warning, TEXT("Building dependencies for '%s' ..."), *GetConstraintLabel(InConstraint));
+		const UTickableTransformConstraint* TransformConstraint = Cast<UTickableTransformConstraint>(Constraint.Get());
+		if (!TransformConstraint || TransformConstraint == InConstraint)
+		{
+			return false;
+		}
+		const UObject* ParentTarget = GetHandleTarget(TransformConstraint->ParentTRSHandle);
+		const UObject* ChildTarget = GetHandleTarget(TransformConstraint->ChildTRSHandle);
+		return ParentTarget == SelfTarget && ChildTarget == SelfTarget;
+	};
+	const TArray< ConstraintWeakPtr > SelfConstraints = Controller.GetConstraintsByPredicate(SelfTargetPredicate);
+	
+	for (const ConstraintWeakPtr& SelfConstraint: SelfConstraints)
+	{
+		const UTickableTransformConstraint* TransformConstraint = Cast<UTickableTransformConstraint>(SelfConstraint);
+
+		// if the new handles depend on that constraint child then, TransformConstraint should tick before
+		if (ParentHandle->HasDirectDependencyWith(*TransformConstraint->ChildTRSHandle))
+		{
+			Controller.SetConstraintsDependencies(TransformConstraint->ConstraintID, InConstraint->ConstraintID);
+			LogDependency(SelfDependencyDesc, TransformConstraint->ChildTRSHandle, ParentHandle, TransformConstraint, InConstraint);
+		}
+		else if (ChildHandle->HasDirectDependencyWith(*TransformConstraint->ChildTRSHandle))
+		{
+			Controller.SetConstraintsDependencies(TransformConstraint->ConstraintID, InConstraint->ConstraintID);
+			LogDependency(SelfDependencyDesc, TransformConstraint->ChildTRSHandle, ChildHandle, TransformConstraint, InConstraint);
+		}
+
+		// if the TransformConstraint handles depend on the new constraint child then, TransformConstraint should tick after
+		if (TransformConstraint->ParentTRSHandle->HasDirectDependencyWith(*ChildHandle))
+		{
+			Controller.SetConstraintsDependencies(InConstraint->ConstraintID, TransformConstraint->ConstraintID);
+			LogDependency(SelfDependencyDesc, ChildHandle, TransformConstraint->ParentTRSHandle, InConstraint, TransformConstraint);
+		}
+		else if (TransformConstraint->ChildTRSHandle->HasDirectDependencyWith(*ChildHandle))
+		{
+			Controller.SetConstraintsDependencies(InConstraint->ConstraintID, TransformConstraint->ConstraintID);
+			LogDependency(SelfDependencyDesc, ChildHandle, TransformConstraint->ChildTRSHandle, InConstraint, TransformConstraint);
+		}
 	}
+}
 
+void FDependencyBuilder::BuildExternalDependencies(UWorld* InWorld, UTickableTransformConstraint* InConstraint)
+{
+	using namespace DependencyLocals;
+	using ConstraintWeakPtr = TWeakObjectPtr<UTickableConstraint>;
+
+	if (!InConstraint || !InConstraint->IsValid())
+	{
+		return;
+	}
+	
 	const UTransformableHandle* ParentHandle = InConstraint->ParentTRSHandle.Get();
 	const UTransformableHandle* ChildHandle = InConstraint->ChildTRSHandle.Get();
+		
+	const UObject* ParentTarget = GetHandleTarget(InConstraint->ParentTRSHandle);
+	const UObject* ChildTarget = GetHandleTarget(InConstraint->ChildTRSHandle);
 	
- 	// get previous child constraints
-	const FConstraintsManagerController& Controller = FConstraintsManagerController::Get(InWorld);
-	TArray<ConstraintWeakPtr> ChildParentConstraints = Controller.GetParentConstraints(ChildHandle->GetHash(), true);
-	ChildParentConstraints.Remove(InConstraint);
-
-	// add dependencies with the last child constraint
-	if (!ChildParentConstraints.IsEmpty())
+	const bool bSelf = ParentTarget && ParentTarget == ChildTarget;
+	if (bSelf)
 	{
-		const FGuid LastChildConstraintID = ChildParentConstraints.Last()->ConstraintID;
-		if (bDebugDependencies)
-		{
-			UE_LOG(LogTemp, Warning, TEXT("Order Dependency: tick after last constraint."));
-		}
-		Controller.SetConstraintsDependencies(LastChildConstraintID, InConstraint->ConstraintID);
+		return;
 	}
 
-	const UObject* ParentTarget = ParentHandle->GetTarget().Get();
-	const UObject* ChildTarget = ChildHandle->GetTarget().Get();
-	const bool bSelf = ParentTarget && ParentTarget == ChildTarget;
+	const FConstraintsManagerController& Controller = FConstraintsManagerController::Get(InWorld);
 	
-	// internal dependencies?
-	if (bSelf && bIncludeTarget)
 	{
-		static const TCHAR* SelfDependencyDesc = TEXT("Self Dependency");
-		
-		const UObject* SelfTarget = ChildTarget;
-		
-		auto SelfTargetPredicate = [InConstraint, SelfTarget](const ConstraintWeakPtr& Constraint)
+		// get all constraints acting on the same target 
+		auto SameChildTargetPredicate = [ChildTarget](const ConstraintWeakPtr& Constraint)
 		{
 			const UTickableTransformConstraint* TransformConstraint = Cast<UTickableTransformConstraint>(Constraint.Get());
-			if (!TransformConstraint || TransformConstraint == InConstraint)
+			if (!TransformConstraint)
 			{
 				return false;
 			}
-			const UObject* ParentTarget = GetHandleTarget(TransformConstraint->ParentTRSHandle);
-			const UObject* ChildTarget = GetHandleTarget(TransformConstraint->ChildTRSHandle);
-			return ParentTarget == SelfTarget && ChildTarget == SelfTarget;
+			const UObject* OtherChildTarget = GetHandleTarget(TransformConstraint->ChildTRSHandle);
+			return OtherChildTarget && OtherChildTarget == ChildTarget;
 		};
+		TArray< ConstraintWeakPtr > ChildTargetParentConstraints = Controller.GetConstraintsByPredicate(SameChildTargetPredicate);
 
-		const TArray< ConstraintWeakPtr > SelfConstraints = Controller.GetConstraintsByPredicate(SelfTargetPredicate);
-		for (const ConstraintWeakPtr& SelfConstraint: SelfConstraints)
+		// store constraint index in this array
+		const int32 ConstrainIndex = ChildTargetParentConstraints.IndexOfByKey(InConstraint);
+		if (ensure(ChildTargetParentConstraints.IsValidIndex(ConstrainIndex)))
 		{
-			const UTickableTransformConstraint* TransformConstraint = Cast<UTickableTransformConstraint>(SelfConstraint);
-
-			// if the new handles depend on that constraint child then, TransformConstraint should tick before
-			if (ParentHandle->HasDirectDependencyWith(*TransformConstraint->ChildTRSHandle))
-			{
-				Controller.SetConstraintsDependencies(TransformConstraint->ConstraintID, InConstraint->ConstraintID);
-				LogDependency(SelfDependencyDesc, TransformConstraint->ChildTRSHandle, ParentHandle, TransformConstraint, InConstraint);
-			}
-			else if (ChildHandle->HasDirectDependencyWith(*TransformConstraint->ChildTRSHandle))
-			{
-				Controller.SetConstraintsDependencies(TransformConstraint->ConstraintID, InConstraint->ConstraintID);
-				LogDependency(SelfDependencyDesc, TransformConstraint->ChildTRSHandle, ChildHandle, TransformConstraint, InConstraint);
-			}
-
-			// if the TransformConstraint handles depend on the new constraint child then, TransformConstraint should tick after
-			if (TransformConstraint->ParentTRSHandle->HasDirectDependencyWith(*ChildHandle))
-			{
-			 	Controller.SetConstraintsDependencies(InConstraint->ConstraintID, TransformConstraint->ConstraintID);
-				LogDependency(SelfDependencyDesc, ChildHandle, TransformConstraint->ParentTRSHandle, InConstraint, TransformConstraint);
-			}
-			else if (TransformConstraint->ChildTRSHandle->HasDirectDependencyWith(*ChildHandle))
-			{
-				Controller.SetConstraintsDependencies(InConstraint->ConstraintID, TransformConstraint->ConstraintID);
-				LogDependency(SelfDependencyDesc, ChildHandle, TransformConstraint->ChildTRSHandle, InConstraint, TransformConstraint);
-			}
+			ChildTargetParentConstraints.RemoveAt(ConstrainIndex);
 		}
-	}
-
-	// make sure we tick after the parent.
-	InConstraint->EnsurePrimaryDependency(InWorld);
-
-	// if child handle is the parent of some other constraints, ensure they will tick after that new one
-	static const TCHAR* ChildDependencyDesc = TEXT("Child Dependency");
-
-	TArray<ConstraintWeakPtr> ChildChildConstraints;
-	UE::TransformConstraintUtil::GetChildrenConstraints(InWorld, InConstraint, ChildChildConstraints, bIncludeTarget && !bSelf);
-	for (const ConstraintWeakPtr& ChildConstraint: ChildChildConstraints)
-	{
-		Controller.SetConstraintsDependencies(InConstraint->ConstraintID, ChildConstraint->ConstraintID);
-		if (const UTickableTransformConstraint* TransformConstraint = Cast<UTickableTransformConstraint>(ChildConstraint))
+	
+		if (!ChildTargetParentConstraints.IsEmpty())
 		{
-			LogDependency(ChildDependencyDesc, ChildHandle, TransformConstraint->ChildTRSHandle, InConstraint, TransformConstraint);
-		}
-	}
-
-	// build dependencies regarding attachments
-	BuildAttachmentsDependencies(InWorld, InConstraint);
+			static const TCHAR* ExternalDependencyDesc = TEXT("External Dependency");
+		
+			TBitArray<> ManagedDependencies(false, ChildTargetParentConstraints.Num());
+		
+			const FConstraintTickFunction& TickFunction = InConstraint->GetTickFunction(InWorld);
+			const TArray<FTickPrerequisite> PrerexCopy = TickFunction.GetPrerequisites();
 
-	if (bIncludeTarget && !bSelf)
-	{
-		{
-			// get all constraints acting on the same target 
-			auto SameChildTargetPredicate = [ChildTarget](const ConstraintWeakPtr& Constraint)
+			for (int32 Index = 0; Index < ChildTargetParentConstraints.Num(); ++Index)
 			{
-				const UTickableTransformConstraint* TransformConstraint = Cast<UTickableTransformConstraint>(Constraint.Get());
-				if (!TransformConstraint)
+				const UTickableTransformConstraint* TransformConstraint = Cast<UTickableTransformConstraint>(ChildTargetParentConstraints[Index]);
+			
+				// if the new handles depend on that constraint child then, TransformConstraint should tick before
+				if (ParentHandle->HasDirectDependencyWith(*TransformConstraint->ChildTRSHandle))
 				{
-					return false;
+					Controller.SetConstraintsDependencies(TransformConstraint->ConstraintID, InConstraint->ConstraintID);
+					ManagedDependencies[Index] = true;
+					LogDependency(ExternalDependencyDesc, TransformConstraint->ChildTRSHandle, ParentHandle, TransformConstraint, InConstraint);
 				}
-				const UObject* OtherChildTarget = GetHandleTarget(TransformConstraint->ChildTRSHandle);
-				return OtherChildTarget && OtherChildTarget == ChildTarget;
-			};
-			TArray< ConstraintWeakPtr > ChildTargetParentConstraints = Controller.GetConstraintsByPredicate(SameChildTargetPredicate);
-
-			// store constraint index in this array
-			const int32 ConstrainIndex = ChildTargetParentConstraints.IndexOfByKey(InConstraint);
-			if (ensure(ChildTargetParentConstraints.IsValidIndex(ConstrainIndex)))
-			{
-				ChildTargetParentConstraints.RemoveAt(ConstrainIndex);
-			}
-		
-			if (!ChildTargetParentConstraints.IsEmpty())
-			{
-				static const TCHAR* ExternalDependencyDesc = TEXT("External Dependency");
-			
-				TBitArray<> ManagedDependencies(false, ChildTargetParentConstraints.Num());
-			
-				const FConstraintTickFunction& TickFunction = InConstraint->GetTickFunction(InWorld);
-				const TArray<FTickPrerequisite> PrerexCopy = TickFunction.GetPrerequisites();
-
-				for (int32 Index = 0; Index < ChildTargetParentConstraints.Num(); ++Index)
+				else if (ChildHandle->HasDirectDependencyWith(*TransformConstraint->ChildTRSHandle))
 				{
-					const UTickableTransformConstraint* TransformConstraint = Cast<UTickableTransformConstraint>(ChildTargetParentConstraints[Index]);
-				
-					// if the new handles depend on that constraint child then, TransformConstraint should tick before
-					if (ParentHandle->HasDirectDependencyWith(*TransformConstraint->ChildTRSHandle))
-					{
-						Controller.SetConstraintsDependencies(TransformConstraint->ConstraintID, InConstraint->ConstraintID);
-						ManagedDependencies[Index] = true;
-						LogDependency(ExternalDependencyDesc, TransformConstraint->ChildTRSHandle, ParentHandle, TransformConstraint, InConstraint);
-					}
-					else if (ChildHandle->HasDirectDependencyWith(*TransformConstraint->ChildTRSHandle))
-					{
-						Controller.SetConstraintsDependencies(TransformConstraint->ConstraintID, InConstraint->ConstraintID);
-						ManagedDependencies[Index] = true;
-						LogDependency(ExternalDependencyDesc, TransformConstraint->ChildTRSHandle, ChildHandle, TransformConstraint, InConstraint);
-					}
+					Controller.SetConstraintsDependencies(TransformConstraint->ConstraintID, InConstraint->ConstraintID);
+					ManagedDependencies[Index] = true;
+					LogDependency(ExternalDependencyDesc, TransformConstraint->ChildTRSHandle, ChildHandle, TransformConstraint, InConstraint);
+				}
 
-					// if the TransformConstraint handles depend on the new constraint child then, TransformConstraint should tick after
-					if (TransformConstraint->ParentTRSHandle->HasDirectDependencyWith(*ChildHandle))
-					{
-						Controller.SetConstraintsDependencies(InConstraint->ConstraintID, TransformConstraint->ConstraintID);
-						ManagedDependencies[Index] = true;
-						LogDependency(ExternalDependencyDesc, ChildHandle, TransformConstraint->ParentTRSHandle, InConstraint, TransformConstraint);
-					}
-					else if (TransformConstraint->ChildTRSHandle->HasDirectDependencyWith(*ChildHandle))
-					{
-						Controller.SetConstraintsDependencies(InConstraint->ConstraintID, TransformConstraint->ConstraintID);
-						ManagedDependencies[Index] = true;
-						LogDependency(ExternalDependencyDesc, ChildHandle, TransformConstraint->ChildTRSHandle, InConstraint, TransformConstraint);
-					}
+				// if the TransformConstraint handles depend on the new constraint child then, TransformConstraint should tick after
+				if (TransformConstraint->ParentTRSHandle->HasDirectDependencyWith(*ChildHandle))
+				{
+					Controller.SetConstraintsDependencies(InConstraint->ConstraintID, TransformConstraint->ConstraintID);
+					ManagedDependencies[Index] = true;
+					LogDependency(ExternalDependencyDesc, ChildHandle, TransformConstraint->ParentTRSHandle, InConstraint, TransformConstraint);
 				}
+				else if (TransformConstraint->ChildTRSHandle->HasDirectDependencyWith(*ChildHandle))
+				{
+					Controller.SetConstraintsDependencies(InConstraint->ConstraintID, TransformConstraint->ConstraintID);
+					ManagedDependencies[Index] = true;
+					LogDependency(ExternalDependencyDesc, ChildHandle, TransformConstraint->ChildTRSHandle, InConstraint, TransformConstraint);
+				}
+			}
 
-				// if we didn't add any new prerequisite then check whether the constraint should tick after the last constraint
-				// acting on the same target to respect order of creation.
-				const bool bPrerexChanged = TickFunction.GetPrerequisites() != PrerexCopy;
-				if (!bPrerexChanged && ManagedDependencies.IsValidIndex(ConstrainIndex-1))
+			// if we didn't add any new prerequisite then check whether the constraint should tick after the last constraint
+			// acting on the same target to respect order of creation.
+			const bool bPrerexChanged = TickFunction.GetPrerequisites() != PrerexCopy;
+			if (!bPrerexChanged && ManagedDependencies.IsValidIndex(ConstrainIndex-1))
+			{
+				for (int32 Index = ConstrainIndex-1; Index >= 0; Index--)
 				{
-					for (int32 Index = ConstrainIndex-1; Index >= 0; Index--)
+					if (!ManagedDependencies[Index])
 					{
-						if (!ManagedDependencies[Index])
+						const UTickableTransformConstraint* LastConstraintSharingSameTarget = Cast<UTickableTransformConstraint>(ChildTargetParentConstraints[Index]);
+						TSet<const FTickFunction*> VisitedFunctions;
+						const FTickFunction& ParentTickFunctionToCheck = LastConstraintSharingSameTarget->GetTickFunction(InWorld);
+						if (!FConstraintCycleChecker::HasPrerequisiteDependencyWith(&TickFunction, &ParentTickFunctionToCheck, VisitedFunctions) &&
+							!FConstraintCycleChecker::HasPrerequisiteDependencyWith(&ParentTickFunctionToCheck, &TickFunction, VisitedFunctions))
 						{
-							const UTickableTransformConstraint* LastConstraintSharingSameTarget = Cast<UTickableTransformConstraint>(ChildTargetParentConstraints[Index]);
-							TSet<const FTickFunction*> VisitedFunctions;
-							const FTickFunction& ParentTickFunctionToCheck = LastConstraintSharingSameTarget->GetTickFunction(InWorld);
-							if (!FConstraintCycleChecker::HasPrerequisiteDependencyWith(&TickFunction, &ParentTickFunctionToCheck, VisitedFunctions) &&
-								!FConstraintCycleChecker::HasPrerequisiteDependencyWith(&ParentTickFunctionToCheck, &TickFunction, VisitedFunctions))
+							Controller.SetConstraintsDependencies(LastConstraintSharingSameTarget->ConstraintID, InConstraint->ConstraintID);
+							if (bDebugDependencies)
 							{
-								Controller.SetConstraintsDependencies(LastConstraintSharingSameTarget->ConstraintID, InConstraint->ConstraintID);
-								if (bDebugDependencies)
-								{
-									UE_LOG(LogTemp, Warning, TEXT("Creation Order Dependency: '%s' must tick before '%s' to respect creation order."),
-									   *GetConstraintLabel(LastConstraintSharingSameTarget), *GetConstraintLabel(InConstraint));
-								}
-								break;
+								UE_LOG(LogTemp, Warning, TEXT("Creation Order Dependency: '%s' must tick before '%s' to respect creation order."),
+								   *GetConstraintLabel(LastConstraintSharingSameTarget), *GetConstraintLabel(InConstraint));
 							}
+							break;
 						}
 					}
 				}
 			}
 		}
+	}
 
-		// get all constraints acting on the same target 
+	// get all constraints acting on the same target 
+	{
+		auto SameParentTargetPredicate = [ParentTarget](const ConstraintWeakPtr& Constraint)
 		{
-			auto SameParentTargetPredicate = [ParentTarget](const ConstraintWeakPtr& Constraint)
-			{
-				const UTickableTransformConstraint* TransformConstraint = Cast<UTickableTransformConstraint>(Constraint.Get());
-				if (!TransformConstraint)
-				{
-					return false;
-				}
-				const UObject* OtherChildTarget = GetHandleTarget(TransformConstraint->ChildTRSHandle);
-				return OtherChildTarget && OtherChildTarget == ParentTarget;
-			};
-			TArray< ConstraintWeakPtr > ParentTargetParentConstraints = Controller.GetConstraintsByPredicate(SameParentTargetPredicate);
-
-			// store constraint index in this array
-			const int32 ConstrainIndex = ParentTargetParentConstraints.IndexOfByKey(InConstraint);
-			if (ParentTargetParentConstraints.IsValidIndex(ConstrainIndex))
+			const UTickableTransformConstraint* TransformConstraint = Cast<UTickableTransformConstraint>(Constraint.Get());
+			if (!TransformConstraint)
 			{
-				ParentTargetParentConstraints.RemoveAt(ConstrainIndex);
+				return false;
 			}
+			const UObject* OtherChildTarget = GetHandleTarget(TransformConstraint->ChildTRSHandle);
+			return OtherChildTarget && OtherChildTarget == ParentTarget;
+		};
+		TArray< ConstraintWeakPtr > ParentTargetParentConstraints = Controller.GetConstraintsByPredicate(SameParentTargetPredicate);
+
+		// store constraint index in this array
+		const int32 ConstrainIndex = ParentTargetParentConstraints.IndexOfByKey(InConstraint);
+		if (ParentTargetParentConstraints.IsValidIndex(ConstrainIndex))
+		{
+			ParentTargetParentConstraints.RemoveAt(ConstrainIndex);
+		}
 
-			if (!ParentTargetParentConstraints.IsEmpty())
+		if (!ParentTargetParentConstraints.IsEmpty())
+		{
+			static const TCHAR* ExternalDependencyDesc = TEXT("External Dependency");
+		
+			TBitArray<> ManagedDependencies(false, ParentTargetParentConstraints.Num());
+		
+			const FConstraintTickFunction& TickFunction = InConstraint->GetTickFunction(InWorld);
+			const TArray<FTickPrerequisite> PrerexCopy = TickFunction.GetPrerequisites();
+
+			for (int32 Index = 0; Index < ParentTargetParentConstraints.Num(); ++Index)
 			{
-				static const TCHAR* ExternalDependencyDesc = TEXT("External Dependency");
-			
-				TBitArray<> ManagedDependencies(false, ParentTargetParentConstraints.Num());
+				const UTickableTransformConstraint* TransformConstraint = Cast<UTickableTransformConstraint>(ParentTargetParentConstraints[Index]);
 			
-				const FConstraintTickFunction& TickFunction = InConstraint->GetTickFunction(InWorld);
-				const TArray<FTickPrerequisite> PrerexCopy = TickFunction.GetPrerequisites();
-
-				for (int32 Index = 0; Index < ParentTargetParentConstraints.Num(); ++Index)
+				// if the new handles depend on that constraint child then, TransformConstraint should tick before
+				if (ParentHandle->HasDirectDependencyWith(*TransformConstraint->ChildTRSHandle))
 				{
-					const UTickableTransformConstraint* TransformConstraint = Cast<UTickableTransformConstraint>(ParentTargetParentConstraints[Index]);
-				
-					// if the new handles depend on that constraint child then, TransformConstraint should tick before
-					if (ParentHandle->HasDirectDependencyWith(*TransformConstraint->ChildTRSHandle))
-					{
-						Controller.SetConstraintsDependencies(TransformConstraint->ConstraintID, InConstraint->ConstraintID);
-						ManagedDependencies[Index] = true;
-						LogDependency(ExternalDependencyDesc, TransformConstraint->ChildTRSHandle, ParentHandle, TransformConstraint, InConstraint);
-					}
+					Controller.SetConstraintsDependencies(TransformConstraint->ConstraintID, InConstraint->ConstraintID);
+					ManagedDependencies[Index] = true;
+					LogDependency(ExternalDependencyDesc, TransformConstraint->ChildTRSHandle, ParentHandle, TransformConstraint, InConstraint);
 				}
+			}
 
-				// if we didn't add any new prerequisite then check whether the constraint should tick after the last constraint
-				// acting on the same target to respect order of creation.
-				const bool bPrerexChanged = TickFunction.GetPrerequisites() != PrerexCopy;
-				const int32 LastIndex = ManagedDependencies.IsValidIndex(ConstrainIndex-1) ? ConstrainIndex-1 : ParentTargetParentConstraints.Num()-1;
-				if (!bPrerexChanged && ManagedDependencies.IsValidIndex(LastIndex))
+			// if we didn't add any new prerequisite then check whether the constraint should tick after the last constraint
+			// acting on the same target to respect order of creation.
+			const bool bPrerexChanged = TickFunction.GetPrerequisites() != PrerexCopy;
+			const int32 LastIndex = ManagedDependencies.IsValidIndex(ConstrainIndex-1) ? ConstrainIndex-1 : ParentTargetParentConstraints.Num()-1;
+			if (!bPrerexChanged && ManagedDependencies.IsValidIndex(LastIndex))
+			{
+				for (int32 Index = LastIndex; Index >= 0; Index--)
 				{
-					for (int32 Index = LastIndex; Index >= 0; Index--)
+					if (!ManagedDependencies[Index])
 					{
-						if (!ManagedDependencies[Index])
+						const UTickableTransformConstraint* LastConstraintSharingSameTarget = Cast<UTickableTransformConstraint>(ParentTargetParentConstraints[Index]);
+						TSet<const FTickFunction*> VisitedFunctions;
+						const FTickFunction& ParentTickFunctionToCheck = LastConstraintSharingSameTarget->GetTickFunction(InWorld);
+						if (!FConstraintCycleChecker::HasPrerequisiteDependencyWith(&TickFunction, &ParentTickFunctionToCheck, VisitedFunctions) &&
+							!FConstraintCycleChecker::HasPrerequisiteDependencyWith(&ParentTickFunctionToCheck, &TickFunction, VisitedFunctions))
 						{
-							const UTickableTransformConstraint* LastConstraintSharingSameTarget = Cast<UTickableTransformConstraint>(ParentTargetParentConstraints[Index]);
-							TSet<const FTickFunction*> VisitedFunctions;
-							const FTickFunction& ParentTickFunctionToCheck = LastConstraintSharingSameTarget->GetTickFunction(InWorld);
-							if (!FConstraintCycleChecker::HasPrerequisiteDependencyWith(&TickFunction, &ParentTickFunctionToCheck, VisitedFunctions) &&
-								!FConstraintCycleChecker::HasPrerequisiteDependencyWith(&ParentTickFunctionToCheck, &TickFunction, VisitedFunctions))
+							Controller.SetConstraintsDependencies(LastConstraintSharingSameTarget->ConstraintID, InConstraint->ConstraintID);
+							if (bDebugDependencies)
 							{
-								Controller.SetConstraintsDependencies(LastConstraintSharingSameTarget->ConstraintID, InConstraint->ConstraintID);
-								if (bDebugDependencies)
-								{
-									UE_LOG(LogTemp, Warning, TEXT("Creation Order Dependency: '%s' must tick before '%s' to respect creation order."),
-									   *GetConstraintLabel(LastConstraintSharingSameTarget), *GetConstraintLabel(InConstraint));
-								}
-								break;
+								UE_LOG(LogTemp, Warning, TEXT("Creation Order Dependency: '%s' must tick before '%s' to respect creation order."),
+								   *GetConstraintLabel(LastConstraintSharingSameTarget), *GetConstraintLabel(InConstraint));
 							}
+							break;
 						}
 					}
 				}
 			}
 		}
 	}
+	
+	// get all constraints having the same parent's target as the new child's target and make sure they tick after if needed 
+	{
+		auto SameParentAsChildTargetPredicate = [ChildTarget, InConstraint](const ConstraintWeakPtr& Constraint)
+		{
+			const UTickableTransformConstraint* TransformConstraint = Cast<UTickableTransformConstraint>(Constraint.Get());
+			if (!TransformConstraint || TransformConstraint == InConstraint)
+			{
+				return false;
+			}
+			const UObject* OtherParentTarget = GetHandleTarget(TransformConstraint->ParentTRSHandle);
+			return OtherParentTarget && OtherParentTarget == ChildTarget;
+		};
+		TArray< ConstraintWeakPtr > ParentTargetChildConstraints = Controller.GetConstraintsByPredicate(SameParentAsChildTargetPredicate);
+
+		const FConstraintTickFunction& TickFunction = InConstraint->GetTickFunction(InWorld);
+		for (const ConstraintWeakPtr& Constraint: ParentTargetChildConstraints)
+		{
+			const UTickableTransformConstraint* ConstraintSharingParentTarget = Cast<UTickableTransformConstraint>(Constraint);
+			const FTickFunction& ParentTickFunctionToCheck = ConstraintSharingParentTarget->GetTickFunction(InWorld);
+
+			TSet<const FTickFunction*> VisitedFunctions0, VisitedFunctions1;
+			
+			if (!FConstraintCycleChecker::HasPrerequisiteDependencyWith(&TickFunction, &ParentTickFunctionToCheck, VisitedFunctions0) &&
+				!FConstraintCycleChecker::HasPrerequisiteDependencyWith(&ParentTickFunctionToCheck, &TickFunction, VisitedFunctions1))
+			{
+				Controller.SetConstraintsDependencies(InConstraint->ConstraintID, ConstraintSharingParentTarget->ConstraintID);
+				if (bDebugDependencies)
+				{
+					UE_LOG(LogTemp, Warning, TEXT("External Dependency: '%s' must tick before '%s' as it acts on its parent."),
+					   *GetConstraintLabel(InConstraint), *GetConstraintLabel(ConstraintSharingParentTarget));
+				}
+			}
+		}
+	}
+}
+
+bool FDependencyBuilder::BuildDependencies(UWorld* InWorld, UTickableTransformConstraint* InConstraint)
+{
+	using namespace DependencyLocals;
+	using ConstraintWeakPtr = TWeakObjectPtr<UTickableConstraint>;
+	
+	if (!ensure(InWorld))
+	{
+		return false;
+	}
+
+	if (!InConstraint || !InConstraint->IsValid())
+	{
+		return false;
+	}
+
+	if (bDebugDependencies)
+	{
+		UE_LOG(LogTemp, Warning, TEXT("Building dependencies for '%s' ..."), *GetConstraintLabel(InConstraint));
+	}
+
+	const UTransformableHandle* ParentHandle = InConstraint->ParentTRSHandle.Get();
+	const UTransformableHandle* ChildHandle = InConstraint->ChildTRSHandle.Get();
+	
+ 	// get previous child constraints
+	const FConstraintsManagerController& Controller = FConstraintsManagerController::Get(InWorld);
+	TArray<ConstraintWeakPtr> ChildParentConstraints = Controller.GetParentConstraints(ChildHandle->GetHash(), true);
+	ChildParentConstraints.Remove(InConstraint);
+
+	// add dependencies with the last child constraint
+	if (!ChildParentConstraints.IsEmpty())
+	{
+		const FGuid LastChildConstraintID = ChildParentConstraints.Last()->ConstraintID;
+		if (bDebugDependencies)
+		{
+			UE_LOG(LogTemp, Warning, TEXT("Order Dependency: tick after last constraint."));
+		}
+		Controller.SetConstraintsDependencies(LastChildConstraintID, InConstraint->ConstraintID);
+	}
+
+	const UObject* ParentTarget = ParentHandle->GetTarget().Get();
+	const UObject* ChildTarget = ChildHandle->GetTarget().Get();
+	const bool bSelf = ParentTarget && ParentTarget == ChildTarget;
+	
+	// internal dependencies?
+	if (bSelf)
+	{
+		BuildSelfDependencies(InWorld, InConstraint);
+	}
+
+	// make sure we tick after the parent.
+	InConstraint->EnsurePrimaryDependency(InWorld);
+
+	// if child handle is the parent of some other constraints, ensure they will tick after that new one
+	static const TCHAR* ChildDependencyDesc = TEXT("Child Dependency");
+
+	TArray<ConstraintWeakPtr> ChildChildConstraints;
+	UE::TransformConstraintUtil::GetChildrenConstraints(InWorld, InConstraint, ChildChildConstraints, !bSelf);
+	for (const ConstraintWeakPtr& ChildConstraint: ChildChildConstraints)
+	{
+		Controller.SetConstraintsDependencies(InConstraint->ConstraintID, ChildConstraint->ConstraintID);
+		if (const UTickableTransformConstraint* TransformConstraint = Cast<UTickableTransformConstraint>(ChildConstraint))
+		{
+			LogDependency(ChildDependencyDesc, ChildHandle, TransformConstraint->ChildTRSHandle, InConstraint, TransformConstraint);
+		}
+	}
+
+	// build dependencies regarding attachments
+	BuildAttachmentsDependencies(InWorld, InConstraint);
+
+	if (!bSelf)
+	{
+		BuildExternalDependencies(InWorld, InConstraint);
+	}
 
 	// warn for possible cycles
 	if (FConstraintCycleChecker::IsCycling(InConstraint->ChildTRSHandle))
diff --git a/Engine/Source/Runtime/Experimental/Animation/Constraints/Public/Transform/DependencyBuilder.h b/Engine/Source/Runtime/Experimental/Animation/Constraints/Public/Transform/DependencyBuilder.h
--- a/Engine/Source/Runtime/Experimental/Animation/Constraints/Public/Transform/DependencyBuilder.h
+++ b/Engine/Source/Runtime/Experimental/Animation/Constraints/Public/Transform/DependencyBuilder.h
@@ -28,6 +28,12 @@ private:
 	
 	/** Ensures that attachment dependencies are reflected at the constraints level. */
 	static void BuildAttachmentsDependencies(UWorld* InWorld, const UTickableTransformConstraint* InConstraint);
+
+	/** Ensures that internal dependencies (control rig only at this point) are addressed at the constraints level. */
+	static void BuildSelfDependencies(UWorld* InWorld, UTickableTransformConstraint* InConstraint);
+
+	/** Ensures that external dependencies are addressed at the constraints level. */
+	static void BuildExternalDependencies(UWorld* InWorld, UTickableTransformConstraint* InConstraint);
 };
 
 /**
