iostore.PurgeInstallCache new arguments
- optionally, specify the number of bytes to be purged, otherwise purges as much as possible
- "defrag" will attempt to defrag cache after purging

iostore.DefragInstallCache
- optionally, specify the number of bytes to be freed through defrag, otherwise defrags as many blocks as possible.

diff --git a/Engine/Source/Runtime/Experimental/IoStore/OnDemand/Private/OnDemandInstallCache.cpp b/Engine/Source/Runtime/Experimental/IoStore/OnDemand/Private/OnDemandInstallCache.cpp
--- a/Engine/Source/Runtime/Experimental/IoStore/OnDemand/Private/OnDemandInstallCache.cpp
+++ b/Engine/Source/Runtime/Experimental/IoStore/OnDemand/Private/OnDemandInstallCache.cpp
@@ -1326,8 +1326,8 @@ public:
 	virtual bool								IsChunkCached(const FIoHash& ChunkHash) override;
 	virtual FIoStatus							PutChunk(FIoBuffer&& Chunk, const FIoHash& ChunkHash) override;
 	virtual FIoStatus							Purge(TSet<FIoHash>&& ChunksToInstall) override;
-	virtual FIoStatus							PurgeAllUnreferenced() override;
-	virtual FIoStatus							DefragAll() override;
+	virtual FIoStatus							PurgeAllUnreferenced(bool bDefrag, const uint64* BytesToPurge = nullptr) override;
+	virtual FIoStatus							DefragAll(const uint64* BytesToFree = nullptr) override;
 	virtual FIoStatus							Flush() override;
 	virtual FOnDemandInstallCacheUsage			GetCacheUsage() override;
 
@@ -1899,7 +1899,7 @@ FIoStatus FOnDemandInstallCache::Purge(TSet<FIoHash>&& ChunksToInstall)
 	return Status;
 }
 
-FIoStatus FOnDemandInstallCache::PurgeAllUnreferenced()
+FIoStatus FOnDemandInstallCache::PurgeAllUnreferenced(bool bDefrag, const uint64* BytesToPurge /*= nullptr*/)
 {
 	FCasBlockInfoMap	BlockInfo;
 	const uint64		TotalCachedBytes = Cas.GetBlockInfo(BlockInfo);
@@ -1912,8 +1912,6 @@ FIoStatus FOnDemandInstallCache::PurgeAllUnreferenced()
 
 	AddReferencesToBlocks(Containers, ChunkEntryIndices, {}, BlockInfo);
 
-	//TODO: Compute fragmentation metric and redownload chunks when this number gets too high
-
 	const uint64 TotalReferencedBytes = Algo::TransformAccumulate(BlockInfo,
 		[](const TPair<FCasBlockId, FCasBlockInfo>& Kv) { return (Kv.Value.RefSize > 0) ? Kv.Value.FileSize : uint64(0); },
 		uint64(0));
@@ -1921,10 +1919,9 @@ FIoStatus FOnDemandInstallCache::PurgeAllUnreferenced()
 	UE_LOG(LogIoStoreOnDemand, Log, TEXT("Purging install cache, MaxCacheSize=%.2lf MiB, CacheSize=%.2lf MiB, ReferencedBytes=%.2lf MiB"),
 		ToMiB(MaxCacheSize), ToMiB(TotalCachedBytes), ToMiB(TotalReferencedBytes));
 
-	uint64 TotalPurgedBytes = 0;
-	FIoStatus Status = Purge(BlockInfo, MaxCacheSize, TotalPurgedBytes);
-
-	// TODO: defrag as well?
+	const uint64	TotalBytesToPurge	= BytesToPurge ? *BytesToPurge : MaxCacheSize;
+	uint64			TotalPurgedBytes	= 0;
+	FIoStatus Status = Purge(BlockInfo, TotalBytesToPurge, TotalPurgedBytes);
 
 	if (TotalPurgedBytes > 0)
 	{
@@ -1937,10 +1934,38 @@ FIoStatus FOnDemandInstallCache::PurgeAllUnreferenced()
 		LogIoStoreOnDemand, Warning, TEXT("Max install cache size exceeded by %.2lf MiB (%.2lf%%)"),
 			ToMiB(NewCachedBytes - MaxCacheSize), 100.0 * (double(NewCachedBytes - MaxCacheSize) / double(MaxCacheSize)));
 
+	if (BytesToPurge)
+	{
+		if (bDefrag)
+		{
+			// Attempt to defrag
+			const uint64 DefragBytesToPurge = TotalBytesToPurge - TotalPurgedBytes;
+			FIoStatus DefragStatus = Defrag(Containers, ChunkEntryIndices, BlockInfo, &DefragBytesToPurge);
+			if (DefragStatus.IsOk() == false)
+			{
+				return FIoStatusBuilder(EIoErrorCode::WriteError)
+					<< FString::Printf(TEXT("Failed to purge %" UINT64_FMT " from install cache after defrag (%s)"), TotalBytesToPurge, *DefragStatus.ToString());
+			}
+		}
+		else
+		{
+			return FIoStatusBuilder(EIoErrorCode::WriteError) << FString::Printf(TEXT("Failed to purge %" UINT64_FMT " from install cache"), TotalBytesToPurge);
+		}
+	}
+	else if (bDefrag)
+	{
+		// Just do Full Defrag
+		FIoStatus DefragStatus = Defrag(Containers, ChunkEntryIndices, BlockInfo);
+		if (DefragStatus.IsOk() == false)
+		{
+			return DefragStatus;
+		}
+	}
+
 	return Status;
 }
 
-FIoStatus FOnDemandInstallCache::DefragAll()
+FIoStatus FOnDemandInstallCache::DefragAll(const uint64* BytesToFree /*= nullptr*/)
 {
 	FCasBlockInfoMap	BlockInfo;
 	const uint64		TotalCachedBytes = Cas.GetBlockInfo(BlockInfo);
@@ -1960,7 +1985,7 @@ FIoStatus FOnDemandInstallCache::DefragAll()
 	UE_LOG(LogIoStoreOnDemand, Log, TEXT("Defragmenting install cache, MaxCacheSize=%.2lf MiB, CacheSize=%.2lf MiB, ReferencedBytes=%.2lf MiB"),
 		ToMiB(MaxCacheSize), ToMiB(TotalCachedBytes), ToMiB(TotalReferencedBytes));
 
-	return Defrag(Containers, ChunkEntryIndices, BlockInfo);
+	return Defrag(Containers, ChunkEntryIndices, BlockInfo, BytesToFree);
 }
 
 uint64 FOnDemandInstallCache::AddReferencesToBlocks(
@@ -2260,6 +2285,8 @@ FIoStatus FOnDemandInstallCache::Defrag(
 		}
 	}
 
+	UE_LOG(LogIoStoreOnDemand, Display, TEXT("Defrag removed %" UINT64_FMT " fragmented bytes of %" UINT64_FMT " total bytes in %i blocks."), FragmentedBytes, TotalBlockSize, BlocksToDefrag.Num());
+
 	return FIoStatus::Ok;
 }
 
diff --git a/Engine/Source/Runtime/Experimental/IoStore/OnDemand/Private/OnDemandInstallCache.h b/Engine/Source/Runtime/Experimental/IoStore/OnDemand/Private/OnDemandInstallCache.h
--- a/Engine/Source/Runtime/Experimental/IoStore/OnDemand/Private/OnDemandInstallCache.h
+++ b/Engine/Source/Runtime/Experimental/IoStore/OnDemand/Private/OnDemandInstallCache.h
@@ -30,8 +30,8 @@ public:
 	virtual bool								IsChunkCached(const FIoHash& ChunkHash) = 0;
 	virtual FIoStatus							PutChunk(FIoBuffer&& Chunk, const FIoHash& ChunkHash) = 0;
 	virtual FIoStatus							Purge(TSet<FIoHash>&& ChunksToInstall) = 0;
-	virtual FIoStatus							PurgeAllUnreferenced() = 0;
-	virtual FIoStatus							DefragAll() = 0;
+	virtual FIoStatus							PurgeAllUnreferenced(bool bDefrag, const uint64* BytesToPurge = nullptr) = 0;
+	virtual FIoStatus							DefragAll(const uint64* BytesToFree = nullptr) = 0;
 	virtual FIoStatus							Flush() = 0;
 	virtual FOnDemandInstallCacheUsage			GetCacheUsage() = 0;
 };
diff --git a/Engine/Source/Runtime/Experimental/IoStore/OnDemand/Private/OnDemandIoStore.cpp b/Engine/Source/Runtime/Experimental/IoStore/OnDemand/Private/OnDemandIoStore.cpp
--- a/Engine/Source/Runtime/Experimental/IoStore/OnDemand/Private/OnDemandIoStore.cpp
+++ b/Engine/Source/Runtime/Experimental/IoStore/OnDemand/Private/OnDemandIoStore.cpp
@@ -1206,32 +1206,32 @@ bool FOnDemandIoStore::Tick()
 	}
 	
 	// Tick purge requests
-	if (LocalPurgeRequests.IsEmpty() == false)
+	for (FSharedPurgeRequest& Request : LocalPurgeRequests)
 	{
-		FOnDemandPurgeResult PurgeResult;
+		const bool bDefrag = EnumHasAnyFlags(Request->Args.Options, EOnDemandPurgeOptions::Defrag);
+		const uint64* BytesToPurge = Request->Args.BytesToPurge.GetPtrOrNull();
 
 		const double StartTime = FPlatformTime::Seconds();
-		PurgeResult.Status = InstallCache->PurgeAllUnreferenced();
+		FOnDemandPurgeResult PurgeResult;
+		PurgeResult.Status = InstallCache->PurgeAllUnreferenced(bDefrag, BytesToPurge);
 		PurgeResult.DurationInSeconds = FPlatformTime::Seconds() - StartTime;
 
-		for (FSharedPurgeRequest& Request : LocalPurgeRequests)
-		{
-			CompletePurgeRequest(*Request, FOnDemandPurgeResult(PurgeResult));
-		}
+		CompletePurgeRequest(*Request, MoveTemp(PurgeResult));
 	}
 
 	// Tick defrag requests
 	if (LocalDefragRequests.IsEmpty() == false)
 	{
-		FOnDemandDefragResult DefragResult;
-
-		const double StartTime = FPlatformTime::Seconds();
-		DefragResult.Status = InstallCache->DefragAll();
-		DefragResult.DurationInSeconds = FPlatformTime::Seconds() - StartTime;
-
 		for (FSharedDefragRequest& Request : LocalDefragRequests)
 		{
-			CompleteDefragRequest(*Request, FOnDemandDefragResult(DefragResult));
+			const uint64* BytesToFree = Request->Args.BytesToFree.GetPtrOrNull();
+			
+			const double StartTime = FPlatformTime::Seconds();
+			FOnDemandDefragResult DefragResult;
+			DefragResult.Status = InstallCache->DefragAll(BytesToFree);
+			DefragResult.DurationInSeconds = FPlatformTime::Seconds() - StartTime;
+
+			CompleteDefragRequest(*Request, MoveTemp(DefragResult));
 		}
 	}
 
diff --git a/Engine/Source/Runtime/Experimental/IoStore/OnDemandCore/Private/IoStoreOnDemand.cpp b/Engine/Source/Runtime/Experimental/IoStore/OnDemandCore/Private/IoStoreOnDemand.cpp
--- a/Engine/Source/Runtime/Experimental/IoStore/OnDemandCore/Private/IoStoreOnDemand.cpp
+++ b/Engine/Source/Runtime/Experimental/IoStore/OnDemandCore/Private/IoStoreOnDemand.cpp
@@ -7,12 +7,12 @@
 #include "HAL/FileManager.h"
 #include "HAL/IConsoleManager.h"
 #include "IO/IoStatus.h"
-#include "IO/IoStoreOnDemand.h"
 #include "Misc/Paths.h"
 #include "Misc/PathViews.h"
 #include "Misc/Parse.h"
 #include "Modules/ModuleManager.h"
-#include "IO/IoStoreOnDemand.h"
+#include "String/Numeric.h"
+
 #if !UE_BUILD_SHIPPING
 #include "Misc/PackageName.h"
 #endif //!UE_BUILD_SHIPPING
@@ -57,12 +57,29 @@ static TIoStatusOr<FOnDemandHostGroup> SplitHostAndTocUrl(FStringView Url, FStri
 }
 
 ////////////////////////////////////////////////////////////////////////////////
-static void PurgeInstallCache()
+static void PurgeInstallCache(const TArray<FString>& Args)
 {
 	if (UE::IoStore::IOnDemandIoStore* IoStore = UE::IoStore::TryGetOnDemandIoStore())
 	{
+		FOnDemandPurgeArgs PurgeArgs;
+		for (const FString& Arg : Args)
+		{
+			if (Arg == TEXTVIEW("defrag"))
+			{
+				PurgeArgs.Options |= EOnDemandPurgeOptions::Defrag;
+			}
+			else if (!PurgeArgs.BytesToPurge && UE::String::IsNumericOnlyDigits(Arg))
+			{
+				uint64 BytesToPurge = 0;
+				if (LexTryParseString(BytesToPurge, *Arg))
+				{
+					PurgeArgs.BytesToPurge = BytesToPurge;
+				}
+			}
+		}
+
 		UE_LOG(LogIoStoreOnDemand, Display, TEXT("Purging on demand install cache"));
-		IoStore->Purge(FOnDemandPurgeArgs(), [](const FOnDemandPurgeResult& Result)
+		IoStore->Purge(MoveTemp(PurgeArgs), [](const FOnDemandPurgeResult& Result)
 		{
 			if (Result.Status.IsOk())
 			{
@@ -81,12 +98,26 @@ static void PurgeInstallCache()
 }
 
 ////////////////////////////////////////////////////////////////////////////////
-static void DefragInstallCache()
+static void DefragInstallCache(const TArray<FString>& Args)
 {
 	if (UE::IoStore::IOnDemandIoStore* IoStore = UE::IoStore::TryGetOnDemandIoStore())
 	{
+		FOnDemandDefragArgs DefragArgs;
+		for (const FString& Arg : Args)
+		{
+			if (UE::String::IsNumericOnlyDigits(Arg))
+			{
+				uint64 BytesToFree = 0;
+				if (LexTryParseString(BytesToFree, *Arg))
+				{
+					DefragArgs.BytesToFree = BytesToFree;
+					break;
+				}
+			}
+		}
+
 		UE_LOG(LogIoStoreOnDemand, Display, TEXT("Defragging on demand install cache"));
-		IoStore->Defrag(FOnDemandDefragArgs(), [](const FOnDemandDefragResult& Result)
+		IoStore->Defrag(MoveTemp(DefragArgs), [](const FOnDemandDefragResult& Result)
 		{
 			if (Result.Status.IsOk())
 			{
@@ -331,13 +362,13 @@ static void InstallPackage(const TArray<FString>& Args)
 static FAutoConsoleCommand PurgeCacheCommand(
 	TEXT("iostore.PurgeInstallCache"),
 	TEXT("Purge On Demand Install Cache"),
-	FConsoleCommandDelegate::CreateStatic(PurgeInstallCache),
+	FConsoleCommandWithArgsDelegate::CreateStatic(PurgeInstallCache),
 	ECVF_Cheat);
 
 static FAutoConsoleCommand DefragCacheCommand(
 	TEXT("iostore.DefragInstallCache"),
 	TEXT("Defragment On Demand Install Cache"),
-	FConsoleCommandDelegate::CreateStatic(DefragInstallCache),
+	FConsoleCommandWithArgsDelegate::CreateStatic(DefragInstallCache),
 	ECVF_Cheat);
 
 static FAutoConsoleCommand PrintCacheUsageCommand(
diff --git a/Engine/Source/Runtime/Experimental/IoStore/OnDemandCore/Public/IO/IoStoreOnDemand.h b/Engine/Source/Runtime/Experimental/IoStore/OnDemandCore/Public/IO/IoStoreOnDemand.h
--- a/Engine/Source/Runtime/Experimental/IoStore/OnDemandCore/Public/IO/IoStoreOnDemand.h
+++ b/Engine/Source/Runtime/Experimental/IoStore/OnDemandCore/Public/IO/IoStoreOnDemand.h
@@ -313,7 +313,7 @@ enum class EOnDemandPurgeOptions
 	/** Trigger callback on game thread. */
 	CallbackOnGameThread = 1 << 0,
 	/** Defrag after purging */
-	Defrag = 1 << 1, // TODO:
+	Defrag = 1 << 1,
 };
 ENUM_CLASS_FLAGS(EOnDemandPurgeOptions)
 
@@ -322,7 +322,8 @@ struct FOnDemandPurgeArgs
 {
 	/** Purge options. */
 	EOnDemandPurgeOptions Options = EOnDemandPurgeOptions::None;
-	// TODO: allow passing in the amount to purge
+	/** Optional Size to purge. If not set will purge all unfreferenced blocks. */
+	TOptional<uint64> BytesToPurge;
 };
 
 /** Holds information about a purge request */
@@ -352,7 +353,8 @@ struct FOnDemandDefragArgs
 {
 	/** Defrag options. */
 	EOnDemandDefragOptions Options = EOnDemandDefragOptions::None;
-	// TODO: allow passing in the amount to defrag
+	/** Optional Size to free. If not set, will degrag all blocks and free all unreferenced chunks. */
+	TOptional<uint64> BytesToFree;
 };
 
 /** Holds information about a defrag request */
