Iris
* Print ReplicationWriter state in Net.Iris.PrintDebugInfoOfNetHandle even if the object is not currently filtered

#rnx

diff --git a/Engine/Source/Runtime/Experimental/Iris/Core/Private/Iris/ReplicationSystem/ObjectReplicationBridgeDebugging.cpp b/Engine/Source/Runtime/Experimental/Iris/Core/Private/Iris/ReplicationSystem/ObjectReplicationBridgeDebugging.cpp
--- a/Engine/Source/Runtime/Experimental/Iris/Core/Private/Iris/ReplicationSystem/ObjectReplicationBridgeDebugging.cpp
+++ b/Engine/Source/Runtime/Experimental/Iris/Core/Private/Iris/ReplicationSystem/ObjectReplicationBridgeDebugging.cpp
@@ -1295,32 +1295,34 @@ void UObjectReplicationBridge::PrintDebugInfoForNetRefHandle(const FNetRefHandle
 		{
 			if (Filtering.GetRelevantObjectsInScope(ConnectionId).IsBitSet(ObjectIndex))
 			{
-				if (const FReplicationWriter* ReplicationWriter = Connection->ReplicationWriter)
-				{
-					// Replication Writer Info
-					UE_LOG(LogIris, Display, TEXT("\tReplicationWriter: %s"), *ReplicationWriter->PrintObjectInfo(ObjectIndex));
-				}
-				else
-				{
-					UE_LOG(LogIris, Display, TEXT("\tReplicationWriter: not found ?"));
-				}
+				UE_LOG(LogIris, Display, TEXT("\tRelevant to connection: %d"), ConnectionId);
+			}
+			else
+			{
+				UE_LOG(LogIris, Display, TEXT("\tNot relevant to connection: %d"), ConnectionId);
+			}
 
-				// Prioritizer
-				UE_LOG(LogIris, Display, TEXT("\tPrioritizer: Priority: %f"), Prioritization.GetObjectPriorityForConnection(ConnectionId, ObjectIndex));
+			if (const FReplicationWriter* ReplicationWriter = Connection->ReplicationWriter)
+			{
+				// Replication Writer Info
+				UE_LOG(LogIris, Display, TEXT("\tReplicationWriter: %s"), *ReplicationWriter->PrintObjectInfo(ObjectIndex));
+			}
+			else
+			{
+				UE_LOG(LogIris, Display, TEXT("\tReplicationWriter: not found ?"));
+			}
 
-				// DeltaCompression info
-				if (DeltaManager.GetDeltaCompressionStatus(ObjectIndex) == ENetObjectDeltaCompressionStatus::Allow)
-				{
-					UE_LOG(LogIris, Display, TEXT("\tDeltaCompression: %s"), *DeltaManager.PrintDeltaCompressionStatus(ConnectionId, ObjectIndex));
-				}
-				else
-				{
-					UE_LOG(LogIris, Display, TEXT("\tDeltaCompression: not enabled"));
-				}
+			// Prioritizer
+			UE_LOG(LogIris, Display, TEXT("\tPrioritizer: Priority: %f"), Prioritization.GetObjectPriorityForConnection(ConnectionId, ObjectIndex));
+
+			// DeltaCompression info
+			if (DeltaManager.GetDeltaCompressionStatus(ObjectIndex) == ENetObjectDeltaCompressionStatus::Allow)
+			{
+				UE_LOG(LogIris, Display, TEXT("\tDeltaCompression: %s"), *DeltaManager.PrintDeltaCompressionStatus(ConnectionId, ObjectIndex));
 			}
 			else
 			{
-				UE_LOG(LogIris, Display, TEXT("\tNot relevant to connection %d"), ConnectionId);
+				UE_LOG(LogIris, Display, TEXT("\tDeltaCompression: not enabled"));
 			}
 		}
 		else
