[NetPhysics] Fix GC post-resim render interpolation issues.
- Call FRewindData::BufferPhysicsResults after GC and CU has run their BufferPhysicsResults_Internal since they alter (child) particle state which is needed before RewindData calculates the post-resim error.
- Cache CU child states pre-resim and calculate post-resim error for each child
- Apply post-resim rotational error in the local space of the particles rotation

diff --git a/Engine/Source/Runtime/Experimental/Chaos/Private/Framework/PhysicsProxyBase.cpp b/Engine/Source/Runtime/Experimental/Chaos/Private/Framework/PhysicsProxyBase.cpp
--- a/Engine/Source/Runtime/Experimental/Chaos/Private/Framework/PhysicsProxyBase.cpp
+++ b/Engine/Source/Runtime/Experimental/Chaos/Private/Framework/PhysicsProxyBase.cpp
@@ -25,7 +25,10 @@ namespace RenderInterpolationCVars
 
 	bool bRenderInterpDebugDraw = false;
 	static FAutoConsoleVariableRef CVarRenderInterpDebugDraw(TEXT("p.RenderInterp.DebugDraw"), bRenderInterpDebugDraw, TEXT("Draw debug lines for physics render interpolation, also needs p.Chaos.DebugDraw.Enabled set"));
-	
+
+	bool bRenderInterpDebugDrawGC = false;
+	static FAutoConsoleVariableRef CVarRenderInterpDebugDrawGC(TEXT("p.RenderInterp.DebugDraw.GC"), bRenderInterpDebugDrawGC, TEXT("Enable render interpolation debug draw for Geometry Collections"));
+
 	float RenderInterpDebugDrawZOffset = 0.0f;
 	static FAutoConsoleVariableRef CVarRenderInterpDebugDrawZOffset(TEXT("p.RenderInterp.DebugDrawZOffset"), RenderInterpDebugDrawZOffset, TEXT("Add Z axis offset to DebugDraw calls for Render Interpolation."));
 }
diff --git a/Engine/Source/Runtime/Experimental/Chaos/Private/PBDRigidsSolver.cpp b/Engine/Source/Runtime/Experimental/Chaos/Private/PBDRigidsSolver.cpp
--- a/Engine/Source/Runtime/Experimental/Chaos/Private/PBDRigidsSolver.cpp
+++ b/Engine/Source/Runtime/Experimental/Chaos/Private/PBDRigidsSolver.cpp
@@ -2297,12 +2297,6 @@ namespace Chaos
 			}
 		}
 
-		// Move post-resim error correction data from RewindData to FPullPhysicsData for marshaling
-		if (FRewindData* RewindData = GetRewindData())
-		{
-			RewindData->BufferPhysicsResults(PullData->DirtyRigidErrors);
-		}
-
 		{
 			ensure(PullData->DirtyGeometryCollections.Num() == 0);	//we only fill this once per frame
 			PullData->DirtyGeometryCollections.Reserve(ActiveGC.Num());
@@ -2347,12 +2341,16 @@ namespace Chaos
 			}
 		}
 
+		// Calculate and move post-resim error correction data from RewindData to FPullPhysicsData for marshaling
+		if (FRewindData* RewindData = GetRewindData())
+		{
+			RewindData->BufferPhysicsResults(PullData->DirtyRigidErrors);
+		}
+
 		// Now that results have been buffered we have completed a solve step so we can broadcast that event
 		EventPostSolve.Broadcast(MLastDt);
 
-
 		Particles.ClearTransientDirty();
-
 	}
 
 	void FPBDRigidsSolver::BeginDestroy()
diff --git a/Engine/Source/Runtime/Experimental/Chaos/Private/PhysicsProxy/ClusterUnionPhysicsProxy.cpp b/Engine/Source/Runtime/Experimental/Chaos/Private/PhysicsProxy/ClusterUnionPhysicsProxy.cpp
--- a/Engine/Source/Runtime/Experimental/Chaos/Private/PhysicsProxy/ClusterUnionPhysicsProxy.cpp
+++ b/Engine/Source/Runtime/Experimental/Chaos/Private/PhysicsProxy/ClusterUnionPhysicsProxy.cpp
@@ -619,9 +619,9 @@ namespace Chaos
 			if (const FQuat* Prev = LerpHelper(PullData.R, ProxyTimestamp->OverWriteR))
 			{
 				const FQuat NewR = bIsReplicationErrorSmoothing ? 
-					InterpData->GetErrorR(*Alpha) * FMath::Lerp(*Prev, NextPullData->R, *Alpha) : 
-					FMath::Lerp(*Prev, NextPullData->R, *Alpha);
-		
+					FMath::Lerp(*Prev, NextPullData->R, *Alpha) * InterpData->GetErrorR(*Alpha) :
+					FMath::Lerp(*Prev, NextPullData->R, *Alpha); // Add rotational error offset in local space
+
 				Particle_External->SetR(NewR, false);
 			}
 
@@ -650,7 +650,7 @@ namespace Chaos
 					if (Error)
 					{
 						Chaos::FDebugDrawQueue::GetInstance().DrawDebugBox(ZOffset + PullData.X, FVector(4, 2, 2), PullData.R, FColor::Red, false, 5.f, 0, 0.5f);
-						Chaos::FDebugDrawQueue::GetInstance().DrawDebugDirectionalArrow(ZOffset + PullData.X, ZOffset + (PullData.X + InterpData->GetErrorX(0)), 1, FColor::Red, false, 5.0f, 0, 0.5f);
+						Chaos::FDebugDrawQueue::GetInstance().DrawDebugDirectionalArrow(ZOffset + (PullData.X + Error->ErrorX), ZOffset + PullData.X, 1, FColor::Red, false, 5.0f, 0, 0.5f);
 					}
 
 					Chaos::FDebugDrawQueue::GetInstance().DrawDebugDirectionalArrow(ZOffset + (Particle_External->GetX() - InterpData->GetErrorX(*Alpha)), ZOffset + Particle_External->GetX(), 1, FColor::Blue, false, 5.0f, 0, 0.5f);
diff --git a/Engine/Source/Runtime/Experimental/Chaos/Private/PhysicsProxy/GeometryCollectionPhysicsProxy.cpp b/Engine/Source/Runtime/Experimental/Chaos/Private/PhysicsProxy/GeometryCollectionPhysicsProxy.cpp
--- a/Engine/Source/Runtime/Experimental/Chaos/Private/PhysicsProxy/GeometryCollectionPhysicsProxy.cpp
+++ b/Engine/Source/Runtime/Experimental/Chaos/Private/PhysicsProxy/GeometryCollectionPhysicsProxy.cpp
@@ -37,6 +37,7 @@
 #include "Chaos/PhysicsObjectInterface.h"
 #include "GeometryCollection/Facades/CollectionAnchoringFacade.h"
 #include "GeometryCollection/Facades/CollectionConnectionGraphFacade.h"
+#include "Chaos/DebugDrawQueue.h"
 
 #define GC_PHYSICSPROXY_CHECK_FOR_NAN_ENABLED 0
 
@@ -4449,11 +4450,46 @@ bool FGeometryCollectionPhysicsProxy::PullFromPhysicsState(const Chaos::FDirtyGe
 					Chaos::FRotation3 NewR;
 					ResultInterpolator.GetPositions(NewX, NewR);
 
-					if (bHasInterpolationData && InterpData->IsErrorSmoothing())
+					const bool bIsReplicationErrorSmoothing = bHasInterpolationData ? InterpData->IsErrorSmoothing() : false;
+
+					if (bIsReplicationErrorSmoothing)
 					{
 						NewX = NewX + InterpData->GetErrorX(*Alpha);
-						NewR = InterpData->GetErrorR(*Alpha) * NewR;
+						NewR = NewR * InterpData->GetErrorR(*Alpha); // Add rotational error offset in local space
+					}
+
+#if CHAOS_DEBUG_DRAW
+					if (RenderInterpolationCVars::bRenderInterpDebugDrawGC)
+					{
+						FGeometryCollectionResults::FEntryIndex EntryIndex = PrevResults.GetEntryIndexByTransformIndex(TransformGroupIndex);
+						FGeometryCollectionResults::FEntryIndex NextEntryIndex = NextResults.GetEntryIndexByTransformIndex(TransformGroupIndex);
+
+						if (EntryIndex != INDEX_NONE && NextEntryIndex != INDEX_NONE)
+						{
+							const FGeometryCollectionResults::FPositionData& PullDataT = PrevResults.GetPositions(EntryIndex);
+							const FGeometryCollectionResults::FPositionData& NextPullDataT = NextResults.GetPositions(NextEntryIndex);
+
+							bool DirectionalDecayPerformed = false;
+							const FVector ZOffset = FVector(0, 0, RenderInterpolationCVars::RenderInterpDebugDrawZOffset);
+							Chaos::FDebugDrawQueue::GetInstance().DrawDebugBox(ZOffset + NextPullDataT.ParticleX, FVector(2, 1, 1), NextPullDataT.ParticleR, FColor::Yellow, false, 3.0f, 0, 0.5f);
+							Chaos::FDebugDrawQueue::GetInstance().DrawDebugDirectionalArrow(ZOffset + PullDataT.ParticleX, ZOffset + NextPullDataT.ParticleX, 0.5f, FColor::Yellow, false, 3.0f, 0, 0.5f);
+							Chaos::FDebugDrawQueue::GetInstance().DrawDebugBox(ZOffset + NewX, FVector(2, 1, 1), NewR, DirectionalDecayPerformed ? FColor::Cyan : FColor::Green, false, 3.0f, 0, 0.5f);
+
+							if (bIsReplicationErrorSmoothing)
+							{
+								if (Error)
+								{
+									Chaos::FDebugDrawQueue::GetInstance().DrawDebugBox(ZOffset + PullDataT.ParticleX, FVector(4, 2, 2), PullDataT.ParticleR, FColor::Red, false, 3.0f, 0, 0.5f);
+									Chaos::FDebugDrawQueue::GetInstance().DrawDebugDirectionalArrow(ZOffset + PullDataT.ParticleX + Error->ErrorX, ZOffset + PullDataT.ParticleX, 1, FColor::Red, false, 3.0f, 0, 0.5f);
+								}
+
+								Chaos::FDebugDrawQueue::GetInstance().DrawDebugDirectionalArrow(ZOffset + (NewX - InterpData->GetErrorX(*Alpha)), ZOffset + NewX, 1, FColor::Blue, false, 3.0f, 0, 0.5f);
+							}
+						}
 					}
+#endif // CHAOS_DEBUG_DRAW
+
+
 
 					const bool XRModified = UpdateGTParticleXR(GTParticle, NewX, NewR, Parameters.bForceUpdateActiveTransforms);
 					bIsCollectionDirty |= XRModified;
diff --git a/Engine/Source/Runtime/Experimental/Chaos/Private/PhysicsProxy/SingleParticlePhysicsProxy.cpp b/Engine/Source/Runtime/Experimental/Chaos/Private/PhysicsProxy/SingleParticlePhysicsProxy.cpp
--- a/Engine/Source/Runtime/Experimental/Chaos/Private/PhysicsProxy/SingleParticlePhysicsProxy.cpp
+++ b/Engine/Source/Runtime/Experimental/Chaos/Private/PhysicsProxy/SingleParticlePhysicsProxy.cpp
@@ -424,7 +424,7 @@ bool FSingleParticlePhysicsProxy::PullFromPhysicsState(const Chaos::FDirtyRigidP
 					FQuat Target = FMath::Lerp(*Prev, NextPullData->R, *Alpha);
 					if (bIsReplicationErrorSmoothing)
 					{
-						Target = InterpData->GetErrorR(*Alpha) * Target;
+						Target = Target * InterpData->GetErrorR(*Alpha); // Add rotational error offset in local space
 					}
 					Rigid->SetR(Target, false);
 					bNeedUpdateShapeBounds = true;
@@ -443,7 +443,7 @@ bool FSingleParticlePhysicsProxy::PullFromPhysicsState(const Chaos::FDirtyRigidP
 						if (Error)
 						{
 							Chaos::FDebugDrawQueue::GetInstance().DrawDebugBox(ZOffset + PullData.X, FVector(4, 2, 2), PullData.R, FColor::Red, false, 5.f, 0, 0.5f);
-							Chaos::FDebugDrawQueue::GetInstance().DrawDebugDirectionalArrow(ZOffset + PullData.X, ZOffset + (PullData.X + InterpData->GetErrorX(0)), 1, FColor::Red, false, 5.0f, 0, 0.5f);
+							Chaos::FDebugDrawQueue::GetInstance().DrawDebugDirectionalArrow(ZOffset + (PullData.X + Error->ErrorX), ZOffset + PullData.X, 1, FColor::Red, false, 5.0f, 0, 0.5f);
 						}
 
 						if (InterpData->IsErrorVelocitySmoothing())
diff --git a/Engine/Source/Runtime/Experimental/Chaos/Private/RewindData.cpp b/Engine/Source/Runtime/Experimental/Chaos/Private/RewindData.cpp
--- a/Engine/Source/Runtime/Experimental/Chaos/Private/RewindData.cpp
+++ b/Engine/Source/Runtime/Experimental/Chaos/Private/RewindData.cpp
@@ -4,6 +4,7 @@
 #include "PBDRigidsSolver.h"
 #include "PhysicsProxy/SingleParticlePhysicsProxy.h"
 #include "PhysicsProxy/ClusterUnionPhysicsProxy.h"
+#include "PhysicsProxy/GeometryCollectionPhysicsProxy.h"
 #include "HAL/IConsoleManager.h"
 #include "Chaos/PBDJointConstraints.h"
 #include "Misc/Paths.h"
@@ -460,7 +461,11 @@ bool FRewindData::RewindToFrame(int32 Frame)
 		
 		FGeometryParticleStateBase& History = DirtyParticleInfo.GetHistory(); //non-const in case we need to record what's at head for a rewind (CurFrame has already been increased to the next frame)
 
-		History.CachePreCorrectionState(*PTParticle);
+		PRAGMA_DISABLE_DEPRECATION_WARNINGS
+		History.CachePreCorrectionState(*PTParticle); // Deprecated UE 5.6
+		PRAGMA_ENABLE_DEPRECATION_WARNINGS
+
+		CachePreResimState(*PTParticle);
 
 		// Todo: This should be set by bubble resimulation so that Dynamic and Sleeping particles outside of relevancy of resim doesn't actually resimulate as dynamic particles.
 		DirtyParticleInfo.bResimAsFollower = PTParticle->ObjectState() != EObjectStateType::Dynamic && PTParticle->ObjectState() != EObjectStateType::Sleeping;
@@ -1027,18 +1032,90 @@ void FRewindData::SpawnProxyIfNeeded(FSingleParticlePhysicsProxy& Proxy)
 	}
 }
 
+void FRewindData::CachePreResimState(FGeometryParticleHandle& Handle)
+{
+	const IPhysicsProxyBase* PhysicsProxy = static_cast<const IPhysicsProxyBase*>(Handle.PhysicsProxy());
+	if (!PhysicsProxy)
+	{
+		return;
+	}
+
+	// Find or add pre-resim error for the particle that has an error
+	auto PreErrorInfo = [this](FGeometryParticleHandle* ParticleHandle) -> FDirtyParticleErrorInfo&
+		{
+			if (FDirtyParticleErrorInfo* Found = DirtyParticlePreResimState.Find(ParticleHandle))
+			{
+				return *Found;
+			}
+
+			return DirtyParticlePreResimState.Add(ParticleHandle, FDirtyParticleErrorInfo(*ParticleHandle));
+		};
+
+	// Cache dirty particle XR before a resimulation
+	PreErrorInfo(&Handle).AccumulateError(Handle.GetX(), Handle.GetR());
+
+	// If particle is a Cluster Union, also cache child particles
+	if (PhysicsProxy->GetType() == EPhysicsProxyType::ClusterUnionProxy)
+	{
+		const FClusterUnionPhysicsProxy* ClusterProxy = static_cast<const FClusterUnionPhysicsProxy*>(PhysicsProxy);
+		if (ClusterProxy)
+		{
+			for (IPhysicsProxyBase* ChildProxyBase : ClusterProxy->GetParticle_Internal()->PhysicsProxies())
+			{
+				switch (ChildProxyBase->GetType())
+				{
+					case EPhysicsProxyType::SingleParticleProxy:
+					{
+						if (FSingleParticlePhysicsProxy* ChildProxy = static_cast<FSingleParticlePhysicsProxy*>(ChildProxyBase))
+						{
+							if (FGeometryParticleHandle* ChildHandle = ChildProxy->GetHandle_LowLevel())
+							{
+								PreErrorInfo(ChildHandle).AccumulateError(ChildHandle->GetX(), ChildHandle->GetR());
+							}
+						}
+						break;
+					}
+					case EPhysicsProxyType::ClusterUnionProxy:
+					{
+						if (FClusterUnionPhysicsProxy* ChildProxy = static_cast<FClusterUnionPhysicsProxy*>(ChildProxyBase))
+						{
+							if (FGeometryParticleHandle* ChildHandle = ChildProxy->GetParticle_Internal())
+							{
+								PreErrorInfo(ChildHandle).AccumulateError(ChildHandle->GetX(), ChildHandle->GetR());
+							}
+						}
+						break;
+					}
+					case EPhysicsProxyType::GeometryCollectionType:
+					{
+						if (FGeometryCollectionPhysicsProxy* ChildProxy = static_cast<FGeometryCollectionPhysicsProxy*>(ChildProxyBase))
+						{
+							if (FGeometryParticleHandle* ChildHandle = ChildProxy->GetInitialRootParticle_Internal())
+							{
+								PreErrorInfo(ChildHandle).AccumulateError(ChildHandle->GetX(), ChildHandle->GetR());
+							}
+						}
+						break;
+					}
+				}
+			}
+		}
+	}
+
+}
+
 template<>
 void FRewindData::AccumulateErrorIfNecessary(FGeometryParticleHandle& Obj, const FFrameAndPhase FrameAndPhase)
 {
-	const FDirtyParticleInfo* DirtyInfo = DirtyParticles.Find(&Obj);
-	if (!DirtyInfo)
+	FDirtyParticleErrorInfo* PreErrorInfo = DirtyParticlePreResimState.Find(&Obj);
+	if (!PreErrorInfo)
 	{
 		return;
 	}
 
 	// Get the error offset after a correction
-	const FVec3 ErrorX = DirtyInfo->GetHistory().PreCorrectionXR.X() - Obj.GetX();
-	FQuat ErrorR = DirtyInfo->GetHistory().PreCorrectionXR.R() * Obj.GetR().Inverse();
+	const FVec3 ErrorX = PreErrorInfo->GetErrorX() - Obj.GetX();
+	FQuat ErrorR = Obj.GetR().Inverse() * PreErrorInfo->GetErrorR(); // ErrorR in local space 
 	ErrorR.EnforceShortestArcWith(FQuat::Identity);
 	ErrorR.Normalize();
 
@@ -1072,21 +1149,18 @@ void FRewindData::BufferPhysicsResults(TMap<const IPhysicsProxyBase*, struct FDi
 			{
 				for (auto& Info : DirtyObjs)
 				{
-					if (Info.bResimAsFollower)
+					if (!Info.GetObjectPtr() || Info.GetObjectPtr()->ResimType() == EResimType::ResimAsFollower)
 					{
 						continue;
 					}
 
-					if (Info.GetObjectPtr())
-					{
-						// Cache the correction offset after a resimulation
-						AccumulateErrorIfNecessary(*Info.GetObjectPtr(), FutureFrame);
-					}
+					// Cache the correction offset after a resimulation
+					AccumulateErrorIfNecessary(*Info.GetObjectPtr(), FutureFrame);
 				}
 			};
 
-		ErrorDataHelper(DirtyParticles);
-		ErrorDataHelper(DirtyJoints);
+		ErrorDataHelper(DirtyParticlePreResimState);
+		DirtyParticlePreResimState.Reset();
 	}
 
 	DirtyRigidErrors.Reserve(DirtyParticleErrors.Num());
@@ -1097,24 +1171,9 @@ void FRewindData::BufferPhysicsResults(TMap<const IPhysicsProxyBase*, struct FDi
 		ErrorData.ErrorX = ErrorInfo.GetErrorX();
 		ErrorData.ErrorR = ErrorInfo.GetErrorR();
 
-		const IPhysicsProxyBase* PhysicsProxy = static_cast<const IPhysicsProxyBase*>(ErrorInfo.GetObjectPtr()->PhysicsProxy());
-		DirtyRigidErrors.Add(PhysicsProxy, ErrorData);
-
-		// todo, move out of RewindData.cpp, possibly to ChaosResultManager.cpp
-		// If ClusterUnion, hand over error to child GeometryCollections which do their own render interpolation and need the error offset individually
-		if (PhysicsProxy->GetType() == EPhysicsProxyType::ClusterUnionProxy)
+		if (const IPhysicsProxyBase* PhysicsProxy = static_cast<const IPhysicsProxyBase*>(ErrorInfo.GetObjectPtr()->PhysicsProxy()))
 		{
-			const FClusterUnionPhysicsProxy* ClusterProxy = static_cast<const FClusterUnionPhysicsProxy*>(PhysicsProxy);
-			if (ClusterProxy)
-			{
-				for (IPhysicsProxyBase* ChildProxy : ClusterProxy->GetParticle_Internal()->PhysicsProxies())
-				{
-					if (ChildProxy->GetType() == EPhysicsProxyType::GeometryCollectionType)
-					{
-						DirtyRigidErrors.Add(ChildProxy, ErrorData);
-					}
-				}
-			}
+			DirtyRigidErrors.Add(PhysicsProxy, ErrorData);
 		}
 	}
 
diff --git a/Engine/Source/Runtime/Experimental/Chaos/Public/Chaos/Framework/PhysicsProxyBase.h b/Engine/Source/Runtime/Experimental/Chaos/Public/Chaos/Framework/PhysicsProxyBase.h
--- a/Engine/Source/Runtime/Experimental/Chaos/Public/Chaos/Framework/PhysicsProxyBase.h
+++ b/Engine/Source/Runtime/Experimental/Chaos/Public/Chaos/Framework/PhysicsProxyBase.h
@@ -36,6 +36,7 @@ namespace RenderInterpolationCVars
 	extern float RenderInterpErrorDirectionalDecayMultiplier;
 	extern bool bRenderInterpErrorVelocityCorrection;
 	extern bool bRenderInterpDebugDraw;
+	extern bool bRenderInterpDebugDrawGC;
 	extern float RenderInterpDebugDrawZOffset;
 }
 
diff --git a/Engine/Source/Runtime/Experimental/Chaos/Public/RewindData.h b/Engine/Source/Runtime/Experimental/Chaos/Public/RewindData.h
--- a/Engine/Source/Runtime/Experimental/Chaos/Public/RewindData.h
+++ b/Engine/Source/Runtime/Experimental/Chaos/Public/RewindData.h
@@ -1047,6 +1047,7 @@ inline int32 ComputeCircularSize(int32 NumFrames) { return NumFrames * FFrameAnd
 
 struct FGeometryParticleStateBase
 {
+	PRAGMA_DISABLE_DEPRECATION_WARNINGS
 	UE_DEPRECATED(5.6, "Use the constructor that takes a @param bCacheOnePhase")
 	explicit FGeometryParticleStateBase(int32 NumFrames)
 	: ParticlePositionRotation(ComputeCircularSize(NumFrames))
@@ -1080,6 +1081,14 @@ struct FGeometryParticleStateBase
 	FGeometryParticleStateBase(FGeometryParticleStateBase&& Other) = default;
 	~FGeometryParticleStateBase() = default;
 
+	// Can be removed when Deprecations are removed
+	FGeometryParticleStateBase() = delete;
+	// Can be removed when Deprecations are removed
+	FGeometryParticleStateBase& operator=(const FGeometryParticleStateBase&) = delete;
+	// Can be removed when Deprecations are removed
+	FGeometryParticleStateBase& operator=(FGeometryParticleStateBase&&) = delete;
+	PRAGMA_ENABLE_DEPRECATION_WARNINGS
+
 	void Release(FDirtyPropertiesPool& Manager)
 	{
 		ParticlePositionRotation.Release(Manager);
@@ -1159,6 +1168,7 @@ struct FGeometryParticleStateBase
 	void SyncDirtyDynamics(FDirtyPropData& DestManager,const FDirtyChaosProperties& Dirty,const FConstDirtyPropData& SrcManager);
 	
 	template<typename TParticle>
+	UE_DEPRECATED(5.6, "Deprecated, use FRewindData::CachePreResimState instead. Not all moving particles are marked as dirty, for example GeometryCollection children of a ClusterUnion.")
 	void CachePreCorrectionState(const TParticle& Particle)
 	{
 		PreCorrectionXR.SetX(Particle.GetX());
@@ -1179,6 +1189,7 @@ struct FGeometryParticleStateBase
 
 	FShapesArrayStateBase ShapesArrayState;
 
+	UE_DEPRECATED(5.6, "Deprecated, use FRewindData::DirtyParticlePreResimState")
 	FParticlePositionRotation PreCorrectionXR;
 };
 
@@ -1979,6 +1990,7 @@ private:
 
 	TDirtyObjects<FDirtyParticleInfo> DirtyParticles;
 	TDirtyObjects<FDirtyJointInfo> DirtyJoints;
+	TDirtyObjects<FDirtyParticleErrorInfo> DirtyParticlePreResimState;
 	TDirtyObjects<FDirtyParticleErrorInfo> DirtyParticleErrors;
 
 	TArray<TWeakPtr<FBaseRewindHistory>> InputHistories; // Todo, deprecate in favor of Particle->Input map?
@@ -2010,6 +2022,8 @@ private:
 	template <bool bSkipDynamics, typename TDirtyInfo>
 	void DesyncIfNecessary(TDirtyInfo& Info, const FFrameAndPhase FrameAndPhase);
 
+	void CachePreResimState(FGeometryParticleHandle& Handle);
+
 	template<typename TObj>
 	void AccumulateErrorIfNecessary(TObj& Handle, const FFrameAndPhase FrameAndPhase) { }
 };
