Add support for FInstancedStruct to UE::JsonObjectGraph::Stringify

#rnx
[RN] UE::JsonObjectGraph::Stringify now includes information about FInstancedStruct members

diff --git a/Engine/Source/Runtime/Experimental/JsonObjectGraph/Private/JsonObjectGraphConventions.h b/Engine/Source/Runtime/Experimental/JsonObjectGraph/Private/JsonObjectGraphConventions.h
--- a/Engine/Source/Runtime/Experimental/JsonObjectGraph/Private/JsonObjectGraphConventions.h
+++ b/Engine/Source/Runtime/Experimental/JsonObjectGraph/Private/JsonObjectGraphConventions.h
@@ -25,6 +25,7 @@
 #define UE_JSON_OBJECT_SERIAL_DATA_KEY_TCHAR TEXT("__SerialData")
 #define UE_JSON_OBJECT_SPARSE_CLASS_DATA_KEY_TCHAR TEXT("__SparseClassData")
 #define UE_JSON_OBJECT_INDIRECTLY_REFERENCED_KEY_TCHAR TEXT("__IndirectlyReferenced")
+#define UE_JSON_SCRIPTSTRUCT_TCHAR TEXT("__UScriptStruct")
 
 // Reference encoding related conventions:
 #define UE_JSON_OBJECT_REF_PREFIX "uobject:"
diff --git a/Engine/Source/Runtime/Experimental/JsonObjectGraph/Private/JsonStringifyImpl.cpp b/Engine/Source/Runtime/Experimental/JsonObjectGraph/Private/JsonStringifyImpl.cpp
--- a/Engine/Source/Runtime/Experimental/JsonObjectGraph/Private/JsonStringifyImpl.cpp
+++ b/Engine/Source/Runtime/Experimental/JsonObjectGraph/Private/JsonStringifyImpl.cpp
@@ -13,6 +13,7 @@
 #include "Serialization/StructuredArchiveSlots.h"
 #include "Serialization/Formatters/JsonArchiveOutputFormatter.h"
 #include "Serialization/ArchiveUObjectFromStructuredArchive.h"
+#include "StructUtils/InstancedStruct.h"
 #include "UObject/EnumProperty.h"
 #include "UObject/PropertyOptional.h"
 #include "UObject/Package.h"
@@ -670,12 +671,59 @@ void FJsonStringifyImpl::WriteStructToJsonWithIdentifier(const TCHAR* Identifier
 
 void FJsonStringifyImpl::WriteStructToJson(const void* StructInstance, const void* DefaultInstance, const UScriptStruct* Struct, const UScriptStruct* DefaultStruct)
 {
+	// FInstancedStruct is a core level construct, and we can usefully decompose it, so try:
+	const FInstancedStruct* StructInstanceTyped = Struct == FInstancedStruct::StaticStruct() ? 
+		static_cast<const FInstancedStruct*>(StructInstance) : nullptr;
+	const FInstancedStruct* DefaultInstanceTyped = DefaultStruct == FInstancedStruct::StaticStruct() ? 
+		static_cast<const FInstancedStruct*>(DefaultInstance) : nullptr;
+	const UScriptStruct* InstancedStructType = StructInstanceTyped ? 
+		StructInstanceTyped->GetScriptStruct() : nullptr;
+	const UScriptStruct* DefaultInstancedStructType = DefaultInstanceTyped ? 
+		DefaultInstanceTyped->GetScriptStruct() : nullptr;
+
+	// for structs with an ambiguous type (like instanced structs), we want to disambiguate
+	// the type stored in the json:
+	const bool bWriteStructType = InstancedStructType && 
+		// UPropertyBag is a transient non native struct, and is regenerated on load. It 
+		// is not meant to have persisent identity - so we don't want to even try
+		// to write it out:
+		(!InstancedStructType->HasAnyFlags(RF_Transient) || InstancedStructType->IsNative());
+		
+	TUniquePtr<FInstancedStruct> FallbackStruct;
+	if(InstancedStructType)
+	{
+		StructInstance = StructInstanceTyped->GetMemory();
+		Struct = InstancedStructType;
+		if(DefaultInstanceTyped)
+		{
+			DefaultInstance = DefaultInstanceTyped->GetMemory();
+			DefaultStruct = DefaultInstancedStructType; 
+		}
+		else
+		{
+			// fall back to default constructed struct
+			FallbackStruct = MakeUnique<FInstancedStruct>(InstancedStructType);
+			DefaultInstance = FallbackStruct.Get()->GetMemory();
+			DefaultStruct = InstancedStructType;
+		}
+	}
+
 	FPendingScope ObjectStartIdentifier( 
 		this, 
-		[Impl = this]() { Impl->Writer->WriteObjectStartInline(); }, 
+		[Impl = this, Struct, bWriteStructType]() 
+		{ 
+			Impl->Writer->WriteObjectStartInline(); 
+			if(bWriteStructType)
+			{
+				Impl->Writer->WriteUtf8Value(UE_JSON_SCRIPTSTRUCT_TCHAR, Impl->WriteObjectReference(Struct));
+			}
+		}, 
 		[Impl = this]() { Impl->Writer->WriteObjectEnd(); }
 		);
 
+	// @todo: can we improve results for structs that have custom serializers? 
+	// We do not capture data outside of FProperty because I use this routine
+	// for inspecting objects and readability is paramount:
 	for (TFieldIterator<FProperty> PropertyIt(Struct); PropertyIt; ++PropertyIt)
 	{
 		FProperty* Property = *PropertyIt;
