Add limited simulation morph target support to ChaosCloth.

Currently only one morph target may be active at a time. You must activate your morph target via the asset or blueprint interactor.

diff --git a/Engine/Plugins/ChaosCloth/Source/ChaosCloth/Private/ChaosCloth/ChaosClothingSimulationCloth.cpp b/Engine/Plugins/ChaosCloth/Source/ChaosCloth/Private/ChaosCloth/ChaosClothingSimulationCloth.cpp
--- a/Engine/Plugins/ChaosCloth/Source/ChaosCloth/Private/ChaosCloth/ChaosClothingSimulationCloth.cpp
+++ b/Engine/Plugins/ChaosCloth/Source/ChaosCloth/Private/ChaosCloth/ChaosClothingSimulationCloth.cpp
@@ -83,25 +83,29 @@ struct FClothingSimulationCloth::FLODData
 	};
 	TMap<FClothingSimulationSolver*, FSolverData> SolverData;
 
+	// Cached property facade data
+	int32 ActiveMorphTarget = INDEX_NONE;
+	float ActiveMorphTargetWeight = 0.f;
+
 	// Stats
 	int32 NumKinematicParticles;
 	int32 NumDynamicParticles;
 
 	FLODData(
-		int32 InNumParticles,
-		const TConstArrayView<uint32>& InIndices,
-		const TConstArrayView<FVector2f>& InPatternPositions,
-		const TConstArrayView<uint32>& InPatternIndices,
-		const TConstArrayView<uint32>& InPatternToWeldedIndices,
+		const FClothingSimulationMesh& Mesh,
+		const int32 LODIndex,
+		bool bUseGeodesicTethers,
 		TMap<FString, TConstArrayView<FRealSingle>>&& InWeightMaps,
 		TMap<FString, const TSet<int32>*>&& InVertexSets,
 		TMap<FString, const TSet<int32>*>&& InFaceSets,
 		TMap<FString, TConstArrayView<int32>>&& InFaceIntMaps,
-		const TArray<TConstArrayView<TTuple<int32, int32, FRealSingle>>>& InTethers,
-		const TSharedPtr<const FManagedArrayCollection>& InManagedArrayCollection);
+		const Softs::FCollectionPropertyFacade& ConfigProperties);
 
 	static FTriangleMesh BuildTriangleMesh(const TConstArrayView<uint32>& Indices, const int32 NumParticles);
 
+	void UpdateCachedProperties(const FClothingSimulationMesh& Mesh,
+		const int32 LODIndex, const Softs::FCollectionPropertyFacade& PropertyCollection, bool bForce = false);
+
 	void AddParticles(FClothingSimulationSolver* Solver, FClothingSimulationCloth* Cloth, int32 LODIndex);
 	void Add(FClothingSimulationSolver* Solver, FClothingSimulationCloth* Cloth, int32 LODIndex);
 	void Remove(FClothingSimulationSolver* Solver);
@@ -114,36 +118,40 @@ struct FClothingSimulationCloth::FLODData
 
 	void UpdateNormals(FClothingSimulationSolver* Solver) const;
 
+
+	// These are only used in Add, so no need to cache them
 	UE_CHAOS_DECLARE_INDEXLESS_PROPERTYCOLLECTION_NAME(MaxDistance, float);
 	UE_CHAOS_DECLARE_INDEXLESS_PROPERTYCOLLECTION_NAME(MassValue, float);
 	UE_CHAOS_DECLARE_INDEXLESS_PROPERTYCOLLECTION_NAME(MultiResCoarseLODIndex, int32);
 	UE_CHAOS_DECLARE_INDEXLESS_PROPERTYCOLLECTION_NAME(IsCoarseMultiResLOD, bool);
 	UE_CHAOS_DECLARE_INDEXLESS_PROPERTYCOLLECTION_NAME(KinematicVertices3D, bool); // Selection set name string property; bool value is not actually used.
 
+	// This can update every frame, so should be cached
+	UE_CHAOS_DECLARE_PROPERTYCOLLECTION_NAME(ActiveMorphTarget, float);
+
 };
 
 FClothingSimulationCloth::FLODData::FLODData(
-	int32 InNumParticles,
-	const TConstArrayView<uint32>& InIndices,
-	const TConstArrayView<FVector2f>& InPatternPositions,
-	const TConstArrayView<uint32>& InPatternIndices,
-	const TConstArrayView<uint32>& InPatternToWeldedIndices,
+	const FClothingSimulationMesh& InMesh,
+	const int32 LODIndex,
+	bool bUseGeodesicTethers,
 	TMap<FString, TConstArrayView<FRealSingle>>&& InWeightMaps,
 	TMap<FString, const TSet<int32>*>&& InVertexSets,
 	TMap<FString, const TSet<int32>*>&& InFaceSets,
 	TMap<FString, TConstArrayView<int32>>&& InFaceIntMaps,
-	const TArray<TConstArrayView<TTuple<int32, int32, FRealSingle>>>& InTethers,
-	const TSharedPtr<const FManagedArrayCollection>& InManagedArrayCollection)
-	: NumParticles(InNumParticles)
+	const Softs::FCollectionPropertyFacade& PropertyCollection)
+	: NumParticles(InMesh.GetNumPoints(LODIndex))
 	, WeightMaps(MoveTemp(InWeightMaps))
 	, VertexSets(MoveTemp(InVertexSets))
 	, FaceSets(MoveTemp(InFaceSets))
 	, FaceIntMaps(MoveTemp(InFaceIntMaps))
-	, Tethers(InTethers)
-	, ManagedArrayCollection(InManagedArrayCollection)
-	, PatternData(NumParticles, InIndices, InPatternPositions, InPatternIndices, InPatternToWeldedIndices)
-	, NoOffsetTriangleMesh(BuildTriangleMesh(InIndices, NumParticles))
+	, Tethers(InMesh.GetTethers(LODIndex, bUseGeodesicTethers))
+	, ManagedArrayCollection(InMesh.GetManagedArrayCollection(LODIndex))
+	, PatternData(NumParticles, InMesh.GetIndices(LODIndex), InMesh.GetPatternPositions(LODIndex), InMesh.GetPatternIndices(LODIndex), InMesh.GetPatternToWeldedIndices(LODIndex))
+	, NoOffsetTriangleMesh(BuildTriangleMesh(InMesh.GetIndices(LODIndex), NumParticles))
+	, ActiveMorphTargetIndex(PropertyCollection)
 {
+	UpdateCachedProperties(InMesh, LODIndex, PropertyCollection, true);
 }
 
 FTriangleMesh FClothingSimulationCloth::FLODData::BuildTriangleMesh(const TConstArrayView<uint32>& Indices, const int32 InNumParticles)
@@ -168,6 +176,20 @@ FTriangleMesh FClothingSimulationCloth::FLODData::BuildTriangleMesh(const TConst
 	return OutTriangleMesh;
 }
 
+void FClothingSimulationCloth::FLODData::UpdateCachedProperties(const FClothingSimulationMesh& InMesh,
+	const int32 LODIndex, const Softs::FCollectionPropertyFacade& InPropertyCollection, bool bForce)
+{
+	if (ActiveMorphTargetIndex != INDEX_NONE && (bForce || IsActiveMorphTargetMutable(InPropertyCollection)))
+	{
+		ActiveMorphTargetWeight = FMath::Clamp(GetActiveMorphTarget(InPropertyCollection), -1.f, 1.f);
+
+		if (bForce || IsActiveMorphTargetStringDirty(InPropertyCollection))
+		{
+			ActiveMorphTarget = InMesh.FindMorphTargetByName(LODIndex, GetActiveMorphTargetString(InPropertyCollection));
+		}
+	}
+}
+
 void FClothingSimulationCloth::FLODData::AddParticles(FClothingSimulationSolver* Solver, FClothingSimulationCloth* Cloth, int32 InLODIndex)
 {
 	check(Solver);
@@ -486,17 +508,15 @@ void FClothingSimulationCloth::SetMesh(FClothingSimulationMesh* InMesh)
 
 		// Add LOD data
 		LODData.Add(MakeUnique<FLODData>(
-			Mesh->GetNumPoints(LODIndex),
-			Mesh->GetIndices(LODIndex),
-			Mesh->GetPatternPositions(LODIndex),
-			Mesh->GetPatternIndices(LODIndex),
-			Mesh->GetPatternToWeldedIndices(LODIndex),
+			*Mesh,
+			LODIndex,
+			bUseGeodesicTethers,
 			MoveTemp(WeightMaps),
 			MoveTemp(VertexSets),
 			MoveTemp(FaceSets),
 			MoveTemp(FaceIntMaps),
-			Mesh->GetTethers(LODIndex, bUseGeodesicTethers),
-			Mesh->GetManagedArrayCollection(LODIndex)));
+			Config->GetProperties(LODIndex)
+		));
 	}
 
 	// Iterate all known solvers
@@ -854,19 +874,28 @@ void FClothingSimulationCloth::Update(FClothingSimulationSolver* Solver)
 		return;
 	}
 
-	Mesh->Update(Solver, PrevLODIndex, LODIndex, PrevParticleRangeId, ParticleRangeId);
-
 	// Retrieve config
 	check(Config);
 	const Softs::FCollectionPropertyFacade& ConfigProperties = Config->GetProperties(LODIndex);
+	int32 CurrentLODMorphTargetIndex = INDEX_NONE;
+	float CurrentLODMorphTargetWeight = 0.f;
+	if (LODIndex != INDEX_NONE)
+	{
+		LODData[LODIndex]->UpdateCachedProperties(*Mesh, LODIndex, ConfigProperties);
+		CurrentLODMorphTargetIndex = LODData[LODIndex]->ActiveMorphTarget;
+		CurrentLODMorphTargetWeight = LODData[LODIndex]->ActiveMorphTargetWeight;
+	}
+
+	Mesh->Update(Solver, PrevLODIndex, LODIndex, PrevParticleRangeId, ParticleRangeId, CurrentLODMorphTargetIndex, CurrentLODMorphTargetWeight);
 
 	const int32 CoarseLODIndex = LODIndex != INDEX_NONE ? LODData[LODIndex]->SolverData.FindChecked(Solver).MultiResCoarseLODIndex : INDEX_NONE;
 	const int32 CoarseParticleRangeId = CoarseLODIndex != INDEX_NONE ? GetParticleRangeId(Solver, CoarseLODIndex) : INDEX_NONE;
 	if (CoarseLODIndex != INDEX_NONE)
 	{
+		LODData[CoarseLODIndex]->UpdateCachedProperties(*Mesh, CoarseLODIndex, Config->GetProperties(CoarseLODIndex));
 		LODData[CoarseLODIndex]->Enable(Solver, true);
 		// TODO: interpolate/ reset when LOD switching to enable multires
-		Mesh->Update(Solver, CoarseLODIndex, CoarseLODIndex, CoarseParticleRangeId, CoarseParticleRangeId);
+		Mesh->Update(Solver, CoarseLODIndex, CoarseLODIndex, CoarseParticleRangeId, CoarseParticleRangeId, LODData[CoarseLODIndex]->ActiveMorphTarget, LODData[CoarseLODIndex]->ActiveMorphTargetWeight);
 	}
 
 	// LOD Switching
diff --git a/Engine/Plugins/ChaosCloth/Source/ChaosCloth/Private/ChaosCloth/ChaosClothingSimulationMesh.cpp b/Engine/Plugins/ChaosCloth/Source/ChaosCloth/Private/ChaosCloth/ChaosClothingSimulationMesh.cpp
--- a/Engine/Plugins/ChaosCloth/Source/ChaosCloth/Private/ChaosCloth/ChaosClothingSimulationMesh.cpp
+++ b/Engine/Plugins/ChaosCloth/Source/ChaosCloth/Private/ChaosCloth/ChaosClothingSimulationMesh.cpp
@@ -187,7 +187,7 @@ FORCEINLINE static void AddInfluence(FVector3f& OutPosition, FVector3f& OutNorma
 	OutNormal += BoneMatrix.TransformVector(RefNormal) * Weight;
 }
 
-void FClothingSimulationMesh::SkinPhysicsMesh(int32 LODIndex, const FReal LocalSpaceScale, const FVec3& LocalSpaceLocation, TArrayView<Softs::FSolverVec3>& OutPositions, TArrayView<Softs::FSolverVec3>& OutNormals) const
+void FClothingSimulationMesh::SkinPhysicsMesh(int32 LODIndex, int32 ActiveMorphTargetIndex, float ActiveMorphTargetWeight, const FReal LocalSpaceScale, const FVec3& LocalSpaceLocation, TArrayView<Softs::FSolverVec3>& OutPositions, TArrayView<Softs::FSolverVec3>& OutNormals) const
 {
 	TRACE_CPUPROFILER_EVENT_SCOPE(FClothingSimulationMesh_SkinPhysicsMesh);
 	SCOPE_CYCLE_COUNTER(STAT_ChaosClothSkinPhysicsMesh);
@@ -209,8 +209,34 @@ void FClothingSimulationMesh::SkinPhysicsMesh(int32 LODIndex, const FReal LocalS
 	check(NumPoints == OutPositions.Num());
 	check(NumPoints == OutNormals.Num());
 	const TConstArrayView<FClothVertBoneData> BoneData = GetBoneData(LODIndex);
-	const TConstArrayView<FVector3f> Positions = GetPositions(LODIndex);
-	const TConstArrayView<FVector3f> Normals = GetNormals(LODIndex);
+	TConstArrayView<FVector3f> Positions = GetPositions(LODIndex);
+	TConstArrayView<FVector3f> Normals = GetNormals(LODIndex);
+	TArray<FVector3f> WritablePositions;
+	TArray<FVector3f> WritableNormals;
+
+	const TConstArrayView<FVector3f> MorphTargetPositionDeltas = GetMorphTargetPositionDeltas(LODIndex, ActiveMorphTargetIndex);
+	const TConstArrayView<FVector3f> MorphTargetTangentZDeltas = GetMorphTargetTangentZDeltas(LODIndex, ActiveMorphTargetIndex);
+	const TConstArrayView<int32> MorphTargetIndices = GetMorphTargetIndices(LODIndex, ActiveMorphTargetIndex);
+
+	if (!FMath::IsNearlyZero(ActiveMorphTargetWeight) && !MorphTargetPositionDeltas.IsEmpty() &&
+		MorphTargetPositionDeltas.Num() == MorphTargetTangentZDeltas.Num() &&
+		MorphTargetPositionDeltas.Num() == MorphTargetIndices.Num())
+	{
+		TRACE_CPUPROFILER_EVENT_SCOPE(FClothingSimulationMesh_SkinPhysicsMesh_MorphTargets);
+		// TODO optimize this
+		WritablePositions = Positions;
+		WritableNormals = Normals;
+
+		for (int32 Index = 0; Index < MorphTargetIndices.Num(); ++Index)
+		{
+			const int32 VertexIndex = MorphTargetIndices[Index];
+			WritablePositions[VertexIndex] += ActiveMorphTargetWeight * MorphTargetPositionDeltas[Index];
+			WritableNormals[VertexIndex] = (WritableNormals[VertexIndex] + ActiveMorphTargetWeight * MorphTargetTangentZDeltas[Index]).GetSafeNormal();
+		}
+
+		Positions = WritablePositions;
+		Normals = WritableNormals;
+	}
 
 #if INTEL_ISPC
 	if (bChaos_SkinPhysicsMesh_ISPC_Enabled)
@@ -275,7 +301,9 @@ void FClothingSimulationMesh::Update(
 	int32 PrevLODIndex,
 	int32 LODIndex,
 	int32 PrevParticleRangeId,
-	int32 ParticleRangeId)
+	int32 ParticleRangeId,
+	int32 ActiveMorphTargetIndex,
+	float ActiveMorphTargetWeight)
 {
 	check(Solver);
 
@@ -291,7 +319,7 @@ void FClothingSimulationMesh::Update(
 	TArrayView<FSolverVec3> OutPositions = Solver->GetAnimationPositionsView(ParticleRangeId);
 	TArrayView<FSolverVec3> OutNormals = Solver->GetAnimationNormalsView(ParticleRangeId);
 	
-	SkinPhysicsMesh(LODIndex, LocalSpaceScale, LocalSpaceLocation, OutPositions, OutNormals);
+	SkinPhysicsMesh(LODIndex, ActiveMorphTargetIndex, ActiveMorphTargetWeight, LocalSpaceScale, LocalSpaceLocation, OutPositions, OutNormals);
 
 	// Update old positions after LOD Switching
 	if (LODIndex != PrevLODIndex)
diff --git a/Engine/Plugins/ChaosCloth/Source/ChaosCloth/Public/ChaosCloth/ChaosClothingSimulationMesh.h b/Engine/Plugins/ChaosCloth/Source/ChaosCloth/Public/ChaosCloth/ChaosClothingSimulationMesh.h
--- a/Engine/Plugins/ChaosCloth/Source/ChaosCloth/Public/ChaosCloth/ChaosClothingSimulationMesh.h
+++ b/Engine/Plugins/ChaosCloth/Source/ChaosCloth/Public/ChaosCloth/ChaosClothingSimulationMesh.h
@@ -143,6 +143,27 @@ namespace Chaos
 			return TSharedPtr<const FManagedArrayCollection>();
 		}
 
+		/** Return the MorphTargetIndex for a given Morph Target Name, if it exists. */
+		virtual int32 FindMorphTargetByName(int32 LODIndex, const FString& Name) const
+		{
+			return INDEX_NONE;
+		}
+
+		virtual TConstArrayView<FVector3f> GetMorphTargetPositionDeltas(int32 LODIndex, int32 MorphTargetIndex) const
+		{
+			return TConstArrayView<FVector3f>();
+		}
+
+		virtual TConstArrayView<FVector3f> GetMorphTargetTangentZDeltas(int32 LODIndex, int32 MorphTargetIndex) const
+		{
+			return TConstArrayView<FVector3f>();
+		}
+
+		virtual TConstArrayView<int32> GetMorphTargetIndices(int32 LODIndex, int32 MorphTargetIndex) const
+		{
+			return TConstArrayView<int32>();
+		}
+
 		/** Deform the specified positions to match the shape of the previous LOD. */
 		UE_API bool WrapDeformLOD(
 			int32 PrevLODIndex,
@@ -168,12 +189,16 @@ namespace Chaos
 			int32 PrevLODIndex,
 			int32 LODIndex,
 			int32 PrevOffset,
-			int32 Offset);
+			int32 Offset,
+			int32 ActiveMorphTargetIndex = INDEX_NONE,
+			float ActiveMorphTargetWeight = 0.f);
 		// ---- End of the Cloth interface ----
 
 	private:
 		UE_API void SkinPhysicsMesh(
 			int32 LODIndex,
+			int32 ActiveMorphTargetIndex,
+			float ActiveMorphTargetWeight,
 			const FReal LocalSpaceScale,
 			const FVec3& LocalSpaceLocation,
 			TArrayView<Softs::FSolverVec3>& OutPositions,
diff --git a/Engine/Plugins/ChaosClothAsset/Source/ChaosClothAsset/Private/ChaosClothAsset/ClothCollection.cpp b/Engine/Plugins/ChaosClothAsset/Source/ChaosClothAsset/Private/ChaosClothAsset/ClothCollection.cpp
--- a/Engine/Plugins/ChaosClothAsset/Source/ChaosClothAsset/Private/ChaosClothAsset/ClothCollection.cpp
+++ b/Engine/Plugins/ChaosClothAsset/Source/ChaosClothAsset/Private/ChaosClothAsset/ClothCollection.cpp
@@ -102,6 +102,22 @@ namespace UE::Chaos::ClothAsset::Private
 		ClothCollectionAttribute::SeamStitchLookup
 	};
 
+	//~ Sim Morph Targets Group
+	static const TArray<FName> SimMorphTargetsGroupAttributes =
+	{
+		ClothCollectionAttribute::SimMorphTargetName,
+		ClothCollectionAttribute::SimMorphTargetVerticesStart,
+		ClothCollectionAttribute::SimMorphTargetVerticesEnd,
+	};
+
+	//~ Sim Morph Targets Vertices Group
+	static const TArray<FName> SimMorphTargetVerticesGroupAttributes =
+	{
+		ClothCollectionAttribute::SimMorphTargetPositionDelta,
+		ClothCollectionAttribute::SimMorphTargetTangentZDelta,
+		ClothCollectionAttribute::SimMorphTargetSimVertex3DIndex,
+	};
+
 	// Render Faces Group
 	static const TArray<FName> RenderFacesGroupAttributes =
 	{
@@ -138,6 +154,8 @@ namespace UE::Chaos::ClothAsset::Private
 		{ ClothCollectionGroup::SimFaces, SimFacesGroupAttributes },
 		{ ClothCollectionGroup::SimVertices2D, SimVertices2DGroupAttributes },
 		{ ClothCollectionGroup::SimVertices3D, SimVertices3DGroupAttributes },
+		{ ClothCollectionGroup::SimMorphTargets, SimMorphTargetsGroupAttributes },
+		{ ClothCollectionGroup::SimMorphTargetVertices, SimMorphTargetVerticesGroupAttributes },
 		{ ClothCollectionGroup::RenderFaces, RenderFacesGroupAttributes },
 		{ ClothCollectionGroup::RenderVertices, RenderVerticesGroupAttributes },
 		{ ClothCollectionGroup::Fabrics, SimFabricGroupAttributes },
@@ -221,6 +239,16 @@ namespace UE::Chaos::ClothAsset
 		SimVertex2DLookup = ManagedArrayCollection->FindAttribute<TArray<int32>>(ClothCollectionAttribute::SimVertex2DLookup, ClothCollectionGroup::SimVertices3D);
 		SeamStitchLookup = ManagedArrayCollection->FindAttribute<TArray<int32>>(ClothCollectionAttribute::SeamStitchLookup, ClothCollectionGroup::SimVertices3D);
 
+		// Sim Morph Targets Group
+		SimMorphTargetName = ManagedArrayCollection->FindAttribute<FString>(ClothCollectionAttribute::SimMorphTargetName, ClothCollectionGroup::SimMorphTargets);
+		SimMorphTargetVerticesStart = ManagedArrayCollection->FindAttribute<int32>(ClothCollectionAttribute::SimMorphTargetVerticesStart, ClothCollectionGroup::SimMorphTargets);
+		SimMorphTargetVerticesEnd = ManagedArrayCollection->FindAttribute<int32>(ClothCollectionAttribute::SimMorphTargetVerticesEnd, ClothCollectionGroup::SimMorphTargets);
+
+		// Sim Morph Targets Vertices Group
+		SimMorphTargetPositionDelta = ManagedArrayCollection->FindAttribute<FVector3f>(ClothCollectionAttribute::SimMorphTargetPositionDelta, ClothCollectionGroup::SimMorphTargetVertices);
+		SimMorphTargetTangentZDelta = ManagedArrayCollection->FindAttribute<FVector3f>(ClothCollectionAttribute::SimMorphTargetTangentZDelta, ClothCollectionGroup::SimMorphTargetVertices);
+		SimMorphTargetSimVertex3DIndex = ManagedArrayCollection->FindAttribute<int32>(ClothCollectionAttribute::SimMorphTargetSimVertex3DIndex, ClothCollectionGroup::SimMorphTargetVertices);
+
 		// Render Faces Group
 		RenderIndices = ManagedArrayCollection->FindAttribute<FIntVector3>(ClothCollectionAttribute::RenderIndices, ClothCollectionGroup::RenderFaces);
 
@@ -243,6 +271,20 @@ namespace UE::Chaos::ClothAsset
 
 	bool FClothCollection::IsValid(EClothCollectionOptionalSchemas OptionalSchemas) const
 	{
+		if (OptionalSchemas == EClothCollectionOptionalSchemas::CookedOnly)
+		{
+			// Sim Morph Targets Group
+			return
+				SimMorphTargetName &&
+				SimMorphTargetVerticesStart &&
+				SimMorphTargetVerticesEnd &&
+
+				// Sim Morph Target Vertices Group
+				SimMorphTargetPositionDelta &&
+				SimMorphTargetTangentZDelta &&
+				SimMorphTargetSimVertex3DIndex;
+		}
+
 		return
 			// LODs Group
 			PhysicsAssetPathName &&
@@ -311,6 +353,18 @@ namespace UE::Chaos::ClothAsset
 			SimVertex2DLookup &&
 			SeamStitchLookup &&
 
+			// Sim Morph Targets Group
+			(!EnumHasAnyFlags(OptionalSchemas, EClothCollectionOptionalSchemas::SimMorphTargets) ||
+				(SimMorphTargetName &&
+				SimMorphTargetVerticesStart &&
+				SimMorphTargetVerticesEnd)) &&
+
+			// Sim Morph Target Vertices Group
+			(!EnumHasAnyFlags(OptionalSchemas, EClothCollectionOptionalSchemas::SimMorphTargets) ||
+				(SimMorphTargetPositionDelta &&
+				SimMorphTargetTangentZDelta &&
+				SimMorphTargetSimVertex3DIndex)) &&
+
 			// Render Faces Group
 			RenderIndices &&
 
@@ -346,6 +400,24 @@ namespace UE::Chaos::ClothAsset
 		FManagedArrayCollection::FConstructionParameters SimFabricsDependency(ClothCollectionGroup::Fabrics, bSaved, bAllowCircularDependency);
 		FManagedArrayCollection::FConstructionParameters SimVertices2DDependency(ClothCollectionGroup::SimVertices2D, bSaved, bAllowCircularDependency);
 		FManagedArrayCollection::FConstructionParameters SimVertices3DDependency(ClothCollectionGroup::SimVertices3D, bSaved, bAllowCircularDependency);  // Any attribute with this dependency must handle welding and splitting in FCollectionClothSeamFacade
+		FManagedArrayCollection::FConstructionParameters SimMorphTargetVerticesDependency(ClothCollectionGroup::SimMorphTargetVertices, bSaved, bAllowCircularDependency);
+
+		if (OptionalSchemas == EClothCollectionOptionalSchemas::CookedOnly)
+		{
+			// Special case: only define cooked schema.
+
+			// Sim Morph Targets Group
+			SimMorphTargetName = &ManagedArrayCollection->AddAttribute<FString>(ClothCollectionAttribute::SimMorphTargetName, ClothCollectionGroup::SimMorphTargets);
+			SimMorphTargetVerticesStart = &ManagedArrayCollection->AddAttribute<int32>(ClothCollectionAttribute::SimMorphTargetVerticesStart, ClothCollectionGroup::SimMorphTargets);
+			SimMorphTargetVerticesEnd = &ManagedArrayCollection->AddAttribute<int32>(ClothCollectionAttribute::SimMorphTargetVerticesEnd, ClothCollectionGroup::SimMorphTargets);
+
+			// Sim Morph Target Vertices Group
+			SimMorphTargetPositionDelta = &ManagedArrayCollection->AddAttribute<FVector3f>(ClothCollectionAttribute::SimMorphTargetPositionDelta, ClothCollectionGroup::SimMorphTargetVertices);
+			SimMorphTargetTangentZDelta = &ManagedArrayCollection->AddAttribute<FVector3f>(ClothCollectionAttribute::SimMorphTargetTangentZDelta, ClothCollectionGroup::SimMorphTargetVertices);
+			SimMorphTargetSimVertex3DIndex = &ManagedArrayCollection->AddAttribute<int32>(ClothCollectionAttribute::SimMorphTargetSimVertex3DIndex, ClothCollectionGroup::SimMorphTargetVertices);
+
+			return;
+		}
 
 		// LODs Group
 		PhysicsAssetPathName = &ManagedArrayCollection->AddAttribute<FString>(ClothCollectionAttribute::PhysicsAssetPathName, ClothCollectionGroup::Lods);
@@ -415,6 +487,19 @@ namespace UE::Chaos::ClothAsset
 		SimVertex2DLookup = &ManagedArrayCollection->AddAttribute<TArray<int32>>(ClothCollectionAttribute::SimVertex2DLookup, ClothCollectionGroup::SimVertices3D, SimVertices2DDependency);
 		SeamStitchLookup = &ManagedArrayCollection->AddAttribute<TArray<int32>>(ClothCollectionAttribute::SeamStitchLookup, ClothCollectionGroup::SimVertices3D, SeamStitchesDependency);
 
+		if (EnumHasAnyFlags(OptionalSchemas, EClothCollectionOptionalSchemas::SimMorphTargets))
+		{
+			// Sim Morph Targets Group
+			SimMorphTargetName = &ManagedArrayCollection->AddAttribute<FString>(ClothCollectionAttribute::SimMorphTargetName, ClothCollectionGroup::SimMorphTargets);
+			SimMorphTargetVerticesStart = &ManagedArrayCollection->AddAttribute<int32>(ClothCollectionAttribute::SimMorphTargetVerticesStart, ClothCollectionGroup::SimMorphTargets, SimMorphTargetVerticesDependency);
+			SimMorphTargetVerticesEnd = &ManagedArrayCollection->AddAttribute<int32>(ClothCollectionAttribute::SimMorphTargetVerticesEnd, ClothCollectionGroup::SimMorphTargets, SimMorphTargetVerticesDependency);
+
+			// Sim Morph Target Vertices Group
+			SimMorphTargetPositionDelta = &ManagedArrayCollection->AddAttribute<FVector3f>(ClothCollectionAttribute::SimMorphTargetPositionDelta, ClothCollectionGroup::SimMorphTargetVertices);
+			SimMorphTargetTangentZDelta = &ManagedArrayCollection->AddAttribute<FVector3f>(ClothCollectionAttribute::SimMorphTargetTangentZDelta, ClothCollectionGroup::SimMorphTargetVertices);
+			SimMorphTargetSimVertex3DIndex = &ManagedArrayCollection->AddAttribute<int32>(ClothCollectionAttribute::SimMorphTargetSimVertex3DIndex, ClothCollectionGroup::SimMorphTargetVertices, SimVertices3DDependency);
+		}
+
 		// Render Faces Group
 		RenderIndices = &ManagedArrayCollection->AddAttribute<FIntVector3>(ClothCollectionAttribute::RenderIndices, ClothCollectionGroup::RenderFaces, RenderVerticesDependency);
 
diff --git a/Engine/Plugins/ChaosClothAsset/Source/ChaosClothAsset/Private/ChaosClothAsset/ClothCollection.h b/Engine/Plugins/ChaosClothAsset/Source/ChaosClothAsset/Private/ChaosClothAsset/ClothCollection.h
--- a/Engine/Plugins/ChaosClothAsset/Source/ChaosClothAsset/Private/ChaosClothAsset/ClothCollection.h
+++ b/Engine/Plugins/ChaosClothAsset/Source/ChaosClothAsset/Private/ChaosClothAsset/ClothCollection.h
@@ -218,6 +218,16 @@ namespace UE::Chaos::ClothAsset
 		const TManagedArray<TArray<int32>>* GetSimVertex2DLookup() const { return SimVertex2DLookup; }
 		const TManagedArray<TArray<int32>>* GetSeamStitchLookup() const { return SeamStitchLookup; }
 
+		//~ Sim Morph Targets Group
+		const TManagedArray<FString>* GetSimMorphTargetName() const { return SimMorphTargetName; }
+		const TManagedArray<int32>* GetSimMorphTargetVerticesStart() const { return SimMorphTargetVerticesStart; }
+		const TManagedArray<int32>* GetSimMorphTargetVerticesEnd() const { return SimMorphTargetVerticesEnd; }
+
+		//~ Sim Morph Targets Vertices Group
+		const TManagedArray<FVector3f>* GetSimMorphTargetPositionDelta() const { return SimMorphTargetPositionDelta; }
+		const TManagedArray<FVector3f>* GetSimMorphTargetTangentZDelta() const { return SimMorphTargetTangentZDelta; }
+		const TManagedArray<int32>* GetSimMorphTargetSimVertex3DIndex() const { return SimMorphTargetSimVertex3DIndex; }
+
 		//~ Render Faces Group
 		const TManagedArray<FIntVector3>* GetRenderIndices() const { return RenderIndices; }
 
@@ -300,6 +310,16 @@ namespace UE::Chaos::ClothAsset
 		TManagedArray<TArray<int32>>* GetSimVertex2DLookup() { return SimVertex2DLookup; }
 		TManagedArray<TArray<int32>>* GetSeamStitchLookup() { return SeamStitchLookup; }
 
+		//~ Sim Morph Targets Group
+		TManagedArray<FString>* GetSimMorphTargetName() { return SimMorphTargetName; }
+		TManagedArray<int32>* GetSimMorphTargetVerticesStart() { return SimMorphTargetVerticesStart; }
+		TManagedArray<int32>* GetSimMorphTargetVerticesEnd() { return SimMorphTargetVerticesEnd; }
+
+		//~ Sim Morph Targets Vertices Group
+		TManagedArray<FVector3f>* GetSimMorphTargetPositionDelta() { return SimMorphTargetPositionDelta; }
+		TManagedArray<FVector3f>* GetSimMorphTargetTangentZDelta() { return SimMorphTargetTangentZDelta; }
+		TManagedArray<int32>* GetSimMorphTargetSimVertex3DIndex() { return SimMorphTargetSimVertex3DIndex; }
+
 		//~ Render Faces Group
 		TManagedArray<FIntVector3>* GetRenderIndices() { return RenderIndices; }
 
@@ -386,6 +406,16 @@ namespace UE::Chaos::ClothAsset
 		TManagedArray<TArray<int32>>* SimVertex2DLookup; // Lookup into corresponding 2D vertices
 		TManagedArray<TArray<int32>>* SeamStitchLookup; // Lookup into any seam stitches which weld this vertex
 
+		//~ Sim Morph Targets Group
+		TManagedArray<FString>* SimMorphTargetName;
+		TManagedArray<int32>* SimMorphTargetVerticesStart;
+		TManagedArray<int32>* SimMorphTargetVerticesEnd;
+
+		//~ Sim Morph Target Vertices Group
+		TManagedArray<FVector3f>* SimMorphTargetPositionDelta;
+		TManagedArray<FVector3f>* SimMorphTargetTangentZDelta;
+		TManagedArray<int32>* SimMorphTargetSimVertex3DIndex;
+
 		//~ Render Faces Group
 		TManagedArray<FIntVector3>* RenderIndices;
 
diff --git a/Engine/Plugins/ChaosClothAsset/Source/ChaosClothAsset/Private/ChaosClothAsset/ClothGeometryTools.cpp b/Engine/Plugins/ChaosClothAsset/Source/ChaosClothAsset/Private/ChaosClothAsset/ClothGeometryTools.cpp
--- a/Engine/Plugins/ChaosClothAsset/Source/ChaosClothAsset/Private/ChaosClothAsset/ClothGeometryTools.cpp
+++ b/Engine/Plugins/ChaosClothAsset/Source/ChaosClothAsset/Private/ChaosClothAsset/ClothGeometryTools.cpp
@@ -907,18 +907,25 @@ namespace UE::Chaos::ClothAsset
 
 	void FClothGeometryTools::BuildSimMeshFromDynamicMesh(
 		const TSharedRef<FManagedArrayCollection>& ClothCollection,
-		const UE::Geometry::FDynamicMesh3& DynamicMesh, int32 UVChannelIndex, const FVector2f& UVScale, bool bAppend, bool bImportNormals)
+		const UE::Geometry::FDynamicMesh3& DynamicMesh, int32 UVChannelIndex, const FVector2f& UVScale, bool bAppend, bool bImportNormals,
+		TArray<int32>* OutSim2DToSourceIndex)
 	{
 		using namespace Private::SimMeshBuilder;
 
 		if (!bAppend)
 		{
 			DeleteSimMesh(ClothCollection);
+
+			if (OutSim2DToSourceIndex)
+			{
+				OutSim2DToSourceIndex->Reset();
+			}
 		}
 
 		const UE::Geometry::FDynamicMeshAttributeSet* const AttributeSet = DynamicMesh.Attributes();
 		const UE::Geometry::FDynamicMeshUVOverlay* const UVOverlay = AttributeSet ? AttributeSet->GetUVLayer(UVChannelIndex) : nullptr;
-		const UE::Geometry::FDynamicMeshVertexSkinWeightsAttribute* SkinWeights = AttributeSet ? AttributeSet->GetSkinWeightsAttribute(FName("Default")) : nullptr; 
+		const UE::Geometry::FDynamicMeshVertexSkinWeightsAttribute* SkinWeights = AttributeSet ? AttributeSet->GetSkinWeightsAttribute(FName("Default")) : nullptr;
+		const UE::Geometry::FNonManifoldMappingSupport NonManifoldMapping(DynamicMesh);
 
 		TArray<FIsland> Islands;
 		if (UVOverlay)
@@ -973,6 +980,15 @@ namespace UE::Chaos::ClothAsset
 					}
 				}
 
+				// Update OutSim2DToSourceIndex
+				if (OutSim2DToSourceIndex)
+				{
+					OutSim2DToSourceIndex->Reserve(OutSim2DToSourceIndex->Num() + Island.PositionToSourceIndex.Num());
+					for (int32 VertexIndex = 0; VertexIndex < Island.PositionToSourceIndex.Num(); ++VertexIndex)
+					{
+						OutSim2DToSourceIndex->Add(NonManifoldMapping.GetOriginalNonManifoldVertexID(Island.PositionToSourceIndex[VertexIndex]));
+					}
+				}
 			}
 		}
 
@@ -1193,6 +1209,8 @@ namespace UE::Chaos::ClothAsset
 #endif
 		::Chaos::Softs::FEmbeddedSpringFacade SpringFacade(ClothCollection.Get(), ClothCollectionGroup::SimVertices3D);
 		SpringFacade.CleanupAndCompactInvalidSprings();
+
+		Cloth.CompactSimMorphTargets();
 	}
 
 
diff --git a/Engine/Plugins/ChaosClothAsset/Source/ChaosClothAsset/Private/ChaosClothAsset/CollectionClothFacade.cpp b/Engine/Plugins/ChaosClothAsset/Source/ChaosClothAsset/Private/ChaosClothAsset/CollectionClothFacade.cpp
--- a/Engine/Plugins/ChaosClothAsset/Source/ChaosClothAsset/Private/ChaosClothAsset/CollectionClothFacade.cpp
+++ b/Engine/Plugins/ChaosClothAsset/Source/ChaosClothAsset/Private/ChaosClothAsset/CollectionClothFacade.cpp
@@ -108,6 +108,16 @@ namespace UE::Chaos::ClothAsset
 		ResultHash = HashCombineFast(ResultHash, ClothCollection->GetElementsTypeHash(ClothCollection->GetTetherReferenceLength()));
 		ResultHash = HashCombineFast(ResultHash, ClothCollection->GetElementsTypeHash(ClothCollection->GetSimVertex2DLookup()));
 		ResultHash = HashCombineFast(ResultHash, ClothCollection->GetElementsTypeHash(ClothCollection->GetSeamStitchLookup()));
+
+		//~ Sim Morph Targets Group
+		ResultHash = HashCombineFast(ResultHash, ClothCollection->GetElementsTypeHash(ClothCollection->GetSimMorphTargetName()));
+		ResultHash = HashCombineFast(ResultHash, ClothCollection->GetElementsTypeHash(ClothCollection->GetSimMorphTargetVerticesStart()));
+		ResultHash = HashCombineFast(ResultHash, ClothCollection->GetElementsTypeHash(ClothCollection->GetSimMorphTargetVerticesEnd()));
+
+		//~ Sim Morph Target Vertices Group
+		ResultHash = HashCombineFast(ResultHash, ClothCollection->GetElementsTypeHash(ClothCollection->GetSimMorphTargetPositionDelta()));
+		ResultHash = HashCombineFast(ResultHash, ClothCollection->GetElementsTypeHash(ClothCollection->GetSimMorphTargetTangentZDelta()));
+		ResultHash = HashCombineFast(ResultHash, ClothCollection->GetElementsTypeHash(ClothCollection->GetSimMorphTargetSimVertex3DIndex()));
 		
 		//~ Fabrics Group
 		ResultHash = HashCombineFast(ResultHash, ClothCollection->GetElementsTypeHash(ClothCollection->GetFabricBendingStiffness()));
@@ -308,6 +318,54 @@ namespace UE::Chaos::ClothAsset
 		return ClothCollection->GetArrayIndexForContainedElement(ClothCollection->GetSimFacesStart(), ClothCollection->GetSimFacesEnd(), FaceIndex);
 	}
 
+	int32 FCollectionClothConstFacade::GetNumSimMorphTargets() const
+	{
+		return ClothCollection->GetNumElements(ClothCollectionGroup::SimMorphTargets);
+	}
+
+	int32 FCollectionClothConstFacade::FindSimMorphTargetIndexByName(const FString& MorphTargetName) const
+	{
+		TConstArrayView<FString> SimMorphTargetNames = GetSimMorphTargetName();
+		for(int32 Index = 0; Index < SimMorphTargetNames.Num(); ++Index)
+		{
+			if (SimMorphTargetNames[Index].Equals(MorphTargetName))
+			{
+				return Index;
+			}
+		}
+		return INDEX_NONE;
+	}
+
+	FCollectionClothSimMorphTargetConstFacade FCollectionClothConstFacade::GetSimMorphTarget(int32 MorphTargetIndex) const
+	{
+		return FCollectionClothSimMorphTargetConstFacade(ClothCollection, MorphTargetIndex);
+	}
+
+	TConstArrayView<FString> FCollectionClothConstFacade::GetSimMorphTargetName() const
+	{
+		return ClothCollection->GetElements(ClothCollection->GetSimMorphTargetName());
+	}
+
+	int32 FCollectionClothConstFacade::GetNumSimMorphTargetVertices() const
+	{
+		return ClothCollection->GetNumElements(ClothCollectionGroup::SimMorphTargetVertices);
+	}
+
+	TConstArrayView<FVector3f> FCollectionClothConstFacade::GetSimMorphTargetPositionDelta() const
+	{
+		return ClothCollection->GetElements(ClothCollection->GetSimMorphTargetPositionDelta());
+	}
+
+	TConstArrayView<FVector3f> FCollectionClothConstFacade::GetSimMorphTargetTangentZDelta() const
+	{
+		return ClothCollection->GetElements(ClothCollection->GetSimMorphTargetTangentZDelta());
+	}
+
+	TConstArrayView<int32> FCollectionClothConstFacade::GetSimMorphTargetSimVertex3DIndex() const
+	{
+		return ClothCollection->GetElements(ClothCollection->GetSimMorphTargetSimVertex3DIndex());
+	}
+
 	int32 FCollectionClothConstFacade::GetNumRenderPatterns() const
 	{
 		return ClothCollection->GetNumElements(ClothCollectionGroup::RenderPatterns);
@@ -569,6 +627,7 @@ namespace UE::Chaos::ClothAsset
 			GetClothCollection()->SetNumElements(0, ClothCollectionGroup::Solvers);  // Reset all Solvers data
 			GetClothCollection()->SetNumElements(1, ClothCollectionGroup::Solvers);  //
 		}
+		SetNumSimMorphTargets(0);
 	}
 
 	void FCollectionClothFacade::Initialize(const FCollectionClothConstFacade& Other)
@@ -669,6 +728,15 @@ namespace UE::Chaos::ClothAsset
 		}
 #endif
 
+		// Sim Morph Targets Group
+		const int32 StartNumSimMorphTargets = GetNumSimMorphTargets();
+		const int32 OtherNumSimMorphTargets = Other.GetNumSimMorphTargets();
+		SetNumSimMorphTargets(StartNumSimMorphTargets + OtherNumSimMorphTargets);
+		for (int32 TargetIndex = 0; TargetIndex < OtherNumSimMorphTargets; ++TargetIndex)
+		{
+			GetSimMorphTarget(TargetIndex + StartNumSimMorphTargets).Initialize(Other.GetSimMorphTarget(TargetIndex), StartNumSimVertices3D);
+		}
+
 		// Render Patterns Group
 		const int32 StartNumRenderPatterns = GetNumRenderPatterns();
 		const int32 OtherNumRenderPatterns = Other.GetNumRenderPatterns();
@@ -718,6 +786,20 @@ namespace UE::Chaos::ClothAsset
 		}
 	}
 
+	void FCollectionClothFacade::InitializeCookedOnly(const FCollectionClothConstFacade& Other)
+	{
+		Reset();
+
+		// Sim Morph Targets Group
+		const int32 StartNumSimMorphTargets = GetNumSimMorphTargets();
+		const int32 OtherNumSimMorphTargets = Other.GetNumSimMorphTargets();
+		SetNumSimMorphTargets(StartNumSimMorphTargets + OtherNumSimMorphTargets);
+		for (int32 TargetIndex = 0; TargetIndex < OtherNumSimMorphTargets; ++TargetIndex)
+		{
+			GetSimMorphTarget(TargetIndex + StartNumSimMorphTargets).Initialize(Other.GetSimMorphTarget(TargetIndex), 0);
+		}
+	}
+
 	void FCollectionClothFacade::SetPhysicsAssetPathName(const FString& PathName)
 	{
 		if (ClothCollection->GetNumElements(ClothCollectionGroup::Lods))
@@ -915,6 +997,84 @@ namespace UE::Chaos::ClothAsset
 		GetClothCollection()->RemoveElements(ClothCollectionGroup::SimPatterns, SortedDeletionList);
 	}
 
+	void FCollectionClothFacade::SetNumSimMorphTargets(int32 InNumMorphTargets)
+	{
+		if (!IsValid(EClothCollectionOptionalSchemas::SimMorphTargets))
+		{
+			DefineSchema(EClothCollectionOptionalSchemas::SimMorphTargets);
+		}
+
+		const int32 NumMorphTargets = GetNumSimMorphTargets();
+
+		for (int32 TargetIndex = InNumMorphTargets; TargetIndex < NumMorphTargets; ++TargetIndex)
+		{
+			GetSimMorphTarget(TargetIndex).Reset();
+		}
+
+		GetClothCollection()->SetNumElements(InNumMorphTargets, ClothCollectionGroup::SimMorphTargets);
+
+		for (int32 TargetIndex = NumMorphTargets; TargetIndex < InNumMorphTargets; ++TargetIndex)
+		{
+			GetSimMorphTarget(TargetIndex).SetDefaults();
+		}
+	}
+
+	int32 FCollectionClothFacade::AddSimMorphTarget()
+	{
+		const int32 TargetIndex = GetNumSimMorphTargets();
+		SetNumSimMorphTargets(TargetIndex + 1);
+		return TargetIndex;
+	}
+
+	FCollectionClothSimMorphTargetFacade FCollectionClothFacade::GetSimMorphTarget(int32 MorphTargetIndex)
+	{
+		return FCollectionClothSimMorphTargetFacade(GetClothCollection(), MorphTargetIndex);
+	}
+
+	void FCollectionClothFacade::RemoveSimMorphTargets(const TArray<int32>& SortedDeletionList)
+	{
+		for (const int32 TargetToRemove : SortedDeletionList)
+		{
+			GetSimMorphTarget(TargetToRemove).Reset();
+		}
+		GetClothCollection()->RemoveElements(ClothCollectionGroup::SimMorphTargets, SortedDeletionList);
+	}
+
+	void FCollectionClothFacade::CompactSimMorphTargets()
+	{
+		TArray<int32> SortedDeletionList;
+		for (int32 TargetIndex = 0; TargetIndex < GetNumSimMorphTargets(); ++TargetIndex)
+		{
+			FCollectionClothSimMorphTargetFacade MorphTarget = GetSimMorphTarget(TargetIndex);
+			MorphTarget.Compact();
+			if (MorphTarget.IsEmpty())
+			{
+				SortedDeletionList.Add(TargetIndex);
+			}
+		}
+		RemoveSimMorphTargets(SortedDeletionList);
+	}
+
+	TArrayView<FVector3f> FCollectionClothFacade::GetSimMorphTargetPositionDelta()
+	{
+		return GetClothCollection()->GetElements(GetClothCollection()->GetSimMorphTargetPositionDelta());
+	}
+
+	TArrayView<FVector3f> FCollectionClothFacade::GetSimMorphTargetTangentZDelta()
+	{
+		return GetClothCollection()->GetElements(GetClothCollection()->GetSimMorphTargetTangentZDelta());
+	}
+
+	TArrayView<int32> FCollectionClothFacade::GetSimMorphTargetSimVertex3DIndex()
+	{
+		return GetClothCollection()->GetElements(GetClothCollection()->GetSimMorphTargetSimVertex3DIndex());
+	}
+
+	void FCollectionClothFacade::RemoveSimMorphTargetVertices3D(const TArray<int32>& SortedDeletionList)
+	{
+		GetClothCollection()->RemoveElements(ClothCollectionGroup::SimMorphTargetVertices, SortedDeletionList);
+	}
+
 	void FCollectionClothFacade::SetNumRenderPatterns(int32 InNumPatterns)
 	{
 		const int32 NumPatterns = GetNumRenderPatterns();
diff --git a/Engine/Plugins/ChaosClothAsset/Source/ChaosClothAsset/Private/ChaosClothAsset/CollectionClothSimMorphTargetFacade.cpp b/Engine/Plugins/ChaosClothAsset/Source/ChaosClothAsset/Private/ChaosClothAsset/CollectionClothSimMorphTargetFacade.cpp
new file mode 100644
--- /dev/null
+++ b/Engine/Plugins/ChaosClothAsset/Source/ChaosClothAsset/Private/ChaosClothAsset/CollectionClothSimMorphTargetFacade.cpp
@@ -0,0 +1,194 @@
+// Copyright Epic Games, Inc. All Rights Reserved.
+
+#include "ChaosClothAsset/CollectionClothSimMorphTargetFacade.h"
+#include "ChaosClothAsset/ClothCollection.h"
+#include "ChaosClothAsset/ClothCollectionGroup.h"
+
+namespace UE::Chaos::ClothAsset
+{
+	const FString& FCollectionClothSimMorphTargetConstFacade::GetSimMorphTargetName() const
+	{
+		static const FString EmptyString;
+		return ClothCollection->GetSimMorphTargetName() && ClothCollection->GetNumElements(ClothCollectionGroup::SimMorphTargets) > GetElementIndex() ?
+			(*ClothCollection->GetSimMorphTargetName())[GetElementIndex()] : EmptyString;
+	}
+
+	int32 FCollectionClothSimMorphTargetConstFacade::GetNumSimMorphTargetVertices() const
+	{
+		return ClothCollection->GetNumElements(
+			ClothCollection->GetSimMorphTargetVerticesStart(),
+			ClothCollection->GetSimMorphTargetVerticesEnd(),
+			GetElementIndex());
+	}
+
+	int32 FCollectionClothSimMorphTargetConstFacade::GetimMorphTargetVerticesOffset() const
+	{
+		return ClothCollection->GetElementsOffset(
+			ClothCollection->GetSimMorphTargetVerticesStart(),
+			GetBaseElementIndex(),
+			GetElementIndex());
+	}
+
+	TConstArrayView<FVector3f> FCollectionClothSimMorphTargetConstFacade::GetSimMorphTargetPositionDelta() const
+	{
+		return ClothCollection->GetElements(
+			ClothCollection->GetSimMorphTargetPositionDelta(),
+			ClothCollection->GetSimMorphTargetVerticesStart(),
+			ClothCollection->GetSimMorphTargetVerticesEnd(),
+			GetElementIndex());
+	}
+
+	TConstArrayView<FVector3f> FCollectionClothSimMorphTargetConstFacade::GetSimMorphTargetTangentZDelta() const
+	{
+		return ClothCollection->GetElements(
+			ClothCollection->GetSimMorphTargetTangentZDelta(),
+			ClothCollection->GetSimMorphTargetVerticesStart(),
+			ClothCollection->GetSimMorphTargetVerticesEnd(),
+			GetElementIndex());
+	}
+
+	TConstArrayView<int32> FCollectionClothSimMorphTargetConstFacade::GetSimMorphTargetSimVertex3DIndex() const
+	{
+		return ClothCollection->GetElements(
+			ClothCollection->GetSimMorphTargetSimVertex3DIndex(),
+			ClothCollection->GetSimMorphTargetVerticesStart(),
+			ClothCollection->GetSimMorphTargetVerticesEnd(),
+			GetElementIndex());
+	}
+
+	bool FCollectionClothSimMorphTargetConstFacade::IsEmpty() const
+	{
+		return GetNumSimMorphTargetVertices() == 0;
+	}
+
+	FCollectionClothSimMorphTargetConstFacade::FCollectionClothSimMorphTargetConstFacade(const TSharedRef<const FClothCollection>& ClothCollection, int32 InMorphTargetIndex)
+		: ClothCollection(ClothCollection)
+		, MorphTargetIndex(InMorphTargetIndex)
+	{
+		check(ClothCollection->IsValid());
+		check(MorphTargetIndex >= 0 && MorphTargetIndex < ClothCollection->GetNumElements(ClothCollectionGroup::SimMorphTargets));
+	}
+
+	void FCollectionClothSimMorphTargetFacade::Reset()
+	{
+		SetNumSimMorphTargetVertices(0);
+		SetDefaults();
+	}
+
+	void FCollectionClothSimMorphTargetFacade::Initialize(const FCollectionClothSimMorphTargetConstFacade& Other, int32 SimVertex3DOffset)
+	{
+		Reset();
+
+		//~ Sim Morph Targets Group
+		SetSimMorphTargetName(Other.GetSimMorphTargetName());
+
+		//~ Sim Morph Target Vertices Group
+		SetNumSimMorphTargetVertices(Other.GetNumSimMorphTargetVertices());
+		FClothCollection::CopyArrayViewData(GetSimMorphTargetPositionDelta(), Other.GetSimMorphTargetPositionDelta());
+		FClothCollection::CopyArrayViewData(GetSimMorphTargetTangentZDelta(), Other.GetSimMorphTargetTangentZDelta());
+		FClothCollection::CopyArrayViewDataAndApplyOffset(GetSimMorphTargetSimVertex3DIndex(), Other.GetSimMorphTargetSimVertex3DIndex(), SimVertex3DOffset);
+	}
+
+	void FCollectionClothSimMorphTargetFacade::Initialize(const FString& Name, const TConstArrayView<FVector3f>& PositionDeltas, const TConstArrayView<FVector3f>& TangentZDeltas, const TConstArrayView<int32>& SimVertex3DIndices)
+	{
+		check(PositionDeltas.Num() == TangentZDeltas.Num());
+		check(PositionDeltas.Num() == SimVertex3DIndices.Num());
+		Reset();
+
+		//~ Sim Morph Targets Group
+		SetSimMorphTargetName(Name);
+
+		//~ Sim Morph Target Vertices Group
+		SetNumSimMorphTargetVertices(PositionDeltas.Num());
+		FClothCollection::CopyArrayViewData(GetSimMorphTargetPositionDelta(), PositionDeltas);
+		FClothCollection::CopyArrayViewData(GetSimMorphTargetTangentZDelta(), TangentZDeltas);
+		FClothCollection::CopyArrayViewData(GetSimMorphTargetSimVertex3DIndex(), SimVertex3DIndices);
+	}
+
+	void FCollectionClothSimMorphTargetFacade::SetSimMorphTargetName(const FString& MorphTargetName)
+	{
+		(*GetClothCollection()->GetSimMorphTargetName())[GetElementIndex()] = MorphTargetName;
+	}
+
+	void FCollectionClothSimMorphTargetFacade::SetNumSimMorphTargetVertices(int32 NumMorphTargetVertices)
+	{
+		GetClothCollection()->SetNumElements(
+			NumMorphTargetVertices,
+			ClothCollectionGroup::SimMorphTargetVertices,
+			GetClothCollection()->GetSimMorphTargetVerticesStart(),
+			GetClothCollection()->GetSimMorphTargetVerticesEnd(),
+			GetElementIndex());
+	}
+
+	void FCollectionClothSimMorphTargetFacade::RemoveSimMorphTargetVertices(const TArray<int32>& SortedDeletionList)
+	{
+		TArray<int32> GlobalIndexSortedDeletionList;
+		const int32 Offset = GetimMorphTargetVerticesOffset();
+		GlobalIndexSortedDeletionList.SetNumUninitialized(SortedDeletionList.Num());
+		for (int32 Idx = 0; Idx < SortedDeletionList.Num(); ++Idx)
+		{
+			GlobalIndexSortedDeletionList[Idx] = SortedDeletionList[Idx] + Offset;
+		}
+
+		GetClothCollection()->RemoveElements(
+			ClothCollectionGroup::SimMorphTargetVertices,
+			GlobalIndexSortedDeletionList,
+			GetClothCollection()->GetSimMorphTargetVerticesStart(),
+			GetClothCollection()->GetSimMorphTargetVerticesEnd(),
+			GetElementIndex());
+	}
+
+	TArrayView<FVector3f> FCollectionClothSimMorphTargetFacade::GetSimMorphTargetPositionDelta()
+	{
+		return GetClothCollection()->GetElements(
+			GetClothCollection()->GetSimMorphTargetPositionDelta(),
+			GetClothCollection()->GetSimMorphTargetVerticesStart(),
+			GetClothCollection()->GetSimMorphTargetVerticesEnd(),
+			GetElementIndex());
+	}
+
+	TArrayView<FVector3f> FCollectionClothSimMorphTargetFacade::GetSimMorphTargetTangentZDelta()
+	{
+		return GetClothCollection()->GetElements(
+			GetClothCollection()->GetSimMorphTargetTangentZDelta(),
+			GetClothCollection()->GetSimMorphTargetVerticesStart(),
+			GetClothCollection()->GetSimMorphTargetVerticesEnd(),
+			GetElementIndex());
+	}
+
+	TArrayView<int32> FCollectionClothSimMorphTargetFacade::GetSimMorphTargetSimVertex3DIndex()
+	{
+		return GetClothCollection()->GetElements(
+			GetClothCollection()->GetSimMorphTargetSimVertex3DIndex(),
+			GetClothCollection()->GetSimMorphTargetVerticesStart(),
+			GetClothCollection()->GetSimMorphTargetVerticesEnd(),
+			GetElementIndex());
+	}
+
+	void FCollectionClothSimMorphTargetFacade::Compact()
+	{
+		TArray<int32> SortedDeletionList;
+		const TConstArrayView<int32> SimVertexIndices = GetSimMorphTargetSimVertex3DIndex();
+		for(int32 Index = 0; Index <SimVertexIndices.Num(); ++Index)
+		{
+			if (SimVertexIndices[Index] == INDEX_NONE)
+			{
+				SortedDeletionList.Add(Index);
+			}
+		}
+		RemoveSimMorphTargetVertices(SortedDeletionList);
+	}
+
+	FCollectionClothSimMorphTargetFacade::FCollectionClothSimMorphTargetFacade(const TSharedRef<FClothCollection>& ClothCollection, int32 InMorphTargetIndex)
+		: FCollectionClothSimMorphTargetConstFacade(ClothCollection, InMorphTargetIndex)
+	{
+	}
+
+	void FCollectionClothSimMorphTargetFacade::SetDefaults()
+	{
+		const int32 ElementIndex = GetElementIndex();
+		SetSimMorphTargetName(FString());
+		(*GetClothCollection()->GetSimMorphTargetVerticesStart())[ElementIndex] = INDEX_NONE;
+		(*GetClothCollection()->GetSimMorphTargetVerticesEnd())[ElementIndex] = INDEX_NONE;
+	}
+}  // End namespace UE::Chaos::ClothAsset
diff --git a/Engine/Plugins/ChaosClothAsset/Source/ChaosClothAsset/Public/ChaosClothAsset/ClothCollectionAttribute.h b/Engine/Plugins/ChaosClothAsset/Source/ChaosClothAsset/Public/ChaosClothAsset/ClothCollectionAttribute.h
--- a/Engine/Plugins/ChaosClothAsset/Source/ChaosClothAsset/Public/ChaosClothAsset/ClothCollectionAttribute.h
+++ b/Engine/Plugins/ChaosClothAsset/Source/ChaosClothAsset/Public/ChaosClothAsset/ClothCollectionAttribute.h
@@ -74,6 +74,16 @@ namespace UE::Chaos::ClothAsset
 		inline const FName SimVertex2DLookup(TEXT("SimVertex2DLookup"));
 		inline const FName SeamStitchLookup(TEXT("SeamStitchLookup"));
 
+		// Sim Morph Targets Group
+		inline const FName SimMorphTargetName(TEXT("SimMorphTargetName"));
+		inline const FName SimMorphTargetVerticesStart(TEXT("SimMorphTargetVerticesStart"));
+		inline const FName SimMorphTargetVerticesEnd(TEXT("SimMorphTargetVerticesEnd"));
+
+		// Sim Morph Target Vertices Group
+		inline const FName SimMorphTargetPositionDelta(TEXT("SimMorphTargetPositionDelta"));
+		inline const FName SimMorphTargetTangentZDelta(TEXT("SimMorphTargetTangentZDelta"));
+		inline const FName SimMorphTargetSimVertex3DIndex(TEXT("SimMorphTargetSimVertex3DIndex"));
+
 		// Render Faces Group
 		inline const FName RenderIndices(TEXT("RenderIndices"));
 
diff --git a/Engine/Plugins/ChaosClothAsset/Source/ChaosClothAsset/Public/ChaosClothAsset/ClothCollectionGroup.h b/Engine/Plugins/ChaosClothAsset/Source/ChaosClothAsset/Public/ChaosClothAsset/ClothCollectionGroup.h
--- a/Engine/Plugins/ChaosClothAsset/Source/ChaosClothAsset/Public/ChaosClothAsset/ClothCollectionGroup.h
+++ b/Engine/Plugins/ChaosClothAsset/Source/ChaosClothAsset/Public/ChaosClothAsset/ClothCollectionGroup.h
@@ -35,5 +35,9 @@ namespace UE::Chaos::ClothAsset
 		inline const FName RenderVertices = FName(TEXT("RenderVertices"));
 		/** Contains all the fabrics simulation parameters used by the sim patterns . */
 		inline const FName Fabrics = FName(TEXT("Fabrics"));
+		/** Contains sim morph target relationships to other groups. */
+		inline const FName SimMorphTargets = FName(TEXT("SimMorphTargets"));
+		/** Contains morph target vertices */
+		inline const FName SimMorphTargetVertices = FName(TEXT("SimMorphTargetVertices"));
 	}
 }  // End namespace UE::Chaos::ClothAsset
diff --git a/Engine/Plugins/ChaosClothAsset/Source/ChaosClothAsset/Public/ChaosClothAsset/ClothCollectionOptionalSchemas.h b/Engine/Plugins/ChaosClothAsset/Source/ChaosClothAsset/Public/ChaosClothAsset/ClothCollectionOptionalSchemas.h
--- a/Engine/Plugins/ChaosClothAsset/Source/ChaosClothAsset/Public/ChaosClothAsset/ClothCollectionOptionalSchemas.h
+++ b/Engine/Plugins/ChaosClothAsset/Source/ChaosClothAsset/Public/ChaosClothAsset/ClothCollectionOptionalSchemas.h
@@ -13,7 +13,9 @@ namespace UE::Chaos::ClothAsset
 		None = 0,
 		RenderDeformer = 1 << 1,
 		Solvers = 1 << 2,
-		Fabrics = 1 << 3
+		Fabrics = 1 << 3,
+		SimMorphTargets = 1 << 4,
+		CookedOnly = 1 << 5, // If this flag is found, only cooked schemas will be considered.
 	};
 	ENUM_CLASS_FLAGS(EClothCollectionOptionalSchemas)
 }  // End namespace UE::Chaos::ClothAsset
diff --git a/Engine/Plugins/ChaosClothAsset/Source/ChaosClothAsset/Public/ChaosClothAsset/ClothGeometryTools.h b/Engine/Plugins/ChaosClothAsset/Source/ChaosClothAsset/Public/ChaosClothAsset/ClothGeometryTools.h
--- a/Engine/Plugins/ChaosClothAsset/Source/ChaosClothAsset/Public/ChaosClothAsset/ClothGeometryTools.h
+++ b/Engine/Plugins/ChaosClothAsset/Source/ChaosClothAsset/Public/ChaosClothAsset/ClothGeometryTools.h
@@ -89,7 +89,7 @@ namespace UE::Chaos::ClothAsset
 		*/
 		static UE_API void BuildSimMeshFromDynamicMesh(
 			const TSharedRef<FManagedArrayCollection>& ClothCollection,
-			const UE::Geometry::FDynamicMesh3& DynamicMesh, int32 UVChannelIndex, const FVector2f& UVScale, bool bAppend, bool bImportNormals = false);
+			const UE::Geometry::FDynamicMesh3& DynamicMesh, int32 UVChannelIndex, const FVector2f& UVScale, bool bAppend, bool bImportNormals = false, TArray<int32>* OutSim2DToSourceIndex = nullptr);
 
 		/**
 		* Remove (topologically) degenerate triangles. Remove any vertices that aren't in a triangle. Compact any lookup arrays that contain INDEX_NONEs.
diff --git a/Engine/Plugins/ChaosClothAsset/Source/ChaosClothAsset/Public/ChaosClothAsset/CollectionClothFacade.h b/Engine/Plugins/ChaosClothAsset/Source/ChaosClothAsset/Public/ChaosClothAsset/CollectionClothFacade.h
--- a/Engine/Plugins/ChaosClothAsset/Source/ChaosClothAsset/Public/ChaosClothAsset/CollectionClothFacade.h
+++ b/Engine/Plugins/ChaosClothAsset/Source/ChaosClothAsset/Public/ChaosClothAsset/CollectionClothFacade.h
@@ -7,6 +7,7 @@
 #include "ChaosClothAsset/CollectionClothFabricFacade.h"
 #include "ChaosClothAsset/CollectionClothRenderPatternFacade.h"
 #include "ChaosClothAsset/CollectionClothSeamFacade.h"
+#include "ChaosClothAsset/CollectionClothSimMorphTargetFacade.h"
 #include "ChaosClothAsset/CollectionClothSimPatternFacade.h"
 #include "ChaosClothAsset/IsUserAttributeType.h"
 
@@ -116,6 +117,26 @@ namespace UE::Chaos::ClothAsset
 		/** Convenience to find which sim pattern a sim face belongs to */
 		UE_API int32 FindSimPatternByFaceIndex(int32 FaceIndex) const;
 
+		//~ Sim Morph Targets Group
+		/** Return the number of sim morph targets in this collection. */
+		UE_API int32 GetNumSimMorphTargets() const;
+		/** Lookup sim morph target by name. Returns INDEX_NONE if not found. */
+		UE_API int32 FindSimMorphTargetIndexByName(const FString& MorphTargetName) const;
+		/** Return a sim morph target facade for the specified morph target index. */
+		UE_API FCollectionClothSimMorphTargetConstFacade GetSimMorphTarget(int32 MorphTargetIndex) const;
+		/** Return a view of all sim morph target names */
+		UE_API TConstArrayView<FString> GetSimMorphTargetName() const;
+
+		//~ Sim Morph Target Vertices Group
+		/** Return the total number of sim morph target vertices in this collection. */
+		UE_API int32 GetNumSimMorphTargetVertices() const;
+		/** Return a view of all sim morph target position deltas */
+		UE_API TConstArrayView<FVector3f> GetSimMorphTargetPositionDelta() const;
+		/** Return a view of all sim morph target tangent z deltas */
+		UE_API TConstArrayView<FVector3f> GetSimMorphTargetTangentZDelta() const;
+		/** Return a view of all sim morph target sim vertex 3d indices */
+		UE_API TConstArrayView<int32> GetSimMorphTargetSimVertex3DIndex() const;
+
 		//~ Render Patterns Group
 		/** Return the number of patterns in this collection. */
 		UE_API int32 GetNumRenderPatterns() const;
@@ -227,8 +248,8 @@ namespace UE::Chaos::ClothAsset
 		/** Append data from another cloth collection. */
 		UE_API void Append(const FCollectionClothConstFacade& Other);
 
-		/** Append to this cloth another cloth. */
-		//void Append(const FCollectionClothConstFacade& Other);
+		/** Copy only data that is not stripped on cook */
+		UE_API void InitializeCookedOnly(const FCollectionClothConstFacade& Other);
 
 		//~ LOD (single per collection) Group
 		/** Set the physics asset path name. */
@@ -284,6 +305,30 @@ namespace UE::Chaos::ClothAsset
 		/** Remove a sorted list of sim patterns. */
 		UE_API void RemoveSimPatterns(const TArray<int32>& SortedDeletionList);
 
+		//~ Sim Morph Targets Group
+		/** Set the new number of sim morph targets to this cloth. */
+		UE_API void SetNumSimMorphTargets(int32 NumMorphTargets);
+		/** Add a sim morph target to this cloth and return its index in the sim morph target pattern list */
+		UE_API int32 AddSimMorphTarget();
+		/** Return a sim morph target facade for the specified morph target index. */
+		UE_API FCollectionClothSimMorphTargetFacade GetSimMorphTarget(int32 MorphTargetIndex);
+		/** Add a new sim morph target to this cloth, and return the morph target facade set to its index. */
+		FCollectionClothSimMorphTargetFacade AddGetSimMorphTarget() { return GetSimMorphTarget(AddSimMorphTarget()); }
+		/** Remove a sorted list of sim morph targets. */
+		UE_API void RemoveSimMorphTargets(const TArray<int32>& SortedDeletionList);
+		/** Remove all morph target vertices with invalid indices. This will also remove any morph targets that are empty. */
+		UE_API void CompactSimMorphTargets();
+
+		//~ Sim Morph Target Vertices Group
+		/** Return a view of all sim morph target position deltas */
+		UE_API TArrayView<FVector3f> GetSimMorphTargetPositionDelta();
+		/** Return a view of all sim morph target tangent z deltas */
+		UE_API TArrayView<FVector3f> GetSimMorphTargetTangentZDelta();
+		/** Return a view of all sim morph target sim vertex 3d indices */
+		UE_API TArrayView<int32> GetSimMorphTargetSimVertex3DIndex();
+		/** Remove a sorted list of sim morph target vertices. */
+		UE_API void RemoveSimMorphTargetVertices3D(const TArray<int32>& SortedDeletionList);
+
 		//~ Render Patterns Group
 		/** Set the new number of patterns to this cloth LOD. */
 		UE_API void SetNumRenderPatterns(int32 NumPatterns);
diff --git a/Engine/Plugins/ChaosClothAsset/Source/ChaosClothAsset/Public/ChaosClothAsset/CollectionClothSimMorphTargetFacade.h b/Engine/Plugins/ChaosClothAsset/Source/ChaosClothAsset/Public/ChaosClothAsset/CollectionClothSimMorphTargetFacade.h
new file mode 100644
--- /dev/null
+++ b/Engine/Plugins/ChaosClothAsset/Source/ChaosClothAsset/Public/ChaosClothAsset/CollectionClothSimMorphTargetFacade.h
@@ -0,0 +1,107 @@
+// Copyright Epic Games, Inc. All Rights Reserved.
+
+#pragma once
+
+#include "Templates/SharedPointer.h"
+
+struct FManagedArrayCollection;
+
+namespace UE::Chaos::ClothAsset
+{
+	/**
+	 * Cloth Asset collection sim morph target facade class to access cloth sim morph target data.
+	 * Constructed from FCollectionClothConstFacade.
+	 * Const access (read only) version.
+	 */
+	class CHAOSCLOTHASSET_API FCollectionClothSimMorphTargetConstFacade
+	{
+	public:
+		FCollectionClothSimMorphTargetConstFacade() = delete;
+
+		FCollectionClothSimMorphTargetConstFacade(const FCollectionClothSimMorphTargetConstFacade&) = delete;
+		FCollectionClothSimMorphTargetConstFacade& operator=(const FCollectionClothSimMorphTargetConstFacade&) = delete;
+
+		FCollectionClothSimMorphTargetConstFacade(FCollectionClothSimMorphTargetConstFacade&&) = default;
+		FCollectionClothSimMorphTargetConstFacade& operator=(FCollectionClothSimMorphTargetConstFacade&&) = default;
+
+		virtual ~FCollectionClothSimMorphTargetConstFacade() = default;
+
+		/** Return the name for this morph target */
+		const FString& GetSimMorphTargetName() const;
+
+		//~ Sim Morph Target Vertices Group
+		/** Return the total number of vertices for this morph target. */
+		int32 GetNumSimMorphTargetVertices() const;
+		/** Return the morph target vertices offset for this morph target in the morph target vertices for the collection. */
+		int32 GetimMorphTargetVerticesOffset() const;
+		TConstArrayView<FVector3f> GetSimMorphTargetPositionDelta() const;
+		TConstArrayView<FVector3f> GetSimMorphTargetTangentZDelta() const;
+		TConstArrayView<int32> GetSimMorphTargetSimVertex3DIndex() const;
+
+		bool IsEmpty() const;
+
+		/** Return the Morph Target index this facade has been created with. */
+		int32 GetMorphTargetIndex() const { return MorphTargetIndex; }
+
+	protected:
+		friend class FCollectionClothSimMorphTargetFacade;  // For other instances access
+		friend class FCollectionClothConstFacade;
+		FCollectionClothSimMorphTargetConstFacade(const TSharedRef<const class FClothCollection>& InClothCollection, int32 InPatternIndex);
+
+		static constexpr int32 GetBaseElementIndex() { return 0; }
+		int32 GetElementIndex() const { return GetBaseElementIndex() + MorphTargetIndex; }
+
+		TSharedRef<const class FClothCollection> ClothCollection;
+		int32 MorphTargetIndex;
+	};
+
+	/**
+	 * Cloth Asset collection sim morph target facade class to access cloth sim morph target data.
+	 * Constructed from FCollectionClothFacade.
+	 * Non-const access (read/write) version.
+	 */
+	class CHAOSCLOTHASSET_API FCollectionClothSimMorphTargetFacade final : public FCollectionClothSimMorphTargetConstFacade
+	{
+	public:
+		FCollectionClothSimMorphTargetFacade() = delete;
+
+		FCollectionClothSimMorphTargetFacade(const FCollectionClothSimMorphTargetFacade&) = delete;
+		FCollectionClothSimMorphTargetFacade& operator=(const FCollectionClothSimMorphTargetFacade&) = delete;
+
+		FCollectionClothSimMorphTargetFacade(FCollectionClothSimMorphTargetFacade&&) = default;
+		FCollectionClothSimMorphTargetFacade& operator=(FCollectionClothSimMorphTargetFacade&&) = default;
+
+		virtual ~FCollectionClothSimMorphTargetFacade() override = default;
+
+		/** Remove all geometry from this morph target. */
+		void Reset();
+
+		/** Initialize from another morph target. Assumes all indices match between source and target. */
+		void Initialize(const FCollectionClothSimMorphTargetConstFacade& Other, int32 SimVertex3DOffset);
+
+		/** Initialize */
+		void Initialize(const FString& Name, const TConstArrayView<FVector3f>& PositionDeltas, const TConstArrayView<FVector3f>& TangentZDeltas, const TConstArrayView<int32>& SimVertex3DIndices);
+
+		/** Set the name for this morph target. */
+		void SetSimMorphTargetName(const FString& MorphTargetName);
+
+		//~ Sim Morph Target Vertices Group
+		/** Grow or shrink the space reserved for morph target vertices for this morph target within the cloth collection and return its start index. */
+		void SetNumSimMorphTargetVertices(int32 NumMorphTargetVertices);
+		void RemoveSimMorphTargetVertices(const TArray<int32>& SortedDeletionList);
+		TArrayView<FVector3f> GetSimMorphTargetPositionDelta();
+		TArrayView<FVector3f> GetSimMorphTargetTangentZDelta();
+		TArrayView<int32> GetSimMorphTargetSimVertex3DIndex();
+
+		/** Remove all morph target vertices with invalid indices. */
+		void Compact();
+
+	protected:
+		friend class FCollectionClothFacade;
+		FCollectionClothSimMorphTargetFacade(const TSharedRef<class FClothCollection>& InClothCollection, int32 InMorphTargetIndex);
+
+		void SetDefaults();
+
+		TSharedRef<class FClothCollection> GetClothCollection() { return ConstCastSharedRef<class FClothCollection>(ClothCollection); }
+	};
+}  // End namespace UE::Chaos::ClothAsset
diff --git a/Engine/Plugins/ChaosClothAsset/Source/ChaosClothAssetEngine/Private/ChaosClothAsset/ClothAsset.cpp b/Engine/Plugins/ChaosClothAsset/Source/ChaosClothAssetEngine/Private/ChaosClothAsset/ClothAsset.cpp
--- a/Engine/Plugins/ChaosClothAsset/Source/ChaosClothAssetEngine/Private/ChaosClothAsset/ClothAsset.cpp
+++ b/Engine/Plugins/ChaosClothAsset/Source/ChaosClothAssetEngine/Private/ChaosClothAsset/ClothAsset.cpp
@@ -151,6 +151,15 @@ namespace UE::Chaos::ClothAsset::Private
 					EmbeddedSpringFacade.Append(InEmbeddedSpringFacade, VertexOffset);
 				}
 
+				// Morph targets
+				const UE::Chaos::ClothAsset::FCollectionClothConstFacade InClothFacade(ClothCollection);
+				if (InClothFacade.IsValid(UE::Chaos::ClothAsset::EClothCollectionOptionalSchemas::CookedOnly))
+				{
+					UE::Chaos::ClothAsset::FCollectionClothFacade ClothFacade(PropertyCollection);
+					ClothFacade.DefineSchema(UE::Chaos::ClothAsset::EClothCollectionOptionalSchemas::CookedOnly);
+					ClothFacade.InitializeCookedOnly(InClothFacade);
+				}
+
 				OutputCollections.Emplace(MoveTemp(PropertyCollection));
 
 				UE_LOG(LogChaosClothAsset, Display, TEXT("TrimOnCook[ON] %s:[%d] [size:%d]"), 
diff --git a/Engine/Plugins/ChaosClothAsset/Source/ChaosClothAssetEngine/Private/ChaosClothAsset/ClothSimulationMesh.cpp b/Engine/Plugins/ChaosClothAsset/Source/ChaosClothAssetEngine/Private/ChaosClothAsset/ClothSimulationMesh.cpp
--- a/Engine/Plugins/ChaosClothAsset/Source/ChaosClothAssetEngine/Private/ChaosClothAsset/ClothSimulationMesh.cpp
+++ b/Engine/Plugins/ChaosClothAsset/Source/ChaosClothAssetEngine/Private/ChaosClothAsset/ClothSimulationMesh.cpp
@@ -4,6 +4,7 @@
 #include "ChaosClothAsset/ClothAsset.h"
 #include "ChaosClothAsset/ClothSimulationModel.h"
 #include "ChaosClothAsset/ClothSimulationContext.h"
+#include "ChaosClothAsset/CollectionClothFacade.h"
 #include "ClothSimulationMesh.h"
 #include "ReferenceSkeleton.h"
 
@@ -27,6 +28,11 @@ namespace UE::Chaos::ClothAsset
 			ReferenceBoneName = ClothSimulationModel.UsedBoneNames[UsedBoneNameIndex];
 		}
 #endif
+		ClothFacades.Reserve(ManagedArrayCollections.Num());
+		for (const TSharedRef<const FManagedArrayCollection>& Collection : ManagedArrayCollections)
+		{
+			ClothFacades.Emplace(Collection);
+		}
 	}
 
 	int32 FClothSimulationMesh::GetNumLODs() const
@@ -239,4 +245,39 @@ namespace UE::Chaos::ClothAsset
 	{
 		return ManagedArrayCollections.IsValidIndex(LODIndex) ? ManagedArrayCollections[LODIndex] : TSharedPtr<const FManagedArrayCollection>();
 	}
+
+	int32 FClothSimulationMesh::FindMorphTargetByName(int32 LODIndex, const FString& Name) const
+	{
+		return ClothFacades.IsValidIndex(LODIndex) ? ClothFacades[LODIndex].FindSimMorphTargetIndexByName(Name) : INDEX_NONE;
+	}
+
+	TConstArrayView<FVector3f> FClothSimulationMesh::GetMorphTargetPositionDeltas(int32 LODIndex, int32 MorphTargetIndex) const
+	{
+		if (ClothFacades.IsValidIndex(LODIndex) && (MorphTargetIndex >= 0 && MorphTargetIndex < ClothFacades[LODIndex].GetNumSimMorphTargets()))
+		{
+			FCollectionClothSimMorphTargetConstFacade MorphTargetFacade = ClothFacades[LODIndex].GetSimMorphTarget(MorphTargetIndex);
+			return MorphTargetFacade.GetSimMorphTargetPositionDelta();
+		}
+		return TConstArrayView<FVector3f>();
+	}
+
+	TConstArrayView<FVector3f> FClothSimulationMesh::GetMorphTargetTangentZDeltas(int32 LODIndex, int32 MorphTargetIndex) const
+	{
+		if (ClothFacades.IsValidIndex(LODIndex) && (MorphTargetIndex >= 0 && MorphTargetIndex < ClothFacades[LODIndex].GetNumSimMorphTargets()))
+		{
+			FCollectionClothSimMorphTargetConstFacade MorphTargetFacade = ClothFacades[LODIndex].GetSimMorphTarget(MorphTargetIndex);
+			return MorphTargetFacade.GetSimMorphTargetTangentZDelta();
+		}
+		return TConstArrayView<FVector3f>();
+	}
+
+	TConstArrayView<int32> FClothSimulationMesh::GetMorphTargetIndices(int32 LODIndex, int32 MorphTargetIndex) const
+	{
+		if (ClothFacades.IsValidIndex(LODIndex) && (MorphTargetIndex >= 0 && MorphTargetIndex < ClothFacades[LODIndex].GetNumSimMorphTargets()))
+		{
+			FCollectionClothSimMorphTargetConstFacade MorphTargetFacade = ClothFacades[LODIndex].GetSimMorphTarget(MorphTargetIndex);
+			return MorphTargetFacade.GetSimMorphTargetSimVertex3DIndex();
+		}
+		return TConstArrayView<int32>();
+	}
 }
diff --git a/Engine/Plugins/ChaosClothAsset/Source/ChaosClothAssetEngine/Private/ChaosClothAsset/ClothSimulationMesh.h b/Engine/Plugins/ChaosClothAsset/Source/ChaosClothAssetEngine/Private/ChaosClothAsset/ClothSimulationMesh.h
--- a/Engine/Plugins/ChaosClothAsset/Source/ChaosClothAssetEngine/Private/ChaosClothAsset/ClothSimulationMesh.h
+++ b/Engine/Plugins/ChaosClothAsset/Source/ChaosClothAssetEngine/Private/ChaosClothAsset/ClothSimulationMesh.h
@@ -9,6 +9,7 @@ struct FChaosClothSimulationModel;
 namespace UE::Chaos::ClothAsset
 {
 	struct FClothSimulationContext;
+	class FCollectionClothConstFacade;
 
 	class FClothSimulationMesh : public ::Chaos::FClothingSimulationMesh
 	{
@@ -57,11 +58,16 @@ namespace UE::Chaos::ClothAsset
 		virtual TConstArrayView<FMeshToMeshVertData> GetTransitionUpSkinData(int32 LODIndex) const override;
 		virtual TConstArrayView<FMeshToMeshVertData> GetTransitionDownSkinData(int32 LODIndex) const override;
 		virtual TSharedPtr<const FManagedArrayCollection> GetManagedArrayCollection(int32 LODIndex) const override;
+		virtual int32 FindMorphTargetByName(int32 LODIndex, const FString& Name) const override;
+		virtual TConstArrayView<FVector3f> GetMorphTargetPositionDeltas(int32 LODIndex, int32 MorphTargetIndex) const override;
+		virtual TConstArrayView<FVector3f> GetMorphTargetTangentZDeltas(int32 LODIndex, int32 MorphTargetIndex) const override;
+		virtual TConstArrayView<int32> GetMorphTargetIndices(int32 LODIndex, int32 MorphTargetIndex) const override;
 		//~ End FClothingSimulationMesh Interface
 
 	private:
 		const FChaosClothSimulationModel& ClothSimulationModel;
 		const FClothSimulationContext& ClothSimulationContext;
 		const TArray<TSharedRef<const FManagedArrayCollection>>& ManagedArrayCollections;
+		TArray<FCollectionClothConstFacade> ClothFacades;
 	};
 }
diff --git a/Engine/Plugins/ChaosClothAssetEditor/Source/ChaosClothAssetDataflowNodes/Private/ChaosClothAsset/ChaosClothAssetDataflowNodesModule.cpp b/Engine/Plugins/ChaosClothAssetEditor/Source/ChaosClothAssetDataflowNodes/Private/ChaosClothAsset/ChaosClothAssetDataflowNodesModule.cpp
--- a/Engine/Plugins/ChaosClothAssetEditor/Source/ChaosClothAssetDataflowNodes/Private/ChaosClothAsset/ChaosClothAssetDataflowNodesModule.cpp
+++ b/Engine/Plugins/ChaosClothAssetEditor/Source/ChaosClothAssetDataflowNodes/Private/ChaosClothAsset/ChaosClothAssetDataflowNodesModule.cpp
@@ -41,6 +41,7 @@
 #include "ChaosClothAsset/SimulationLongRangeAttachmentConfigNode.h"
 #include "ChaosClothAsset/SimulationMassConfigNode.h"
 #include "ChaosClothAsset/SimulationMaxDistanceConfigNode.h"
+#include "ChaosClothAsset/SimulationMorphTargetConfigNode.h"
 #include "ChaosClothAsset/SimulationPBDAreaSpringConfigNode.h"
 #include "ChaosClothAsset/SimulationPBDBendingElementConfigNode.h"
 #include "ChaosClothAsset/SimulationPBDBendingSpringConfigNode.h"
@@ -133,6 +134,7 @@ namespace UE::Chaos::ClothAsset
 			DATAFLOW_NODE_REGISTER_CREATION_FACTORY(FChaosClothAssetSimulationLongRangeAttachmentConfigNode_v2);
 			DATAFLOW_NODE_REGISTER_CREATION_FACTORY(FChaosClothAssetSimulationMassConfigNode);
 			DATAFLOW_NODE_REGISTER_CREATION_FACTORY(FChaosClothAssetSimulationMaxDistanceConfigNode);
+			DATAFLOW_NODE_REGISTER_CREATION_FACTORY(FChaosClothAssetSimulationMorphTargetConfigNode);
 			DATAFLOW_NODE_REGISTER_CREATION_FACTORY(FChaosClothAssetSimulationPBDAreaSpringConfigNode);
 			DATAFLOW_NODE_REGISTER_CREATION_FACTORY(FChaosClothAssetSimulationPBDBendingElementConfigNode);
 			DATAFLOW_NODE_REGISTER_CREATION_FACTORY(FChaosClothAssetSimulationPBDBendingSpringConfigNode);
diff --git a/Engine/Plugins/ChaosClothAssetEditor/Source/ChaosClothAssetDataflowNodes/Private/ChaosClothAsset/ClothDataflowTools.cpp b/Engine/Plugins/ChaosClothAssetEditor/Source/ChaosClothAssetDataflowNodes/Private/ChaosClothAsset/ClothDataflowTools.cpp
--- a/Engine/Plugins/ChaosClothAssetEditor/Source/ChaosClothAssetDataflowNodes/Private/ChaosClothAsset/ClothDataflowTools.cpp
+++ b/Engine/Plugins/ChaosClothAssetEditor/Source/ChaosClothAssetDataflowNodes/Private/ChaosClothAsset/ClothDataflowTools.cpp
@@ -435,7 +435,7 @@ namespace UE::Chaos::ClothAsset
 		ClothPatternFacade.SetRenderMaterialPathName(RenderMaterialPathName);
 	}
 	
-	void FClothDataflowTools::AddSimPatternsFromSkeletalMeshSection(const TSharedRef<FManagedArrayCollection>& ClothCollection, const FSkeletalMeshLODModel& SkeletalMeshModel, const int32 SectionIndex, const int32 UVChannelIndex, const FVector2f& UVScale, bool bImportNormals)
+	void FClothDataflowTools::AddSimPatternsFromSkeletalMeshSection(const TSharedRef<FManagedArrayCollection>& ClothCollection, const FSkeletalMeshLODModel& SkeletalMeshModel, const int32 SectionIndex, const int32 UVChannelIndex, const FVector2f& UVScale, bool bImportNormals, TArray<int32>* OutSim2DToSourceVertex)
 	{
 		check(SectionIndex < SkeletalMeshModel.Sections.Num());
 
@@ -452,7 +452,17 @@ namespace UE::Chaos::ClothAsset
 		UE::Geometry::FNonManifoldMappingSupport::AttachNonManifoldVertexMappingData(SkelMeshSectionToDynamicMesh.ToSrcVertIDMap, DynamicMesh);
 
 		constexpr bool bAppend = true;
-		FClothGeometryTools::BuildSimMeshFromDynamicMesh(ClothCollection, DynamicMesh, UVChannelIndex, UVScale, bAppend, bImportNormals);
+		const int32 OrigNumSimVertices2D = FCollectionClothConstFacade(ClothCollection).GetNumSimVertices2D();
+		check(!OutSim2DToSourceVertex || OutSim2DToSourceVertex->Num() == OrigNumSimVertices2D);
+		FClothGeometryTools::BuildSimMeshFromDynamicMesh(ClothCollection, DynamicMesh, UVChannelIndex, UVScale, bAppend, bImportNormals, OutSim2DToSourceVertex);
+		if (OutSim2DToSourceVertex)
+		{
+			// SrcVertID doesn't include SourceSection.BaseVertexIndex.
+			for (int32 Index = OrigNumSimVertices2D; Index < OutSim2DToSourceVertex->Num(); ++Index)
+			{
+				(*OutSim2DToSourceVertex)[Index] += SkeletalMeshModel.Sections[SectionIndex].BaseVertexIndex;
+			}
+		}
 	}
 
 	void FClothDataflowTools::LogAndToastWarning(const FDataflowNode& DataflowNode, const FText& Headline, const FText& Details)
diff --git a/Engine/Plugins/ChaosClothAssetEditor/Source/ChaosClothAssetDataflowNodes/Private/ChaosClothAsset/ClothDataflowTools.h b/Engine/Plugins/ChaosClothAssetEditor/Source/ChaosClothAssetDataflowNodes/Private/ChaosClothAsset/ClothDataflowTools.h
--- a/Engine/Plugins/ChaosClothAssetEditor/Source/ChaosClothAssetDataflowNodes/Private/ChaosClothAsset/ClothDataflowTools.h
+++ b/Engine/Plugins/ChaosClothAssetEditor/Source/ChaosClothAssetDataflowNodes/Private/ChaosClothAsset/ClothDataflowTools.h
@@ -23,7 +23,7 @@ namespace UE::Chaos::ClothAsset
 	{
 		static void AddRenderPatternFromSkeletalMeshSection(const TSharedRef<FManagedArrayCollection>& ClothCollection, const FSkeletalMeshLODModel& SkeletalMeshModel, const int32 SectionIndex, const FString& RenderMaterialPathName);
 
-		static void AddSimPatternsFromSkeletalMeshSection(const TSharedRef<FManagedArrayCollection>& ClothCollection, const FSkeletalMeshLODModel& SkeletalMeshModel, const int32 SectionIndex, const int32 UVChannelIndex, const FVector2f& UVScale, bool bImportNormals = false);
+		static void AddSimPatternsFromSkeletalMeshSection(const TSharedRef<FManagedArrayCollection>& ClothCollection, const FSkeletalMeshLODModel& SkeletalMeshModel, const int32 SectionIndex, const int32 UVChannelIndex, const FVector2f& UVScale, bool bImportNormals = false, TArray<int32>* OutSim2DToSourceVertex = nullptr);
 
 		static void LogAndToastWarning(const FDataflowNode& DataflowNode, const FText& Headline, const FText& Details);
 
diff --git a/Engine/Plugins/ChaosClothAssetEditor/Source/ChaosClothAssetDataflowNodes/Private/ChaosClothAsset/SimulationMorphTargetConfigNode.cpp b/Engine/Plugins/ChaosClothAssetEditor/Source/ChaosClothAssetDataflowNodes/Private/ChaosClothAsset/SimulationMorphTargetConfigNode.cpp
new file mode 100644
--- /dev/null
+++ b/Engine/Plugins/ChaosClothAssetEditor/Source/ChaosClothAssetDataflowNodes/Private/ChaosClothAsset/SimulationMorphTargetConfigNode.cpp
@@ -0,0 +1,18 @@
+// Copyright Epic Games, Inc. All Rights Reserved.
+
+#include "ChaosClothAsset/SimulationMorphTargetConfigNode.h"
+#include "Chaos/CollectionPropertyFacade.h"
+
+#include UE_INLINE_GENERATED_CPP_BY_NAME(SimulationMorphTargetConfigNode)
+
+FChaosClothAssetSimulationMorphTargetConfigNode::FChaosClothAssetSimulationMorphTargetConfigNode(const UE::Dataflow::FNodeParameters& InParam, FGuid InGuid)
+	: FChaosClothAssetSimulationBaseConfigNode(InParam, InGuid)
+{
+	RegisterCollectionConnections();
+}
+
+void FChaosClothAssetSimulationMorphTargetConfigNode::AddProperties(FPropertyHelper& PropertyHelper) const
+{
+	PropertyHelper.SetPropertyString(FName(TEXT("ActiveMorphTarget")), ActiveMorphTarget.Name);
+	PropertyHelper.SetProperty(FName(TEXT("ActiveMorphTarget")), ActiveMorphTarget.Weight);
+}
diff --git a/Engine/Plugins/ChaosClothAssetEditor/Source/ChaosClothAssetDataflowNodes/Private/ChaosClothAsset/SkeletalMeshImportNode.cpp b/Engine/Plugins/ChaosClothAssetEditor/Source/ChaosClothAssetDataflowNodes/Private/ChaosClothAsset/SkeletalMeshImportNode.cpp
--- a/Engine/Plugins/ChaosClothAssetEditor/Source/ChaosClothAssetDataflowNodes/Private/ChaosClothAsset/SkeletalMeshImportNode.cpp
+++ b/Engine/Plugins/ChaosClothAssetEditor/Source/ChaosClothAssetDataflowNodes/Private/ChaosClothAsset/SkeletalMeshImportNode.cpp
@@ -4,6 +4,7 @@
 #include "ChaosClothAsset/ClothDataflowTools.h"
 #include "ChaosClothAsset/ClothGeometryTools.h"
 #include "ChaosClothAsset/CollectionClothFacade.h"
+#include "Animation/MorphTarget.h"
 #include "Animation/Skeleton.h"
 #include "Dataflow/DataflowInputOutput.h"
 #include "Engine/SkeletalMesh.h"
@@ -17,6 +18,24 @@
 
 #define LOCTEXT_NAMESPACE "ChaosClothAssetSkeletalMeshImportNode"
 
+namespace UE::Chaos::ClothAsset::Private
+{
+	TMap<int32, int32> GenerateSKMToSim3DLookup(const FCollectionClothConstFacade& ClothFacade, const TArray<int32>& Sim2DToSKMIndex)
+	{
+		TConstArrayView<int32 > SimVertex3DLookup = ClothFacade.GetSimVertex3DLookup();
+		check(Sim2DToSKMIndex.Num() == SimVertex3DLookup.Num());
+	
+		TMap<int32, int32> SKMToSim3DLookup;
+		for (int32 Sim2DIndex = 0; Sim2DIndex < Sim2DToSKMIndex.Num(); ++Sim2DIndex)
+		{
+			const int32 Sim3DIndex = SimVertex3DLookup[Sim2DIndex];
+			const int32 SKMIndex = Sim2DToSKMIndex[Sim2DIndex];
+			SKMToSim3DLookup.Add(SKMIndex, Sim3DIndex);
+		}
+		return SKMToSim3DLookup;
+	}
+}
+
 FChaosClothAssetSkeletalMeshImportNode_v2::FChaosClothAssetSkeletalMeshImportNode_v2(const UE::Dataflow::FNodeParameters& InParam, FGuid InGuid)
 	: FDataflowNode(InParam, InGuid)
 {
@@ -55,6 +74,8 @@ void FChaosClothAssetSkeletalMeshImportNode_v2::Evaluate(UE::Dataflow::FContext&
 			const int32 LastSection = bImportSingleSection ? SectionIndex : LODModel.Sections.Num() - 1;
 			constexpr bool bImportSimMeshNormals = true;
 
+			TArray<int32> Sim2DToSKMIndex;
+
 			for (int32 Section = FirstSection; Section <= LastSection; ++Section)
 			{
 				const bool bIsValidSection = LODModel.Sections.IsValidIndex(Section);;
@@ -71,7 +92,7 @@ void FChaosClothAssetSkeletalMeshImportNode_v2::Evaluate(UE::Dataflow::FContext&
 
 				if (bImportSimMesh)
 				{
-					FClothDataflowTools::AddSimPatternsFromSkeletalMeshSection(ClothCollection, LODModel, Section, UVChannel, UVScale, bImportSimMeshNormals);
+					FClothDataflowTools::AddSimPatternsFromSkeletalMeshSection(ClothCollection, LODModel, Section, UVChannel, UVScale, bImportSimMeshNormals, &Sim2DToSKMIndex);
 				}
 
 				if (bImportRenderMesh)
@@ -83,6 +104,68 @@ void FChaosClothAssetSkeletalMeshImportNode_v2::Evaluate(UE::Dataflow::FContext&
 				}
 			}
 
+			if (bImportSimMorphTargets)
+			{
+				// Use Sim2DToSKMIndex to get a map from SKMIndex to Sim3D index
+				const TMap<int32, int32> SKMToSim3DLookup = Private::GenerateSKMToSim3DLookup(ClothFacade, Sim2DToSKMIndex);
+
+				const TArray<TObjectPtr<UMorphTarget>>& MorphTargets = InSkeletalMesh->GetMorphTargets();
+				for (const UMorphTarget* const MorphTarget : MorphTargets)
+				{
+					if (MorphTarget && MorphTarget->HasDataForLOD(LODIndex))
+					{
+						const FMorphTargetLODModel& MorphTargetLODModel = MorphTarget->GetMorphLODModels()[LODIndex];
+						
+						TArray<int32> MissingIndices;
+						TArray<FVector3f> PositionDelta;
+						TArray<FVector3f> NormalZDelta;
+						TArray<int32> SimIndex;
+						PositionDelta.Reserve(MorphTargetLODModel.NumVertices);
+						NormalZDelta.Reserve(MorphTargetLODModel.NumVertices);
+						SimIndex.Reserve(MorphTargetLODModel.NumVertices);
+						for (const FMorphTargetDelta& Delta : MorphTargetLODModel.Vertices)
+						{
+							if (const int32* const SimIndex3D = SKMToSim3DLookup.Find(Delta.SourceIdx))
+							{
+								PositionDelta.Add(Delta.PositionDelta);
+								NormalZDelta.Add(Delta.TangentZDelta);
+								SimIndex.Add(*SimIndex3D);
+							}
+							else
+							{
+								MissingIndices.Add(Delta.SourceIdx);
+							}
+						}
+
+						const int32 NumTargetVertices = PositionDelta.Num();
+						check(NormalZDelta.Num() == NumTargetVertices);
+						check(SimIndex.Num() == NumTargetVertices);
+						if (NumTargetVertices > 0)
+						{
+							FCollectionClothSimMorphTargetFacade MorphTargetFacade = ClothFacade.AddGetSimMorphTarget();
+							MorphTargetFacade.Initialize(MorphTarget->GetName(), PositionDelta, NormalZDelta, SimIndex);
+							if (!MissingIndices.IsEmpty())
+							{
+								FString MissingIndicesString;
+								for (int32 Index = 0; Index < MissingIndices.Num(); ++Index)
+								{
+									MissingIndicesString += FString::FromInt(MissingIndices[Index]);
+									if (Index != MissingIndices.Num() - 1)
+									{
+										MissingIndicesString += ", ";
+									}
+								}
+								Context.Warning(FString::Printf(TEXT("Failed to find corresponding sim vertex index for the following indices for morph target '%s': %s"), *MorphTarget->GetName(), *MissingIndicesString), this, Out);
+							}
+						}
+						else
+						{
+							Context.Warning(FString::Printf(TEXT("Failed to import morph target '%s': no valid vertices found"), *MorphTarget->GetName()), this, Out);
+						}
+					}
+				}
+			}
+
 			FClothGeometryTools::CleanupAndCompactMesh(ClothCollection);
 
 			if (const UPhysicsAsset* PhysicsAsset = bSetPhysicsAsset ? InSkeletalMesh->GetPhysicsAsset() : nullptr)
diff --git a/Engine/Plugins/ChaosClothAssetEditor/Source/ChaosClothAssetDataflowNodes/Public/ChaosClothAsset/SimulationMorphTargetConfigNode.h b/Engine/Plugins/ChaosClothAssetEditor/Source/ChaosClothAssetDataflowNodes/Public/ChaosClothAsset/SimulationMorphTargetConfigNode.h
new file mode 100644
--- /dev/null
+++ b/Engine/Plugins/ChaosClothAssetEditor/Source/ChaosClothAssetDataflowNodes/Public/ChaosClothAsset/SimulationMorphTargetConfigNode.h
@@ -0,0 +1,41 @@
+// Copyright Epic Games, Inc. All Rights Reserved.
+
+#pragma once 
+
+#include "ChaosClothAsset/SimulationBaseConfigNode.h"
+#include "SimulationMorphTargetConfigNode.generated.h"
+
+/** The active sim morph target.
+*/
+USTRUCT()
+struct FChaosClothAssetMorphTargetSelection
+{
+	GENERATED_BODY()
+
+	/** Morph target name */
+	UPROPERTY(EditAnywhere, Category ="Simulation Morph Target Properties", meta = (InteractorName = "ActiveMorphTarget"))
+	FString Name;
+	/** Morph target weight */
+	UPROPERTY(EditAnywhere, Category = "Simulation Morph Target Properties", meta = (InteractorName = "ActiveMorphTarget"))
+	float Weight = 0.f;
+};
+
+/** Simulation Morph Target configuration node. This node is necessary to set sim morph targets (e.g., via BP interactor)*/
+USTRUCT(Meta = (DataflowCloth))
+struct FChaosClothAssetSimulationMorphTargetConfigNode : public FChaosClothAssetSimulationBaseConfigNode
+{
+	GENERATED_USTRUCT_BODY()
+	DATAFLOW_NODE_DEFINE_INTERNAL(FChaosClothAssetSimulationMorphTargetConfigNode, "SimulationMorphTargetConfig", "Cloth", "Cloth Simulation Morph Target Config")
+
+public:
+	/**
+	 * The name of the active sim morph target
+	 */
+	UPROPERTY(EditAnywhere, Category = "Simulation Morph Target Properties", meta = (InteractorName = "ActiveMorphTarget"))
+	FChaosClothAssetMorphTargetSelection ActiveMorphTarget;
+
+	FChaosClothAssetSimulationMorphTargetConfigNode(const UE::Dataflow::FNodeParameters& InParam, FGuid InGuid = FGuid::NewGuid());
+
+private:
+	virtual void AddProperties(FPropertyHelper& PropertyHelper) const override;
+};
diff --git a/Engine/Plugins/ChaosClothAssetEditor/Source/ChaosClothAssetDataflowNodes/Public/ChaosClothAsset/SkeletalMeshImportNode.h b/Engine/Plugins/ChaosClothAssetEditor/Source/ChaosClothAssetDataflowNodes/Public/ChaosClothAsset/SkeletalMeshImportNode.h
--- a/Engine/Plugins/ChaosClothAssetEditor/Source/ChaosClothAssetDataflowNodes/Public/ChaosClothAsset/SkeletalMeshImportNode.h
+++ b/Engine/Plugins/ChaosClothAssetEditor/Source/ChaosClothAssetDataflowNodes/Public/ChaosClothAsset/SkeletalMeshImportNode.h
@@ -69,6 +69,10 @@ public:
 	UPROPERTY(EditAnywhere, Category = "Skeletal Mesh Import")
 	bool bSetPhysicsAsset = false;
 
+	/** Import morph targets as Sim Mesh Morph Targets */
+	UPROPERTY(EditAnywhere, Category = "Skeletal Mesh Import")
+	bool bImportSimMorphTargets = false;
+
 	FChaosClothAssetSkeletalMeshImportNode_v2(const UE::Dataflow::FNodeParameters& InParam, FGuid InGuid = FGuid::NewGuid());
 
 private:
diff --git a/Engine/Source/Runtime/Experimental/Chaos/Public/Chaos/CollectionPropertyFacade.h b/Engine/Source/Runtime/Experimental/Chaos/Public/Chaos/CollectionPropertyFacade.h
--- a/Engine/Source/Runtime/Experimental/Chaos/Public/Chaos/CollectionPropertyFacade.h
+++ b/Engine/Source/Runtime/Experimental/Chaos/Public/Chaos/CollectionPropertyFacade.h
@@ -551,76 +551,76 @@ namespace Chaos::Softs
 // Use this macro to add shorthands for property getters and direct access through the declared key index
 #define UE_CHAOS_DECLARE_PROPERTYCOLLECTION_NAME(PropertyName, Type) \
 	UE_CHAOS_DECLARE_INDEXLESS_PROPERTYCOLLECTION_NAME(PropertyName, Type) \
-	Type GetLow##PropertyName(const FCollectionPropertyConstFacade& PropertyCollection) const \
+	Type GetLow##PropertyName(const FCollectionPropertyConstFacade& InPropertyCollection) const \
 	{ \
-		checkSlow(PropertyName##Index == PropertyCollection.GetKeyIndex(PropertyName##Name.ToString())); \
+		checkSlow(PropertyName##Index == InPropertyCollection.GetKeyIndex(PropertyName##Name.ToString())); \
 		checkf(PropertyName##Index != INDEX_NONE, TEXT("The default value getter that sets the property index must be called once prior to calling this function.")); \
-		return PropertyCollection.GetLowValue<Type>(PropertyName##Index); \
+		return InPropertyCollection.GetLowValue<Type>(PropertyName##Index); \
 	} \
-	Type GetHigh##PropertyName(const FCollectionPropertyConstFacade& PropertyCollection) const \
+	Type GetHigh##PropertyName(const FCollectionPropertyConstFacade& InPropertyCollection) const \
 	{ \
-		checkSlow(PropertyName##Index == PropertyCollection.GetKeyIndex(PropertyName##Name.ToString())); \
+		checkSlow(PropertyName##Index == InPropertyCollection.GetKeyIndex(PropertyName##Name.ToString())); \
 		checkf(PropertyName##Index != INDEX_NONE, TEXT("The default value getter that sets the property index must be called once prior to calling this function.")); \
-		return PropertyCollection.GetHighValue<Type>(PropertyName##Index); \
+		return InPropertyCollection.GetHighValue<Type>(PropertyName##Index); \
 	} \
-	TPair<Type, Type> GetWeighted##PropertyName(const FCollectionPropertyConstFacade& PropertyCollection) const \
+	TPair<Type, Type> GetWeighted##PropertyName(const FCollectionPropertyConstFacade& InPropertyCollection) const \
 	{ \
-		checkSlow(PropertyName##Index == PropertyCollection.GetKeyIndex(PropertyName##Name.ToString())); \
+		checkSlow(PropertyName##Index == InPropertyCollection.GetKeyIndex(PropertyName##Name.ToString())); \
 		checkf(PropertyName##Index != INDEX_NONE, TEXT("The default value getter that sets the property index must be called once prior to calling this function.")); \
-		return PropertyCollection.GetWeightedValue<Type>(PropertyName##Index); \
+		return InPropertyCollection.GetWeightedValue<Type>(PropertyName##Index); \
 	} \
-	FVector2f GetWeightedFloat##PropertyName(const FCollectionPropertyConstFacade& PropertyCollection) const \
+	FVector2f GetWeightedFloat##PropertyName(const FCollectionPropertyConstFacade& InPropertyCollection) const \
 	{ \
-		checkSlow(PropertyName##Index == PropertyCollection.GetKeyIndex(PropertyName##Name.ToString())); \
+		checkSlow(PropertyName##Index == InPropertyCollection.GetKeyIndex(PropertyName##Name.ToString())); \
 		checkf(PropertyName##Index != INDEX_NONE, TEXT("The default value getter that sets the property index must be called once prior to calling this function.")); \
-		return PropertyCollection.GetWeightedFloatValue(PropertyName##Index); \
+		return InPropertyCollection.GetWeightedFloatValue(PropertyName##Index); \
 	} \
-	Type Get##PropertyName(const FCollectionPropertyConstFacade& PropertyCollection) const \
+	Type Get##PropertyName(const FCollectionPropertyConstFacade& InPropertyCollection) const \
 	{ \
-		checkSlow(PropertyName##Index == PropertyCollection.GetKeyIndex(PropertyName##Name.ToString())); \
+		checkSlow(PropertyName##Index == InPropertyCollection.GetKeyIndex(PropertyName##Name.ToString())); \
 		checkf(PropertyName##Index != INDEX_NONE, TEXT("The default value getter that sets the property index must be called once prior to calling this function.")); \
-		return PropertyCollection.GetValue<Type>(PropertyName##Index); \
+		return InPropertyCollection.GetValue<Type>(PropertyName##Index); \
 	} \
-	const FString& Get##PropertyName##String(const FCollectionPropertyConstFacade& PropertyCollection) \
+	const FString& Get##PropertyName##String(const FCollectionPropertyConstFacade& InPropertyCollection) \
 	{ \
-		checkSlow(PropertyName##Index == PropertyCollection.GetKeyIndex(PropertyName##Name.ToString())); \
+		checkSlow(PropertyName##Index == InPropertyCollection.GetKeyIndex(PropertyName##Name.ToString())); \
 		checkf(PropertyName##Index != INDEX_NONE, TEXT("The default value getter that sets the property index must be called once prior to calling this function.")); \
-		return PropertyCollection.GetStringValue(PropertyName##Index); \
+		return InPropertyCollection.GetStringValue(PropertyName##Index); \
 	} \
-	bool Is##PropertyName##Enabled(const FCollectionPropertyConstFacade& PropertyCollection) const \
+	bool Is##PropertyName##Enabled(const FCollectionPropertyConstFacade& InPropertyCollection) const \
 	{ \
-		checkSlow(PropertyName##Index == PropertyCollection.GetKeyIndex(PropertyName##Name.ToString())); \
+		checkSlow(PropertyName##Index == InPropertyCollection.GetKeyIndex(PropertyName##Name.ToString())); \
 		checkf(PropertyName##Index != INDEX_NONE, TEXT("The default value getter that sets the property index must be called once prior to calling this function.")); \
-		return PropertyCollection.IsEnabled(PropertyName##Index); \
+		return InPropertyCollection.IsEnabled(PropertyName##Index); \
 	} \
-	bool Is##PropertyName##Animatable(const FCollectionPropertyConstFacade& PropertyCollection) const\
+	bool Is##PropertyName##Animatable(const FCollectionPropertyConstFacade& InPropertyCollection) const\
 	{ \
-		checkSlow(PropertyName##Index == PropertyCollection.GetKeyIndex(PropertyName##Name.ToString())); \
+		checkSlow(PropertyName##Index == InPropertyCollection.GetKeyIndex(PropertyName##Name.ToString())); \
 		checkf(PropertyName##Index != INDEX_NONE, TEXT("The default value getter that sets the property index must be called once prior to calling this function.")); \
-		return PropertyCollection.IsAnimatable(PropertyName##Index); \
+		return InPropertyCollection.IsAnimatable(PropertyName##Index); \
 	} \
-	bool Is##PropertyName##Dirty(const FCollectionPropertyConstFacade& PropertyCollection) const \
+	bool Is##PropertyName##Dirty(const FCollectionPropertyConstFacade& InPropertyCollection) const \
 	{ \
-		checkSlow(PropertyName##Index == PropertyCollection.GetKeyIndex(PropertyName##Name.ToString())); \
+		checkSlow(PropertyName##Index == InPropertyCollection.GetKeyIndex(PropertyName##Name.ToString())); \
 		checkf(PropertyName##Index != INDEX_NONE, TEXT("The default value getter that sets the property index must be called once prior to calling this function.")); \
-		return PropertyCollection.IsDirty(PropertyName##Index); \
+		return InPropertyCollection.IsDirty(PropertyName##Index); \
 	} \
-	bool Is##PropertyName##StringDirty(const FCollectionPropertyConstFacade& PropertyCollection) const \
+	bool Is##PropertyName##StringDirty(const FCollectionPropertyConstFacade& InPropertyCollection) const \
 	{ \
-		checkSlow(PropertyName##Index == PropertyCollection.GetKeyIndex(PropertyName##Name.ToString())); \
+		checkSlow(PropertyName##Index == InPropertyCollection.GetKeyIndex(PropertyName##Name.ToString())); \
 		checkf(PropertyName##Index != INDEX_NONE, TEXT("The default value getter that sets the property index must be called once prior to calling this function.")); \
-		return PropertyCollection.IsStringDirty(PropertyName##Index); \
+		return InPropertyCollection.IsStringDirty(PropertyName##Index); \
 	} \
-	bool Is##PropertyName##Mutable(const FCollectionPropertyConstFacade& PropertyCollection) const \
+	bool Is##PropertyName##Mutable(const FCollectionPropertyConstFacade& InPropertyCollection) const \
 	{ \
 		return PropertyName##Index != INDEX_NONE && \
-			PropertyCollection.IsAnimatable(PropertyName##Index) && (PropertyCollection.IsDirty(PropertyName##Index)); \
+			InPropertyCollection.IsAnimatable(PropertyName##Index) && (InPropertyCollection.IsDirty(PropertyName##Index)); \
 	} \
 	struct F##PropertyName##Index \
 	{ \
 		int32 Index = INDEX_NONE; \
 		explicit F##PropertyName##Index(EForceInit) : Index(INDEX_NONE) {} \
-		explicit F##PropertyName##Index(const FCollectionPropertyConstFacade& PropertyCollection) : Index(PropertyCollection.GetKeyIndex(PropertyName##Name.ToString())) {} \
+		explicit F##PropertyName##Index(const FCollectionPropertyConstFacade& InPropertyCollection) : Index(InPropertyCollection.GetKeyIndex(PropertyName##Name.ToString())) {} \
 		operator int32() const { return Index; } \
 	} PropertyName##Index;
 
