Fix wasted bit in stringify options

#rnx

diff --git a/Engine/Source/Runtime/Experimental/JsonObjectGraph/Public/JsonObjectGraph/Stringify.h b/Engine/Source/Runtime/Experimental/JsonObjectGraph/Public/JsonObjectGraph/Stringify.h
--- a/Engine/Source/Runtime/Experimental/JsonObjectGraph/Public/JsonObjectGraph/Stringify.h
+++ b/Engine/Source/Runtime/Experimental/JsonObjectGraph/Public/JsonObjectGraph/Stringify.h
@@ -15,9 +15,9 @@ enum class EJsonStringifyFlags : uint8
 	/** The default representation of an object attempts to be complete and stable across time */
 	Default = 0,
 	/** Filters editor only data such that it is not written to the Json */
-	FilterEditorOnlyData = 1 << 1,
+	FilterEditorOnlyData = 1 << 0,
 	/** Disables delta encoding such that all properties are encoded in the Json, rather than only changes from the objects' archetypes */
-	DisableDeltaEncoding = 1 << 2,
+	DisableDeltaEncoding = 1 << 1,
 };
 ENUM_CLASS_FLAGS(EJsonStringifyFlags);
 
