Dataflow : fix issue with drag and drop-ing a reroute node into an input or output  node when the concrete type is a structure
- Changed the method used to find the object corresponding to the type

diff --git a/Engine/Source/Runtime/Experimental/Dataflow/Engine/Private/Dataflow/DataflowSubGraphNodes.cpp b/Engine/Source/Runtime/Experimental/Dataflow/Engine/Private/Dataflow/DataflowSubGraphNodes.cpp
--- a/Engine/Source/Runtime/Experimental/Dataflow/Engine/Private/Dataflow/DataflowSubGraphNodes.cpp
+++ b/Engine/Source/Runtime/Experimental/Dataflow/Engine/Private/Dataflow/DataflowSubGraphNodes.cpp
@@ -8,6 +8,7 @@
 #include "Dataflow/DataflowSubGraph.h"
 #include "Dataflow/DataflowNodeFactory.h"
 #include "GraphEditAction.h"
+#include "UObject/UObjectGlobals.h"
 
 namespace UE::Dataflow::SubGraphNodes::Private
 {
@@ -188,12 +189,10 @@ namespace UE::Dataflow::SubGraphNodes::Private
 				}
 				else if (TypeStr.StartsWith("F"))
 				{
-					static UPackage* CoreUObjectPkg = FindObjectChecked<UPackage>(nullptr, TEXT("/Script/CoreUObject"));
-					const FName ShortTypeName(TypeStr.RightChop(1));
-					UScriptStruct* Result = (UScriptStruct*)StaticFindObjectFastInternal(UScriptStruct::StaticClass(), CoreUObjectPkg, ShortTypeName, false, RF_NoFlags, EInternalObjectFlags::None);
-					if (Result)
+					const FString ShortTypeName = TypeStr.RightChop(1);
+					if (const UScriptStruct* ScriptStruct = FindFirstObjectSafe<UScriptStruct>(*ShortTypeName))
 					{
-						return FPropertyBagPropertyDesc(Connection.GetName(), EPropertyBagPropertyType::Struct, Result);
+						return FPropertyBagPropertyDesc(Connection.GetName(), EPropertyBagPropertyType::Struct, ScriptStruct);
 					}
 				}
 			}
