Dataflow : add a generic version fo math Clamp

diff --git a/Engine/Plugins/Experimental/GeometryCollectionPlugin/Source/GeometryCollectionNodes/Public/Dataflow/GeometryCollectionMathNodes.h b/Engine/Plugins/Experimental/GeometryCollectionPlugin/Source/GeometryCollectionNodes/Public/Dataflow/GeometryCollectionMathNodes.h
--- a/Engine/Plugins/Experimental/GeometryCollectionPlugin/Source/GeometryCollectionNodes/Public/Dataflow/GeometryCollectionMathNodes.h
+++ b/Engine/Plugins/Experimental/GeometryCollectionPlugin/Source/GeometryCollectionNodes/Public/Dataflow/GeometryCollectionMathNodes.h
@@ -692,9 +692,10 @@ public:
 /**
 *
 * Clamps X to be between Min and Max, inclusive
+* DEPRECATED 5.6: Use the Clamp node that support all numerics types
 *
 */
-USTRUCT(meta = (DataflowGeometryCollection))
+USTRUCT(meta = (DataflowGeometryCollection, Deprecated = "5.6"))
 struct FClampDataflowNode : public FDataflowNode
 {
 	GENERATED_USTRUCT_BODY()
diff --git a/Engine/Source/Runtime/Experimental/Dataflow/Core/Private/Dataflow/DataflowMathNodes.cpp b/Engine/Source/Runtime/Experimental/Dataflow/Core/Private/Dataflow/DataflowMathNodes.cpp
--- a/Engine/Source/Runtime/Experimental/Dataflow/Core/Private/Dataflow/DataflowMathNodes.cpp
+++ b/Engine/Source/Runtime/Experimental/Dataflow/Core/Private/Dataflow/DataflowMathNodes.cpp
@@ -15,6 +15,7 @@ namespace UE::Dataflow
 		DATAFLOW_NODE_REGISTER_CREATION_FACTORY(FDataflowMathAbsNode);
 		DATAFLOW_NODE_REGISTER_CREATION_FACTORY(FDataflowMathAddNode);
 		DATAFLOW_NODE_REGISTER_CREATION_FACTORY(FDataflowMathCeilNode);
+		DATAFLOW_NODE_REGISTER_CREATION_FACTORY(FDataflowMathClampNode);
 		DATAFLOW_NODE_REGISTER_CREATION_FACTORY(FDataflowMathConstantNode);
 		DATAFLOW_NODE_REGISTER_CREATION_FACTORY(FDataflowMathCubeNode);
 		DATAFLOW_NODE_REGISTER_CREATION_FACTORY(FDataflowMathDivideNode);
@@ -716,6 +717,23 @@ void FDataflowMathConstantNode::Evaluate(UE::Dataflow::FContext& Context, const
 	}
 }
 
+//-----------------------------------------------------------------------------------------------
+
+FDataflowMathClampNode::FDataflowMathClampNode(const UE::Dataflow::FNodeParameters& InParam, FGuid InGuid)
+	: FDataflowMathOneInputOperatorNode(InParam, InGuid)
+{
+	RegisterInputsAndOutputs();
+	Min.Value = 0.0;
+	Max.Value = 1.0;
+}
+
+double FDataflowMathClampNode::ComputeResult(UE::Dataflow::FContext& Context, double InA) const
+{
+	const double InMin = GetValue(Context, &Min);
+	const double InMax = GetValue(Context, &Max);
+	return FMath::Clamp(InA, InMin, InMax);
+}
+
 //--------------------------------------------------------------------------
 //
 // Trigonometric nodes
diff --git a/Engine/Source/Runtime/Experimental/Dataflow/Core/Public/Dataflow/DataflowMathNodes.h b/Engine/Source/Runtime/Experimental/Dataflow/Core/Public/Dataflow/DataflowMathNodes.h
--- a/Engine/Source/Runtime/Experimental/Dataflow/Core/Public/Dataflow/DataflowMathNodes.h
+++ b/Engine/Source/Runtime/Experimental/Dataflow/Core/Public/Dataflow/DataflowMathNodes.h
@@ -496,6 +496,24 @@ private:
 	double GetConstant() const;
 };
 
+/** Clamp(A, Min, Max) clamp a value to specific range (inclusive) */
+USTRUCT()
+struct FDataflowMathClampNode : public FDataflowMathOneInputOperatorNode
+{
+	GENERATED_USTRUCT_BODY()
+	DATAFLOW_NODE_DEFINE_INTERNAL(FDataflowMathClampNode, "Clamp", DATAFLOW_MATH_NODES_CATEGORY, "")
+
+	UPROPERTY(meta = (DataflowInput))
+	FDataflowNumericTypes Min;
+
+	UPROPERTY(meta = (DataflowInput))
+	FDataflowNumericTypes Max;
+
+public:
+	FDataflowMathClampNode(const UE::Dataflow::FNodeParameters& InParam, FGuid InGuid = FGuid::NewGuid());
+	virtual double ComputeResult(UE::Dataflow::FContext& Context, double InA) const override;
+};
+
 //--------------------------------------------------------------------------
 //
 // Trigonometric nodes
@@ -615,41 +633,23 @@ public:
 
 //--------------------------------------------------------------------------
 //
-// Trigonometric nodes
+// Other
 //
 //--------------------------------------------------------------------------
-
-
 /*
-Nodes Left to be converted :
+Nodes Left to be converted from the Geometry collection math nodes:
 	FFloatMathExpressionDataflowNode);
 	FMathExpressionDataflowNode);
-	
-	FClampDataflowNode);
 	FFitDataflowNode);
 	FEFitDataflowNode);
-	
 	FLerpDataflowNode);
 	FWrapDataflowNode);
 
-
-
-	// vectors - requires Vector any type ? 
-	FNormalizeToRangeDataflowNode);
-	FScaleVectorDataflowNode);
-	FDotProductDataflowNode);
-	FCrossProductDataflowNode);
-	FNormalizeDataflowNode);
-	FLengthDataflowNode);
-	FDistanceDataflowNode);
-	FIsNearlyZeroDataflowNode);
-
 	// random 
 	FRandomFloatDataflowNode);
 	FRandomFloatInRangeDataflowNode);
 	FRandomUnitVectorDataflowNode);
 	FRandomUnitVectorInConeDataflowNode);
-
 */
 
 namespace UE::Dataflow
