Dataflow : fix to properly notify of a dataflow asset change to resync variables
- Also fix the geometry collection not forwarding the property change to the base class

diff --git a/Engine/Source/Runtime/Experimental/Dataflow/Engine/Private/Dataflow/DataflowInstance.cpp b/Engine/Source/Runtime/Experimental/Dataflow/Engine/Private/Dataflow/DataflowInstance.cpp
--- a/Engine/Source/Runtime/Experimental/Dataflow/Engine/Private/Dataflow/DataflowInstance.cpp
+++ b/Engine/Source/Runtime/Experimental/Dataflow/Engine/Private/Dataflow/DataflowInstance.cpp
@@ -473,22 +473,35 @@ FDataflowInstance::FDataflowInstance(UObject* InOwner, UDataflow* InDataflowAsse
 					VariableOverrides.OnDataflowVariablesChanged(InDataflowAsset, InVariableName);
 				}
 			});
+		OnOwnerPropertyChangedHandle = FCoreUObjectDelegates::OnObjectPropertyChanged.AddRaw(this, &FDataflowInstance::OnOwnerPostEditChangeProperty);
 	}
 #endif
 }
 
 FDataflowInstance::~FDataflowInstance()
 {
+#if WITH_EDITOR
 	if (Owner)
 	{
 		FDataflowAssetDelegates::OnVariablesChanged.RemoveAll(Owner.Get());
+		FCoreUObjectDelegates::OnObjectPropertyChanged.Remove(OnOwnerPropertyChangedHandle);
 	}
+#endif
 }
 
 #if WITH_EDITOR
-void FDataflowInstance::OnOwnerPostEditChangeProperty(struct FPropertyChangedEvent& PropertyChangedEvent)
+void FDataflowInstance::OnOwnerPostEditChangeProperty(UObject* InObject, struct FPropertyChangedEvent& InPropertyChangedEvent)
 {
-	VariableOverrides.OnOwnerPostEditChangeProperty(PropertyChangedEvent);
+	if (InObject == Owner)
+	{
+		VariableOverrides.OnOwnerPostEditChangeProperty(InPropertyChangedEvent);
+
+		// if the Dataflow asset has changed we need to resync the variables
+		if (InPropertyChangedEvent.GetPropertyName() == GetDataflowAssetPropertyName())
+		{
+			SyncVariables();
+		}
+	}
 	// should check for terminal node / dataflow asset changes ?
 }
 
diff --git a/Engine/Source/Runtime/Experimental/Dataflow/Engine/Public/Dataflow/DataflowInstance.h b/Engine/Source/Runtime/Experimental/Dataflow/Engine/Public/Dataflow/DataflowInstance.h
--- a/Engine/Source/Runtime/Experimental/Dataflow/Engine/Public/Dataflow/DataflowInstance.h
+++ b/Engine/Source/Runtime/Experimental/Dataflow/Engine/Public/Dataflow/DataflowInstance.h
@@ -247,7 +247,9 @@ private:
 	friend class UE::Dataflow::Private::FVariablesOverridesDetails;
 
 #if WITH_EDITOR
-	void OnOwnerPostEditChangeProperty(struct FPropertyChangedEvent& PropertyChangedEvent);
+	void OnOwnerPostEditChangeProperty(UObject* InObject, struct FPropertyChangedEvent& InPropertyChangedEvent);
+
+	FDelegateHandle OnOwnerPropertyChangedHandle;
 #endif
 
 	// Owning Object
diff --git a/Engine/Source/Runtime/Experimental/GeometryCollectionEngine/Private/GeometryCollection/GeometryCollectionObject.cpp b/Engine/Source/Runtime/Experimental/GeometryCollectionEngine/Private/GeometryCollection/GeometryCollectionObject.cpp
--- a/Engine/Source/Runtime/Experimental/GeometryCollectionEngine/Private/GeometryCollection/GeometryCollectionObject.cpp
+++ b/Engine/Source/Runtime/Experimental/GeometryCollectionEngine/Private/GeometryCollection/GeometryCollectionObject.cpp
@@ -1920,6 +1920,8 @@ void UGeometryCollection::PostEditChangeProperty(struct FPropertyChangedEvent& P
 		RebuildRenderData();
 	}
 	InvalidateDataflowContents();
+
+	Super::PostEditChangeProperty(PropertyChangedEvent);
 }
 
 bool UGeometryCollection::Modify(bool bAlwaysMarkDirty /*= true*/)
