Surrounded usages of chaos debug draw with the correct define to fix compile errors, such as when debug draw is enabled in Test builds.

#rnx

diff --git a/Engine/Plugins/Experimental/Water/Source/Runtime/Public/BuoyancyComponentSimulation.h b/Engine/Plugins/Experimental/Water/Source/Runtime/Public/BuoyancyComponentSimulation.h
--- a/Engine/Plugins/Experimental/Water/Source/Runtime/Public/BuoyancyComponentSimulation.h
+++ b/Engine/Plugins/Experimental/Water/Source/Runtime/Public/BuoyancyComponentSimulation.h
@@ -431,7 +431,7 @@ public:
 			}
 			PontoonIndex++;
 		}
-#if ENABLE_DRAW_DEBUG
+#if CHAOS_DEBUG_DRAW
 		if (CVarWaterDebugBuoyancy.GetValueOnAnyThread())
 		{
 			const float NumPoints = CVarWaterBuoyancyDebugPoints.GetValueOnAnyThread();
diff --git a/Engine/Source/Runtime/Experimental/GeometryCollectionEngine/Private/GeometryCollection/GeometryCollectionComponent.cpp b/Engine/Source/Runtime/Experimental/GeometryCollectionEngine/Private/GeometryCollection/GeometryCollectionComponent.cpp
--- a/Engine/Source/Runtime/Experimental/GeometryCollectionEngine/Private/GeometryCollection/GeometryCollectionComponent.cpp
+++ b/Engine/Source/Runtime/Experimental/GeometryCollectionEngine/Private/GeometryCollection/GeometryCollectionComponent.cpp
@@ -2958,7 +2958,7 @@ bool UGeometryCollectionComponent::ProcessRepData(const float DeltaTime, const f
 	bool ReturnValue = ::ProcessRepDataCommon(DeltaTime, SimTime, RepData, PhysicsProxy, VersionProcessed, OneOffActivatedProcessed, LastHardsnapTimeInMs, bReplicateMovement);
 
 
-#if ENABLE_DRAW_DEBUG
+#if CHAOS_DEBUG_DRAW
 	
 	if (bGeometryCollectionDebugDrawRep)
 	{
