On Demand Install Cache
- Reserve some of disk quota for defrag overhead
- At startup, detect if cache is overbudget and purge blocks until the cache is in budget
- Move block info sort by last access time into FOnDemandInstallCache::Purge
- FOnDemandInstallCache::Purge now prunes the blocks it deletes from the block info map - this fixes an issue with defrag after purge possibly consider blocks that no longer exist
- Defrag now fails if there is not enough disk budget left to support defrag block overhead.

diff --git a/Engine/Source/Runtime/Experimental/IoStore/OnDemand/Private/OnDemandInstallCache.cpp b/Engine/Source/Runtime/Experimental/IoStore/OnDemand/Private/OnDemandInstallCache.cpp
--- a/Engine/Source/Runtime/Experimental/IoStore/OnDemand/Private/OnDemandInstallCache.cpp
+++ b/Engine/Source/Runtime/Experimental/IoStore/OnDemand/Private/OnDemandInstallCache.cpp
@@ -167,7 +167,7 @@ struct FCas
 	FSharedFileOpenResult		OpenRead(FCasBlockId BlockId);
 	FSharedFileOpenAsyncResult	OpenAsyncRead(FCasBlockId BlockId);
 	void						OnFileHandleDeleted(FCasBlockId BlockId);
-	FUniqueFileHandle			OpenWrite(FCasBlockId BlockId);
+	FUniqueFileHandle			OpenWrite(FCasBlockId BlockId) const;
 	void						TrackAccess(FCasBlockId BlockId, int64 UtcTicks);
 	void						TrackAccess(FCasBlockId BlockId) { TrackAccess(BlockId, FDateTime::UtcNow().GetTicks()); }
 	void						TrackAccessIfNewer(FCasBlockId BlockId, int64 UtcTicks);
@@ -396,7 +396,7 @@ void FCas::OnFileHandleDeleted(FCasBlockId BlockId)
 	}
 }
 
-FUniqueFileHandle FCas::OpenWrite(FCasBlockId BlockId)
+FUniqueFileHandle FCas::OpenWrite(FCasBlockId BlockId) const
 {
 	IPlatformFile&	Ipf = FPlatformFileManager::Get().GetPlatformFile();
 	const FString	Filename = GetBlockFilename(BlockId);
@@ -1337,7 +1337,7 @@ private:
 													const TArray<TBitArray<>>& ChunkEntryIndices,
 													const TSet<FIoHash>& ChunksToInstall,
 													FCasBlockInfoMap& BlockInfoMap) const;
-	FIoStatus									Purge(const FCasBlockInfoMap& BlockInfo, uint64 TotalBytesToPurge, uint64& OutTotalPurgedBytes);
+	FIoStatus									Purge(FCasBlockInfoMap& BlockInfo, uint64 TotalBytesToPurge, uint64& OutTotalPurgedBytes);
 	FIoStatus									Defrag(
 													const TArray<FSharedOnDemandContainer>& Containers,
 													const TArray<TBitArray<>>& ChunkEntryIndices,
@@ -1373,12 +1373,20 @@ FOnDemandInstallCache::FOnDemandInstallCache(const FOnDemandInstallCacheConfig&
 	, MaxCacheSize(Config.DiskQuota)
 	, MaxJournalSize(Config.JournalMaxSize)
 {
-
-	// TODO: Reserve one block of the configured cache size be resvered for defrag overhead
-
 	UE_LOG(LogIoStoreOnDemand, Log, TEXT("Initializing install cache, MaxCacheSize=%.2lf MiB, MaxJournalSize=%.2lf KiB"),
 		ToMiB(MaxCacheSize), ToKiB(MaxJournalSize));
 
+	const uint64 MinDiskQuota = 2 * Cas.MaxBlockSize;
+	if (MaxCacheSize < MinDiskQuota)
+	{
+		UE_LOG(LogIoStoreOnDemand, Error, TEXT("Failed to initialize install cache - disk quota must be at least %.2lf MiB"), ToMiB(MinDiskQuota));
+		return;
+	}
+
+	// Reserve one block of space for defragmentation overhead
+	MaxCacheSize -= Cas.MaxBlockSize;
+	UE_LOG(LogIoStoreOnDemand, Log, TEXT("Effective MaxCacheSize without defragmentation space is MaxCacheSize=%.2lf MiB"), ToMiB(MaxCacheSize));
+
 	FIoStatus Status = Cas.Initialize(CacheDirectory);
 	if (Status.IsOk() == false)
 	{
@@ -1468,6 +1476,35 @@ FOnDemandInstallCache::FOnDemandInstallCache(const FOnDemandInstallCacheConfig&
 		Status = FCasJournal::Commit(MoveTemp(Transaction));
 	}
 
+	// Check that Cas is in-budget and shrink if needed
+	if (Status.IsOk())
+	{
+		FCasBlockInfoMap	BlockInfo;
+		uint64				TotalSize = Cas.GetBlockInfo(BlockInfo);
+		if (TotalSize > MaxCacheSize)
+		{
+			const uint64 TotalBytesToPurge = TotalSize - MaxCacheSize;
+
+			UE_LOG(LogIoStoreOnDemand, Warning, TEXT("Cache size is greater than disk quota - Purging install cache, MaxCacheSize=%.2lf MiB, TotalSize=%.2lf MiB, TotalBytesToPurge=%.2lf MiB"),
+				ToMiB(MaxCacheSize), ToMiB(TotalSize), ToMiB(TotalBytesToPurge));
+			
+			uint64 TotalPurgedBytes = 0;
+			FIoStatus PurgeStatus = Purge(BlockInfo, TotalBytesToPurge, TotalPurgedBytes);
+
+			if (PurgeStatus.IsOk() && TotalPurgedBytes >= TotalBytesToPurge)
+			{
+				UE_LOG(LogIoStoreOnDemand, Warning, TEXT("Successfully purged %.2lf MiB from install cache"), ToMiB(TotalPurgedBytes));
+			}
+			else
+			{
+				UE_LOG(LogIoStoreOnDemand, Error, TEXT("Failed to purge %.2lf MiB from install cache. Actually purged %.2lf MiB from install cache"),
+					ToMiB(TotalBytesToPurge), ToMiB(TotalPurgedBytes));
+
+				Status = FIoStatusBuilder(EIoErrorCode::WriteError) << FString::Printf(TEXT("Failed to purge overbudget cache (%s)"), *PurgeStatus.ToString());
+			}
+		}
+	}
+
 	Cas.Compact();
 
 	if (Status.IsOk() == false)
@@ -1849,11 +1886,6 @@ FIoStatus FOnDemandInstallCache::Purge(TSet<FIoHash>&& ChunksToInstall)
 
 	//TODO: Compute fragmentation metric and redownload chunks when this number gets too high
 
-	BlockInfo.ValueSort([](const FCasBlockInfo& LHS, const FCasBlockInfo& RHS)
-	{
-		return LHS.LastAccess < RHS.LastAccess;
-	});
-
 	const uint64 TotalReferencedBytes = Algo::TransformAccumulate(BlockInfo,
 		[](const TPair<FCasBlockId, FCasBlockInfo>& Kv){ return (Kv.Value.RefSize > 0) ? Kv.Value.FileSize : uint64(0); },
 		uint64(0));
@@ -2030,14 +2062,19 @@ uint64 FOnDemandInstallCache::AddReferencesToBlocks(
 	return TotalUncachedBytes;
 }
 
-FIoStatus FOnDemandInstallCache::Purge(const FCasBlockInfoMap& BlockInfo, const uint64 TotalBytesToPurge, uint64& OutTotalPurgedBytes)
+FIoStatus FOnDemandInstallCache::Purge(FCasBlockInfoMap& BlockInfo, const uint64 TotalBytesToPurge, uint64& OutTotalPurgedBytes)
 {
+	BlockInfo.ValueSort([](const FCasBlockInfo& LHS, const FCasBlockInfo& RHS)
+	{
+		return LHS.LastAccess < RHS.LastAccess;
+	});
+
 	OutTotalPurgedBytes = 0;
 
-	for (const TPair<FCasBlockId, FCasBlockInfo>& Kv : BlockInfo)
+	for (auto It = BlockInfo.CreateIterator(); It; ++It)
 	{
-		const FCasBlockId BlockId = Kv.Key;
-		const FCasBlockInfo& Info = Kv.Value;
+		const FCasBlockId BlockId = It->Key;
+		const FCasBlockInfo& Info = It->Value;
 		if (Info.RefSize > 0)
 		{
 			continue;
@@ -2057,6 +2094,8 @@ FIoStatus FOnDemandInstallCache::Purge(const FCasBlockInfoMap& BlockInfo, const
 
 		OutTotalPurgedBytes += Info.FileSize;
 
+		It.RemoveCurrent();
+
 		for (const FCasAddr& Addr : RemovedChunks)
 		{
 			Transaction.ChunkLocation(FCasLocation::Invalid, Addr);
@@ -2083,6 +2122,19 @@ FIoStatus FOnDemandInstallCache::Defrag(
 	FCasBlockInfoMap& BlockInfo, 
 	const uint64* TotalBytesToFree /*= nullptr*/)
 {
+	const uint64 TotalCachedBytes = Algo::TransformAccumulate(BlockInfo,
+		[](const TPair<FCasBlockId, FCasBlockInfo>& Kv) { return Kv.Value.FileSize; },
+		uint64(0));
+
+	if (TotalCachedBytes > MaxCacheSize)
+	{
+		// Ruh-Roh! There's not enough of the disk quota left to run a defrag!
+		const FString ErrorMsg = FString::Printf(TEXT("Cache size is greater than disk quota - Cannot Defragment!, MaxCacheSize=%.2lf MiB, TotalCachedBytes=%.2lf MiB"), 
+			ToMiB(MaxCacheSize), ToMiB(TotalCachedBytes));
+		UE_LOG(LogIoStoreOnDemand, Error, TEXT("%s"), *ErrorMsg);
+		return FIoStatusBuilder(EIoErrorCode::WriteError) << ErrorMsg;
+	}
+
 	struct FDefragBlockReferencedChunk
 	{
 		uint32 BlockOffset = 0;
@@ -2256,16 +2308,6 @@ FIoStatus FOnDemandInstallCache::Defrag(
 			check(DefragPendingChunks.IsEmpty());
 		}
 
-		// TODO: At startup, as part of the Verify, the cache needs to be check if there is enough room for the defrag chunk, and
-		// purge down to the correct size if not. Its possible something went wrong before the old block could be deleted.
-		// Alternatively we could create a transaction for every defragged block processed but its possible that would include
-		// multiple flushes. Is that OK?
-
-		// TODO: if something goes wrong before we get the chunk deletion here, we can be left in an over-allocated state.
-		// If we just keep trying to defrag in that state, we may just continue making new blocks without cleaning up old blocks.
-		// Its probably best that all subsequent defrag attempts fail if we ever leak a block here. A restart would be required to heal 
-		// the cache in that case.
-
 		FCasJournal::FTransaction Transaction = FCasJournal::Begin(GetJournalFilename());
 
 		// Flushing should overwrite the lookup info for the cas addr to point at the new block.
