Enable iostore.EnableDefrag by default

diff --git a/Engine/Source/Runtime/Experimental/IoStore/OnDemand/Private/OnDemandInstallCache.cpp b/Engine/Source/Runtime/Experimental/IoStore/OnDemand/Private/OnDemandInstallCache.cpp
--- a/Engine/Source/Runtime/Experimental/IoStore/OnDemand/Private/OnDemandInstallCache.cpp
+++ b/Engine/Source/Runtime/Experimental/IoStore/OnDemand/Private/OnDemandInstallCache.cpp
@@ -48,7 +48,7 @@ namespace UE::IoStore
 ///////////////////////////////////////////////////////////////////////////////
 namespace CVars
 {
-	static bool GIoStoreOnDemandEnableDefrag = false;
+	static bool GIoStoreOnDemandEnableDefrag = true;
 	static FAutoConsoleVariableRef CVar_IoStoreOnDemandEnableDefrag(
 		TEXT("iostore.EnableDefrag"),
 		GIoStoreOnDemandEnableDefrag,
