Iris
* Now intercepting StopReplication requests done while inside ReceiveUpdate.
* Those objects now stop replicating only after we finished processing all received data from the current frame.

#rnx

diff --git a/Engine/Source/Runtime/Experimental/Iris/Core/Private/Iris/ReplicationSystem/ReplicationBridge.cpp b/Engine/Source/Runtime/Experimental/Iris/Core/Private/Iris/ReplicationSystem/ReplicationBridge.cpp
--- a/Engine/Source/Runtime/Experimental/Iris/Core/Private/Iris/ReplicationSystem/ReplicationBridge.cpp
+++ b/Engine/Source/Runtime/Experimental/Iris/Core/Private/Iris/ReplicationSystem/ReplicationBridge.cpp
@@ -77,6 +77,13 @@ void UReplicationBridge::PostReceiveUpdate()
 	check(bInReceiveUpdate == true);
 	bInReceiveUpdate = false;
 
+	// Now process all StopReplication calls done while inside ReceiveUpdate
+	for (const auto& It : HandlesToStopReplicating)
+	{
+		StopReplicatingNetRefHandle(It.Key, It.Value);
+	}
+	HandlesToStopReplicating.Reset();
+	
 	OnPostReceiveUpdate();
 }
 
@@ -622,6 +629,26 @@ void UReplicationBridge::StopReplicatingNetRefHandle(FNetRefHandle Handle, EEndR
 		return;
 	}
 
+	// While we are inside ReceiveUpdate, queue stop replication requests instead of immediately stopping replication
+	// This allows us the apply any received updates before we cut off this object
+	if (IsInReceiveUpdate())
+	{
+		UE_LOG(LogIrisBridge, Verbose, TEXT("Delayed request to StopReplicating %s (flags: %s) because it was called while inside ReceiveUpdate"), *NetRefHandleManager->PrintObjectFromNetRefHandle(Handle), *LexToString(EndReplicationFlags));
+		
+		// Detect if we have diverging EndReplicationFlags for the same netobject
+#if DO_ENSURE
+		{
+			EEndReplicationFlags* PreviousFlags = HandlesToStopReplicating.Find(Handle);
+			const bool bPreviousFlagsMatch = PreviousFlags ? (*PreviousFlags) == EndReplicationFlags : true;
+			ensureMsgf(bPreviousFlagsMatch, TEXT("Received multiple StopReplicating calls for %s with different EndReplicationFlags: Previous: %s | Newest: %s"),
+				*NetRefHandleManager->PrintObjectFromNetRefHandle(Handle), *LexToString(*PreviousFlags), *LexToString(EndReplicationFlags));
+		}
+#endif
+
+		HandlesToStopReplicating.Add(Handle, EndReplicationFlags);
+		return;
+	}
+
 	const FInternalNetRefIndex InternalReplicationIndex = NetRefHandleManager->GetInternalIndex(Handle);
 	if (NetRefHandleManager->IsLocal(InternalReplicationIndex))
 	{
diff --git a/Engine/Source/Runtime/Experimental/Iris/Core/Public/Iris/ReplicationSystem/ReplicationBridge.h b/Engine/Source/Runtime/Experimental/Iris/Core/Public/Iris/ReplicationSystem/ReplicationBridge.h
--- a/Engine/Source/Runtime/Experimental/Iris/Core/Public/Iris/ReplicationSystem/ReplicationBridge.h
+++ b/Engine/Source/Runtime/Experimental/Iris/Core/Public/Iris/ReplicationSystem/ReplicationBridge.h
@@ -344,6 +344,9 @@ private:
 	/** Tracks if we are in the middle of processing incoming data */
 	bool bInReceiveUpdate = false;
 
+	/** List of replicated objects that requested to stop replicating while we were in ReceiveUpdate */
+	TMap<FNetRefHandle, EEndReplicationFlags> HandlesToStopReplicating;
+
 private:
 
 	// DestructionInfos
