Fix compil error in unity builds

diff --git a/Engine/Source/Runtime/Experimental/Iris/Core/Private/Iris/ReplicationSystem/ObjectReplicationBridge.cpp b/Engine/Source/Runtime/Experimental/Iris/Core/Private/Iris/ReplicationSystem/ObjectReplicationBridge.cpp
--- a/Engine/Source/Runtime/Experimental/Iris/Core/Private/Iris/ReplicationSystem/ObjectReplicationBridge.cpp
+++ b/Engine/Source/Runtime/Experimental/Iris/Core/Private/Iris/ReplicationSystem/ObjectReplicationBridge.cpp
@@ -57,7 +57,6 @@ namespace UE::Net::Private::ObjectBridgeDebugging
 DEFINE_LOG_CATEGORY(LogIrisFilterConfig)
 
 #define UE_LOG_OBJECTREPLICATIONBRIDGE(Category, Format, ...)  UE_LOG(LogIrisBridge, Category, TEXT("ObjectReplicationBridge(%u)::") Format, GetReplicationSystem()->GetId(), ##__VA_ARGS__)
-#define UE_LOG_BRIDGEID(Category, Verbosity, Format, ...)  UE_LOG(Category, Verbosity, TEXT("ObjectReplicationBridge(%u)::") Format, GetReplicationSystem()->GetId(), ##__VA_ARGS__)
 
 static bool bUseFrequencyBasedPolling = true;
 static FAutoConsoleVariableRef CVarUseFrequencyBasedPolling(
diff --git a/Engine/Source/Runtime/Experimental/Iris/Core/Private/Iris/ReplicationSystem/ReplicationBridge.cpp b/Engine/Source/Runtime/Experimental/Iris/Core/Private/Iris/ReplicationSystem/ReplicationBridge.cpp
--- a/Engine/Source/Runtime/Experimental/Iris/Core/Private/Iris/ReplicationSystem/ReplicationBridge.cpp
+++ b/Engine/Source/Runtime/Experimental/Iris/Core/Private/Iris/ReplicationSystem/ReplicationBridge.cpp
@@ -31,7 +31,6 @@
 DEFINE_LOG_CATEGORY(LogIrisBridge)
 
 #define UE_LOG_REPLICATIONBRIDGE(Category, Format, ...)  UE_LOG(LogIrisBridge, Category, TEXT("ReplicationBridge(%u)::") Format, ReplicationSystem->GetId(), ##__VA_ARGS__)
-#define UE_LOG_BRIDGEID(Category, Verbosity, Format, ...)  UE_LOG(Category, Verbosity, TEXT("ReplicationBridge(%u)::") Format, ReplicationSystem->GetId(), ##__VA_ARGS__)
 
 static bool bEnableFlushReliableRPCOnDestroy = true;
 static FAutoConsoleVariableRef CVarEnableFlushReliableRPCOnDestroy(
diff --git a/Engine/Source/Runtime/Experimental/Iris/Core/Public/Iris/ReplicationSystem/ReplicationBridge.h b/Engine/Source/Runtime/Experimental/Iris/Core/Public/Iris/ReplicationSystem/ReplicationBridge.h
--- a/Engine/Source/Runtime/Experimental/Iris/Core/Public/Iris/ReplicationSystem/ReplicationBridge.h
+++ b/Engine/Source/Runtime/Experimental/Iris/Core/Public/Iris/ReplicationSystem/ReplicationBridge.h
@@ -61,6 +61,8 @@ namespace UE::Net
 	IRISCORE_API DECLARE_LOG_CATEGORY_EXTERN(LogIrisBridge, Log, All);
 #endif
 
+#define UE_LOG_BRIDGEID(Category, Verbosity, Format, ...)  UE_LOG(Category, Verbosity, TEXT("ReplicationBridge(%u)::") Format, GetReplicationSystem()->GetId(), ##__VA_ARGS__)
+
 struct FReplicationBridgeSerializationContext
 {
 	FReplicationBridgeSerializationContext(UE::Net::FNetSerializationContext& InSerialiazationContext, uint32 InConnectionId, bool bInIsDestructionInfo = false);
