Dataflow : fix clamp node not being able to be placed on a graph

diff --git a/Engine/Source/Runtime/Experimental/Dataflow/Core/Private/Dataflow/DataflowMathNodes.cpp b/Engine/Source/Runtime/Experimental/Dataflow/Core/Private/Dataflow/DataflowMathNodes.cpp
--- a/Engine/Source/Runtime/Experimental/Dataflow/Core/Private/Dataflow/DataflowMathNodes.cpp
+++ b/Engine/Source/Runtime/Experimental/Dataflow/Core/Private/Dataflow/DataflowMathNodes.cpp
@@ -723,6 +723,8 @@ FDataflowMathClampNode::FDataflowMathClampNode(const UE::Dataflow::FNodeParamete
 	: FDataflowMathOneInputOperatorNode(InParam, InGuid)
 {
 	RegisterInputsAndOutputs();
+	RegisterInputConnection(&Min);
+	RegisterInputConnection(&Max);
 	Min.Value = 0.0;
 	Max.Value = 1.0;
 }
diff --git a/Engine/Source/Runtime/Experimental/Dataflow/Core/Public/Dataflow/DataflowMathNodes.h b/Engine/Source/Runtime/Experimental/Dataflow/Core/Public/Dataflow/DataflowMathNodes.h
--- a/Engine/Source/Runtime/Experimental/Dataflow/Core/Public/Dataflow/DataflowMathNodes.h
+++ b/Engine/Source/Runtime/Experimental/Dataflow/Core/Public/Dataflow/DataflowMathNodes.h
@@ -503,10 +503,10 @@ struct FDataflowMathClampNode : public FDataflowMathOneInputOperatorNode
 	GENERATED_USTRUCT_BODY()
 	DATAFLOW_NODE_DEFINE_INTERNAL(FDataflowMathClampNode, "Clamp", DATAFLOW_MATH_NODES_CATEGORY, "")
 
-	UPROPERTY(meta = (DataflowInput))
+	UPROPERTY(EditAnywhere, Category = "Clamp", meta = (DataflowInput))
 	FDataflowNumericTypes Min;
 
-	UPROPERTY(meta = (DataflowInput))
+	UPROPERTY(EditAnywhere, Category = "Clamp", meta = (DataflowInput))
 	FDataflowNumericTypes Max;
 
 public:
