ChaosDD - fix debug draw verbose log. Change to UE_LOGFMT

diff --git a/Engine/Source/Runtime/Experimental/Chaos/Public/ChaosDebugDraw/ChaosDDFrame.h b/Engine/Source/Runtime/Experimental/Chaos/Public/ChaosDebugDraw/ChaosDDFrame.h
--- a/Engine/Source/Runtime/Experimental/Chaos/Public/ChaosDebugDraw/ChaosDDFrame.h
+++ b/Engine/Source/Runtime/Experimental/Chaos/Public/ChaosDebugDraw/ChaosDDFrame.h
@@ -45,12 +45,18 @@ namespace ChaosDD::Private
 				Commands.Reserve(InCommandQueueLength);
 				LatentCommands.Reserve(InCommandQueueLength);
 			}
+			BuildName();
 		}
 
 		virtual ~FChaosDDFrame()
 		{
 		}
 
+		const FString& GetName() const
+		{
+			return Name;
+		}
+
 		FChaosDDTimelinePtr GetTimeline() const
 		{
 			return Timeline.Pin();
@@ -235,8 +241,12 @@ namespace ChaosDD::Private
 			, Commands(MoveTemp(Other.Commands))
 			, LatentCommands(MoveTemp(Other.LatentCommands))
 		{
+			BuildName();
 		}
 
+		// Implemented in ChaosDDTimeline.h
+		void BuildName();
+
 		FChaosDDTimelineWeakPtr Timeline;
 		int64 FrameIndex;
 		double Time;
@@ -245,6 +255,8 @@ namespace ChaosDD::Private
 		int32 CommandBudget;
 		int32 CommandCost;
 
+		FString Name;
+
 		// Debug draw commands
 		TArray<FChaosDDCommand> Commands;
 
diff --git a/Engine/Source/Runtime/Experimental/Chaos/Public/ChaosDebugDraw/ChaosDDTimeline.h b/Engine/Source/Runtime/Experimental/Chaos/Public/ChaosDebugDraw/ChaosDDTimeline.h
--- a/Engine/Source/Runtime/Experimental/Chaos/Public/ChaosDebugDraw/ChaosDDTimeline.h
+++ b/Engine/Source/Runtime/Experimental/Chaos/Public/ChaosDebugDraw/ChaosDDTimeline.h
@@ -116,6 +116,21 @@ namespace ChaosDD::Private
 		int32 CommandQueueLength;
 		FCriticalSection FramesCS;
 	};
+
+
+	// Here to avoid circular dependency with ChaosDDFrame.h
+	inline void FChaosDDFrame::BuildName()
+	{
+		if (FChaosDDTimelinePtr PinnedTimeline = GetTimeline())
+		{
+			Name = FString::Format(TEXT("{0}: {1}"), { PinnedTimeline->GetName(), GetFrameIndex() });
+		}
+		else
+		{
+			Name = FString::Format(TEXT("<Global>: {0}"), { GetFrameIndex()});
+		}
+	}
+
 }
 
 #endif
diff --git a/Engine/Source/Runtime/Experimental/ChaosSolverEngine/Private/Chaos/ChaosDebugDrawSubsystem.cpp b/Engine/Source/Runtime/Experimental/ChaosSolverEngine/Private/Chaos/ChaosDebugDrawSubsystem.cpp
--- a/Engine/Source/Runtime/Experimental/ChaosSolverEngine/Private/Chaos/ChaosDebugDrawSubsystem.cpp
+++ b/Engine/Source/Runtime/Experimental/ChaosSolverEngine/Private/Chaos/ChaosDebugDrawSubsystem.cpp
@@ -13,6 +13,7 @@
 #include "GameFramework/Pawn.h"
 #include "GameFramework/PlayerController.h"
 #include "HAL/CriticalSection.h"
+#include "Logging/StructuredLog.h"
 #include "Physics/Experimental/PhysScene_Chaos.h"
 #include "VisualLogger/VisualLogger.h"
 
@@ -281,7 +282,7 @@ bool UChaosDebugDrawSubsystem::ShouldCreateSubsystem(UObject* Outer) const
 		const bool bIsWorldTypeSupported = World->IsPreviewWorld() ? Chaos::CVars::bChaosDebugDraw_PreviewWorldsEnabled : true;
 		const bool bCreateDebugDraw = !IsRunningCommandlet() && bIsWorldTypeSupported;
 
-		UE_LOG(LogChaosDD, Log, TEXT("%s Chaos Debug Draw Scene for world %s"), bCreateDebugDraw ? TEXT("Creating") : TEXT("Not creating"), *World->GetName());
+		UE_LOGFMT(LogChaosDD, Log, "{0} Chaos Debug Draw Scene for world {1}", bCreateDebugDraw ? TEXT("Creating") : TEXT("Not creating"), World->GetName());
 		
 		return bCreateDebugDraw;
 	}
@@ -340,7 +341,7 @@ void UChaosDebugDrawSubsystem::Deinitialize()
 void UChaosDebugDrawSubsystem::Startup()
 {
 #if CHAOS_DEBUG_DRAW
-	UE_LOG(LogChaosDD, Log, TEXT("Chaos Debug Draw Startup"));
+	UE_LOGFMT(LogChaosDD, Log, "Chaos Debug Draw Startup");
 
 	OnTickWorldStartDelegate = FWorldDelegates::OnWorldPreActorTick.AddStatic(&StaticOnWorldTickStart);
 	OnTickWorldEndDelegate = FWorldDelegates::OnWorldPostActorTick.AddStatic(&StaticOnWorldTickEnd);
@@ -350,7 +351,7 @@ void UChaosDebugDrawSubsystem::Startup()
 void UChaosDebugDrawSubsystem::Shutdown()
 {
 #if CHAOS_DEBUG_DRAW
-	UE_LOG(LogChaosDD, Log, TEXT("Chaos Debug Draw Shutdown"));
+	UE_LOGFMT(LogChaosDD, Log, "Chaos Debug Draw Shutdown");
 
 	FWorldDelegates::OnWorldPreActorTick.Remove(OnTickWorldStartDelegate);
 	FWorldDelegates::OnWorldPostActorTick.Remove(OnTickWorldEndDelegate);
@@ -532,7 +533,7 @@ void UChaosDebugDrawSubsystem::RenderFrame(FChaosDDRenderer& Renderer, const Cha
 
 	if (bDebugDrawEnabled && Frame.IsValid() && (Frame->GetNumCommands() + Frame->GetNumLatentCommands() > 0))
 	{
-		UE_LOG(LogChaosDD, VeryVerbose, TEXT("Render %s %d %d+%d Commands"), *Frame->GetTimeline()->GetName(), Frame->GetFrameIndex(), Frame->GetNumCommands(), Frame->GetNumLatentCommands());
+		UE_LOGFMT(LogChaosDD, VeryVerbose, "Render {0} {1}+{2} Commands", Frame->GetName(), Frame->GetNumCommands(), Frame->GetNumLatentCommands());
 
 		// Render the legacy commands
 		Frame->VisitLatentCommands(
