ChaosDD - fix debug draw context scope on physics thread

diff --git a/Engine/Source/Runtime/Experimental/Chaos/Private/PBDRigidsSolver.cpp b/Engine/Source/Runtime/Experimental/Chaos/Private/PBDRigidsSolver.cpp
--- a/Engine/Source/Runtime/Experimental/Chaos/Private/PBDRigidsSolver.cpp
+++ b/Engine/Source/Runtime/Experimental/Chaos/Private/PBDRigidsSolver.cpp
@@ -1353,9 +1353,11 @@ namespace Chaos
 		const FReal StartSimTime = GetSolverTime();
 
 #if CHAOS_DEBUG_DRAW
-		if (!GetEvolution()->IsResimming())
+		ChaosDD::Private::FChaosDDTimelineContext DDContext;
+		bool bEnableDebugDrawCapture = !GetEvolution()->IsResimming();
+		if (bEnableDebugDrawCapture)
 		{
-			ChaosDD::Private::FChaosDDScopeTimelineContext DDContext(CDDFrameTimeline, StartSimTime, GetLastDt());
+			DDContext.BeginFrame(CDDFrameTimeline, StartSimTime, GetLastDt());
 		}
 #endif
 
@@ -1382,6 +1384,13 @@ namespace Chaos
 			//final step so we can destroy proxies
 			DestroyPendingProxies_Internal();
 		}
+
+#if CHAOS_DEBUG_DRAW
+		if (bEnableDebugDrawCapture)
+		{
+			DDContext.EndFrame();
+		}
+#endif
 	}
 
 #if CHAOS_DEBUG_DRAW
