Dataflow : Fix subgraph calling node to evaluate all outputs as an optimization to avoid multiple unecessary re-evaluation

diff --git a/Engine/Source/Runtime/Experimental/Dataflow/Engine/Private/Dataflow/DataflowSubGraphNodes.cpp b/Engine/Source/Runtime/Experimental/Dataflow/Engine/Private/Dataflow/DataflowSubGraphNodes.cpp
--- a/Engine/Source/Runtime/Experimental/Dataflow/Engine/Private/Dataflow/DataflowSubGraphNodes.cpp
+++ b/Engine/Source/Runtime/Experimental/Dataflow/Engine/Private/Dataflow/DataflowSubGraphNodes.cpp
@@ -641,10 +641,19 @@ void FDataflowCallSubGraphNode::Evaluate(UE::Dataflow::FContext& Context, const
 							}
 						}
 					}
+
+					TArray<FDataflowOutput*> OutputsToEvaluate = GetOutputs();
+
+					// Even if the evelaute function is for a single output 
+					// We are pulling all outputs otherwise this may force us to evaluate the loops multiple time 
+					// whioch may be more expensive
 					for (int32 Iteration = 0; Iteration < NumIterations; ++Iteration)
 					{
 						SubGraphContext.SetCurrentArrayIndex(Iteration);
-						SubGraphContext.EvaluateSubGraph(*SubGraphOutputNode, *this, *Out);
+						for (const FDataflowOutput* OutToEvaluate : OutputsToEvaluate)
+						{
+							SubGraphContext.EvaluateSubGraph(*SubGraphOutputNode, *this, *OutToEvaluate);
+						}
 						SubGraphContext.ClearAllData();
 					}
 				}
