Minor refactor of get install size for on-demand content.

#rnx

diff --git a/Engine/Source/Runtime/Experimental/IoStore/OnDemand/Private/OnDemandIoStore.cpp b/Engine/Source/Runtime/Experimental/IoStore/OnDemand/Private/OnDemandIoStore.cpp
--- a/Engine/Source/Runtime/Experimental/IoStore/OnDemand/Private/OnDemandIoStore.cpp
+++ b/Engine/Source/Runtime/Experimental/IoStore/OnDemand/Private/OnDemandIoStore.cpp
@@ -873,6 +873,8 @@ TIoStatusOr<uint64> FOnDemandIoStore::GetInstallSize(const FOnDemandGetInstallSi
 {
 	TRACE_CPUPROFILER_EVENT_SCOPE(FOnDemandIoStore::GetInstallSize);
 
+	using namespace UE::IoStore::Private;
+
 	TSet<FSharedOnDemandContainer> ContainersForInstallation;
 	TSet<FPackageId> PackageIdsToInstall;
 
@@ -886,32 +888,31 @@ TIoStatusOr<uint64> FOnDemandIoStore::GetInstallSize(const FOnDemandGetInstallSi
 		return Status;
 	}
 
-	Private::EBuildInstallDataOptions BuildInstallDataOptions = Private::EBuildInstallDataOptions::None;
-	if (EnumHasAnyFlags(Args.Options, EOnDemandGetInstallSizeOptions::IncludeSoftReferences))
-	{
-		BuildInstallDataOptions |= Private::EBuildInstallDataOptions::IncludeSoftReferences;
-	}
-
-	Private::FInstallData InstallData;
+	const bool bIncludeSoftReferences = EnumHasAnyFlags(Args.Options, EOnDemandGetInstallSizeOptions::IncludeSoftReferences);
+	TArray<FResolvedContainerChunks> ResolvedChunks;
 	TSet<FPackageId> Missing;
 
-	FIoStatus Status = Private::BuildInstallData(ContainersForInstallation, PackageIdsToInstall, BuildInstallDataOptions, InstallData, Missing);
-	if (Status.IsOk() == false)
-	{
-		return Status;
-	}
-
-	const uint64 RetSize = Algo::TransformAccumulate(InstallData,
-		[](const Private::FInstallData::ElementType& Pair) { return Pair.Value.TotalSize; },
+	ResolveChunksToInstall(
+		ContainersForInstallation,
+		PackageIdsToInstall,
+		bIncludeSoftReferences,
+		ResolvedChunks,
+		Missing);
+
+	const uint64 InstallSize = Algo::TransformAccumulate(
+		ResolvedChunks,
+		[](const FResolvedContainerChunks& R) { return R.TotalSize; },
 		uint64(0));
 
-	return RetSize;
+	return InstallSize;
 }
 
 FIoStatus FOnDemandIoStore::GetInstallSizesByMountId(const FOnDemandGetInstallSizeArgs& Args, TMap<FString, uint64>& OutSizesByMountId) const
 {
 	TRACE_CPUPROFILER_EVENT_SCOPE(FOnDemandIoStore::GetInstallSizesByMountId);
 
+	using namespace UE::IoStore::Private;
+
 	TSet<FSharedOnDemandContainer> ContainersForInstallation;
 	TSet<FPackageId> PackageIdsToInstall;
 
@@ -925,24 +926,20 @@ FIoStatus FOnDemandIoStore::GetInstallSizesByMountId(const FOnDemandGetInstallSi
 		return Status;
 	}
 
-	Private::EBuildInstallDataOptions BuildInstallDataOptions = Private::EBuildInstallDataOptions::None;
-	if (EnumHasAnyFlags(Args.Options, EOnDemandGetInstallSizeOptions::IncludeSoftReferences))
-	{
-		BuildInstallDataOptions |= Private::EBuildInstallDataOptions::IncludeSoftReferences;
-	}
-
-	Private::FInstallData InstallData;
+	const bool bIncludeSoftReferences = EnumHasAnyFlags(Args.Options, EOnDemandGetInstallSizeOptions::IncludeSoftReferences);
+	TArray<FResolvedContainerChunks> ResolvedChunks;
 	TSet<FPackageId> Missing;
 
-	FIoStatus Status = Private::BuildInstallData(ContainersForInstallation, PackageIdsToInstall, BuildInstallDataOptions, InstallData, Missing);
-	if (Status.IsOk() == false)
-	{
-		return Status;
-	}
+	ResolveChunksToInstall(
+		ContainersForInstallation,
+		PackageIdsToInstall,
+		bIncludeSoftReferences,
+		ResolvedChunks,
+		Missing);
 
-	for (const TPair<FSharedOnDemandContainer, Private::FContainerInstallData>& Pair : InstallData)
+	for (const FResolvedContainerChunks& R : ResolvedChunks)
 	{
-		OutSizesByMountId.FindOrAdd(Pair.Key->MountId, 0) += Pair.Value.TotalSize;
+		OutSizesByMountId.FindOrAdd(R.Container->MountId, 0) += R.TotalSize;
 	}
 
 	return EIoErrorCode::Ok;
