Added cvar for erroring on downloading specific assets.

diff --git a/Engine/Source/Runtime/Core/Public/IO/IoStatus.h b/Engine/Source/Runtime/Core/Public/IO/IoStatus.h
--- a/Engine/Source/Runtime/Core/Public/IO/IoStatus.h
+++ b/Engine/Source/Runtime/Core/Public/IO/IoStatus.h
@@ -34,7 +34,8 @@ enum class EIoErrorCode
 	PendingEncryptionKey,
 	Disabled,
 	NotInstalled,
-	PendingHostGroup
+	PendingHostGroup,
+	Last
 };
 
 /** Get I/O error code description. */
diff --git a/Engine/Source/Runtime/Experimental/IoStore/OnDemand/Private/OnDemandContentInstaller.cpp b/Engine/Source/Runtime/Experimental/IoStore/OnDemand/Private/OnDemandContentInstaller.cpp
--- a/Engine/Source/Runtime/Experimental/IoStore/OnDemand/Private/OnDemandContentInstaller.cpp
+++ b/Engine/Source/Runtime/Experimental/IoStore/OnDemand/Private/OnDemandContentInstaller.cpp
@@ -11,6 +11,7 @@
 #include "IO/IoStoreOnDemand.h"
 #include "Serialization/PackageStore.h"
 #include "IO/IoStatus.h"
+#include "Misc/PackageName.h"
 #include "OnDemandHttpThread.h"
 #include "OnDemandInstallCache.h"
 #include "OnDemandIoStore.h"
@@ -19,7 +20,17 @@
 
 namespace UE::IoStore
 {
-
+	namespace CVars
+	{
+#if !UE_BUILD_SHIPPING
+		FString IoStoreErrorOnRequest = "";
+		static FAutoConsoleVariableRef CVar_IoStoreErrorOnRequest(
+			TEXT("iostore.ErrorOnRequest"),
+			IoStoreErrorOnRequest,
+			TEXT("When the request with a debug name partially matching this cvar is found iostore will error with a random error.")
+			);
+#endif
+	}
 namespace Private
 {
 
@@ -317,6 +328,21 @@ void FOnDemandContentInstaller::ProcessRequest(FInstallRequestState& RequestStat
 		return CompleteRequest(RequestState);
 	}
 
+
+#if !UE_BUILD_SHIPPING
+	if (!UE::IoStore::CVars::IoStoreErrorOnRequest.IsEmpty())
+	{
+		if (FCString::Strstr(*ContentHandle->DebugName, *UE::IoStore::CVars::IoStoreErrorOnRequest) != nullptr || 
+			FCString::Strstr(*Request.Args.DebugName, *UE::IoStore::CVars::IoStoreErrorOnRequest) != nullptr)
+		{
+			Request.ErrorCode = (EIoErrorCode)((FMath::Rand() / (RAND_MAX / ((int)EIoErrorCode::Last-1))) + 1);
+			RequestState.ErrorReason = FString::Printf(TEXT("Debug error requested on debug name %s"), *UE::IoStore::CVars::IoStoreErrorOnRequest);
+			return CompleteRequest(RequestState);
+		}
+	}
+#endif
+
+
 	if (Request.ErrorCode.load(std::memory_order_relaxed) == EIoErrorCode::Cancelled)
 	{
 		return CompleteRequest(RequestState);
