Don't Prune/Defrag and Put chunks in parallel.

- updated content installer to process installs one at the time to prevent pruning/defrag running in parallel with puts

#rnx

diff --git a/Engine/Source/Runtime/Experimental/IoStore/OnDemand/Private/OnDemandContentInstaller.cpp b/Engine/Source/Runtime/Experimental/IoStore/OnDemand/Private/OnDemandContentInstaller.cpp
--- a/Engine/Source/Runtime/Experimental/IoStore/OnDemand/Private/OnDemandContentInstaller.cpp
+++ b/Engine/Source/Runtime/Experimental/IoStore/OnDemand/Private/OnDemandContentInstaller.cpp
@@ -2,6 +2,7 @@
 
 #include "OnDemandContentInstaller.h"
 
+#include "Algo/RemoveIf.h"
 #include "Async/UniqueLock.h"
 #include "Containers/RingBuffer.h"
 #include "HAL/Platform.h"
@@ -11,6 +12,7 @@
 #include "IO/IoStoreOnDemand.h"
 #include "Serialization/PackageStore.h"
 #include "IO/IoStatus.h"
+#include "Misc/Timespan.h"
 #include "Misc/PackageName.h"
 #include "OnDemandHttpThread.h"
 #include "OnDemandInstallCache.h"
@@ -210,13 +212,13 @@ void ResolveChunksToInstall(
 FOnDemandContentInstaller::FOnDemandContentInstaller(FOnDemandIoStore& InIoStore, FOnDemandHttpThread& InHttpClient)
 	: IoStore(InIoStore)
 	, HttpClient(InHttpClient)
-	, ProcessPipe(TEXT("IoStoreOnDemandProcessPipe"))
-	, HttpCompletionPipe(TEXT("IoStoreOnDemandHttpCompletionPipe"))
+	, InstallerPipe(TEXT("IoStoreOnDemandInstallerPipe"))
 {
 }
 
 FOnDemandContentInstaller::~FOnDemandContentInstaller()
 {
+	Shutdown();
 }
 
 FSharedInternalInstallRequest FOnDemandContentInstaller::Enqueue(
@@ -232,14 +234,16 @@ FSharedInternalInstallRequest FOnDemandContentInstaller::Enqueue(
 	FInstallRequestState* RequestState = nullptr;
 	{
 		UE::TUniqueLock Lock(Mutex);
-		RequestState = RequestStateAllocator.Construct();
-		RequestState->Request = InstallRequest;
+		RequestState			= RequestStateAllocator.Construct();
+		RequestState->SeqNo		= NextInstallSeqNo++;
+		RequestState->Request	= InstallRequest;
+
 		RequestStateByHandle.Add(InstallRequest, RequestState);
 	}
 
 	FOnDemandContentInstallerStats::OnRequestEnqueued();
 
-	ProcessPipe.Launch(
+	InstallerPipe.Launch(
 		TEXT("ProcessIoStoreOnDemandInstallRequest"),
 		[this, RequestState] { ProcessRequest(*RequestState); },
 		UE::Tasks::ETaskPriority::BackgroundLow);
@@ -249,40 +253,67 @@ FSharedInternalInstallRequest FOnDemandContentInstaller::Enqueue(
 
 void FOnDemandContentInstaller::Cancel(FSharedInternalInstallRequest InstallRequest)
 {
-	EIoErrorCode Expected = EIoErrorCode::Unknown;
-	if (InstallRequest->ErrorCode.compare_exchange_strong(Expected, EIoErrorCode::Cancelled))
-	{
-		UE_LOG(LogIoStoreOnDemand, Log, TEXT("Cancelling install request, ContentHandle=(%s)"),
-			*LexToString(InstallRequest->Args.ContentHandle));
-	}
+	InstallerPipe.Launch(
+		TEXT("CancelIoStoreOnDemandInstallRequest"),
+		[this, InstallRequest]
+		{
+			EIoErrorCode Expected = EIoErrorCode::Unknown;
+			if (InstallRequest->ErrorCode.compare_exchange_strong(Expected, EIoErrorCode::Cancelled) == false)
+			{
+				return;
+			}
+
+			UE_LOG(LogIoStoreOnDemand, Log, TEXT("Cancelling install request, ContentHandle=(%s)"),
+				*LexToString(InstallRequest->Args.ContentHandle));
+
+			FInstallRequestState* ToComplete = nullptr;
+			{
+				UE::TUniqueLock Lock(Mutex);
+				if (FInstallRequestState* RequestState = RequestStateByHandle.FindRef(InstallRequest))
+				{
+					if (RequestQueue.Remove(RequestState) > 0)
+					{
+						ToComplete = RequestState;
+						RequestQueue.Heapify(ByInstallPriority);
+					}
+				}
+			}
+
+			if (ToComplete != nullptr)
+			{
+				CompleteRequest(*ToComplete);
+			}
+		});
 }
 
 void FOnDemandContentInstaller::UpdatePriority(FSharedInternalInstallRequest InstallRequest, int32 NewPriority)
 {
-	// Use the HTTP completion pipeline to synchronize access to the pending HTTP request(s)
-	HttpCompletionPipe.Launch(
+	InstallerPipe.Launch(
 		TEXT("UpdateIoStoreOnDemandInstallRequestPriority"),
 		[this, InstallRequest, NewPriority]
 		{
 			UE::TUniqueLock Lock(Mutex);
 			if (FInstallRequestState* RequestState = RequestStateByHandle.FindRef(InstallRequest))
 			{
-				if (!RequestState->bHttpRequestsIssued.load(std::memory_order_relaxed) || NewPriority <= RequestState->Priority)
-				{
-					return;
-				}
-
 				UE_LOG(LogIoStoreOnDemand, Log, TEXT("Updating install request priority, ContentHandle=(%s), CurrentPriority=%d, NewPriority=%d"),
 					*LexToString(InstallRequest->Args.ContentHandle), RequestState->Priority, NewPriority);
 
 				RequestState->Priority = NewPriority;
-				for (FChunkHttpRequestHandle& PendingHttpRequest : RequestState->HttpRequestHandles)
+
+				if (RequestState->bHttpRequestsIssued.load(std::memory_order_relaxed))
 				{
-					if (PendingHttpRequest.Handle != nullptr)
+					for (FChunkHttpRequestHandle& PendingHttpRequest : RequestState->HttpRequestHandles)
 					{
-						HttpClient.ReprioritizeRequest(PendingHttpRequest.Handle, RequestState->Priority);
+						if (PendingHttpRequest.Handle != nullptr)
+						{
+							HttpClient.ReprioritizeRequest(PendingHttpRequest.Handle, RequestState->Priority);
+						}
 					}
 				}
+				else
+				{
+					RequestQueue.Heapify(ByInstallPriority);
+				}
 			}
 		});
 }
@@ -394,11 +425,8 @@ void FOnDemandContentInstaller::ProcessRequest(FInstallRequestState& RequestStat
 	uint64 TotalContentSize	= 0;	
 	uint64 TotalInstallSize	= 0;
 
+	// Find all chunks we need to fetch from the resolved chunk(s)
 	{
-		// TODO: This may double-request chunks. Chunks don't become referenced until the request completes, so two requests that
-		// share chunks can race to download the same chunk and it may be cached more than once.
-
-		TSet<FIoHash> ChunksToInstall;
 		for (int32 ContainerIndex = 0; FResolvedContainerChunks& ResolvedChunks : RequestState.ResolvedChunks)
 		{
 			for (int32 EntryIndex : ResolvedChunks.EntryIndices)
@@ -407,7 +435,6 @@ void FOnDemandContentInstaller::ProcessRequest(FInstallRequestState& RequestStat
 				const bool bInstalled = IoStore.InstallCache->IsChunkCached(Entry.Hash);
 				if (bInstalled == false)
 				{
-					ChunksToInstall.Add(Entry.Hash);
 					RequestState.HttpRequestHandles.Add(FChunkHttpRequestHandle
 					{
 						.Handle			= nullptr,
@@ -421,17 +448,6 @@ void FOnDemandContentInstaller::ProcessRequest(FInstallRequestState& RequestStat
 			}
 			ContainerIndex++;
 		}
-
-		if (ChunksToInstall.IsEmpty() == false)
-		{
-			if (FIoStatus Status = IoStore.InstallCache->Purge(MoveTemp(ChunksToInstall)); !Status.IsOk())
-			{
-				Request.ErrorCode			= Status.GetErrorCode();
-				RequestState.ErrorReason	= Status.ToString(); 
-
-				return CompleteRequest(RequestState);
-			}
-		}
 	}
 
 	RequestState.Progress.TotalContentSize		= TotalContentSize;
@@ -442,7 +458,116 @@ void FOnDemandContentInstaller::ProcessRequest(FInstallRequestState& RequestStat
 	if (RequestState.HttpRequestHandles.IsEmpty())
 	{
 		Request.ErrorCode = EIoErrorCode::Ok;
-		return CompleteRequest(RequestState); 
+		return CompleteRequest(RequestState);
+	}
+
+	if (Request.ErrorCode.load(std::memory_order_relaxed) == EIoErrorCode::Cancelled)
+	{
+		return CompleteRequest(RequestState);
+	}
+
+	bool bStartDownload = false;
+	{
+		TUniqueLock Lock(Mutex);
+		if (CurrentRequest == nullptr)
+		{
+			CurrentRequest	= &RequestState;
+			bStartDownload	= true;
+		}
+		else
+		{
+			RequestQueue.HeapPush(&RequestState, ByInstallPriority);
+		}
+	}
+
+	if (bStartDownload)
+	{
+		check(CurrentRequest == &RequestState);
+		const bool bRemoveAlreadyCachedChunks = false;
+		StartDownloadChunks(RequestState, bRemoveAlreadyCachedChunks);
+	}
+}
+
+void FOnDemandContentInstaller::TryStartDownloadChunks()
+{
+	if (bShuttingDown.load(std::memory_order_relaxed))
+	{
+		return;
+	}
+
+	FInstallRequestState* NextRequest = nullptr;
+	{
+		TUniqueLock Lock(Mutex);
+		if (CurrentRequest == nullptr && RequestQueue.IsEmpty() == false)
+		{
+			RequestQueue.HeapPop(NextRequest, ByInstallPriority, EAllowShrinking::No);
+			CurrentRequest = NextRequest;
+		}
+	}
+
+	if (NextRequest != nullptr)
+	{
+		InstallerPipe.Launch(
+			TEXT("StartIoStoreOnDemandInstallRequest"),
+			[this, NextRequest] { StartDownloadChunks(*NextRequest, /*bRemoveAlreadyCachedChunks*/ true); },
+			UE::Tasks::ETaskPriority::BackgroundLow);
+	}
+}
+
+void FOnDemandContentInstaller::StartDownloadChunks(FInstallRequestState& RequestState, bool bRemoveAlreadyCachedChunks)
+{
+	FOnDemandInternalInstallRequest& Request = *RequestState.Request;
+	check(RequestState.HttpRequestHandles.IsEmpty() == false);
+
+	if (Request.ErrorCode.load(std::memory_order_relaxed) == EIoErrorCode::Cancelled)
+	{
+		return CompleteRequest(RequestState);
+	}
+
+	if (bRemoveAlreadyCachedChunks)
+	{
+		RequestState.HttpRequestHandles.SetNum(
+			Algo::RemoveIf(
+				RequestState.HttpRequestHandles,
+				[this, &RequestState](FChunkHttpRequestHandle& HttpRequest)
+				{
+					FSharedOnDemandContainer& Container		= RequestState.ResolvedChunks[HttpRequest.ContainerIndex].Container;
+					const FOnDemandChunkEntry& ChunkEntry	= Container->ChunkEntries[HttpRequest.EntryIndex];
+
+					if (IoStore.InstallCache->IsChunkCached(ChunkEntry.Hash))
+					{
+						RequestState.Progress.TotalInstallSize -= ChunkEntry.EncodedSize;
+						return true;
+					}
+
+					return false;
+				}));
+
+		if (RequestState.HttpRequestHandles.IsEmpty())
+		{
+			Request.ErrorCode = EIoErrorCode::Ok;
+			return CompleteRequest(RequestState);
+		}
+	}
+
+	// Make sure we have enough space in the cache
+	{
+		TSet<FIoHash> ChunksToInstall;
+		for (FChunkHttpRequestHandle& HttpRequest : RequestState.HttpRequestHandles)
+		{
+			FSharedOnDemandContainer& Container		= RequestState.ResolvedChunks[HttpRequest.ContainerIndex].Container;
+			const FOnDemandChunkEntry& ChunkEntry	= Container->ChunkEntries[HttpRequest.EntryIndex];
+
+			ChunksToInstall.Add(ChunkEntry.Hash);
+		}
+
+		if (FIoStatus Status = IoStore.InstallCache->Purge(MoveTemp(ChunksToInstall)); !Status.IsOk())
+		{
+			Request.ErrorCode			= Status.GetErrorCode();
+			RequestState.ErrorReason	= Status.ToString(); 
+
+			return CompleteRequest(RequestState);
+		}
 	}
 
 	if (Request.ErrorCode.load(std::memory_order_relaxed) == EIoErrorCode::Cancelled)
@@ -465,7 +590,7 @@ void FOnDemandContentInstaller::ProcessRequest(FInstallRequestState& RequestStat
 			[this, &RequestState, &HttpRequest](uint32 HttpStatusCode, FStringView ErrorReason, FIoBuffer&& Chunk)
 			{
 				TRACE_CPUPROFILER_EVENT_SCOPE(FOnDemandContentInstaller::Callback);
-				HttpCompletionPipe.Launch(
+				InstallerPipe.Launch(
 					TEXT("ProcessIoStoreOnDemandDownloadedChunk"),
 					[this, &RequestState, &HttpRequest, HttpStatusCode, ErrorReason = FString(ErrorReason), Chunk = MoveTemp(Chunk)]() mutable
 					{
@@ -538,6 +663,11 @@ void FOnDemandContentInstaller::ProcessDownloadedChunk(
 	const bool bCompleted = ++RequestState.DownloadedChunkCount >= RequestState.HttpRequestHandles.Num();
 	if (bCompleted)
 	{
+		if (FIoStatus Status = IoStore.InstallCache->Flush(); !Status.IsOk())
+		{
+			Request.ErrorCode			= Status.GetErrorCode();
+			RequestState.ErrorReason	= Status.ToString();
+		}
 		if (Request.ErrorCode == EIoErrorCode::Unknown)
 		{
 			Request.ErrorCode = EIoErrorCode::Ok;
@@ -575,7 +705,7 @@ void FOnDemandContentInstaller::NotifyProgress(FInstallRequestState& RequestStat
 	{
 		return;
 	}
-	RequestState.LastProgressCycles = Cycles; 
+	RequestState.LastProgressCycles = Cycles;
 
 	FOnDemandInstallProgress Progress = RequestState.Progress;
 	if (EnumHasAnyFlags(Request.Args.Options, EOnDemandInstallOptions::CallbackOnGameThread))
@@ -600,14 +730,8 @@ void FOnDemandContentInstaller::CompleteRequest(FInstallRequestState& RequestSta
 	FSharedInternalInstallRequest Request = MoveTemp(RequestState.Request);
 	ensure(Request->ErrorCode != EIoErrorCode::Unknown);
 
-	if (FIoStatus Status = IoStore.InstallCache->Flush(); !Status.IsOk())
-	{
-		Request->ErrorCode			= Status.GetErrorCode();
-		RequestState.ErrorReason	= Status.ToString();
-	}
-
 	// Mark all resolved chunk(s) as referenced by the content handle and notify the package store to update
-	if (RequestState.ResolvedChunkCount > 0)
+	if (Request->ErrorCode == EIoErrorCode::Ok && RequestState.ResolvedChunkCount > 0)
 	{
 		{
 			FOnDemandContentHandle& ContentHandle = Request->Args.ContentHandle;
@@ -652,10 +776,16 @@ void FOnDemandContentInstaller::CompleteRequest(FInstallRequestState& RequestSta
 
 	{
 		UE::TUniqueLock Lock(Mutex);
+		if (CurrentRequest == &RequestState)
+		{
+			CurrentRequest = nullptr;
+		}
 		RequestStateByHandle.Remove(Request); 
 		RequestStateAllocator.Destroy(&RequestState);
 	}
 
+	TryStartDownloadChunks();
+
 	FOnDemandInstallRequest::EStatus RequestStatus = FOnDemandInstallRequest::EStatus::None;
 	switch (InstallResult.Status.GetErrorCode())
 	{
@@ -693,4 +823,52 @@ void FOnDemandContentInstaller::CompleteRequest(FInstallRequestState& RequestSta
 	}
 }
 
+void FOnDemandContentInstaller::Shutdown()
+{
+	bShuttingDown					= true;
+	const double MaxTimeSeconds		= 10.0;
+	const uint64 StartTimeCycles	= FPlatformTime::Cycles64();
+
+	// Wait for the currently installing request to finish
+	for (;;)
+	{
+		if (FPlatformTime::ToSeconds64(FPlatformTime::Cycles64() - StartTimeCycles) > MaxTimeSeconds)
+		{
+			UE_LOG(LogIoStoreOnDemand, Warning, TEXT("Content installer shutdown cancelled after %.2lf"), MaxTimeSeconds);
+			break;
+		}
+
+		InstallerPipe.WaitUntilEmpty(FTimespan::FromSeconds(1.0));
+		{
+			TUniqueLock Lock(Mutex);
+			if (CurrentRequest != nullptr)
+			{
+				continue;
+			}
+		}
+	}
+
+	// Cancel all remaining request(s)
+	for (;;)
+	{
+		FInstallRequestState* RequestState = nullptr;
+		{
+			TUniqueLock Lock(Mutex);
+			if (RequestQueue.IsEmpty() == false)
+			{
+				RequestQueue.HeapPop(RequestState, ByInstallPriority, EAllowShrinking::No);
+			}
+		}
+
+		if (RequestState == nullptr)
+		{
+			break;
+		}
+
+		FOnDemandInternalInstallRequest& Request = *RequestState->Request;
+		Request.ErrorCode = EIoErrorCode::Cancelled;
+		CompleteRequest(*RequestState);
+	}
+}
+
 } // namespace UE::IoStore
diff --git a/Engine/Source/Runtime/Experimental/IoStore/OnDemand/Private/OnDemandContentInstaller.h b/Engine/Source/Runtime/Experimental/IoStore/OnDemand/Private/OnDemandContentInstaller.h
--- a/Engine/Source/Runtime/Experimental/IoStore/OnDemand/Private/OnDemandContentInstaller.h
+++ b/Engine/Source/Runtime/Experimental/IoStore/OnDemand/Private/OnDemandContentInstaller.h
@@ -73,9 +73,20 @@ class FOnDemandContentInstaller
 		uint64										LastProgressCycles = 0;
 		uint64										StartTimeCycles = FPlatformTime::Cycles64();
 		int32										Priority = 0; 
+		uint32										SeqNo = 0;
 		std::atomic_bool							bHttpRequestsIssued{false};
 	};
 
+	static bool ByInstallPriority(const FInstallRequestState& LHS, const FInstallRequestState& RHS)
+	{
+		if (LHS.Priority == RHS.Priority)
+		{
+			return LHS.SeqNo < RHS.SeqNo;
+		}
+
+		return LHS.Priority > RHS.Priority;
+	}
+
 public:
 									FOnDemandContentInstaller(FOnDemandIoStore& IoStore, FOnDemandHttpThread& HttpClient);
 									~FOnDemandContentInstaller();
@@ -89,6 +100,8 @@ public:
 
 private:
 	void							ProcessRequest(FInstallRequestState& RequestState);
+	void							TryStartDownloadChunks();
+	void							StartDownloadChunks(FInstallRequestState& RequestState, bool bRemoveAlreadyCachedChunks);
 	void							ProcessDownloadedChunk(
 										FInstallRequestState& RequestState,
 										FChunkHttpRequestHandle& HttpRequest,
@@ -97,18 +110,22 @@ private:
 										FIoBuffer&& Chunk);
 	void							NotifyProgress(FInstallRequestState& RequestState);
 	void							CompleteRequest(FInstallRequestState& RequestState);
+	void							Shutdown();
 
 	using FInstallRequestStateMap	= TMap<FSharedInternalInstallRequest, FInstallRequestState*>;
 	using FRequestStateAllocator	= TSingleThreadedSlabAllocator<FInstallRequestState>;
 
-	FOnDemandIoStore&			IoStore;
-	FOnDemandHttpThread&		HttpClient;
-	UE::Tasks::FPipe			ProcessPipe;
-	UE::Tasks::FPipe			HttpCompletionPipe;
-
-	FMutex						Mutex;
-	FInstallRequestStateMap		RequestStateByHandle;
-	FRequestStateAllocator		RequestStateAllocator;
+	FOnDemandIoStore&				IoStore;
+	FOnDemandHttpThread&			HttpClient;
+	UE::Tasks::FPipe				InstallerPipe;
+
+	FMutex							Mutex;
+	FInstallRequestStateMap			RequestStateByHandle;
+	FRequestStateAllocator			RequestStateAllocator;
+	TArray<FInstallRequestState*>	RequestQueue;
+	FInstallRequestState*			CurrentRequest = nullptr;
+	uint32							NextInstallSeqNo = 0;
+	std::atomic_bool				bShuttingDown{false};
 };
 
 } // namespace UE::IoStore
