Iris
* Stop executing RPCs on root objects after a client called StopReplication on it

#rnx

diff --git a/Engine/Source/Runtime/Experimental/Iris/Core/Private/Iris/ReplicationSystem/NetBlob/NetRPC.cpp b/Engine/Source/Runtime/Experimental/Iris/Core/Private/Iris/ReplicationSystem/NetBlob/NetRPC.cpp
--- a/Engine/Source/Runtime/Experimental/Iris/Core/Private/Iris/ReplicationSystem/NetBlob/NetRPC.cpp
+++ b/Engine/Source/Runtime/Experimental/Iris/Core/Private/Iris/ReplicationSystem/NetBlob/NetRPC.cpp
@@ -470,7 +470,7 @@ void FNetRPC::CallFunction(FNetRPCCallContext& CallContext)
 		if (!TargetObjectReference.IsValid())
 		{
 			UE_LOG(LogIrisRpc, Error, TEXT("Rejected RPC %s due to missing object or function for object: %s."),
-				*(Function?Function->GetName():NetRPC_GetDebugFunctionName(ReplicationSystem, NetObjectReference, FunctionLocator)),
+				*(Function ? Function->GetName() : NetRPC_GetDebugFunctionName(ReplicationSystem, NetObjectReference, FunctionLocator)),
 				*NetRPC_GetDebugObjectRefName(Context, NetObjectReference, Object));
 		}
 		else
@@ -482,6 +482,26 @@ void FNetRPC::CallFunction(FNetRPCCallContext& CallContext)
 		return;
 	}
 
+	// Make sure the root object has an instance protocol
+	{
+		const FNetRefHandleManager& NetRefHandleManager = ReplicationSystem->GetReplicationSystemInternal()->GetNetRefHandleManager();
+		const FNetRefHandle RootObjectNetRefHandle = NetObjectReference.GetRefHandle();
+		const FInternalNetRefIndex InternalIndex = NetRefHandleManager.GetInternalIndex(RootObjectNetRefHandle);
+
+		ensureMsgf(InternalIndex != FNetRefHandleManager::InvalidInternalIndex, TEXT("Processing NetRPC for an object before it's internal index was assigned. NetRefHandle: %s | Instance: %s"),
+			*RootObjectNetRefHandle.ToString(), *GetNameSafe(Object));
+
+		// If the object stopped replicating locally, discard all received RPCs so they don't get executed.
+		if (!NetRefHandleManager.HasInstanceProtocol(InternalIndex))
+		{
+			UE_LOG(LogIrisRpc, Warning, TEXT("Rejected RPC: %s due to object: %s having stopped replication"),
+				*(Function ? Function->GetName() : NetRPC_GetDebugFunctionName(ReplicationSystem, NetObjectReference, FunctionLocator)),
+				*NetRPC_GetDebugObjectRefName(Context, NetObjectReference, Object));
+
+			return;
+		}
+	}
+
 	if ((Function->FunctionFlags & FUNC_Net) == 0)
 	{
 		UE_LOG(LogIrisRpc, Error, TEXT("Rejected %s function %s due to it not being a Net function. %s : %s"), (Function->FunctionFlags & FUNC_NetReliable ? TEXT("reliable") : TEXT("unreliable")), ToCStr(Function->GetName()), *NetObjectReference.ToString(), *Object->GetFullName());
diff --git a/Engine/Source/Runtime/Experimental/Iris/Core/Private/Iris/ReplicationSystem/NetRefHandleManager.cpp b/Engine/Source/Runtime/Experimental/Iris/Core/Private/Iris/ReplicationSystem/NetRefHandleManager.cpp
--- a/Engine/Source/Runtime/Experimental/Iris/Core/Private/Iris/ReplicationSystem/NetRefHandleManager.cpp
+++ b/Engine/Source/Runtime/Experimental/Iris/Core/Private/Iris/ReplicationSystem/NetRefHandleManager.cpp
@@ -299,6 +299,12 @@ const FReplicationInstanceProtocol* FNetRefHandleManager::DetachInstanceProtocol
 	return nullptr;
 }
 
+bool FNetRefHandleManager::HasInstanceProtocol(FInternalNetRefIndex InternalIndex) const
+{
+	check(InternalIndex == InvalidInternalIndex || AssignedInternalIndices.GetBit(InternalIndex));
+	return ReplicatedObjectData[InternalIndex].InstanceProtocol != nullptr;
+}
+
 FNetRefHandle FNetRefHandleManager::AllocateNetRefHandle(bool bIsStatic)
 {
 	uint64& NextHandleId = bIsStatic ? NextStaticHandleIndex : NextDynamicHandleIndex;
diff --git a/Engine/Source/Runtime/Experimental/Iris/Core/Private/Iris/ReplicationSystem/NetRefHandleManager.h b/Engine/Source/Runtime/Experimental/Iris/Core/Private/Iris/ReplicationSystem/NetRefHandleManager.h
--- a/Engine/Source/Runtime/Experimental/Iris/Core/Private/Iris/ReplicationSystem/NetRefHandleManager.h
+++ b/Engine/Source/Runtime/Experimental/Iris/Core/Private/Iris/ReplicationSystem/NetRefHandleManager.h
@@ -165,6 +165,9 @@ public:
 	// Used when we want to detach the instance protocol due to async destroy
 	const FReplicationInstanceProtocol* DetachInstanceProtocol(FInternalNetRefIndex InternalIndex);
 
+	/** Does this object has an instance protocol attached to it */
+	bool HasInstanceProtocol(FInternalNetRefIndex InternalIndex) const;
+
 	// Creates a DestructionInfo Handle which is used to replicate persistently destroyed static objects when late joining or streaming in levels with static objects that already has been destroyed
 	// Returns a new handle used to replicate the destruction info, and adds a cross-reference to the original handle being destroyed if that still exists 
 	// Handle is a reference to an already replicated Handle that is about to be destroyed
