Always enable on-demand HTTP thread and remove old installer path.

#rnx

diff --git a/Engine/Source/Runtime/Experimental/IoStore/OnDemand/Private/OnDemandHttpThread.cpp b/Engine/Source/Runtime/Experimental/IoStore/OnDemand/Private/OnDemandHttpThread.cpp
--- a/Engine/Source/Runtime/Experimental/IoStore/OnDemand/Private/OnDemandHttpThread.cpp
+++ b/Engine/Source/Runtime/Experimental/IoStore/OnDemand/Private/OnDemandHttpThread.cpp
@@ -2,8 +2,6 @@
 
 #include "OnDemandHttpThread.h"
 
-#if UE_ENABLE_REQUEST_THREAD
-
 #include "Async/UniqueLock.h"
 #include "HAL/IConsoleManager.h"
 #include "HAL/LowLevelMemTracker.h"
@@ -50,23 +48,79 @@ static const TCHAR* LexToString(EThreadPriority Priority)
 namespace UE::IoStore
 {
 
-////////////////////////////////////////
-// START EXTERN CVARS
-////////////////////////////////////////
-
-extern int32 GIasHttpHealthCheckWaitTime;
-extern int32 GOnDemandBackendThreadPriorityIndex;
-extern int32 GIasHttpTimeOutMs;
-extern int32 GIasHttpConnectionCount;
-extern int32 GIasHttpRecvBufKiB;
-extern int32 GIasHttpRetryCount;
-extern int32 GIasHttpFailTimeOutMs;
-extern bool GIasHttpAllowChunkedXfer;
-extern int32 GIasHttpConcurrentRequests;
-extern int32 GIasHttpRateLimitKiBPerSecond;
-extern int32 GIasHttpPollTimeoutMs;
-
-const int32 GREPLACEOnDemandBackendThreadPriorities[5] =
+///////////////////////////////////////////////////////////////////////////////
+int32 GIasHttpHealthCheckWaitTime = 3000;
+static FAutoConsoleVariableRef CVar_IasHttpHealthCheckWaitTime(
+	TEXT("ias.HttpHealthCheckWaitTime"),
+	GIasHttpHealthCheckWaitTime,
+	TEXT("Number of milliseconds to wait before reconnecting to avaiable endpoint(s)")
+);
+
+int32 GOnDemandBackendThreadPriorityIndex = 4; // EThreadPriority::TPri_AboveNormal
+FAutoConsoleVariableRef CVarOnDemandBackendThreadPriority(
+	TEXT("ias.onDemandBackendThreadPriority"),
+	GOnDemandBackendThreadPriorityIndex,
+	TEXT("Thread priority of the on demand backend thread: 0=Lowest, 1=BelowNormal, 2=SlightlyBelowNormal, 3=Normal, 4=AboveNormal\n")
+	TEXT("Note that this is switchable at runtime"),
+	ECVF_Default);
+
+int32 GIasHttpConnectionCount = 4;
+static FAutoConsoleVariableRef CVar_IasHttpConnectionCount(
+	TEXT("ias.HttpConnectionCount"),
+	GIasHttpConnectionCount,
+	TEXT("Number of open HTTP connections to the on demand endpoint(s).")
+);
+
+int32 GIasHttpRecvBufKiB = -1;
+static FAutoConsoleVariableRef CVar_GIasHttpRecvBufKiB(
+	TEXT("ias.HttpRecvBufKiB"),
+	GIasHttpRecvBufKiB,
+	TEXT("Recv buffer size")
+);
+
+int32 GIasHttpRetryCount = -1;
+static FAutoConsoleVariableRef CVar_IasHttpRetryCount(
+	TEXT("ias.HttpRetryCount"),
+	GIasHttpRetryCount,
+	TEXT("Number of times that a request should be retried before being considered failed. A negative value will use the default behaviour, which is one retry per host url provided.")
+);
+
+int32 GIasHttpFailTimeOutMs = 4 * 1000;
+static FAutoConsoleVariableRef CVar_IasHttpFailTimeOutMs(
+	TEXT("ias.HttpFailTimeOutMs"),
+	GIasHttpFailTimeOutMs,
+	TEXT("Fail infinite network waits that take longer than this (in ms, a value of zero will use the default)")
+);
+
+bool GIasHttpAllowChunkedXfer = false;
+static FAutoConsoleVariableRef CVar_IasHttpAllowChunkedXfer(
+	TEXT("ias.HttpAllowChunkedXfer"),
+	GIasHttpAllowChunkedXfer,
+	TEXT("Enable/disable IAS' support for chunked transfer encoding")
+);
+
+int32 GIasHttpConcurrentRequests = 8;
+static FAutoConsoleVariableRef CVar_IasHttpConcurrentRequests(
+	TEXT("ias.HttpConcurrentRequests"),
+	GIasHttpConcurrentRequests,
+	TEXT("Number of concurrent requests in the http client.")
+);
+
+int32 GIasHttpRateLimitKiBPerSecond = 0;
+static FAutoConsoleVariableRef CVar_GIasHttpRateLimitKiBPerSecond(
+	TEXT("ias.HttpRateLimitKiBPerSecond"),
+	GIasHttpRateLimitKiBPerSecond,
+	TEXT("Http throttle limit in KiBPerSecond")
+);
+
+int32 GIasHttpPollTimeoutMs = 17;
+static FAutoConsoleVariableRef CVar_GIasHttpPollTimeoutMs(
+	TEXT("ias.HttpPollTimeoutMs"),
+	GIasHttpPollTimeoutMs,
+	TEXT("Http tick poll timeout in milliseconds")
+);
+
+const int32 GOnDemandBackendThreadPriorities[5] =
 {
 	EThreadPriority::TPri_Lowest,
 	EThreadPriority::TPri_BelowNormal,
@@ -75,6 +129,12 @@ const int32 GREPLACEOnDemandBackendThreadPriorities[5] =
 	EThreadPriority::TPri_AboveNormal
 };
 
+////////////////////////////////////////
+// START EXTERN CVARS
+////////////////////////////////////////
+
+extern int32 GIasHttpTimeOutMs;
+
 ////////////////////////////////////////
 // END EXTERN CVARS
 ////////////////////////////////////////
@@ -113,20 +173,15 @@ struct FHttpRequest
 	int32 Priority = 0;
 };
 
-FOnDemandHttpThread* FOnDemandHttpThread::SingletonInstance = nullptr;
-
 FOnDemandHttpThread::FOnDemandHttpThread()
 {
 	LLM_SCOPE_BYTAG(Ias);
 
-	checkf(SingletonInstance == nullptr, TEXT("Attempting to create multiple instances of FOnDemandHttpThread"));
-	SingletonInstance = this;
-
 #if UE_IOSTORE_ONDEMAND_NO_HTTP_THREAD
 	ensure(TryCreateHttpClient());
 #else
-	const int32 ThreadPriorityIndex			= FMath::Clamp(GOnDemandBackendThreadPriorityIndex, 0, (int32)UE_ARRAY_COUNT(GREPLACEOnDemandBackendThreadPriorities) - 1);
-	EThreadPriority DesiredThreadPriority	= (EThreadPriority)GREPLACEOnDemandBackendThreadPriorities[ThreadPriorityIndex];
+	const int32 ThreadPriorityIndex			= FMath::Clamp(GOnDemandBackendThreadPriorityIndex, 0, (int32)UE_ARRAY_COUNT(GOnDemandBackendThreadPriorities) - 1);
+	EThreadPriority DesiredThreadPriority	= (EThreadPriority)GOnDemandBackendThreadPriorities[ThreadPriorityIndex];
 	ThreadPriority							= DesiredThreadPriority;
 	const uint32 StackSize					= 0; // Use default stack size
 	const uint64 ThreadAffinityMask			= FGenericPlatformAffinity::GetNoAffinityMask();
@@ -139,21 +194,12 @@ FOnDemandHttpThread::FOnDemandHttpThread()
 
 FOnDemandHttpThread::~FOnDemandHttpThread()
 {
-	checkf(SingletonInstance == this, TEXT("SingletonInstance was reassigned"));
-	SingletonInstance = nullptr;
-
 	Thread.Reset();
 
 	// TODO: Not 100% sure this is still needed?
 	DrainHttpRequests();
 }
 
-FOnDemandHttpThread& FOnDemandHttpThread::Get()
-{
-	check(SingletonInstance != nullptr);
-	return *SingletonInstance;
-}
-
 FOnDemandHttpThread::FRequestHandle FOnDemandHttpThread::IssueRequest(const FOnDemandChunkInfo& ChunkInfo, const FIoOffsetAndLength& ReadRange, int32 Priority, FCompletionCallback&& CompletionCallback)
 {
 	return IssueRequest(FOnDemandChunkInfo(ChunkInfo), ReadRange, Priority, MoveTemp(CompletionCallback));
@@ -410,8 +456,8 @@ bool FOnDemandHttpThread::TryCreateHttpClient()
 void FOnDemandHttpThread::UpdateThreadPriorityIfNeeded()
 {
 	// Read the thread priority from the cvar
-	int32 ThreadPriorityIndex = FMath::Clamp(GOnDemandBackendThreadPriorityIndex, 0, (int32)UE_ARRAY_COUNT(GREPLACEOnDemandBackendThreadPriorities) - 1);
-	EThreadPriority DesiredThreadPriority = (EThreadPriority)GREPLACEOnDemandBackendThreadPriorities[ThreadPriorityIndex];
+	int32 ThreadPriorityIndex = FMath::Clamp(GOnDemandBackendThreadPriorityIndex, 0, (int32)UE_ARRAY_COUNT(GOnDemandBackendThreadPriorities) - 1);
+	EThreadPriority DesiredThreadPriority = (EThreadPriority)GOnDemandBackendThreadPriorities[ThreadPriorityIndex];
 	if (DesiredThreadPriority != ThreadPriority)
 	{
 		UE_LOGFMT(LogIas, Log, "Updated IoStoreOnDemand.Http thread priority to '{Priority}'", LexToString(DesiredThreadPriority));
@@ -434,5 +480,3 @@ void FOnDemandHttpThread::DestroyRequest(FHttpRequest* Request)
 }
 
 } //namespace UE::IoStore
-
-#endif // UE_ENABLE_REQUEST_THREAD
diff --git a/Engine/Source/Runtime/Experimental/IoStore/OnDemand/Private/OnDemandHttpThread.h b/Engine/Source/Runtime/Experimental/IoStore/OnDemand/Private/OnDemandHttpThread.h
--- a/Engine/Source/Runtime/Experimental/IoStore/OnDemand/Private/OnDemandHttpThread.h
+++ b/Engine/Source/Runtime/Experimental/IoStore/OnDemand/Private/OnDemandHttpThread.h
@@ -2,15 +2,6 @@
 
 #pragma once
 
-#ifndef UE_ENABLE_REQUEST_THREAD
-	#define UE_ENABLE_REQUEST_THREAD 1
-#endif
-
-#define UE_ENABLE_IAS_REQUEST_THREAD (1 && UE_ENABLE_REQUEST_THREAD)
-#define UE_ENABLE_IAD_REQUEST_THREAD (1 && UE_ENABLE_REQUEST_THREAD)
-
-#if UE_ENABLE_REQUEST_THREAD
-
 #include "Async/Mutex.h"
 #include "Containers/StringFwd.h"
 #include "Delegates/Delegate.h"
@@ -104,7 +95,6 @@ private:
 	void DestroyRequest(FHttpRequest* Request);
 
 private:
-	static FOnDemandHttpThread* SingletonInstance;
 	TUniquePtr<FRunnableThread> Thread;
 
 	FEventRef TickThreadEvent;
@@ -124,18 +114,3 @@ private:
 };
 
 } //namespace UE::IoStore
-
-#else
-
-namespace UE::IoStore
-{
-
-// Dummy class for FOnDemandIoStore::OnDemandHttpThread
-// Remove when we delete UE_ENABLE_REQUEST_THREAD
-class FOnDemandHttpThread final
-{
-};
-
-} //namespace UE::IoStore
-
-#endif // UE_ENABLE_REQUEST_THREAD
diff --git a/Engine/Source/Runtime/Experimental/IoStore/OnDemand/Private/OnDemandIoDispatcherBackend.cpp b/Engine/Source/Runtime/Experimental/IoStore/OnDemand/Private/OnDemandIoDispatcherBackend.cpp
--- a/Engine/Source/Runtime/Experimental/IoStore/OnDemand/Private/OnDemandIoDispatcherBackend.cpp
+++ b/Engine/Source/Runtime/Experimental/IoStore/OnDemand/Private/OnDemandIoDispatcherBackend.cpp
@@ -10,8 +10,6 @@
 #include "HAL/Platform.h"
 #include "HAL/PlatformTime.h"
 #include "HAL/PreprocessorHelpers.h"
-#include "HAL/Runnable.h"
-#include "HAL/RunnableThread.h"
 #include "IO/Http/LaneTrace.h"
 #include "IO/IoAllocators.h"
 #include "IO/IoAllocators.h"
@@ -99,55 +97,6 @@ static FAutoConsoleVariableRef CVar_IasHttpPrimaryEndpoint(
 	TEXT("Primary endpoint to use returned from the distribution endpoint")
 );
 
-int32 GIasHttpPollTimeoutMs = 17;
-static FAutoConsoleVariableRef CVar_GIasHttpPollTimeoutMs(
-	TEXT("ias.HttpPollTimeoutMs"),
-	GIasHttpPollTimeoutMs,
-	TEXT("Http tick poll timeout in milliseconds")
-);
-
-int32 GIasHttpRateLimitKiBPerSecond = 0;
-static FAutoConsoleVariableRef CVar_GIasHttpRateLimitKiBPerSecond(
-	TEXT("ias.HttpRateLimitKiBPerSecond"),
-	GIasHttpRateLimitKiBPerSecond,
-	TEXT("Http throttle limit in KiBPerSecond")
-);
-
-int32 GIasHttpRecvBufKiB = -1;
-static FAutoConsoleVariableRef CVar_GIasHttpRecvBufKiB(
-	TEXT("ias.HttpRecvBufKiB"),
-	GIasHttpRecvBufKiB,
-	TEXT("Recv buffer size")
-);
-
-int32 GIasHttpConcurrentRequests = 8;
-static FAutoConsoleVariableRef CVar_IasHttpConcurrentRequests(
-	TEXT("ias.HttpConcurrentRequests"),
-	GIasHttpConcurrentRequests,
-	TEXT("Number of concurrent requests in the http client.")
-);
-
-int32 GIasHttpConnectionCount = 4;
-static FAutoConsoleVariableRef CVar_IasHttpConnectionCount(
-	TEXT("ias.HttpConnectionCount"),
-	GIasHttpConnectionCount,
-	TEXT("Number of open HTTP connections to the on demand endpoint(s).")
-);
-
-int32 GIasHttpFailTimeOutMs = 4 * 1000;
-static FAutoConsoleVariableRef CVar_IasHttpFailTimeOutMs(
-	TEXT("ias.HttpFailTimeOutMs"),
-	GIasHttpFailTimeOutMs,
-	TEXT("Fail infinite network waits that take longer than this (in ms, a value of zero will use the default)")
-);
-
-int32 GIasHttpRetryCount = -1;
-static FAutoConsoleVariableRef CVar_IasHttpRetryCount(
-	TEXT("ias.HttpRetryCount"),
-	GIasHttpRetryCount,
-	TEXT("Number of times that a request should be retried before being considered failed. A negative value will use the default behaviour, which is one retry per host url provided.")
-);
-
 int32 GIasHttpTimeOutMs = 10 * 1000;
 static FAutoConsoleVariableRef CVar_IasHttpTimeOutMs(
 	TEXT("ias.HttpTimeOutMs"),
@@ -155,13 +104,6 @@ static FAutoConsoleVariableRef CVar_IasHttpTimeOutMs(
 	TEXT("Time out value for HTTP requests in milliseconds")
 );
 
-int32 GIasHttpHealthCheckWaitTime = 3000;
-static FAutoConsoleVariableRef CVar_IasHttpHealthCheckWaitTime(
-	TEXT("ias.HttpHealthCheckWaitTime"),
-	GIasHttpHealthCheckWaitTime,
-	TEXT("Number of milliseconds to wait before reconnecting to avaiable endpoint(s)")
-);
-
 bool GIasHttpEnabled = true;
 static FAutoConsoleVariableRef CVar_IasHttpEnabled(
 	TEXT("ias.HttpEnabled"),
@@ -204,13 +146,6 @@ static FAutoConsoleVariableRef CVar_DistributedEndpointAttemptCount(
 	TEXT("Number of times we should try to resolve a distributed endpoint befor eusing the fallback url (if there is one)")
 );
 
-bool GIasHttpAllowChunkedXfer = false;
-static FAutoConsoleVariableRef CVar_IasHttpAllowChunkedXfer(
-	TEXT("ias.HttpAllowChunkedXfer"),
-	GIasHttpAllowChunkedXfer,
-	TEXT("Enable/disable IAS' support for chunked transfer encoding")
-);
-
 bool GIasEnableWriteOnlyDecoding = false;
 static FAutoConsoleVariableRef CVar_WriteOnlyDecoding(
 	TEXT("ias.EnableWriteOnlyDecoding"),
@@ -255,34 +190,6 @@ FAutoConsoleVariableRef CVarCompleteMaterializeTaskPriority(
 	ECVF_Default
 );
 
-// Thread priority cvar (settable at runtime)
-// We declare these explicitly rather than just casting the cvar in case the enum changes in future
-const int32 GOnDemandBackendThreadPriorities[] =
-{
-	EThreadPriority::TPri_Lowest,
-	EThreadPriority::TPri_BelowNormal,
-	EThreadPriority::TPri_SlightlyBelowNormal,
-	EThreadPriority::TPri_Normal,
-	EThreadPriority::TPri_AboveNormal
-};
-
-const TCHAR* GOnDemandBackendThreadPriorityNames[] =
-{
-	TEXT("TPri_Lowest"),
-	TEXT("TPri_BelowNormal"),
-	TEXT("TPri_SlightlyBelowNormal"),
-	TEXT("TPri_Normal"),
-	TEXT("TPri_AboveNormal")
-};
-
-int32 GOnDemandBackendThreadPriorityIndex = 4; // EThreadPriority::TPri_AboveNormal
-FAutoConsoleVariableRef CVarOnDemandBackendThreadPriority(
-	TEXT("ias.onDemandBackendThreadPriority"),
-	GOnDemandBackendThreadPriorityIndex,
-	TEXT("Thread priority of the on demand backend thread: 0=Lowest, 1=BelowNormal, 2=SlightlyBelowNormal, 3=Normal, 4=AboveNormal\n")
-	TEXT("Note that this is switchable at runtime"),
-	ECVF_Default);
-
 ///////////////////////////////////////////////////////////////////////////////
 
 [[nodiscard]] UE::Tasks::ETaskPriority GetRequestCompletionTaskPriority()
@@ -483,9 +390,7 @@ struct FChunkRequest
 	FChunkRequestParams Params;
 	FIoRequestImpl* RequestHead;
 	FIoRequestImpl* RequestTail;
-#if UE_ENABLE_IAS_REQUEST_THREAD
 	FOnDemandHttpThread::FRequestHandle HttpRequest = 0;
-#endif //UE_ENABLE_IAS_REQUEST_THREAD
 	FIASHostGroup HostGroup; // Still used in a few places even when UE_ENABLE_IAS_REQUEST_THREAD is enabled
 	FIoBuffer Chunk;
 	uint64 StartTime;
@@ -555,8 +460,7 @@ static void LogIoResult(
 
 ///////////////////////////////////////////////////////////////////////////////
 class FOnDemandIoBackend final
-	: public FRunnable
-	, public IOnDemandIoDispatcherBackend
+	: public IOnDemandIoDispatcherBackend
 {
 	using FIoRequestQueue = TThreadSafeIntrusiveQueue<FIoRequestImpl>;
 	using FChunkRequestQueue = TThreadSafeIntrusiveQueue<FChunkRequest>;
@@ -641,7 +545,7 @@ class FOnDemandIoBackend final
 			return ChunkRequest;
 		}
 
-		bool Cancel(FIoRequestImpl* Request, IIasCache* TheCache)
+		bool Cancel(FIoRequestImpl* Request, FOnDemandHttpThread& TheHttpClient, IIasCache* TheCache)
 		{
 #if UE_ALLOW_DISABLE_CANCELLING
 			if (GIasCancelationEnabled == false)
@@ -681,9 +585,7 @@ class FOnDemandIoBackend final
 					check(ChunkRequest.RequestTail == nullptr);
 
 					ChunkRequest.bCancelled = true;
-#if UE_ENABLE_IAS_REQUEST_THREAD
-					FOnDemandHttpThread::Get().CancelRequest(ChunkRequest.HttpRequest);
-#endif // UE_ENABLE_IAS_REQUEST_THREAD
+					TheHttpClient.CancelRequest(ChunkRequest.HttpRequest);
 
 					if (TheCache != nullptr)
 					{
@@ -737,7 +639,7 @@ class FOnDemandIoBackend final
 	};
 public:
 
-	FOnDemandIoBackend(const FOnDemandEndpointConfig& InConfig, FOnDemandIoStore& InIoStore, TUniquePtr<IIasCache>&& InCache);
+	FOnDemandIoBackend(const FOnDemandEndpointConfig& InConfig, FOnDemandIoStore& InIoStore, FOnDemandHttpThread& InHttpClient, TUniquePtr<IIasCache>&& InCache);
 	virtual ~FOnDemandIoBackend();
 
 	// I/O dispatcher backend
@@ -763,54 +665,36 @@ public:
 	
 	virtual FOnDemandStreamingCacheUsage GetCacheUsage() const override;
 
-	// Runnable
-	virtual bool Init() override { return true; }
-	virtual uint32 Run() override;
-	virtual void Stop() override;
-
 private:
 	bool Resolve(FIoRequestImpl* Request);
 
-	void ConditionallyStartBackendThread();
+	void BeginCreateDefaultHostGroup();
 	void IssueHttpRequest(FChunkRequest* ChunkRequest);
 	void CompleteRequest(FChunkRequest* ChunkRequest);
 	void CompleteCacheRead(FChunkRequest* ChunkRequest);
 	bool ResolveDefaultHostGroup();
 	bool CreateDefaultHostGroup();
 	bool ResolveDistributedEndpoint(const FDistributedEndpointUrl& Url, TArray<FString>& OutUrls);
-	void ProcessHttpRequests(FMultiEndpointHttpClient& HttpClient, int32 MaxConcurrentRequests);
 	int32 WaitForCompleteRequestTasks(float WaitTimeSeconds, float PollTimeSeconds);
-	void DrainHttpRequests();
-	void UpdateThreadPriorityIfNeeded();
 	void OnThreadTickIdle();
 
 #if UE_IAS_DEBUG_CONSOLE_CMDS
 	void RunRequestTest(const TArray<FString>& Args) const;
 #endif // UE_IAS_DEBUG_CONSOLE_CMDS
 
-
 	FOnDemandIoStore& IoStore;
+	FOnDemandHttpThread& HttpClient;
 	TUniquePtr<IIasCache> Cache;
 	TSharedPtr<const FIoDispatcherBackendContext> BackendContext;
-	TUniquePtr<FRunnableThread> BackendThread;
-	FEventRef TickBackendEvent;
 	FChunkRequests ChunkRequests;
 	FIoRequestQueue CompletedRequests;
-	FChunkRequestQueue HttpRequests;
-	TArray<FChunkRequest*> PendingCacheGets;
 	FBackendStatus BackendStatus;
 	TArray<FString> DefaultUrls;
 	FOnDemandIoBackendStats Stats;
 	FDistributedEndpointUrl DistributionUrl;
 	FEventRef DistributedEndpointEvent;
 	FIASHostGroup DefaultHostGroup;
-
-	EThreadPriority CurrentThreadPriority;
-
-#if UE_ENABLE_IAS_REQUEST_THREAD
 	FDelegateHandle OnThreadTickIdleHandle;
-#endif // UE_ENABLE_IAS_REQUEST_THREAD
-
 	mutable FRWLock Lock;
 	std::atomic_uint32_t InflightCacheRequestCount{0};
 	std::atomic_bool bStopRequested{false};
@@ -821,11 +705,11 @@ private:
 };
 
 ///////////////////////////////////////////////////////////////////////////////
-FOnDemandIoBackend::FOnDemandIoBackend(const FOnDemandEndpointConfig& Config, FOnDemandIoStore& InIoStore, TUniquePtr<IIasCache>&& InCache)
+FOnDemandIoBackend::FOnDemandIoBackend(const FOnDemandEndpointConfig& Config, FOnDemandIoStore& InIoStore, FOnDemandHttpThread& InHttpClient, TUniquePtr<IIasCache>&& InCache)
 	: IoStore(InIoStore)
+	, HttpClient(InHttpClient)
 	, Cache(MoveTemp(InCache))
 	, Stats(BackendStatus)
-	, CurrentThreadPriority(EThreadPriority::TPri_Num)
 {
 	FAnsiString EndpointTestPath = StringCast<ANSICHAR>(*Config.TocPath).Get();
 
@@ -853,8 +737,6 @@ FOnDemandIoBackend::FOnDemandIoBackend(const FOnDemandEndpointConfig& Config, FO
 			{
 				UE_LOG(LogIas, Display, TEXT("User invoked http error via 'ias.InvokeHttpFailure'"));
 				FHostGroupManager::Get().DisconnectAll();
-
-				TickBackendEvent->Trigger(); // Wake up the backend in case it is ideal waiting for requests
 			}),
 		ECVF_Default)
 	);
@@ -900,8 +782,8 @@ void FOnDemandIoBackend::Initialize(TSharedRef<const FIoDispatcherBackendContext
 	LLM_SCOPE_BYTAG(Ias);
 	UE_LOG(LogIas, Log, TEXT("Initializing on demand I/O dispatcher backend"));
 	BackendContext = Context;
-
-	ConditionallyStartBackendThread();
+	BeginCreateDefaultHostGroup();
+	OnThreadTickIdleHandle = HttpClient.OnTickIdle().AddRaw(this, &FOnDemandIoBackend::OnThreadTickIdle);
 }
 
 void FOnDemandIoBackend::Shutdown()
@@ -915,20 +797,11 @@ void FOnDemandIoBackend::Shutdown()
 
 	UE_LOG(LogIas, Log, TEXT("Shutting down on demand I/O dispatcher backend"));
 
-#if UE_ENABLE_IAS_REQUEST_THREAD
 	if (OnThreadTickIdleHandle.IsValid())
 	{
-		FOnDemandHttpThread::Get().OnTickIdle().Remove(OnThreadTickIdleHandle);
+		HttpClient.OnTickIdle().Remove(OnThreadTickIdleHandle);
 		OnThreadTickIdleHandle.Reset();
 	}
-#endif // UE_ENABLE_IAS_REQUEST_THREAD
-
-	// Stop and wait for our backend thread to finish.
-	// Note, the IoDispatcher typically waits for all its pending io requests before shutting down its backends.
-	BackendThread.Reset();
-
-	// Drain any remaining (cancelled) http requests that already been completed from the IoDispatcher point of view.
-	DrainHttpRequests();
 
 	// The CompleteRequest tasks may still be executing a while after the IoDispatcher has been notified about the completed io requests.
 	const int32 NumPending = WaitForCompleteRequestTasks(5.0f, 0.1f);
@@ -937,19 +810,6 @@ void FOnDemandIoBackend::Shutdown()
 	BackendContext.Reset();
 }
 
-void FOnDemandIoBackend::UpdateThreadPriorityIfNeeded()
-{
-	// Read the thread priority from the cvar
-	int32 ThreadPriorityIndex = FMath::Clamp(GOnDemandBackendThreadPriorityIndex, 0, (int32)UE_ARRAY_COUNT(GOnDemandBackendThreadPriorities) - 1);
-	EThreadPriority DesiredThreadPriority = (EThreadPriority)GOnDemandBackendThreadPriorities[ThreadPriorityIndex];
-	if (DesiredThreadPriority != CurrentThreadPriority)
-	{
-		UE_LOG(LogIas, Log, TEXT("Setting backend http thread priority to %s"), GOnDemandBackendThreadPriorityNames[ThreadPriorityIndex]);
-		FPlatformProcess::SetThreadPriority(DesiredThreadPriority);
-		CurrentThreadPriority = DesiredThreadPriority;
-	}
-}
-
 void FOnDemandIoBackend::OnThreadTickIdle()
 {
 	if (BackendStatus.ShouldAbandonCache())
@@ -964,21 +824,8 @@ void FOnDemandIoBackend::OnThreadTickIdle()
 	}
 }
 
-void FOnDemandIoBackend::ConditionallyStartBackendThread()
+void FOnDemandIoBackend::BeginCreateDefaultHostGroup()
 {
-#if !UE_ENABLE_IAS_REQUEST_THREAD
-	FWriteScopeLock _(Lock);
-	if (BackendThread.IsValid() == false)
-	{
-		// Read the desired thread priority from the cvar
-		int32 ThreadPriorityIndex = FMath::Clamp(GOnDemandBackendThreadPriorityIndex, 0, (int32)UE_ARRAY_COUNT(GOnDemandBackendThreadPriorities) - 1);
-		EThreadPriority DesiredThreadPriority = (EThreadPriority)GOnDemandBackendThreadPriorities[ThreadPriorityIndex];
-		CurrentThreadPriority = DesiredThreadPriority;
-
-		BackendThread.Reset(FRunnableThread::Create(this, TEXT("Ias.Http"), 0, DesiredThreadPriority));
-	}
-#else
-
 	UE::Tasks::Launch(UE_SOURCE_LOCATION, [this]()
 		{
 			if (CreateDefaultHostGroup())
@@ -986,15 +833,11 @@ void FOnDemandIoBackend::ConditionallyStartBackendThread()
 				BackendStatus.SetHttpEnabled(true);
 			}
 		});
-
-	OnThreadTickIdleHandle = FOnDemandHttpThread::Get().OnTickIdle().AddRaw(this, &FOnDemandIoBackend::OnThreadTickIdle);
-#endif // UE_ENABLE_IAS_REQUEST_THREAD
 }
 
 void FOnDemandIoBackend::IssueHttpRequest(FChunkRequest* ChunkRequest)
 {
-#if UE_ENABLE_IAS_REQUEST_THREAD
-	ChunkRequest->HttpRequest = FOnDemandHttpThread::Get().IssueRequest(ChunkRequest->Params.ChunkInfo, ChunkRequest->Params.ChunkRange, ChunkRequest->Priority,
+	ChunkRequest->HttpRequest = HttpClient.IssueRequest(ChunkRequest->Params.ChunkInfo, ChunkRequest->Params.ChunkRange, ChunkRequest->Priority,
 		[this, ChunkRequest](uint32 StatusCode, FStringView ErrorReason, FIoBuffer&& Data)
 		{
 			ChunkRequest->Chunk = MoveTemp(Data);
@@ -1013,11 +856,6 @@ void FOnDemandIoBackend::IssueHttpRequest(FChunkRequest* ChunkRequest)
 		});
 
 	check(ChunkRequest->HttpRequest != nullptr);
-#else
-	Stats.OnHttpEnqueue();
-	HttpRequests.EnqueueByPriority(ChunkRequest);
-	TickBackendEvent->Trigger();
-#endif //UE_ENABLE_IAS_REQUEST_THREAD
 }
 
 void FOnDemandIoBackend::CompleteRequest(FChunkRequest* ChunkRequest)
@@ -1112,11 +950,6 @@ void FOnDemandIoBackend::CompleteRequest(FChunkRequest* ChunkRequest)
 		Cache->Put(ChunkRequest->Params.ChunkKey, Chunk);
 	}
 
-	if (BackendStatus.ShouldAbandonCache() && InflightCacheRequestCount.load(std::memory_order_relaxed) == 0)
-	{
-		TickBackendEvent->Trigger();
-	}
-
 	ChunkRequests.Release(ChunkRequest);
 }
 
@@ -1200,12 +1033,7 @@ bool FOnDemandIoBackend::Resolve(FIoRequestImpl* Request)
 	{
 		if (bUpdatePriority)
 		{
-#if UE_ENABLE_IAS_REQUEST_THREAD
-			FOnDemandHttpThread::Get().ReprioritizeRequest(ChunkRequest->HttpRequest, ChunkRequest->Priority);
-#else
-			TRACE_CPUPROFILER_EVENT_SCOPE(IasBackend::UpdatePriorityForIoRequest);
-			HttpRequests.Reprioritize(ChunkRequest, ChunkRequest->Priority);
-#endif // UE_ENABLE_IAS_REQUEST_THREAD
+			HttpClient.ReprioritizeRequest(ChunkRequest->HttpRequest, ChunkRequest->Priority);
 		}
 		// The chunk for the request is already inflight 
 		return true;
@@ -1277,7 +1105,7 @@ void FOnDemandIoBackend::ResolveIoRequests(FIoRequestList Requests, FIoRequestLi
 
 void FOnDemandIoBackend::CancelIoRequest(FIoRequestImpl* Request)
 {
-	if (ChunkRequests.Cancel(Request, Cache.Get()))
+	if (ChunkRequests.Cancel(Request, HttpClient, Cache.Get()))
 	{
 		CompletedRequests.Enqueue(Request);
 		BackendContext->WakeUpDispatcherThreadDelegate.Execute();
@@ -1289,11 +1117,7 @@ void FOnDemandIoBackend::UpdatePriorityForIoRequest(FIoRequestImpl* Request)
 	TRACE_CPUPROFILER_EVENT_SCOPE(IasBackend::UpdatePriorityForIoRequest);
 	if (FChunkRequest* ChunkRequest = ChunkRequests.TryUpdatePriority(Request))
 	{
-#if UE_ENABLE_IAS_REQUEST_THREAD
-		FOnDemandHttpThread::Get().ReprioritizeRequest(ChunkRequest->HttpRequest, ChunkRequest->Priority);
-#else
-		HttpRequests.Reprioritize(ChunkRequest, ChunkRequest->Priority);
-#endif // UE_ENABLE_IAS_REQUEST_THREAD
+		HttpClient.ReprioritizeRequest(ChunkRequest->HttpRequest, ChunkRequest->Priority);
 	}
 }
 
@@ -1542,190 +1366,6 @@ FOnDemandStreamingCacheUsage FOnDemandIoBackend::GetCacheUsage() const
 	return Usage; 
 }
 
-void FOnDemandIoBackend::ProcessHttpRequests(FMultiEndpointHttpClient& HttpClient, int32 MaxConcurrentRequests)
-{
-	int32 NumConcurrentRequests = 0;
-	FChunkRequest* NextChunkRequest = HttpRequests.Dequeue();
-
-	// TODO: It would be better to only update connections as they need it, consider doing this on
-	// hostgroup connect/disconnect events.
-	HttpClient.UpdateConnections();
-
-	while (NextChunkRequest)
-	{
-		while (NextChunkRequest)
-		{
-			{
-				TRACE_CPUPROFILER_EVENT_SCOPE(IasBackend::IssueHttpGet);
-				FChunkRequest* ChunkRequest = NextChunkRequest;
-				NextChunkRequest = ChunkRequest->NextRequest;
-				ChunkRequest->NextRequest = nullptr;
-
-				Stats.OnHttpDequeue();
-
-				if (ChunkRequest->bCancelled)
-				{
-					CompleteRequest(ChunkRequest);
-					Stats.OnHttpCancel();
-				}
-				else if (!ChunkRequest->HostGroup.IsConnected() || !BackendStatus.IsHttpEnabled())
-				{
-					CompleteRequest(ChunkRequest);
-					// Technically this request is being skipped because of a pre-existing error. It is not
-					// an error itself and it is not being canceled by higher level code. However we do not
-					// currently have a statistic for that and we have to call one of the existing types in
-					// order to correctly reduce the pending count.
-					Stats.OnHttpCancel();
-				}
-				else
-				{
-					TAnsiStringBuilder<256> Url;
-					ChunkRequest->Params.GetUrl(Url);
-
-					NumConcurrentRequests++;
-
-					auto AnsiUrl = StringCast<ANSICHAR>(Url.GetData(), Url.Len());
-
-					HttpClient.Get(ChunkRequest->HostGroup.GetUnderlyingHostGroup(), AnsiUrl, ChunkRequest->Params.ChunkRange,
-						[this, ChunkRequest, &NumConcurrentRequests]
-						(FMultiEndpointHttpClientResponse&& HttpResponse)
-						{
-							TRACE_CPUPROFILER_EVENT_SCOPE(IasBackend::RequestCallback);
-							NumConcurrentRequests--;
-
-							if (HttpResponse.RetryCount > 0)
-							{
-								Stats.OnHttpRetry();
-							}
-
-							Stats.OnHttpCdnCacheReply(HttpResponse.CDNCacheStatus);
-
-							if (HttpResponse.IsOk())
-							{
-								ChunkRequest->HostGroup.OnSuccessfulResponse();
-
-								ChunkRequest->Chunk = MoveTemp(HttpResponse.Body);
-								Stats.OnHttpGet(ChunkRequest->Chunk.DataSize(), HttpResponse.Sample.RecvMs);
-
-								UE::Tasks::Launch(UE_SOURCE_LOCATION, [this, ChunkRequest]()
-									{
-										CompleteRequest(ChunkRequest);
-									}, GetRequestCompletionTaskPriority());
-							}
-							else
-							{
-								Stats.OnHttpError();
-
-								if (ChunkRequest->HostGroup.OnFailedResponse())
-								{
-									// A disconnect was triggered
-									Stats.OnHttpDisconnected();
-								}
-
-								CompleteRequest(ChunkRequest);
-							}
-						});
-				}
-			}
-
-			if (NumConcurrentRequests >= MaxConcurrentRequests)
-			{
-				TRACE_CPUPROFILER_EVENT_SCOPE(IasBackend::TickHttpSaturated);
-				while (NumConcurrentRequests >= MaxConcurrentRequests) //-V654
-				{
-					HttpClient.Tick(MAX_uint32, GIasHttpRateLimitKiBPerSecond);
-				}
-			}
-
-			if (!NextChunkRequest)
-			{
-				NextChunkRequest = HttpRequests.Dequeue();
-			}
-		}
-
-		{
-			// Keep processing pending connections until all requests are completed or a new one is issued
-			TRACE_CPUPROFILER_EVENT_SCOPE(IasBackend::TickHttp);
-			while (HttpClient.Tick(GIasHttpPollTimeoutMs, GIasHttpRateLimitKiBPerSecond))
-			{
-				if (!NextChunkRequest)
-				{
-					NextChunkRequest = HttpRequests.Dequeue();
-				}
-				if (NextChunkRequest)
-				{
-					break;
-				}
-			}
-		}
-	} 
-}
-
-void FOnDemandIoBackend::DrainHttpRequests()
-{
-	FChunkRequest* Iterator = HttpRequests.Dequeue();
-	while (Iterator != nullptr)
-	{
-		FChunkRequest* Request = Iterator;
-		Iterator = Iterator->NextRequest;
-
-		Stats.OnHttpDequeue();
-		CompleteRequest(Request);
-		Stats.OnHttpCancel();
-	}
-}
-
-uint32 FOnDemandIoBackend::Run()
-{
-	LLM_SCOPE_BYTAG(Ias);
-
-	if (!CreateDefaultHostGroup())
-	{
-		return 0;
-	}
-
-	TUniquePtr<FMultiEndpointHttpClient> HttpClient = FMultiEndpointHttpClient::Create(FMultiEndpointHttpClientConfig
-		{
-			.MaxConnectionCount = GIasHttpConnectionCount,
-			.ReceiveBufferSize = GIasHttpRecvBufKiB >= 0 ? GIasHttpRecvBufKiB << 10 : -1,
-			.MaxRetryCount = GIasHttpRetryCount,
-			.TimeoutMs = GIasHttpFailTimeOutMs,
-			.Redirects = EHttpRedirects::Disabled,
-			.bEnableThreadSafetyChecks = true,
-			.bAllowChunkedTransfer = GIasHttpAllowChunkedXfer,
-			.LogCategory = &LogIas,
-			.LogVerbosity = ELogVerbosity::VeryVerbose
-		});
-	check(HttpClient.IsValid());
-
-	BackendStatus.SetHttpEnabled(true);
-
-	while (!bStopRequested)
-	{
-		UpdateThreadPriorityIfNeeded();
-
-		ProcessHttpRequests(*HttpClient, FMath::Min(GIasHttpConcurrentRequests, 32));
-
-		if (!bStopRequested)
-		{
-			FHostGroupManager::Get().Tick(GIasHttpTimeOutMs, bStopRequested);
-
-			OnThreadTickIdle();
-
-			TickBackendEvent->Wait(FHostGroupManager::Get().GetNumDisconnctedHosts() > 0 ? GIasHttpHealthCheckWaitTime : MAX_uint32);
-		}
-	}
-
-	return 0;
-}
-
-void FOnDemandIoBackend::Stop()
-{
-	bStopRequested = true;
-	TickBackendEvent->Trigger();
-	DistributedEndpointEvent->Trigger();
-}
-
 int32 FOnDemandIoBackend::WaitForCompleteRequestTasks(float WaitTimeSeconds, float PollTimeSeconds)
 {
 	const double StartTime = FPlatformTime::Seconds();
@@ -1783,10 +1423,11 @@ void FOnDemandIoBackend::RunRequestTest(const TArray<FString>& Args) const
 TSharedPtr<IOnDemandIoDispatcherBackend> MakeOnDemandIoDispatcherBackend(
 	const FOnDemandEndpointConfig& Config,
 	FOnDemandIoStore& IoStore,
+	FOnDemandHttpThread& HttpClient,
 	TUniquePtr<IIasCache>&& Cache)
 {
 	return MakeShareable<IOnDemandIoDispatcherBackend>(
-		new FOnDemandIoBackend(Config, IoStore, MoveTemp(Cache)));
+		new FOnDemandIoBackend(Config, IoStore, HttpClient, MoveTemp(Cache)));
 }
 
 ///////////////////////////////////////////////////////////////////////////////
diff --git a/Engine/Source/Runtime/Experimental/IoStore/OnDemand/Private/OnDemandIoDispatcherBackend.h b/Engine/Source/Runtime/Experimental/IoStore/OnDemand/Private/OnDemandIoDispatcherBackend.h
--- a/Engine/Source/Runtime/Experimental/IoStore/OnDemand/Private/OnDemandIoDispatcherBackend.h
+++ b/Engine/Source/Runtime/Experimental/IoStore/OnDemand/Private/OnDemandIoDispatcherBackend.h
@@ -12,6 +12,7 @@ namespace UE::IoStore
 {
 
 class FOnDemandIoStore;
+class FOnDemandHttpThread;
 class IAnalyticsRecording;
 class IIasCache;
 struct FOnDemandStreamingCacheUsage;
@@ -88,6 +89,7 @@ public:
 TSharedPtr<IOnDemandIoDispatcherBackend> MakeOnDemandIoDispatcherBackend(
 	const FOnDemandEndpointConfig& Config,
 	FOnDemandIoStore& IoStore,
+	FOnDemandHttpThread& HttpClient,
 	TUniquePtr<IIasCache>&& Cache);
 
 } // namespace UE::IoStore
diff --git a/Engine/Source/Runtime/Experimental/IoStore/OnDemand/Private/OnDemandIoStore.cpp b/Engine/Source/Runtime/Experimental/IoStore/OnDemand/Private/OnDemandIoStore.cpp
--- a/Engine/Source/Runtime/Experimental/IoStore/OnDemand/Private/OnDemandIoStore.cpp
+++ b/Engine/Source/Runtime/Experimental/IoStore/OnDemand/Private/OnDemandIoStore.cpp
@@ -193,201 +193,6 @@ static TIoStatusOr<FSharedContainerHeader> DeserializeContainerHeader(const FOnD
 	return OutHeader; 
 }
 
-///////////////////////////////////////////////////////////////////////////////
-enum class EChunkCacheState : uint8
-{
-	None = 0,
-	Cached,
-	NotCached
-};
-
-struct FContainerInstallData
-{
-	TSet<FPackageId>				ResolvedPackageIds;
-	TMap<uint32, EChunkCacheState>	ResolvedChunks;
-	uint64							TotalSize = 0;
-};
-
-using FInstallData = TMap<FSharedOnDemandContainer, FContainerInstallData>;
-
-enum class EBuildInstallDataOptions
-{
-	None = 0,
-	IncludeSoftReferences = 1 << 0,
-};
-ENUM_CLASS_FLAGS(EBuildInstallDataOptions);
-
-static FIoStatus BuildInstallData(
-	const TSet<FSharedOnDemandContainer>& Containers,
-	const TSet<FPackageId>& PackageIds,
-	const EBuildInstallDataOptions Options,
-	FInstallData& OutInstallData,
-	TSet<FPackageId>& OutMissing)
-{
-	TRACE_CPUPROFILER_EVENT_SCOPE(FOnDemandIoStore::BuildInstallData);
-
-	OutInstallData.Reserve(Containers.Num());
-
-	// For now we always download these required chunks
-	for (const FSharedOnDemandContainer& Container : Containers)
-	{
-		FContainerInstallData& Data = OutInstallData.FindOrAdd(Container);
-
-		for (int32 EntryIndex = 0; const FIoChunkId& ChunkId : Container->ChunkIds)
-		{
-			switch(ChunkId.GetChunkType())
-			{
-				case EIoChunkType::ExternalFile:
-				case EIoChunkType::ShaderCodeLibrary:
-				case EIoChunkType::ShaderCode:
-				{
-					Data.ResolvedChunks.Add(EntryIndex, EChunkCacheState::None);
-					Data.TotalSize += Container->ChunkEntries[EntryIndex].EncodedSize;
-				}
-				default:
-					break;
-			}
-			++EntryIndex;
-		}
-	}
-
-	if (PackageIds.IsEmpty())
-	{
-		return EIoErrorCode::Ok;
-	}
-
-	// Traverse dependencies for each package ID
-	TArray<FPackageId> ResolvedPackageIds;
-	{
-		TRACE_CPUPROFILER_EVENT_SCOPE(FOnDemandIoStore::ResolveDependencies);
-
-		TRingBuffer<FPackageId> Queue;
-		TSet<FPackageId>		Visitied;
-
-		Visitied.Reserve(PackageIds.Num());
-		Queue.Reserve(PackageIds.Num());
-		for (const FPackageId& PackageId : PackageIds)
-		{
-			Queue.Add(PackageId);
-		}
-
-		FPackageStore& PackageStore = FPackageStore::Get();
-		FPackageStoreReadScope _(PackageStore);
-
-		while (!Queue.IsEmpty())
-		{
-			FPackageId PackageId = Queue.PopFrontValue();
-			{
-				FName		SourcePackageName;
-				FPackageId	RedirectedToPackageId;
-				if (PackageStore.GetPackageRedirectInfo(PackageId, SourcePackageName, RedirectedToPackageId))
-				{
-					PackageId = RedirectedToPackageId;
-				}
-			}
-
-			bool bIsAlreadyInSet = false;
-			Visitied.Add(PackageId, &bIsAlreadyInSet);
-			if (bIsAlreadyInSet)
-			{
-				continue;
-			}
-
-			FPackageStoreEntry PackageStoreEntry;
-			const EPackageStoreEntryStatus EntryStatus = PackageStore.GetPackageStoreEntry(PackageId, NAME_None, PackageStoreEntry);
-			if (EntryStatus != EPackageStoreEntryStatus::Missing)
-			{
-				ResolvedPackageIds.Add(PackageId);
-				for (const FPackageId& ImportedPackageId : PackageStoreEntry.ImportedPackageIds)
-				{
-					if (!Visitied.Contains(ImportedPackageId))
-					{
-						Queue.Add(ImportedPackageId);
-					}
-				}
-
-				if (EnumHasAnyFlags(Options, EBuildInstallDataOptions::IncludeSoftReferences))
-				{
-					TConstArrayView<FPackageId> SoftReferences;
-					TConstArrayView<uint32> Indices = PackageStore.GetSoftReferences(PackageId, SoftReferences);
-					for (uint32 Idx : Indices)
-					{
-						const FPackageId& SoftRef = SoftReferences[Idx];
-						if (!Visitied.Contains(SoftRef))
-						{
-							Queue.Add(SoftRef);
-						}
-					}
-				}
-			}
-			else
-			{
-				OutMissing.Add(PackageId);
-			}
-		}
-	}
-
-	auto FindChunkEntry = [&OutInstallData](const FIoChunkId& ChunkId, int32& OutIndex) -> TPair<FSharedOnDemandContainer, FContainerInstallData>*
-	{
-		for (TPair<FSharedOnDemandContainer, FContainerInstallData>& Kv : OutInstallData)
-		{
-			if (OutIndex = Kv.Key->FindChunkEntryIndex(ChunkId); OutIndex != INDEX_NONE)
-			{
-				return &Kv;
-			}
-		}
-
-		return nullptr; 
-	};
-
-	// Resolve all chunk ID's from the resolved package ID's
-	for (const FPackageId& PackageId : ResolvedPackageIds)
-	{
-		const FIoChunkId PackageChunkId								= CreatePackageDataChunkId(PackageId);
-		int32 EntryIndex											= INDEX_NONE; 
-		TPair<FSharedOnDemandContainer, FContainerInstallData>*	Kv	= FindChunkEntry(PackageChunkId, EntryIndex);
-
-		if (Kv == nullptr) 
-		{
-			// The chunk resides in a base game container
-			continue;
-		}
-
-		check(EntryIndex != INDEX_NONE);
-		FOnDemandContainer& Container		= *Kv->Key;
-		FContainerInstallData& InstallData	= Kv->Value;
-
-		InstallData.ResolvedPackageIds.Add(PackageId);
-		InstallData.ResolvedChunks.Add(EntryIndex, EChunkCacheState::None);
-		InstallData.TotalSize += Container.ChunkEntries[EntryIndex].EncodedSize;
-
-		// TODO: Installing optional bulkdata should probably be an install argument
-		const EIoChunkType AdditionalPackageChunkTypes[] =
-		{
-			EIoChunkType::BulkData,
-			EIoChunkType::OptionalBulkData,
-			EIoChunkType::MemoryMappedBulkData 
-		};
-
-		for (EIoChunkType ChunkType : AdditionalPackageChunkTypes)
-		{
-			// TODO: For Mutable we need to traverse all possible bulk data chunk indices?
-			const FIoChunkId ChunkId = CreateBulkDataIoChunkId(PackageId.Value(), 0, 0, ChunkType);
-			if (Kv = FindChunkEntry(ChunkId, EntryIndex); Kv != nullptr)
-			{
-				check(EntryIndex != INDEX_NONE);
-				FOnDemandContainer& OtherContainer		= *Kv->Key;
-				FContainerInstallData& OtherInstallData	= Kv->Value;
-
-				OtherInstallData.ResolvedChunks.Add(EntryIndex, EChunkCacheState::None);
-				OtherInstallData.TotalSize += OtherContainer.ChunkEntries[EntryIndex].EncodedSize;
-			}
-		}
-	}
-
-	return EIoErrorCode::Ok;
-}
-
 } // namespace UE::IoStore::Private
 
 ///////////////////////////////////////////////////////////////////////////////
@@ -407,13 +212,8 @@ FString FOnDemandContainer::UniqueName() const
 ///////////////////////////////////////////////////////////////////////////////
 FOnDemandIoStore::FOnDemandIoStore()
 {
-#if UE_ENABLE_REQUEST_THREAD
-	OnDemandHttpThread = MakeUnique<FOnDemandHttpThread>();
-#if UE_ENABLE_IAD_REQUEST_THREAD
-	Installer = MakeUnique<FOnDemandContentInstaller>(*this, FOnDemandHttpThread::Get());
-#endif
-#endif // UE_ENABLE_REQUEST_THREAD
-
+	HttpClient = MakeUnique<FOnDemandHttpThread>();
+	Installer = MakeUnique<FOnDemandContentInstaller>(*this, *HttpClient);
 	FEncryptionKeyManager::Get().OnKeyAdded().AddRaw(this, &FOnDemandIoStore::OnEncryptionKeyAdded);
 }
 
@@ -436,9 +236,8 @@ FOnDemandIoStore::~FOnDemandIoStore()
 		TickFuture.Wait();
 	}
 
-#if UE_ENABLE_REQUEST_THREAD
-	OnDemandHttpThread.Reset();
-#endif // UE_ENABLE_REQUEST_THREAD
+	Installer.Reset();
+	HttpClient.Reset();
 }
 
 FIoStatus FOnDemandIoStore::Initialize()
@@ -715,7 +514,7 @@ FIoStatus FOnDemandIoStore::InitializeStreamingBackend(const FOnDemandEndpointCo
 		UE_LOG(LogIas, Log, TEXT("File cache disabled - streaming only (zero-quota)"));
 	}
 
-	StreamingBackend = MakeOnDemandIoDispatcherBackend(EndpointConfig, *this, MoveTemp(StreamingCache));
+	StreamingBackend = MakeOnDemandIoDispatcherBackend(EndpointConfig, *this, *HttpClient, MoveTemp(StreamingCache));
 	if (StreamingBackend.IsValid() == false)
 	{
 		return FIoStatus(EIoErrorCode::InvalidParameter, TEXT("Failed to create streaming backend"));
@@ -758,28 +557,12 @@ FOnDemandInstallRequest FOnDemandIoStore::Install(
 {
 	LLM_SCOPE_BYNAME(TEXT("FileSystem/OnDemandIoStore"));
 	
-#if UE_ENABLE_IAD_REQUEST_THREAD
 	FSharedInternalInstallRequest InstallRequest = Installer->Enqueue(
 		MoveTemp(Args),
 		MoveTemp(OnCompleted),
 		MoveTemp(OnProgress));
 
 	return FOnDemandInstallRequest(AsWeak(), InstallRequest);
-#else
-	FSharedInternalInstallRequest InstallRequest = MakeShared<FOnDemandInternalInstallRequest, ESPMode::ThreadSafe>(
-			MoveTemp(Args),
-			MoveTemp(OnCompleted),
-			MoveTemp(OnProgress));
-
-	{
-		UE::TUniqueLock Lock(RequestMutex);
-		InstallRequests.Add(InstallRequest);
-	}
-
-	TryEnterTickLoop();
-
-	return FOnDemandInstallRequest(AsWeak(), InstallRequest);
-#endif
 }
 
 void FOnDemandIoStore::Purge(FOnDemandPurgeArgs&& Args, FOnDemandPurgeCompleted&& OnCompleted)
@@ -1232,34 +1015,6 @@ bool FOnDemandIoStore::Tick()
 		}
 	}
 
-	FSharedInternalInstallRequest InstallRequest;
-	{
-		UE::TUniqueLock Lock(RequestMutex);
-		if (InstallRequests.IsEmpty() == false)
-		{
-			InstallRequests.StableSort(FOnDemandInternalInstallRequest::SortByPriority);
-			InstallRequest = InstallRequests.Pop(EAllowShrinking::No);
-		}
-	}
-
-	// Tick install request(s)
-	if (InstallRequest.IsValid())
-	{
-		// Process the install request if the error code is unknown otherwise the request has been cancelled and completed
-		EIoErrorCode Expected = EIoErrorCode::Unknown;
-		if (InstallRequest->ErrorCode.compare_exchange_strong(Expected, EIoErrorCode::Ok))
-		{
-			FOnDemandInstallResult Result = TickInstallRequest(*InstallRequest);
-			CompleteInstallRequest(InstallRequest, MoveTemp(Result));
-		}
-		else
-		{
-			UE_LOG(LogIoStoreOnDemand, Log, TEXT("Skipping cancelled install request, %s"), *LexToString(*InstallRequest));
-		}
-
-		bTicked = true;
-	}
-
 	return bTicked;
 }
 
@@ -1584,421 +1339,6 @@ void FOnDemandIoStore::CompleteMountRequest(FMountRequest& Request, FOnDemandMou
 	}
 }
 
-FOnDemandInstallResult FOnDemandIoStore::TickInstallRequest(FOnDemandInternalInstallRequest& InstallRequest)
-{
-	TRACE_CPUPROFILER_EVENT_SCOPE(FOnDemandIoStore::InstallContent);
-
-	UE_LOG(LogIoStoreOnDemand, Log, TEXT("Ticking install request, ContentHandle='%s'"),
-		*LexToString(InstallRequest.Args.ContentHandle));
-
-	FOnDemandInstallResult OutResult;
-	FIoStatus& Status					= OutResult.Status;
-	FOnDemandInstallProgress& Progress	= OutResult.Progress;
-
-	const double StartTime = FPlatformTime::Seconds();
-	ON_SCOPE_EXIT
-	{
-		OutResult.DurationInSeconds = FPlatformTime::Seconds() - StartTime;
-	};
-
-	auto CheckAndSetCancelled = [&InstallRequest, &Status]() -> bool
-	{
-		if (InstallRequest.IsCancelRequested())
-		{
-			Status = FIoStatus(EIoErrorCode::Cancelled);
-			return true;
-		}
-
-		return false;
-	};
-
-	if (InstallRequest.Args.ContentHandle.IsValid() == false)
-	{
-		Status = FIoStatusBuilder(EIoErrorCode::InvalidParameter) << TEXT("Invalid content handle");
-		return OutResult;
-	}
-
-	if (InstallCache.IsValid() == false)
-	{
-		Status = FIoStatusBuilder(EIoErrorCode::InvalidParameter) << TEXT("Install cache not initialized");
-		return OutResult;
-	}
-
-	check(InstallCache.IsValid() && PackageStoreBackend.IsValid()); 
-
-	TAnsiStringBuilder<512> ChunkUrl;
-	TSet<FSharedOnDemandContainer> ContainersForInstallation;
-	TSet<FPackageId> PackageIdsToInstall;
-
-#if !UE_BUILD_SHIPPING
-	if (!GIoStoreBreakOnRequest.IsEmpty())
-	{
-		if (FCString::Strstr(*InstallRequest.Args.DebugName, *GIoStoreBreakOnRequest))
-		{
-			UE_DEBUG_BREAK();
-		}
-	}
-#endif
-
-	if (Status = GetContainersAndPackagesForInstall(
-		InstallRequest.Args.MountId,
-		InstallRequest.Args.TagSets,
-		InstallRequest.Args.PackageIds,
-		ContainersForInstallation,
-		PackageIdsToInstall); !Status.IsOk())
-	{
-		return OutResult;
-	}
-
-	// Its OK for PackageIdsToInstall to be empty at this point. Any chunks not referenced by a package must still be installed.
-
-	// Find all I/O chunks fo the specified list of packages
-	Private::EBuildInstallDataOptions BuildInstallDataOptions = Private::EBuildInstallDataOptions::None;
-	if (EnumHasAnyFlags(InstallRequest.Args.Options, EOnDemandInstallOptions::InstallSoftReferences))
-	{
-		BuildInstallDataOptions |= Private::EBuildInstallDataOptions::IncludeSoftReferences;
-	}
-
-	Private::FInstallData InstallData;
-	TSet<FPackageId> Missing;
-
-	Status = Private::BuildInstallData(ContainersForInstallation, PackageIdsToInstall, BuildInstallDataOptions, InstallData, Missing);
-	if (Status.IsOk() == false)
-	{
-		return OutResult;
-	}
-
-	// Check the other I/O backends for missing package chunks
-	uint32 MissingCount = 0;
-	for (const FPackageId& PackageId : Missing)
-	{
-		const FIoChunkId ChunkId = CreatePackageDataChunkId(PackageId);
-		if (FIoDispatcher::Get().DoesChunkExist(ChunkId) == false)
-		{
-			UE_CLOG(MissingCount == 0, LogIoStoreOnDemand, Warning, TEXT("Failed to resolve the following chunk(s) for content handle '%s':"),
-				*LexToString(InstallRequest.Args.ContentHandle));
-
-			UE_LOG(LogIoStoreOnDemand, Warning, TEXT("ChunkId='%s'"), *LexToString(ChunkId));
-			MissingCount++;
-		}
-	}
-
-	if (MissingCount > 0 && GIoStoreOnDemandTreatMissingChunksAsError)
-	{
-		UE_LOG(LogIoStoreOnDemand, Error, TEXT("Install request failed to due missing chunk(s), ContentHandle='%s'"),
-			*LexToString(InstallRequest.Args.ContentHandle));
-		Status = FIoStatus(EIoErrorCode::UnknownChunkID);
-		return OutResult;
-	}
-
-	if (CheckAndSetCancelled())
-	{
-		return OutResult;
-	}
-
-	// Purge
-	{
-		TRACE_CPUPROFILER_EVENT_SCOPE(FOnDemandIoStore::PurgeInstallCache);
-
-		TSet<FIoHash> ChunksToInstall;
-		for (const auto& Kv : InstallData)
-		{
-			FOnDemandContainer&						Container = *Kv.Key;
-			const Private::FContainerInstallData&	Data = Kv.Value;
-
-			for (const TPair<uint32, Private::EChunkCacheState>& ChunkPair : Data.ResolvedChunks)
-			{
-				const int32 EntryIndex = ChunkPair.Key;
-
-				const FOnDemandChunkEntry& Entry = Container.ChunkEntries[EntryIndex];
-				ChunksToInstall.Add(Entry.Hash);
-			}
-		}
-
-		if (Status = InstallCache->Purge(MoveTemp(ChunksToInstall)); Status.IsOk() == false)
-		{
-			return OutResult;
-		}
-	}
-
-	if (CheckAndSetCancelled())
-	{
-		return OutResult;
-	}
-
-	// Download all chunks
-	const int32 MaxConcurrentRequests = 16;
-	int32 ConcurrentRequests = 0;
-
-	FMultiEndpointHttpClientConfig HttpConfig;
-	HttpConfig.MaxConnectionCount = 8;
-	HttpConfig.MaxRetryCount = 2;
-	HttpConfig.LogCategory = &LogIoStoreOnDemand;
-
-	TUniquePtr<FMultiEndpointHttpClient> HttpClient = FMultiEndpointHttpClient::Create(MoveTemp(HttpConfig));
-	if (HttpClient.IsValid() == false)
-	{
-		Status = FIoStatusBuilder(EIoErrorCode::InvalidCode) << TEXT("Failed to initialize HTTP client");
-		return OutResult;
-	}
-
-	uint64 TotalChunkCount		= 0;
-	uint64 TotalBytes			= 0;
-	uint64 ToDownloadChunkCount	= 0;
-	uint64 ToDownloadBytes		= 0;
-	
-	for (Private::FInstallData::ElementType& Kv : InstallData)
-	{
-		const FOnDemandContainer& Container		= *Kv.Key;
-		Private::FContainerInstallData& Data	= Kv.Value;
-
-		for (TPair<uint32, Private::EChunkCacheState>& ChunkPair : Data.ResolvedChunks)
-		{
-			const int32 EntryIndex						= ChunkPair.Key;
-			Private::EChunkCacheState& ChunkCacheState	= ChunkPair.Value;
-
-			const FIoChunkId& ChunkId = Container.ChunkIds[EntryIndex];
-			const FOnDemandChunkEntry& ChunkEntry = Container.ChunkEntries[EntryIndex];
-
-			++TotalChunkCount;
-			TotalBytes += ChunkEntry.EncodedSize;
-
-			if (InstallCache->IsChunkCached(ChunkEntry.Hash))
-			{
-				ChunkCacheState = Private::EChunkCacheState::Cached;
-				continue;
-			}
-
-			ChunkCacheState = Private::EChunkCacheState::NotCached;
-
-			++ToDownloadChunkCount;
-			ToDownloadBytes += ChunkEntry.EncodedSize;
-		}
-	}
-
-	Progress.TotalContentSize = TotalBytes;
-	Progress.TotalInstallSize = ToDownloadBytes;
-
-	ProgressInstallRequest(InstallRequest, Progress);
-
-	{
-		TRACE_CPUPROFILER_EVENT_SCOPE(FOnDemandIoStore::InstallContent::HttpRequests);
-		for (const Private::FInstallData::ElementType& Kv : InstallData)
-		{
-			const FOnDemandContainer& Container			= *Kv.Key;
-			const Private::FContainerInstallData& Data	= Kv.Value;
-
-			for (const TPair<uint32, Private::EChunkCacheState>& ChunkPair : Data.ResolvedChunks)
-			{
-				const int32 EntryIndex							= ChunkPair.Key;
-				const Private::EChunkCacheState ChunkCacheState	= ChunkPair.Value;
-
-				const FIoChunkId& ChunkId						= Container.ChunkIds[EntryIndex];
-				const FOnDemandChunkEntry& ChunkEntry			= Container.ChunkEntries[EntryIndex];
-
-				if (ChunkCacheState == Private::EChunkCacheState::Cached)
-				{
-					continue;
-				}
-
-				++ConcurrentRequests;
-
-				HttpClient->Get(
-					Container.HostGroup.GetUnderlyingHostGroup(),
-					Private::GetChunkUrl(Container, ChunkEntry, ChunkUrl).ToView(),
-					[this, &InstallRequest, &OutResult, ChunkId, ChunkEntry, &ConcurrentRequests]
-					(FMultiEndpointHttpClientResponse&& HttpResponse)
-					{
-						ConcurrentRequests--;
-
-						FIoStatus& Status					= OutResult.Status;
-						FOnDemandInstallProgress& Progress	= OutResult.Progress;
-
-						if (!Status.IsOk())
-						{
-							// a different request already failed, don't stomp the status if this one succeeded
-							return;
-						}
-						
-						if (!HttpResponse.IsOk())
-						{
-							Status = FIoStatus(EIoErrorCode::ReadError) << 
-								FString::Printf(TEXT("HTTP Status %u, "), HttpResponse.StatusCode) << 
-								HttpResponse.Reason;
-							return;
-						}
-
-						FIoBuffer& Chunk		= HttpResponse.Body;
-						const FIoHash ChunkHash	= FIoHash::HashBuffer(Chunk.GetView());
-						if (ChunkHash != ChunkEntry.Hash)
-						{
-							Status = FIoStatus(EIoErrorCode::ReadError, TEXTVIEW("Hash mismatch"));
-							return;
-						}
-						Status = InstallCache->PutChunk(MoveTemp(Chunk), ChunkHash);
-						if (!Status.IsOk())
-						{
-							return;
-						}
-
-						// TODO: Is this good enough progress or should it be fine grained bytes from the HttpClient?
-						Progress.CurrentInstallSize += ChunkEntry.EncodedSize;
-						ProgressInstallRequest(InstallRequest, Progress);
-					});
-
-				while (ConcurrentRequests >= MaxConcurrentRequests)
-				{
-					HttpClient->Tick();
-				}
-
-				if ((Status.IsOk() == false) || CheckAndSetCancelled())
-				{
-					// TODO: Should there be a way to flush/cancel the HTTPClient?
-					while (HttpClient->Tick())
-						;
-
-					return OutResult;
-				}
-			}
-		}
-
-		while (HttpClient->Tick())
-			;
-	}
-
-	if (Status.IsOk() == false)
-	{
-		return OutResult;
-	}
-	
-	if (Status = InstallCache->Flush(); Status.IsOk() == false)
-	{
-		return OutResult;
-	}
-
-	if (CheckAndSetCancelled())
-	{
-		return OutResult;
-	}
-
-	const TSharedPtr<FOnDemandInternalContentHandle>& ContentHandle = InstallRequest.Args.ContentHandle.Handle;
-	if (ContentHandle->IoStore.IsValid() == false)
-	{
-		// First time this content handle is used
-		ContentHandle->IoStore = AsWeak(); 
-	}
-
-	for (const auto& Kv : InstallData)
-	{
-		FOnDemandContainer& Container				= *Kv.Key;
-		const Private::FContainerInstallData& Data	= Kv.Value;
-
-		{
-			UE::TUniqueLock Lock(ContainerMutex);
-			FOnDemandChunkEntryReferences& References = Container.FindOrAddChunkEntryReferences(*ContentHandle);
-			for (const TPair<uint32, Private::EChunkCacheState>& ChunkPair : Data.ResolvedChunks)
-			{
-				const int32 EntryIndex = ChunkPair.Key;
-				References.Indices[EntryIndex] = true;
-			}
-		}
-	}
-
-	if (InstallData.IsEmpty() == false)
-	{
-		// Referencers were added so the package store must be refreshed
-		PackageStoreBackend->NeedsUpdate(EOnDemandPackageStoreUpdateMode::ReferencedPackages);
-	}
-
-	double DurationInSeconds = FPlatformTime::Seconds() - StartTime;
-	#if !UE_BUILD_SHIPPING
-	if (GIoStoreOnDemandRestrictKibsPerSecond > 0.0f)
-	{
-		
-		float MinSleepTime = ((float)(TotalBytes)) / (GIoStoreOnDemandRestrictKibsPerSecond * 1024.0f);
-		float SleepTime = MinSleepTime - (float)(DurationInSeconds);
-		if (SleepTime > 0.0f)
-		{
-			FPlatformProcess::SleepNoStats(SleepTime);
-			DurationInSeconds = FPlatformTime::Seconds() - StartTime;
-		}
-	}
-	if (GIoStoreOnDemandMinInstallDurationSeconds > 0.0f)
-	{
-		if (float SleepTime = GIoStoreOnDemandMinInstallDurationSeconds - float(DurationInSeconds); SleepTime > 0.0f)
-		{
-			FPlatformProcess::SleepNoStats(SleepTime);
-			DurationInSeconds = FPlatformTime::Seconds() - StartTime;
-		}
-	}
-	#endif
-	UE_LOG(LogIoStoreOnDemand, Log, TEXT("Downloaded %llu (%.2lf MiB) of total %llu (%.2lf MiB) requested chunk(s), Duration=%.2lfs, ContentHandle='%s'"),
-		ToDownloadChunkCount, double(ToDownloadBytes) / 1024.0 / 1024.0, TotalChunkCount, double(TotalBytes) / 1024.0 / 1024.0, DurationInSeconds, *LexToString(*ContentHandle));
-
-	return OutResult;
-}
-
-void FOnDemandIoStore::CompleteInstallRequest(FSharedInternalInstallRequest Request, FOnDemandInstallResult&& InstallResult)
-{
-	FOnDemandInstallRequest::EStatus Status = FOnDemandInstallRequest::EStatus::None;
-	switch (InstallResult.Status.GetErrorCode())
-	{
-		case EIoErrorCode::Ok:
-			Status = FOnDemandInstallRequest::Ok;
-			break;
-		case EIoErrorCode::Cancelled:
-			Status = FOnDemandInstallRequest::Cancelled;
-			break;
-		default:
-			Status = FOnDemandInstallRequest::Error;
-			break;
-	}
-
-	if (!Request->OnCompleted)
-	{
-		Request->Status.store(Status);
-	}
-	else if (EnumHasAnyFlags(Request->Args.Options, EOnDemandInstallOptions::CallbackOnGameThread))
-	{
-		ExecuteOnGameThread(
-			UE_SOURCE_LOCATION,
-			[Request, InstallResult = MoveTemp(InstallResult), Status]() mutable
-			{
-				FOnDemandInstallCompleted OnCompleted = MoveTemp(Request->OnCompleted);
-				OnCompleted(MoveTemp(InstallResult));
-				Request->Status.store(Status);
-			});
-	}
-	else
-	{
-		FOnDemandInstallCompleted OnCompleted = MoveTemp(Request->OnCompleted);
-		OnCompleted(MoveTemp(InstallResult));
-		Request->Status.store(Status);
-	}
-}
-
-void FOnDemandIoStore::ProgressInstallRequest(const FOnDemandInternalInstallRequest& Request, const FOnDemandInstallProgress& Progress)
-{
-	if (!Request.OnProgress)
-	{
-		return;
-	}
-
-	if (EnumHasAnyFlags(Request.Args.Options, EOnDemandInstallOptions::CallbackOnGameThread))
-	{
-		ExecuteOnGameThread(
-			UE_SOURCE_LOCATION,
-			[&Request, Progress]()
-			{
-				Request.OnProgress(Progress);
-			});
-	}
-	else
-	{
-		Request.OnProgress(Progress);
-	}
-}
-
 void FOnDemandIoStore::CompletePurgeRequest(FPurgeRequest& Request, FOnDemandPurgeResult&& Result)
 {
 	if (!Request.OnCompleted)
@@ -2432,35 +1772,12 @@ void FOnDemandIoStore::GetReferencedContentByHandle(TMap<FOnDemandWeakContentHan
 void FOnDemandIoStore::CancelInstallRequest(FSharedInternalInstallRequest InstallRequest)
 {
 	check(InstallRequest.IsValid());
-
-#if UE_ENABLE_IAD_REQUEST_THREAD
 	Installer->Cancel(InstallRequest);
-#else
-	// Cancel the request if the error code is unknown otherwise the request will be completed in the tick loop 
-	EIoErrorCode Expected = EIoErrorCode::Unknown;
-	if (InstallRequest->ErrorCode.compare_exchange_strong(Expected, EIoErrorCode::Cancelled))
-	{
-		UE_LOG(LogIoStoreOnDemand, Log, TEXT("Cancelling install request, %s"), *LexToString(*InstallRequest));
-		FOnDemandInstallResult InstallResult
-		{
-			.Status				= FIoStatus(EIoErrorCode::Cancelled),
-			.DurationInSeconds 	= FPlatformTime::Seconds() - InstallRequest->StartTime
-		};
-		CompleteInstallRequest(InstallRequest, MoveTemp(InstallResult));
-	}
-#endif
 }
 
 void FOnDemandIoStore::UpdateInstallRequestPriority(FSharedInternalInstallRequest InstallRequest, int32 NewPriority)
 {
-#if UE_ENABLE_IAD_REQUEST_THREAD
 	Installer->UpdatePriority(InstallRequest, NewPriority);
-#else
-	{
-		UE::TUniqueLock Lock(RequestMutex);
-		InstallRequest->Args.Priority = NewPriority;
-	}
-#endif
 }
 
 FOnDemandCacheUsage FOnDemandIoStore::GetCacheUsage(const FOnDemandGetCacheUsageArgs& Args) const
diff --git a/Engine/Source/Runtime/Experimental/IoStore/OnDemand/Private/OnDemandIoStore.h b/Engine/Source/Runtime/Experimental/IoStore/OnDemand/Private/OnDemandIoStore.h
--- a/Engine/Source/Runtime/Experimental/IoStore/OnDemand/Private/OnDemandIoStore.h
+++ b/Engine/Source/Runtime/Experimental/IoStore/OnDemand/Private/OnDemandIoStore.h
@@ -379,9 +379,8 @@ private:
 											const TArray<FString>& ExistingContainers,
 											TArray<FSharedOnDemandContainer>& Out);
 
+	TUniquePtr<FOnDemandHttpThread>			HttpClient;
 	TUniquePtr<FOnDemandContentInstaller>	Installer;
-	TUniquePtr<FOnDemandHttpThread>				OnDemandHttpThread;
-
 	FSharedInstallCache						InstallCache;
 	FSharedPackageStoreBackend				PackageStoreBackend;
 	FSharedStreamingBackend					StreamingBackend;
@@ -392,7 +391,6 @@ private:
 	mutable UE::FMutex						ContainerMutex;
 
 	TArray<FSharedMountRequest>				MountRequests;
-	TArray<FSharedInternalInstallRequest>	InstallRequests;
 	TArray<FSharedPurgeRequest>				PurgeRequests;
 	TArray<FSharedDefragRequest>			DefragRequests;
 	UE::FMutex								RequestMutex;
