Split IAS and IAD http analytics into two different data sets.

#rnx

### Misc
- Add a new enum EHttpRequestType to help classify the type of request being made which is now required when making a call to FOnDemandHttpThread::IssueRequest.
- Apply the same enum to the http methods of FOnDemandIoBackendStats.
-- Add additional code documentation for the http methods.

### Statistics
- All http statistics are now arrays using EHttpRequestType::NUM_SOURCES as the number of array elements so we can track the different types independently.
-- To make this easier to do we now have our own counter macros UE_IAX_COUNTER and UE_IAX_MEMORY_COUNTER which hides this.
-- We also need our own csv macros UE_IAX_CSV_DEFINE_STAT and UE_IAX_CSV_CUSTOM_STAT_DEFINED as well as creating an IAD csv category.
-- Note that the big problem for both of these is the requirement to provide a full static string for the name at compile time, so we can't just make a FHttpStatistics struct and have one per request type.
- FAnalyticsRecording doesn't really work with the new way we record analytics so for now I added a new define UE_ENABLE_ANALYTICS_RECORDING to enable/disable the feature and left it as disabled by default. The code will compile but it's only able to record IAS statistics at the moment.
- Changed the Ias/HttpDurationMs counter to record the duration of each Get rather than a rolling average, which didn't make much sense.
- The rolling average code (which we still use for display and logging purposes) has been grouped under FHttpRecentHistoryStatistics.
- Pretty much every place of code that was using the http stats is now abstracted to a lambda that takes a EHttpRequestType so we can write out the IAS and IAD stats and easily add another request type in the future if needed.
- Moved the IAS cache statistics to their own line in the display debug text. They no longer share a line with the IAS http stats.
- Replaced use of Add(1) and Add(-1) with counters to increment/decrement for visibility.

diff --git a/Engine/Source/Runtime/Experimental/IoStore/OnDemand/Private/OnDemandContentInstaller.cpp b/Engine/Source/Runtime/Experimental/IoStore/OnDemand/Private/OnDemandContentInstaller.cpp
--- a/Engine/Source/Runtime/Experimental/IoStore/OnDemand/Private/OnDemandContentInstaller.cpp
+++ b/Engine/Source/Runtime/Experimental/IoStore/OnDemand/Private/OnDemandContentInstaller.cpp
@@ -602,7 +602,7 @@ void FOnDemandContentInstaller::StartDownloadChunks(FInstallRequestState& Reques
 							MoveTemp(Chunk));
 					},
 					UE::Tasks::ETaskPriority::BackgroundLow);
-			});
+			}, EHttpRequestType::Installed);
 	}
 	
 	RequestState.bHttpRequestsIssued = true;
diff --git a/Engine/Source/Runtime/Experimental/IoStore/OnDemand/Private/OnDemandHttpThread.cpp b/Engine/Source/Runtime/Experimental/IoStore/OnDemand/Private/OnDemandHttpThread.cpp
--- a/Engine/Source/Runtime/Experimental/IoStore/OnDemand/Private/OnDemandHttpThread.cpp
+++ b/Engine/Source/Runtime/Experimental/IoStore/OnDemand/Private/OnDemandHttpThread.cpp
@@ -169,6 +169,7 @@ struct FHttpRequest
 	FOnDemandHttpThread::FCompletionCallback CompletionCallback;
 	FOnDemandChunkInfo ChunkInfo;
 	FIoOffsetAndLength ChunkRange;
+	EHttpRequestType Type;
 	bool bCancelled = false;
 	int32 Priority = 0;
 };
@@ -200,12 +201,12 @@ FOnDemandHttpThread::~FOnDemandHttpThread()
 	DrainHttpRequests();
 }
 
-FOnDemandHttpThread::FRequestHandle FOnDemandHttpThread::IssueRequest(const FOnDemandChunkInfo& ChunkInfo, const FIoOffsetAndLength& ReadRange, int32 Priority, FCompletionCallback&& CompletionCallback)
+FOnDemandHttpThread::FRequestHandle FOnDemandHttpThread::IssueRequest(const FOnDemandChunkInfo& ChunkInfo, const FIoOffsetAndLength& ReadRange, int32 Priority, FCompletionCallback&& CompletionCallback, EHttpRequestType Type)
 {
-	return IssueRequest(FOnDemandChunkInfo(ChunkInfo), ReadRange, Priority, MoveTemp(CompletionCallback));
+	return IssueRequest(FOnDemandChunkInfo(ChunkInfo), ReadRange, Priority, MoveTemp(CompletionCallback), Type);
 }
 
-FOnDemandHttpThread::FRequestHandle FOnDemandHttpThread::IssueRequest(FOnDemandChunkInfo&& ChunkInfo, const FIoOffsetAndLength& ReadRange, int32 Priority, FCompletionCallback&& CompletionCallback)
+FOnDemandHttpThread::FRequestHandle FOnDemandHttpThread::IssueRequest(FOnDemandChunkInfo&& ChunkInfo, const FIoOffsetAndLength& ReadRange, int32 Priority, FCompletionCallback&& CompletionCallback, EHttpRequestType Type)
 {
 	TRACE_CPUPROFILER_EVENT_SCOPE(FOnDemandHttpThread::IssueRequest);
 
@@ -213,9 +214,10 @@ FOnDemandHttpThread::FRequestHandle FOnDemandHttpThread::IssueRequest(FOnDemandC
 
 	Request->ChunkRange = ReadRange;
 	Request->CompletionCallback = MoveTemp(CompletionCallback);
+	Request->Type = Type;
 	Request->Priority = Priority;
 
-	FOnDemandIoBackendStats::Get()->OnHttpEnqueue();
+	FOnDemandIoBackendStats::Get()->OnHttpEnqueue(Request->Type);
 
 	HttpRequests.EnqueueByPriority(Request);
 #if UE_IOSTORE_ONDEMAND_NO_HTTP_THREAD 
@@ -317,12 +319,12 @@ void FOnDemandHttpThread::TickRequests()
 				NextHttpRequest = HttpRequest->NextRequest;
 				HttpRequest->NextRequest = nullptr;
 
-				FOnDemandIoBackendStats::Get()->OnHttpDequeue();
+				FOnDemandIoBackendStats::Get()->OnHttpDequeue(HttpRequest->Type);
 
 				if (HttpRequest->bCancelled)
 				{
 					HttpRequest->OnRequestCompleted(0, TEXTVIEW("Request cancelled"), FIoBuffer());
-					FOnDemandIoBackendStats::Get()->OnHttpCancel();
+					FOnDemandIoBackendStats::Get()->OnHttpCancel(HttpRequest->Type);
 
 					DestroyRequest(HttpRequest);
 				}
@@ -334,7 +336,7 @@ void FOnDemandHttpThread::TickRequests()
 					// an error itself and it is not being canceled by higher level code. However we do not
 					// currently have a statistic for that and we have to call one of the existing types in
 					// order to correctly reduce the pending count.
-					FOnDemandIoBackendStats::Get()->OnHttpCancel();
+					FOnDemandIoBackendStats::Get()->OnHttpCancel(HttpRequest->Type);
 
 					DestroyRequest(HttpRequest);
 				}
@@ -355,23 +357,23 @@ void FOnDemandHttpThread::TickRequests()
 
 							if (HttpResponse.RetryCount > 0)
 							{
-								FOnDemandIoBackendStats::Get()->OnHttpRetry();
+								FOnDemandIoBackendStats::Get()->OnHttpRetry(HttpRequest->Type);
 							}
 
-							FOnDemandIoBackendStats::Get()->OnHttpCdnCacheReply(HttpResponse.CDNCacheStatus);
+							FOnDemandIoBackendStats::Get()->OnHttpCdnCacheReply(HttpRequest->Type, HttpResponse.CDNCacheStatus);
 
 							if (HttpResponse.IsOk())
 							{
 								HttpRequest->HostGroup().OnSuccessfulResponse();
 
-								FOnDemandIoBackendStats::Get()->OnHttpGet(HttpResponse.Body.DataSize(), HttpResponse.DurationMilliseconds);
+								FOnDemandIoBackendStats::Get()->OnHttpGet(HttpRequest->Type, HttpResponse.Body.DataSize(), HttpResponse.DurationMilliseconds);
 
 								HttpRequest->OnRequestCompleted(HttpResponse.StatusCode, HttpResponse.Reason, MoveTemp(HttpResponse.Body));
 								DestroyRequest(HttpRequest);
 							}
 							else
 							{
-								FOnDemandIoBackendStats::Get()->OnHttpError();
+								FOnDemandIoBackendStats::Get()->OnHttpError(HttpRequest->Type);
 
 								if (HttpRequest->HostGroup().OnFailedResponse())
 								{
@@ -428,10 +430,10 @@ void FOnDemandHttpThread::DrainHttpRequests()
 		FHttpRequest* Request = Iterator;
 		Iterator = Iterator->NextRequest;
 
-		FOnDemandIoBackendStats::Get()->OnHttpDequeue();
+		FOnDemandIoBackendStats::Get()->OnHttpDequeue(Request->Type);
 		Request->OnRequestCompleted(0, TEXTVIEW("Request cancelled due to shutdown"), FIoBuffer());
 		DestroyRequest(Request);
-		FOnDemandIoBackendStats::Get()->OnHttpCancel();
+		FOnDemandIoBackendStats::Get()->OnHttpCancel(Request->Type);
 	}
 }
 
diff --git a/Engine/Source/Runtime/Experimental/IoStore/OnDemand/Private/OnDemandHttpThread.h b/Engine/Source/Runtime/Experimental/IoStore/OnDemand/Private/OnDemandHttpThread.h
--- a/Engine/Source/Runtime/Experimental/IoStore/OnDemand/Private/OnDemandHttpThread.h
+++ b/Engine/Source/Runtime/Experimental/IoStore/OnDemand/Private/OnDemandHttpThread.h
@@ -24,6 +24,15 @@ class FMultiEndpointHttpClient;
 struct FHttpRequest;
 struct FOnDemandChunkInfo;
 
+/** Describes the system making the request */
+enum class EHttpRequestType : uint8
+{
+	Streaming = 0,	/** IAS */
+	Installed,		/** IAD */
+
+	NUM_SOURCES
+};
+
 class FOnDemandHttpThread final
 	: public FRunnable
 	, public FSingleThreadRunnable
@@ -50,7 +59,7 @@ public:
 	 * @return				A handle representing the request which can be used to modify it while the request is still in
 	 *						flight. The handle will no longer be valid after the CompletionCallback has been invoked.
 	 */
-	FRequestHandle IssueRequest(const FOnDemandChunkInfo& ChunkInfo, const FIoOffsetAndLength& ReadRange, int32 Priority, FCompletionCallback&& CompletionCallback);
+	FRequestHandle IssueRequest(const FOnDemandChunkInfo& ChunkInfo, const FIoOffsetAndLength& ReadRange, int32 Priority, FCompletionCallback&& CompletionCallback, EHttpRequestType Type);
 	
 	/**
 	 * Issue a http request to read from a chunk.
@@ -64,7 +73,7 @@ public:
 	 * @return				A handle representing the request which can be used to modify it while the request is still in
 	 *						flight. The handle will no longer be valid after the CompletionCallback has been invoked.
 	 */
-	FRequestHandle IssueRequest(FOnDemandChunkInfo&& ChunkInfo, const FIoOffsetAndLength& ReadRange, int32 Priority, FCompletionCallback&& CompletionCallback);
+	FRequestHandle IssueRequest(FOnDemandChunkInfo&& ChunkInfo, const FIoOffsetAndLength& ReadRange, int32 Priority, FCompletionCallback&& CompletionCallback, EHttpRequestType Type);
 
 	void ReprioritizeRequest(FRequestHandle Request, int32 NewPriority);
 	void CancelRequest(FRequestHandle Request);
diff --git a/Engine/Source/Runtime/Experimental/IoStore/OnDemand/Private/OnDemandIoDispatcherBackend.cpp b/Engine/Source/Runtime/Experimental/IoStore/OnDemand/Private/OnDemandIoDispatcherBackend.cpp
--- a/Engine/Source/Runtime/Experimental/IoStore/OnDemand/Private/OnDemandIoDispatcherBackend.cpp
+++ b/Engine/Source/Runtime/Experimental/IoStore/OnDemand/Private/OnDemandIoDispatcherBackend.cpp
@@ -853,7 +853,7 @@ void FOnDemandIoBackend::IssueHttpRequest(FChunkRequest* ChunkRequest)
 			{
 				CompleteRequest(ChunkRequest);
 			}
-		});
+		}, EHttpRequestType::Streaming);
 
 	check(ChunkRequest->HttpRequest != nullptr);
 }
@@ -911,7 +911,7 @@ void FOnDemandIoBackend::CompleteRequest(FChunkRequest* ChunkRequest)
 				else
 				{
 					// Currently not being cached implies that the request was made via http
-					Stats.OnHttpDecodeError();
+					Stats.OnHttpDecodeError(EHttpRequestType::Streaming);
 				}
 			}
 		}
diff --git a/Engine/Source/Runtime/Experimental/IoStore/OnDemand/Private/Statistics.cpp b/Engine/Source/Runtime/Experimental/IoStore/OnDemand/Private/Statistics.cpp
--- a/Engine/Source/Runtime/Experimental/IoStore/OnDemand/Private/Statistics.cpp
+++ b/Engine/Source/Runtime/Experimental/IoStore/OnDemand/Private/Statistics.cpp
@@ -12,12 +12,14 @@
 #include "Internationalization/Internationalization.h"
 #include "Math/UnrealMathUtility.h"
 #include "Misc/CoreDelegates.h"
+#include "OnDemandHttpThread.h"
 #include "Templates/Requires.h"
 
 LLM_DEFINE_TAG(Ias);
 
 #if IAS_WITH_STATISTICS
 
+#define UE_ENABLE_ANALYTICS_RECORDING 0
 #define UE_ENABLE_ONSCREEN_STATISTICS !UE_BUILD_SHIPPING
 
 namespace UE::IoStore
@@ -246,13 +248,18 @@ private:
 		TCounterInt(...)  {}
 		void Set(int64 i) { V = i; }
 		void Add(int64 d) { V += d; }
+		void Increment() { V++; }
+		void Decrement() { --V; }
 		int64 Get() const { return V;}
 		Type V = 0;
 	};
 	using FCounterInt		= TCounterInt<int64>;
 	using FCounterAtomicInt = TCounterInt<std::atomic<int64>>;
-#endif 
+#endif
 
+#define UE_IAX_COUNTER(Name, Type) Type G##Name[(uint8)EHttpRequestType::NUM_SOURCES] = {{TEXT(UE_STRINGIZE(Ias/Name)), TraceCounterDisplayHint_None}, {TEXT(UE_STRINGIZE(Iad/Name)), TraceCounterDisplayHint_None}};
+#define UE_IAX_MEMORY_COUNTER(Name, Type) Type G##Name[(uint8)EHttpRequestType::NUM_SOURCES] = {{TEXT(UE_STRINGIZE(Ias/Name)), TraceCounterDisplayHint_None}, {TEXT(UE_STRINGIZE(Iad/Name)), TraceCounterDisplayHint_Memory}};
+	
 // iorequest stats
 FCounterInt				GIoRequestCount(TEXT("Ias/IoRequestCount"), TraceCounterDisplayHint_None);
 FCounterAtomicInt		GIoRequestReadCount(TEXT("Ias/IoRequestReadCount"), TraceCounterDisplayHint_None);
@@ -272,36 +279,90 @@ int64					GCacheMaxBytes = 0;
 FCounterAtomicInt		GCachePendingBytes(TEXT("Ias/CachePendingBytes"), TraceCounterDisplayHint_Memory);
 FCounterAtomicInt		GCacheReadBytes(TEXT("Ias/CacheReadBytes"), TraceCounterDisplayHint_Memory);
 FCounterAtomicInt		GCacheRejectBytes(TEXT("Ias/CachePutRejectBytes"), TraceCounterDisplayHint_Memory);
+
 // http stats
 bool					GHttpDistributedEndpointResolved = false;
 FCounterInt				GHttpConnectCount(TEXT("Ias/HttpConnectCount"), TraceCounterDisplayHint_None);
 FCounterInt				GHttpDisconnectCount(TEXT("Ias/HttpDisconnectCount"), TraceCounterDisplayHint_None);
-FCounterInt				GHttpGetCount(TEXT("Ias/HttpGetCount"), TraceCounterDisplayHint_None);
-FCounterInt				GHttpErrorCount(TEXT("Ias/HttpErrorCount"), TraceCounterDisplayHint_None);
-FCounterAtomicInt		GHttpDecodeErrorCount(TEXT("Ias/HttpDecodeErrorCount"), TraceCounterDisplayHint_None);
-FCounterInt				GHttpRetryCount(TEXT("Ias/HttpRetryCount"), TraceCounterDisplayHint_None);
-FCounterInt				GHttpCancelCount(TEXT("Ias/HttpCancelCount"), TraceCounterDisplayHint_None);
-FCounterAtomicInt		GHttpPendingCount(TEXT("Ias/HttpPendingCount"), TraceCounterDisplayHint_None);
-FCounterInt				GHttpInflightCount(TEXT("Ias/HttpInflightCount"), TraceCounterDisplayHint_None);
-FCounterInt				GHttpDownloadedBytes(TEXT("Ias/HttpDownloadedBytes"), TraceCounterDisplayHint_Memory);
-FCounterInt				GHttpBandwidthMpbs(TEXT("Ias/HttpBandwidthMbps"), TraceCounterDisplayHint_None);
-FCounterInt				GHttpDurationMs(TEXT("Ias/HttpDurationMs"), TraceCounterDisplayHint_None);
-FCounterInt				GHttpDurationMsAvg(TEXT("Ias/HttpDurationMsAvg"), TraceCounterDisplayHint_None);
-FCounterInt				GHttpDurationMsMax(TEXT("Ias/HttpDurationMsMax"), TraceCounterDisplayHint_None);
-int64					GHttpDurationMsSum = 0;
-constexpr int64			GHttpHistoryCount = 16;
-int64					GHttpHistoryDuration[GHttpHistoryCount] = {};
-int64					GHttpHistoryBytes[GHttpHistoryCount] = {};
-int64					GHttpHistoryTotalDuration = 0;
-int64					GHttpHistoryTotalBytes = 0;
-int64 					GHttpHistoryIndex = 0;
+UE_IAX_COUNTER(HttpGetCount, FCounterInt);
+UE_IAX_COUNTER(HttpErrorCount, FCounterInt);
+UE_IAX_COUNTER(HttpDecodeErrorCount, FCounterAtomicInt);
+UE_IAX_COUNTER(HttpRetryCount, FCounterInt);
+UE_IAX_COUNTER(HttpCancelCount, FCounterInt);
+UE_IAX_COUNTER(HttpPendingCount, FCounterAtomicInt);
+UE_IAX_COUNTER(HttpInflightCount, FCounterInt);
+UE_IAX_MEMORY_COUNTER(HttpDownloadedBytes, FCounterInt);
+UE_IAX_COUNTER(HttpDurationMs, FCounterInt);
+UE_IAX_COUNTER(HttpBandwidthMbps, FCounterInt);
+double GHttpDurationMsAvg[(uint8)EHttpRequestType::NUM_SOURCES] = { 0.0 };
+int32 GHttpDurationMsMax[(uint8)EHttpRequestType::NUM_SOURCES] = { 0 };
+int64 GHttpDurationMsSum[(uint8)EHttpRequestType::NUM_SOURCES] = { 0 };
+
+struct FHttpRecentHistoryStatistics
+{
+	static const int64 HistoryCount = 16;
+
+	int64			Duration[HistoryCount] = {};
+	int64			Bytes[HistoryCount] = {};
+	int64			TotalDuration = 0;
+	int64			MaxDuration = 0;
+	int64			TotalBytes = 0;
+	int64 			Index = 0;
+
+	void OnGet(uint64 SizeBytes, uint64 DurationMs)
+	{
+		int64 OldDuration = Duration[Index];
+		int64 NewDuration = (int64)DurationMs;
+
+		TotalDuration -= OldDuration;
+		TotalDuration += NewDuration;
+		Duration[Index] = NewDuration;
+
+		TotalBytes -= Bytes[Index];
+		TotalBytes += SizeBytes;
+		Bytes[Index] = SizeBytes;
+
+		MaxDuration = FMath::Max(NewDuration, MaxDuration);
+
+		Index = (Index + 1) % HistoryCount;
+	}
+
+	int64 GetBandwidthMbps() const
+	{
+		return (TotalBytes * 8) / (TotalDuration + 1) / 1000;
+	}
+
+	double GetAverage() const
+	{
+		return static_cast<double>(TotalDuration) / static_cast<double>(HistoryCount);
+	}
+
+	int64 GetMaxDuration() const
+	{
+		return MaxDuration;
+	}
+};
+
+FHttpRecentHistoryStatistics GHttpHistory[(uint8)EHttpRequestType::NUM_SOURCES];
+
+// Experimental Http Stats
 static uint32			GHttpCdnCacheHit = 0;
 static uint32			GHttpCdnCacheMiss = 0;
 static uint32			GHttpCdnCacheUnknown = 0;
 
 ////////////////////////////////////////////////////////////////////////////////
 // CSV STATS
+
+#if CSV_PROFILER && !CSV_PROFILER_MINIMAL
+	#define UE_IAX_CSV_DEFINE_STAT(StatName) FCsvDeclaredStat _GCsvStat_##StatName[(uint8)EHttpRequestType::NUM_SOURCES] = {{TEXT(#StatName),CSV_CATEGORY_INDEX(Ias)},{TEXT(#StatName), CSV_CATEGORY_INDEX(Iad)}}
+	#define UE_IAX_CSV_CUSTOM_STAT_DEFINED(RequestType, StatName, Value, Op) FCsvProfiler::RecordCustomStat(_GCsvStat_##StatName[(uint8)RequestType].Name, _GCsvStat_##StatName[(uint8)RequestType].CategoryIndex, Value, Op);
+#else
+	#define UE_IAX_CSV_DEFINE_STAT(StatName)
+	#define UE_IAX_CSV_CUSTOM_STAT_DEFINED(RequestType, StatName, Value, Op)
+#endif
+
 CSV_DEFINE_CATEGORY(Ias, true);
+CSV_DEFINE_CATEGORY(Iad, true);
 // iorequest per frame stats
 CSV_DEFINE_STAT(Ias, FrameIoRequestCount);
 CSV_DEFINE_STAT(Ias, FrameIoRequestReadCount);
@@ -319,15 +380,15 @@ CSV_DEFINE_STAT(Ias, CacheWrittenMB);
 CSV_DEFINE_STAT(Ias, CacheReadMB);
 CSV_DEFINE_STAT(Ias, CacheRejectedMB);
 // http stat totals
-CSV_DEFINE_STAT(Ias, HttpGetCount);
-CSV_DEFINE_STAT(Ias, HttpRetryCount);
-CSV_DEFINE_STAT(Ias, HttpCancelCount);
-CSV_DEFINE_STAT(Ias, HttpErrorCount);
-CSV_DEFINE_STAT(Ias, HttpPendingCount);
-CSV_DEFINE_STAT(Ias, HttpDownloadedMB);
-CSV_DEFINE_STAT(Ias, HttpBandwidthMpbs);
-CSV_DEFINE_STAT(Ias, HttpDurationMsAvg);
-CSV_DEFINE_STAT(Ias, HttpDurationMsMax);
+UE_IAX_CSV_DEFINE_STAT(HttpGetCount);
+UE_IAX_CSV_DEFINE_STAT(HttpRetryCount);
+UE_IAX_CSV_DEFINE_STAT(HttpCancelCount);
+UE_IAX_CSV_DEFINE_STAT(HttpErrorCount);
+UE_IAX_CSV_DEFINE_STAT(HttpPendingCount);
+UE_IAX_CSV_DEFINE_STAT(HttpDownloadedMB);
+UE_IAX_CSV_DEFINE_STAT(HttpBandwidthMpbs);
+UE_IAX_CSV_DEFINE_STAT(HttpDurationMsAvg);
+UE_IAX_CSV_DEFINE_STAT(HttpDurationMsMax);
 
 ////////////////////////////////////////////////////////////////////////////////
 // TRACE STATS IAD
@@ -375,25 +436,22 @@ FOnDemandIoBackendStats::FOnDemandIoBackendStats(FBackendStatus& InStatus)
 		CSV_CUSTOM_STAT_DEFINED(CacheWrittenMB, CWrittenKiB * OneOver1024, ECsvCustomStatOp::Set);
 		CSV_CUSTOM_STAT_DEFINED(CacheReadMB, CReadKiB * OneOver1024, ECsvCustomStatOp::Set);
 		CSV_CUSTOM_STAT_DEFINED(CacheRejectedMB, CRejectedKiB * OneOver1024, ECsvCustomStatOp::Set);
-
+		
 		// http stat totals
-		int32 HGetCount = int32(GHttpGetCount.Get());
-		int32 HRetryCount = int32(GHttpRetryCount.Get());
-		int32 HCancelCount = int32(GHttpCancelCount.Get());
-		int32 HErrorCount = int32(GHttpErrorCount.Get());
-		int32 HPendingCount = int32(GHttpPendingCount.Get());
-		float HDownloadedKiB = (float)GHttpDownloadedBytes.Get() * OneOver1024;
-		int32 HBandwidthMpbs = int32(GHttpBandwidthMpbs.Get());
-		int32 HDurationMsAvg = int32(GHttpDurationMsAvg.Get());
-		int32 HDurationMsMax = int32(GHttpDurationMsMax.Get());
-		CSV_CUSTOM_STAT_DEFINED(HttpGetCount, HGetCount, ECsvCustomStatOp::Set);
-		CSV_CUSTOM_STAT_DEFINED(HttpCancelCount, HCancelCount, ECsvCustomStatOp::Set);
-		CSV_CUSTOM_STAT_DEFINED(HttpErrorCount, HErrorCount, ECsvCustomStatOp::Set);
-		CSV_CUSTOM_STAT_DEFINED(HttpPendingCount, HPendingCount, ECsvCustomStatOp::Set);
-		CSV_CUSTOM_STAT_DEFINED(HttpDownloadedMB, HDownloadedKiB * OneOver1024, ECsvCustomStatOp::Set);
-		CSV_CUSTOM_STAT_DEFINED(HttpBandwidthMpbs, HBandwidthMpbs, ECsvCustomStatOp::Set);
-		CSV_CUSTOM_STAT_DEFINED(HttpDurationMsAvg, HDurationMsAvg, ECsvCustomStatOp::Set);
-		CSV_CUSTOM_STAT_DEFINED(HttpDurationMsMax, HDurationMsMax, ECsvCustomStatOp::Set);
+		auto HttpCsv = [](EHttpRequestType Type)
+			{
+				UE_IAX_CSV_CUSTOM_STAT_DEFINED(Type, HttpGetCount, (int32)GHttpGetCount[(uint8)Type].Get(), ECsvCustomStatOp::Set);
+				UE_IAX_CSV_CUSTOM_STAT_DEFINED(Type, HttpCancelCount, (int32)GHttpCancelCount[(uint8)Type].Get(), ECsvCustomStatOp::Set);
+				UE_IAX_CSV_CUSTOM_STAT_DEFINED(Type, HttpErrorCount, (int32)GHttpErrorCount[(uint8)Type].Get(), ECsvCustomStatOp::Set);
+				UE_IAX_CSV_CUSTOM_STAT_DEFINED(Type, HttpPendingCount, (int32)GHttpPendingCount[(uint8)Type].Get(), ECsvCustomStatOp::Set);
+				UE_IAX_CSV_CUSTOM_STAT_DEFINED(Type, HttpDownloadedMB, (float)GHttpDownloadedBytes[(uint8)Type].Get() * OneOver1024 * OneOver1024, ECsvCustomStatOp::Set);
+				UE_IAX_CSV_CUSTOM_STAT_DEFINED(Type, HttpBandwidthMpbs, (int32)GHttpBandwidthMbps[(uint8)Type].Get(), ECsvCustomStatOp::Set);
+				UE_IAX_CSV_CUSTOM_STAT_DEFINED(Type, HttpDurationMsAvg, (int32)GHttpDurationMsAvg[(uint8)Type], ECsvCustomStatOp::Set);
+				UE_IAX_CSV_CUSTOM_STAT_DEFINED(Type, HttpDurationMsMax, GHttpDurationMsMax[(uint8)Type], ECsvCustomStatOp::Set);
+			};
+
+		HttpCsv(EHttpRequestType::Streaming);
+		HttpCsv(EHttpRequestType::Installed);
 
 		if (GIasStatisticsLogInterval > 0.f)
 		{
@@ -410,8 +468,22 @@ FOnDemandIoBackendStats::FOnDemandIoBackendStats(FBackendStatus& InStatus)
 					UE_LOG(LogIas, Log, TEXT("CacheStats: Disabled"));
 				}
 
-				UE_LOG(LogIas, Log, TEXT("HttpStats: DownloadedKiB=%d, Get=%d, Retry=%d, Cancel=%d, Error=%d, CurPending=%d, CurDurationMsAvg=%d, CurDurationMsMax=%d"),
-					(int32)HDownloadedKiB, HGetCount, HRetryCount, HCancelCount, HErrorCount, HPendingCount, HDurationMsAvg, HDurationMsMax);
+				auto HttpLog = [](const TCHAR* Title, EHttpRequestType Type)
+				{
+						UE_LOG(LogIas, Log, TEXT("%s - HttpStats: DownloadedKiB=%d, Get=%d, Retry=%d, Cancel=%d, Error=%d, CurPending=%d, CurDurationMsAvg=%d, CurDurationMsMax=%d"),
+							Title, 
+							(int32)((float)GHttpDownloadedBytes[(uint8)Type].Get() * OneOver1024),
+							(int32)GHttpGetCount[(uint8)Type].Get(),
+							(int32)GHttpRetryCount[(uint8)Type].Get(),
+							(int32)GHttpCancelCount[(uint8)Type].Get(),
+							(int32)GHttpErrorCount[(uint8)Type].Get(),
+							(int32)GHttpPendingCount[(uint8)Type].Get(),
+							(int32)GHttpDurationMsAvg[(uint8)Type],
+							GHttpDurationMsMax[(uint8)Type]);
+				};
+
+				HttpLog(TEXT("IAS"), EHttpRequestType::Streaming);
+				HttpLog(TEXT("IAD"), EHttpRequestType::Installed);
 
 				{
 					const uint64 TotalInstallCount		= GTotalInstallCount.Get();
@@ -475,34 +547,47 @@ FOnDemandIoBackendStats::FOnDemandIoBackendStats(FBackendStatus& InStatus)
 				Out.Add(Verbosity, FText::FromStringView(Text.ToView()));
 			});
 
-			TStringBuilder<256> CachingText;
-			if (BackendStatus.IsCacheEnabled())
 			{
-				CachingText << TEXT("Cached: ") << FText::AsMemory(GCacheCachedBytes.Get()).ToString();
-				CachingText << TEXT(" | Rejected: ") << FText::AsMemory(GCacheRejectBytes.Get()).ToString();
-				CachingText << TEXT(" | Read: ") << FText::AsMemory(GCacheReadBytes.Get()).ToString();
-				CachingText << TEXT(" (") << GCacheGetCount.Get() << TEXT(")");
+				TStringBuilder<256> CachingText;
+				CachingText << TEXT("IAS Cache Stats: ");
+				if (BackendStatus.IsCacheEnabled())
+				{
+					CachingText << TEXT("Cached: ") << FText::AsMemory(GCacheCachedBytes.Get()).ToString();
+					CachingText << TEXT(" | Rejected: ") << FText::AsMemory(GCacheRejectBytes.Get()).ToString();
+					CachingText << TEXT(" | Read: ") << FText::AsMemory(GCacheReadBytes.Get()).ToString();
+					CachingText << TEXT(" (") << GCacheGetCount.Get() << TEXT(")");
+				}
+				else
+				{
+					CachingText << TEXTVIEW("Caching Disabled");
+				}
+
+				Out.Add(FCoreDelegates::EOnScreenMessageSeverity::Info, FText::FromStringView(CachingText.ToView()));
 			}
-			else
+
 			{
-				CachingText = TEXTVIEW("Caching Disabled");
+				auto HttpStats = [&Out](const TCHAR* Title, EHttpRequestType Type)
+				{
+					TStringBuilder<256> Builder;
+					Builder.Appendf(
+						TEXT("%s Backend Stats: Downloaded: %s (%d) Avg %d ms | Retries: %d | Pending: %d"),
+						Title,
+						*FText::AsMemory(GHttpDownloadedBytes[(uint8)Type].Get()).ToString(),
+						GHttpGetCount[(uint8)Type].Get(),
+						(int32)GHttpDurationMsAvg[(uint8)Type],
+						GHttpRetryCount[(uint8)Type].Get(),
+						GHttpPendingCount[(uint8)Type].Get()
+					);
+
+					Out.Add(FCoreDelegates::EOnScreenMessageSeverity::Info, FText::FromStringView(Builder.ToView()));
+				};
+
+				HttpStats(TEXT("IAS"), EHttpRequestType::Streaming);
+				HttpStats(TEXT("IAD"), EHttpRequestType::Installed);
 			}
 
-			TStringBuilder<256> Builder;
-			Builder.Appendf(
-				TEXT("IAS Backend Stats:   %s | Downloaded: %s (%d) Avg %d ms | Retries: %d | Pending: %d"),
-				CachingText.ToString(),
-				*FText::AsMemory(GHttpDownloadedBytes.Get()).ToString(),
-				GHttpGetCount.Get(),
-				GHttpDurationMsAvg.Get(),
-				GHttpRetryCount.Get(),
-				GHttpPendingCount.Get()
-			);
-
-			Out.Add(FCoreDelegates::EOnScreenMessageSeverity::Info, FText::FromStringView(Builder.ToView()));
-
 			{
-				Builder.Reset();
+				TStringBuilder<256> Builder;
 				Builder << TEXT("IAS CDN: Hit/Miss/NoHdr: ");
 
 				Builder << GHttpCdnCacheHit;
@@ -521,13 +606,13 @@ FOnDemandIoBackendStats::FOnDemandIoBackendStats(FBackendStatus& InStatus)
 				Out.Add(FCoreDelegates::EOnScreenMessageSeverity::Info, FText::FromStringView(Builder.ToView()));
 			}
 
-			if (GHttpDecodeErrorCount.Get() > 0 || GCacheDecodeErrorCount.Get() > 0 || GHttpErrorCount.Get() > 0)
+			if (GHttpDecodeErrorCount[(uint8)EHttpRequestType::Streaming].Get() > 0 || GCacheDecodeErrorCount.Get() > 0 || GHttpErrorCount[(uint8)EHttpRequestType::Streaming].Get() > 0)
 			{
-				Builder.Reset();
+				TStringBuilder<256> Builder;
 				Builder.Appendf(TEXT("IAS Errors: Cache Decode: %d | Http Decode: %d | Http: %d"),
 					GCacheDecodeErrorCount.Get(),
-					GHttpDecodeErrorCount.Get(),
-					GHttpErrorCount.Get()
+					GHttpDecodeErrorCount[(uint8)EHttpRequestType::Streaming].Get(),
+					GHttpErrorCount[(uint8)EHttpRequestType::Streaming].Get()
 				);
 
 				Out.Add(FCoreDelegates::EOnScreenMessageSeverity::Error, FText::FromStringView(Builder.ToView()));
@@ -562,62 +647,71 @@ void FOnDemandIoBackendStats::ReportGeneralAnalytics(TArray<FAnalyticsEventAttri
 
 void FOnDemandIoBackendStats::ReportEndPointAnalytics(TArray<FAnalyticsEventAttribute>& OutAnalyticsArray) const
 {
-	// We use this macro with counters that are used elsewhere meaning we can't just reset them
-	// each time we send an analytics payload. The macro will track the running total and only
-	// pass on the delta to the analytics array since the last time this method was called.
-#define TRACK_DELTA(Name, Value) TEXT(Name), GDeltaTracking.Get(TEXTVIEW(Name), Value)
+
 
 	if (GIasReportHttpAnalyticsEnabled)
 	{
-		const int64 ByteCount = GDeltaTracking.Get(TEXTVIEW("IasHttpDownloadedBytes"), GHttpDownloadedBytes.Get());
-		const int64 GetCount = GDeltaTracking.Get(TEXTVIEW("IasHttpGetCount"), GHttpGetCount.Get());
+		auto ReportHttpStats = [&OutAnalyticsArray](const TCHAR* Prefix, EHttpRequestType Type)
+		{
+#define UE_PREFIX(Name) WriteToString<128>(Prefix, Name)
+#define UE_TRACK_DELTA(Name, Value) *Name, GDeltaTracking.Get(Name, Value)
 
-		const double DataRateBPS = GHttpDurationMsSum > 0 ? static_cast<double>(ByteCount) / (static_cast<double>(GHttpDurationMsSum) / 1000.0) : 0.0;
-		const double DurationMean = GetCount ? static_cast<double>(GHttpDurationMsSum) / static_cast<double>(GetCount) : 0.0;
+			const int64 ByteCount = GDeltaTracking.Get(UE_PREFIX(TEXT("HttpDownloadedBytes")), GHttpDownloadedBytes[(uint8)Type].Get());
+			const int64 GetCount = GDeltaTracking.Get(UE_PREFIX(TEXT("HttpGetCount")), GHttpGetCount[(uint8)Type].Get());
 
-		AppendAnalyticsEventAttributeArray(OutAnalyticsArray
-			,TRACK_DELTA("IasHttpErrorCount", GHttpErrorCount.Get())
-			,TRACK_DELTA("IasHttpDecodeErrors", GHttpDecodeErrorCount.Get())
-			,TRACK_DELTA("IasHttpRetryCount", GHttpRetryCount.Get())
-			,TEXT("IasHttpGetCount"), GetCount
-			,TEXT("IasHttpDownloadedBytes"), ByteCount
-			,TEXT("IasHttpDurationMean"), DurationMean
-			,TEXT("IasHttpDurationSum"), GHttpDurationMsSum
-			,TEXT("IasHttpDataRateMean"), DataRateBPS
-		);
+			const double DataRateBPS = GHttpDurationMsSum[(uint8)Type] > 0 ? static_cast<double>(ByteCount) / (static_cast<double>(GHttpDurationMsSum[(uint8)Type]) / 1000.0) : 0.0;
+			const double DurationMean = GetCount ? static_cast<double>(GHttpDurationMsSum[(uint8)Type]) / static_cast<double>(GetCount) : 0.0;
+
+			AppendAnalyticsEventAttributeArray(OutAnalyticsArray
+				, UE_TRACK_DELTA(UE_PREFIX("HttpErrorCount"), GHttpErrorCount[(uint8)Type].Get())
+				, UE_TRACK_DELTA(UE_PREFIX("HttpDecodeErrors"), GHttpDecodeErrorCount[(uint8)Type].Get())
+				, UE_TRACK_DELTA(UE_PREFIX("HttpRetryCount"), GHttpRetryCount[(uint8)Type].Get())
+				, *UE_PREFIX(TEXT("HttpGetCount")), GetCount
+				, *UE_PREFIX(TEXT("HttpDownloadedBytes")), ByteCount
+				, *UE_PREFIX(TEXT("HttpDurationMean")), DurationMean
+				, *UE_PREFIX(TEXT("HttpDurationSum")), GHttpDurationMsSum[(uint8)Type]
+				, *UE_PREFIX(TEXT("HttpDataRateMean")), DataRateBPS
+			);
+
+			// These values we can just reset as they are only being used with analytics
+			GHttpDurationMsSum[(uint8)Type] = 0;
+#undef UE_TRACK_DELTA
+#undef UE_PREFIX
+		};
+
+		ReportHttpStats(TEXT("Ias"), EHttpRequestType::Streaming);
+		ReportHttpStats(TEXT("Iad"), EHttpRequestType::Installed);
 	}
 
 	if (GIasReportCacheAnalyticsEnabled)
 	{
+#define UE_TRACK_DELTA(Name, Value) TEXT(Name), GDeltaTracking.Get(TEXTVIEW(Name), Value)
 		const int64 CacheTotalCount = GCacheGetCount.Get() + GCachePutCount.Get();
 		const float CacheUsagePercent = GCacheMaxBytes > 0 ? 100.f * (float(GCacheCachedBytes.Get()) / float(GCacheMaxBytes)) : 0.f;
 
 		AppendAnalyticsEventAttributeArray(OutAnalyticsArray
 
-			,TEXT("IasCacheEnabled"), BackendStatus.IsCacheEnabled()
+			, TEXT("IasCacheEnabled"), BackendStatus.IsCacheEnabled()
 
-			,TRACK_DELTA("IasCacheTotalCount", CacheTotalCount)
-			,TRACK_DELTA("IasCacheErrorCount", GCacheErrorCount.Get())
-			,TRACK_DELTA("IasCacheDecodeErrors", GCacheDecodeErrorCount.Get())
-			,TRACK_DELTA("IasCacheGetCount", GCacheGetCount.Get())
-			,TRACK_DELTA("IasCachePutCount", GCachePutCount.Get())
+			, UE_TRACK_DELTA("IasCacheTotalCount", CacheTotalCount)
+			, UE_TRACK_DELTA("IasCacheErrorCount", GCacheErrorCount.Get())
+			, UE_TRACK_DELTA("IasCacheDecodeErrors", GCacheDecodeErrorCount.Get())
+			, UE_TRACK_DELTA("IasCacheGetCount", GCacheGetCount.Get())
+			, UE_TRACK_DELTA("IasCachePutCount", GCachePutCount.Get())
 
-			,TEXT("IasCacheCachedBytes"), GCacheCachedBytes.Get()
-			,TEXT("IasCacheMaxBytes"), GCacheMaxBytes
-			,TEXT("IasCacheUsagePercent"), CacheUsagePercent
+			, TEXT("IasCacheCachedBytes"), GCacheCachedBytes.Get()
+			, TEXT("IasCacheMaxBytes"), GCacheMaxBytes
+			, TEXT("IasCacheUsagePercent"), CacheUsagePercent
 
-			,TRACK_DELTA("IasCacheWriteBytes", GCacheWrittenBytes.Get())
-			,TRACK_DELTA("IasCacheReadBytes", GCacheReadBytes.Get())
-			,TRACK_DELTA("IasCacheRejectBytes", GCacheRejectBytes.Get())
+			, UE_TRACK_DELTA("IasCacheWriteBytes", GCacheWrittenBytes.Get())
+			, UE_TRACK_DELTA("IasCacheReadBytes", GCacheReadBytes.Get())
+			, UE_TRACK_DELTA("IasCacheRejectBytes", GCacheRejectBytes.Get())
 		);
-
-		// These values we can just reset as they are only being used with analytics
-		GHttpDurationMsSum = 0;
+#undef UE_TRACK_DELTA
 	}
-
-#undef TRACK_DELTA
 }
 
+#if UE_ENABLE_ANALYTICS_RECORDING
 
 class FAnalyticsRecording : public IAnalyticsRecording
 {
@@ -765,12 +859,12 @@ private:
 	/** Http Stats */
 	struct FHttpStats
 	{
-		FTrackedValue<FCounterInt, GHttpErrorCount>				ErrorCount;
-		FTrackedValue<FCounterAtomicInt, GHttpDecodeErrorCount> DecodeErrorCount;
-		FTrackedValue<FCounterInt, GHttpRetryCount>				RetryCount;
-		FTrackedValue<FCounterInt, GHttpGetCount>				GetCount;
-		FTrackedValue<FCounterInt, GHttpDownloadedBytes>		DownloadedBytes;
-		FTrackedValue<int64, GHttpDurationMsSum>				TotalDuration;
+		FTrackedValue<FCounterInt, GHttpErrorCount[0]>				ErrorCount;
+		FTrackedValue<FCounterAtomicInt, GHttpDecodeErrorCount[0]>	DecodeErrorCount;
+		FTrackedValue<FCounterInt, GHttpRetryCount[0]>				RetryCount;
+		FTrackedValue<FCounterInt, GHttpGetCount[0]>				GetCount;
+		FTrackedValue<FCounterInt, GHttpDownloadedBytes[0]>			DownloadedBytes;
+		FTrackedValue<int64, GHttpDurationMsSum[0]>					TotalDuration;
 		
 	} Http;
 
@@ -791,21 +885,27 @@ private:
 	bool bRecording = true;
 };
 
+#endif //UE_ENABLE_ANALYTICS_RECORDING
+
 TUniquePtr<IAnalyticsRecording> FOnDemandIoBackendStats::StartAnalyticsRecording() const
 {
+#if UE_ENABLE_ANALYTICS_RECORDING
 	TUniquePtr<FAnalyticsRecording> Recording = MakeUnique<FAnalyticsRecording>(BackendStatus);
 	return Recording;
+#else
+	return nullptr;
+#endif // UE_ENABLE_ANALYTICS_RECORDING
 }
 
 void FOnDemandIoBackendStats::OnIoRequestEnqueue()
 {
-	GIoRequestCount.Add(1);
+	GIoRequestCount.Increment();
 	CSV_CUSTOM_STAT_DEFINED(FrameIoRequestCount, int32(GIoRequestCount.Get()), ECsvCustomStatOp::Set);
 }
 
 void FOnDemandIoBackendStats::OnIoRequestComplete(uint64 Size, uint64 Duration)
 {
-	GIoRequestReadCount.Add(1);
+	GIoRequestReadCount.Increment();
 	GIoRequestReadBytes.Add(Size);
 
 	CSV_CUSTOM_STAT_DEFINED(FrameIoRequestReadCount, int32(GIoRequestReadCount.Get()), ECsvCustomStatOp::Set);
@@ -814,45 +914,45 @@ void FOnDemandIoBackendStats::OnIoRequestComplete(uint64 Size, uint64 Duration)
 
 void FOnDemandIoBackendStats::OnIoRequestCancel()
 {
-	GIoRequestCancelCount.Add(1);
+	GIoRequestCancelCount.Increment();
 	CSV_CUSTOM_STAT_DEFINED(FrameIoRequestCancelCount, int32(GIoRequestCancelCount.Get()), ECsvCustomStatOp::Set);
 }
 
 void FOnDemandIoBackendStats::OnIoRequestError()
 {
-	GIoRequestErrorCount.Add(1);
+	GIoRequestErrorCount.Increment();
 	CSV_CUSTOM_STAT_DEFINED(FrameIoRequestErrorCount, int32(GIoRequestErrorCount.Get()), ECsvCustomStatOp::Set);
 }
 
 void FOnDemandIoBackendStats::OnCacheError()
 {
-	GCacheErrorCount.Add(1);
+	GCacheErrorCount.Increment();
 }
 
 void FOnDemandIoBackendStats::OnCacheDecodeError()
 {
-	GCacheDecodeErrorCount.Add(1);
+	GCacheDecodeErrorCount.Increment();
 }
 
 void FOnDemandIoBackendStats::OnCacheGet(uint64 DataSize)
 {
-	GCacheGetCount.Add(1);
+	GCacheGetCount.Increment();
 	GCacheReadBytes.Add(DataSize);
 }
 
 void FOnDemandIoBackendStats::OnCachePut()
 {
-	GCachePutCount.Add(1);
+	GCachePutCount.Increment();
 }
 
 void FOnDemandIoBackendStats::OnCachePutExisting(uint64 /*DataSize*/)
 {
-	GCachePutExistingCount.Add(1);
+	GCachePutExistingCount.Increment();
 }
 
 void FOnDemandIoBackendStats::OnCachePutReject(uint64 DataSize)
 {
-	GCachePutRejectCount.Add(1);
+	GCachePutRejectCount.Increment();
 	GCacheRejectBytes.Add(DataSize);
 }
 
@@ -883,84 +983,65 @@ void FOnDemandIoBackendStats::OnHttpDistributedEndpointResolved()
 
 void FOnDemandIoBackendStats::OnHttpConnected()
 {
-	GHttpConnectCount.Add(1);
+	GHttpConnectCount.Increment();
 }
 
 void FOnDemandIoBackendStats::OnHttpDisconnected()
 {
-	GHttpDisconnectCount.Add(1);
+	GHttpDisconnectCount.Increment();
 }
 
-void FOnDemandIoBackendStats::OnHttpEnqueue()
+void FOnDemandIoBackendStats::OnHttpEnqueue(EHttpRequestType Type)
 {
-	GHttpPendingCount.Add(1);
+	GHttpPendingCount[(uint8)Type].Increment();
 }
 
-void FOnDemandIoBackendStats::OnHttpDequeue()
+void FOnDemandIoBackendStats::OnHttpDequeue(EHttpRequestType Type)
 {
-	GHttpInflightCount.Add(1);
+	GHttpInflightCount[(uint8)Type].Increment();
 }
 
-void FOnDemandIoBackendStats::OnHttpGet(uint64 SizeBytes, uint64 DurationMs)
+void FOnDemandIoBackendStats::OnHttpGet(EHttpRequestType Type, uint64 SizeBytes, uint64 DurationMs)
 {
-	GHttpPendingCount.Add(-1);
-	GHttpInflightCount.Add(-1);
-	GHttpGetCount.Add(1);
-	GHttpDownloadedBytes.Add(SizeBytes);
-	GHttpDurationMsSum += DurationMs;
-	GHttpDurationMs.Set(DurationMs);
-
-	int64 OldDuration = GHttpHistoryDuration[GHttpHistoryIndex];
-	int64 NewDuration = (int64)DurationMs;
-	
-	GHttpHistoryTotalDuration -= OldDuration;
-	GHttpHistoryTotalDuration += NewDuration;
-	GHttpHistoryDuration[GHttpHistoryIndex] = NewDuration;
-
-	GHttpHistoryTotalBytes -= GHttpHistoryBytes[GHttpHistoryIndex];
-	GHttpHistoryTotalBytes += SizeBytes;
-	GHttpHistoryBytes[GHttpHistoryIndex] = SizeBytes;
-
-	GHttpBandwidthMpbs.Set((GHttpHistoryTotalBytes*8)/(GHttpHistoryTotalDuration+1)/1000);
-	GHttpDurationMsAvg.Set(GHttpHistoryTotalDuration/GHttpHistoryCount);
-
-	if (NewDuration > GHttpDurationMsMax.Get())
-	{
-		GHttpDurationMsMax.Set(NewDuration);
-	}
-	else if (OldDuration == GHttpDurationMsMax.Get())
-	{
-		GHttpDurationMsMax.Set(*Algo::MaxElement(GHttpHistoryDuration));
-	}
-
-	GHttpHistoryIndex = (GHttpHistoryIndex + 1) % GHttpHistoryCount;
+	GHttpPendingCount[(uint8)Type].Decrement();
+	GHttpInflightCount[(uint8)Type].Decrement();
+	GHttpGetCount[(uint8)Type].Increment();
+	GHttpDownloadedBytes[(uint8)Type].Add(SizeBytes);
+	GHttpDurationMsSum[(uint8)Type] += DurationMs;
+	GHttpDurationMs[(uint8)Type].Set(DurationMs);
+
+	GHttpHistory[(uint8)Type].OnGet(SizeBytes, DurationMs);
+
+	GHttpBandwidthMbps[(uint8)Type].Set(GHttpHistory[(uint8)Type].GetBandwidthMbps());
+	GHttpDurationMsAvg[(uint8)Type] = GHttpHistory[(uint8)Type].GetAverage();
+	GHttpDurationMsMax[(uint8)Type] = (int32)GHttpHistory[(uint8)Type].GetMaxDuration();
 }
 
-void FOnDemandIoBackendStats::OnHttpCancel()
+void FOnDemandIoBackendStats::OnHttpCancel(EHttpRequestType Type)
 {
-	GHttpInflightCount.Add(-1);
-	GHttpPendingCount.Add(-1);
-	GHttpCancelCount.Add(1);
+	GHttpInflightCount[(uint8)Type].Decrement();
+	GHttpPendingCount[(uint8)Type].Decrement();
+	GHttpCancelCount[(uint8)Type].Increment();
 }
 
-void FOnDemandIoBackendStats::OnHttpRetry()
+void FOnDemandIoBackendStats::OnHttpRetry(EHttpRequestType Type)
 {
-	GHttpRetryCount.Add(1);
+	GHttpRetryCount[(uint8)Type].Increment();
 }
 
-void FOnDemandIoBackendStats::OnHttpError()
+void FOnDemandIoBackendStats::OnHttpError(EHttpRequestType Type)
 {
-	GHttpPendingCount.Add(-1);
-	GHttpInflightCount.Add(-1);
-	GHttpErrorCount.Add(1);
+	GHttpPendingCount[(uint8)Type].Decrement();
+	GHttpInflightCount[(uint8)Type].Decrement();
+	GHttpErrorCount[(uint8)Type].Increment();
 }
 
-void FOnDemandIoBackendStats::OnHttpDecodeError()
+void FOnDemandIoBackendStats::OnHttpDecodeError(EHttpRequestType Type)
 {
-	GHttpDecodeErrorCount.Add(1);
+	GHttpDecodeErrorCount[(uint8)Type].Increment();
 }
 
-void FOnDemandIoBackendStats::OnHttpCdnCacheReply(int32 Reply)
+void FOnDemandIoBackendStats::OnHttpCdnCacheReply(EHttpRequestType Type, int32 Reply)
 {
 	switch (Reply)
 	{
@@ -973,7 +1054,7 @@ void FOnDemandIoBackendStats::OnHttpCdnCacheReply(int32 Reply)
 ////////////////////////////////////////////////////////////////////////////////
 void FOnDemandContentInstallerStats::OnRequestEnqueued()
 {
-	GInflightInstallCount.Add(1);
+	GInflightInstallCount.Increment();
 }
 
 void FOnDemandContentInstallerStats::OnRequestCompleted(
@@ -990,20 +1071,20 @@ void FOnDemandContentInstallerStats::OnRequestCompleted(
 		return;
 	}
 
-	GInflightInstallCount.Add(-1);
+	GInflightInstallCount.Decrement();
 	if (ErrorCode != EIoErrorCode::Ok)
 	{
 		if (ErrorCode == EIoErrorCode::UnknownChunkID || ErrorCode == EIoErrorCode::NotFound)
 		{
-			GTotalMissingChunkInstallErrorCount.Add(1);
+			GTotalMissingChunkInstallErrorCount.Increment();
 		}
 		else if (ErrorCode == EIoErrorCode::WriteError)
 		{
-			GTotalWriteErrorInstallErrorCount.Add(1);
+			GTotalWriteErrorInstallErrorCount.Increment();
 		}
-		GTotalInstallErrorCount.Add(1);
+		GTotalInstallErrorCount.Increment();
 	}
-	GTotalInstallCount.Add(1);
+	GTotalInstallCount.Increment();
 	GTotalDownloadedChunkCount.Add(DownloadedChunkCount);
 	GTotalDownloadedChunkBytes.Add(DownloadedBytes);
 	GTotalInstallDurationMs.Add(uint64(FPlatformTime::ToMilliseconds64(DurationCycles)));
diff --git a/Engine/Source/Runtime/Experimental/IoStore/OnDemand/Private/Statistics.h b/Engine/Source/Runtime/Experimental/IoStore/OnDemand/Private/Statistics.h
--- a/Engine/Source/Runtime/Experimental/IoStore/OnDemand/Private/Statistics.h
+++ b/Engine/Source/Runtime/Experimental/IoStore/OnDemand/Private/Statistics.h
@@ -22,6 +22,7 @@ namespace UE::IoStore
 {
 
 class IAnalyticsRecording;
+enum class EHttpRequestType : uint8;
 struct FAvailableEndPoints;
 
 #if IAS_WITH_STATISTICS
@@ -62,18 +63,22 @@ public:
 	void OnCacheWriteBytes(uint64 WriteSize) IAS_STATISTICS_IMPL()
 	void OnCacheSetMaxBytes(uint64 TotalSize) IAS_STATISTICS_IMPL()
 
+	// Per session
 	void OnHttpDistributedEndpointResolved() IAS_STATISTICS_IMPL()
+
+	// Per hostgroup
 	void OnHttpConnected() IAS_STATISTICS_IMPL()
 	void OnHttpDisconnected() IAS_STATISTICS_IMPL()
 
-	void OnHttpEnqueue() IAS_STATISTICS_IMPL()
-	void OnHttpCancel() IAS_STATISTICS_IMPL()
-	void OnHttpDequeue() IAS_STATISTICS_IMPL()
-	void OnHttpGet(uint64 SizeBytes, uint64 DurationMs) IAS_STATISTICS_IMPL()
-	void OnHttpRetry() IAS_STATISTICS_IMPL()
-	void OnHttpError() IAS_STATISTICS_IMPL()
-	void OnHttpDecodeError() IAS_STATISTICS_IMPL()
-	void OnHttpCdnCacheReply(int32 Reply) IAS_STATISTICS_IMPL()
+	// From FOnDemandHttpThread
+	void OnHttpEnqueue(EHttpRequestType Type) IAS_STATISTICS_IMPL()
+	void OnHttpCancel(EHttpRequestType Type) IAS_STATISTICS_IMPL()
+	void OnHttpDequeue(EHttpRequestType Type) IAS_STATISTICS_IMPL()
+	void OnHttpGet(EHttpRequestType Type, uint64 SizeBytes, uint64 DurationMs) IAS_STATISTICS_IMPL()
+	void OnHttpRetry(EHttpRequestType Type) IAS_STATISTICS_IMPL()
+	void OnHttpError(EHttpRequestType Type) IAS_STATISTICS_IMPL()
+	void OnHttpDecodeError(EHttpRequestType Type) IAS_STATISTICS_IMPL()
+	void OnHttpCdnCacheReply(EHttpRequestType Type,int32 Reply) IAS_STATISTICS_IMPL()
 
 private:
 
