Fix broken collision behavior when converting to nonprobe.

The ConvertToNonProbe contact modifier doesn't initialize certain constraint fields that would be handled normally through the non-probe path in midphase generation. This change updates the shape transforms and constraint in the contact modifier to improve stability.

diff --git a/Engine/Source/Runtime/Experimental/Chaos/Private/Chaos/ContactModification.cpp b/Engine/Source/Runtime/Experimental/Chaos/Private/Chaos/ContactModification.cpp
--- a/Engine/Source/Runtime/Experimental/Chaos/Private/Chaos/ContactModification.cpp
+++ b/Engine/Source/Runtime/Experimental/Chaos/Private/Chaos/ContactModification.cpp
@@ -4,6 +4,8 @@
 #include "Chaos/Collision/PBDCollisionConstraint.h"
 #include "Chaos/Particle/ParticleUtilities.h"
 #include "Chaos/PBDCollisionConstraints.h"
+#include "Chaos/CollisionResolution.h"
+
 namespace Chaos
 {
 	void FContactPairModifier::Disable()
@@ -649,6 +651,13 @@ namespace Chaos
 	{
 		Constraint.SetModifierApplied();
 
+		const FRigidTransform3& ShapeWorldTransform0 = Constraint.GetContact().GetShapeRelativeTransform0() * FParticleUtilities::GetActorWorldTransform(FConstGenericParticleHandle(Constraint.GetParticle0()));
+		const FRigidTransform3& ShapeWorldTransform1 = Constraint.GetContact().GetShapeRelativeTransform1() * FParticleUtilities::GetActorWorldTransform(FConstGenericParticleHandle(Constraint.GetParticle1()));
+		
+		Constraint.SetShapeWorldTransforms(ShapeWorldTransform0, ShapeWorldTransform1);
+
+		Collisions::UpdateConstraint(Constraint, ShapeWorldTransform0, ShapeWorldTransform1, Dt);
+
 		Constraint.SetIsProbe(false);
 	}
 
