Dataflow : add API to subgraph input and output nodes for adding new connections

diff --git a/Engine/Source/Runtime/Experimental/Dataflow/Engine/Private/Dataflow/DataflowSubGraphNodes.cpp b/Engine/Source/Runtime/Experimental/Dataflow/Engine/Private/Dataflow/DataflowSubGraphNodes.cpp
--- a/Engine/Source/Runtime/Experimental/Dataflow/Engine/Private/Dataflow/DataflowSubGraphNodes.cpp
+++ b/Engine/Source/Runtime/Experimental/Dataflow/Engine/Private/Dataflow/DataflowSubGraphNodes.cpp
@@ -296,6 +296,20 @@ const FDataflowConnection* FDataflowSubGraphInputNode::OnDropConnectionOnNode(co
 	return nullptr;
 }
 
+void FDataflowSubGraphInputNode::AddConnectionsTo(TArrayView<FDataflowConnection*> Connections)
+{
+	using namespace  UE::Dataflow::SubGraphNodes;
+	for (const FDataflowConnection* Connection : Connections)
+	{
+		if (Connection && SupportsDropConnectionOnNode(Connection->GetType(), Connection->GetDirection()))
+		{
+			const FPropertyBagPropertyDesc Desc = Private::GetPropertyBagPropertyDescFromDataflowConnection(*Connection);
+			PropertyBag.AddProperties({ Desc });
+		}
+	}
+	DynamicConnections.Refresh();
+}
+
 //////////////////////////////////////////////////////////////////////////////////////////////////////////
 
 FDataflowSubGraphOutputNode::FDataflowSubGraphOutputNode(const UE::Dataflow::FNodeParameters& InParam, FGuid InGuid)
@@ -343,6 +357,20 @@ const FDataflowConnection* FDataflowSubGraphOutputNode::OnDropConnectionOnNode(c
 	return nullptr;
 }
 
+void FDataflowSubGraphOutputNode::AddConnectionsTo(TArrayView<FDataflowConnection*> Connections)
+{
+	using namespace  UE::Dataflow::SubGraphNodes;
+	for (const FDataflowConnection* Connection : Connections)
+	{
+		if (Connection && SupportsDropConnectionOnNode(Connection->GetType(), Connection->GetDirection()))
+		{
+			const FPropertyBagPropertyDesc Desc = Private::GetPropertyBagPropertyDescFromDataflowConnection(*Connection);
+			PropertyBag.AddProperties({ Desc });
+		}
+	}
+	DynamicConnections.Refresh();
+}
+
 //////////////////////////////////////////////////////////////////////////////////////////////////////////
 FDataflowSubGraphGetCurrentIndexNode::FDataflowSubGraphGetCurrentIndexNode(const UE::Dataflow::FNodeParameters& Param, FGuid InGuid)
 	: FDataflowNode(Param, InGuid)
diff --git a/Engine/Source/Runtime/Experimental/Dataflow/Engine/Public/Dataflow/DataflowSubGraphNodes.h b/Engine/Source/Runtime/Experimental/Dataflow/Engine/Public/Dataflow/DataflowSubGraphNodes.h
--- a/Engine/Source/Runtime/Experimental/Dataflow/Engine/Public/Dataflow/DataflowSubGraphNodes.h
+++ b/Engine/Source/Runtime/Experimental/Dataflow/Engine/Public/Dataflow/DataflowSubGraphNodes.h
@@ -47,6 +47,8 @@ public:
 	virtual const FInstancedPropertyBag& GetPropertyBag(const FDataflowDynamicConnections* Caller) override { return PropertyBag; }
 	/* End IOwnerInterface */
 
+	DATAFLOWENGINE_API void AddConnectionsTo(TArrayView<FDataflowConnection*> Connections);
+
 private:
 	virtual void Evaluate(UE::Dataflow::FContext& Context, const FDataflowOutput* Out) const override;
 	virtual void PostSerialize(const FArchive& Ar) override;
@@ -84,6 +86,8 @@ public:
 	virtual const FInstancedPropertyBag& GetPropertyBag(const FDataflowDynamicConnections* Caller) override { return PropertyBag; }
 	/* End IOwnerInterface */
 
+	DATAFLOWENGINE_API void AddConnectionsTo(TArrayView<FDataflowConnection*> Connections);
+
 private:
 	virtual void Evaluate(UE::Dataflow::FContext& Context, const FDataflowOutput* Out) const override;
 	virtual void PostSerialize(const FArchive& Ar) override;
