Add EOnDemandGetCacheUsageOptions::DumpHandlesToResults to add handle info to FOnDemandInstallCacheUsage

diff --git a/Engine/Source/Runtime/Experimental/IoStore/OnDemand/Private/OnDemandIoStore.cpp b/Engine/Source/Runtime/Experimental/IoStore/OnDemand/Private/OnDemandIoStore.cpp
--- a/Engine/Source/Runtime/Experimental/IoStore/OnDemand/Private/OnDemandIoStore.cpp
+++ b/Engine/Source/Runtime/Experimental/IoStore/OnDemand/Private/OnDemandIoStore.cpp
@@ -1782,43 +1782,64 @@ void FOnDemandIoStore::UpdateInstallRequestPriority(FSharedInternalInstallReques
 
 FOnDemandCacheUsage FOnDemandIoStore::GetCacheUsage(const FOnDemandGetCacheUsageArgs& Args) const
 {
-	if (EnumHasAnyFlags(Args.Options, EOnDemandGetCacheUsageOptions::DumpHandlesToLog))
+	FOnDemandCacheUsage CacheUsage;
+
+	if (InstallCache.IsValid())
 	{
-		TMap<FOnDemandWeakContentHandle, TArray<FOnDemandContainerChunkEntryReferences>> RefsByHandle;
-		GetReferencedContentByHandle(RefsByHandle);
+		CacheUsage.InstallCache = InstallCache->GetCacheUsage();
 
-		for (const TPair<FOnDemandWeakContentHandle, TArray<FOnDemandContainerChunkEntryReferences>>& Kv : RefsByHandle)
+		if (EnumHasAnyFlags(Args.Options, EOnDemandGetCacheUsageOptions::DumpHandlesToLog | EOnDemandGetCacheUsageOptions::DumpHandlesToResults))
 		{
-			const FOnDemandWeakContentHandle& WeakHandle = Kv.Key;
-			const TArray<FOnDemandContainerChunkEntryReferences>& ContainerReferences = Kv.Value;
+			TMap<FOnDemandWeakContentHandle, TArray<FOnDemandContainerChunkEntryReferences>> RefsByHandle;
+			GetReferencedContentByHandle(RefsByHandle);
+
+			if (EnumHasAnyFlags(Args.Options, EOnDemandGetCacheUsageOptions::DumpHandlesToResults))
+			{
+				CacheUsage.InstallCache.ReferencedBytesByHandle.Reserve(RefsByHandle.Num());
+			}
 
-			uint64 ReferencedBytesByHandle = 0;
-			for (const FOnDemandContainerChunkEntryReferences& Refs : ContainerReferences)
+			for (const TPair<FOnDemandWeakContentHandle, TArray<FOnDemandContainerChunkEntryReferences>>& Kv : RefsByHandle)
 			{
-				for (int32 EntryIndex = 0; const FOnDemandChunkEntry & ChunkEntry : Refs.Container->ChunkEntries)
+				const FOnDemandWeakContentHandle& WeakHandle = Kv.Key;
+				const TArray<FOnDemandContainerChunkEntryReferences>& ContainerReferences = Kv.Value;
+
+				uint64 ReferencedBytesByHandle = 0;
+				for (const FOnDemandContainerChunkEntryReferences& Refs : ContainerReferences)
 				{
-					if (Refs.Indices[EntryIndex])
+					for (int32 EntryIndex = 0; const FOnDemandChunkEntry & ChunkEntry : Refs.Container->ChunkEntries)
 					{
-						ReferencedBytesByHandle += ChunkEntry.EncodedSize;
+						if (Refs.Indices[EntryIndex])
+						{
+							ReferencedBytesByHandle += ChunkEntry.EncodedSize;
+						}
+						EntryIndex++;
 					}
-					EntryIndex++;
 				}
-			}
 
-			if (ReferencedBytesByHandle > 0)
-			{
-				UE_LOG(LogIoStoreOnDemand, Log, TEXT("HandleId=0x%llX, DebugName='%s', ReferencedBytes=%.2lf KiB"),
-					WeakHandle.HandleId, *WeakHandle.DebugName, double(ReferencedBytesByHandle) / 1024.0);
+				if (ReferencedBytesByHandle > 0)
+				{
+					if (EnumHasAnyFlags(Args.Options, EOnDemandGetCacheUsageOptions::DumpHandlesToLog))
+					{
+						UE_LOG(LogIoStoreOnDemand, Log, TEXT("HandleId=0x%llX, DebugName='%s', ReferencedBytes=%.2lf KiB"),
+							WeakHandle.HandleId, *WeakHandle.DebugName, double(ReferencedBytesByHandle) / 1024.0);
+					}
+
+					if (EnumHasAnyFlags(Args.Options, EOnDemandGetCacheUsageOptions::DumpHandlesToResults))
+					{
+						CacheUsage.InstallCache.ReferencedBytesByHandle.Add(
+							FOnDemandInstallHandleCacheUsage
+							{
+								.HandleId = WeakHandle.HandleId,
+								.DebugName = WeakHandle.DebugName,
+								.ReferencedBytes = ReferencedBytesByHandle
+							}
+						);
+					}
+				}
 			}
 		}
 	}
 
-	FOnDemandCacheUsage CacheUsage;
-	if (InstallCache.IsValid())
-	{
-		CacheUsage.InstallCache = InstallCache->GetCacheUsage();
-	}
-
 	if (StreamingBackend.IsValid())
 	{
 		CacheUsage.StreamingCache = StreamingBackend->GetCacheUsage();
diff --git a/Engine/Source/Runtime/Experimental/IoStore/OnDemandCore/Public/IO/IoStoreOnDemand.h b/Engine/Source/Runtime/Experimental/IoStore/OnDemandCore/Public/IO/IoStoreOnDemand.h
--- a/Engine/Source/Runtime/Experimental/IoStore/OnDemandCore/Public/IO/IoStoreOnDemand.h
+++ b/Engine/Source/Runtime/Experimental/IoStore/OnDemandCore/Public/IO/IoStoreOnDemand.h
@@ -399,6 +399,8 @@ enum class EOnDemandGetCacheUsageOptions
 	None = 0,
 	/** Dump Handle information to log */
 	DumpHandlesToLog = 1 << 0,
+	/** Dump Handle information to returned FOnDemandInstallCacheUsage */
+	DumpHandlesToResults = 1 << 1,
 };
 ENUM_CLASS_FLAGS(EOnDemandGetCacheUsageOptions);
 
@@ -408,6 +410,13 @@ struct FOnDemandGetCacheUsageArgs
 	EOnDemandGetCacheUsageOptions Options = EOnDemandGetCacheUsageOptions::None;
 };
 
+struct FOnDemandInstallHandleCacheUsage
+{
+	UPTRINT			HandleId = 0;
+	FSharedString	DebugName;
+	uint64			ReferencedBytes = 0;
+};
+
 /** Holds information about the install cache usage */
 struct FOnDemandInstallCacheUsage
 {
@@ -416,6 +425,9 @@ struct FOnDemandInstallCacheUsage
 	uint64 ReferencedBlockSize = 0;
 	uint64 FragmentedChunksSize = 0;
 
+	// Only populated if DumpHandlesToResults is specified
+	TArray<FOnDemandInstallHandleCacheUsage> ReferencedBytesByHandle;
+
 	UE_API friend FStringBuilderBase& operator<<(FStringBuilderBase& Sb, const FOnDemandInstallCacheUsage& CacheUsage);
 };
 
