Chaos:: Fix wrong rotation axis when using simd for linear joint solver.

diff --git a/Engine/Source/Runtime/Experimental/Chaos/Private/Chaos/Joint/PBDJointCachedSolverGaussSeidel.cpp b/Engine/Source/Runtime/Experimental/Chaos/Private/Chaos/Joint/PBDJointCachedSolverGaussSeidel.cpp
--- a/Engine/Source/Runtime/Experimental/Chaos/Private/Chaos/Joint/PBDJointCachedSolverGaussSeidel.cpp
+++ b/Engine/Source/Runtime/Experimental/Chaos/Private/Chaos/Joint/PBDJointCachedSolverGaussSeidel.cpp
@@ -2616,8 +2616,8 @@ void FPBDJointCachedSolver::InitSwingTwistDrives(
 
 	if (RotationDrives.bUseSimd)
 	{
-		FVec3 ConstraintAxes[3] = { ConnectorRs[1] * FJointConstants::TwistAxis(), ConnectorRs[1] * FJointConstants::Swing1Axis(), ConnectorRs[1] * FJointConstants::Swing2Axis() };
-		InitRotationConstraintDriveSimd(ConstraintAxes, FRealSingle(Dt), FVec3(DTwistAngle, DSwing1Angle, DSwing2Angle));
+		FVec3 ConstraintAxes[3] = { ConnectorRs[1] * FJointConstants::TwistAxis(), ConnectorRs[1] * FJointConstants::Swing2Axis(), ConnectorRs[1] * FJointConstants::Swing1Axis() };
+		InitRotationConstraintDriveSimd(ConstraintAxes, FRealSingle(Dt), FVec3(DTwistAngle, DSwing2Angle, DSwing1Angle));
 
 		RotationDrives.Simd.ConstraintHardStiffness = VectorZeroFloat();
 		RotationDrives.Simd.ConstraintLambda = VectorZeroFloat();
