Enable some p.Chaos cvars in Shipping so they can be used to tweak game-specific performance.

diff --git a/Engine/Source/Runtime/Experimental/Chaos/Private/Chaos/Framework/Parallel.cpp b/Engine/Source/Runtime/Experimental/Chaos/Private/Chaos/Framework/Parallel.cpp
--- a/Engine/Source/Runtime/Experimental/Chaos/Private/Chaos/Framework/Parallel.cpp
+++ b/Engine/Source/Runtime/Experimental/Chaos/Private/Chaos/Framework/Parallel.cpp
@@ -14,7 +14,6 @@ namespace Chaos
 	CHAOS_API int32 SmallBatchSize = 10;
 	CHAOS_API int32 LargeBatchSize = 100;
 	
-#if !UE_BUILD_SHIPPING
 	CHAOS_API bool bDisablePhysicsParallelFor = false;
 	CHAOS_API bool bDisableParticleParallelFor = false;
 	CHAOS_API bool bDisableCollisionParallelFor = false;
@@ -27,7 +26,6 @@ namespace Chaos
 	FAutoConsoleVariableRef CVarMaxRangeBatchWorkers(TEXT("p.Chaos.MaxNumWorkers"), MaxNumWorkers, TEXT("Set the max number of workers for physics"));
 	FAutoConsoleVariableRef CVarSmallBatchSize(TEXT("p.Chaos.SmallBatchSize"), SmallBatchSize, TEXT("Small batch size for chaos parallel loops"));
 	FAutoConsoleVariableRef CVarLargeBatchSize(TEXT("p.Chaos.LargeBatchSize"), LargeBatchSize, TEXT("Large batch size for chaos parallel loops"));
-#endif
 }
 
 void Chaos::InnerPhysicsParallelFor(int32 Num, TFunctionRef<void(int32)> InCallable, bool bForceSingleThreaded)
diff --git a/Engine/Source/Runtime/Experimental/Chaos/Public/Chaos/Framework/Parallel.h b/Engine/Source/Runtime/Experimental/Chaos/Public/Chaos/Framework/Parallel.h
--- a/Engine/Source/Runtime/Experimental/Chaos/Public/Chaos/Framework/Parallel.h
+++ b/Engine/Source/Runtime/Experimental/Chaos/Public/Chaos/Framework/Parallel.h
@@ -17,13 +17,7 @@ namespace Chaos
 	CHAOS_API extern int32 MaxNumWorkers;
 	CHAOS_API extern int32 SmallBatchSize;
 	CHAOS_API extern int32 LargeBatchSize;
-#if UE_BUILD_SHIPPING
-	const bool bDisablePhysicsParallelFor = false;
-	const bool bDisableParticleParallelFor = false;
-	const bool bDisableCollisionParallelFor = false;
-#else
 	CHAOS_API extern bool bDisablePhysicsParallelFor;
 	CHAOS_API extern bool bDisableParticleParallelFor;
 	CHAOS_API extern bool bDisableCollisionParallelFor;
-#endif
 }
\ No newline at end of file
