Geometry collection : fix potential crash when calling SetLocalRestTransforms with an collection that has no bones

diff --git a/Engine/Source/Runtime/Experimental/GeometryCollectionEngine/Private/GeometryCollection/GeometryCollectionComponent.cpp b/Engine/Source/Runtime/Experimental/GeometryCollectionEngine/Private/GeometryCollection/GeometryCollectionComponent.cpp
--- a/Engine/Source/Runtime/Experimental/GeometryCollectionEngine/Private/GeometryCollection/GeometryCollectionComponent.cpp
+++ b/Engine/Source/Runtime/Experimental/GeometryCollectionEngine/Private/GeometryCollection/GeometryCollectionComponent.cpp
@@ -6902,7 +6902,7 @@ void UGeometryCollectionComponent::SetLocalRestTransforms(const TArray<FTransfor
 	{
 		// get the current transform and will use this array to update
 		TArray<FTransform> LocalTransforms = GetInitialLocalRestTransforms();
-		if (LocalTransforms.Num() == NewTransforms.Num())
+		if (LocalTransforms.Num() == NewTransforms.Num() && NewTransforms.Num() > 0)
 		{
 			FGeometryCollection& GeometryCollection = *RestCollection->GetGeometryCollection();
 			const TManagedArray<FTransform>* MassToLocal = GeometryCollection.FindAttribute<FTransform>("MassToLocal", FGeometryCollection::TransformGroup);
