Chaos Field: fix issues with computing two level deep diosabled child particle sample position
- make sure to gather all relative transform up to the non-disabled parent instead of assuming that the first parent will be enough

diff --git a/Engine/Source/Runtime/Experimental/Chaos/Public/PhysicsProxy/FieldSystemProxyHelper.h b/Engine/Source/Runtime/Experimental/Chaos/Public/PhysicsProxy/FieldSystemProxyHelper.h
--- a/Engine/Source/Runtime/Experimental/Chaos/Public/PhysicsProxy/FieldSystemProxyHelper.h
+++ b/Engine/Source/Runtime/Experimental/Chaos/Public/PhysicsProxy/FieldSystemProxyHelper.h
@@ -91,14 +91,21 @@ namespace Chaos
 					// the disabled child particles into world space.
 					if (ClusterHandle->Disabled() == true)
 					{
-						if (FPBDRigidClusteredParticleHandle* ParentHandle = ClusterHandle->Parent())
+						FRigidTransform3 RelativeFrame = ClusterHandle->ChildToParent();
+
+						// go up the chain until no parent or a non-disabled parent ( should normally stop at a non disabled one )
+						FPBDRigidClusteredParticleHandle* ParentHandle = ClusterHandle->Parent();
+						while (ParentHandle && ParentHandle->Disabled())
+						{
+							RelativeFrame = RelativeFrame * ParentHandle->ChildToParent();
+							ParentHandle = ParentHandle->Parent();
+						}
+
+						if (ParentHandle && ParentHandle->Disabled() == false)
 						{
-							if (ParentHandle->Disabled() == false)
-							{
-								const FRigidTransform3 ParentWorldTM(ParentHandle->GetP(), ParentHandle->GetQ());
-								const FRigidTransform3 ChildFrame = ClusterHandle->ChildToParent() * ParentWorldTM;
-								SamplePosition = ChildFrame.GetTranslation();
-							}
+							const FRigidTransform3 ParentWorldTM(ParentHandle->GetP(), ParentHandle->GetQ());
+							const FRigidTransform3 ChildFrame = RelativeFrame * ParentWorldTM;
+							SamplePosition = ChildFrame.GetTranslation();
 						}
 					}
 				}
