Chaos - fix debug draw of optimized convexes (not working since the change to the new debug draw backend)

#rnx

diff --git a/Engine/Source/Runtime/Experimental/Chaos/Private/Chaos/DebugDraw/DebugDrawParticle.cpp b/Engine/Source/Runtime/Experimental/Chaos/Private/Chaos/DebugDraw/DebugDrawParticle.cpp
--- a/Engine/Source/Runtime/Experimental/Chaos/Private/Chaos/DebugDraw/DebugDrawParticle.cpp
+++ b/Engine/Source/Runtime/Experimental/Chaos/Private/Chaos/DebugDraw/DebugDrawParticle.cpp
@@ -33,6 +33,7 @@ namespace Chaos::DebugDraw
 	extern bool bChaosDebugDebugDrawShowSimOnlyShapes;
 	extern bool bChaosDebugDebugDrawShowProbeOnlyShapes;
 
+	extern bool bChaosDebugDebugDrawShowOptimizedConvexes;
 }
 
 namespace Chaos::CVars
@@ -103,6 +104,7 @@ namespace Chaos
 		FChaosDDParticleData(
 			const FConstGenericParticleHandle& InParticle)
 		{
+			check(InParticle.IsValid());
 			IslandId = (InParticle->GetConstraintGraphNode() != nullptr) ? InParticle->GetConstraintGraphNode()->GetIslandId() : INDEX_NONE;
 			ClusterId = (InParticle->CastToClustered() != nullptr) ? FMath::Abs(InParticle->CastToClustered()->ClusterGroupIndex()) : INDEX_NONE;
 			ObjectState = InParticle->ObjectState();
@@ -118,16 +120,20 @@ namespace Chaos
 		uint8 bIsQuery: 1;
 		uint8 bIsSim: 1;
 		uint8 bIsProbe: 1;
-		uint8 bIsOptimized: 1;
+		uint8 bIsNonOptimized : 1;
+		uint8 bIsOptimized : 1;
 
 		FChaosDDShapeData(
-			const FShapeInstancePtr& InShapeInstance, 
+			const FShapeInstance* InShapeInstance,
+			bool bInIsNonOptimized,
 			bool bInIsOptimized)
 		{
+			check(InShapeInstance != nullptr);
 			CollisionTraceFlag = InShapeInstance->GetCollisionTraceType();
 			bIsQuery = InShapeInstance->GetQueryEnabled();
 			bIsSim = InShapeInstance->GetSimEnabled();
 			bIsProbe = InShapeInstance->GetIsProbe();
+			bIsNonOptimized = bInIsNonOptimized;
 			bIsOptimized = bInIsOptimized;
 		}
 	};
@@ -175,6 +181,16 @@ namespace Chaos
 				return false;
 			}
 
+			if (!ShapeData.bIsNonOptimized && !DebugDraw::bChaosDebugDebugDrawShowOptimizedConvexes)
+			{
+				return false;
+			}
+
+			if (!ShapeData.bIsOptimized && DebugDraw::bChaosDebugDebugDrawShowOptimizedConvexes)
+			{
+				return false;
+			}
+
 			// Depending on the shape settings, we may not show the simple or complex shape
 			const EImplicitObjectType InnerType = GetInnerType(Implicit->GetType());
 			const bool bIsMesh = (InnerType == ImplicitObjectType::TriangleMesh);
@@ -241,7 +257,7 @@ namespace Chaos
 			{
 				if (InnerType == ImplicitObjectType::Convex)
 				{
-					if (ShapeData.bIsOptimized)
+					if (!ShapeData.bIsNonOptimized)
 					{
 						return FColor::Green;
 					}
@@ -270,13 +286,25 @@ namespace Chaos
 		}
 
 		static void Draw(
-			const FRigidTransform3& SpaceTransform, 
+			const FRigidTransform3& SpaceTransform,
 			const FConstGenericParticleHandle& InParticle,
-			const FShapeInstancePtr& InShapeInstance,
+			const FShapeInstance* InShapeInstance,
+			const FImplicitObject* InImplicitObject,
+			const FRigidTransform3& InRelativeTransform,
+			bool bInIsNonOptimized,
 			bool bInIsOptimized,
-			bool bAutoColor,
+			bool bInAutoColor,
 			const FColor& InColor = FColor::Purple)
 		{
+			if (!InParticle.IsValid())
+			{
+				return;
+			}
+			if (InShapeInstance == nullptr)
+			{
+				return;
+			}
+
 			const int32 Cost = GetParticleCommandCost(InParticle);
 			const FBox3d Bounds = FBox3d((InParticle->WorldSpaceInflatedBounds().Min()), InParticle->WorldSpaceInflatedBounds().Max()).TransformBy(SpaceTransform);
 
@@ -294,56 +322,91 @@ namespace Chaos
 
 			const FRigidTransform3 ParticleTransform = InParticle->GetTransformPQ() * SpaceTransform;
 			const FChaosDDParticleData ParticleData = FChaosDDParticleData(InParticle);
-			const FChaosDDShapeData ShapeData = FChaosDDShapeData(InShapeInstance, bInIsOptimized);
+			const FChaosDDShapeData ShapeData = FChaosDDShapeData(InShapeInstance, bInIsNonOptimized, bInIsOptimized);
+
+			Writer.EnqueueCommand(
+				[
+					Transform = InRelativeTransform * ParticleTransform,
+					ImplicitObject = FConstImplicitObjectPtr(InImplicitObject),
+					ParticleData,
+					ShapeData,
+					InColor,
+					bInAutoColor
+				]
+				(ChaosDD::Private::IChaosDDRenderer& Renderer)
+				{
+					if (ShouldRender(ImplicitObject.GetReference(), ParticleData, ShapeData, Renderer.IsServer()))
+					{
+						Chaos::Private::ChaosDDRenderImplicitObject(
+							Renderer,
+							ImplicitObject.GetReference(),
+							Transform,
+							bInAutoColor ? GetRenderColor(ImplicitObject.GetReference(), ParticleData, ShapeData, Renderer.IsServer()) : InColor,
+							GetLineThickness(),
+							GetDuration());
+					}
+				});
+		}
+
+		static void Draw(
+			const FRigidTransform3& SpaceTransform, 
+			const FConstGenericParticleHandle& InParticle,
+			const FShapeInstance* InShapeInstance,
+			bool bInIsNonOptimized,
+			bool bInIsOptimized,
+			bool bInAutoColor,
+			const FColor& InColor = FColor::Purple)
+		{
+			if (!InParticle.IsValid())
+			{
+				return;
+			}
+			if (InShapeInstance == nullptr)
+			{
+				return;
+			}
+
+			const FBox3d Bounds = FBox3d((InParticle->WorldSpaceInflatedBounds().Min()), InParticle->WorldSpaceInflatedBounds().Max()).TransformBy(SpaceTransform);
+
+			ChaosDD::Private::FChaosDDFrameWriter Writer = ChaosDD::Private::FChaosDDContext::GetWriter();
+
+			if (!Writer.IsInDrawRegion(Bounds))
+			{
+				return;
+			}
+
+			const FChaosDDParticleData ParticleData = FChaosDDParticleData(InParticle);
+			const FChaosDDShapeData ShapeData = FChaosDDShapeData(InShapeInstance, bInIsNonOptimized, bInIsOptimized);
 
 			InShapeInstance->GetGeometry()->VisitLeafObjects(
-				[&Writer, &ParticleTransform, &ParticleData, &ShapeData, &InColor, &bAutoColor](const FImplicitObject* LeafImplicitObject, const FRigidTransform3& LeafRelativeTransform, const int32 UnusedRootObjectIndex, const int32 UnusedObjectIndex, const int32 UnusedLeafObjectIndex)
+				[&Writer, &SpaceTransform, InParticle, InShapeInstance, bInIsNonOptimized, bInIsOptimized, bInAutoColor, &InColor](const FImplicitObject* LeafImplicitObject, const FRigidTransform3& LeafRelativeTransform, const int32 UnusedRootObjectIndex, const int32 UnusedObjectIndex, const int32 UnusedLeafObjectIndex)
 				{
-					Writer.EnqueueCommand(
-						[
-							Transform = LeafRelativeTransform * ParticleTransform,
-							ImplicitObject = Chaos::FConstImplicitObjectPtr(LeafImplicitObject, /*bAddRef*/true),
-							ParticleData,
-							ShapeData,
-							Color = InColor,
-							bAutoColor
-						]
-						(ChaosDD::Private::IChaosDDRenderer& Renderer)
-						{
-							if (ShouldRender(ImplicitObject, ParticleData, ShapeData, Renderer.IsServer()))
-							{
-								Chaos::Private::ChaosDDRenderImplicitObject(
-									Renderer,
-									ImplicitObject,
-									Transform,
-									bAutoColor ? GetRenderColor(ImplicitObject, ParticleData, ShapeData, Renderer.IsServer()) : Color,
-									GetLineThickness(),
-									GetDuration());
-							}
-						});
-
+					FChaosDDParticleShape::Draw(SpaceTransform, InParticle, InShapeInstance, LeafImplicitObject, LeafRelativeTransform, bInIsNonOptimized, bInIsOptimized, bInAutoColor, InColor);
 				});
 
 		}
+
 	};
 
 	void FChaosDDParticle::DrawShapes(
 		const FRigidTransform3& SpaceTransform,
 		const FGeometryParticleHandle* InParticleHandle)
 	{
+		// Record the optimized geometry if there is any.
+		// If there is no optimized geo, we report the regular geometry as the optimized geometry as well 
+		// as the non-optimized so that DebugDrawOptimized will show regular geo instead.
+		bool bHasOptimizedShapes = DrawOptimizedShapes(InParticleHandle);
+
 		for (const Chaos::FShapeInstancePtr& ShapeInstance : InParticleHandle->ShapeInstances())
 		{
-			FChaosDDParticleShape::Draw(SpaceTransform, InParticleHandle, ShapeInstance, /*bIsOptimized*/false, /*bAutoColor*/true);
+			FChaosDDParticleShape::Draw(SpaceTransform, InParticleHandle, ShapeInstance.Get(), /*bIsNonOptimized*/true, /*bIsOptimized*/!bHasOptimizedShapes, /*bAutoColor*/true);
 		}
 	}
 
 	void FChaosDDParticle::DrawShapes(
 		const FGeometryParticleHandle* InParticleHandle)
 	{
-		for (const Chaos::FShapeInstancePtr& ShapeInstance : InParticleHandle->ShapeInstances())
-		{
-			FChaosDDParticleShape::Draw(FRigidTransform3::Identity, InParticleHandle, ShapeInstance, /*bIsOptimized*/false, /*bAutoColor*/true);
-		}
+		FChaosDDParticle::DrawShapes(FRigidTransform3::Identity, InParticleHandle);
 	}
 
 	void FChaosDDParticle::DrawShapes(
@@ -352,7 +415,7 @@ namespace Chaos
 	{
 		for (const Chaos::FShapeInstancePtr& ShapeInstance : InParticleHandle->ShapeInstances())
 		{
-			FChaosDDParticleShape::Draw(FRigidTransform3::Identity, InParticleHandle, ShapeInstance, /*bIsOptimized*/false, /*bAutoColor*/false, Color);
+			FChaosDDParticleShape::Draw(FRigidTransform3::Identity, InParticleHandle, ShapeInstance.Get(), /*bIsNonOptimized*/true, /*bIsOptimized*/false, /*bAutoColor*/false, Color);
 		}
 	}
 
@@ -360,15 +423,34 @@ namespace Chaos
 		const FGeometryParticleHandle* InParticleHandle)
 	{
 		const Private::FConvexOptimizer* ConvexOptimizer = InParticleHandle->CastToClustered() ? InParticleHandle->CastToClustered()->ConvexOptimizer().Get() : nullptr;
-		if (ConvexOptimizer != nullptr)
+		if ((ConvexOptimizer == nullptr) || !ConvexOptimizer->IsValid())
 		{
-			for (const Chaos::FShapeInstancePtr& ShapeInstance : ConvexOptimizer->GetShapeInstances())
-			{
-				FChaosDDParticleShape::Draw(FRigidTransform3::Identity, InParticleHandle, ShapeInstance, /*bIsOptimized*/true, /*bAutoColor*/true);
-			}
-			return true;
+			return false;
 		}
-		return false;
+
+		const auto& DrawOptimizedConvex = 
+			[InParticleHandle, ConvexOptimizer](const FImplicitObject* ImplicitObject, const FRigidTransform3& RelativeTransform, const int32 RootObjectIndex, const int32 ObjectIndex, const int32 LeafObjectIndex)
+			{
+				// The convex optimizer may reuse ShapeInstances from the particle, but with a different implicit object
+				// It's like it is pretending to be a Union-of-Unions (in which case the shapes in the child unions share 
+				// a ShapeInstance)
+				const FShapeInstance* ShapeInstance = nullptr;
+				if ((RootObjectIndex == INDEX_NONE) && (ConvexOptimizer->GetShapeInstances().Num() > 0))
+				{
+					ShapeInstance = ConvexOptimizer->GetShapeInstances()[0].Get();
+				}
+				else
+				{
+					const int32 ShapeIndex = (InParticleHandle->ShapeInstances().IsValidIndex(RootObjectIndex)) ? RootObjectIndex : 0;
+					ShapeInstance = InParticleHandle->ShapeInstances()[ShapeIndex].Get();
+				}
+
+				FChaosDDParticleShape::Draw(FRigidTransform3::Identity, InParticleHandle, ShapeInstance, ImplicitObject, RelativeTransform, /*bIsNonOptimized*/false, /*bIsOptimized*/true, /*bAutoColor*/true);
+			};
+
+		ConvexOptimizer->VisitCollisionObjects(DrawOptimizedConvex);
+
+		return true;
 	}
 
 }
