Chaos - fix assert in collision detection after disabling Convex Optimization from a Midphase Modifier

Disabling Convex Optimization requires that collisions be destroyed if they were created with convex optimization enabled. Before we destroy collisions we must remove them from the constraint graph. The MidPhase modifier now handles this, but had to be moved from CollisionConstraints to the Evolution which knows about the graph.

#rnx

diff --git a/Engine/Source/Runtime/Experimental/Chaos/Private/Chaos/Collision/ParticlePairMidPhase.cpp b/Engine/Source/Runtime/Experimental/Chaos/Private/Chaos/Collision/ParticlePairMidPhase.cpp
--- a/Engine/Source/Runtime/Experimental/Chaos/Private/Chaos/Collision/ParticlePairMidPhase.cpp
+++ b/Engine/Source/Runtime/Experimental/Chaos/Private/Chaos/Collision/ParticlePairMidPhase.cpp
@@ -1217,6 +1217,37 @@ namespace Chaos
 		NewConstraints.Empty();
 	}
 
+	void FGenericParticlePairMidPhase::SetConvexOptimizationIsActive(bool bSetConvexOptimizationIsActive, TFunctionRef<void(FPBDCollisionConstraint*)> OnDestroyingCollision)
+	{
+		Flags.bIsModified = true;
+		Flags.bIsConvexOptimizationActive = bSetConvexOptimizationIsActive;
+
+		// We need to reset collisions if we enable/disable use of the convex optimizer because we are switching 
+		// to a different set of collision shapes. This is necessary because the "key" used to identify a shape 
+		// would not be unique between the two sets of shapes so we would end up with mismatched collisions in the map.
+		// 
+		// NOTE: Collisions must be removed from the constraint graph before being destroyed. The caller is
+		// responsible for this because we don't know about the graph here (and shouldn't because the solver
+		// may not be using a graph).
+		// @todo(chaos): see if we can clean this up a bit?
+		// 
+		// NOTE: bIsConvexOptimizationActive can be changed in a modifier which gets reset every frame so we
+		// use bIsConvexOptimizationUsed to remember what we actually used rather than just resetting on change.
+		if (Flags.bIsConvexOptimizationActive != Flags.bIsConvexOptimizationUsed)
+		{
+			// The caller should remove collisions from the graph here
+			for (TMap<uint64, FPBDCollisionConstraintPtr>::ElementType& KVP : Constraints)
+			{
+				FPBDCollisionConstraintPtr& Constraint = KVP.Value;
+				OnDestroyingCollision(Constraint.Get());
+			}
+
+			// NOTE: An assert here would indicate that the OnDestroyingCollisions callback did not 
+			// remove collisions from the constraint graph.
+			Constraints.Reset();
+		}
+	}
+
 	void FGenericParticlePairMidPhase::BuildDetectorsImpl()
 	{
 		check(Particle0->GetGeometry() != nullptr);
@@ -1247,20 +1278,14 @@ namespace Chaos
 			? ClusteredHandle1->ConvexOptimizer().Get()
 			: nullptr;
 
-		// We need to reset collisions if we enable/disable use of the convex optimizer
-		// because we are switching to a different set of collision shapes.
 		// bIsConvexOptimizationActive can be changed in a modifier which gets reset every frame so we
-		// need to remember what we actually used rather than just resetting on change.
+		// need to remember what we actually used when generating collisions, so that we know if we
+		// have to destroy those collisions if it changes again (see SetConvexOptimizationIsActive)
 		// NOTE: We don't have to handle the case where bHasConvexOptimizer0 or bHasConvexOptimizer1
 		// changes because we already reset collisions whenever geometry is replaced or rebuilt.
 		const bool bHasConvexOptimizer0 = (ConvexOptimizer0 != nullptr) && ConvexOptimizer0->IsValid();
 		const bool bHasConvexOptimizer1 = (ConvexOptimizer1 != nullptr) && ConvexOptimizer1->IsValid();
-		const bool bShouldUseConvexOptimization = Flags.bIsConvexOptimizationActive && (bHasConvexOptimizer0 || bHasConvexOptimizer1);
-		if (Flags.bIsConvexOptimizationUsed != bShouldUseConvexOptimization)
-		{
-			Flags.bIsConvexOptimizationUsed = bShouldUseConvexOptimization;
-			Constraints.Reset();
-		}
+		Flags.bIsConvexOptimizationUsed = (bHasConvexOptimizer0 || bHasConvexOptimizer1);
 
 		// See if we have a BVH for either/both of the particles
 		const Private::FImplicitBVH* BVH0 = nullptr;
@@ -2018,7 +2043,7 @@ namespace Chaos
 		TArray<uint64, TMemStackAllocator<alignof(uint32)>> Pruned;
 		Pruned.Reserve(Constraints.Num());
 
-		for (auto& KVP : Constraints)
+		for (TMap<uint64, FPBDCollisionConstraintPtr>::ElementType& KVP : Constraints)
 		{
 			const uint64 CollisionKey = KVP.Key;
 			FPBDCollisionConstraintPtr& Constraint = KVP.Value;
@@ -2040,7 +2065,7 @@ namespace Chaos
 
 	void FGenericParticlePairMidPhase::WakeCollisionsImpl(const int32 CurrentEpoch)
 	{
-		for (auto& KVP : Constraints)
+		for (TMap<uint64, FPBDCollisionConstraintPtr>::ElementType& KVP : Constraints)
 		{
 			FPBDCollisionConstraintPtr& Constraint = KVP.Value;
 			if (Constraint->GetContainerCookie().LastUsedEpoch >= LastUsedEpoch)
diff --git a/Engine/Source/Runtime/Experimental/Chaos/Private/Chaos/MidPhaseModification.cpp b/Engine/Source/Runtime/Experimental/Chaos/Private/Chaos/MidPhaseModification.cpp
--- a/Engine/Source/Runtime/Experimental/Chaos/Private/Chaos/MidPhaseModification.cpp
+++ b/Engine/Source/Runtime/Experimental/Chaos/Private/Chaos/MidPhaseModification.cpp
@@ -6,6 +6,7 @@
 #include "Chaos/Collision/ParticlePairMidPhase.h"
 #include "Chaos/Particle/ParticleUtilities.h"
 #include "Chaos/PBDCollisionConstraints.h"
+#include "Chaos/PBDRigidsEvolutionGBF.h"
 #include "Chaos/ParticleHandleFwd.h"
 
 namespace Chaos
@@ -22,7 +23,18 @@ namespace Chaos
 
 	void FMidPhaseModifier::DisableConvexOptimization()
 	{
-		MidPhase->SetConvexOptimizationIsActive(false);
+		// SetConvexOptimizationIsActive may destroy constraints. If it does, we need to remove
+		// them from the constraint graph first.
+		// NOTE: We only destroy collisions if convex optimization use was changed. Otherwise
+		// we would end up unnecessarily removing collisions every frame if DisableConvexOptimization()
+		// is called repreatedly.
+		const auto& RemoveCollisionFromGraph = [this](FPBDCollisionConstraint* InConstraint) -> ECollisionVisitorResult
+			{
+				Accessor->GetEvolution().RemoveConstraintFromConstraintGraph(InConstraint->GetConstraintHandle());
+				return ECollisionVisitorResult::Continue;
+			};
+
+		MidPhase->SetConvexOptimizationIsActive(false, RemoveCollisionFromGraph);
 	}
 
 	void FMidPhaseModifier::GetParticles(
@@ -116,7 +128,7 @@ namespace Chaos
 	{
 		// This visitor is a thin function wrapper so that users can't directly access
 		// the midphase itself, but instead can access a modifier.
-		ConstraintAllocator.VisitMidPhases([this, &Visitor](FParticlePairMidPhase& MidPhase)
+		Evolution.GetCollisionConstraints().GetConstraintAllocator().VisitMidPhases([this, &Visitor](FParticlePairMidPhase& MidPhase)
 		{
 			FMidPhaseModifier MidPhaseModifier(&MidPhase, this);
 			Visitor(MidPhaseModifier);
diff --git a/Engine/Source/Runtime/Experimental/Chaos/Private/Chaos/PBDCollisionConstraints.cpp b/Engine/Source/Runtime/Experimental/Chaos/Private/Chaos/PBDCollisionConstraints.cpp
--- a/Engine/Source/Runtime/Experimental/Chaos/Private/Chaos/PBDCollisionConstraints.cpp
+++ b/Engine/Source/Runtime/Experimental/Chaos/Private/Chaos/PBDCollisionConstraints.cpp
@@ -5,7 +5,6 @@
 #include "Chaos/CastingUtilities.h"
 #include "Chaos/ChaosPerfTest.h"
 #include "Chaos/ContactModification.h"
-#include "Chaos/MidPhaseModification.h"
 #include "Chaos/CCDModification.h"
 #include "Chaos/CollisionResolution.h"
 #include "Chaos/Collision/CollisionPruning.h"
@@ -611,15 +610,6 @@ namespace Chaos
 		}
 	}
 
-	void FPBDCollisionConstraints::ApplyMidPhaseModifier(const TArray<ISimCallbackObject*>& MidPhaseModifiers, FReal Dt)
-	{
-		FMidPhaseModifierAccessor ModifierAccessor(GetConstraintAllocator());
-		for(ISimCallbackObject* ModifierCallback : MidPhaseModifiers)
-		{
-			ModifierCallback->MidPhaseModification_Internal(ModifierAccessor);
-		}
-	}
-
 	void FPBDCollisionConstraints::ApplyCCDModifier(const TArray<ISimCallbackObject*>& CCDModifiers, FReal Dt)
 	{
 		FCCDModifierAccessor ModifierAccessor(Dt);
diff --git a/Engine/Source/Runtime/Experimental/Chaos/Private/Chaos/PBDRigidsEvolutionGBF.cpp b/Engine/Source/Runtime/Experimental/Chaos/Private/Chaos/PBDRigidsEvolutionGBF.cpp
--- a/Engine/Source/Runtime/Experimental/Chaos/Private/Chaos/PBDRigidsEvolutionGBF.cpp
+++ b/Engine/Source/Runtime/Experimental/Chaos/Private/Chaos/PBDRigidsEvolutionGBF.cpp
@@ -8,6 +8,7 @@
 #include "Chaos/ImplicitObjectTransformed.h"
 #include "Chaos/ImplicitObjectUnion.h"
 #include "Chaos/MassConditioning.h"
+#include "Chaos/MidPhaseModification.h"
 #include "Chaos/PBDCollisionConstraints.h"
 #include "Chaos/PerParticleEtherDrag.h"
 #include "Chaos/PerParticleEulerStepVelocity.h"
@@ -580,7 +581,7 @@ void FPBDRigidsEvolutionGBF::AdvanceOneTimeStepImpl(const FReal Dt, const FSubSt
 		if (MidPhaseModifiers)
 		{
 			SCOPE_CYCLE_COUNTER(STAT_Evolution_MidPhaseModifierCallback);
-			CollisionConstraints.ApplyMidPhaseModifier(*MidPhaseModifiers, Dt);
+			ApplyMidPhaseModifier(Dt);
 		}
 
 		{
@@ -1091,6 +1092,18 @@ void FPBDRigidsEvolutionGBF::ApplyKinematicTargets(const FReal Dt, const FReal S
 	Particles.UpdateDirtyViews();
 }
 
+void FPBDRigidsEvolutionGBF::ApplyMidPhaseModifier(FReal Dt)
+{
+	if (MidPhaseModifiers != nullptr)
+	{
+		FMidPhaseModifierAccessor MidPhaseModifier(*this);
+		for (ISimCallbackObject* ModifierCallback : *MidPhaseModifiers)
+		{
+			ModifierCallback->MidPhaseModification_Internal(MidPhaseModifier);
+		}
+	}
+}
+
 void FPBDRigidsEvolutionGBF::SetIsDeterministic(const bool bInIsDeterministic)
 {
 	// We detect collisions in parallel, so order is non-deterministic without additional processing
diff --git a/Engine/Source/Runtime/Experimental/Chaos/Public/Chaos/Collision/ParticlePairMidPhase.h b/Engine/Source/Runtime/Experimental/Chaos/Public/Chaos/Collision/ParticlePairMidPhase.h
--- a/Engine/Source/Runtime/Experimental/Chaos/Public/Chaos/Collision/ParticlePairMidPhase.h
+++ b/Engine/Source/Runtime/Experimental/Chaos/Public/Chaos/Collision/ParticlePairMidPhase.h
@@ -289,7 +289,7 @@ namespace Chaos
 		 * @brief Call a lambda on each active collision constraint
 		 * This includes sleeping constraints, but skips constraints that are were not used on the last awake tick
 		 * but are still kept around as an optimization.
-		 * @param Visitor functor with signature ECollisionVisitorResult(FPBDCollisionConstaint& Constraint)
+		 * @param Visitor functor with signature ECollisionVisitorResult(FPBDCollisionConstraint& Constraint)
 		*/
 		template<typename TLambda>
 		ECollisionVisitorResult VisitCollisions(const TLambda& Visitor, const ECollisionVisitorFlags VisitFlags = ECollisionVisitorFlags::VisitDefault);
@@ -297,7 +297,7 @@ namespace Chaos
 
 		/**
 		 * @brief Call a lambda on each active collision constraint
-		 * @param Visitor functor with signature ECollisionVisitorResult(const FPBDCollisionConstaint& Constraint)
+		 * @param Visitor functor with signature ECollisionVisitorResult(const FPBDCollisionConstraint& Constraint)
 		 */
 		template<typename TLambda>
 		ECollisionVisitorResult VisitConstCollisions(const TLambda& Visitor, const ECollisionVisitorFlags VisitFlags = ECollisionVisitorFlags::VisitDefault) const;
@@ -353,12 +353,20 @@ namespace Chaos
 			Flags.bIsCCDActive = bCCDIsActive;
 		}
 
+		UE_DEPRECATED(5.6, "Use SetConvexOptimizationIsActive which takes an OnDestroyingCollision functor")
 		void SetConvexOptimizationIsActive(bool bSetConvexOptimizationIsActive)
 		{
-			Flags.bIsModified = true;
-			Flags.bIsConvexOptimizationActive = bSetConvexOptimizationIsActive;
+			SetConvexOptimizationIsActive(bSetConvexOptimizationIsActive, [](FPBDCollisionConstraint*) {});
 		}
 
+
+		/**
+		 * @brief Override whether convex optimization if used this mid-phase. Used by the MidPhase modifier and gets reset every frame.
+		 * NOTE: when this setting is changed, collisions will be destroyed so any external systems referencing them will need to 
+		 * drop their references.
+		 */
+		CHAOS_API virtual void SetConvexOptimizationIsActive(bool bSetConvexOptimizationIsActive, TFunctionRef<void(FPBDCollisionConstraint*)> OnDestroyingCollision) {}
+
 		/*
 		 * True if CCD is supported by either particle
 		 */
@@ -575,6 +583,8 @@ namespace Chaos
 
 		CHAOS_API bool UpdateCollisionCCD( FPBDCollisionConstraint* Constraint, const FReal CullDistance, const FReal Dt, const FCollisionContext& Context);
 
+		CHAOS_API virtual void SetConvexOptimizationIsActive(bool bSetConvexOptimizationIsActive, TFunctionRef<void(FPBDCollisionConstraint*)> OnDestroyingCollision) override final;
+
 		TMap<uint64, FPBDCollisionConstraintPtr> Constraints;
 		TArray<FPBDCollisionConstraint*> NewConstraints;
 	};
diff --git a/Engine/Source/Runtime/Experimental/Chaos/Public/Chaos/MidPhaseModification.h b/Engine/Source/Runtime/Experimental/Chaos/Public/Chaos/MidPhaseModification.h
--- a/Engine/Source/Runtime/Experimental/Chaos/Public/Chaos/MidPhaseModification.h
+++ b/Engine/Source/Runtime/Experimental/Chaos/Public/Chaos/MidPhaseModification.h
@@ -156,8 +156,8 @@ namespace Chaos
 	class FMidPhaseModifierAccessor
 	{
 	public:
-		FMidPhaseModifierAccessor(Private::FCollisionConstraintAllocator& InConstraintAllocator)
-			: ConstraintAllocator(InConstraintAllocator)
+		FMidPhaseModifierAccessor(FPBDRigidsEvolutionGBF& InEvolution)
+			: Evolution(InEvolution)
 		{ }
 
 		// Get an object which allows for range iteration over the list of
@@ -171,7 +171,10 @@ namespace Chaos
 		CHAOS_API void VisitMidPhases(const TFunction<void(FMidPhaseModifier&)>& Visitor);
 
 	private:
+		friend class FMidPhaseModifier;
 
-		Private::FCollisionConstraintAllocator& ConstraintAllocator;
+		FPBDRigidsEvolutionGBF& GetEvolution() { return Evolution; }
+
+		FPBDRigidsEvolutionGBF& Evolution;
 	};
 }
diff --git a/Engine/Source/Runtime/Experimental/Chaos/Public/Chaos/PBDCollisionConstraints.h b/Engine/Source/Runtime/Experimental/Chaos/Public/Chaos/PBDCollisionConstraints.h
--- a/Engine/Source/Runtime/Experimental/Chaos/Public/Chaos/PBDCollisionConstraints.h
+++ b/Engine/Source/Runtime/Experimental/Chaos/Public/Chaos/PBDCollisionConstraints.h
@@ -121,7 +121,8 @@ public:
 	/**
 	 * Apply modifiers to particle pair midphases
 	 */
-	CHAOS_API void ApplyMidPhaseModifier(const TArray<ISimCallbackObject*>& MidPhaseModifiers, FReal Dt);
+	UE_DEPRECATED(5.6, "Not supported. Use FPBDRigidsEvolutionGBF::ApplyMidPhaseModifier")
+	CHAOS_API void ApplyMidPhaseModifier(const TArray<ISimCallbackObject*>& MidPhaseModifiers, FReal Dt) {}
 
 	/**
 	 * Apply modifiers to CCD results
diff --git a/Engine/Source/Runtime/Experimental/Chaos/Public/Chaos/PBDRigidsEvolutionGBF.h b/Engine/Source/Runtime/Experimental/Chaos/Public/Chaos/PBDRigidsEvolutionGBF.h
--- a/Engine/Source/Runtime/Experimental/Chaos/Public/Chaos/PBDRigidsEvolutionGBF.h
+++ b/Engine/Source/Runtime/Experimental/Chaos/Public/Chaos/PBDRigidsEvolutionGBF.h
@@ -275,7 +275,9 @@ namespace Chaos
 
 	protected:
 
-		CHAOS_API void AdvanceOneTimeStepImpl(const FReal dt, const FSubStepInfo& SubStepInfo);
+		CHAOS_API void AdvanceOneTimeStepImpl(const FReal Dt, const FSubStepInfo& SubStepInfo);
+
+		void ApplyMidPhaseModifier(FReal Dt);
 
 		// Update the particle transform and fix collision anchors (used by client corrections)
 		void ApplyParticleTransformCorrectionImpl(FGeometryParticleHandle* InParticle, const FRigidTransform3& InTransform, const bool bInRecalculateFriction = true);
