Fixing ModuleFactoryRegister so that its static registry is shared across modules.

#rnx

diff --git a/Engine/Source/Runtime/Experimental/ChaosVehicles/ChaosVehiclesCore/Private/SimModule/ModuleFactoryRegister.cpp b/Engine/Source/Runtime/Experimental/ChaosVehicles/ChaosVehiclesCore/Private/SimModule/ModuleFactoryRegister.cpp
--- a/Engine/Source/Runtime/Experimental/ChaosVehicles/ChaosVehiclesCore/Private/SimModule/ModuleFactoryRegister.cpp
+++ b/Engine/Source/Runtime/Experimental/ChaosVehicles/ChaosVehiclesCore/Private/SimModule/ModuleFactoryRegister.cpp
@@ -5,6 +5,13 @@
 
 namespace Chaos
 {
+
+	FModuleFactoryRegister& FModuleFactoryRegister::Get()
+	{
+		static FModuleFactoryRegister Instance;
+		return Instance;
+	}
+
 	void FModuleFactoryRegister::RegisterFactory(const FName TypeName, TWeakPtr<IFactoryModule> InFactory)
 	{
 		RegisterFactory(GetModuleHash(TypeName), InFactory);
diff --git a/Engine/Source/Runtime/Experimental/ChaosVehicles/ChaosVehiclesCore/Public/SimModule/ModuleFactoryRegister.h b/Engine/Source/Runtime/Experimental/ChaosVehicles/ChaosVehiclesCore/Public/SimModule/ModuleFactoryRegister.h
--- a/Engine/Source/Runtime/Experimental/ChaosVehicles/ChaosVehiclesCore/Public/SimModule/ModuleFactoryRegister.h
+++ b/Engine/Source/Runtime/Experimental/ChaosVehicles/ChaosVehiclesCore/Public/SimModule/ModuleFactoryRegister.h
@@ -18,11 +18,7 @@ namespace Chaos
 	class FModuleFactoryRegister
 	{
 	public:
-		static FModuleFactoryRegister& Get()
-		{
-			static FModuleFactoryRegister Instance;
-			return Instance;
-		}
+		UE_API static FModuleFactoryRegister& Get();
 		
 		UE_API void RegisterFactory(const FName TypeName, TWeakPtr<IFactoryModule> InFactory);
 		UE_API void RegisterFactory(const uint32 TypeNameHash, TWeakPtr<IFactoryModule> InFactory);
