IoStoreOnDemand - Move purge/defrag requests to content installer pipe

Per - please replace this with your more robust change.

diff --git a/Engine/Source/Runtime/Experimental/IoStore/OnDemand/Private/OnDemandContentInstaller.cpp b/Engine/Source/Runtime/Experimental/IoStore/OnDemand/Private/OnDemandContentInstaller.cpp
--- a/Engine/Source/Runtime/Experimental/IoStore/OnDemand/Private/OnDemandContentInstaller.cpp
+++ b/Engine/Source/Runtime/Experimental/IoStore/OnDemand/Private/OnDemandContentInstaller.cpp
@@ -318,6 +318,77 @@ void FOnDemandContentInstaller::UpdatePriority(FSharedInternalInstallRequest Ins
 		});
 }
 
+void FOnDemandContentInstaller::EnqueuePurge(FOnDemandPurgeArgs&& Args, FOnDemandPurgeCompleted&& OnCompleted)
+{
+	InstallerPipe.Launch(
+		TEXT("ProcessIoStoreOnDemandPurgeRequest"),
+		[this, Args = MoveTemp(Args), OnCompleted = MoveTemp(OnCompleted), StartTimeCycles = FPlatformTime::Cycles64()]() mutable
+		{
+			const bool bDefrag = EnumHasAnyFlags(Args.Options, EOnDemandPurgeOptions::Defrag);
+			const uint64* BytesToPurge = Args.BytesToPurge.GetPtrOrNull();
+
+			FOnDemandPurgeResult Result;
+			Result.Status = IoStore.InstallCache->PurgeAllUnreferenced(bDefrag, BytesToPurge);
+
+			const uint64 DurationCycles = FPlatformTime::Cycles64() - StartTimeCycles;
+			Result.DurationInSeconds = FPlatformTime::ToSeconds64(DurationCycles);
+
+			if (OnCompleted)
+			{
+				if (EnumHasAnyFlags(Args.Options, EOnDemandPurgeOptions::CallbackOnGameThread))
+				{
+					ExecuteOnGameThread(
+						UE_SOURCE_LOCATION,
+						[OnCompleted = MoveTemp(OnCompleted), Result = MoveTemp(Result)]() mutable
+						{
+							OnCompleted(MoveTemp(Result));
+						});
+				}
+				else
+				{
+					OnCompleted(MoveTemp(Result));
+				}
+			}
+		},
+		UE::Tasks::ETaskPriority::BackgroundLow
+	);
+}
+
+void FOnDemandContentInstaller::EnqueueDefrag(FOnDemandDefragArgs&& Args, FOnDemandDefragCompleted&& OnCompleted)
+{
+	InstallerPipe.Launch(
+		TEXT("ProcessIoStoreOnDemandDefragRequest"),
+		[this, Args = MoveTemp(Args), OnCompleted = MoveTemp(OnCompleted), StartTimeCycles = FPlatformTime::Cycles64()]() mutable
+		{
+			const uint64* BytesToFree = Args.BytesToFree.GetPtrOrNull();
+
+			FOnDemandDefragResult Result;
+			Result.Status = IoStore.InstallCache->DefragAll(BytesToFree);
+
+			const uint64 DurationCycles = FPlatformTime::Cycles64() - StartTimeCycles;
+			Result.DurationInSeconds = FPlatformTime::ToSeconds64(DurationCycles);
+
+			if (OnCompleted)
+			{
+				if (EnumHasAnyFlags(Args.Options, EOnDemandDefragOptions::CallbackOnGameThread))
+				{
+					ExecuteOnGameThread(
+						UE_SOURCE_LOCATION,
+						[OnCompleted = MoveTemp(OnCompleted), Result = MoveTemp(Result)]() mutable
+						{
+							OnCompleted(MoveTemp(Result));
+						});
+				}
+				else
+				{
+					OnCompleted(MoveTemp(Result));
+				}
+			}
+		},
+		UE::Tasks::ETaskPriority::BackgroundLow
+	);
+}
+
 void FOnDemandContentInstaller::ProcessRequest(FInstallRequestState& RequestState)
 {
 	using namespace UE::IoStore::Private;
diff --git a/Engine/Source/Runtime/Experimental/IoStore/OnDemand/Private/OnDemandContentInstaller.h b/Engine/Source/Runtime/Experimental/IoStore/OnDemand/Private/OnDemandContentInstaller.h
--- a/Engine/Source/Runtime/Experimental/IoStore/OnDemand/Private/OnDemandContentInstaller.h
+++ b/Engine/Source/Runtime/Experimental/IoStore/OnDemand/Private/OnDemandContentInstaller.h
@@ -98,6 +98,9 @@ public:
 	void							Cancel(FSharedInternalInstallRequest InstallRequest);
 	void							UpdatePriority(FSharedInternalInstallRequest InstallRequest, int32 NewPriority);
 
+	void							EnqueuePurge(FOnDemandPurgeArgs&& Args, FOnDemandPurgeCompleted&& OnCompleted);
+	void							EnqueueDefrag(FOnDemandDefragArgs&& Args, FOnDemandDefragCompleted&& OnCompleted);
+
 private:
 	void							ProcessRequest(FInstallRequestState& RequestState);
 	void							TryStartDownloadChunks();
diff --git a/Engine/Source/Runtime/Experimental/IoStore/OnDemand/Private/OnDemandIoStore.cpp b/Engine/Source/Runtime/Experimental/IoStore/OnDemand/Private/OnDemandIoStore.cpp
--- a/Engine/Source/Runtime/Experimental/IoStore/OnDemand/Private/OnDemandIoStore.cpp
+++ b/Engine/Source/Runtime/Experimental/IoStore/OnDemand/Private/OnDemandIoStore.cpp
@@ -567,36 +567,12 @@ FOnDemandInstallRequest FOnDemandIoStore::Install(
 
 void FOnDemandIoStore::Purge(FOnDemandPurgeArgs&& Args, FOnDemandPurgeCompleted&& OnCompleted)
 {
-	// TODO: redirect to installer?
-
-	LLM_SCOPE_BYNAME(TEXT("FileSystem/OnDemandIoStore"));
-	FSharedPurgeRequest PurgeRequest	= MakeShared<FPurgeRequest>();
-	PurgeRequest->Args					= MoveTemp(Args);
-	PurgeRequest->OnCompleted			= MoveTemp(OnCompleted);
-
-	{
-		UE::TUniqueLock Lock(RequestMutex);
-		PurgeRequests.Add(MoveTemp(PurgeRequest));
-	}
-
-	TryEnterTickLoop();
+	Installer->EnqueuePurge(MoveTemp(Args), MoveTemp(OnCompleted));
 }
 
 void FOnDemandIoStore::Defrag(FOnDemandDefragArgs&& Args, FOnDemandDefragCompleted&& OnCompleted)
 {
-	// TODO: redirect to installer?
-
-	LLM_SCOPE_BYNAME(TEXT("FileSystem/OnDemandIoStore"));
-	FSharedDefragRequest DefragRequest	= MakeShared<FDefragRequest>();
-	DefragRequest->Args					= MoveTemp(Args);
-	DefragRequest->OnCompleted			= MoveTemp(OnCompleted);
-
-	{
-		UE::TUniqueLock Lock(RequestMutex);
-		DefragRequests.Add(MoveTemp(DefragRequest));
-	}
-
-	TryEnterTickLoop();
+	Installer->EnqueueDefrag(MoveTemp(Args), MoveTemp(OnCompleted));
 }
 
 FIoStatus FOnDemandIoStore::Unmount(FStringView MountId)
@@ -975,46 +951,6 @@ bool FOnDemandIoStore::Tick()
 			});
 	}
 
-	TArray<FSharedPurgeRequest> LocalPurgeRequests;
-	TArray<FSharedDefragRequest> LocalDefragRequests;
-	{
-		UE::TUniqueLock Lock(RequestMutex);
-		LocalPurgeRequests = MoveTemp(PurgeRequests);
-		PurgeRequests.Empty();
-		LocalDefragRequests = MoveTemp(DefragRequests);
-		DefragRequests.Empty();
-	}
-	
-	// Tick purge requests
-	for (FSharedPurgeRequest& Request : LocalPurgeRequests)
-	{
-		const bool bDefrag = EnumHasAnyFlags(Request->Args.Options, EOnDemandPurgeOptions::Defrag);
-		const uint64* BytesToPurge = Request->Args.BytesToPurge.GetPtrOrNull();
-
-		const double StartTime = FPlatformTime::Seconds();
-		FOnDemandPurgeResult PurgeResult;
-		PurgeResult.Status = InstallCache->PurgeAllUnreferenced(bDefrag, BytesToPurge);
-		PurgeResult.DurationInSeconds = FPlatformTime::Seconds() - StartTime;
-
-		CompletePurgeRequest(*Request, MoveTemp(PurgeResult));
-	}
-
-	// Tick defrag requests
-	if (LocalDefragRequests.IsEmpty() == false)
-	{
-		for (FSharedDefragRequest& Request : LocalDefragRequests)
-		{
-			const uint64* BytesToFree = Request->Args.BytesToFree.GetPtrOrNull();
-			
-			const double StartTime = FPlatformTime::Seconds();
-			FOnDemandDefragResult DefragResult;
-			DefragResult.Status = InstallCache->DefragAll(BytesToFree);
-			DefragResult.DurationInSeconds = FPlatformTime::Seconds() - StartTime;
-
-			CompleteDefragRequest(*Request, MoveTemp(DefragResult));
-		}
-	}
-
 	return bTicked;
 }
 
@@ -1339,52 +1275,6 @@ void FOnDemandIoStore::CompleteMountRequest(FMountRequest& Request, FOnDemandMou
 	}
 }
 
-void FOnDemandIoStore::CompletePurgeRequest(FPurgeRequest& Request, FOnDemandPurgeResult&& Result)
-{
-	if (!Request.OnCompleted)
-	{
-		return;
-	}
-
-	if (EnumHasAnyFlags(Request.Args.Options, EOnDemandPurgeOptions::CallbackOnGameThread))
-	{
-		ExecuteOnGameThread(
-			UE_SOURCE_LOCATION,
-			[OnCompleted = MoveTemp(Request.OnCompleted), Result = MoveTemp(Result)]() mutable
-			{
-				OnCompleted(MoveTemp(Result));
-			});
-	}
-	else
-	{
-		FOnDemandPurgeCompleted OnCompleted = MoveTemp(Request.OnCompleted);
-		OnCompleted(MoveTemp(Result));
-	}
-}
-
-void FOnDemandIoStore::CompleteDefragRequest(FDefragRequest& Request, FOnDemandDefragResult&& Result)
-{
-	if (!Request.OnCompleted)
-	{
-		return;
-	}
-
-	if (EnumHasAnyFlags(Request.Args.Options, EOnDemandDefragOptions::CallbackOnGameThread))
-	{
-		ExecuteOnGameThread(
-			UE_SOURCE_LOCATION,
-			[OnCompleted = MoveTemp(Request.OnCompleted), Result = MoveTemp(Result)]() mutable
-			{
-				OnCompleted(MoveTemp(Result));
-			});
-	}
-	else
-	{
-		FOnDemandDefragCompleted OnCompleted = MoveTemp(Request.OnCompleted);
-		OnCompleted(MoveTemp(Result));
-	}
-}
-
 void FOnDemandIoStore::OnEncryptionKeyAdded(const FGuid& Id, const FAES::FAESKey& Key)
 {
 	TRACE_CPUPROFILER_EVENT_SCOPE(FOnDemandIoStore::OnEncryptionKeyAdded);
diff --git a/Engine/Source/Runtime/Experimental/IoStore/OnDemand/Private/OnDemandIoStore.h b/Engine/Source/Runtime/Experimental/IoStore/OnDemand/Private/OnDemandIoStore.h
--- a/Engine/Source/Runtime/Experimental/IoStore/OnDemand/Private/OnDemandIoStore.h
+++ b/Engine/Source/Runtime/Experimental/IoStore/OnDemand/Private/OnDemandIoStore.h
@@ -74,7 +74,7 @@ struct FOnDemandTagSet
 ///////////////////////////////////////////////////////////////////////////////
 struct FOnDemandChunkEntryReferences
 {
-	UPTRINT		ContentHandleId = 0;	
+	UPTRINT		ContentHandleId = 0;
 	TBitArray<>	Indices;
 };
 
@@ -296,22 +296,6 @@ class FOnDemandIoStore
 
 	using FSharedMountRequest	= TSharedRef<FMountRequest>;
 
-	struct FPurgeRequest
-	{
-		FOnDemandPurgeArgs					Args;
-		FOnDemandPurgeCompleted				OnCompleted;
-	};
-
-	using FSharedPurgeRequest = TSharedRef<FPurgeRequest>;
-
-	struct FDefragRequest
-	{
-		FOnDemandDefragArgs					Args;
-		FOnDemandDefragCompleted			OnCompleted;
-	};
-
-	using FSharedDefragRequest = TSharedRef<FDefragRequest>;
-
 public:
 	UE_NONCOPYABLE(FOnDemandIoStore);
 
@@ -366,11 +350,6 @@ private:
 	bool								Tick();
 	FIoStatus							TickMountRequest(FMountRequest& MountRequest);
 	void								CompleteMountRequest(FMountRequest& MountRequest, FOnDemandMountResult&& MountResult);
-	FOnDemandInstallResult				TickInstallRequest(FOnDemandInternalInstallRequest& InstallRequest);
-	void								CompleteInstallRequest(FSharedInternalInstallRequest InstallRequest, FOnDemandInstallResult&& InstallResult);
-	void								ProgressInstallRequest(const FOnDemandInternalInstallRequest& InstallRequest, const FOnDemandInstallProgress& Progress);
-	void								CompletePurgeRequest(FPurgeRequest& PurgeRequest, FOnDemandPurgeResult&& Result);
-	void								CompleteDefragRequest(FDefragRequest& DefragRequest, FOnDemandDefragResult&& Result);
 	void								OnEncryptionKeyAdded(const FGuid& Id, const FAES::FAESKey& Key);
 	static void							CreateContainersFromToc(
 											FStringView MountId,
@@ -391,8 +370,6 @@ private:
 	mutable UE::FMutex						ContainerMutex;
 
 	TArray<FSharedMountRequest>				MountRequests;
-	TArray<FSharedPurgeRequest>				PurgeRequests;
-	TArray<FSharedDefragRequest>			DefragRequests;
 	UE::FMutex								RequestMutex;
 
 	EOnDemandStreamingOptions				StreamingOptions = EOnDemandStreamingOptions::Default;
