Move Purge and Defrag request handling to FOnDemandContentInstaller.

#rnx

diff --git a/Engine/Source/Runtime/Experimental/IoStore/OnDemand/Private/OnDemandContentInstaller.cpp b/Engine/Source/Runtime/Experimental/IoStore/OnDemand/Private/OnDemandContentInstaller.cpp
--- a/Engine/Source/Runtime/Experimental/IoStore/OnDemand/Private/OnDemandContentInstaller.cpp
+++ b/Engine/Source/Runtime/Experimental/IoStore/OnDemand/Private/OnDemandContentInstaller.cpp
@@ -19,6 +19,7 @@
 #include "OnDemandIoStore.h"
 #include "OnDemandPackageStoreBackend.h"
 #include "Statistics.h"
+#include <atomic>
 
 namespace UE::IoStore
 {
@@ -208,6 +209,9 @@ void ResolveChunksToInstall(
 
 } // namespace Private
 
+////////////////////////////////////////////////////////////////////////////////
+uint32 FOnDemandContentInstaller::FRequest::NextSeqNo = 0;
+
 ////////////////////////////////////////////////////////////////////////////////
 FOnDemandContentInstaller::FOnDemandContentInstaller(FOnDemandIoStore& InIoStore, FOnDemandHttpThread& InHttpClient)
 	: IoStore(InIoStore)
@@ -221,194 +225,206 @@ FOnDemandContentInstaller::~FOnDemandContentInstaller()
 	Shutdown();
 }
 
-FSharedInternalInstallRequest FOnDemandContentInstaller::Enqueue(
+FSharedInternalInstallRequest FOnDemandContentInstaller::EnqueueInstallRequest(
 	FOnDemandInstallArgs&& Args,
 	FOnDemandInstallCompleted&& OnCompleted,
 	FOnDemandInstallProgressed&& OnProgress)
 {
-	FSharedInternalInstallRequest InstallRequest = MakeShared<FOnDemandInternalInstallRequest, ESPMode::ThreadSafe>(
-		MoveTemp(Args),
-		MoveTemp(OnCompleted),
-		MoveTemp(OnProgress));
-	
-	FInstallRequestState* RequestState = nullptr;
+	FRequest* Request = nullptr;
 	{
-		UE::TUniqueLock Lock(Mutex);
-		RequestState			= RequestStateAllocator.Construct();
-		RequestState->SeqNo		= NextInstallSeqNo++;
-		RequestState->Request	= InstallRequest;
-
-		RequestStateByHandle.Add(InstallRequest, RequestState);
+		TUniqueLock Lock(Mutex);
+		Request = RequestAllocator.Construct(MoveTemp(Args), MoveTemp(OnCompleted), MoveTemp(OnProgress));
 	}
 
+	FSharedInternalInstallRequest InstallRequest = MakeShared<FOnDemandInternalInstallRequest, ESPMode::ThreadSafe>(UPTRINT(Request));
+	Request->AsInstall().Request = InstallRequest;
+
 	FOnDemandContentInstallerStats::OnRequestEnqueued();
 
 	InstallerPipe.Launch(
 		TEXT("ProcessIoStoreOnDemandInstallRequest"),
-		[this, RequestState] { ProcessRequest(*RequestState); },
+		[this, Request] { ProcessInstallRequest(*Request); },
 		UE::Tasks::ETaskPriority::BackgroundLow);
 
 	return InstallRequest;
 }
 
-void FOnDemandContentInstaller::Cancel(FSharedInternalInstallRequest InstallRequest)
+void FOnDemandContentInstaller::EnqueuePurgeRequest(FOnDemandPurgeArgs&& Args, FOnDemandPurgeCompleted&& OnCompleted)
+{
+	{
+		TUniqueLock Lock(Mutex);
+		RequestQueue.HeapPush(RequestAllocator.Construct(MoveTemp(Args), MoveTemp(OnCompleted)), RequestSortPredicate);
+	}
+
+	TryExecuteNextRequest();
+}
+
+void FOnDemandContentInstaller::EnqueueDefragRequest(FOnDemandDefragArgs&& Args, FOnDemandDefragCompleted&& OnCompleted)
+{
+	{
+		TUniqueLock Lock(Mutex);
+		RequestQueue.HeapPush(RequestAllocator.Construct(MoveTemp(Args), MoveTemp(OnCompleted)), RequestSortPredicate);
+	}
+
+	TryExecuteNextRequest();
+}
+
+void FOnDemandContentInstaller::CancelInstallRequest(FSharedInternalInstallRequest InstallRequest)
 {
 	InstallerPipe.Launch(
 		TEXT("CancelIoStoreOnDemandInstallRequest"),
 		[this, InstallRequest]
 		{
-			EIoErrorCode Expected = EIoErrorCode::Unknown;
-			if (InstallRequest->ErrorCode.compare_exchange_strong(Expected, EIoErrorCode::Cancelled) == false)
+			FRequest* ToComplete = nullptr;
 			{
-				return;
-			}
+				UE::TUniqueLock Lock(Mutex);
 
-			UE_LOG(LogIoStoreOnDemand, Log, TEXT("Cancelling install request, ContentHandle=(%s)"),
-				*LexToString(InstallRequest->Args.ContentHandle));
+				if (InstallRequest->InstallerRequest == 0)
+				{
+					return;
+				}
 
-			FInstallRequestState* ToComplete = nullptr;
-			{
-				UE::TUniqueLock Lock(Mutex);
-				if (FInstallRequestState* RequestState = RequestStateByHandle.FindRef(InstallRequest))
+				FRequest* Request = reinterpret_cast<FRequest*>(InstallRequest->InstallerRequest);
+
+				EIoErrorCode Expected = EIoErrorCode::Unknown;
+				if (Request->ErrorCode.compare_exchange_strong(Expected, EIoErrorCode::Cancelled) == false)
 				{
-					if (RequestQueue.Remove(RequestState) > 0)
-					{
-						ToComplete = RequestState;
-						RequestQueue.Heapify(ByInstallPriority);
-					}
+					return;
+				}
+
+				UE_LOG(LogIoStoreOnDemand, Log, TEXT("Cancelling install request, ContentHandle=(%s)"),
+					*LexToString(Request->AsInstall().Args.ContentHandle));
+
+				if (RequestQueue.Remove(Request) > 0)
+				{
+					ToComplete = Request;
+					RequestQueue.Heapify(RequestSortPredicate);
 				}
 			}
 
 			if (ToComplete != nullptr)
 			{
-				CompleteRequest(*ToComplete);
+				CompleteInstallRequest(*ToComplete);
 			}
 		});
 }
 
-void FOnDemandContentInstaller::UpdatePriority(FSharedInternalInstallRequest InstallRequest, int32 NewPriority)
+void FOnDemandContentInstaller::UpdateInstallRequestPriority(FSharedInternalInstallRequest InstallRequest, int32 NewPriority)
 {
 	InstallerPipe.Launch(
 		TEXT("UpdateIoStoreOnDemandInstallRequestPriority"),
 		[this, InstallRequest, NewPriority]
 		{
 			UE::TUniqueLock Lock(Mutex);
-			if (FInstallRequestState* RequestState = RequestStateByHandle.FindRef(InstallRequest))
+
+			if (InstallRequest->InstallerRequest == 0)
 			{
-				UE_LOG(LogIoStoreOnDemand, Log, TEXT("Updating install request priority, ContentHandle=(%s), CurrentPriority=%d, NewPriority=%d"),
-					*LexToString(InstallRequest->Args.ContentHandle), RequestState->Priority, NewPriority);
+				return;
+			}
 
-				RequestState->Priority = NewPriority;
+			FRequest&			Request	= *reinterpret_cast<FRequest*>(InstallRequest->InstallerRequest);
+			FRequest::FInstall& Install	= Request.AsInstall();
 
-				if (RequestState->bHttpRequestsIssued.load(std::memory_order_relaxed))
+			UE_LOG(LogIoStoreOnDemand, Log, TEXT("Updating install request priority, ContentHandle=(%s), CurrentPriority=%d, NewPriority=%d"),
+				*LexToString(Install.Args.ContentHandle), Request.Priority, NewPriority);
+
+			Request.Priority = NewPriority;
+
+			if (Install.bHttpRequestsIssued.load(std::memory_order_seq_cst))
+			{
+				for (FChunkHttpRequestHandle& PendingHttpRequest : Install.HttpRequestHandles)
 				{
-					for (FChunkHttpRequestHandle& PendingHttpRequest : RequestState->HttpRequestHandles)
+					if (PendingHttpRequest.Handle != nullptr)
 					{
-						if (PendingHttpRequest.Handle != nullptr)
-						{
-							HttpClient.ReprioritizeRequest(PendingHttpRequest.Handle, RequestState->Priority);
-						}
+						HttpClient.ReprioritizeRequest(PendingHttpRequest.Handle, Request.Priority);
 					}
 				}
-				else
-				{
-					RequestQueue.Heapify(ByInstallPriority);
-				}
+			}
+			else
+			{
+				RequestQueue.Heapify(RequestSortPredicate);
 			}
 		});
 }
 
-void FOnDemandContentInstaller::EnqueuePurge(FOnDemandPurgeArgs&& Args, FOnDemandPurgeCompleted&& OnCompleted)
+void FOnDemandContentInstaller::TryExecuteNextRequest()
 {
-	InstallerPipe.Launch(
-		TEXT("ProcessIoStoreOnDemandPurgeRequest"),
-		[this, Args = MoveTemp(Args), OnCompleted = MoveTemp(OnCompleted), StartTimeCycles = FPlatformTime::Cycles64()]() mutable
-		{
-			const bool bDefrag = EnumHasAnyFlags(Args.Options, EOnDemandPurgeOptions::Defrag);
-			const uint64* BytesToPurge = Args.BytesToPurge.GetPtrOrNull();
-
-			FOnDemandPurgeResult Result;
-			Result.Status = IoStore.InstallCache->PurgeAllUnreferenced(bDefrag, BytesToPurge);
+	if (bShuttingDown.load(std::memory_order_relaxed))
+	{
+		return;
+	}
 
-			const uint64 DurationCycles = FPlatformTime::Cycles64() - StartTimeCycles;
-			Result.DurationInSeconds = FPlatformTime::ToSeconds64(DurationCycles);
+	FRequest* NextRequest = nullptr;
+	{
+		TUniqueLock Lock(Mutex);
+		if (CurrentRequest == nullptr && RequestQueue.IsEmpty() == false)
+		{
+			RequestQueue.HeapPop(NextRequest, RequestSortPredicate, EAllowShrinking::No);
+			CurrentRequest = NextRequest;
+		}
+	}
 
-			if (OnCompleted)
-			{
-				if (EnumHasAnyFlags(Args.Options, EOnDemandPurgeOptions::CallbackOnGameThread))
-				{
-					ExecuteOnGameThread(
-						UE_SOURCE_LOCATION,
-						[OnCompleted = MoveTemp(OnCompleted), Result = MoveTemp(Result)]() mutable
-						{
-							OnCompleted(MoveTemp(Result));
-						});
-				}
-				else
-				{
-					OnCompleted(MoveTemp(Result));
-				}
-			}
-		},
-		UE::Tasks::ETaskPriority::BackgroundLow
-	);
+	if (NextRequest != nullptr)
+	{
+		InstallerPipe.Launch(
+			TEXT("ExecuteRequest"),
+			[this, NextRequest] { ExecuteRequest(*NextRequest); },
+			UE::Tasks::ETaskPriority::BackgroundLow);
+	}
 }
 
-void FOnDemandContentInstaller::EnqueueDefrag(FOnDemandDefragArgs&& Args, FOnDemandDefragCompleted&& OnCompleted)
+void FOnDemandContentInstaller::ExecuteRequest(FRequest& Request)
 {
-	InstallerPipe.Launch(
-		TEXT("ProcessIoStoreOnDemandDefragRequest"),
-		[this, Args = MoveTemp(Args), OnCompleted = MoveTemp(OnCompleted), StartTimeCycles = FPlatformTime::Cycles64()]() mutable
+	struct FVisitor
+	{
+		void operator()(FEmptyVariantState& Empty)
 		{
-			const uint64* BytesToFree = Args.BytesToFree.GetPtrOrNull();
+			ensure(false);
+		}
 
-			FOnDemandDefragResult Result;
-			Result.Status = IoStore.InstallCache->DefragAll(BytesToFree);
+		void operator()(FRequest::FInstall&)
+		{
+			Installer.ExecuteInstallRequest(Request, /*bRemoveAlreadyCachedChunks*/ true);
+		}
 
-			const uint64 DurationCycles = FPlatformTime::Cycles64() - StartTimeCycles;
-			Result.DurationInSeconds = FPlatformTime::ToSeconds64(DurationCycles);
+		void operator()(FRequest::FPurge&)
+		{
+			Installer.ExecutePurgeRequest(Request);
+		}
 
-			if (OnCompleted)
-			{
-				if (EnumHasAnyFlags(Args.Options, EOnDemandDefragOptions::CallbackOnGameThread))
-				{
-					ExecuteOnGameThread(
-						UE_SOURCE_LOCATION,
-						[OnCompleted = MoveTemp(OnCompleted), Result = MoveTemp(Result)]() mutable
-						{
-							OnCompleted(MoveTemp(Result));
-						});
-				}
-				else
-				{
-					OnCompleted(MoveTemp(Result));
-				}
-			}
-		},
-		UE::Tasks::ETaskPriority::BackgroundLow
-	);
+		void operator()(FRequest::FDefrag& Defrag)
+		{
+			Installer.ExecuteDefragRequest(Request);
+		}
+
+		FOnDemandContentInstaller&	Installer;
+		FRequest&					Request;
+	};
+
+	FVisitor Visitor { .Installer = *this, .Request = Request };
+	Visit(Visitor, Request.Variant);
 }
 
-void FOnDemandContentInstaller::ProcessRequest(FInstallRequestState& RequestState)
+void FOnDemandContentInstaller::ProcessInstallRequest(FRequest& Request)
 {
 	using namespace UE::IoStore::Private;
 
-	TRACE_CPUPROFILER_EVENT_SCOPE(FOnDemandContentInstaller::ProcessRequest);
+	TRACE_CPUPROFILER_EVENT_SCOPE(FOnDemandContentInstaller::ProcessInstallRequest);
 
-	FOnDemandInternalInstallRequest& Request = *RequestState.Request;
+	FRequest::FInstall& InstallRequest = Request.AsInstall();
 
 	UE_LOG(LogIoStoreOnDemand, Log, TEXT("Processing install request, ContentHandle=(%s)"),
-		*LexToString(Request.Args.ContentHandle));
+		*LexToString(InstallRequest.Args.ContentHandle));
 
-	if (Request.Args.ContentHandle.IsValid() == false)
+	if (InstallRequest.Args.ContentHandle.IsValid() == false)
 	{
-		Request.ErrorCode			= EIoErrorCode::InvalidParameter;
-		RequestState.ErrorReason	= TEXT("Invalid content handle");
+		Request.ErrorCode	= EIoErrorCode::InvalidParameter;
+		Request.ErrorReason	= TEXT("Invalid content handle");
 
-		return CompleteRequest(RequestState);
+		return CompleteInstallRequest(Request);
 	}
 
-	const TSharedPtr<FOnDemandInternalContentHandle>& ContentHandle = Request.Args.ContentHandle.Handle;
+	const TSharedPtr<FOnDemandInternalContentHandle>& ContentHandle = InstallRequest.Args.ContentHandle.Handle;
 	if (ContentHandle->IoStore.IsValid() == false)
 	{
 		// First time this content handle is used
@@ -418,45 +434,44 @@ void FOnDemandContentInstaller::ProcessRequest(FInstallRequestState& RequestStat
 	TSet<FSharedOnDemandContainer> ContainersForInstallation;
 	TSet<FPackageId> PackageIdsToInstall;
 	if (FIoStatus Status = IoStore.GetContainersAndPackagesForInstall(
-		Request.Args.MountId,
-		Request.Args.TagSets,
-		Request.Args.PackageIds,
+		InstallRequest.Args.MountId,
+		InstallRequest.Args.TagSets,
+		InstallRequest.Args.PackageIds,
 		ContainersForInstallation,
 		PackageIdsToInstall); !Status.IsOk())
 	{
-		Request.ErrorCode			= Status.GetErrorCode();
-		RequestState.ErrorReason	= Status.ToString(); 
+		Request.ErrorCode	= Status.GetErrorCode();
+		Request.ErrorReason	= Status.ToString(); 
 
-		return CompleteRequest(RequestState);
+		return CompleteInstallRequest(Request);
 	}
 
-
 #if !UE_BUILD_SHIPPING
 	if (!UE::IoStore::CVars::IoStoreErrorOnRequest.IsEmpty())
 	{
 		if (FCString::Strstr(*ContentHandle->DebugName, *UE::IoStore::CVars::IoStoreErrorOnRequest) != nullptr || 
-			FCString::Strstr(*Request.Args.DebugName, *UE::IoStore::CVars::IoStoreErrorOnRequest) != nullptr)
+			FCString::Strstr(*InstallRequest.Args.DebugName, *UE::IoStore::CVars::IoStoreErrorOnRequest) != nullptr)
 		{
-			Request.ErrorCode = (EIoErrorCode)((FMath::Rand() / (RAND_MAX / ((int)EIoErrorCode::Last-1))) + 1);
-			RequestState.ErrorReason = FString::Printf(TEXT("Debug error requested on debug name %s"), *UE::IoStore::CVars::IoStoreErrorOnRequest);
-			return CompleteRequest(RequestState);
+			Request.ErrorCode	= (EIoErrorCode)((FMath::Rand() / (RAND_MAX / ((int)EIoErrorCode::Last-1))) + 1);
+			Request.ErrorReason = FString::Printf(TEXT("Debug error requested on debug name %s"), *UE::IoStore::CVars::IoStoreErrorOnRequest);
+
+			return CompleteInstallRequest(Request);
 		}
 	}
 #endif
 
-
 	if (Request.ErrorCode.load(std::memory_order_relaxed) == EIoErrorCode::Cancelled)
 	{
-		return CompleteRequest(RequestState);
+		return CompleteInstallRequest(Request);
 	}
 
 	TSet<FPackageId> Missing;
-	const bool bIncludeSoftReferences = EnumHasAnyFlags(Request.Args.Options, EOnDemandInstallOptions::InstallSoftReferences);
+	const bool bIncludeSoftReferences = EnumHasAnyFlags(InstallRequest.Args.Options, EOnDemandInstallOptions::InstallSoftReferences);
 	Private::ResolveChunksToInstall(
 		ContainersForInstallation,
 		PackageIdsToInstall,
 		bIncludeSoftReferences,
-		RequestState.ResolvedChunks,
+		InstallRequest.ResolvedChunks,
 		Missing);
 
 	// Check the other I/O backends for missing package chunks
@@ -471,7 +486,7 @@ void FOnDemandContentInstaller::ProcessRequest(FInstallRequestState& RequestStat
 			if (IoDispatcher.DoesChunkExist(ChunkId) == false)
 			{
 				UE_CLOG(MissingCount == 0, LogIoStoreOnDemand, Warning, TEXT("Failed to resolve the following chunk(s) for content handle '%s':"),
-					*LexToString(Request.Args.ContentHandle));
+					*LexToString(InstallRequest.Args.ContentHandle));
 
 				UE_LOG(LogIoStoreOnDemand, Warning, TEXT("\tChunkId='%s'"), *LexToString(ChunkId));
 				MissingCount++;
@@ -480,33 +495,38 @@ void FOnDemandContentInstaller::ProcessRequest(FInstallRequestState& RequestStat
 
 		if (MissingCount > 0) 
 		{
-			Request.ErrorCode			= EIoErrorCode::UnknownChunkID;
-			RequestState.ErrorReason	= FString::Printf(TEXT("Missing chunk(s), Count=%u, ContentHandle='%s'"),
-				MissingCount, *LexToString(Request.Args.ContentHandle));
+			Request.ErrorCode	= EIoErrorCode::UnknownChunkID;
+			Request.ErrorReason	= FString::Printf(TEXT("Missing chunk(s), Count=%u, ContentHandle='%s'"),
+				MissingCount, *LexToString(InstallRequest.Args.ContentHandle));
 
-			return CompleteRequest(RequestState);
+			return CompleteInstallRequest(Request);
 		}
 	}
 
 	if (Request.ErrorCode.load(std::memory_order_relaxed) == EIoErrorCode::Cancelled)
 	{
-		return CompleteRequest(RequestState);
+		return CompleteInstallRequest(Request);
 	}
 
-	uint64 TotalContentSize	= 0;	
+	uint64 TotalContentSize	= 0;
 	uint64 TotalInstallSize	= 0;
 
 	// Find all chunks we need to fetch from the resolved chunk(s)
 	{
-		for (int32 ContainerIndex = 0; FResolvedContainerChunks& ResolvedChunks : RequestState.ResolvedChunks)
+		for (int32 ContainerIndex = 0; FResolvedContainerChunks& ResolvedChunks : InstallRequest.ResolvedChunks)
 		{
+			TArray<int32, TInlineAllocator<64>> CachedEntryIndices;
 			for (int32 EntryIndex : ResolvedChunks.EntryIndices)
 			{
 				const FOnDemandChunkEntry& Entry = ResolvedChunks.Container->ChunkEntries[EntryIndex];
-				const bool bInstalled = IoStore.InstallCache->IsChunkCached(Entry.Hash);
-				if (bInstalled == false)
+				const bool bCached = IoStore.InstallCache->IsChunkCached(Entry.Hash);
+				if (bCached)
 				{
-					RequestState.HttpRequestHandles.Add(FChunkHttpRequestHandle
+					CachedEntryIndices.Add(EntryIndex);
+				}
+				else
+				{
+					InstallRequest.HttpRequestHandles.Add(FChunkHttpRequestHandle
 					{
 						.Handle			= nullptr,
 						.ContainerIndex = ContainerIndex,
@@ -514,119 +534,110 @@ void FOnDemandContentInstaller::ProcessRequest(FInstallRequestState& RequestStat
 					});
 					TotalInstallSize += Entry.EncodedSize;
 				}
-				++RequestState.ResolvedChunkCount;
+				++InstallRequest.ResolvedChunkCount;
 				TotalContentSize += Entry.EncodedSize;
 			}
+
+			// Add references to existing chunk(s)
+			if (CachedEntryIndices.IsEmpty() == false)
+			{
+				TUniqueLock Lock(IoStore.ContainerMutex);
+
+				FOnDemandChunkEntryReferences& References = ResolvedChunks.Container->FindOrAddChunkEntryReferences(*ContentHandle);
+				for (int32 EntryIndex : CachedEntryIndices)
+				{
+					References.Indices[EntryIndex] = true;
+				}
+			}
+
 			ContainerIndex++;
 		}
 	}
 
-	RequestState.Progress.TotalContentSize		= TotalContentSize;
-	RequestState.Progress.TotalInstallSize		= TotalInstallSize;
-	RequestState.Progress.CurrentInstallSize	= 0;
-	RequestState.Priority						= Request.Args.Priority;
+	InstallRequest.Progress.TotalContentSize	= TotalContentSize;
+	InstallRequest.Progress.TotalInstallSize	= TotalInstallSize;
+	InstallRequest.Progress.CurrentInstallSize	= 0;
+	Request.Priority							= InstallRequest.Args.Priority;
 
-	if (RequestState.HttpRequestHandles.IsEmpty())
+	if (InstallRequest.HttpRequestHandles.IsEmpty())
 	{
 		Request.ErrorCode = EIoErrorCode::Ok;
-		return CompleteRequest(RequestState);
+
+		return CompleteInstallRequest(Request);
 	}
 
 	if (Request.ErrorCode.load(std::memory_order_relaxed) == EIoErrorCode::Cancelled)
 	{
-		return CompleteRequest(RequestState);
+		return CompleteInstallRequest(Request);
 	}
 
-	bool bStartDownload = false;
+	bool bExecuteRequest = false;
 	{
 		TUniqueLock Lock(Mutex);
 		if (CurrentRequest == nullptr)
 		{
-			CurrentRequest	= &RequestState;
-			bStartDownload	= true;
+			CurrentRequest	= &Request;
+			bExecuteRequest	= true;
 		}
 		else
 		{
-			RequestQueue.HeapPush(&RequestState, ByInstallPriority);
+			RequestQueue.HeapPush(&Request, RequestSortPredicate);
 		}
 	}
 
-	if (bStartDownload)
+	if (bExecuteRequest)
 	{
-		check(CurrentRequest == &RequestState);
+		check(CurrentRequest == &Request);
 		const bool bRemoveAlreadyCachedChunks = false;
-		StartDownloadChunks(RequestState, bRemoveAlreadyCachedChunks);
+		ExecuteInstallRequest(Request, bRemoveAlreadyCachedChunks);
 	}
 }
 
-void FOnDemandContentInstaller::TryStartDownloadChunks()
+void FOnDemandContentInstaller::ExecuteInstallRequest(FRequest& Request, bool bRemoveAlreadyCachedChunks)
 {
-	if (bShuttingDown.load(std::memory_order_relaxed))
-	{
-		return;
-	}
-
-	FInstallRequestState* NextRequest = nullptr;
-	{
-		TUniqueLock Lock(Mutex);
-		if (CurrentRequest == nullptr && RequestQueue.IsEmpty() == false)
-		{
-			RequestQueue.HeapPop(NextRequest, ByInstallPriority, EAllowShrinking::No);
-			CurrentRequest = NextRequest;
-		}
-	}
+	check(Request.IsInstall());
+	check(&Request == CurrentRequest);
 
-	if (NextRequest != nullptr)
-	{
-		InstallerPipe.Launch(
-			TEXT("StartIoStoreOnDemandInstallRequest"),
-			[this, NextRequest] { StartDownloadChunks(*NextRequest, /*bRemoveAlreadyCachedChunks*/ true); },
-			UE::Tasks::ETaskPriority::BackgroundLow);
-	}
-}
-
-void FOnDemandContentInstaller::StartDownloadChunks(FInstallRequestState& RequestState, bool bRemoveAlreadyCachedChunks)
-{
-	FOnDemandInternalInstallRequest& Request = *RequestState.Request;
-	check(RequestState.HttpRequestHandles.IsEmpty() == false);
+	FRequest::FInstall& InstallRequest = Request.AsInstall();
+	check(InstallRequest.HttpRequestHandles.IsEmpty() == false);
 
 	if (Request.ErrorCode.load(std::memory_order_relaxed) == EIoErrorCode::Cancelled)
 	{
-		return CompleteRequest(RequestState);
+		return CompleteInstallRequest(Request);
 	}
 
 	if (bRemoveAlreadyCachedChunks)
 	{
-		RequestState.HttpRequestHandles.SetNum(
+		InstallRequest.HttpRequestHandles.SetNum(
 			Algo::RemoveIf(
-				RequestState.HttpRequestHandles,
-				[this, &RequestState](FChunkHttpRequestHandle& HttpRequest)
+				InstallRequest.HttpRequestHandles,
+				[this, &InstallRequest](FChunkHttpRequestHandle& HttpRequest)
 				{
-					FSharedOnDemandContainer& Container		= RequestState.ResolvedChunks[HttpRequest.ContainerIndex].Container;
+					FSharedOnDemandContainer& Container		= InstallRequest.ResolvedChunks[HttpRequest.ContainerIndex].Container;
 					const FOnDemandChunkEntry& ChunkEntry	= Container->ChunkEntries[HttpRequest.EntryIndex];
 
 					if (IoStore.InstallCache->IsChunkCached(ChunkEntry.Hash))
 					{
-						RequestState.Progress.TotalInstallSize -= ChunkEntry.EncodedSize;
+						InstallRequest.Progress.TotalInstallSize -= ChunkEntry.EncodedSize;
 						return true;
 					}
 
 					return false;
 				}));
 
-		if (RequestState.HttpRequestHandles.IsEmpty())
+		if (InstallRequest.HttpRequestHandles.IsEmpty())
 		{
 			Request.ErrorCode = EIoErrorCode::Ok;
-			return CompleteRequest(RequestState);
+			return CompleteInstallRequest(Request);
 		}
 	}
 
 	// Make sure we have enough space in the cache
 	{
 		TSet<FIoHash> ChunksToInstall;
-		for (FChunkHttpRequestHandle& HttpRequest : RequestState.HttpRequestHandles)
+		for (FChunkHttpRequestHandle& HttpRequest : InstallRequest.HttpRequestHandles)
 		{
-			FSharedOnDemandContainer& Container		= RequestState.ResolvedChunks[HttpRequest.ContainerIndex].Container;
+			FSharedOnDemandContainer& Container		= InstallRequest.ResolvedChunks[HttpRequest.ContainerIndex].Container;
 			const FOnDemandChunkEntry& ChunkEntry	= Container->ChunkEntries[HttpRequest.EntryIndex];
 
 			ChunksToInstall.Add(ChunkEntry.Hash);
@@ -634,39 +645,39 @@ void FOnDemandContentInstaller::StartDownloadChunks(FInstallRequestState& Reques
 
 		if (FIoStatus Status = IoStore.InstallCache->Purge(MoveTemp(ChunksToInstall)); !Status.IsOk())
 		{
-			Request.ErrorCode			= Status.GetErrorCode();
-			RequestState.ErrorReason	= Status.ToString(); 
+			Request.ErrorCode	= Status.GetErrorCode();
+			Request.ErrorReason	= Status.ToString(); 
 
-			return CompleteRequest(RequestState);
+			return CompleteInstallRequest(Request);
 		}
 	}
 
 	if (Request.ErrorCode.load(std::memory_order_relaxed) == EIoErrorCode::Cancelled)
 	{
-		return CompleteRequest(RequestState);
+		return CompleteInstallRequest(Request);
 	}
 
-	NotifyProgress(RequestState);
+	NotifyInstallProgress(Request);
 
-	for (FChunkHttpRequestHandle& HttpRequest : RequestState.HttpRequestHandles)
+	for (FChunkHttpRequestHandle& HttpRequest : InstallRequest.HttpRequestHandles)
 	{
 		TRACE_CPUPROFILER_EVENT_SCOPE(FOnDemandContentInstaller::IssueRequest);
-		FSharedOnDemandContainer& Container		= RequestState.ResolvedChunks[HttpRequest.ContainerIndex].Container;
+		FSharedOnDemandContainer& Container		= InstallRequest.ResolvedChunks[HttpRequest.ContainerIndex].Container;
 		const FOnDemandChunkEntry& ChunkEntry	= Container->ChunkEntries[HttpRequest.EntryIndex];
 
 		HttpRequest.Handle = HttpClient.IssueRequest(
 			FOnDemandChunkInfo(Container, ChunkEntry),
 			FIoOffsetAndLength(),
-			RequestState.Priority,
-			[this, &RequestState, &HttpRequest](uint32 HttpStatusCode, FStringView ErrorReason, FIoBuffer&& Chunk)
+			Request.Priority,
+			[this, &Request, &HttpRequest](uint32 HttpStatusCode, FStringView ErrorReason, FIoBuffer&& Chunk)
 			{
 				TRACE_CPUPROFILER_EVENT_SCOPE(FOnDemandContentInstaller::Callback);
 				InstallerPipe.Launch(
 					TEXT("ProcessIoStoreOnDemandDownloadedChunk"),
-					[this, &RequestState, &HttpRequest, HttpStatusCode, ErrorReason = FString(ErrorReason), Chunk = MoveTemp(Chunk)]() mutable
+					[this, &Request, &HttpRequest, HttpStatusCode, ErrorReason = FString(ErrorReason), Chunk = MoveTemp(Chunk)]() mutable
 					{
 						ProcessDownloadedChunk(
-							RequestState,
+							Request,
 							HttpRequest,
 							HttpStatusCode,
 							MoveTemp(ErrorReason),
@@ -676,11 +687,66 @@ void FOnDemandContentInstaller::StartDownloadChunks(FInstallRequestState& Reques
 			}, EHttpRequestType::Installed);
 	}
 	
-	RequestState.bHttpRequestsIssued = true;
+	InstallRequest.bHttpRequestsIssued = true;
+}
+
+void FOnDemandContentInstaller::ExecutePurgeRequest(FRequest& Request)
+{
+	check(Request.IsPurge());
+	check(&Request == CurrentRequest);
+
+	FRequest::FPurge& PurgeRequest	= Request.AsPurge();
+	const bool bDefrag				= EnumHasAnyFlags(PurgeRequest.Args.Options, EOnDemandPurgeOptions::Defrag);
+	const uint64* BytesToPurge		= PurgeRequest.Args.BytesToPurge.GetPtrOrNull();
+
+	UE_LOG(LogIoStoreOnDemand, Log, TEXT("Executing purge request, BytesToPurge=%llu, Defrag='%s'"),
+		BytesToPurge != nullptr ? *BytesToPurge : -1, bDefrag ? TEXT("True") : TEXT("False"));
+
+	if (Request.ErrorCode.load(std::memory_order_relaxed) == EIoErrorCode::Cancelled)
+	{
+		return CompletePurgeRequest(Request);
+	}
+
+	const FIoStatus Status	= IoStore.InstallCache->PurgeAllUnreferenced(bDefrag, BytesToPurge);
+	Request.ErrorCode		= Status.GetErrorCode();
+
+	if (!Status.IsOk())
+	{
+		Request.ErrorReason = Status.ToString();
+	}
+
+	CompletePurgeRequest(Request); 
+}
+
+void FOnDemandContentInstaller::ExecuteDefragRequest(FRequest& Request)
+{
+	check(Request.IsDefrag());
+	check(&Request == CurrentRequest);
+
+	FRequest::FDefrag& DefragRequest	= Request.AsDefrag();
+	const uint64* BytesToFree			= DefragRequest.Args.BytesToFree.GetPtrOrNull();
+
+	UE_LOG(LogIoStoreOnDemand, Log, TEXT("Executing defrag request, BytesToFree=%llu"),
+		BytesToFree != nullptr ? *BytesToFree : -1); 
+
+	if (Request.ErrorCode.load(std::memory_order_relaxed) == EIoErrorCode::Cancelled)
+	{
+		return CompleteDefragRequest(Request);
+	}
+
+	const FIoStatus Status	= IoStore.InstallCache->DefragAll(BytesToFree);
+	Request.ErrorCode		= Status.GetErrorCode();
+
+	if (!Status.IsOk())
+	{
+		Request.ErrorReason = Status.ToString();
+	}
+
+	CompleteDefragRequest(Request); 
 }
 
 void FOnDemandContentInstaller::ProcessDownloadedChunk(
-	FInstallRequestState& RequestState,
+	FRequest& Request,
 	FChunkHttpRequestHandle& HttpRequest,
 	uint32 HttpStatusCode,
 	FString&& ErrorReason,
@@ -688,17 +754,17 @@ void FOnDemandContentInstaller::ProcessDownloadedChunk(
 {
 	TRACE_CPUPROFILER_EVENT_SCOPE(FOnDemandContentInstaller::ProcessDownloadedChunk);
 
-	FOnDemandInternalInstallRequest& Request	= *RequestState.Request;
+	FRequest::FInstall& InstallRequest			= Request.AsInstall();
 	HttpRequest.Handle							= nullptr;
-	FSharedOnDemandContainer& Container			= RequestState.ResolvedChunks[HttpRequest.ContainerIndex].Container;
+	FSharedOnDemandContainer& Container			= InstallRequest.ResolvedChunks[HttpRequest.ContainerIndex].Container;
 	const FOnDemandChunkEntry& ChunkEntry		= Container->ChunkEntries[HttpRequest.EntryIndex];
 	const FIoChunkId& ChunkId					= Container->ChunkIds[HttpRequest.EntryIndex];
 
-	RequestState.Progress.CurrentInstallSize += ChunkEntry.EncodedSize;
+	InstallRequest.Progress.CurrentInstallSize	+= ChunkEntry.EncodedSize;
 
 	UE_LOG(LogIoStoreOnDemand, Verbose, TEXT("Install progress %.2lf/%.2lf MiB, ContentHandle=(%s), ChunkId='%s', ChunkSize=%.2lf KiB, HttpStatus=%u"),
-		double(RequestState.Progress.CurrentInstallSize) / 1024.0 / 1024.0, double(RequestState.Progress.TotalInstallSize) / 1024.0 / 1024.0,
-		*LexToString(Request.Args.ContentHandle), *LexToString(ChunkId), double(ChunkEntry.EncodedSize) / 1024.0, HttpStatusCode);
+		double(InstallRequest.Progress.CurrentInstallSize) / 1024.0 / 1024.0, double(InstallRequest.Progress.TotalInstallSize) / 1024.0 / 1024.0,
+		*LexToString(InstallRequest.Args.ContentHandle), *LexToString(ChunkId), double(ChunkEntry.EncodedSize) / 1024.0, HttpStatusCode);
 
 	if (Request.ErrorCode == EIoErrorCode::Unknown)
 	{
@@ -711,104 +777,113 @@ void FOnDemandContentInstaller::ProcessDownloadedChunk(
 			{
 				if (FIoStatus Status = IoStore.InstallCache->PutChunk(MoveTemp(Chunk), ChunkHash); !Status.IsOk())
 				{
-					Request.ErrorCode			= Status.GetErrorCode();
-					RequestState.ErrorReason	= Status.ToString();
+					Request.ErrorCode	= Status.GetErrorCode();
+					Request.ErrorReason	= Status.ToString();
 				}
 			}
 			else
 			{
-				Request.ErrorCode			= EIoErrorCode::ReadError;
-				RequestState.ErrorReason	= FString::Printf(TEXT("Hash mismatch, ChunkId='%s', ExpectedHash='%s', ActualHash='%s'"),
+				Request.ErrorCode	= EIoErrorCode::ReadError;
+				Request.ErrorReason	= FString::Printf(TEXT("Hash mismatch, ChunkId='%s', ExpectedHash='%s', ActualHash='%s'"),
 					*LexToString(ChunkId), *LexToString(ChunkEntry.Hash), *LexToString(ChunkHash));
 			}
 		}
 		else
 		{
-			Request.ErrorCode			= EIoErrorCode::ReadError;
-			RequestState.ErrorReason	= FString::Printf(TEXT("Http failure, StatusCode=%u, Reason=%s"), HttpStatusCode, *ErrorReason);
+			Request.ErrorCode	= EIoErrorCode::ReadError;
+			Request.ErrorReason	= FString::Printf(TEXT("Http failure, StatusCode=%u, Reason=%s"), HttpStatusCode, *ErrorReason);
 		}
 	}
 
-	NotifyProgress(RequestState);
-
-	const bool bCompleted = ++RequestState.DownloadedChunkCount >= RequestState.HttpRequestHandles.Num();
+	const bool bCompleted = ++InstallRequest.DownloadedChunkCount >= InstallRequest.HttpRequestHandles.Num();
 	if (bCompleted)
 	{
 		if (FIoStatus Status = IoStore.InstallCache->Flush(); !Status.IsOk())
 		{
-			Request.ErrorCode			= Status.GetErrorCode();
-			RequestState.ErrorReason	= Status.ToString();
+			Request.ErrorCode	= Status.GetErrorCode();
+			Request.ErrorReason	= Status.ToString();
 		}
 		if (Request.ErrorCode == EIoErrorCode::Unknown)
 		{
 			Request.ErrorCode = EIoErrorCode::Ok;
 		}
-		CompleteRequest(RequestState);
+
+		CompleteInstallRequest(Request);
 	}
-	else if (Request.ErrorCode != EIoErrorCode::Unknown)
+	else 
 	{
-		int32 NumCancelled = 0;
-		for (FChunkHttpRequestHandle& PendingHttpRequest : RequestState.HttpRequestHandles)
+		NotifyInstallProgress(Request);
+
+		if (Request.ErrorCode != EIoErrorCode::Unknown)
 		{
-			if (PendingHttpRequest.Handle != nullptr)
+			int32 NumCancelled = 0;
+			for (FChunkHttpRequestHandle& PendingHttpRequest : InstallRequest.HttpRequestHandles)
 			{
-				HttpClient.CancelRequest(PendingHttpRequest.Handle);
-				++NumCancelled;
+				if (PendingHttpRequest.Handle != nullptr)
+				{
+					HttpClient.CancelRequest(PendingHttpRequest.Handle);
+					++NumCancelled;
+				}
 			}
-		}
 
-		UE_CLOG(NumCancelled == 0, LogIoStoreOnDemand, Log, TEXT("Cancelled %d HTTP request(s) due to install error"), NumCancelled);
+			UE_CLOG(NumCancelled == 0, LogIoStoreOnDemand, Log, TEXT("Cancelled %d HTTP request(s) due to install error"), NumCancelled);
+		}
 	}
 }
 
-void FOnDemandContentInstaller::NotifyProgress(FInstallRequestState& RequestState)
+void FOnDemandContentInstaller::NotifyInstallProgress(FRequest& Request)
 {
-	FOnDemandInternalInstallRequest& Request = *RequestState.Request;
+	ensure(Request.IsInstall());
 
-	if (!Request.OnProgress)
+	FRequest::FInstall& InstallRequest = Request.AsInstall();
+
+	if (!InstallRequest.OnProgress)
 	{
 		return;
 	}
 
 	const uint64 Cycles = FPlatformTime::Cycles64();
-	const double SecondsSinceLastProgress = FPlatformTime::ToSeconds64(Cycles - RequestState.LastProgressCycles);
-	if (SecondsSinceLastProgress < .25)
+	const double SecondsSinceLastProgress = FPlatformTime::ToSeconds64(Cycles - InstallRequest.LastProgressCycles);
+	if (InstallRequest.bNotifyingProgressOnGameThread.load(std::memory_order_seq_cst) || SecondsSinceLastProgress < .25)
 	{
 		return;
 	}
-	RequestState.LastProgressCycles = Cycles;
+	InstallRequest.LastProgressCycles = Cycles;
 
-	FOnDemandInstallProgress Progress = RequestState.Progress;
-	if (EnumHasAnyFlags(Request.Args.Options, EOnDemandInstallOptions::CallbackOnGameThread))
+	//TODO: Remove support for notifying progress on the game thread
+	FOnDemandInstallProgress Progress = InstallRequest.Progress;
+	if (EnumHasAnyFlags(InstallRequest.Args.Options, EOnDemandInstallOptions::CallbackOnGameThread))
 	{
+		InstallRequest.bNotifyingProgressOnGameThread = true;
 		ExecuteOnGameThread(
 			UE_SOURCE_LOCATION,
-			[Request = RequestState.Request, Progress]()
+			[&InstallRequest, Progress]()
 			{
-				Request->OnProgress(Progress);
+				InstallRequest.OnProgress(InstallRequest.Progress);
+				InstallRequest.bNotifyingProgressOnGameThread = false;
 			});
 	}
 	else
 	{
-		Request.OnProgress(Progress);
+		InstallRequest.OnProgress(Progress);
 	}
 }
 
-void FOnDemandContentInstaller::CompleteRequest(FInstallRequestState& RequestState)
+void FOnDemandContentInstaller::CompleteInstallRequest(FRequest& Request)
 {
 	using namespace UE::IoStore::Private;
 
-	FSharedInternalInstallRequest Request = MoveTemp(RequestState.Request);
-	ensure(Request->ErrorCode != EIoErrorCode::Unknown);
+	FRequest::FInstall& InstallRequest = Request.AsInstall();
+	ensure(Request.ErrorCode != EIoErrorCode::Unknown);
 
 	// Mark all resolved chunk(s) as referenced by the content handle and notify the package store to update
-	if (Request->ErrorCode == EIoErrorCode::Ok && RequestState.ResolvedChunkCount > 0)
+	if (Request.ErrorCode == EIoErrorCode::Ok && InstallRequest.ResolvedChunkCount > 0)
 	{
 		{
-			FOnDemandContentHandle& ContentHandle = Request->Args.ContentHandle;
+			FOnDemandContentHandle& ContentHandle = InstallRequest.Args.ContentHandle;
 
 			TUniqueLock Lock(IoStore.ContainerMutex);
-			for (FResolvedContainerChunks& ResolvedChunks : RequestState.ResolvedChunks)
+			for (FResolvedContainerChunks& ResolvedChunks : InstallRequest.ResolvedChunks)
 			{
 				const FSharedOnDemandContainer& Container = ResolvedChunks.Container;
 				FOnDemandChunkEntryReferences& References = Container->FindOrAddChunkEntryReferences(*ContentHandle.Handle);
@@ -822,40 +897,71 @@ void FOnDemandContentInstaller::CompleteRequest(FInstallRequestState& RequestSta
 		IoStore.PackageStoreBackend->NeedsUpdate(EOnDemandPackageStoreUpdateMode::ReferencedPackages);
 	}
 
-	const uint64 DurationCycles = FPlatformTime::Cycles64() - RequestState.StartTimeCycles; 
-	const double CacheHitRatio	= RequestState.Progress.TotalContentSize > 0
-		? double(RequestState.Progress.TotalContentSize - RequestState.Progress.TotalInstallSize) / double(RequestState.Progress.TotalContentSize)
+	const uint64 DurationCycles = FPlatformTime::Cycles64() - Request.StartTimeCycles; 
+	const double CacheHitRatio	= InstallRequest.Progress.TotalContentSize > 0
+		? double(InstallRequest.Progress.TotalContentSize - InstallRequest.Progress.TotalInstallSize) / double(InstallRequest.Progress.TotalContentSize)
 		: 0.0;
 
 	FOnDemandInstallResult InstallResult;
-	InstallResult.Status				= Request->ErrorCode == EIoErrorCode::Ok ? FIoStatus::Ok : FIoStatus(Request->ErrorCode, RequestState.ErrorReason);
-	InstallResult.Progress				= RequestState.Progress;
+	InstallResult.Status				= Request.ErrorCode == EIoErrorCode::Ok ? FIoStatus::Ok : FIoStatus(Request.ErrorCode, Request.ErrorReason);
+	InstallResult.Progress				= InstallRequest.Progress;
 	InstallResult.DurationInSeconds		= FPlatformTime::ToSeconds64(DurationCycles);
 
 	FOnDemandContentInstallerStats::OnRequestCompleted(
-		RequestState.ResolvedChunkCount,
+		InstallRequest.ResolvedChunkCount,
 		InstallResult.Progress.TotalContentSize,
-		static_cast<uint64>(RequestState.HttpRequestHandles.Num()),
+		static_cast<uint64>(InstallRequest.HttpRequestHandles.Num()),
 		InstallResult.Progress.TotalInstallSize,
 		CacheHitRatio,
 		DurationCycles,
-		Request->ErrorCode);
+		Request.ErrorCode);
 
 	UE_LOG(LogIoStoreOnDemand, Log, TEXT("Install request completed, Result='%s', ContentHandle=(%s), ContentSize=%.2lf MiB, InstallSize=%.2lf MiB, CacheHitRatio=%d %%, Duration=%d ms"),
-		GetIoErrorText(InstallResult.Status.GetErrorCode()), *LexToString(Request->Args.ContentHandle), double(InstallResult.Progress.TotalContentSize) / 1024.0 / 1024.0,
+		GetIoErrorText(InstallResult.Status.GetErrorCode()), *LexToString(InstallRequest.Args.ContentHandle), double(InstallResult.Progress.TotalContentSize) / 1024.0 / 1024.0,
 		double(InstallResult.Progress.TotalInstallSize) / 1024.0 / 1024.0, int32(CacheHitRatio * 100), int32(InstallResult.DurationInSeconds * 1000));
 
+	const bool bCallbackOnGameThread			= EnumHasAnyFlags(InstallRequest.Args.Options, EOnDemandInstallOptions::CallbackOnGameThread);
+	FOnDemandInstallCompleted OnCompleted		= MoveTemp(InstallRequest.OnCompleted);
+	FSharedInternalInstallRequest ClientRequest = MoveTemp(InstallRequest.Request);
+
+	// TODO: Remove support for notifying progress on game thread
+	if (InstallRequest.bNotifyingProgressOnGameThread)
+	{
+		const double WaitTimeoutSeconds		= 5.0;
+		const double WarningTimeoutSeconds	= .5;
+		double WaitTimeSeconds				= 0.0;
+		double WaitStartTimeSeconds			= FPlatformTime::Seconds();
+
+		while (InstallRequest.bNotifyingProgressOnGameThread)
+		{
+			WaitTimeSeconds = FPlatformTime::Seconds() - WaitStartTimeSeconds;
+			if (WaitTimeSeconds > WaitTimeoutSeconds)
+			{
+				UE_LOG(LogIoStoreOnDemand, Error,
+					TEXT("Content installer exceeded timeout threshold of %.2f seconds waiting for progress callback to finish"),
+					WaitTimeoutSeconds);
+				break;
+			}
+			FPlatformProcess::SleepNoStats(0.017f);
+		}
+
+		UE_CLOG(WaitTimeSeconds > WarningTimeoutSeconds, LogIoStoreOnDemand, Warning,
+			TEXT("Content installer waited %.2lf seconds for progress callback to finish"),
+			WaitTimeSeconds);
+	}
+
 	{
 		UE::TUniqueLock Lock(Mutex);
-		if (CurrentRequest == &RequestState)
+		ClientRequest->InstallerRequest = 0;
+
+		if (&Request == CurrentRequest)
 		{
 			CurrentRequest = nullptr;
 		}
-		RequestStateByHandle.Remove(Request); 
-		RequestStateAllocator.Destroy(&RequestState);
+		RequestAllocator.Destroy(&Request);
 	}
 
-	TryStartDownloadChunks();
+	TryExecuteNextRequest();
 
 	FOnDemandInstallRequest::EStatus RequestStatus = FOnDemandInstallRequest::EStatus::None;
 	switch (InstallResult.Status.GetErrorCode())
@@ -871,74 +977,170 @@ void FOnDemandContentInstaller::CompleteRequest(FInstallRequestState& RequestSta
 			break;
 	}
 
-	if (!Request->OnCompleted)
+	if (!OnCompleted)
 	{
-		Request->Status.store(RequestStatus);
+		ClientRequest->Status.store(RequestStatus);
 	}
-	else if (EnumHasAnyFlags(Request->Args.Options, EOnDemandInstallOptions::CallbackOnGameThread))
+	else if (bCallbackOnGameThread)
 	{
 		ExecuteOnGameThread(
 			UE_SOURCE_LOCATION,
-			[Request, InstallResult = MoveTemp(InstallResult), RequestStatus]() mutable
+			[ClientRequest, InstallResult = MoveTemp(InstallResult), RequestStatus, OnCompleted = MoveTemp(OnCompleted)]() mutable
 			{
-				FOnDemandInstallCompleted OnCompleted = MoveTemp(Request->OnCompleted);
 				OnCompleted(MoveTemp(InstallResult));
-				Request->Status.store(RequestStatus);
+				ClientRequest->Status.store(RequestStatus);
 			});
 	}
 	else
 	{
-		FOnDemandInstallCompleted OnCompleted = MoveTemp(Request->OnCompleted);
 		OnCompleted(MoveTemp(InstallResult));
-		Request->Status.store(RequestStatus);
+		ClientRequest->Status.store(RequestStatus);
+	}
+}
+
+void FOnDemandContentInstaller::CompletePurgeRequest(FRequest& Request)
+{
+	check(&Request == CurrentRequest);
+
+	const uint64 DurationCycles = FPlatformTime::Cycles64() - Request.StartTimeCycles; 
+	
+	FOnDemandPurgeResult PurgeResult;
+	PurgeResult.Status				= Request.ErrorCode == EIoErrorCode::Ok ? FIoStatus::Ok : FIoStatus(Request.ErrorCode, Request.ErrorReason);
+	PurgeResult.DurationInSeconds	= FPlatformTime::ToSeconds64(DurationCycles);
+
+	UE_LOG(LogIoStoreOnDemand, Log, TEXT("Purge request completed, Result='%s', Duration=%d ms"),
+		GetIoErrorText(PurgeResult.Status.GetErrorCode()), int32(PurgeResult.DurationInSeconds * 1000));
+	
+	const bool bCallbackOnGameThread	= EnumHasAnyFlags(Request.AsPurge().Args.Options, EOnDemandPurgeOptions::CallbackOnGameThread);
+	FOnDemandPurgeCompleted OnCompleted = MoveTemp(Request.AsPurge().OnCompleted);
+	{
+		UE::TUniqueLock Lock(Mutex);
+		if (&Request == CurrentRequest)
+		{
+			CurrentRequest = nullptr;
+		}
+		RequestAllocator.Destroy(&Request);
+	}
+
+	TryExecuteNextRequest();
+
+	if (!OnCompleted)
+	{
+		return;
+	}
+
+	if (bCallbackOnGameThread)
+	{
+		ExecuteOnGameThread(
+			UE_SOURCE_LOCATION,
+			[OnCompleted = MoveTemp(OnCompleted), PurgeResult = MoveTemp(PurgeResult)]() mutable
+			{
+				OnCompleted(MoveTemp(PurgeResult));
+			});
+	}
+	else
+	{
+		OnCompleted(MoveTemp(PurgeResult));
+	}
+}
+
+void FOnDemandContentInstaller::CompleteDefragRequest(FRequest& Request)
+{
+	check(&Request == CurrentRequest);
+
+	const uint64 DurationCycles = FPlatformTime::Cycles64() - Request.StartTimeCycles; 
+
+	FOnDemandDefragResult DefragResult;
+	DefragResult.Status				= Request.ErrorCode == EIoErrorCode::Ok ? FIoStatus::Ok : FIoStatus(Request.ErrorCode, Request.ErrorReason);
+	DefragResult.DurationInSeconds	= FPlatformTime::ToSeconds64(DurationCycles);
+
+	UE_LOG(LogIoStoreOnDemand, Log, TEXT("Defrag request completed, Result='%s', Duration=%d ms"),
+		GetIoErrorText(DefragResult.Status.GetErrorCode()), int32(DefragResult.DurationInSeconds * 1000));
+
+	const bool bCallbackOnGameThread		= EnumHasAnyFlags(Request.AsDefrag().Args.Options, EOnDemandDefragOptions::CallbackOnGameThread);
+	FOnDemandDefragCompleted OnCompleted	= MoveTemp(Request.AsDefrag().OnCompleted);
+	{
+		UE::TUniqueLock Lock(Mutex);
+		if (&Request == CurrentRequest)
+		{
+			CurrentRequest = nullptr;
+		}
+		RequestAllocator.Destroy(&Request);
+	}
+
+	TryExecuteNextRequest();	
+
+	if (!OnCompleted)
+	{
+		return;
+	}
+
+	if (bCallbackOnGameThread)
+	{
+		ExecuteOnGameThread(
+			UE_SOURCE_LOCATION,
+			[OnCompleted = MoveTemp(OnCompleted), DefragResult = MoveTemp(DefragResult)]() mutable
+			{
+				OnCompleted(MoveTemp(DefragResult));
+			});
+	}
+	else
+	{
+		OnCompleted(MoveTemp(DefragResult));
 	}
 }
 
 void FOnDemandContentInstaller::Shutdown()
 {
 	bShuttingDown					= true;
-	const double MaxTimeSeconds		= 10.0;
+	const double WaitTimeoutSeconds	= 5.0;
 	const uint64 StartTimeCycles	= FPlatformTime::Cycles64();
 
-	// Wait for the currently installing request to finish
+	// Wait for the current request to finish
 	for (;;)
 	{
-		if (FPlatformTime::ToSeconds64(FPlatformTime::Cycles64() - StartTimeCycles) > MaxTimeSeconds)
+		double WaitTimeSeconds = FPlatformTime::ToSeconds64(FPlatformTime::Cycles64() - StartTimeCycles);
+		if (WaitTimeSeconds > WaitTimeoutSeconds)
 		{
-			UE_LOG(LogIoStoreOnDemand, Warning, TEXT("Content installer shutdown cancelled after %.2lf"), MaxTimeSeconds);
+			UE_LOG(LogIoStoreOnDemand, Warning, TEXT("Content installer shutdown cancelled after %.2lf"), WaitTimeSeconds);
 			break;
 		}
 
 		InstallerPipe.WaitUntilEmpty(FTimespan::FromSeconds(1.0));
 		{
 			TUniqueLock Lock(Mutex);
-			if (CurrentRequest != nullptr)
+			if (CurrentRequest == nullptr)
 			{
-				continue;
+				break;
 			}
 		}
 	}
 
+	{
+		TUniqueLock Lock(Mutex);
+		UE_CLOG(CurrentRequest != nullptr, LogIoStoreOnDemand, Error, TEXT("Content installer has still inflight request(s) while shutting down"));
+	}
+
 	// Cancel all remaining request(s)
 	for (;;)
 	{
-		FInstallRequestState* RequestState = nullptr;
+		FRequest* NextRequest = nullptr;
 		{
 			TUniqueLock Lock(Mutex);
 			if (RequestQueue.IsEmpty() == false)
 			{
-				RequestQueue.HeapPop(RequestState, ByInstallPriority, EAllowShrinking::No);
+				RequestQueue.HeapPop(NextRequest, RequestSortPredicate, EAllowShrinking::No);
+				CurrentRequest = NextRequest;
 			}
 		}
 
-		if (RequestState == nullptr)
+		if (NextRequest == nullptr)
 		{
 			break;
 		}
 
-		FOnDemandInternalInstallRequest& Request = *RequestState->Request;
-		Request.ErrorCode = EIoErrorCode::Cancelled;
-		CompleteRequest(*RequestState);
+		NextRequest->ErrorCode = EIoErrorCode::Cancelled;
+		ExecuteRequest(*NextRequest);
 	}
 }
 
diff --git a/Engine/Source/Runtime/Experimental/IoStore/OnDemand/Private/OnDemandContentInstaller.h b/Engine/Source/Runtime/Experimental/IoStore/OnDemand/Private/OnDemandContentInstaller.h
--- a/Engine/Source/Runtime/Experimental/IoStore/OnDemand/Private/OnDemandContentInstaller.h
+++ b/Engine/Source/Runtime/Experimental/IoStore/OnDemand/Private/OnDemandContentInstaller.h
@@ -8,6 +8,7 @@
 #include "IO/PackageId.h"
 #include "IO/IoStoreOnDemandInternals.h"
 #include "IO/IoAllocators.h"
+#include "Misc/TVariant.h"
 #include "OnDemandIoStore.h"
 #include "Tasks/Pipe.h"
 
@@ -61,73 +62,143 @@ class FOnDemandContentInstaller
 		int32	EntryIndex = INDEX_NONE;
 	};
 
-	struct FInstallRequestState
+	struct FRequest
 	{
-		FSharedInternalInstallRequest				Request;
-		TArray<Private::FResolvedContainerChunks>	ResolvedChunks;
-		TArray<FChunkHttpRequestHandle>				HttpRequestHandles;
-		FString										ErrorReason;
-		FOnDemandInstallProgress					Progress;
-		uint64										ResolvedChunkCount = 0;
-		uint64										DownloadedChunkCount = 0;
-		uint64										LastProgressCycles = 0;
-		uint64										StartTimeCycles = FPlatformTime::Cycles64();
-		int32										Priority = 0; 
-		uint32										SeqNo = 0;
-		std::atomic_bool							bHttpRequestsIssued{false};
+		struct FInstall
+		{
+			FInstall(
+				FOnDemandInstallArgs&& InArgs,
+				FOnDemandInstallCompleted&& InOnCompleted,
+				FOnDemandInstallProgressed&& InOnProgress)
+					: Args(MoveTemp(InArgs))
+					, OnCompleted(MoveTemp(InOnCompleted))
+					, OnProgress(MoveTemp(InOnProgress)) { }
+
+			FOnDemandInstallArgs						Args;
+			FOnDemandInstallCompleted					OnCompleted;
+			FOnDemandInstallProgressed					OnProgress;
+			FSharedInternalInstallRequest				Request;
+			TArray<Private::FResolvedContainerChunks>	ResolvedChunks;
+			TArray<FChunkHttpRequestHandle>				HttpRequestHandles;
+			FOnDemandInstallProgress					Progress;
+			uint64										ResolvedChunkCount = 0;
+			uint64										DownloadedChunkCount = 0;
+			uint64										LastProgressCycles = 0;
+			std::atomic_bool							bHttpRequestsIssued{false};
+			std::atomic_bool							bNotifyingProgressOnGameThread{false};
+		};
+
+		struct FPurge
+		{
+			FPurge(FOnDemandPurgeArgs&& InArgs, FOnDemandPurgeCompleted&& InOnCompleted)
+				: Args(MoveTemp(InArgs))
+				, OnCompleted(MoveTemp(InOnCompleted)) { }
+
+			FOnDemandPurgeArgs		Args;
+			FOnDemandPurgeCompleted	OnCompleted;
+		};
+
+		struct FDefrag
+		{
+			FDefrag(FOnDemandDefragArgs&& InArgs, FOnDemandDefragCompleted&& InOnCompleted)
+				: Args(MoveTemp(InArgs))
+				, OnCompleted(MoveTemp(InOnCompleted)) { }
+
+			FOnDemandDefragArgs			Args;
+			FOnDemandDefragCompleted	OnCompleted;
+		};
+
+		using FRequestVariant = TVariant<FEmptyVariantState, FInstall, FPurge, FDefrag>;
+
+		FRequest(FOnDemandInstallArgs&& Args, FOnDemandInstallCompleted&& OnCompleted, FOnDemandInstallProgressed&& OnProgress)
+		{
+			Variant.Emplace<FInstall>(MoveTemp(Args), MoveTemp(OnCompleted), MoveTemp(OnProgress));
+		}
+
+		FRequest(FOnDemandPurgeArgs&& Args, FOnDemandPurgeCompleted&& OnCompleted)
+		{
+			Variant.Emplace<FPurge>(MoveTemp(Args), MoveTemp(OnCompleted));
+		}
+
+		FRequest(FOnDemandDefragArgs&& Args, FOnDemandDefragCompleted&& OnCompleted)
+		{
+			Variant.Emplace<FDefrag>(MoveTemp(Args), MoveTemp(OnCompleted));
+		}
+
+		bool			IsInstall() const	{ return Variant.IsType<FInstall>(); } 
+		bool			IsPurge() const		{ return Variant.IsType<FPurge>(); } 
+		bool			IsDefrag() const	{ return Variant.IsType<FDefrag>(); } 
+
+		FInstall&		AsInstall()			{ return Variant.Get<FInstall>(); }
+		FPurge&			AsPurge()			{ return Variant.Get<FPurge>(); }
+		FDefrag&		AsDefrag()			{ return Variant.Get<FDefrag>(); }
+
+		static uint32	NextSeqNo;
+
+		uint32						SeqNo = NextSeqNo++;
+		uint32						Priority = 0;
+		uint64						StartTimeCycles = FPlatformTime::Cycles64();
+		FString						ErrorReason;
+		std::atomic<EIoErrorCode>	ErrorCode{EIoErrorCode::Unknown};
+		FRequestVariant				Variant;
 	};
 
-	static bool ByInstallPriority(const FInstallRequestState& LHS, const FInstallRequestState& RHS)
+	static bool RequestSortPredicate(const FRequest& LHS, const FRequest& RHS)
 	{
-		if (LHS.Priority == RHS.Priority)
+		if (LHS.Variant.GetIndex() == RHS.Variant.GetIndex())
 		{
-			return LHS.SeqNo < RHS.SeqNo;
+			if (LHS.Priority == RHS.Priority)
+			{
+				return LHS.SeqNo < RHS.SeqNo;
+			}
+
+			return LHS.Priority > RHS.Priority;
 		}
 
-		return LHS.Priority > RHS.Priority;
+		return LHS.SeqNo < RHS.SeqNo;
 	}
 
+	using FRequestAllocator	= TSingleThreadedSlabAllocator<FRequest, 32>;
+
 public:
 									FOnDemandContentInstaller(FOnDemandIoStore& IoStore, FOnDemandHttpThread& HttpClient);
 									~FOnDemandContentInstaller();
 
-	FSharedInternalInstallRequest	Enqueue(
+	FSharedInternalInstallRequest	EnqueueInstallRequest(
 										FOnDemandInstallArgs&& Args,
 										FOnDemandInstallCompleted&& OnCompleted,
 										FOnDemandInstallProgressed&& OnProgress);
-	void							Cancel(FSharedInternalInstallRequest InstallRequest);
-	void							UpdatePriority(FSharedInternalInstallRequest InstallRequest, int32 NewPriority);
-
-	void							EnqueuePurge(FOnDemandPurgeArgs&& Args, FOnDemandPurgeCompleted&& OnCompleted);
-	void							EnqueueDefrag(FOnDemandDefragArgs&& Args, FOnDemandDefragCompleted&& OnCompleted);
-
+	void							EnqueuePurgeRequest(FOnDemandPurgeArgs&& Args, FOnDemandPurgeCompleted&& OnCompleted);
+	void							EnqueueDefragRequest(FOnDemandDefragArgs&& Args, FOnDemandDefragCompleted&& OnCompleted);
+	void							CancelInstallRequest(FSharedInternalInstallRequest InstallRequest);
+	void							UpdateInstallRequestPriority(FSharedInternalInstallRequest InstallRequest, int32 NewPriority);
 private:
-	void							ProcessRequest(FInstallRequestState& RequestState);
-	void							TryStartDownloadChunks();
-	void							StartDownloadChunks(FInstallRequestState& RequestState, bool bRemoveAlreadyCachedChunks);
+	void							TryExecuteNextRequest();
+	void							ExecuteRequest(FRequest& Request);
+	void							ProcessInstallRequest(FRequest& Request);
+	void							ExecuteInstallRequest(FRequest& Request, bool bRemoveAlreadyCachedChunks);
+	void							ExecutePurgeRequest(FRequest& Request);
+	void							ExecuteDefragRequest(FRequest& Request);
 	void							ProcessDownloadedChunk(
-										FInstallRequestState& RequestState,
+										FRequest& Request,
 										FChunkHttpRequestHandle& HttpRequest,
 										uint32 HttpStatusCode,
 										FString&& ErrorReason,
 										FIoBuffer&& Chunk);
-	void							NotifyProgress(FInstallRequestState& RequestState);
-	void							CompleteRequest(FInstallRequestState& RequestState);
+	void							NotifyInstallProgress(FRequest& Request);
+	void							CompleteInstallRequest(FRequest& Request);
+	void							CompletePurgeRequest(FRequest& Request);
+	void							CompleteDefragRequest(FRequest& Request);
 	void							Shutdown();
 
-	using FInstallRequestStateMap	= TMap<FSharedInternalInstallRequest, FInstallRequestState*>;
-	using FRequestStateAllocator	= TSingleThreadedSlabAllocator<FInstallRequestState>;
-
 	FOnDemandIoStore&				IoStore;
 	FOnDemandHttpThread&			HttpClient;
 	UE::Tasks::FPipe				InstallerPipe;
 
 	FMutex							Mutex;
-	FInstallRequestStateMap			RequestStateByHandle;
-	FRequestStateAllocator			RequestStateAllocator;
-	TArray<FInstallRequestState*>	RequestQueue;
-	FInstallRequestState*			CurrentRequest = nullptr;
-	uint32							NextInstallSeqNo = 0;
+	FRequestAllocator				RequestAllocator;
+	TArray<FRequest*>				RequestQueue;
+	FRequest*						CurrentRequest = nullptr;
 	std::atomic_bool				bShuttingDown{false};
 };
 
diff --git a/Engine/Source/Runtime/Experimental/IoStore/OnDemand/Private/OnDemandIoStore.cpp b/Engine/Source/Runtime/Experimental/IoStore/OnDemand/Private/OnDemandIoStore.cpp
--- a/Engine/Source/Runtime/Experimental/IoStore/OnDemand/Private/OnDemandIoStore.cpp
+++ b/Engine/Source/Runtime/Experimental/IoStore/OnDemand/Private/OnDemandIoStore.cpp
@@ -557,7 +557,7 @@ FOnDemandInstallRequest FOnDemandIoStore::Install(
 {
 	LLM_SCOPE_BYNAME(TEXT("FileSystem/OnDemandIoStore"));
 	
-	FSharedInternalInstallRequest InstallRequest = Installer->Enqueue(
+	FSharedInternalInstallRequest InstallRequest = Installer->EnqueueInstallRequest(
 		MoveTemp(Args),
 		MoveTemp(OnCompleted),
 		MoveTemp(OnProgress));
@@ -567,12 +567,14 @@ FOnDemandInstallRequest FOnDemandIoStore::Install(
 
 void FOnDemandIoStore::Purge(FOnDemandPurgeArgs&& Args, FOnDemandPurgeCompleted&& OnCompleted)
 {
-	Installer->EnqueuePurge(MoveTemp(Args), MoveTemp(OnCompleted));
+	LLM_SCOPE_BYNAME(TEXT("FileSystem/OnDemandIoStore"));
+	Installer->EnqueuePurgeRequest(MoveTemp(Args), MoveTemp(OnCompleted));
 }
 
 void FOnDemandIoStore::Defrag(FOnDemandDefragArgs&& Args, FOnDemandDefragCompleted&& OnCompleted)
 {
-	Installer->EnqueueDefrag(MoveTemp(Args), MoveTemp(OnCompleted));
+	LLM_SCOPE_BYNAME(TEXT("FileSystem/OnDemandIoStore"));
+	Installer->EnqueueDefragRequest(MoveTemp(Args), MoveTemp(OnCompleted));
 }
 
 FIoStatus FOnDemandIoStore::Unmount(FStringView MountId)
@@ -1662,12 +1664,12 @@ void FOnDemandIoStore::GetReferencedContentByHandle(TMap<FOnDemandWeakContentHan
 void FOnDemandIoStore::CancelInstallRequest(FSharedInternalInstallRequest InstallRequest)
 {
 	check(InstallRequest.IsValid());
-	Installer->Cancel(InstallRequest);
+	Installer->CancelInstallRequest(InstallRequest);
 }
 
 void FOnDemandIoStore::UpdateInstallRequestPriority(FSharedInternalInstallRequest InstallRequest, int32 NewPriority)
 {
-	Installer->UpdatePriority(InstallRequest, NewPriority);
+	Installer->UpdateInstallRequestPriority(InstallRequest, NewPriority);
 }
 
 FOnDemandCacheUsage FOnDemandIoStore::GetCacheUsage(const FOnDemandGetCacheUsageArgs& Args) const
diff --git a/Engine/Source/Runtime/Experimental/IoStore/OnDemandCore/Internal/IO/IoStoreOnDemandInternals.h b/Engine/Source/Runtime/Experimental/IoStore/OnDemandCore/Internal/IO/IoStoreOnDemandInternals.h
--- a/Engine/Source/Runtime/Experimental/IoStore/OnDemandCore/Internal/IO/IoStoreOnDemandInternals.h
+++ b/Engine/Source/Runtime/Experimental/IoStore/OnDemandCore/Internal/IO/IoStoreOnDemandInternals.h
@@ -63,42 +63,17 @@ class FOnDemandInternalInstallRequest
 {
 public:
 	using EStatus = FOnDemandRequest::EStatus;
+	
+	explicit FOnDemandInternalInstallRequest(UPTRINT InInstallerRequest)
+		: InstallerRequest(InInstallerRequest) { }
 
-	UE_API static uint32 NextSeqNo;
-
-	static bool SortByPriority(FSharedInternalInstallRequest LHS, FSharedInternalInstallRequest RHS)
-	{
-		if (LHS->Args.Priority == RHS->Args.Priority)
-		{
-			return LHS->SeqNo > RHS->SeqNo;
-		}
-		return static_cast<uint8>(LHS->Args.Priority) < static_cast<uint8>(RHS->Args.Priority);
-	};
-
-	FOnDemandInternalInstallRequest(
-		FOnDemandInstallArgs&& InArgs,
-		FOnDemandInstallCompleted&& InOnCompleted,
-		FOnDemandInstallProgressed&& InOnProgress)
-			: Args(MoveTemp(InArgs))
-			, OnCompleted(MoveTemp(InOnCompleted))
-			, OnProgress(MoveTemp(InOnProgress))
-	{ }
-
-	bool IsCancelRequested() const { return EIoErrorCode::Cancelled == ErrorCode.load(std::memory_order_relaxed); }
-
-	FOnDemandInstallArgs		Args;
-	FOnDemandInstallCompleted	OnCompleted;
-	FOnDemandInstallProgressed	OnProgress;
-	uint32						SeqNo = NextSeqNo++;
-	std::atomic<EStatus>		Status{EStatus::Pending};
-	std::atomic<EIoErrorCode>	ErrorCode{EIoErrorCode::Unknown};
+	UPTRINT					InstallerRequest = 0;
+	std::atomic<EStatus>	Status{EStatus::Pending};
 };
 
 } // namespace UE::IoStore
 
 ////////////////////////////////////////////////////////////////////////////////////////////////////////////////////////////////////////////////////////////////
 UE_API FString LexToString(const UE::IoStore::FOnDemandInternalContentHandle& Handle);
-UE_API FString LexToString(const UE::IoStore::FOnDemandInternalInstallRequest& Request);
-
 
 #undef UE_API
diff --git a/Engine/Source/Runtime/Experimental/IoStore/OnDemandCore/Private/IoStoreOnDemandInternals.cpp b/Engine/Source/Runtime/Experimental/IoStore/OnDemandCore/Private/IoStoreOnDemandInternals.cpp
--- a/Engine/Source/Runtime/Experimental/IoStore/OnDemandCore/Private/IoStoreOnDemandInternals.cpp
+++ b/Engine/Source/Runtime/Experimental/IoStore/OnDemandCore/Private/IoStoreOnDemandInternals.cpp
@@ -58,8 +58,6 @@ FString LexToString(const FOnDemandContentHandle& Handle)
 	return Handle.IsValid() ? ::LexToString(*Handle.Handle) : TEXT("Invalid");
 }
 
-uint32 FOnDemandInternalInstallRequest::NextSeqNo = 0;
-
 ///////////////////////////////////////////////////////////////////////////////
 FOnDemandInstallRequest::FOnDemandInstallRequest()
 {
@@ -122,14 +120,3 @@ FString LexToString(const UE::IoStore::FOnDemandInternalContentHandle& Handle)
 {
 	return FString::Printf(TEXT("Id=0x%llX DebugName='%s'"), Handle.HandleId(), *Handle.DebugName);
 }
-
-FString LexToString(const UE::IoStore::FOnDemandInternalInstallRequest& Request)
-{
-	return FString::Printf(TEXT("Id=0x%llX, DebugName='%s', ContentHandle='%s'"),
-		UPTRINT(&Request),
-		Request.Args.DebugName.IsEmpty()
-			? TEXT("NoName")
-			: *Request.Args.DebugName,
-		*LexToString(Request.Args.ContentHandle));
-}
-
