Basic on-demand install analytics.

#rnx

diff --git a/Engine/Source/Runtime/Experimental/IoStore/OnDemand/Private/OnDemandContentInstaller.cpp b/Engine/Source/Runtime/Experimental/IoStore/OnDemand/Private/OnDemandContentInstaller.cpp
--- a/Engine/Source/Runtime/Experimental/IoStore/OnDemand/Private/OnDemandContentInstaller.cpp
+++ b/Engine/Source/Runtime/Experimental/IoStore/OnDemand/Private/OnDemandContentInstaller.cpp
@@ -347,6 +347,11 @@ void FOnDemandContentInstaller::UpdateInstallRequestPriority(FSharedInternalInst
 		});
 }
 
+void FOnDemandContentInstaller::ReportAnalytics(TArray<FAnalyticsEventAttribute>& OutAnalyticsArray) const
+{
+	FOnDemandContentInstallerStats::ReportAnalytics(OutAnalyticsArray);
+}
+
 void FOnDemandContentInstaller::TryExecuteNextRequest()
 {
 	if (bShuttingDown.load(std::memory_order_relaxed))
diff --git a/Engine/Source/Runtime/Experimental/IoStore/OnDemand/Private/OnDemandContentInstaller.h b/Engine/Source/Runtime/Experimental/IoStore/OnDemand/Private/OnDemandContentInstaller.h
--- a/Engine/Source/Runtime/Experimental/IoStore/OnDemand/Private/OnDemandContentInstaller.h
+++ b/Engine/Source/Runtime/Experimental/IoStore/OnDemand/Private/OnDemandContentInstaller.h
@@ -15,6 +15,7 @@
 #include <atomic>
 
 class FIoBuffer;
+struct FAnalyticsEventAttribute;
 
 namespace UE::IoStore
 {
@@ -172,6 +173,8 @@ public:
 	void							EnqueueDefragRequest(FOnDemandDefragArgs&& Args, FOnDemandDefragCompleted&& OnCompleted);
 	void							CancelInstallRequest(FSharedInternalInstallRequest InstallRequest);
 	void							UpdateInstallRequestPriority(FSharedInternalInstallRequest InstallRequest, int32 NewPriority);
+	void							ReportAnalytics(TArray<FAnalyticsEventAttribute>& OutAnalyticsArray) const;
+
 private:
 	void							TryExecuteNextRequest();
 	void							ExecuteRequest(FRequest& Request);
diff --git a/Engine/Source/Runtime/Experimental/IoStore/OnDemand/Private/OnDemandIoStore.cpp b/Engine/Source/Runtime/Experimental/IoStore/OnDemand/Private/OnDemandIoStore.cpp
--- a/Engine/Source/Runtime/Experimental/IoStore/OnDemand/Private/OnDemandIoStore.cpp
+++ b/Engine/Source/Runtime/Experimental/IoStore/OnDemand/Private/OnDemandIoStore.cpp
@@ -1756,6 +1756,11 @@ void FOnDemandIoStore::SetStreamingOptions(EOnDemandStreamingOptions Options)
 
 void FOnDemandIoStore::ReportAnalytics(TArray<FAnalyticsEventAttribute>& OutAnalyticsArray) const
 {
+	if (Installer.IsValid())
+	{
+		Installer->ReportAnalytics(OutAnalyticsArray);
+	}
+
 	if (StreamingBackend.IsValid())
 	{
 		return StreamingBackend->ReportAnalytics(OutAnalyticsArray);
diff --git a/Engine/Source/Runtime/Experimental/IoStore/OnDemand/Private/Statistics.cpp b/Engine/Source/Runtime/Experimental/IoStore/OnDemand/Private/Statistics.cpp
--- a/Engine/Source/Runtime/Experimental/IoStore/OnDemand/Private/Statistics.cpp
+++ b/Engine/Source/Runtime/Experimental/IoStore/OnDemand/Private/Statistics.cpp
@@ -43,6 +43,12 @@ static FAutoConsoleVariableRef CVar_ReportCacheAnalytics(
 	GIasReportCacheAnalyticsEnabled,
 	TEXT("Enables reporting statics on our file cache usage to the analytics system"));
 
+bool GIadReportAnalyticsEnabled = true;
+static FAutoConsoleVariableRef CVar_ReportIadAnalyticsEnabled(
+	TEXT("iad.ReportAnalytics"),
+	GIadReportAnalyticsEnabled,
+	TEXT("Enables reporting analytics for individual asset downloads."));
+
 #if UE_ENABLE_ONSCREEN_STATISTICS
 
 bool GIasDisplayOnScreenStatistics = false;
@@ -1092,6 +1098,22 @@ void FOnDemandContentInstallerStats::OnRequestCompleted(
 	//GAvgInstallCacheHitRatio.Set(GInstallCacheHitRatio.GetMean());
 }
 
+void FOnDemandContentInstallerStats::ReportAnalytics(TArray<FAnalyticsEventAttribute>& OutAnalyticsArray)
+{
+	if (GIadReportAnalyticsEnabled == false)
+	{
+		return;
+	}
+
+	AppendAnalyticsEventAttributeArray(OutAnalyticsArray
+		, TEXT("IadTotalInstallCount"), GTotalInstallCount.Get()
+		, TEXT("IadTotalInstallErrorCount"), GTotalInstallErrorCount.Get()
+		, TEXT("IadTotalWriteErrorInstallErrorCount"), GTotalWriteErrorInstallErrorCount.Get()
+		, TEXT("IadTotalInstallDurationMs"), GTotalInstallDurationMs.Get()
+		, TEXT("IadAvgCacheHitRatio"), GInstallCacheHitRatio.GetMean()
+	);
+}
+
 } // namespace UE::IoStore
 
 #undef UE_ENABLE_ONSCREEN_STATISTICS
diff --git a/Engine/Source/Runtime/Experimental/IoStore/OnDemand/Private/Statistics.h b/Engine/Source/Runtime/Experimental/IoStore/OnDemand/Private/Statistics.h
--- a/Engine/Source/Runtime/Experimental/IoStore/OnDemand/Private/Statistics.h
+++ b/Engine/Source/Runtime/Experimental/IoStore/OnDemand/Private/Statistics.h
@@ -100,6 +100,7 @@ public:
 		double CacheHitRatio,
 		uint64 DurationCycles,
 		EIoErrorCode ErrorCode) IAS_STATISTICS_IMPL()
+	static void ReportAnalytics(TArray<FAnalyticsEventAttribute>& OutAnalyticsArray) IAS_STATISTICS_IMPL()
 };
 
 #undef IAS_STATISTICS_IMPL
