Detect and remove null constriants from CCD swept constraints prior to application

diff --git a/Engine/Source/Runtime/Experimental/Chaos/Private/Chaos/CCDUtilities.cpp b/Engine/Source/Runtime/Experimental/Chaos/Private/Chaos/CCDUtilities.cpp
--- a/Engine/Source/Runtime/Experimental/Chaos/Private/Chaos/CCDUtilities.cpp
+++ b/Engine/Source/Runtime/Experimental/Chaos/Private/Chaos/CCDUtilities.cpp
@@ -13,6 +13,9 @@ namespace Chaos
 {
 	namespace CVars
 	{
+		bool bClearNullCCDConstraintsOnApply = true;
+		FAutoConsoleVariableRef CVarClearNullCCDConstraintsOnApply(TEXT("p.Chaos.Collision.CCD.ClearNullConstraints"), bClearNullCCDConstraintsOnApply, TEXT("Remove null constraints encountered during correction application for CCD."));
+
 		// @todo(chaos): most of these cvars should be settings per particle
 
 		// NOTE: With this disabled secondary CCD collisions will often be missed
@@ -918,6 +921,14 @@ namespace Chaos
 		//		CCD - Kinematic
 		TArray<FPBDCollisionConstraint*> Constraints;
 		Constraints.Append(SweptConstraints);
+
+		if(Chaos::CVars::bClearNullCCDConstraintsOnApply)
+		{
+			Constraints.Remove(nullptr);
+			const int32 NumNullConstraints = SweptConstraints.Num() - Constraints.Num();
+			ensureMsgf(NumNullConstraints == 0, TEXT("Found %d nullptr constraints during CCD correction application"), NumNullConstraints);
+		}
+
 		Constraints.Sort([](const FPBDCollisionConstraint& L, const FPBDCollisionConstraint& R)
 			{
 				const FGenericParticleHandle LP0 = L.GetParticle0();
