Dataflow : fix issue with for each calls when there's no iteration ( empty array )
- set default values on the output

diff --git a/Engine/Source/Runtime/Experimental/Dataflow/Engine/Private/Dataflow/DataflowSubGraphNodes.cpp b/Engine/Source/Runtime/Experimental/Dataflow/Engine/Private/Dataflow/DataflowSubGraphNodes.cpp
--- a/Engine/Source/Runtime/Experimental/Dataflow/Engine/Private/Dataflow/DataflowSubGraphNodes.cpp
+++ b/Engine/Source/Runtime/Experimental/Dataflow/Engine/Private/Dataflow/DataflowSubGraphNodes.cpp
@@ -672,17 +672,34 @@ void FDataflowCallSubGraphNode::Evaluate(UE::Dataflow::FContext& Context, const
 
 					TArray<FDataflowOutput*> OutputsToEvaluate = GetOutputs();
 
-					// Even if the evelaute function is for a single output 
-					// We are pulling all outputs otherwise this may force us to evaluate the loops multiple time 
-					// whioch may be more expensive
-					for (int32 Iteration = 0; Iteration < NumIterations; ++Iteration)
+					if (NumIterations > 0)
 					{
-						SubGraphContext.SetCurrentArrayIndex(Iteration);
+						// Even if the evaluate function is for a single output 
+						// We are pulling all outputs otherwise this may force us to evaluate the loops multiple time 
+						// whioch may be more expensive
+						for (int32 Iteration = 0; Iteration < NumIterations; ++Iteration)
+						{
+							SubGraphContext.SetCurrentArrayIndex(Iteration);
+							for (const FDataflowOutput* OutToEvaluate : OutputsToEvaluate)
+							{
+								if (OutToEvaluate)
+								{
+									SubGraphContext.EvaluateSubGraph(*SubGraphOutputNode, *this, *OutToEvaluate);
+								}
+							}
+							SubGraphContext.ClearAllData();
+						}
+					}
+					else
+					{
+						// no iteration , mean we need to set default values to the output
 						for (const FDataflowOutput* OutToEvaluate : OutputsToEvaluate)
 						{
-							SubGraphContext.EvaluateSubGraph(*SubGraphOutputNode, *this, *OutToEvaluate);
+							if (OutToEvaluate)
+							{
+								OutToEvaluate->SetNullValue(Context);
+							}
 						}
-						SubGraphContext.ClearAllData();
 					}
 				}
 			}
