Chaos - fix assert in collision detection after disabling Convex Optimization from a Midphase Modifier
- clean up CCD TOI so it is always reset (mainly to make debugging CCD issues easier)

Disabling Convex Optimization requires that collisions be destroyed if they were created with convex optimization enabled. Before we destroy collisions we must remove them from the constraint graph. The MidPhase modifier now handles this, but had to be moved from CollisionConstraints to the Evolution which knows about the graph.

#rnx

diff --git a/Engine/Source/Runtime/Experimental/Chaos/Private/Chaos/Collision/ParticlePairMidPhase.cpp b/Engine/Source/Runtime/Experimental/Chaos/Private/Chaos/Collision/ParticlePairMidPhase.cpp
--- a/Engine/Source/Runtime/Experimental/Chaos/Private/Chaos/Collision/ParticlePairMidPhase.cpp
+++ b/Engine/Source/Runtime/Experimental/Chaos/Private/Chaos/Collision/ParticlePairMidPhase.cpp
@@ -26,6 +26,7 @@ TRACE_DECLARE_INT_COUNTER_EXTERN(ChaosTraceCounter_MidPhase_NumShapePair);
 TRACE_DECLARE_INT_COUNTER_EXTERN(ChaosTraceCounter_MidPhase_NumGeneric);
 
 extern bool Chaos_Collision_NarrowPhase_AABBBoundsCheck;
+extern bool bChaosCCDResetTOIFix;
 
 // Enable for extended stats. Slow but useful for determining counts and relative costs of midphases
 #if 0
@@ -62,6 +63,10 @@ namespace Chaos
 		FAutoConsoleVariableRef CVarChaos_Collision_EnableShapePairs(TEXT("p.Chaos.Collision.EnableShapePairs"), bChaos_Collision_MidPhase_EnableShapePairs, TEXT(""));
 		FAutoConsoleVariableRef CVarChaos_Collision_MaxShapePairs(TEXT("p.Chaos.Collision.MaxShapePairs"), Chaos_Collision_MidPhase_MaxShapePairs, TEXT(""));
 
+		// See FGenericParticlePairMidPhase::SetConvexOptimizationIsActive
+		bool bChaosMidPhaseModifierConvexOptimizerFix = true;
+		FAutoConsoleVariableRef CVarChaosMidPhaseModifierConvexOptimizerFix(TEXT("p.Chaos.Bugs.MidPhaseModifierConvexOptimizerFix"), bChaosMidPhaseModifierConvexOptimizerFix, TEXT(""));
+
 		extern int32 ChaosOneWayInteractionPairCollisionMode;
 
 		extern bool bChaosForceMACD;
@@ -480,7 +485,10 @@ namespace Chaos
 			// Constraint may have been previously used with CCD enabled (e.g., a midphase modifier)
 			// so we need to make sure that the CCD flag is disabled
 			Constraint->SetCCDEnabled(false);
-			Constraint->ResetCCDTimeOfImpact();
+			if (bChaosCCDResetTOIFix)
+			{
+				Constraint->ResetCCDTimeOfImpact();
+			}
 
 			// If the constraint was not used last frame, it needs to be reset, otherwise we will try to reuse
 			if (!bWasUpdatedLastTick || (Constraint->GetManifoldPoints().Num() == 0))
@@ -561,7 +569,10 @@ namespace Chaos
 			check(!Constraint->IsEnabled());
 			Constraint->SetCCDEnabled(true);
 			Constraint->SetCCDSweepEnabled(bEnableCCDSweep);
-			Constraint->ResetCCDTimeOfImpact();
+			if (bChaosCCDResetTOIFix)
+			{
+				Constraint->ResetCCDTimeOfImpact();
+			}
 		}
 
 		if (!bEnableCCDSweep)
@@ -1233,18 +1244,21 @@ namespace Chaos
 		// 
 		// NOTE: bIsConvexOptimizationActive can be changed in a modifier which gets reset every frame so we
 		// use bIsConvexOptimizationUsed to remember what we actually used rather than just resetting on change.
-		if (Flags.bIsConvexOptimizationActive != Flags.bIsConvexOptimizationUsed)
+		if (CVars::bChaosMidPhaseModifierConvexOptimizerFix)
 		{
-			// The caller should remove collisions from the graph here
-			for (TMap<uint64, FPBDCollisionConstraintPtr>::ElementType& KVP : Constraints)
+			if (Flags.bIsConvexOptimizationActive != Flags.bIsConvexOptimizationUsed)
 			{
-				FPBDCollisionConstraintPtr& Constraint = KVP.Value;
-				OnDestroyingCollision(Constraint.Get());
-			}
+				// The caller should remove collisions from the graph here
+				for (TMap<uint64, FPBDCollisionConstraintPtr>::ElementType& KVP : Constraints)
+				{
+					FPBDCollisionConstraintPtr& Constraint = KVP.Value;
+					OnDestroyingCollision(Constraint.Get());
+				}
 
-			// NOTE: An assert here would indicate that the OnDestroyingCollisions callback did not 
-			// remove collisions from the constraint graph.
-			Constraints.Reset();
+				// NOTE: An assert here would indicate that the OnDestroyingCollisions callback did not 
+				// remove collisions from the constraint graph.
+				Constraints.Reset();
+			}
 		}
 	}
 
@@ -1858,7 +1872,10 @@ namespace Chaos
 			// CCD may be temporarily disabled by the user (via a midphase modifier) or because we are moving slowly.
 			Constraint->SetCCDEnabled(Flags.bIsCCDActive);
 			Constraint->SetCCDSweepEnabled(bUseCCDSweep);
-			Constraint->ResetCCDTimeOfImpact();
+			if (bChaosCCDResetTOIFix)
+			{
+				Constraint->ResetCCDTimeOfImpact();
+			}
 
 			// NOTE: Probe constraints are always active and we run collision detection for them at the end opf the frame
 			bool bIsActive = true;
diff --git a/Engine/Source/Runtime/Experimental/Chaos/Private/Chaos/CollisionResolution.cpp b/Engine/Source/Runtime/Experimental/Chaos/Private/Chaos/CollisionResolution.cpp
--- a/Engine/Source/Runtime/Experimental/Chaos/Private/Chaos/CollisionResolution.cpp
+++ b/Engine/Source/Runtime/Experimental/Chaos/Private/Chaos/CollisionResolution.cpp
@@ -79,6 +79,9 @@ FAutoConsoleVariableRef CVarChaosCollisionAABBBoundsCheck(TEXT("p.Chaos.Collisio
 bool bChaos_Collision_ShapesArrayMode = true;
 FAutoConsoleVariableRef CVarChaos_Collision_ShapesArrayMode(TEXT("p.Chaos.Collision.ShapesArrayMode"), bChaos_Collision_ShapesArrayMode, TEXT(""));
 
+bool bChaosCCDResetTOIFix = true;
+FAutoConsoleVariableRef CVarChaos_Collision_ResetTOIFix(TEXT("p.Chaos.Bugs.ResetTOIFix"), bChaosCCDResetTOIFix, TEXT(""));
+
 namespace Chaos
 {
 	namespace Collisions
@@ -1046,6 +1049,11 @@ namespace Chaos
 			const FReal LengthCCD = DirCCD.SafeNormalize();
 			const bool bUseCCD = (LengthCCD > 0);
 
+			if (!bChaosCCDResetTOIFix)
+			{
+				Constraint.ResetCCDTimeOfImpact();
+			}
+
 			if (bUseCCD)
 			{
 				const EContactShapesType ShapePairType = Constraint.GetShapesType();
diff --git a/Engine/Source/Runtime/Experimental/Chaos/Private/Chaos/MidPhaseModification.cpp b/Engine/Source/Runtime/Experimental/Chaos/Private/Chaos/MidPhaseModification.cpp
--- a/Engine/Source/Runtime/Experimental/Chaos/Private/Chaos/MidPhaseModification.cpp
+++ b/Engine/Source/Runtime/Experimental/Chaos/Private/Chaos/MidPhaseModification.cpp
@@ -30,11 +30,17 @@ namespace Chaos
 		// is called repreatedly.
 		const auto& RemoveCollisionFromGraph = [this](FPBDCollisionConstraint* InConstraint) -> ECollisionVisitorResult
 			{
-				Accessor->GetEvolution().RemoveConstraintFromConstraintGraph(InConstraint->GetConstraintHandle());
+				if ((Accessor != nullptr) && (InConstraint != nullptr))
+				{
+					Accessor->GetEvolution().RemoveConstraintFromConstraintGraph(InConstraint->GetConstraintHandle());
+				}
 				return ECollisionVisitorResult::Continue;
 			};
 
-		MidPhase->SetConvexOptimizationIsActive(false, RemoveCollisionFromGraph);
+		if (MidPhase != nullptr)
+		{
+			MidPhase->SetConvexOptimizationIsActive(false, RemoveCollisionFromGraph);
+		}
 	}
 
 	void FMidPhaseModifier::GetParticles(
