Dataflow : fix rendering of box and sphere through the render interface  to make sure we have normals and vertex color

diff --git a/Engine/Source/Runtime/Experimental/Chaos/Private/GeometryCollection/Facades/CollectionRenderingFacade.cpp b/Engine/Source/Runtime/Experimental/Chaos/Private/GeometryCollection/Facades/CollectionRenderingFacade.cpp
--- a/Engine/Source/Runtime/Experimental/Chaos/Private/GeometryCollection/Facades/CollectionRenderingFacade.cpp
+++ b/Engine/Source/Runtime/Experimental/Chaos/Private/GeometryCollection/Facades/CollectionRenderingFacade.cpp
@@ -214,34 +214,81 @@ namespace GeometryCollection::Facades
 
 	void FRenderingFacade::AddBox(const FVector3f& InMinVertex, const FVector3f& InMaxVertex)
 	{
-		constexpr int32 NumVertices = 8;
+		constexpr const int32 NumCorners = 8;
+		const FVector3f Corners[NumCorners] =
+		{
+			FVector3f(InMinVertex.X, InMinVertex.Y, InMinVertex.Z), // -X / -Y / -Z
+			FVector3f(InMaxVertex.X, InMinVertex.Y, InMinVertex.Z), // +X / -Y / -Z
+			FVector3f(InMaxVertex.X, InMaxVertex.Y, InMinVertex.Z), // +X / +Y / -Z
+			FVector3f(InMinVertex.X, InMaxVertex.Y, InMinVertex.Z), // -X / +Y / -Z
+			FVector3f(InMinVertex.X, InMinVertex.Y, InMaxVertex.Z), // -X / -Y / +Z
+			FVector3f(InMaxVertex.X, InMinVertex.Y, InMaxVertex.Z), // +X / -Y / +Z
+			FVector3f(InMaxVertex.X, InMaxVertex.Y, InMaxVertex.Z), // +X / +Y / +Z
+			FVector3f(InMinVertex.X, InMaxVertex.Y, InMaxVertex.Z), // -X / +Y / +Z
+		};
+
+		constexpr int32 NumVertices = 24; // 6 faces with each 4 vertices
 		constexpr int32 NumTriangles = 12;
 
 		TArray<FVector3f> Vertices; Vertices.AddUninitialized(NumVertices);
+		TArray<FVector3f> Normals; Normals.AddUninitialized(NumVertices);
 		TArray<FIntVector> Tris; Tris.AddUninitialized(NumTriangles);
 
-		// Add vertices
-		Vertices[0] = InMinVertex;
-		Vertices[1] = FVector3f(InMaxVertex.X, InMinVertex.Y, InMinVertex.Z);
-		Vertices[2] = FVector3f(InMaxVertex.X, InMaxVertex.Y, InMinVertex.Z);
-		Vertices[3] = FVector3f(InMinVertex.X, InMaxVertex.Y, InMinVertex.Z);
-		Vertices[4] = FVector3f(InMinVertex.X, InMinVertex.Y, InMaxVertex.Z);
-		Vertices[5] = FVector3f(InMaxVertex.X, InMinVertex.Y, InMaxVertex.Z);
-		Vertices[6] = InMaxVertex;
-		Vertices[7] = FVector3f(InMinVertex.X, InMaxVertex.Y, InMaxVertex.Z);
-
-		AddTriangle(Chaos::FTriangle(Vertices[0], Vertices[1], Vertices[3]));
-		AddTriangle(Chaos::FTriangle(Vertices[1], Vertices[2], Vertices[3]));
-		AddTriangle(Chaos::FTriangle(Vertices[0], Vertices[4], Vertices[1]));
-		AddTriangle(Chaos::FTriangle(Vertices[4], Vertices[5], Vertices[1]));
-		AddTriangle(Chaos::FTriangle(Vertices[5], Vertices[2], Vertices[1]));
-		AddTriangle(Chaos::FTriangle(Vertices[5], Vertices[6], Vertices[2]));
-		AddTriangle(Chaos::FTriangle(Vertices[3], Vertices[2], Vertices[6]));
-		AddTriangle(Chaos::FTriangle(Vertices[7], Vertices[3], Vertices[6]));
-		AddTriangle(Chaos::FTriangle(Vertices[0], Vertices[3], Vertices[7]));
-		AddTriangle(Chaos::FTriangle(Vertices[4], Vertices[0], Vertices[7]));
-		AddTriangle(Chaos::FTriangle(Vertices[5], Vertices[4], Vertices[7]));
-		AddTriangle(Chaos::FTriangle(Vertices[5], Vertices[7], Vertices[6]));
+		// +Z face
+		Vertices[0] = Corners[0]; Normals[0] = FVector3f::ZAxisVector;
+		Vertices[1] = Corners[1]; Normals[1] = FVector3f::ZAxisVector;
+		Vertices[2] = Corners[2]; Normals[2] = FVector3f::ZAxisVector;
+		Vertices[3] = Corners[3]; Normals[3] = FVector3f::ZAxisVector;
+		Tris[0] = FIntVector{ 0, 1, 3 };
+		Tris[1] = FIntVector{ 1, 2, 3 };
+
+		// +Z face
+		Vertices[4] = Corners[4]; Normals[4] = -FVector3f::ZAxisVector;
+		Vertices[5] = Corners[5]; Normals[5] = -FVector3f::ZAxisVector;
+		Vertices[6] = Corners[6]; Normals[6] = -FVector3f::ZAxisVector;
+		Vertices[7] = Corners[7]; Normals[7] = -FVector3f::ZAxisVector;
+		Tris[2] = FIntVector{ 5, 4, 7 };
+		Tris[3] = FIntVector{ 5, 7, 6 };
+
+		// +X face
+		Vertices[8]  = Corners[1]; Normals[8]  = FVector3f::XAxisVector;
+		Vertices[9]  = Corners[2]; Normals[9]  = FVector3f::XAxisVector;
+		Vertices[10] = Corners[6]; Normals[10] = FVector3f::XAxisVector;
+		Vertices[11] = Corners[5]; Normals[11] = FVector3f::XAxisVector;
+		Tris[4] = FIntVector{ 8, 11, 10 };
+		Tris[5] = FIntVector{ 8, 10, 9 };
+
+		// -X face
+		Vertices[12] = Corners[0]; Normals[12] = -FVector3f::XAxisVector;
+		Vertices[13] = Corners[3]; Normals[13] = -FVector3f::XAxisVector;
+		Vertices[14] = Corners[7]; Normals[14] = -FVector3f::XAxisVector;
+		Vertices[15] = Corners[4]; Normals[15] = -FVector3f::XAxisVector;
+		Tris[6] = FIntVector{ 12, 13, 14 };
+		Tris[7] = FIntVector{ 12, 14, 15 };
+				
+		// +Y face
+		Vertices[16] = Corners[3]; Normals[16] = FVector3f::YAxisVector;
+		Vertices[17] = Corners[2]; Normals[17] = FVector3f::YAxisVector;
+		Vertices[18] = Corners[6]; Normals[18] = FVector3f::YAxisVector;
+		Vertices[19] = Corners[7]; Normals[19] = FVector3f::YAxisVector;
+		Tris[8] = FIntVector{ 16, 17, 18 };
+		Tris[9] = FIntVector{ 16, 18, 19 };
+
+		// -Y face
+		Vertices[20] = Corners[1]; Normals[20] = -FVector3f::YAxisVector;
+		Vertices[21] = Corners[0]; Normals[21] = -FVector3f::YAxisVector;
+		Vertices[22] = Corners[4]; Normals[22] = -FVector3f::YAxisVector;
+		Vertices[23] = Corners[5]; Normals[23] = -FVector3f::YAxisVector;
+		Tris[10] = FIntVector{ 20, 21, 22 };
+		Tris[11] = FIntVector{ 20, 22, 23 };
+
+		TArray<FLinearColor> Colors; Colors.AddUninitialized(Vertices.Num());
+		for (int32 VertexIdx = 0; VertexIdx < Colors.Num(); ++VertexIdx)
+		{
+			Colors[VertexIdx] = FLinearColor::White;
+		}
+
+		AddSurface(MoveTemp(Vertices), MoveTemp(Tris), MoveTemp(Normals), MoveTemp(Colors));
 	}
 
 	void FRenderingFacade::AddSphere(const FSphere& InSphere, const FLinearColor& InColor)
@@ -255,6 +302,7 @@ namespace GeometryCollection::Facades
 		constexpr int32 NumTriangles = 48;
 
 		TArray<FVector3f> Vertices; Vertices.AddUninitialized(NumVertices);
+		TArray<FVector3f> VertexNormals; VertexNormals.AddUninitialized(NumVertices);
 		TArray<FIntVector> Tris; Tris.AddUninitialized(NumTriangles);
 
 		// Add vertices
@@ -287,6 +335,7 @@ namespace GeometryCollection::Facades
 
 		for (int32 Idx = 0; Idx < NumVertices; ++Idx)
 		{
+			VertexNormals[Idx] = Vertices[Idx]; // normalize local vector
 			Vertices[Idx] *= InRadius;
 			Vertices[Idx] += FVector3f(InCenter);
 		}
@@ -317,15 +366,12 @@ namespace GeometryCollection::Facades
 		Tris[44] = FIntVector(22, 1, 23); Tris[45] = FIntVector(23, 1, 24);
 		Tris[46] = FIntVector(24, 1, 25); Tris[47] = FIntVector(25, 1, 18);
 
-		TArray<FVector3f> VertexNormals; VertexNormals.AddUninitialized(NumVertices);
-		// TODO: Compute vertex normals
-
 		// Add VertexNormal and VertexColor
+		const FLinearColor LinearColor(InColor);
 		TArray<FLinearColor> VertexColors; VertexColors.AddUninitialized(Vertices.Num());
-		for (int32 VertexIdx = 0; VertexIdx < VertexNormals.Num(); ++VertexIdx)
+		for (int32 VertexIdx = 0; VertexIdx < VertexColors.Num(); ++VertexIdx)
 		{
-			VertexNormals[VertexIdx] = FVector3f(0.f);
-			VertexColors[VertexIdx] = FLinearColor(InColor);
+			VertexColors[VertexIdx] = LinearColor;
 		}
 
 		AddSurface(MoveTemp(Vertices), MoveTemp(Tris), MoveTemp(VertexNormals), MoveTemp(VertexColors));
