I/O store on-demand install cache stats and analytics.

#rnx

diff --git a/Engine/Source/Runtime/Experimental/IoStore/OnDemand/Private/OnDemandInstallCache.cpp b/Engine/Source/Runtime/Experimental/IoStore/OnDemand/Private/OnDemandInstallCache.cpp
--- a/Engine/Source/Runtime/Experimental/IoStore/OnDemand/Private/OnDemandInstallCache.cpp
+++ b/Engine/Source/Runtime/Experimental/IoStore/OnDemand/Private/OnDemandInstallCache.cpp
@@ -3,6 +3,7 @@
 #include "OnDemandInstallCache.h"
 #include "OnDemandHttpClient.h"
 #include "OnDemandIoStore.h"
+#include "Statistics.h"
 
 #include "Algo/Accumulate.h"
 #include "Algo/Find.h"
@@ -824,6 +825,7 @@ FIoStatus FCasJournal::Commit(FTransaction&& Transaction)
 
 		if (FileSize < sizeof(FHeader))
 		{
+			FOnDemandInstallCacheStats::OnJournalCommit(EIoErrorCode::ReadError, 0);
 			return FIoStatusBuilder(EIoErrorCode::FileOpenFailed)
 				<< TEXT("Failed to validate CAS journal file '")
 				<< Transaction.JournalFile
@@ -833,6 +835,7 @@ FIoStatus FCasJournal::Commit(FTransaction&& Transaction)
 		FHeader Header;
 		if ((FileHandle->Read(reinterpret_cast<uint8*>(&Header), sizeof(FHeader)) == false) || (Header.IsValid() == false))
 		{
+			FOnDemandInstallCacheStats::OnJournalCommit(EIoErrorCode::SignatureError, 0);
 			return FIoStatusBuilder(EIoErrorCode::ReadError)
 				<< TEXT("Failed to validate CAS journal header '")
 				<< Transaction.JournalFile
@@ -842,6 +845,7 @@ FIoStatus FCasJournal::Commit(FTransaction&& Transaction)
 		const int64 FooterPos = FileSize - sizeof(FFooter);
 		if (FileHandle->Seek(FooterPos) == false)
 		{
+			FOnDemandInstallCacheStats::OnJournalCommit(EIoErrorCode::SignatureError, 0);
 			return FIoStatusBuilder(EIoErrorCode::ReadError)
 				<< TEXT("Failed to validate CAS journal footer '")
 				<< Transaction.JournalFile
@@ -851,6 +855,7 @@ FIoStatus FCasJournal::Commit(FTransaction&& Transaction)
 		FFooter Footer;
 		if ((FileHandle->Read(reinterpret_cast<uint8*>(&Footer), sizeof(FFooter)) == false) || (Footer.IsValid() == false))
 		{
+			FOnDemandInstallCacheStats::OnJournalCommit(EIoErrorCode::SignatureError, 0);
 			return FIoStatusBuilder(EIoErrorCode::ReadError)
 				<< TEXT("Failed to validate CAS journal footer '")
 				<< Transaction.JournalFile
@@ -867,6 +872,7 @@ FIoStatus FCasJournal::Commit(FTransaction&& Transaction)
 
 		if ((EntriesPos < 0) || (FileHandle->Seek(EntriesPos) == false))
 		{
+			FOnDemandInstallCacheStats::OnJournalCommit(EIoErrorCode::WriteError, 0);
 			return FIoStatusBuilder(EIoErrorCode::FileOpenFailed)
 				<< TEXT("Failed to open CAS journal '")
 				<< Transaction.JournalFile
@@ -878,6 +884,7 @@ FIoStatus FCasJournal::Commit(FTransaction&& Transaction)
 			reinterpret_cast<const uint8*>(Transaction.Entries.GetData()),
 			TotalEntrySize) == false)
 		{
+			FOnDemandInstallCacheStats::OnJournalCommit(EIoErrorCode::WriteError, 0);
 			return FIoStatusBuilder(EIoErrorCode::WriteError)
 				<< TEXT("Failed to write CAS journal entries to '")
 				<< Transaction.JournalFile
@@ -888,6 +895,7 @@ FIoStatus FCasJournal::Commit(FTransaction&& Transaction)
 		FMemory::Memcpy(&Footer.Magic, &FFooter::MagicSequence, sizeof(FFooter::MagicSequence));
 		if (FileHandle->Write(reinterpret_cast<uint8*>(&Footer), sizeof(FFooter)) == false)
 		{
+			FOnDemandInstallCacheStats::OnJournalCommit(EIoErrorCode::WriteError, 0);
 			return FIoStatusBuilder(EIoErrorCode::WriteError)
 				<< TEXT("Failed to write CAS journal footer to '")
 				<< Transaction.JournalFile
@@ -896,12 +904,15 @@ FIoStatus FCasJournal::Commit(FTransaction&& Transaction)
 
 		if (FileHandle->Flush() == false)
 		{
+			FOnDemandInstallCacheStats::OnJournalCommit(EIoErrorCode::WriteError, 0);
 			return EIoErrorCode::WriteError;
 		}
 
 		UE_LOG(LogIoStoreOnDemand, Log, TEXT("Committed %d CAS journal entries of total %.2lf KiB to '%s'"),
 			Transaction.Entries.Num(), ToKiB(TotalEntrySize), *Transaction.JournalFile);
 
+		FOnDemandInstallCacheStats::OnJournalCommit(EIoErrorCode::Ok, TotalEntrySize);
+
 		return EIoErrorCode::Ok;
 	}
 }
@@ -1336,7 +1347,8 @@ private:
 													const TArray<FSharedOnDemandContainer>& Containers, 
 													const TArray<TBitArray<>>& ChunkEntryIndices,
 													const TSet<FIoHash>& ChunksToInstall,
-													FCasBlockInfoMap& BlockInfoMap) const;
+													FCasBlockInfoMap& BlockInfoMap,
+													uint64& OutTotalReferencedBytes) const;
 	FIoStatus									Purge(FCasBlockInfoMap& BlockInfo, uint64 TotalBytesToPurge, uint64& OutTotalPurgedBytes);
 	FIoStatus									Defrag(
 													const TArray<FSharedOnDemandContainer>& Containers,
@@ -1467,6 +1479,8 @@ FOnDemandInstallCache::FOnDemandInstallCache(const FOnDemandInstallCacheConfig&
 	TArray<FCasAddr> RemovedChunks;
 	if (FIoStatus Verify = Cas.Verify(RemovedChunks); !Verify.IsOk())
 	{
+		FOnDemandInstallCacheStats::OnCasVerificationError(RemovedChunks.Num());
+
 		// Try to recover if the CAS blocks on disk doesn't match
 		FCasJournal::FTransaction Transaction = FCasJournal::Begin(GetJournalFilename());
 		for (const FCasAddr& Addr : RemovedChunks)
@@ -1511,6 +1525,8 @@ FOnDemandInstallCache::FOnDemandInstallCache(const FOnDemandInstallCacheConfig&
 	{
 		if (Status.GetErrorCode() != EIoErrorCode::NotFound)
 		{
+			FOnDemandInstallCacheStats::OnStartupError(Status.GetErrorCode());
+
 			UE_LOG(LogIoStoreOnDemand, Warning, TEXT("Failed to replay install cache journal file '%s, reason '%s'"),
 				*Status.ToString(), *JournalFile);
 
@@ -1874,24 +1890,38 @@ FIoStatus FOnDemandInstallCache::Purge(TSet<FIoHash>&& ChunksToInstall)
 	IoStore.GetReferencedContent(Containers, ChunkEntryIndices);
 	check(Containers.Num() == ChunkEntryIndices.Num());
 
-	const uint64 TotalUncachedBytes = AddReferencesToBlocks(Containers, ChunkEntryIndices, ChunksToInstall, BlockInfo);
+	uint64 ReferencedBytes = 0;
+	uint64 FragmentedBytes = 0;
+	uint64 TotalReferencedBlockBytes = 0;
+
+	const uint64 TotalUncachedBytes = AddReferencesToBlocks(Containers, ChunkEntryIndices, ChunksToInstall, BlockInfo, ReferencedBytes);
+	for (const TPair<FCasBlockId, FCasBlockInfo>& Kv : BlockInfo)
+	{
+		const FCasBlockInfo& Info = Kv.Value;
+		if (Info.RefSize < Info.FileSize)
+		{
+			FragmentedBytes += (Info.FileSize - Info.RefSize);
+		}
+		if (Info.RefSize > 0)
+		{
+			TotalReferencedBlockBytes += Info.FileSize;
+		}
+	}
+
+	FOnDemandInstallCacheStats::OnCacheUsage(MaxCacheSize, TotalCachedBytes, TotalReferencedBlockBytes, ReferencedBytes, FragmentedBytes);
 
 	const uint64 TotalRequiredBytes = TotalCachedBytes + TotalUncachedBytes;
 	if (TotalRequiredBytes <= MaxCacheSize)
 	{
-		UE_LOG(LogIoStoreOnDemand, Log, TEXT("Skipping cache purge, MaxCacheSize=%.2lf MiB, CacheSize=%.2lf MiB, UncachedSize=%.2lf MiB"),
-			ToMiB(MaxCacheSize), ToMiB(TotalCachedBytes), ToMiB(TotalUncachedBytes));
+		UE_LOG(LogIoStoreOnDemand, Log, TEXT("Skipping cache purge, MaxCacheSize=%.2lf MiB, CacheSize=%.2lf MiB, ReferencedBlockSize=%.2lf MiB, ReferencedSize=%.2lf MiB, FragmentedBytes=%.2lf MiB, UncachedSize=%.2lf MiB"),
+			ToMiB(MaxCacheSize), ToMiB(TotalCachedBytes), ToMiB(TotalReferencedBlockBytes), ToMiB(ReferencedBytes), ToMiB(FragmentedBytes), ToMiB(TotalUncachedBytes));
 		return FIoStatus::Ok;
 	}
 
 	//TODO: Compute fragmentation metric and redownload chunks when this number gets too high
 
-	const uint64 TotalReferencedBytes = Algo::TransformAccumulate(BlockInfo,
-		[](const TPair<FCasBlockId, FCasBlockInfo>& Kv){ return (Kv.Value.RefSize > 0) ? Kv.Value.FileSize : uint64(0); },
-		uint64(0));
-
-	UE_LOG(LogIoStoreOnDemand, Log, TEXT("Purging install cache, MaxCacheSize=%.2lf MiB, CacheSize=%.2lf MiB, UncachedSize=%.2lf MiB, ReferencedBytes=%.2lf MiB"),
-		ToMiB(MaxCacheSize), ToMiB(TotalCachedBytes), ToMiB(TotalUncachedBytes), ToMiB(TotalReferencedBytes));
+	UE_LOG(LogIoStoreOnDemand, Log, TEXT("Purging install cache, MaxCacheSize=%.2lf MiB, CacheSize=%.2lf MiB, ReferencedBlockSize=%.2lf MiB, ReferencedSize=%.2lf MiB, FragmentedBytes=%.2lf MiB, UncachedSize=%.2lf MiB"),
+		ToMiB(MaxCacheSize), ToMiB(TotalCachedBytes), ToMiB(TotalReferencedBlockBytes), ToMiB(ReferencedBytes), ToMiB(FragmentedBytes), ToMiB(TotalUncachedBytes));
 
 	const uint64	TotalBytesToPurge	= TotalRequiredBytes - MaxCacheSize;
 	uint64			TotalPurgedBytes	= 0;
@@ -1909,6 +1939,8 @@ FIoStatus FOnDemandInstallCache::Purge(TSet<FIoHash>&& ChunksToInstall)
 		LogIoStoreOnDemand, Warning, TEXT("Max install cache size exceeded by %.2lf MiB (%.2lf%%)"),
 			ToMiB(NewCachedBytes - MaxCacheSize), 100.0 * (double(NewCachedBytes - MaxCacheSize) / double(MaxCacheSize)));
 
+	FOnDemandInstallCacheStats::OnPurge(Status.GetErrorCode(), MaxCacheSize, NewCachedBytes, TotalBytesToPurge, TotalPurgedBytes);
+
 	if (TotalPurgedBytes < TotalBytesToPurge)
 	{
 		if (UE::IoStore::CVars::GIoStoreOnDemandEnableDefrag)
@@ -1942,7 +1974,8 @@ FIoStatus FOnDemandInstallCache::PurgeAllUnreferenced(bool bDefrag, const uint64
 	IoStore.GetReferencedContent(Containers, ChunkEntryIndices);
 	check(Containers.Num() == ChunkEntryIndices.Num());
 
-	AddReferencesToBlocks(Containers, ChunkEntryIndices, {}, BlockInfo);
+	uint64 ReferencedBytes = 0;
+	AddReferencesToBlocks(Containers, ChunkEntryIndices, {}, BlockInfo, ReferencedBytes);
 
 	const uint64 TotalReferencedBytes = Algo::TransformAccumulate(BlockInfo,
 		[](const TPair<FCasBlockId, FCasBlockInfo>& Kv) { return (Kv.Value.RefSize > 0) ? Kv.Value.FileSize : uint64(0); },
@@ -2008,14 +2041,15 @@ FIoStatus FOnDemandInstallCache::DefragAll(const uint64* BytesToFree /*= nullptr
 	IoStore.GetReferencedContent(Containers, ChunkEntryIndices);
 	check(Containers.Num() == ChunkEntryIndices.Num());
 
-	AddReferencesToBlocks(Containers, ChunkEntryIndices, {}, BlockInfo);
+	uint64 ReferencedBytes = 0;
+	AddReferencesToBlocks(Containers, ChunkEntryIndices, {}, BlockInfo, ReferencedBytes);
 
-	const uint64 TotalReferencedBytes = Algo::TransformAccumulate(BlockInfo,
+	const uint64 TotalReferencedBlockBytes = Algo::TransformAccumulate(BlockInfo,
 		[](const TPair<FCasBlockId, FCasBlockInfo>& Kv) { return (Kv.Value.RefSize > 0) ? Kv.Value.FileSize : uint64(0); },
 		uint64(0));
 
-	UE_LOG(LogIoStoreOnDemand, Log, TEXT("Defragmenting install cache, MaxCacheSize=%.2lf MiB, CacheSize=%.2lf MiB, ReferencedBytes=%.2lf MiB"),
-		ToMiB(MaxCacheSize), ToMiB(TotalCachedBytes), ToMiB(TotalReferencedBytes));
+	UE_LOG(LogIoStoreOnDemand, Log, TEXT("Defragmenting install cache, MaxCacheSize=%.2lf MiB, CacheSize=%.2lf MiB, ReferencedBlockSize=%.2lf MiB, ReferencedSize=%.2lf MiB"),
+		ToMiB(MaxCacheSize), ToMiB(TotalCachedBytes), ToMiB(TotalReferencedBlockBytes), ToMiB(ReferencedBytes));
 
 	return Defrag(Containers, ChunkEntryIndices, BlockInfo, BytesToFree);
 }
@@ -2024,38 +2058,44 @@ uint64 FOnDemandInstallCache::AddReferencesToBlocks(
 	const TArray<FSharedOnDemandContainer>& Containers,
 	const TArray<TBitArray<>>& ChunkEntryIndices,
 	const TSet<FIoHash>& ChunksToInstall,
-	FCasBlockInfoMap& BlockInfoMap) const
+	FCasBlockInfoMap& BlockInfoMap,
+	uint64& OutTotalReferencedBytes) const
 {
 	uint64 TotalUncachedBytes = 0;
+	OutTotalReferencedBytes	= 0;
 
 	for (int32 Index = 0; FSharedOnDemandContainer Container : Containers)
 	{
 		const TBitArray<>& IsReferenced = ChunkEntryIndices[Index++];
 		for (int32 EntryIndex = 0; const FOnDemandChunkEntry& Entry : Container->ChunkEntries)
 		{
-			const bool bIsReferenced = IsReferenced[EntryIndex++];
-			const bool bToInstall = ChunksToInstall.Contains(Entry.Hash);
-			
-			if (!bIsReferenced && !bToInstall)
-			{
-				continue;
-			}
-			
-			FCasBlockInfo* BlockInfo = nullptr;
-			if (FCasLocation Loc = Cas.FindChunk(Entry.Hash); Loc.IsValid())
+			if (IsReferenced[EntryIndex])
 			{
-				BlockInfo = BlockInfoMap.Find(Loc.BlockId);
-				ensure(BlockInfo);
-			}
-
-			if (BlockInfo)
-			{
-				BlockInfo->RefSize += Entry.EncodedSize;
+				OutTotalReferencedBytes += Entry.EncodedSize;
+				if (FCasLocation Loc = Cas.FindChunk(Entry.Hash); Loc.IsValid())
+				{
+					if (FCasBlockInfo* BlockInfo = BlockInfoMap.Find(Loc.BlockId))
+					{
+						BlockInfo->RefSize += Entry.EncodedSize;
+					}
+					else
+					{
+						UE_LOG(LogIoStoreOnDemand, Error, TEXT("Failed to find CAS block info for referenced chunk, ChunkId='%s', Container='%s'"),
+							*LexToString(Container->ChunkIds[EntryIndex]), *Container->Name);
+					}
+				}
+				else
+				{
+					UE_LOG(LogIoStoreOnDemand, Error, TEXT("Failed to find CAS location for referenced chunk, ChunkId='%s', Container='%s'"),
+						*LexToString(Container->ChunkIds[EntryIndex]), *Container->Name);
+				}
 			}
-			else if (bToInstall)
+			else if (ChunksToInstall.Contains(Entry.Hash))
 			{
 				TotalUncachedBytes += Entry.EncodedSize;
 			}
+
+			EntryIndex++;
 		}
 	}
 
@@ -2194,6 +2234,7 @@ FIoStatus FOnDemandInstallCache::Defrag(
 
 		if (!bPossibleToFreeBytes)
 		{
+			FOnDemandInstallCacheStats::OnDefrag(EIoErrorCode::WriteError, 0);
 			return FIoStatusBuilder(EIoErrorCode::WriteError) << FString::Printf(TEXT("Defrag failed - cannot free %" UINT64_FMT), *TotalBytesToFree);
 		}
 	}
@@ -2269,6 +2310,7 @@ FIoStatus FOnDemandInstallCache::Defrag(
 				const FString Filename = Cas.GetBlockFilename(DefragBlock.BlockId);
 				const FString ErrorMsg = FileOpenResult.GetError().GetMessage();
 				UE_LOG(LogIoStoreOnDemand, Error, TEXT("Failed to open CAS block '%s' for reading, reason '%s'"), *Filename, *ErrorMsg);
+				FOnDemandInstallCacheStats::OnDefrag(EIoErrorCode::FileOpenFailed, 0);
 				return FIoStatusBuilder(EIoErrorCode::FileOpenFailed) << ErrorMsg;
 			}
 
@@ -2284,6 +2326,7 @@ FIoStatus FOnDemandInstallCache::Defrag(
 				const bool bOk = FileHandle->Read(Buffer.GetData(), Buffer.GetSize());
 				if (!bOk)
 				{
+					FOnDemandInstallCacheStats::OnDefrag(EIoErrorCode::ReadError, 0);
 					return EIoErrorCode::ReadError;
 				}
 
@@ -2291,6 +2334,7 @@ FIoStatus FOnDemandInstallCache::Defrag(
 				{
 					if (FIoStatus Status = FlushPendingChunks(DefragPendingChunks, DefragBlock.LastAccess); Status.IsOk() == false)
 					{
+						FOnDemandInstallCacheStats::OnDefrag(Status.GetErrorCode(), 0);
 						return Status;
 					}
 					check(DefragPendingChunks.IsEmpty());
@@ -2303,6 +2347,7 @@ FIoStatus FOnDemandInstallCache::Defrag(
 
 			if (FIoStatus Status = FlushPendingChunks(DefragPendingChunks); Status.IsOk() == false)
 			{
+				FOnDemandInstallCacheStats::OnDefrag(Status.GetErrorCode(), 0);
 				return Status;
 			}
 			check(DefragPendingChunks.IsEmpty());
@@ -2323,12 +2368,15 @@ FIoStatus FOnDemandInstallCache::Defrag(
 
 		if (FIoStatus Status = FCasJournal::Commit(MoveTemp(Transaction)); !Status.IsOk())
 		{
+			FOnDemandInstallCacheStats::OnDefrag(Status.GetErrorCode(), 0);
 			return Status;
 		}
 	}
 
 	UE_LOG(LogIoStoreOnDemand, Display, TEXT("Defrag removed %" UINT64_FMT " fragmented bytes of %" UINT64_FMT " total bytes in %i blocks."), FragmentedBytes, TotalBlockSize, BlocksToDefrag.Num());
 
+	FOnDemandInstallCacheStats::OnDefrag(EIoErrorCode::Ok, FragmentedBytes);
+
 	return FIoStatus::Ok;
 }
 
@@ -2359,7 +2407,8 @@ FOnDemandInstallCacheUsage FOnDemandInstallCache::GetCacheUsage()
 	IoStore.GetReferencedContent(Containers, ChunkEntryIndices);
 	check(Containers.Num() == ChunkEntryIndices.Num());
 
-	AddReferencesToBlocks(Containers, ChunkEntryIndices, {}, BlockInfo);
+	uint64 ReferencedBytes = 0;
+	AddReferencesToBlocks(Containers, ChunkEntryIndices, {}, BlockInfo, ReferencedBytes);
 
 	uint64 FragmentedBytes = 0;
 	uint64 ReferencedBlockBytes = 0;
@@ -2384,6 +2433,7 @@ FOnDemandInstallCacheUsage FOnDemandInstallCache::GetCacheUsage()
 		.MaxSize = MaxCacheSize,
 		.TotalSize = TotalCachedBytes,
 		.ReferencedBlockSize = ReferencedBlockBytes,
+		.ReferencedSize = ReferencedBytes,
 		.FragmentedChunksSize = FragmentedBytes,
 	};
 }
@@ -2432,6 +2482,7 @@ FIoStatus FOnDemandInstallCache::FlushPendingChunksImpl(FPendingChunks& Chunks,
 		TUniquePtr<IFileHandle>	CasFileHandle = Cas.OpenWrite(CurrentBlockId);
 		if (CasFileHandle.IsValid() == false)
 		{
+			FOnDemandInstallCacheStats::OnFlush(EIoErrorCode::WriteError, 0);
 			return FIoStatusBuilder(EIoErrorCode::FileOpenFailed)
 				<< TEXT("Failed to open cache block file '")
 				<< Cas.GetBlockFilename(CurrentBlockId)
@@ -2477,10 +2528,13 @@ FIoStatus FOnDemandInstallCache::FlushPendingChunksImpl(FPendingChunks& Chunks,
 
 			if (CasFileHandle->Flush() == false)
 			{
+				FOnDemandInstallCacheStats::OnFlush(EIoErrorCode::WriteError, Ar.Tell());
 				return FIoStatusBuilder(EIoErrorCode::WriteError)
 					<< TEXT("Failed to flush cache block to disk");
 			}
 
+			FOnDemandInstallCacheStats::OnFlush(EIoErrorCode::Ok, Ar.Tell());
+
 			check(ChunkHashes.Num() == Offsets.Num());
 			check(CurrentBlockId.IsValid());
 			for (int32 Idx = 0, Count = Offsets.Num(); Idx < Count; ++Idx)
diff --git a/Engine/Source/Runtime/Experimental/IoStore/OnDemand/Private/Statistics.cpp b/Engine/Source/Runtime/Experimental/IoStore/OnDemand/Private/Statistics.cpp
--- a/Engine/Source/Runtime/Experimental/IoStore/OnDemand/Private/Statistics.cpp
+++ b/Engine/Source/Runtime/Experimental/IoStore/OnDemand/Private/Statistics.cpp
@@ -2,6 +2,8 @@
 
 #include "Statistics.h"
 
+#include "Async/Mutex.h"
+#include "Async/UniqueLock.h"
 #include "Algo/MaxElement.h"
 #include "AnalyticsEventAttribute.h"
 #include "HAL/IConsoleManager.h"
@@ -409,6 +411,27 @@ FCounterAtomicInt		GTotalDownloadedChunkBytes(TEXT("Iad/DownloadedBytes"), Trace
 FCounterAtomicInt		GTotalInstallDurationMs(TEXT("Iad/TotalInstallDurationMs"), TraceCounterDisplayHint_None);
 //FCounterAtomicFloat	GAvgInstallCacheHitRatio(TEXT("Iad/AvgInstallCacheHitRatio"), TraceCounterDisplayHint_None);
 
+struct FInstallCacheStats
+{
+	void Lock() { Mutex.Lock(); }
+	void Unlock() { Mutex.Unlock(); }
+
+	FMutex Mutex; 
+	uint64 VerificationRemovedBlockCount = 0;
+	uint64 StartupErrorCount = 0;
+	uint64 FlushErrorCount = 0;
+	uint64 PurgedBytes = 0;
+	uint64 PurgeErrorCount = 0;
+	uint64 FragmentedBytes = 0;
+	uint64 DefragErrorCount = 0;
+	uint64 JournalCommitErrorCount = 0;
+	uint64 MaxCacheUsageSize = 0;
+	uint64 MaxReferencedBlockSize = 0;
+	uint64 MaxReferencedSize = 0;
+	uint64 MaxFragmentedSize = 0;
+};
+static FInstallCacheStats InstallCacheStats;
+
 ////////////////////////////////////////////////////////////////////////////////
 static FOnDemandIoBackendStats* GStatistics = nullptr;
 static FDelegateHandle GStatisticsEndFrameDelegateHandle;
@@ -1105,15 +1128,104 @@ void FOnDemandContentInstallerStats::ReportAnalytics(TArray<FAnalyticsEventAttri
 		return;
 	}
 
+	UE::TUniqueLock Lock(InstallCacheStats);
+
 	AppendAnalyticsEventAttributeArray(OutAnalyticsArray
 		, TEXT("IadTotalInstallCount"), GTotalInstallCount.Get()
 		, TEXT("IadTotalInstallErrorCount"), GTotalInstallErrorCount.Get()
 		, TEXT("IadTotalWriteErrorInstallErrorCount"), GTotalWriteErrorInstallErrorCount.Get()
+		, TEXT("IadTotalDownloadedBytes"), GTotalDownloadedChunkBytes.Get()
 		, TEXT("IadTotalInstallDurationMs"), GTotalInstallDurationMs.Get()
 		, TEXT("IadAvgCacheHitRatio"), GInstallCacheHitRatio.GetMean()
+		, TEXT("IadInstallCacheStartupErrorCount"), InstallCacheStats.StartupErrorCount
+		, TEXT("IadInstallCacheVerificationRemovedBlockCount"), InstallCacheStats.VerificationRemovedBlockCount 
+		, TEXT("IadInstallCacheFlushErrorCount"), InstallCacheStats.FlushErrorCount 
+		, TEXT("IadInstallCachePurgeErrorCount"), InstallCacheStats.PurgeErrorCount
+		, TEXT("IadInstallCacheDefragErrorCount"), InstallCacheStats.DefragErrorCount 
+		, TEXT("IadInstallInstallCacheJournalCommitErrorCount"), InstallCacheStats.DefragErrorCount 
+		, TEXT("IadInstallInstallCacheMaxUsageSize"), InstallCacheStats.MaxCacheUsageSize 
+		, TEXT("IadInstallInstallCacheMaxReferencedBlockSize"), InstallCacheStats.MaxReferencedBlockSize
+		, TEXT("IadInstallInstallCacheMaxReferencedSize"), InstallCacheStats.MaxReferencedSize
+		, TEXT("IadInstallInstallCacheMaxFragmentedSize"), InstallCacheStats.MaxFragmentedSize
 	);
 }
 
+////////////////////////////////////////////////////////////////////////////////
+void FOnDemandInstallCacheStats::OnStartupError(EIoErrorCode ErrorCode)
+{
+	UE::TUniqueLock Lock(InstallCacheStats);
+	InstallCacheStats.StartupErrorCount++;
+}
+
+void FOnDemandInstallCacheStats::OnFlush(EIoErrorCode ErrorCode, int64 ByteCount)
+{
+	if (ErrorCode != EIoErrorCode::Ok)
+	{
+		UE::TUniqueLock Lock(InstallCacheStats);
+		InstallCacheStats.FlushErrorCount++;
+	}
+}
+
+void FOnDemandInstallCacheStats::OnJournalCommit(EIoErrorCode ErrorCode, int64 ByteCount)
+{
+	if (ErrorCode != EIoErrorCode::Ok)
+	{
+		UE::TUniqueLock Lock(InstallCacheStats);
+		InstallCacheStats.JournalCommitErrorCount++;
+	}
+}
+
+void FOnDemandInstallCacheStats::OnCasVerificationError(int32 RemoveChunks)
+{
+	UE::TUniqueLock Lock(InstallCacheStats);
+	InstallCacheStats.VerificationRemovedBlockCount += RemoveChunks;
+}
+
+void FOnDemandInstallCacheStats::OnPurge(
+	EIoErrorCode ErrorCode,
+	uint64 MaxCacheSize,
+	uint64 NewCacheSize,
+	uint64 BytesToPurge,
+	uint64 PurgedBytes)
+{
+	UE::TUniqueLock Lock(InstallCacheStats);
+	if (ErrorCode == EIoErrorCode::Ok)
+	{
+		InstallCacheStats.PurgedBytes += PurgedBytes;
+	}
+	else 
+	{
+		InstallCacheStats.PurgeErrorCount++;
+	}
+}
+
+void FOnDemandInstallCacheStats::OnDefrag(EIoErrorCode ErrorCode, uint64 FragmentedBytes)
+{
+	UE::TUniqueLock Lock(InstallCacheStats);
+	if (ErrorCode == EIoErrorCode::Ok)
+	{
+		InstallCacheStats.FragmentedBytes += FragmentedBytes;
+	}
+	else
+	{
+		InstallCacheStats.DefragErrorCount++;
+	}
+}
+
+void FOnDemandInstallCacheStats::OnCacheUsage(
+	uint64 MaxCacheSize,
+	uint64 CacheSize,
+	uint64 ReferencedBlockSize,
+	uint64 ReferencedSize,
+	uint64 FragmentedSize)
+{
+	UE::TUniqueLock Lock(InstallCacheStats);
+	InstallCacheStats.MaxCacheUsageSize			= FMath::Max(InstallCacheStats.MaxCacheUsageSize, CacheSize);
+	InstallCacheStats.MaxReferencedBlockSize	= FMath::Max(InstallCacheStats.MaxReferencedBlockSize, ReferencedBlockSize);
+	InstallCacheStats.MaxReferencedSize			= FMath::Max(InstallCacheStats.MaxReferencedSize, ReferencedSize);
+	InstallCacheStats.MaxFragmentedSize			= FMath::Max(InstallCacheStats.MaxFragmentedSize, FragmentedSize);
+}
+
 } // namespace UE::IoStore
 
 #undef UE_ENABLE_ONSCREEN_STATISTICS
diff --git a/Engine/Source/Runtime/Experimental/IoStore/OnDemand/Private/Statistics.h b/Engine/Source/Runtime/Experimental/IoStore/OnDemand/Private/Statistics.h
--- a/Engine/Source/Runtime/Experimental/IoStore/OnDemand/Private/Statistics.h
+++ b/Engine/Source/Runtime/Experimental/IoStore/OnDemand/Private/Statistics.h
@@ -103,6 +103,29 @@ public:
 	static void ReportAnalytics(TArray<FAnalyticsEventAttribute>& OutAnalyticsArray) IAS_STATISTICS_IMPL()
 };
 
+////////////////////////////////////////////////////////////////////////////////
+class FOnDemandInstallCacheStats
+{
+public:
+	static void OnStartupError(EIoErrorCode ErrorCode) IAS_STATISTICS_IMPL()
+	static void OnFlush(EIoErrorCode ErrorCode, int64 ByteCount) IAS_STATISTICS_IMPL()
+	static void OnJournalCommit(EIoErrorCode ErrorCode, int64 ByteCount) IAS_STATISTICS_IMPL()
+	static void OnCasVerificationError(int32 RemoveChunks) IAS_STATISTICS_IMPL()
+	static void OnPurge(
+		EIoErrorCode ErrorCode,
+		uint64 MaxCacheSize,
+		uint64 NewCacheSize,
+		uint64 BytesToPurge,
+		uint64 PurgedBytes) IAS_STATISTICS_IMPL()
+	static void OnDefrag(EIoErrorCode ErrorCode, uint64 FragmentedBytes) IAS_STATISTICS_IMPL()
+	static void OnCacheUsage(
+		uint64 MaxCacheSize,
+		uint64 CacheSize,
+		uint64 ReferencedBlockSize,
+		uint64 ReferencedSize,
+		uint64 FragmentedSize) IAS_STATISTICS_IMPL()
+};
+
 #undef IAS_STATISTICS_IMPL
 
 } // namespace UE::IoStore
diff --git a/Engine/Source/Runtime/Experimental/IoStore/OnDemandCore/Private/IoStoreOnDemand.cpp b/Engine/Source/Runtime/Experimental/IoStore/OnDemandCore/Private/IoStoreOnDemand.cpp
--- a/Engine/Source/Runtime/Experimental/IoStore/OnDemandCore/Private/IoStoreOnDemand.cpp
+++ b/Engine/Source/Runtime/Experimental/IoStore/OnDemandCore/Private/IoStoreOnDemand.cpp
@@ -400,10 +400,11 @@ static const TCHAR* NotInitializedError = TEXT("I/O store on-demand not initiali
 ////////////////////////////////////////////////////////////////////////////////
 FStringBuilderBase& operator<<(FStringBuilderBase& Sb, const FOnDemandInstallCacheUsage& CacheUsage)
 {
-	Sb.Appendf(TEXT("MaxSize=%.2lf MiB, TotalSize=%.2lf MiB, ReferencedBlockSize=%.2lf MiB, FragmentedChunksSize=%.2lf MiB"),
+	Sb.Appendf(TEXT("MaxSize=%.2lf MiB, TotalSize=%.2lf MiB, ReferencedBlockSize=%.2lf MiB, ReferencedSize=%.2lf MiB, FragmentedChunksSize=%.2lf MiB"),
 		double(CacheUsage.MaxSize) / 1024.0 / 1024.0,
 		double(CacheUsage.TotalSize) / 1024.0 / 1024.0,
 		double(CacheUsage.ReferencedBlockSize) / 1024.0 / 1024.0,
+		double(CacheUsage.ReferencedSize) / 1024.0 / 1024.0,
 		double(CacheUsage.FragmentedChunksSize) / 1024.0 / 1024.0);
 
 	return Sb;
diff --git a/Engine/Source/Runtime/Experimental/IoStore/OnDemandCore/Public/IO/IoStoreOnDemand.h b/Engine/Source/Runtime/Experimental/IoStore/OnDemandCore/Public/IO/IoStoreOnDemand.h
--- a/Engine/Source/Runtime/Experimental/IoStore/OnDemandCore/Public/IO/IoStoreOnDemand.h
+++ b/Engine/Source/Runtime/Experimental/IoStore/OnDemandCore/Public/IO/IoStoreOnDemand.h
@@ -423,6 +423,7 @@ struct FOnDemandInstallCacheUsage
 	uint64 MaxSize = 0;
 	uint64 TotalSize = 0;
 	uint64 ReferencedBlockSize = 0;
+	uint64 ReferencedSize = 0;
 	uint64 FragmentedChunksSize = 0;
 
 	// Only populated if DumpHandlesToResults is specified
