Dataflow : make the code that reads the variable from property bag more robust
- fix the issue with string that were getting deleted when going out of scope
- now keep the results around to make the pointer we get out of TryGetValue remains valid

diff --git a/Engine/Source/Runtime/Experimental/Dataflow/Engine/Private/Dataflow/DataflowVariableNodes.cpp b/Engine/Source/Runtime/Experimental/Dataflow/Engine/Private/Dataflow/DataflowVariableNodes.cpp
--- a/Engine/Source/Runtime/Experimental/Dataflow/Engine/Private/Dataflow/DataflowVariableNodes.cpp
+++ b/Engine/Source/Runtime/Experimental/Dataflow/Engine/Private/Dataflow/DataflowVariableNodes.cpp
@@ -289,9 +289,10 @@ void FGetDataflowVariableNode::Evaluate(UE::Dataflow::FContext& Context, const F
 	case EPropertyBagContainerType::None:
 	{
 		bool OutBool = false;
-		if (const bool* Value = Variables.GetValueBool(Node.VariableName).TryGetValue())
+		TValueOrError<bool, EPropertyBagResult> Result = Variables.GetValueBool(Node.VariableName);
+		if (Result.HasValue())
 		{
-			OutBool = *Value;
+			OutBool = Result.GetValue();
 		}
 		Out.SetValue<bool>(MoveTemp(OutBool), Context);
 		break;
@@ -299,12 +300,17 @@ void FGetDataflowVariableNode::Evaluate(UE::Dataflow::FContext& Context, const F
 	case EPropertyBagContainerType::Array:
 	{
 		TArray<bool> OutBoolArray;
-		if (const FPropertyBagArrayRef* ArrayRef = Variables.GetArrayRef(Node.VariableName).TryGetValue())
+		TValueOrError<const FPropertyBagArrayRef, EPropertyBagResult> Result = Variables.GetArrayRef(Node.VariableName);
+		if (Result.HasValue())
 		{
-			OutBoolArray.SetNum(ArrayRef->Num());
-			for (int32 Index = 0; Index < ArrayRef->Num(); Index++)
+			const FPropertyBagArrayRef ArrayRef = Result.GetValue();
+			if (const int32 Count = ArrayRef.Num())
 			{
-				OutBoolArray[Index] = ArrayRef->GetValueBool(Index).GetValue();
+				OutBoolArray.SetNum(Count);
+				for (int32 Index = 0; Index < Count; Index++)
+				{
+					OutBoolArray[Index] = ArrayRef.GetValueBool(Index).GetValue();
+				}
 			}
 		}
 		Out.SetValue(MoveTemp(OutBoolArray), Context);
@@ -320,9 +326,10 @@ void FGetDataflowVariableNode::Evaluate(UE::Dataflow::FContext& Context, const F
 	case EPropertyBagContainerType::None:
 	{
 		uint8 OutUInt8 = 0;
-		if (const uint8* Value = Variables.GetValueByte(Node.VariableName).TryGetValue())
+		TValueOrError<uint8, EPropertyBagResult> Result = Variables.GetValueByte(Node.VariableName);
+		if (Result.HasValue())
 		{
-			OutUInt8 = *Value;
+			OutUInt8 = Result.GetValue();
 		}
 		Out.SetValue<uint8>(MoveTemp(OutUInt8), Context);
 		break;
@@ -330,12 +337,17 @@ void FGetDataflowVariableNode::Evaluate(UE::Dataflow::FContext& Context, const F
 	case EPropertyBagContainerType::Array:
 	{
 		TArray<uint8> OutInt8Array;
-		if (const FPropertyBagArrayRef* ArrayRef = Variables.GetArrayRef(Node.VariableName).TryGetValue())
+		TValueOrError<const FPropertyBagArrayRef, EPropertyBagResult> Result = Variables.GetArrayRef(Node.VariableName);
+		if (Result.HasValue())
 		{
-			OutInt8Array.SetNum(ArrayRef->Num());
-			for (int32 Index = 0; Index < ArrayRef->Num(); Index++)
+			const FPropertyBagArrayRef ArrayRef = Result.GetValue();
+			if (const int32 Count = ArrayRef.Num())
 			{
-				OutInt8Array[Index] = ArrayRef->GetValueByte(Index).GetValue();
+				OutInt8Array.SetNum(Count);
+				for (int32 Index = 0; Index < Count; Index++)
+				{
+					OutInt8Array[Index] = ArrayRef.GetValueByte(Index).GetValue();
+				}
 			}
 		}
 		Out.SetValue(MoveTemp(OutInt8Array), Context);
@@ -351,9 +363,10 @@ void FGetDataflowVariableNode::Evaluate(UE::Dataflow::FContext& Context, const F
 	case EPropertyBagContainerType::None:
 	{
 		int32 OutInt32 = 0;
-		if (const int32* Value = Variables.GetValueInt32(Node.VariableName).TryGetValue())
+		TValueOrError<int32, EPropertyBagResult> Result = Variables.GetValueInt32(Node.VariableName);
+		if (Result.HasValue())
 		{
-			OutInt32 = *Value;
+			OutInt32 = Result.GetValue();
 		}
 		Out.SetValue<int32>(MoveTemp(OutInt32), Context);
 		break;
@@ -361,12 +374,17 @@ void FGetDataflowVariableNode::Evaluate(UE::Dataflow::FContext& Context, const F
 	case EPropertyBagContainerType::Array:
 	{
 		TArray<int32> OutInt32Array;
-		if (const FPropertyBagArrayRef* ArrayRef = Variables.GetArrayRef(Node.VariableName).TryGetValue())
+		TValueOrError<const FPropertyBagArrayRef, EPropertyBagResult> Result = Variables.GetArrayRef(Node.VariableName);
+		if (Result.HasValue())
 		{
-			OutInt32Array.SetNum(ArrayRef->Num());
-			for (int32 Index = 0; Index < ArrayRef->Num(); Index++)
+			const FPropertyBagArrayRef ArrayRef = Result.GetValue();
+			if (const int32 Count = ArrayRef.Num())
 			{
-				OutInt32Array[Index] = ArrayRef->GetValueInt32(Index).GetValue();
+				OutInt32Array.SetNum(Count);
+				for (int32 Index = 0; Index < Count; Index++)
+				{
+					OutInt32Array[Index] = ArrayRef.GetValueInt32(Index).GetValue();
+				}
 			}
 		}
 		Out.SetValue(MoveTemp(OutInt32Array), Context);
@@ -382,9 +400,10 @@ void FGetDataflowVariableNode::Evaluate(UE::Dataflow::FContext& Context, const F
 	case EPropertyBagContainerType::None:
 	{
 		int64 OutInt64 = 0;
-		if (const int64* Value = Variables.GetValueInt64(Node.VariableName).TryGetValue())
+		TValueOrError<int64, EPropertyBagResult> Result = Variables.GetValueInt64(Node.VariableName);
+		if (Result.HasValue())
 		{
-			OutInt64 = *Value;
+			OutInt64 = Result.GetValue();
 		}
 		Out.SetValue<int64>(MoveTemp(OutInt64), Context);
 		break;
@@ -392,12 +411,17 @@ void FGetDataflowVariableNode::Evaluate(UE::Dataflow::FContext& Context, const F
 	case EPropertyBagContainerType::Array:
 	{
 		TArray<int64> OutInt64Array;
-		if (const FPropertyBagArrayRef* ArrayRef = Variables.GetArrayRef(Node.VariableName).TryGetValue())
+		TValueOrError<const FPropertyBagArrayRef, EPropertyBagResult> Result = Variables.GetArrayRef(Node.VariableName);
+		if (Result.HasValue())
 		{
-			OutInt64Array.SetNum(ArrayRef->Num());
-			for (int32 Index = 0; Index < ArrayRef->Num(); Index++)
+			const FPropertyBagArrayRef ArrayRef = Result.GetValue();
+			if (const int32 Count = ArrayRef.Num())
 			{
-				OutInt64Array[Index] = ArrayRef->GetValueInt64(Index).GetValue();
+				OutInt64Array.SetNum(Count);
+				for (int32 Index = 0; Index < Count; Index++)
+				{
+					OutInt64Array[Index] = ArrayRef.GetValueInt64(Index).GetValue();
+				}
 			}
 		}
 		Out.SetValue(MoveTemp(OutInt64Array), Context);
@@ -413,9 +437,10 @@ void FGetDataflowVariableNode::Evaluate(UE::Dataflow::FContext& Context, const F
 	case EPropertyBagContainerType::None:
 	{
 		float OutFloat = 0.0f;
-		if (const float* Value = Variables.GetValueFloat(Node.VariableName).TryGetValue())
+		TValueOrError<float, EPropertyBagResult> Result = Variables.GetValueFloat(Node.VariableName);
+		if (Result.HasValue())
 		{
-			OutFloat = *Value;
+			OutFloat = Result.GetValue();
 		}
 		Out.SetValue<float>(MoveTemp(OutFloat), Context);
 		break;
@@ -423,12 +448,17 @@ void FGetDataflowVariableNode::Evaluate(UE::Dataflow::FContext& Context, const F
 	case EPropertyBagContainerType::Array:
 	{
 		TArray<float> OutFloatArray;
-		if (const FPropertyBagArrayRef* ArrayRef = Variables.GetArrayRef(Node.VariableName).TryGetValue())
+		TValueOrError<const FPropertyBagArrayRef, EPropertyBagResult> Result = Variables.GetArrayRef(Node.VariableName);
+		if (Result.HasValue())
 		{
-			OutFloatArray.SetNum(ArrayRef->Num());
-			for (int32 Index = 0; Index < ArrayRef->Num(); Index++)
+			const FPropertyBagArrayRef ArrayRef = Result.GetValue();
+			if (const int32 Count = ArrayRef.Num())
 			{
-				OutFloatArray[Index] = ArrayRef->GetValueFloat(Index).GetValue();
+				OutFloatArray.SetNum(Count);
+				for (int32 Index = 0; Index < Count; Index++)
+				{
+					OutFloatArray[Index] = ArrayRef.GetValueFloat(Index).GetValue();
+				}
 			}
 		}
 		Out.SetValue(MoveTemp(OutFloatArray), Context);
@@ -446,9 +476,10 @@ void FGetDataflowVariableNode::Evaluate(UE::Dataflow::FContext& Context, const F
 	case EPropertyBagContainerType::None:
 	{
 		float OutFloat = 0.0;
-		if (const double* Value = Variables.GetValueDouble(Node.VariableName).TryGetValue())
+		TValueOrError<double, EPropertyBagResult> Result = Variables.GetValueDouble(Node.VariableName);
+		if (Result.HasValue())
 		{
-			OutFloat = *Value;
+			OutFloat = static_cast<float>(Result.GetValue());
 		}
 		Out.SetValue<float>(MoveTemp(OutFloat), Context);
 		break;
@@ -456,12 +487,17 @@ void FGetDataflowVariableNode::Evaluate(UE::Dataflow::FContext& Context, const F
 	case EPropertyBagContainerType::Array:
 	{
 		TArray<float> OutFloatArray;
-		if (const FPropertyBagArrayRef* ArrayRef = Variables.GetArrayRef(Node.VariableName).TryGetValue())
+		TValueOrError<const FPropertyBagArrayRef, EPropertyBagResult> Result = Variables.GetArrayRef(Node.VariableName);
+		if (Result.HasValue())
 		{
-			OutFloatArray.SetNum(ArrayRef->Num());
-			for (int32 Index = 0; Index < ArrayRef->Num(); Index++)
+			const FPropertyBagArrayRef ArrayRef = Result.GetValue();
+			if (const int32 Count = ArrayRef.Num())
 			{
-				OutFloatArray[Index] = static_cast<float>(ArrayRef->GetValueDouble(Index).GetValue());
+				OutFloatArray.SetNum(Count);
+				for (int32 Index = 0; Index < Count; Index++)
+				{
+					OutFloatArray[Index] = static_cast<float>(ArrayRef.GetValueDouble(Index).GetValue());
+				}
 			}
 		}
 		Out.SetValue(MoveTemp(OutFloatArray), Context);
@@ -477,9 +513,10 @@ void FGetDataflowVariableNode::Evaluate(UE::Dataflow::FContext& Context, const F
 	case EPropertyBagContainerType::None:
 	{
 		FName OutName;
-		if (const FName* Value = Variables.GetValueName(Node.VariableName).TryGetValue())
+		TValueOrError<FName, EPropertyBagResult> Result = Variables.GetValueName(Node.VariableName);
+		if (Result.HasValue())
 		{
-			OutName = *Value;
+			OutName = Result.GetValue();
 		}
 		Out.SetValue<FName>(MoveTemp(OutName), Context);
 		break;
@@ -487,12 +524,17 @@ void FGetDataflowVariableNode::Evaluate(UE::Dataflow::FContext& Context, const F
 	case EPropertyBagContainerType::Array:
 	{
 		TArray<FName> OutNameArray;
-		if (const FPropertyBagArrayRef* ArrayRef = Variables.GetArrayRef(Node.VariableName).TryGetValue())
+		TValueOrError<const FPropertyBagArrayRef, EPropertyBagResult> Result = Variables.GetArrayRef(Node.VariableName);
+		if (Result.HasValue())
 		{
-			OutNameArray.SetNum(ArrayRef->Num());
-			for (int32 Index = 0; Index < ArrayRef->Num(); Index++)
+			const FPropertyBagArrayRef ArrayRef = Result.GetValue();
+			if (const int32 Count = ArrayRef.Num())
 			{
-				OutNameArray[Index] = ArrayRef->GetValueName(Index).GetValue();
+				OutNameArray.SetNum(Count);
+				for (int32 Index = 0; Index < Count; Index++)
+				{
+					OutNameArray[Index] = ArrayRef.GetValueName(Index).GetValue();
+				}
 			}
 		}
 		Out.SetValue(MoveTemp(OutNameArray), Context);
@@ -508,9 +550,10 @@ void FGetDataflowVariableNode::Evaluate(UE::Dataflow::FContext& Context, const F
 	case EPropertyBagContainerType::None:
 	{
 		FString OutString;
-		if (const FString* Value = Variables.GetValueString(Node.VariableName).TryGetValue())
+		TValueOrError<FString, EPropertyBagResult> Result = Variables.GetValueString(Node.VariableName);
+		if (Result.HasValue())
 		{
-			OutString = *Value;
+			OutString = Result.GetValue();
 		}
 		Out.SetValue<FString>(MoveTemp(OutString), Context);
 		break;
@@ -518,12 +561,17 @@ void FGetDataflowVariableNode::Evaluate(UE::Dataflow::FContext& Context, const F
 	case EPropertyBagContainerType::Array:
 	{
 		TArray<FString> OutStringArray;
-		if (const FPropertyBagArrayRef* ArrayRef = Variables.GetArrayRef(Node.VariableName).TryGetValue())
+		TValueOrError<const FPropertyBagArrayRef, EPropertyBagResult> Result = Variables.GetArrayRef(Node.VariableName);
+		if (Result.HasValue())
 		{
-			OutStringArray.SetNum(ArrayRef->Num());
-			for (int32 Index = 0; Index < ArrayRef->Num(); Index++)
+			const FPropertyBagArrayRef ArrayRef = Result.GetValue();
+			if (const int32 Count = ArrayRef.Num())
 			{
-				OutStringArray[Index] = ArrayRef->GetValueString(Index).GetValue();
+				OutStringArray.SetNum(Count);
+				for (int32 Index = 0; Index < Count; Index++)
+				{
+					OutStringArray[Index] = ArrayRef.GetValueString(Index).GetValue();
+				}
 			}
 		}
 		Out.SetValue(MoveTemp(OutStringArray), Context);
@@ -539,9 +587,10 @@ void FGetDataflowVariableNode::Evaluate(UE::Dataflow::FContext& Context, const F
 	case EPropertyBagContainerType::None:
 	{
 		FString OutString;
-		if (const FText* Value = Variables.GetValueText(Node.VariableName).TryGetValue())
+		TValueOrError<FText, EPropertyBagResult> Result = Variables.GetValueText(Node.VariableName);
+		if (Result.HasValue())
 		{
-			OutString = Value->ToString();
+			OutString = Result.GetValue().ToString();
 		}
 		Out.SetValue<FString>(MoveTemp(OutString), Context);
 		break;
@@ -549,12 +598,17 @@ void FGetDataflowVariableNode::Evaluate(UE::Dataflow::FContext& Context, const F
 	case EPropertyBagContainerType::Array:
 	{
 		TArray<FString> OutStringArray;
-		if (const FPropertyBagArrayRef* ArrayRef = Variables.GetArrayRef(Node.VariableName).TryGetValue())
+		TValueOrError<const FPropertyBagArrayRef, EPropertyBagResult> Result = Variables.GetArrayRef(Node.VariableName);
+		if (Result.HasValue())
 		{
-			OutStringArray.SetNum(ArrayRef->Num());
-			for (int32 Index = 0; Index < ArrayRef->Num(); Index++)
+			const FPropertyBagArrayRef ArrayRef = Result.GetValue();
+			if (const int32 Count = ArrayRef.Num())
 			{
-				OutStringArray[Index] = ArrayRef->GetValueText(Index).GetValue().ToString();
+				OutStringArray.SetNum(Count);
+				for (int32 Index = 0; Index < Count; Index++)
+				{
+					OutStringArray[Index] = ArrayRef.GetValueText(Index).GetValue().ToString();
+				}
 			}
 		}
 		Out.SetValue(MoveTemp(OutStringArray), Context);
@@ -570,9 +624,10 @@ void FGetDataflowVariableNode::Evaluate(UE::Dataflow::FContext& Context, const F
 	case EPropertyBagContainerType::None:
 	{
 		TObjectPtr<UObject> OutObject = nullptr;
-		if (UObject* const* Value = Variables.GetValueObject<UObject>(Node.VariableName).TryGetValue())
+		TValueOrError<UObject*, EPropertyBagResult> Result = Variables.GetValueObject(Node.VariableName);
+		if (Result.HasValue())
 		{
-			OutObject = *Value;
+			OutObject = Result.GetValue();
 		}
 		Out.SetValue<TObjectPtr<UObject>>(MoveTemp(OutObject), Context);
 		break;
@@ -580,12 +635,17 @@ void FGetDataflowVariableNode::Evaluate(UE::Dataflow::FContext& Context, const F
 	case EPropertyBagContainerType::Array:
 	{
 		TArray<TObjectPtr<UObject>> OutObjectArray;
-		if (const FPropertyBagArrayRef* ArrayRef = Variables.GetArrayRef(Node.VariableName).TryGetValue())
+		TValueOrError<const FPropertyBagArrayRef, EPropertyBagResult> Result = Variables.GetArrayRef(Node.VariableName);
+		if (Result.HasValue())
 		{
-			OutObjectArray.SetNum(ArrayRef->Num());
-			for (int32 Index = 0; Index < ArrayRef->Num(); Index++)
+			const FPropertyBagArrayRef ArrayRef = Result.GetValue();
+			if (const int32 Count = ArrayRef.Num())
 			{
-				OutObjectArray[Index] = ArrayRef->GetValueObject(Index).GetValue();
+				OutObjectArray.SetNum(Count);
+				for (int32 Index = 0; Index < Count; Index++)
+				{
+					OutObjectArray[Index] = ArrayRef.GetValueObject(Index).GetValue();
+				}
 			}
 		}
 		Out.SetValue(MoveTemp(OutObjectArray), Context);
@@ -601,9 +661,10 @@ void FGetDataflowVariableNode::Evaluate(UE::Dataflow::FContext& Context, const F
 	case EPropertyBagContainerType::None:
 	{
 		FConstStructView OutStruct;
-		if (const FStructView* const Value = Variables.GetValueStruct(Node.VariableName).TryGetValue())
+		TValueOrError<FStructView, EPropertyBagResult> Result = Variables.GetValueStruct(Node.VariableName);
+		if (Result.HasValue())
 		{
-			OutStruct = *Value;
+			OutStruct = Result.GetValue();
 		}
 		Out.SetValueFromStructView(OutStruct, Context);
 		break;
@@ -611,11 +672,14 @@ void FGetDataflowVariableNode::Evaluate(UE::Dataflow::FContext& Context, const F
 	case EPropertyBagContainerType::Array:
 	{
 		FConstStructArrayView OutStructArrayView;
-		if (const FPropertyBagArrayRef* const ArrayRef = Variables.GetArrayRef(Node.VariableName).TryGetValue())
+		TValueOrError<const FPropertyBagArrayRef, EPropertyBagResult> Result = Variables.GetArrayRef(Node.VariableName);
+		if (Result.HasValue())
 		{
-			if (const int32 Count = ArrayRef->Num())
+			const FPropertyBagArrayRef ArrayRef = Result.GetValue();
+			if (const int32 Count = ArrayRef.Num())
 			{
-				if (const FStructView* const Value = ArrayRef->GetValueStruct(0).TryGetValue())
+				TValueOrError<FStructView, EPropertyBagResult> ResultElt = ArrayRef.GetValueStruct(0);
+				if (const FStructView* const Value = ResultElt.TryGetValue())
 				{
 					if (const UScriptStruct* const ScriptStruct = Value->GetScriptStruct())
 					{
