Split I/O store on-demmand trace counters and analytics and some stats fixups.

#rnx

diff --git a/Engine/Source/Runtime/Experimental/IoStore/OnDemand/Private/Statistics.cpp b/Engine/Source/Runtime/Experimental/IoStore/OnDemand/Private/Statistics.cpp
--- a/Engine/Source/Runtime/Experimental/IoStore/OnDemand/Private/Statistics.cpp
+++ b/Engine/Source/Runtime/Experimental/IoStore/OnDemand/Private/Statistics.cpp
@@ -399,38 +399,79 @@ UE_IAX_CSV_DEFINE_STAT(HttpDurationMsAvg);
 UE_IAX_CSV_DEFINE_STAT(HttpDurationMsMax);
 
 ////////////////////////////////////////////////////////////////////////////////
-// TRACE STATS IAD
-FIncrementalVariance	GInstallCacheHitRatio;
-FCounterAtomicInt		GTotalInstallCount(TEXT("Iad/InstallCount"), TraceCounterDisplayHint_None);
-FCounterAtomicInt		GTotalInstallErrorCount(TEXT("Iad/InstallErrorCount"), TraceCounterDisplayHint_None);
-FCounterAtomicInt		GTotalMissingChunkInstallErrorCount(TEXT("Iad/MissingChunkInstallErrorCount"), TraceCounterDisplayHint_None);
-FCounterAtomicInt		GTotalWriteErrorInstallErrorCount(TEXT("Iad/WriteErrorInstallErrorCount"), TraceCounterDisplayHint_None);
-FCounterAtomicInt		GInflightInstallCount(TEXT("Iad/InflightInstallCount"), TraceCounterDisplayHint_None);
-FCounterAtomicInt		GTotalDownloadedChunkCount(TEXT("Iad/DownloadedChunks"), TraceCounterDisplayHint_None);
-FCounterAtomicInt		GTotalDownloadedChunkBytes(TEXT("Iad/DownloadedBytes"), TraceCounterDisplayHint_Memory);
-FCounterAtomicInt		GTotalInstallDurationMs(TEXT("Iad/TotalInstallDurationMs"), TraceCounterDisplayHint_None);
-//FCounterAtomicFloat	GAvgInstallCacheHitRatio(TEXT("Iad/AvgInstallCacheHitRatio"), TraceCounterDisplayHint_None);
-
-struct FInstallCacheStats
+#if COUNTERSTRACE_ENABLED
+struct FInstallerTraceCounters
+{
+	using FCounterFloat = FCountersTrace::FCounterFloat;
+
+	void Lock()		{ Mutex.Lock(); }
+	void Unlock()	{ Mutex.Unlock(); }
+	static FMutex Mutex; 
+
+	FInstallerTraceCounters()
+		: InstallCount(TEXT("Iad/InstallCount"), TraceCounterDisplayHint_None)
+		, InflightInstallCount(TEXT("Iad/InflightInstallCount"), TraceCounterDisplayHint_None)
+		, DownloadedBytes(TEXT("Iad/DownloadedBytes"), TraceCounterDisplayHint_Memory)
+		, AvgInstallDurationMs(TEXT("Iad/AvgInstallDurationMs"), TraceCounterDisplayHint_None)
+		, AvgCacheHitRatio(TEXT("Iad/AvgCacheHitRatio"), TraceCounterDisplayHint_None)
+	{ }
+
+	FCounterInt				InstallCount;
+	FCounterInt				InflightInstallCount;
+	FCounterInt				DownloadedBytes; 
+	FCounterFloat			AvgInstallDurationMs;
+	FCounterFloat			AvgCacheHitRatio;
+	FIncrementalVariance	InstallDurationMs;
+	FIncrementalVariance	CacheHitRatio;
+};
+FMutex FInstallerTraceCounters::Mutex;
+FInstallerTraceCounters InstallerTraceCounters;
+#endif // COUNTERSTRACE_ENABLED
+
+////////////////////////////////////////////////////////////////////////////////
+struct FInstallerAnalytics
 {
-	void Lock() { Mutex.Lock(); }
-	void Unlock() { Mutex.Unlock(); }
+	void Lock()		{ Mutex.Lock(); }
+	void Unlock()	{ Mutex.Unlock(); }
+	static FMutex Mutex; 
+
+	uint64 InstallCount = 0;
+	uint64 InstallErrorCount = 0;
+	uint64 DownloadedBytes = 0;
+	uint64 TotalInstallDurationMs = 0;
+	double TotalCacheHitRatio = 0;
+};
+FMutex FInstallerAnalytics::Mutex;
+static FInstallerAnalytics InstallerAnalytics;
+
+////////////////////////////////////////////////////////////////////////////////
+struct FInstallCacheAnalytics
+{
+	void Lock()		{ Mutex.Lock(); }
+	void Unlock()	{ Mutex.Unlock(); }
+	static FMutex Mutex; 
 
-	FMutex Mutex; 
 	uint64 VerificationRemovedBlockCount = 0;
 	uint64 StartupErrorCount = 0;
+	uint64 FlushCount = 0;
 	uint64 FlushErrorCount = 0;
+	uint64 FlushedBytes = 0;
 	uint64 PurgedBytes = 0;
+	uint64 PurgeCount = 0;
 	uint64 PurgeErrorCount = 0;
 	uint64 FragmentedBytes = 0;
+	uint64 DefragCount = 0;
 	uint64 DefragErrorCount = 0;
+	uint64 JournalCommitCount = 0;
 	uint64 JournalCommitErrorCount = 0;
+	uint64 MaxCacheSize = 0;
 	uint64 MaxCacheUsageSize = 0;
 	uint64 MaxReferencedBlockSize = 0;
 	uint64 MaxReferencedSize = 0;
 	uint64 MaxFragmentedSize = 0;
 };
-static FInstallCacheStats InstallCacheStats;
+FMutex FInstallCacheAnalytics::Mutex;
+static FInstallCacheAnalytics InstallCacheAnalytics;
 
 ////////////////////////////////////////////////////////////////////////////////
 static FOnDemandIoBackendStats* GStatistics = nullptr;
@@ -514,22 +555,16 @@ FOnDemandIoBackendStats::FOnDemandIoBackendStats(FBackendStatus& InStatus)
 				HttpLog(TEXT("IAS"), EHttpRequestType::Streaming);
 				HttpLog(TEXT("IAD"), EHttpRequestType::Installed);
 
+#if COUNTERSTRACE_ENABLED
 				{
-					const uint64 TotalInstallCount		= GTotalInstallCount.Get();
-					const uint64 TotalInstallDurationMs	= GTotalInstallDurationMs.Get();
-					const uint64 TotalDownloadedBytes	= GTotalDownloadedChunkBytes.Get();
-					uint64 AvgInstallDurationMs			= 0;
-					double ThroughputKbitPerSec			= 0.0;
-
-					if (TotalInstallCount > 0)
-					{
-						AvgInstallDurationMs	= TotalInstallDurationMs / TotalInstallCount;
-						ThroughputKbitPerSec	= (double(TotalDownloadedBytes * 8)  / (double(TotalInstallDurationMs) / 1000.0)) / 1000.0;
-					}
-
-					UE_LOG(LogIoStoreOnDemand, Log, TEXT("IadStats: InstallCount=%llu, Downloaded=%d KiB, Throughput=%d kbit/s, AvgInstallDuration=%d ms, AvgCacheHitRatio=%d %%"),
-						TotalInstallCount, int32(TotalDownloadedBytes / 1024), int32(ThroughputKbitPerSec), AvgInstallDurationMs, int32(GInstallCacheHitRatio.GetMean() * 100.0));
+					TUniqueLock Lock(InstallerTraceCounters);
+					UE_LOG(LogIoStoreOnDemand, Log, TEXT("IadStats: InstallCount=%llu, Downloaded=%d KiB, AvgInstallDuration=%dms, AvgCacheHitRatio=%d%%"),
+						InstallerTraceCounters.InstallCount.Get(),
+						(int32(InstallerTraceCounters.DownloadedBytes.Get()) / 1024),
+						int32(InstallerTraceCounters.AvgInstallDurationMs.Get()),
+						int32(InstallerTraceCounters.CacheHitRatio.GetMean() * 100.0));
 				}
+#endif // COUNTERSTRACE_ENABLED
 				LastLogTime = Time;
 			}
 		}
@@ -1083,7 +1118,10 @@ void FOnDemandIoBackendStats::OnHttpCdnCacheReply(EHttpRequestType Type, int32 R
 ////////////////////////////////////////////////////////////////////////////////
 void FOnDemandContentInstallerStats::OnRequestEnqueued()
 {
-	GInflightInstallCount.Increment();
+#if COUNTERSTRACE_ENABLED
+	TUniqueLock Lock(InstallerTraceCounters);
+	InstallerTraceCounters.InflightInstallCount.Increment();
+#endif
 }
 
 void FOnDemandContentInstallerStats::OnRequestCompleted(
@@ -1100,25 +1138,30 @@ void FOnDemandContentInstallerStats::OnRequestCompleted(
 		return;
 	}
 
-	GInflightInstallCount.Decrement();
-	if (ErrorCode != EIoErrorCode::Ok)
+#if COUNTERSTRACE_ENABLED
 	{
-		if (ErrorCode == EIoErrorCode::UnknownChunkID || ErrorCode == EIoErrorCode::NotFound)
-		{
-			GTotalMissingChunkInstallErrorCount.Increment();
-		}
-		else if (ErrorCode == EIoErrorCode::WriteError)
+		TUniqueLock Lock(InstallerTraceCounters);
+		InstallerTraceCounters.InstallCount.Increment();
+		InstallerTraceCounters.InflightInstallCount.Decrement();
+		InstallerTraceCounters.DownloadedBytes.Add(DownloadedBytes);
+		InstallerTraceCounters.InstallDurationMs.Increment(FPlatformTime::ToMilliseconds64(DurationCycles));
+		InstallerTraceCounters.CacheHitRatio.Increment(CacheHitRatio);
+		InstallerTraceCounters.AvgInstallDurationMs.Set(InstallerTraceCounters.InstallDurationMs.GetMean());
+		InstallerTraceCounters.AvgCacheHitRatio.Set(InstallerTraceCounters.InstallDurationMs.GetMean());
+	}
+#endif // COUNTERSTRACE_ENABLED
+
+	{
+		TUniqueLock Lock(InstallerAnalytics);
+		if (ErrorCode != EIoErrorCode::Ok)
 		{
-			GTotalWriteErrorInstallErrorCount.Increment();
+			InstallerAnalytics.InstallErrorCount++;
 		}
-		GTotalInstallErrorCount.Increment();
+		InstallerAnalytics.InstallCount++;
+		InstallerAnalytics.DownloadedBytes += DownloadedBytes;
+		InstallerAnalytics.TotalInstallDurationMs += uint64(FPlatformTime::ToMilliseconds64(DurationCycles));
+		InstallerAnalytics.TotalCacheHitRatio += CacheHitRatio;
 	}
-	GTotalInstallCount.Increment();
-	GTotalDownloadedChunkCount.Add(DownloadedChunkCount);
-	GTotalDownloadedChunkBytes.Add(DownloadedBytes);
-	GTotalInstallDurationMs.Add(uint64(FPlatformTime::ToMilliseconds64(DurationCycles)));
-	GInstallCacheHitRatio.Increment(CacheHitRatio);
-	//GAvgInstallCacheHitRatio.Set(GInstallCacheHitRatio.GetMean());
 }
 
 void FOnDemandContentInstallerStats::ReportAnalytics(TArray<FAnalyticsEventAttribute>& OutAnalyticsArray)
@@ -1128,57 +1171,77 @@ void FOnDemandContentInstallerStats::ReportAnalytics(TArray<FAnalyticsEventAttri
 		return;
 	}
 
-	UE::TUniqueLock Lock(InstallCacheStats);
+	FInstallerAnalytics CurrentInstallerAnalytics;
+	{
+		TUniqueLock Lock(InstallerAnalytics);
+		CurrentInstallerAnalytics	= InstallerAnalytics;
+		InstallerAnalytics			= FInstallerAnalytics();
+	}
+	FInstallCacheAnalytics CurrentInstallCacheAnalytics;
+	{
+		TUniqueLock Lock(InstallCacheAnalytics);
+		CurrentInstallCacheAnalytics	= InstallCacheAnalytics;
+		InstallCacheAnalytics			= FInstallCacheAnalytics();
+	}
 
 	AppendAnalyticsEventAttributeArray(OutAnalyticsArray
-		, TEXT("IadTotalInstallCount"), GTotalInstallCount.Get()
-		, TEXT("IadTotalInstallErrorCount"), GTotalInstallErrorCount.Get()
-		, TEXT("IadTotalWriteErrorInstallErrorCount"), GTotalWriteErrorInstallErrorCount.Get()
-		, TEXT("IadTotalDownloadedBytes"), GTotalDownloadedChunkBytes.Get()
-		, TEXT("IadTotalInstallDurationMs"), GTotalInstallDurationMs.Get()
-		, TEXT("IadAvgCacheHitRatio"), GInstallCacheHitRatio.GetMean()
-		, TEXT("IadInstallCacheStartupErrorCount"), InstallCacheStats.StartupErrorCount
-		, TEXT("IadInstallCacheVerificationRemovedBlockCount"), InstallCacheStats.VerificationRemovedBlockCount 
-		, TEXT("IadInstallCacheFlushErrorCount"), InstallCacheStats.FlushErrorCount 
-		, TEXT("IadInstallCachePurgeErrorCount"), InstallCacheStats.PurgeErrorCount
-		, TEXT("IadInstallCacheDefragErrorCount"), InstallCacheStats.DefragErrorCount 
-		, TEXT("IadInstallInstallCacheJournalCommitErrorCount"), InstallCacheStats.DefragErrorCount 
-		, TEXT("IadInstallInstallCacheMaxUsageSize"), InstallCacheStats.MaxCacheUsageSize 
-		, TEXT("IadInstallInstallCacheMaxReferencedBlockSize"), InstallCacheStats.MaxReferencedBlockSize
-		, TEXT("IadInstallInstallCacheMaxReferencedSize"), InstallCacheStats.MaxReferencedSize
-		, TEXT("IadInstallInstallCacheMaxFragmentedSize"), InstallCacheStats.MaxFragmentedSize
+		, TEXT("IadTotalInstallCount"), CurrentInstallerAnalytics.InstallCount
+		, TEXT("IadTotalInstallErrorCount"), CurrentInstallerAnalytics.InstallErrorCount 
+		, TEXT("IadTotalDownloadedBytes"), CurrentInstallerAnalytics.DownloadedBytes 
+		, TEXT("IadTotalInstallDurationMs"), CurrentInstallerAnalytics.TotalInstallDurationMs 
+		, TEXT("IadAvgInstallDurationMs"), (double(CurrentInstallerAnalytics.TotalInstallDurationMs) / double(CurrentInstallerAnalytics.InstallCount))
+		, TEXT("IadAvgCacheHitRatio"), (CurrentInstallerAnalytics.TotalCacheHitRatio / double(CurrentInstallerAnalytics.InstallCount))
+		, TEXT("IadInstallCacheStartupErrorCount"), CurrentInstallCacheAnalytics.StartupErrorCount
+		, TEXT("IadInstallCacheVerificationRemovedBlockCount"), CurrentInstallCacheAnalytics.VerificationRemovedBlockCount 
+		, TEXT("IadInstallCacheFlushCount"), CurrentInstallCacheAnalytics.FlushCount 
+		, TEXT("IadInstallCacheFlushErrorCount"), CurrentInstallCacheAnalytics.FlushErrorCount 
+		, TEXT("IadInstallCacheFlushedBytes"), CurrentInstallCacheAnalytics.FlushedBytes
+		, TEXT("IadInstallCachePurgeCount"), CurrentInstallCacheAnalytics.PurgeCount
+		, TEXT("IadInstallCachePurgeErrorCount"), CurrentInstallCacheAnalytics.PurgeErrorCount
+		, TEXT("IadInstallCacheDefragCount"), CurrentInstallCacheAnalytics.DefragCount 
+		, TEXT("IadInstallCacheDefragErrorCount"), CurrentInstallCacheAnalytics.DefragErrorCount 
+		, TEXT("IadInstallInstallCacheJournalCommitCount"), CurrentInstallCacheAnalytics.JournalCommitCount 
+		, TEXT("IadInstallInstallCacheJournalCommitErrorCount"), CurrentInstallCacheAnalytics.JournalCommitErrorCount
+		, TEXT("IadInstallInstallCacheMaxSize"), CurrentInstallCacheAnalytics.MaxCacheSize 
+		, TEXT("IadInstallInstallCacheMaxUsageSize"), CurrentInstallCacheAnalytics.MaxCacheUsageSize 
+		, TEXT("IadInstallInstallCacheMaxReferencedBlockSize"), CurrentInstallCacheAnalytics.MaxReferencedBlockSize
+		, TEXT("IadInstallInstallCacheMaxReferencedSize"), CurrentInstallCacheAnalytics.MaxReferencedSize
+		, TEXT("IadInstallInstallCacheMaxFragmentedSize"), CurrentInstallCacheAnalytics.MaxFragmentedSize
 	);
 }
 
 ////////////////////////////////////////////////////////////////////////////////
 void FOnDemandInstallCacheStats::OnStartupError(EIoErrorCode ErrorCode)
 {
-	UE::TUniqueLock Lock(InstallCacheStats);
-	InstallCacheStats.StartupErrorCount++;
+	UE::TUniqueLock Lock(InstallCacheAnalytics);
+	InstallCacheAnalytics.StartupErrorCount++;
 }
 
 void FOnDemandInstallCacheStats::OnFlush(EIoErrorCode ErrorCode, int64 ByteCount)
 {
+	UE::TUniqueLock Lock(InstallCacheAnalytics);
 	if (ErrorCode != EIoErrorCode::Ok)
 	{
-		UE::TUniqueLock Lock(InstallCacheStats);
-		InstallCacheStats.FlushErrorCount++;
+		InstallCacheAnalytics.FlushErrorCount++;
 	}
+	InstallCacheAnalytics.FlushCount++;
+	InstallCacheAnalytics.FlushedBytes += ByteCount;
 }
 
 void FOnDemandInstallCacheStats::OnJournalCommit(EIoErrorCode ErrorCode, int64 ByteCount)
 {
+	UE::TUniqueLock Lock(InstallCacheAnalytics);
 	if (ErrorCode != EIoErrorCode::Ok)
 	{
-		UE::TUniqueLock Lock(InstallCacheStats);
-		InstallCacheStats.JournalCommitErrorCount++;
+		InstallCacheAnalytics.JournalCommitErrorCount++;
 	}
+	InstallCacheAnalytics.JournalCommitCount++;
 }
 
 void FOnDemandInstallCacheStats::OnCasVerificationError(int32 RemoveChunks)
 {
-	UE::TUniqueLock Lock(InstallCacheStats);
-	InstallCacheStats.VerificationRemovedBlockCount += RemoveChunks;
+	UE::TUniqueLock Lock(InstallCacheAnalytics);
+	InstallCacheAnalytics.VerificationRemovedBlockCount += RemoveChunks;
 }
 
 void FOnDemandInstallCacheStats::OnPurge(
@@ -1188,28 +1251,24 @@ void FOnDemandInstallCacheStats::OnPurge(
 	uint64 BytesToPurge,
 	uint64 PurgedBytes)
 {
-	UE::TUniqueLock Lock(InstallCacheStats);
-	if (ErrorCode == EIoErrorCode::Ok)
-	{
-		InstallCacheStats.PurgedBytes += PurgedBytes;
-	}
-	else 
+	UE::TUniqueLock Lock(InstallCacheAnalytics);
+	if (ErrorCode != EIoErrorCode::Ok)
 	{
-		InstallCacheStats.PurgeErrorCount++;
+		InstallCacheAnalytics.PurgeErrorCount++;
 	}
+	InstallCacheAnalytics.PurgeCount++;
+	InstallCacheAnalytics.PurgedBytes += PurgedBytes;
 }
 
 void FOnDemandInstallCacheStats::OnDefrag(EIoErrorCode ErrorCode, uint64 FragmentedBytes)
 {
-	UE::TUniqueLock Lock(InstallCacheStats);
-	if (ErrorCode == EIoErrorCode::Ok)
-	{
-		InstallCacheStats.FragmentedBytes += FragmentedBytes;
-	}
-	else
+	UE::TUniqueLock Lock(InstallCacheAnalytics);
+	if (ErrorCode != EIoErrorCode::Ok)
 	{
-		InstallCacheStats.DefragErrorCount++;
+		InstallCacheAnalytics.DefragErrorCount++;
 	}
+	InstallCacheAnalytics.DefragCount++;
+	InstallCacheAnalytics.FragmentedBytes += FragmentedBytes;
 }
 
 void FOnDemandInstallCacheStats::OnCacheUsage(
@@ -1219,11 +1278,12 @@ void FOnDemandInstallCacheStats::OnCacheUsage(
 	uint64 ReferencedSize,
 	uint64 FragmentedSize)
 {
-	UE::TUniqueLock Lock(InstallCacheStats);
-	InstallCacheStats.MaxCacheUsageSize			= FMath::Max(InstallCacheStats.MaxCacheUsageSize, CacheSize);
-	InstallCacheStats.MaxReferencedBlockSize	= FMath::Max(InstallCacheStats.MaxReferencedBlockSize, ReferencedBlockSize);
-	InstallCacheStats.MaxReferencedSize			= FMath::Max(InstallCacheStats.MaxReferencedSize, ReferencedSize);
-	InstallCacheStats.MaxFragmentedSize			= FMath::Max(InstallCacheStats.MaxFragmentedSize, FragmentedSize);
+	UE::TUniqueLock Lock(InstallCacheAnalytics);
+	InstallCacheAnalytics.MaxCacheSize				= MaxCacheSize; 
+	InstallCacheAnalytics.MaxCacheUsageSize			= FMath::Max(InstallCacheAnalytics.MaxCacheUsageSize, CacheSize);
+	InstallCacheAnalytics.MaxReferencedBlockSize	= FMath::Max(InstallCacheAnalytics.MaxReferencedBlockSize, ReferencedBlockSize);
+	InstallCacheAnalytics.MaxReferencedSize			= FMath::Max(InstallCacheAnalytics.MaxReferencedSize, ReferencedSize);
+	InstallCacheAnalytics.MaxFragmentedSize			= FMath::Max(InstallCacheAnalytics.MaxFragmentedSize, FragmentedSize);
 }
 
 } // namespace UE::IoStore
