Chaos
- Robustify FCCDManager::ApplyCorrections to cover null particles (destroyed in a callback?)
- Undo CL 40145926 which adds newly awakened collisions in the Active lists and put it behind a disabled cvar

We speculate that adding newly awakened collisions to the Active list can lead to Active Collisions with garbage particle pointers. Possibly caused by some unexpected behavior in a callback. Until we track down that issue, disable the activation of awakened collisions.

diff --git a/Engine/Source/Runtime/Experimental/Chaos/Private/Chaos/CCDUtilities.cpp b/Engine/Source/Runtime/Experimental/Chaos/Private/Chaos/CCDUtilities.cpp
--- a/Engine/Source/Runtime/Experimental/Chaos/Private/Chaos/CCDUtilities.cpp
+++ b/Engine/Source/Runtime/Experimental/Chaos/Private/Chaos/CCDUtilities.cpp
@@ -924,9 +924,14 @@ namespace Chaos
 
 		if(Chaos::CVars::bClearNullCCDConstraintsOnApply)
 		{
-			Constraints.Remove(nullptr);
-			const int32 NumNullConstraints = SweptConstraints.Num() - Constraints.Num();
-			ensureMsgf(NumNullConstraints == 0, TEXT("Found %d nullptr constraints during CCD correction application"), NumNullConstraints);
+			Constraints.RemoveAll(
+				[](FPBDCollisionConstraint* C)
+				{
+					return (C == nullptr)
+						|| (C->GetParticle0() == nullptr)
+						|| (C->GetParticle1() == nullptr);
+				}
+			);
 		}
 
 		Constraints.Sort([](const FPBDCollisionConstraint& L, const FPBDCollisionConstraint& R)
diff --git a/Engine/Source/Runtime/Experimental/Chaos/Private/Chaos/Collision/ParticlePairMidPhase.cpp b/Engine/Source/Runtime/Experimental/Chaos/Private/Chaos/Collision/ParticlePairMidPhase.cpp
--- a/Engine/Source/Runtime/Experimental/Chaos/Private/Chaos/Collision/ParticlePairMidPhase.cpp
+++ b/Engine/Source/Runtime/Experimental/Chaos/Private/Chaos/Collision/ParticlePairMidPhase.cpp
@@ -63,6 +63,11 @@ namespace Chaos
 		FAutoConsoleVariableRef CVarChaos_Collision_EnableShapePairs(TEXT("p.Chaos.Collision.EnableShapePairs"), bChaos_Collision_MidPhase_EnableShapePairs, TEXT(""));
 		FAutoConsoleVariableRef CVarChaos_Collision_MaxShapePairs(TEXT("p.Chaos.Collision.MaxShapePairs"), Chaos_Collision_MidPhase_MaxShapePairs, TEXT(""));
 
+		// This should be enabled but it exposes a bug that must be fixed first: collisions with 
+		// invalid particle pointers are being activated causing a crash in CCDUtilities::ApplyCorrections
+		bool bChaosMidPhaseActivateWakingConstraints = false;
+		FAutoConsoleVariableRef CVarChaosMidPhaseActivateWakingConstraints(TEXT("p.Chaos.Bugs.ActivateWakingConstraints"), bChaosMidPhaseActivateWakingConstraints, TEXT(""));
+
 		// See FGenericParticlePairMidPhase::SetConvexOptimizationIsActive
 		bool bChaosMidPhaseModifierConvexOptimizerFix = true;
 		FAutoConsoleVariableRef CVarChaosMidPhaseModifierConvexOptimizerFix(TEXT("p.Chaos.Bugs.MidPhaseModifierConvexOptimizerFix"), bChaosMidPhaseModifierConvexOptimizerFix, TEXT(""));
@@ -694,9 +699,22 @@ namespace Chaos
 	{
 		if (Constraint.IsValid() && IsUsedSince(SleepEpoch))
 		{
-			// Update the epoch and add to the active list
-			Constraint->ConcreteContainer()->GetConstraintAllocator().AddActiveConstraint(Constraint.Get());
-			LastUsedEpoch = Constraint->GetContainerCookie().LastUsedEpoch;
+			if (bChaosMidPhaseActivateWakingConstraints)
+			{
+				// Update the epoch and add to the active list
+				Constraint->ConcreteContainer()->GetConstraintAllocator().AddActiveConstraint(Constraint.Get());
+				LastUsedEpoch = Constraint->GetContainerCookie().LastUsedEpoch;
+			}
+			else
+			{
+				// We need to refresh the epoch so that the constraint state will be used as the previous
+				// state if the pair is still colliding in the next tick. However, we are not in the active 
+				// collisions array so reset the indices.
+				Constraint->GetContainerCookie().LastUsedEpoch = CurrentEpoch;
+				Constraint->GetContainerCookie().ConstraintIndex = INDEX_NONE;
+				Constraint->GetContainerCookie().CCDConstraintIndex = INDEX_NONE;
+				LastUsedEpoch = CurrentEpoch;
+			}
 
 			// We have skipped collision detection for this particle because it was asleep, so we need to update the transforms...
 			// NOTE: this relies on the shape world transforms being up-to-date. They are usually updated in Integrate which
@@ -2087,8 +2105,17 @@ namespace Chaos
 			FPBDCollisionConstraintPtr& Constraint = KVP.Value;
 			if (Constraint->GetContainerCookie().LastUsedEpoch >= LastUsedEpoch)
 			{
-				// Update the epoch and add to the active list
-				Constraint->ConcreteContainer()->GetConstraintAllocator().AddActiveConstraint(Constraint.Get());
+				if (bChaosMidPhaseActivateWakingConstraints)
+				{
+					// Update the epoch and add to the active list
+					Constraint->ConcreteContainer()->GetConstraintAllocator().AddActiveConstraint(Constraint.Get());
+				}
+				else
+				{
+					Constraint->GetContainerCookie().LastUsedEpoch = CurrentEpoch;
+					Constraint->GetContainerCookie().ConstraintIndex = INDEX_NONE;
+					Constraint->GetContainerCookie().CCDConstraintIndex = INDEX_NONE;
+				}
 			}
 		}
 	}
@@ -2284,8 +2311,17 @@ namespace Chaos
 	{
 		if (Constraint.IsValid() && (Constraint->GetContainerCookie().LastUsedEpoch >= LastUsedEpoch))
 		{
-			// Update the epoch and add to the active list
-			Constraint->ConcreteContainer()->GetConstraintAllocator().AddActiveConstraint(Constraint.Get());
+			if (bChaosMidPhaseActivateWakingConstraints)
+			{
+				// Update the epoch and add to the active list
+				Constraint->ConcreteContainer()->GetConstraintAllocator().AddActiveConstraint(Constraint.Get());
+			}
+			else
+			{
+				Constraint->GetContainerCookie().LastUsedEpoch = CurrentEpoch;
+				Constraint->GetContainerCookie().ConstraintIndex = INDEX_NONE;
+				Constraint->GetContainerCookie().CCDConstraintIndex = INDEX_NONE;
+			}
 		}
 	}
 
