CitySample - Graphical artifacts on Player Character's model on PS5
RBAN instability was found when binary was compiled with Horde
Building locally binary was all fine.
 inverse transform is using no scale function to compute capsule capsule

diff --git a/Engine/Source/Runtime/Experimental/Chaos/Private/Chaos/Collision/ContactPointsMiscShapes.cpp b/Engine/Source/Runtime/Experimental/Chaos/Private/Chaos/Collision/ContactPointsMiscShapes.cpp
--- a/Engine/Source/Runtime/Experimental/Chaos/Private/Chaos/Collision/ContactPointsMiscShapes.cpp
+++ b/Engine/Source/Runtime/Experimental/Chaos/Private/Chaos/Collision/ContactPointsMiscShapes.cpp
@@ -496,7 +496,7 @@ namespace Chaos
 			const FVec3f LocationA = P1 + Dir * A.GetRadiusf();
 			const FVec3f LocationB = P2 - Dir * B.GetRadiusf();
 
-			Result.ShapeContactPoints[0] = CapsuleAToCapsuleBf.InverseTransformPosition(LocationA);
+			Result.ShapeContactPoints[0] = CapsuleAToCapsuleBf.InverseTransformPositionNoScale(LocationA);
 			Result.ShapeContactPoints[1] = LocationB;
 			Result.ShapeContactNormal = Normal;
 			Result.Phi = NewPhi;
diff --git a/Engine/Source/Runtime/Experimental/Chaos/Private/Chaos/CollisionOneShotManifoldsMiscShapes.cpp b/Engine/Source/Runtime/Experimental/Chaos/Private/Chaos/CollisionOneShotManifoldsMiscShapes.cpp
--- a/Engine/Source/Runtime/Experimental/Chaos/Private/Chaos/CollisionOneShotManifoldsMiscShapes.cpp
+++ b/Engine/Source/Runtime/Experimental/Chaos/Private/Chaos/CollisionOneShotManifoldsMiscShapes.cpp
@@ -396,7 +396,7 @@ namespace Chaos
 			// We always add the closest point to the manifold
 			// We may also add 2 more points generated from the end cap positions of the smaller capsule
 			FContactPointf ClosestContactPoint;
-			ClosestContactPoint.ShapeContactPoints[0] = CapsuleAToCapsuleBf.InverseTransformPosition(ClosestLocationA);
+			ClosestContactPoint.ShapeContactPoints[0] = CapsuleAToCapsuleBf.InverseTransformPositionNoScale(ClosestLocationA);
 			ClosestContactPoint.ShapeContactPoints[1] = ClosestLocationB;
 			ClosestContactPoint.ShapeContactNormal = ClosestNormal;
 			ClosestContactPoint.Phi = ClosestPhi;
