Iris
* Clients will now report to the server if a root object's replication is blocked by must be mapped references after a timeout expires
* Server will ensure and log information about the object and the references that are not loading

#rnx

diff --git a/Engine/Source/Runtime/Engine/Private/DataChannel.cpp b/Engine/Source/Runtime/Engine/Private/DataChannel.cpp
--- a/Engine/Source/Runtime/Engine/Private/DataChannel.cpp
+++ b/Engine/Source/Runtime/Engine/Private/DataChannel.cpp
@@ -1758,6 +1758,7 @@ IMPLEMENT_CONTROL_CHANNEL_MESSAGE(BeaconAssignGUID);
 IMPLEMENT_CONTROL_CHANNEL_MESSAGE(BeaconNetGUIDAck);
 IMPLEMENT_CONTROL_CHANNEL_MESSAGE(IrisProtocolMismatch);
 IMPLEMENT_CONTROL_CHANNEL_MESSAGE(IrisNetRefHandleError);
+IMPLEMENT_CONTROL_CHANNEL_MESSAGE(IrisNetRefHandleErrorWithArray);
 
 void UControlChannel::Init( UNetConnection* InConnection, int32 InChannelIndex, EChannelCreateFlags CreateFlags )
 {
@@ -1993,6 +1994,21 @@ void UControlChannel::ReceivedBunch( FInBunch& Bunch )
 				{
 					IrisRepSystem->ReportErrorWithNetRefHandle(NetRefHandleError, NetRefHandleId, Connection->GetConnectionHandle().GetParentConnectionId());
 				}
+#endif
+			}
+		}
+		else if (MessageType == NMT_IrisNetRefHandleErrorWithArray)
+		{
+			UE::Net::ENetRefHandleError NetRefHandleError = UE::Net::ENetRefHandleError::None;
+			uint64 NetRefHandleId = 0;
+			TArray<uint64> ExtraNetRefs;
+			if (FNetControlMessage<NMT_IrisNetRefHandleErrorWithArray>::Receive(Bunch, NetRefHandleError, NetRefHandleId, ExtraNetRefs))
+			{
+#if UE_WITH_IRIS
+				if (UReplicationSystem* IrisRepSystem = Connection->Driver->GetReplicationSystem())
+				{
+					IrisRepSystem->ReportErrorWithNetRefHandle(NetRefHandleError, NetRefHandleId, Connection->GetConnectionHandle().GetParentConnectionId(), ExtraNetRefs);
+				}
 #endif
 			}
 		}
@@ -2079,6 +2095,9 @@ void UControlChannel::ReceivedBunch( FInBunch& Bunch )
 				case NMT_IrisNetRefHandleError:
 					FNetControlMessage<NMT_IrisNetRefHandleError>::Discard(Bunch);
 					break;
+				case NMT_IrisNetRefHandleErrorWithArray:
+					FNetControlMessage<NMT_IrisNetRefHandleErrorWithArray>::Discard(Bunch);
+					break;
 				case NMT_JoinNoPawn:
 					FNetControlMessage<NMT_JoinNoPawn>::Discard(Bunch);
 					break;
diff --git a/Engine/Source/Runtime/Engine/Private/Net/Iris/ReplicationSystem/EngineReplicationBridge.cpp b/Engine/Source/Runtime/Engine/Private/Net/Iris/ReplicationSystem/EngineReplicationBridge.cpp
--- a/Engine/Source/Runtime/Engine/Private/Net/Iris/ReplicationSystem/EngineReplicationBridge.cpp
+++ b/Engine/Source/Runtime/Engine/Private/Net/Iris/ReplicationSystem/EngineReplicationBridge.cpp
@@ -822,6 +822,30 @@ void UEngineReplicationBridge::SendErrorWithNetRefHandle(UE::Net::ENetRefHandleE
 	}
 }
 
+void UEngineReplicationBridge::SendErrorWithNetRefHandle(UE::Net::ENetRefHandleError ErrorType, FNetRefHandle RefHandle, uint32 ConnectionId, const TArray<FNetRefHandle>& ExtraNetRefHandle)
+{
+	if (NetDriver)
+	{
+		if (UNetConnection* ClientConnection = NetDriver->GetConnectionByHandle(UE::Net::FConnectionHandle(ConnectionId)))
+		{
+			uint64 HandleId = RefHandle.GetId();
+
+			TArray<uint64> ArrayIds;
+			ArrayIds.Reserve(ExtraNetRefHandle.Num());
+			for(FNetRefHandle ExtraHandle : ExtraNetRefHandle)
+			{
+				ArrayIds.Emplace(ExtraHandle.GetId());
+			}
+
+			FNetControlMessage<NMT_IrisNetRefHandleErrorWithArray>::Send(ClientConnection, ErrorType, HandleId, ArrayIds);
+		}
+		else
+		{
+			UE_LOG(LogIrisBridge, Error, TEXT("UEngineReplicationBridge::SendErrorWithNetRefHandle could not find Connection for id:%u"), ConnectionId);
+		}
+	}
+}
+
 void UEngineReplicationBridge::ActorChangedLevel(const AActor* Actor, const ULevel* PreviousLevel)
 {
 	if (!UE::Net::Private::bEnableActorLevelChanges)
diff --git a/Engine/Source/Runtime/Engine/Public/Net/DataChannel.h b/Engine/Source/Runtime/Engine/Public/Net/DataChannel.h
--- a/Engine/Source/Runtime/Engine/Public/Net/DataChannel.h
+++ b/Engine/Source/Runtime/Engine/Public/Net/DataChannel.h
@@ -217,6 +217,7 @@ DEFINE_CONTROL_CHANNEL_MESSAGE(BeaconNetGUIDAck, 28, FString); // server receive
 
 DEFINE_CONTROL_CHANNEL_MESSAGE(IrisProtocolMismatch, 29, uint64); // client has a different protocol hash from the server
 DEFINE_CONTROL_CHANNEL_MESSAGE(IrisNetRefHandleError, 30, UE::Net::ENetRefHandleError, uint64); // a specific handle caused an error and we want the remote connection to log all information it has on it
+DEFINE_CONTROL_CHANNEL_MESSAGE(IrisNetRefHandleErrorWithArray, 33, UE::Net::ENetRefHandleError, uint64, TArray<uint64>); // a specific handle caused 
 
 // Called in place of Join or JoinSplit during connection handshake, and will created a connection that has the following properties:
 //	* Won't spawn an instance of APawn.
diff --git a/Engine/Source/Runtime/Engine/Public/Net/Iris/ReplicationSystem/EngineReplicationBridge.h b/Engine/Source/Runtime/Engine/Public/Net/Iris/ReplicationSystem/EngineReplicationBridge.h
--- a/Engine/Source/Runtime/Engine/Public/Net/Iris/ReplicationSystem/EngineReplicationBridge.h
+++ b/Engine/Source/Runtime/Engine/Public/Net/Iris/ReplicationSystem/EngineReplicationBridge.h
@@ -102,6 +102,9 @@ public:
 	/** Tell the remote connection that we detected a reading error with a specific replicated object */
 	ENGINE_API virtual void SendErrorWithNetRefHandle(UE::Net::ENetRefHandleError ErrorType, FNetRefHandle RefHandle, uint32 ConnectionId) override;
 
+	/** Tell the remote connection about an error with extra information */
+	ENGINE_API virtual void SendErrorWithNetRefHandle(UE::Net::ENetRefHandleError ErrorType, FNetRefHandle RefHandle, uint32 ConnectionId, const TArray<FNetRefHandle>& ExtraNetRefHandle) override;
+
 	/** Add the rootobject to the level's filter group so it will only be relevant if the connection has that level streamed in. */
 	ENGINE_API void AddRootObjectToLevelGroup(const UObject* RootObject, const ULevel* Level);
 	
diff --git a/Engine/Source/Runtime/Experimental/Iris/Core/Private/Iris/ReplicationSystem/ObjectReferenceCache.cpp b/Engine/Source/Runtime/Experimental/Iris/Core/Private/Iris/ReplicationSystem/ObjectReferenceCache.cpp
--- a/Engine/Source/Runtime/Experimental/Iris/Core/Private/Iris/ReplicationSystem/ObjectReferenceCache.cpp
+++ b/Engine/Source/Runtime/Experimental/Iris/Core/Private/Iris/ReplicationSystem/ObjectReferenceCache.cpp
@@ -2029,7 +2029,7 @@ void FObjectReferenceCache::AsyncPackageCallback(const FName& PackageName, UPack
 
 		// This won't be the exact amount of time that we spent loading the package, but should
 		// give us a close enough estimate (within a frame time).
-		const double LoadTime = (ReplicationSystem->GetElapsedTime() - PendingLoadRequest->RequestStartTime);
+		//const double LoadTime = (ReplicationSystem->GetElapsedTime() - PendingLoadRequest->RequestStartTime);
 		//if (GGuidCacheTrackAsyncLoadingGUIDThreshold > 0.f &&
 		//	LoadTime >= GGuidCacheTrackAsyncLoadingGUIDThreshold)
 		//{
@@ -2125,5 +2125,16 @@ void FObjectReferenceCache::RemoveTrackedQueuedBatchObjectReference(const FNetRe
 	}
 }
 
+const TCHAR* FObjectReferenceCache::GetObjectRelativePath(FNetRefHandle NetRefHandle) const
+{
+	if (const FCachedNetObjectReference* CacheObject = ReferenceHandleToCachedReference.Find(NetRefHandle))
+	{
+		const TCHAR* ResolvedToken = StringTokenStore->ResolveToken(CacheObject->RelativePath);
+		return ResolvedToken;
+	}
 
+	return nullptr;
 }
+
+
+} // end namespace UE::Net::Private
diff --git a/Engine/Source/Runtime/Experimental/Iris/Core/Private/Iris/ReplicationSystem/ObjectReferenceCache.h b/Engine/Source/Runtime/Experimental/Iris/Core/Private/Iris/ReplicationSystem/ObjectReferenceCache.h
--- a/Engine/Source/Runtime/Experimental/Iris/Core/Private/Iris/ReplicationSystem/ObjectReferenceCache.h
+++ b/Engine/Source/Runtime/Experimental/Iris/Core/Private/Iris/ReplicationSystem/ObjectReferenceCache.h
@@ -151,6 +151,9 @@ public:
 
 	FString DescribeObjectReference(const FNetObjectReference Ref, const FNetObjectResolveContext& ResolveContext);
 
+	/** Return the stored relative path of a replicated object if it's in the cache. */
+	const TCHAR* GetObjectRelativePath(FNetRefHandle NetRefHandle) const;
+
 private:
 
 	struct FCachedNetObjectReference
diff --git a/Engine/Source/Runtime/Experimental/Iris/Core/Private/Iris/ReplicationSystem/ObjectReplicationBridge.cpp b/Engine/Source/Runtime/Experimental/Iris/Core/Private/Iris/ReplicationSystem/ObjectReplicationBridge.cpp
--- a/Engine/Source/Runtime/Experimental/Iris/Core/Private/Iris/ReplicationSystem/ObjectReplicationBridge.cpp
+++ b/Engine/Source/Runtime/Experimental/Iris/Core/Private/Iris/ReplicationSystem/ObjectReplicationBridge.cpp
@@ -2513,6 +2513,10 @@ void UObjectReplicationBridge::OnErrorWithNetRefHandleReported(UE::Net::ENetRefH
 				{
 					ensureMsgf(false, TEXT("NetRefHandle error: Replication disabled for %s. Client was unable to read the creation data and did not spawn the netobject."), *NetRefHandleManager->PrintObjectFromNetRefHandle(RefHandle));
 				} break;
+				case ENetRefHandleError::BlockedByMustBeMapped:
+				{
+					ensureMsgf(false, TEXT("NetRefHandle error: Replication blocked for %s due to ? unmapped references."), *NetRefHandleManager->PrintObjectFromNetRefHandle(RefHandle));
+				} break;
 				default:
 				{
 					ensureMsgf(false, TEXT("NetRefHandle error: %s reported for %s. This should not occur. Missing switch case or wrong error type used."), LexToString(ErrorType), *NetRefHandleManager->PrintObjectFromNetRefHandle(RefHandle));
@@ -2531,6 +2535,74 @@ void UObjectReplicationBridge::OnErrorWithNetRefHandleReported(UE::Net::ENetRefH
 	UE_LOG_BRIDGEID(LogIris, Error, TEXT("OnErrorWithNetRefHandleReported: %s from client %s. Problematic object was %s"),  LexToString(ErrorType), *PrintConnectionInfo(ConnectionId), *NetRefHandleManager->PrintObjectFromIndex(ObjectInternalIndex));
 }
 
+void UObjectReplicationBridge::OnErrorWithNetRefHandleReported(UE::Net::ENetRefHandleError ErrorType, FNetRefHandle RefHandle, uint32 ConnectionId, const TArray<FNetRefHandle>& ExtraHandles)
+{
+	using namespace UE::Net;
+	using namespace UE::Net::Private;
+
+	// Ensure at the end so the log contains all the relevant information
+	ON_SCOPE_EXIT
+	{
+		if (GEnsureNetRefHandleError != -1 && (GEnsureNetRefHandleError == 0 || EnumHasAnyFlags((ENetRefHandleError)GEnsureNetRefHandleError,ErrorType)))
+		{
+			// Use different calls for every error type to prevent only reporting the first error that occured.
+			switch (ErrorType)
+			{
+				case ENetRefHandleError::BitstreamCorrupted:
+				{
+					ensureMsgf(false, TEXT("NetRefHandle error: Bitstream corrupted reported with %s. Look at the log for important information on the object tied to the handle."), *NetRefHandleManager->PrintObjectFromNetRefHandle(RefHandle));
+				} break;
+				case ENetRefHandleError::ReplicationDisabled:
+				{
+					ensureMsgf(false, TEXT("NetRefHandle error: Replication disabled for %s. Client was unable to read the creation data and did not spawn the netobject."), *NetRefHandleManager->PrintObjectFromNetRefHandle(RefHandle));
+				} break;
+				case ENetRefHandleError::BlockedByMustBeMapped:
+				{
+					ensureMsgf(false, TEXT("NetRefHandle error: Replication blocked for %s due to %d unmapped references."), *NetRefHandleManager->PrintObjectFromNetRefHandle(RefHandle), ExtraHandles.Num());
+				} break;
+				default:
+				{
+					ensureMsgf(false, TEXT("NetRefHandle error: %s reported for %s. This should not occur. Missing switch case or wrong error type used."), LexToString(ErrorType), *NetRefHandleManager->PrintObjectFromNetRefHandle(RefHandle));
+				} break;
+			}
+		}
+	};
+
+	// Print the blocked object
+	{
+		const FInternalNetRefIndex ObjectInternalIndex = NetRefHandleManager->GetInternalIndex(RefHandle);
+		if (ObjectInternalIndex == FNetRefHandleManager::InvalidInternalIndex)
+		{
+			UE_LOG_BRIDGEID(LogIris, Warning, TEXT("OnErrorWithNetRefHandleReported: %s from client %s for %s but object has no InternalIndex."), LexToString(ErrorType), *PrintConnectionInfo(ConnectionId), *RefHandle.ToString());
+			return;
+		}
+
+		UE_LOG_BRIDGEID(LogIris, Error, TEXT("OnErrorWithNetRefHandleReported: %s from client %s. Problematic object was %s"), LexToString(ErrorType), *PrintConnectionInfo(ConnectionId), *NetRefHandleManager->PrintObjectFromIndex(ObjectInternalIndex));
+	}
+
+	// Make a custom printer for the specific error
+	auto PrintExtraObjectInfo = [&](FNetRefHandle ObjectHandle) -> FString
+	{
+		// Error specific message
+		if (ErrorType == ENetRefHandleError::BlockedByMustBeMapped)
+		{
+			return FString::Printf(TEXT("MustBeMapped reference: %s | PathName: %s"), *NetRefHandleManager->PrintObjectFromNetRefHandle(ObjectHandle), GetObjectReferenceCache()->GetObjectRelativePath(ObjectHandle));
+		}
+
+		// For not yet implemented errors
+		return FString::Printf(TEXT("Extra object: %s"), *NetRefHandleManager->PrintObjectFromNetRefHandle(ObjectHandle));
+	};
+
+	// Print the extra objects sent with the error
+	for (FNetRefHandle ExtraHandle : ExtraHandles)
+	{
+		UE_LOG_BRIDGEID(LogIris, Error, TEXT("\t%s"), *PrintExtraObjectInfo(ExtraHandle));
+	}
+
+	// All all internal info we have on the blocked object
+	PrintDebugInfoForNetRefHandle(RefHandle, ConnectionId, UObjectReplicationBridge::EPrintDebugInfoTraits::NoProtocolState);
+}
+
 TArray<uint32> UObjectReplicationBridge::FindConnectionsFromArgs(const TArray<FString>& Args) const
 {
 	using namespace UE::Net::Private;
@@ -2562,4 +2634,4 @@ TArray<uint32> UObjectReplicationBridge::FindConnectionsFromArgs(const TArray<FS
 	}
 
 	return ConnectionList;
-}
\ No newline at end of file
+}
diff --git a/Engine/Source/Runtime/Experimental/Iris/Core/Private/Iris/ReplicationSystem/ReplicationReader.cpp b/Engine/Source/Runtime/Experimental/Iris/Core/Private/Iris/ReplicationSystem/ReplicationReader.cpp
--- a/Engine/Source/Runtime/Experimental/Iris/Core/Private/Iris/ReplicationSystem/ReplicationReader.cpp
+++ b/Engine/Source/Runtime/Experimental/Iris/Core/Private/Iris/ReplicationSystem/ReplicationReader.cpp
@@ -119,6 +119,12 @@ static FAutoConsoleVariableRef CVarQueuedBatchTimeoutSeconds(
 	QueuedBatchTimeoutSeconds,
 	TEXT("Time in seconds to wait for queued batches to flush before logging a warning."));
 
+static int32 MaxMustBeMappedHandleArray = 16;
+static FAutoConsoleVariableRef CVarMaxMustBeMappedHandleArray(
+	TEXT("net.Iris.MaxMustBeMappedHandleArray"),
+	MaxMustBeMappedHandleArray,
+	TEXT("Maximum number of unresolved references that will be included with the QueuedBatch timeout error message."));
+
 static const FName NetError_FailedToFindAttachmentQueue("Failed to find attachment queue");
 
 //------------------------------------------------------------------------
@@ -2765,14 +2771,20 @@ void FReplicationReader::ProcessQueuedBatches()
 						, PendingBatchData.PendingMustBeMappedReferences.Num()
 						, QueuedBatchPendingTime);
 
-					if (UE_LOG_ACTIVE(LogIris, Log))
+					TArray<FNetRefHandle> PendingMustBeMapped;
+					PendingMustBeMapped.Reserve(MaxMustBeMappedHandleArray);
+
+					for (const FNetRefHandle& NetRefHandle : PendingBatchData.PendingMustBeMappedReferences)
 					{
-						for (const FNetRefHandle& NetRefHandle : PendingBatchData.PendingMustBeMappedReferences)
+						UE_LOG(LogIris, Log, TEXT("  NetRefHandle Pending resolve [%s]"), *ObjectReferenceCache->DescribeObjectReference(FObjectReferenceCache::MakeNetObjectReference(NetRefHandle), ResolveContext));
+						if (PendingMustBeMapped.Num() < MaxMustBeMappedHandleArray)
 						{
-							UE_LOG(LogIris, Log, TEXT("  NetRefHandle Pending resolve [%s]"), *ObjectReferenceCache->DescribeObjectReference(FObjectReferenceCache::MakeNetObjectReference(NetRefHandle), ResolveContext));
+							PendingMustBeMapped.Emplace(NetRefHandle);
 						}
 					}
 
+					ReplicationBridge->SendErrorWithNetRefHandle(UE::Net::ENetRefHandleError::BlockedByMustBeMapped, PendingBatchData.Handle, Parameters.ConnectionId, PendingMustBeMapped);
+
 					// Reset counter
 					PendingBatchData.PendingBatchTryProcessCount = 0;
 				}
diff --git a/Engine/Source/Runtime/Experimental/Iris/Core/Private/Iris/ReplicationSystem/ReplicationSystem.cpp b/Engine/Source/Runtime/Experimental/Iris/Core/Private/Iris/ReplicationSystem/ReplicationSystem.cpp
--- a/Engine/Source/Runtime/Experimental/Iris/Core/Private/Iris/ReplicationSystem/ReplicationSystem.cpp
+++ b/Engine/Source/Runtime/Experimental/Iris/Core/Private/Iris/ReplicationSystem/ReplicationSystem.cpp
@@ -2031,6 +2031,20 @@ void UReplicationSystem::ReportErrorWithNetRefHandle(UE::Net::ENetRefHandleError
 	Impl->ReplicationSystemInternal.GetReplicationBridge()->OnErrorWithNetRefHandleReported(ErrorType, NetRefHandle, ConnectionId);
 }
 
+void UReplicationSystem::ReportErrorWithNetRefHandle(UE::Net::ENetRefHandleError ErrorType, uint64 NetRefHandleId, uint32 ConnectionId, const TArray<uint64>& ExtraNetRefHandles)
+{
+	using namespace UE::Net::Private;
+	const FNetRefHandle NetRefHandle = FNetRefHandleManager::MakeNetRefHandle(NetRefHandleId, GetId());
+	TArray<FNetRefHandle> ExtraHandles;
+	ExtraHandles.Reserve(ExtraNetRefHandles.Num());
+	for (uint64 HandleId : ExtraNetRefHandles)
+	{
+		ExtraHandles.Emplace(FNetRefHandleManager::MakeNetRefHandle(HandleId, GetId()));
+	}
+
+	Impl->ReplicationSystemInternal.GetReplicationBridge()->OnErrorWithNetRefHandleReported(ErrorType, NetRefHandle, ConnectionId, ExtraHandles);
+}
+
 void UReplicationSystem::CollectNetMetrics(UE::Net::FNetMetrics& OutNetMetrics) const
 {
 	Impl->CollectNetMetrics(OutNetMetrics);
diff --git a/Engine/Source/Runtime/Experimental/Iris/Core/Public/Iris/ReplicationSystem/ObjectReplicationBridge.h b/Engine/Source/Runtime/Experimental/Iris/Core/Public/Iris/ReplicationSystem/ObjectReplicationBridge.h
--- a/Engine/Source/Runtime/Experimental/Iris/Core/Public/Iris/ReplicationSystem/ObjectReplicationBridge.h
+++ b/Engine/Source/Runtime/Experimental/Iris/Core/Public/Iris/ReplicationSystem/ObjectReplicationBridge.h
@@ -299,6 +299,7 @@ protected:
 	IRISCORE_API virtual void DetachInstance(FNetRefHandle Handle) override;
 	IRISCORE_API virtual void OnProtocolMismatchReported(FNetRefHandle RefHandle, uint32 ConnectionId) override;
 	IRISCORE_API virtual void OnErrorWithNetRefHandleReported(UE::Net::ENetRefHandleError ErrorType, FNetRefHandle RefHandle, uint32 ConnectionId) override;
+	IRISCORE_API virtual void OnErrorWithNetRefHandleReported(UE::Net::ENetRefHandleError ErrorType, FNetRefHandle RefHandle, uint32 ConnectionId, const TArray<FNetRefHandle>& ExtraHandles) override;
 	
 protected:
 	
diff --git a/Engine/Source/Runtime/Experimental/Iris/Core/Public/Iris/ReplicationSystem/ReplicationBridge.h b/Engine/Source/Runtime/Experimental/Iris/Core/Public/Iris/ReplicationSystem/ReplicationBridge.h
--- a/Engine/Source/Runtime/Experimental/Iris/Core/Public/Iris/ReplicationSystem/ReplicationBridge.h
+++ b/Engine/Source/Runtime/Experimental/Iris/Core/Public/Iris/ReplicationSystem/ReplicationBridge.h
@@ -247,13 +247,15 @@ protected:
 	IRISCORE_API virtual bool IsAllowedToDestroyInstance(const UObject* Instance) const;
 
 	/** Called when a remote connection detected a protocol mismatch when trying to instantiate the NetRefHandle replicated object. */
-	virtual void OnProtocolMismatchReported(FNetRefHandle RefHandle, uint32 ConnectionId) {}
+	IRISCORE_API virtual void OnProtocolMismatchReported(FNetRefHandle RefHandle, uint32 ConnectionId) {}
 
 	/** Called when a remote connection has a critical error caused by a specific NetRefHandle */
-	virtual void OnErrorWithNetRefHandleReported(UE::Net::ENetRefHandleError ErrorType, FNetRefHandle RefHandle, uint32 ConnectionId) {}
+	IRISCORE_API virtual void OnErrorWithNetRefHandleReported(UE::Net::ENetRefHandleError ErrorType, FNetRefHandle RefHandle, uint32 ConnectionId) {}
+	IRISCORE_API virtual void OnErrorWithNetRefHandleReported(UE::Net::ENetRefHandleError ErrorType, FNetRefHandle RefHandle, uint32 ConnectionId, const TArray<FNetRefHandle>& ExtraNetRefHandle) {}
 
 	/** Tell the remote connection that we detected a reading error with a specific replicated object */
-	virtual void SendErrorWithNetRefHandle(UE::Net::ENetRefHandleError ErrorType, FNetRefHandle RefHandle, uint32 ConnectionId) {}
+	IRISCORE_API virtual void SendErrorWithNetRefHandle(UE::Net::ENetRefHandleError ErrorType, FNetRefHandle RefHandle, uint32 ConnectionId) {}
+	IRISCORE_API virtual void SendErrorWithNetRefHandle(UE::Net::ENetRefHandleError ErrorType, FNetRefHandle RefHandle, uint32 ConnectionId, const TArray<FNetRefHandle>& ExtraNetRefHandle) {}
 
 private:
 
diff --git a/Engine/Source/Runtime/Experimental/Iris/Core/Public/Iris/ReplicationSystem/ReplicationSystem.h b/Engine/Source/Runtime/Experimental/Iris/Core/Public/Iris/ReplicationSystem/ReplicationSystem.h
--- a/Engine/Source/Runtime/Experimental/Iris/Core/Public/Iris/ReplicationSystem/ReplicationSystem.h
+++ b/Engine/Source/Runtime/Experimental/Iris/Core/Public/Iris/ReplicationSystem/ReplicationSystem.h
@@ -711,6 +711,9 @@ public:
 	/** Called when a connection reports a critical error with a netrefhandle object */
 	IRISCORE_API void ReportErrorWithNetRefHandle(UE::Net::ENetRefHandleError ErrorType, uint64 NetRefHandleId, uint32 ConnectionId);
 
+	/** Called when a connection reports a critical error with a netrefhandle object and extra handles */
+	IRISCORE_API void ReportErrorWithNetRefHandle(UE::Net::ENetRefHandleError ErrorType, uint64 NetRefHandleId, uint32 ConnectionId, const TArray<uint64>& ExtraNetRefHandles);
+
 	/** Collect relevant metrics that could be sent to analytics */
 	IRISCORE_API void CollectNetMetrics(UE::Net::FNetMetrics& OutNetMetrics) const;
 
diff --git a/Engine/Source/Runtime/Net/Core/Public/Net/Core/Connection/NetEnums.h b/Engine/Source/Runtime/Net/Core/Public/Net/Core/Connection/NetEnums.h
--- a/Engine/Source/Runtime/Net/Core/Public/Net/Core/Connection/NetEnums.h
+++ b/Engine/Source/Runtime/Net/Core/Public/Net/Core/Connection/NetEnums.h
@@ -131,6 +131,7 @@ enum class ENetRefHandleError : uint32
 	None = 0x00,
 	BitstreamCorrupted = 0x01, // An error while reading detected an irrecoverable bitstream corruption. Client was forced to disconnect
 	ReplicationDisabled = 0x02, // An error while reading was recoverable but prevents this NetObject from replicating.
+	BlockedByMustBeMapped = 0x03, // An object is blocked due to MustBeMapped references not resolving fast enough
 };
 ENUM_CLASS_FLAGS(ENetRefHandleError);
 
@@ -141,6 +142,7 @@ inline const TCHAR* LexToString(ENetRefHandleError NetRefHandleError)
 		case ENetRefHandleError::None: return TEXT("None");
 		case ENetRefHandleError::BitstreamCorrupted: return TEXT("BitstreamCorrupted");
 		case ENetRefHandleError::ReplicationDisabled: return TEXT("ReplicationDisabled");
+		case ENetRefHandleError::BlockedByMustBeMapped: return TEXT("BlockedByMustBeMapped");
 		default: return TEXT("missing");
 	}
 }
