Write out type description for UPropertyBag based structs, as this is the serialization semantic used in the binary serializer

#rnx

diff --git a/Engine/Source/Runtime/CoreUObject/Public/StructUtils/PropertyBag.h b/Engine/Source/Runtime/CoreUObject/Public/StructUtils/PropertyBag.h
--- a/Engine/Source/Runtime/CoreUObject/Public/StructUtils/PropertyBag.h
+++ b/Engine/Source/Runtime/CoreUObject/Public/StructUtils/PropertyBag.h
@@ -1261,7 +1261,7 @@ class UPropertyBagMissingObject : public UObject
 	GENERATED_BODY()
 };
 
-
+namespace UE::Private { struct FJsonStringifyImpl; }
 /**
  * A script struct that is used to store the value of the property bag instance.
  * References to UPropertyBag cannot be serialized, instead the array of the properties
@@ -1325,7 +1325,9 @@ protected:
 
 	std::atomic<int32> RefCount = 0;
 	
+	// both of these types need to serialize their property bag:
 	friend struct FInstancedPropertyBag;
+	friend struct ::UE::Private::FJsonStringifyImpl;
 };
 
 #undef UE_API
\ No newline at end of file
diff --git a/Engine/Source/Runtime/Experimental/JsonObjectGraph/Private/JsonObjectGraphConventions.h b/Engine/Source/Runtime/Experimental/JsonObjectGraph/Private/JsonObjectGraphConventions.h
--- a/Engine/Source/Runtime/Experimental/JsonObjectGraph/Private/JsonObjectGraphConventions.h
+++ b/Engine/Source/Runtime/Experimental/JsonObjectGraph/Private/JsonObjectGraphConventions.h
@@ -26,6 +26,8 @@
 #define UE_JSON_OBJECT_SPARSE_CLASS_DATA_KEY_TCHAR TEXT("__SparseClassData")
 #define UE_JSON_OBJECT_INDIRECTLY_REFERENCED_KEY_TCHAR TEXT("__IndirectlyReferenced")
 #define UE_JSON_SCRIPTSTRUCT_TCHAR TEXT("__UScriptStruct")
+#define UE_JSON_PROPERTYBAG_TCHAR TEXT("__UPropertyBag")
+#define UE_JSON_PROPERTYDESCS_TCHAR TEXT("PropertyDescs")
 
 // Reference encoding related conventions:
 #define UE_JSON_OBJECT_REF_PREFIX "uobject:"
diff --git a/Engine/Source/Runtime/Experimental/JsonObjectGraph/Private/JsonStringifyImpl.cpp b/Engine/Source/Runtime/Experimental/JsonObjectGraph/Private/JsonStringifyImpl.cpp
--- a/Engine/Source/Runtime/Experimental/JsonObjectGraph/Private/JsonStringifyImpl.cpp
+++ b/Engine/Source/Runtime/Experimental/JsonObjectGraph/Private/JsonStringifyImpl.cpp
@@ -14,6 +14,7 @@
 #include "Serialization/Formatters/JsonArchiveOutputFormatter.h"
 #include "Serialization/ArchiveUObjectFromStructuredArchive.h"
 #include "StructUtils/InstancedStruct.h"
+#include "StructUtils/PropertyBag.h"
 #include "UObject/EnumProperty.h"
 #include "UObject/PropertyOptional.h"
 #include "UObject/Package.h"
@@ -681,14 +682,16 @@ void FJsonStringifyImpl::WriteStructToJson(const void* StructInstance, const voi
 	const UScriptStruct* DefaultInstancedStructType = DefaultInstanceTyped ? 
 		DefaultInstanceTyped->GetScriptStruct() : nullptr;
 
-	// for structs with an ambiguous type (like instanced structs), we want to disambiguate
-	// the type stored in the json:
+	// For structs with an ambiguous type (like instanced structs), we want to disambiguate
+	// the type stored in the json, additionally for UPropertyBag, which is a transient 
+	// non native struct and is regenerated on load we want to store its current layout.
+	// It's possible we could allow structures to customize their representation, but 
+	// there are enough serialization customization facilities already and I do not want
+	// to introduce more.
+	const UPropertyBag* const PropertyBag = Cast<UPropertyBag>(InstancedStructType);
 	const bool bWriteStructType = InstancedStructType && 
-		// UPropertyBag is a transient non native struct, and is regenerated on load. It 
-		// is not meant to have persisent identity - so we don't want to even try
-		// to write it out:
-		(!InstancedStructType->HasAnyFlags(RF_Transient) || InstancedStructType->IsNative());
-		
+		!PropertyBag;
+
 	TUniquePtr<FInstancedStruct> FallbackStruct;
 	if(InstancedStructType)
 	{
@@ -710,12 +713,16 @@ void FJsonStringifyImpl::WriteStructToJson(const void* StructInstance, const voi
 
 	FPendingScope ObjectStartIdentifier( 
 		this, 
-		[Impl = this, Struct, bWriteStructType]() 
+		[Impl = this, InstancedStructType, bWriteStructType, PropertyBag]() 
 		{ 
 			Impl->Writer->WriteObjectStartInline(); 
-			if(bWriteStructType)
+			if(PropertyBag)
+			{
+				Impl->WritePropertyBagDescToJson(PropertyBag);
+			}
+			else if(bWriteStructType)
 			{
-				Impl->Writer->WriteUtf8Value(UE_JSON_SCRIPTSTRUCT_TCHAR, Impl->WriteObjectReference(Struct));
+				Impl->Writer->WriteUtf8Value(UE_JSON_SCRIPTSTRUCT_TCHAR, Impl->WriteObjectReference(InstancedStructType));
 			}
 		}, 
 		[Impl = this]() { Impl->Writer->WriteObjectEnd(); }
@@ -847,6 +854,21 @@ void FJsonStringifyImpl::WriteOptionalToJson(const void* OptionalInstnace, const
 	Writer->WriteObjectEnd();
 }
 
+void FJsonStringifyImpl::WritePropertyBagDescToJson(const UPropertyBag* PropertyBag)
+{
+	const FArrayProperty* PropertyDescsProperty = CastFieldChecked<FArrayProperty>(
+		PropertyBag->GetClass()->FindPropertyByName(GET_MEMBER_NAME_CHECKED(UPropertyBag, PropertyDescs)));
+
+	Writer->WriteIdentifierPrefix(UE_JSON_PROPERTYBAG_TCHAR);
+	Writer->WriteObjectStartInline();
+	Writer->WriteIdentifierPrefix(UE_JSON_PROPERTYDESCS_TCHAR);
+	WriteArrayToJson(
+		PropertyDescsProperty->ContainerPtrToValuePtr<void>(PropertyBag),
+		PropertyDescsProperty
+	);
+	Writer->WriteObjectEnd();
+}
+
 TArray<uint8> FJsonStringifyImpl::SerialDataToJson(const UObject* Object, int32 InitialIndentLevel)
 {
 	FJsonStringifyArchive SerialData(Object, InitialIndentLevel, this, Versions, MemoryWriter.IsFilterEditorOnly());
diff --git a/Engine/Source/Runtime/Experimental/JsonObjectGraph/Private/JsonStringifyImpl.h b/Engine/Source/Runtime/Experimental/JsonObjectGraph/Private/JsonStringifyImpl.h
--- a/Engine/Source/Runtime/Experimental/JsonObjectGraph/Private/JsonStringifyImpl.h
+++ b/Engine/Source/Runtime/Experimental/JsonObjectGraph/Private/JsonStringifyImpl.h
@@ -12,6 +12,7 @@ class FArrayProperty;
 class FSetProperty;
 class FMapProperty;
 class FOptionalProperty;
+class UPropertyBag;
 
 namespace UE::Private 
 {
@@ -61,6 +62,7 @@ private:
 	void WriteSetToJson(const void* SetInstance, const FSetProperty* SetProperty);
 	void WriteMapToJson(const void* MapInstance, const FMapProperty* MapProperty);
 	void WriteOptionalToJson(const void* OptionalInstnace, const FOptionalProperty* OptionalProperty);
+	void WritePropertyBagDescToJson(const UPropertyBag* PropertyBag);
 
 	TArray<uint8> SerialDataToJson(const UObject* Object, int32 InitialIndentLevel);
 #if WITH_TEXT_ARCHIVE_SUPPORT
