FOnDemandInstallCache::AddReferencesToBlocks - Anything in ChunksToInstall should be considered referenced to be safe.

diff --git a/Engine/Source/Runtime/Experimental/IoStore/OnDemand/Private/OnDemandInstallCache.cpp b/Engine/Source/Runtime/Experimental/IoStore/OnDemand/Private/OnDemandInstallCache.cpp
--- a/Engine/Source/Runtime/Experimental/IoStore/OnDemand/Private/OnDemandInstallCache.cpp
+++ b/Engine/Source/Runtime/Experimental/IoStore/OnDemand/Private/OnDemandInstallCache.cpp
@@ -2069,28 +2069,34 @@ uint64 FOnDemandInstallCache::AddReferencesToBlocks(
 		const TBitArray<>& IsReferenced = ChunkEntryIndices[Index++];
 		for (int32 EntryIndex = 0; const FOnDemandChunkEntry& Entry : Container->ChunkEntries)
 		{
-			if (IsReferenced[EntryIndex])
+			const bool bToInstall = ChunksToInstall.Contains(Entry.Hash);
+			const bool bIsReferenced = IsReferenced[EntryIndex];
+
+			FCasBlockInfo* BlockInfo = nullptr;
+			if (bToInstall || bIsReferenced)
 			{
 				OutTotalReferencedBytes += Entry.EncodedSize;
 				if (FCasLocation Loc = Cas.FindChunk(Entry.Hash); Loc.IsValid())
 				{
-					if (FCasBlockInfo* BlockInfo = BlockInfoMap.Find(Loc.BlockId))
-					{
-						BlockInfo->RefSize += Entry.EncodedSize;
-					}
-					else
+					BlockInfo = BlockInfoMap.Find(Loc.BlockId);
+					if (!BlockInfo)
 					{
-						UE_LOG(LogIoStoreOnDemand, Error, TEXT("Failed to find CAS block info for referenced chunk, ChunkId='%s', Container='%s'"),
+						UE_CLOG(bIsReferenced, LogIoStoreOnDemand, Error, TEXT("Failed to find CAS block info for referenced chunk, ChunkId='%s', Container='%s'"),
 							*LexToString(Container->ChunkIds[EntryIndex]), *Container->Name);
 					}
 				}
 				else
 				{
-					UE_LOG(LogIoStoreOnDemand, Error, TEXT("Failed to find CAS location for referenced chunk, ChunkId='%s', Container='%s'"),
+					UE_CLOG(bIsReferenced, LogIoStoreOnDemand, Error, TEXT("Failed to find CAS location for referenced chunk, ChunkId='%s', Container='%s'"),
 						*LexToString(Container->ChunkIds[EntryIndex]), *Container->Name);
 				}
 			}
-			else if (ChunksToInstall.Contains(Entry.Hash))
+
+			if (BlockInfo)
+			{
+				BlockInfo->RefSize += Entry.EncodedSize;
+			}
+			else if (bToInstall)
 			{
 				TotalUncachedBytes += Entry.EncodedSize;
 			}
