Chaos - fix potential crash in FCCDManager::ApplyCorrections when awakened constaints are added to the active list (required for correct collision enter/exit events).
- disable the previous workaround, but leave it as an option for now (cvar)

FCCDManager was holding a TArrayView on the active CCD constraints, but that array might relocate during island waking.

#rnx

diff --git a/Engine/Source/Runtime/Experimental/Chaos/Private/Chaos/CCDUtilities.cpp b/Engine/Source/Runtime/Experimental/Chaos/Private/Chaos/CCDUtilities.cpp
--- a/Engine/Source/Runtime/Experimental/Chaos/Private/Chaos/CCDUtilities.cpp
+++ b/Engine/Source/Runtime/Experimental/Chaos/Private/Chaos/CCDUtilities.cpp
@@ -97,20 +97,19 @@ namespace Chaos
 		return INDEX_NONE;
 	}
 
-	void FCCDManager::ApplyConstraintsPhaseCCD(const FReal Dt, Private::FCollisionConstraintAllocator *CollisionAllocator, const int32 NumDynamicParticles)
+	void FCCDManager::ApplyCCD(const FReal Dt, TArrayView<FPBDCollisionConstraint* const> SweptConstraints, const int32 NumDynamicParticles)
 	{
-		SweptConstraints = CollisionAllocator->GetCCDConstraints();
 		if (SweptConstraints.Num() > 0)
 		{
 			ApplySweptConstraints(Dt, SweptConstraints, NumDynamicParticles);
-			UpdateSweptConstraints(Dt, CollisionAllocator);
+			UpdateSweptConstraints(Dt);
 			OverwriteXUsingV(Dt);
 		}
 	}
 
-	void FCCDManager::ApplySweptConstraints(const FReal Dt, TArrayView<FPBDCollisionConstraint* const> InSweptConstraints, const int32 NumDynamicParticles)
+	void FCCDManager::ApplySweptConstraints(const FReal Dt, TArrayView<FPBDCollisionConstraint* const> SweptConstraints, const int32 NumDynamicParticles)
 	{
-		const bool bNeedCCDSolve = Init(Dt, NumDynamicParticles);
+		const bool bNeedCCDSolve = Init(Dt, SweptConstraints, NumDynamicParticles);
 		if (!bNeedCCDSolve)
 		{
 			return;
@@ -136,10 +135,13 @@ namespace Chaos
 		});
 	}
 
-	bool FCCDManager::Init(const FReal Dt, const int32 NumDynamicParticles)
+	bool FCCDManager::Init(const FReal Dt, TArrayView<FPBDCollisionConstraint* const> SweptConstraints, const int32 NumDynamicParticles)
 	{
+		// We store pointers to CCDParticle in CCDConstraint and GroupedCCDParticles, therefore we need to make sure to reserve enough space 
+		// for TArray CCDParticles so that reallocation does not happen in the for loop. If reallocation happens, some pointers may become 
+		// invalid and this could cause serious bugs. We know that the number of CCDParticles cannot exceed SweptConstraints.Num() * 2 or NumDynamicParticles. 
+		// Therefore this makes sure reallocation does not happen.
 		CCDParticles.Reset();
-		// We store pointers to CCDParticle in CCDConstraint and GroupedCCDParticles, therefore we need to make sure to reserve enough space for TArray CCDParticles so that reallocation does not happen in the for loop. If reallocation happens, some pointers may become invalid and this could cause serious bugs. We know that the number of CCDParticles cannot exceed SweptConstraints.Num() * 2 or NumDynamicParticles. Therefore this makes sure reallocation does not happen.
 		CCDParticles.Reserve(FMath::Min(SweptConstraints.Num() * 2, NumDynamicParticles));
 		ParticleToCCDParticle.Reset();
 		CCDConstraints.Reset();
@@ -167,7 +169,8 @@ namespace Chaos
 				continue;
 			}
 
-			// Create CCDParticle for all dynamic particles affected by swept constraints (UseCCD() could be either true or false). For static or kinematic particles, this pointer remains to be nullptr.
+			// Create CCDParticle for all dynamic particles affected by swept constraints (UseCCD() could be either true or false). 
+			// For static or kinematic particles, this pointer remains to be nullptr.
 			FCCDParticle* CCDParticlePair[2] = {nullptr, nullptr};
 			bool IsDynamic[2] = {false, false};
 			FVec3 Displacements[2] = {FVec3(0.f), FVec3(0.f)};
@@ -181,6 +184,9 @@ namespace Chaos
 					FCCDParticle** FoundCCDParticle = ParticleToCCDParticle.Find(RigidParticle);
 					if (!FoundCCDParticle)
 					{
+						// See comments above - we cannot allow CCDParticles to resize
+						// @todo(chaos): clean this up - use indices, chunked array, or something
+						check(CCDParticles.Num() < CCDParticles.Max());
 						CCDParticles.Add(FCCDParticle(RigidParticle));
 						CCDParticle = &CCDParticles.Last();
 						ParticleToCCDParticle.Add(RigidParticle, CCDParticle);
@@ -196,7 +202,8 @@ namespace Chaos
 
 				if (RigidParticle)
 				{
-					// One can also use P - X for dynamic particles. But notice that for kinematic particles, both P and X are end-frame positions and P - X won't work for kinematic particles.
+					// One can also use P - X for dynamic particles. But notice that for kinematic particles, 
+					// both P and X are end-frame positions and P - X won't work for kinematic particles.
 					Displacements[i] = RigidParticle->GetV() * Dt;
 				}
 			}
@@ -872,7 +879,7 @@ namespace Chaos
 		}
 	}
 
-	void FCCDManager::UpdateSweptConstraints(const FReal Dt, Private::FCollisionConstraintAllocator *CollisionAllocator)
+	void FCCDManager::UpdateSweptConstraints(const FReal Dt)
 	{
 		// Build the set of collision whose contact data will be out of date because we moved one or both of its particles. 
 		// This is all collision constraints, including non-swept ones, for any particle that was relocated by the CCD sweep 
@@ -908,20 +915,17 @@ namespace Chaos
 		}
 	}
 
-	void FCCDManager::ApplyCorrections(const FReal Dt)
+	void FCCDManager::ProjectCCD(const FReal Dt, TArrayView<FPBDCollisionConstraint* const> SweptConstraints)
 	{
-		// Build a list of CCD constraints sorted so that the most important ones are last.
-		// The primary goal is to prevent CCD objects from leaving the world, so we process collisions
-		// with statics/kinematics last. The secondary goal is to prevent CCD objects passing through each 
-		// other. And finally we try to prevent CCD object passing through non-CCD dynamics but aren't
-		// too worried if that happens.
-		// Sort order: 
-		//		CCD - Dynamic
-		//		CCD - CCD
-		//		CCD - Kinematic
-		TArray<FPBDCollisionConstraint*> Constraints;
-		Constraints.Append(SweptConstraints);
-
+		// NOTE: The SweptConstraints list may be different that the one used in ApplyCCD
+		// earlier in the frame but that's ok - we still want to enforce depenetration.
+		// - some may be removed (or nulled) if particles were destroyed
+		// - some may be added if an island was awakened
+		TArray<FPBDCollisionConstraint*> Constraints = TArray<FPBDCollisionConstraint*>(SweptConstraints);
+
+		// Remove null constraints in the active list. Also check for null pointers (which is probably not required 
+		// any more - we used to just null the particle pointers in collisions and let them auto-clean up later
+		// but now we destroy them immediately).
 		if(Chaos::CVars::bClearNullCCDConstraintsOnApply)
 		{
 			Constraints.RemoveAll(
@@ -934,6 +938,16 @@ namespace Chaos
 			);
 		}
 
+		// Build a list of CCD constraints sorted so that the most important ones are last.
+		// The primary goal is to prevent CCD objects from leaving the world, so we process collisions
+		// with statics/kinematics last. The secondary goal is to prevent CCD objects passing through each 
+		// other. And finally we try to prevent CCD object passing through non-CCD dynamics but aren't
+		// too worried if that happens.
+		//
+		// Sort order based on the two particle types in the collision:
+		//		CCD - Dynamic
+		//		CCD - CCD
+		//		CCD - Kinematic
 		Constraints.Sort([](const FPBDCollisionConstraint& L, const FPBDCollisionConstraint& R)
 			{
 				const FGenericParticleHandle LP0 = L.GetParticle0();
diff --git a/Engine/Source/Runtime/Experimental/Chaos/Private/Chaos/Collision/ParticlePairMidPhase.cpp b/Engine/Source/Runtime/Experimental/Chaos/Private/Chaos/Collision/ParticlePairMidPhase.cpp
--- a/Engine/Source/Runtime/Experimental/Chaos/Private/Chaos/Collision/ParticlePairMidPhase.cpp
+++ b/Engine/Source/Runtime/Experimental/Chaos/Private/Chaos/Collision/ParticlePairMidPhase.cpp
@@ -65,7 +65,7 @@ namespace Chaos
 
 		// This should be enabled but it exposes a bug that must be fixed first: collisions with 
 		// invalid particle pointers are being activated causing a crash in CCDUtilities::ApplyCorrections
-		bool bChaosMidPhaseActivateWakingConstraints = false;
+		bool bChaosMidPhaseActivateWakingConstraints = true;
 		FAutoConsoleVariableRef CVarChaosMidPhaseActivateWakingConstraints(TEXT("p.Chaos.Bugs.ActivateWakingConstraints"), bChaosMidPhaseActivateWakingConstraints, TEXT(""));
 
 		// See FGenericParticlePairMidPhase::SetConvexOptimizationIsActive
diff --git a/Engine/Source/Runtime/Experimental/Chaos/Private/Chaos/PBDRigidsEvolutionGBF.cpp b/Engine/Source/Runtime/Experimental/Chaos/Private/Chaos/PBDRigidsEvolutionGBF.cpp
--- a/Engine/Source/Runtime/Experimental/Chaos/Private/Chaos/PBDRigidsEvolutionGBF.cpp
+++ b/Engine/Source/Runtime/Experimental/Chaos/Private/Chaos/PBDRigidsEvolutionGBF.cpp
@@ -621,7 +621,7 @@ void FPBDRigidsEvolutionGBF::AdvanceOneTimeStepImpl(const FReal Dt, const FSubSt
 	{
 		SCOPE_CYCLE_COUNTER(STAT_Evolution_CCD);
 		CSV_SCOPED_TIMING_STAT(PhysicsVerbose, CCD);
-		CCDManager.ApplyConstraintsPhaseCCD(Dt, &CollisionConstraints.GetConstraintAllocator(), Particles.GetActiveParticlesView().Num());
+		CCDManager.ApplyCCD(Dt, CollisionConstraints.GetConstraintAllocator().GetCCDConstraints(), Particles.GetActiveParticlesView().Num());
 	}
 
 	if (CollisionModifiers && !CVars::bChaosCollisionModiferBeforeCCD)
@@ -696,7 +696,7 @@ void FPBDRigidsEvolutionGBF::AdvanceOneTimeStepImpl(const FReal Dt, const FSubSt
 		{
 			SCOPE_CYCLE_COUNTER(STAT_Evolution_CCDCorrection);
 			CSV_SCOPED_TIMING_STAT(PhysicsVerbose, StepSolver_CCDCorrection);
-			CCDManager.ApplyCorrections(Dt);
+			CCDManager.ProjectCCD(Dt, CollisionConstraints.GetConstraintAllocator().GetCCDConstraints());
 		}
 	}
 
diff --git a/Engine/Source/Runtime/Experimental/Chaos/Public/Chaos/CCDUtilities.h b/Engine/Source/Runtime/Experimental/Chaos/Public/Chaos/CCDUtilities.h
--- a/Engine/Source/Runtime/Experimental/Chaos/Public/Chaos/CCDUtilities.h
+++ b/Engine/Source/Runtime/Experimental/Chaos/Public/Chaos/CCDUtilities.h
@@ -84,18 +84,28 @@ namespace Chaos
 	{
 	public:
 		FCCDManager(){}
-		void ApplyConstraintsPhaseCCD(const FReal Dt, Private::FCollisionConstraintAllocator *CollisionAllocator, const int32 NumDynamicParticles = TNumericLimits<int32>::Max());
 
-		// A post process on CCD contacts to ensure that CCD objects never pass through non-dynamic objects
-		void ApplyCorrections(const FReal Dt);
+		// Apply the CCD solver to the set of constraints. Currently this just moves the particles so that they overlap by
+		// an amount that can be resolved by the standard collision solver.
+		void ApplyCCD(const FReal Dt, TArrayView<FPBDCollisionConstraint* const> SweptConstraints, const int32 NumDynamicParticles);
+
+		// Fix up any remaining penetration following the main constraint solver phase. E.g., this prevents heavy objects
+		// pushing light CCD objects through the world.
+		void ProjectCCD(const FReal Dt, TArrayView<FPBDCollisionConstraint* const> SweptConstraints);
+
+		UE_DEPRECATED(5.6, "No longer functioning. Use ApplyCCD")
+		void ApplyConstraintsPhaseCCD(const FReal Dt, Private::FCollisionConstraintAllocator* CollisionAllocator, const int32 NumDynamicParticles = TNumericLimits<int32>::Max()) {}
+
+		UE_DEPRECATED(5.6, "No longer functioning. Use ProjectCCD")
+		void ApplyCorrections(const FReal Dt) {}
 
 	private:
-		void ApplySweptConstraints(const FReal Dt, TArrayView<FPBDCollisionConstraint* const> InSweptConstraints, const int32 NumDynamicParticles = TNumericLimits<int32>::Max());
+		void ApplySweptConstraints(const FReal Dt, TArrayView<FPBDCollisionConstraint* const> SweptConstraints, const int32 NumDynamicParticles = TNumericLimits<int32>::Max());
 		bool UpdateParticleSweptConstraints(FCCDParticle* CCDParticle, const FReal IslandTOI, const FReal Dt);
 		// This is called after ApplySweptConstraints. This function updates manifold data which will be used in normal solve.
-		void UpdateSweptConstraints(const FReal Dt, Private::FCollisionConstraintAllocator *CollisionAllocator);
+		void UpdateSweptConstraints(const FReal Dt);
 		void OverwriteXUsingV(const FReal Dt);
-		bool Init(const FReal Dt, const int32 NumDynamicParticles = TNumericLimits<int32>::Max());
+		bool Init(const FReal Dt, TArrayView<FPBDCollisionConstraint* const> SweptConstraints, const int32 NumDynamicParticles = TNumericLimits<int32>::Max());
 		void AssignParticleIslandsAndGroupParticles();
 		void AssignConstraintIslandsAndRecordConstraintNum();
 		void GroupConstraintsWithIslands();
@@ -109,10 +119,9 @@ namespace Chaos
 		void ClipParticleP(FCCDParticle *CCDParticle, const FVec3 Offset) const;
 		void ApplyImpulse(FCCDConstraint *CCDConstraint);
 
-		TArrayView<FPBDCollisionConstraint* const> SweptConstraints;
 		TArray<FCCDParticle> CCDParticles;
 		TArray<FCCDParticle*> GroupedCCDParticles; // FCCDParticles are grouped based on their islands.
-		TMap<TPBDRigidParticleHandle<FReal, 3>*, FCCDParticle*> ParticleToCCDParticle; // This map only contains dynamic particles.
+		TMap<FPBDRigidParticleHandle*, FCCDParticle*> ParticleToCCDParticle; // This map only contains dynamic particles.
 		TArray<FCCDConstraint> CCDConstraints;
 		TArray<FCCDConstraint*> SortedCCDConstraints; // Constraints are first grouped into islands and then sorted based on TOI.
 
