Expose more functions for indexed cache storage management

diff --git a/Engine/Source/Runtime/Experimental/IndexedCacheStorage/IndexedCacheStorage.Build.cs b/Engine/Source/Runtime/Experimental/IndexedCacheStorage/IndexedCacheStorage.Build.cs
new file mode 100644
--- /dev/null
+++ b/Engine/Source/Runtime/Experimental/IndexedCacheStorage/IndexedCacheStorage.Build.cs
@@ -0,0 +1,14 @@
+// Copyright Epic Games, Inc. All Rights Reserved.
+
+using UnrealBuildTool;
+
+public class IndexedCacheStorage : ModuleRules
+{
+	public IndexedCacheStorage(ReadOnlyTargetRules Target) : base(Target)
+	{
+		PublicDependencyModuleNames.Add("Core");
+
+		bAllowConfidentialPlatformDefines = true;
+		CppCompileWarningSettings.UnsafeTypeCastWarningLevel = WarningLevel.Error; 
+	}
+}
diff --git a/Engine/Source/Runtime/Experimental/IndexedCacheStorage/Private/IndexedCacheStorage.cpp b/Engine/Source/Runtime/Experimental/IndexedCacheStorage/Private/IndexedCacheStorage.cpp
new file mode 100644
--- /dev/null
+++ b/Engine/Source/Runtime/Experimental/IndexedCacheStorage/Private/IndexedCacheStorage.cpp
@@ -0,0 +1,75 @@
+// Copyright Epic Games, Inc. All Rights Reserved.
+#include "IndexedCacheStorage.h"
+
+namespace Experimental
+{
+
+/**
+ * Generic implementation of IIndexedCacheStorage
+ */
+class FGenericIndexedCacheStorage : public Experimental::IIndexedCacheStorage
+{
+public:
+	virtual bool SupportsIndexedCacheStorage() override
+	{
+		return false;
+	}
+
+	virtual uint32 GetNumIndexedCacheStorages() override
+	{
+		return 0;
+	}
+
+	virtual bool CreateCacheStorage(uint64 InRequestNumberOfBytes, int32& OutCacheStorageIndex) override
+	{
+		return false;
+	}
+
+	virtual bool DestroyCacheStorage(int32 InCacheStorageIndex) override
+	{
+		return false;
+	}
+
+	virtual void EnumerateCacheStorages(TArray<int32>& OutCacheStorageIndices) override
+	{
+	}
+	
+	virtual void GetCacheStorageInfo(int32 InCacheStorageIndex, uint64& OutCacheStorageSize) override
+	{
+	}
+	
+	virtual void GetCacheStorageMountPath(FString& OutPlatformMountPath, const FString& InMountName) override
+	{
+	}
+
+	virtual bool MountCacheStorage(FString& OutPlatformMountPath, int32 InCacheStorageIndex, const FString& InMountName) override
+	{
+		return false;
+	}
+
+	virtual void UnmountCacheStorage(const FString& InMountName) override
+	{
+	}
+
+};
+
+#if !defined(HAS_CUSTOM_INDEXED_CACHE_STORAGE)
+#	define HAS_CUSTOM_INDEXED_CACHE_STORAGE 0
+#endif
+
+#if (HAS_CUSTOM_INDEXED_CACHE_STORAGE == 0)
+IIndexedCacheStorage& IIndexedCacheStorageModule::GetOrCreateIndexedCacheStorage()
+{
+	static FGenericIndexedCacheStorage Instance;
+	return Instance;
+}
+#endif
+
+IIndexedCacheStorage& IIndexedCacheStorage::Get()
+{
+	return IIndexedCacheStorageModule::Get().GetOrCreateIndexedCacheStorage();
+}
+
+
+}
+IMPLEMENT_MODULE(Experimental::IIndexedCacheStorageModule, IndexedCacheStorage);
diff --git a/Engine/Source/Runtime/Experimental/IndexedCacheStorage/Public/IndexedCacheStorage.h b/Engine/Source/Runtime/Experimental/IndexedCacheStorage/Public/IndexedCacheStorage.h
new file mode 100644
--- /dev/null
+++ b/Engine/Source/Runtime/Experimental/IndexedCacheStorage/Public/IndexedCacheStorage.h
@@ -0,0 +1,114 @@
+// Copyright Epic Games, Inc. All Rights Reserved.
+
+#pragma once
+
+#include "CoreTypes.h"
+#include "Containers/Array.h"
+
+#include "Modules/ModuleInterface.h"
+#include "Modules/ModuleManager.h"
+
+#define UE_API INDEXEDCACHESTORAGE_API
+
+
+namespace Experimental
+{
+
+class IIndexedCacheStorage;
+
+class IIndexedCacheStorageModule : public IModuleInterface
+{
+public:
+	static IIndexedCacheStorageModule& Get() 
+	{
+		return FModuleManager::Get().LoadModuleChecked<IIndexedCacheStorageModule>("IndexedCacheStorage");
+	}
+
+	IIndexedCacheStorage& GetOrCreateIndexedCacheStorage();
+
+};
+
+
+/**
+  * Interface for managing cache storages: some platforms may store user data in different cache storages that may be created / destroyed
+  * This allows the title to manage its disk space with a finer granularity. 
+  * The data put in these caches need to be re-generated/re-obtained if needed.
+  */
+class IIndexedCacheStorage
+{
+public:
+	UE_API static IIndexedCacheStorage& Get();
+
+	/**
+	  * Check if indexed cache storage is available on the current platform.
+	  * 
+	  * @return True is available, false otherwise.
+	  */
+	UE_API virtual bool SupportsIndexedCacheStorage() = 0;
+
+	/**
+	  * Retrieve the number of indexed cache entries available on the platform.
+	  * 
+	  * @return number of indexed cache entries available on the platform.
+	  */
+	UE_API virtual uint32 GetNumIndexedCacheStorages() = 0;
+
+	/**
+	  * Requests a new indexed cache entry, capable of holding up to RequestNumberOfBytes.
+	  * 
+ 	  * @param InRequestNumberOfBytes number of bytes this entry needs to hold
+ 	  * @param OutCacheStorageIndex - the index of the new entry
+	  * @return True on success, otherwise false
+	  */
+	UE_API virtual bool CreateCacheStorage(uint64 InRequestNumberOfBytes, int32& OutCacheStorageIndex) = 0;
+
+	/**
+	  * Destroy an indexed cache entry.
+	  * 
+ 	  * @param InCacheStorageIndex - the index of the entry
+	  * @return True on success, otherwise false
+	  */
+	UE_API virtual bool DestroyCacheStorage(int32 InCacheStorageIndex) = 0;
+
+	/**
+	  * Enumerate existing cache entries.
+	  * 
+ 	  * @param OutCacheStorageIndices - indices of existing cache storages
+	  */
+	UE_API virtual void EnumerateCacheStorages(TArray<int32>& OutCacheStorageIndices) = 0;
+	
+	/**
+	  * Queries cache storage.
+	  * 
+ 	  * @param InCacheStorageIndex - the index of the cache
+ 	  * @param OutCacheStorageSize - How large the cache storage is
+	  */
+	UE_API virtual void GetCacheStorageInfo(int32 InCacheStorageIndex, uint64& OutCacheStorageSize) = 0;
+	
+	/**
+	  * Gets the platform mount path for a cached storage based on its generic Name
+	  * 
+ 	  * @param OutPlatformMountPath - the platform specific mount path
+ 	  * @param InMountName - the generic mount name
+	  */
+	UE_API virtual void GetCacheStorageMountPath(FString& OutPlatformMountPath, const FString& InMountName) = 0;
+
+	/**
+	  * Mounts a cache storage so file operations can be done
+	  * 
+ 	  * @param OutPlatformMountPath - the platform specific mount path. Same name as the one returned by GetCacheStorageMountPath
+ 	  * @param InCacheStorageIndex - the index of the cache
+ 	  * @param InMountName - the generic mount name
+	  */
+	UE_API virtual bool MountCacheStorage(FString& OutPlatformMountPath, int32 InCacheStorageIndex, const FString& InMountName) = 0;
+
+	/**
+	  * Unmounts a cache storage
+ 	  * @param InMountName - the generic mount name
+	  */
+	UE_API virtual void UnmountCacheStorage(const FString& InMountName) = 0;
+};
+
+}
+
+#undef UE_API
\ No newline at end of file
