Added cvar iax.ForceInsecureHttp to disable https endpoints when creating on-demand host groups.

#rnx

diff --git a/Engine/Source/Runtime/Experimental/IoStore/OnDemandCore/Private/OnDemandHostGroup.cpp b/Engine/Source/Runtime/Experimental/IoStore/OnDemandCore/Private/OnDemandHostGroup.cpp
--- a/Engine/Source/Runtime/Experimental/IoStore/OnDemandCore/Private/OnDemandHostGroup.cpp
+++ b/Engine/Source/Runtime/Experimental/IoStore/OnDemandCore/Private/OnDemandHostGroup.cpp
@@ -5,10 +5,25 @@
 #include "Containers/AnsiString.h"
 #include "Containers/Array.h"
 #include "Containers/StringView.h"
+#include "HAL/IConsoleManager.h"
 
 namespace UE::IoStore
 {
 
+////////////////////////////////////////////////////////////////////////////////
+namespace Config
+{
+
+bool bForceInsecureHttp = true;
+static FAutoConsoleVariableRef CVar_ForceInsecureHttp(
+	TEXT("iax.ForceInsecureHttp"),
+	bForceInsecureHttp,
+	TEXT("Whether to force the use of insecure HTTP"),
+	ECVF_ReadOnly
+);
+
+} // UE::IoStore::Config
+
 namespace Private
 {
 
@@ -97,7 +112,12 @@ TIoStatusOr<FOnDemandHostGroup> FOnDemandHostGroup::Create(FAnsiStringView Url)
 	}
 
 	FSharedImpl Impl = MakeShared<FImpl>();
-	Impl->HostUrls.Add(FAnsiString(Url));
+	FAnsiString AnsiUrl(Url);
+	if (Config::bForceInsecureHttp)
+	{
+		AnsiUrl.ReplaceInline("https", "http");
+	}
+	Impl->HostUrls.Add(MoveTemp(AnsiUrl));
 	Impl->PrimaryIndex = 0;
 
 	return FOnDemandHostGroup(MoveTemp(Impl)); 
@@ -131,7 +151,13 @@ TIoStatusOr<FOnDemandHostGroup> FOnDemandHostGroup::Create(TConstArrayView<FAnsi
 		{
 			return FIoStatus(EIoErrorCode::InvalidParameter, Reason);
 		}
-		Impl->HostUrls.Emplace(UrlView);
+
+		FAnsiString AnsiUrl(Url);
+		if (Config::bForceInsecureHttp)
+		{
+			AnsiUrl.ReplaceInline("https", "http");
+		}
+		Impl->HostUrls.Add(MoveTemp(AnsiUrl));
 	}
 	Impl->PrimaryIndex = 0; 
 
@@ -163,6 +189,11 @@ TIoStatusOr<FOnDemandHostGroup> FOnDemandHostGroup::Create(TConstArrayView<FStri
 		{
 			return FIoStatus(EIoErrorCode::InvalidParameter, Reason);
 		}
+
+		if (Config::bForceInsecureHttp)
+		{
+			AnsiUrl.ReplaceInline("https", "http");
+		}
 		Impl->HostUrls.Emplace(MoveTemp(AnsiUrl));
 	}
 	Impl->PrimaryIndex = 0;
