Remove useless ensure

#rnx

diff --git a/Engine/Source/Runtime/Experimental/Animation/Constraints/Private/Transform/TransformConstraintUtil.cpp b/Engine/Source/Runtime/Experimental/Animation/Constraints/Private/Transform/TransformConstraintUtil.cpp
--- a/Engine/Source/Runtime/Experimental/Animation/Constraints/Private/Transform/TransformConstraintUtil.cpp
+++ b/Engine/Source/Runtime/Experimental/Animation/Constraints/Private/Transform/TransformConstraintUtil.cpp
@@ -85,7 +85,7 @@ TOptional<FTransform> FConstraintsInteractionCache::GetParentTransform(const UOb
 const TArray< TWeakObjectPtr<UTickableConstraint> >& FConstraintsInteractionCache::Get(const uint32 InHandleHash, UWorld* InWorld)
 {
 	static const TArray< TWeakObjectPtr<UTickableConstraint> > Empty;
-	if (!ensure(InHandleHash != 0 && InWorld))
+	if (!InWorld || InHandleHash == 0)
 	{
 		return Empty;
 	}
