[ModularVehicle] Add an alternative way to quantize inputs

Add an enum EModuleInputQuantizationType on the ModularVehicleBaseComponent  to choose the quantization type:
- Default_16Bits: the current input quantization
- Custom_8Bits: quantize and serialize on uint8. Made to replicate a legacy vehicle movement

A CVar EnableCustomInputQuantization is added to enable the custom replication mode.

diff --git a/Engine/Plugins/Experimental/ChaosModularVehicle/Source/ChaosModularVehicleEngine/Private/ChaosModularVehicle/ChaosSimModuleManagerAsyncCallback.cpp b/Engine/Plugins/Experimental/ChaosModularVehicle/Source/ChaosModularVehicleEngine/Private/ChaosModularVehicle/ChaosSimModuleManagerAsyncCallback.cpp
--- a/Engine/Plugins/Experimental/ChaosModularVehicle/Source/ChaosModularVehicleEngine/Private/ChaosModularVehicle/ChaosSimModuleManagerAsyncCallback.cpp
+++ b/Engine/Plugins/Experimental/ChaosModularVehicle/Source/ChaosModularVehicleEngine/Private/ChaosModularVehicle/ChaosSimModuleManagerAsyncCallback.cpp
@@ -6,6 +6,7 @@
 #include "ChaosModularVehicle/ModularVehicleSimulationCU.h"
 #include "ChaosModularVehicle/ModuleInputTokenStore.h"
 #include "PBDRigidsSolver.h"
+#include "ChaosModularVehicle/ModularVehicleBaseComponent.h"
 #include "ChaosModularVehicle/ModuleInputTokenStore.h"
 #include "GeometryCollection/GeometryCollectionParticlesData.h"
 #include "PhysicsProxy/GeometryCollectionPhysicsProxy.h"
@@ -260,6 +261,12 @@ void FModularVehicleAsyncInput::ProcessInputs()
 
 bool FNetworkModularVehicleInputs::NetSerialize(FArchive& Ar, class UPackageMap* Map, bool& bOutSuccess)
 {
+	EModuleInputQuantizationType InputQuantizationType = EModuleInputQuantizationType::Default_16Bits;
+	if (const UModularVehicleBaseComponent* ModularVehicleBaseComponent = Cast<UModularVehicleBaseComponent>(ImplementationComponent.Get()))
+	{
+		InputQuantizationType = ModularVehicleBaseComponent->InputQuantizationType;	
+	}
+
 	FNetworkPhysicsData::SerializeFrames(Ar);
 
 	Ar.SerializeBits(&VehicleInputs.Reverse, 1);
@@ -295,20 +302,20 @@ bool FNetworkModularVehicleInputs::NetSerialize(FArchive& Ar, class UPackageMap*
 			InputValues[I].SetApplyInputDecay(InputStateData.DecayValues[I]);
 			if (PreviousInputValues.Num() == InputValues.Num())
 			{
-				InputValues[I].DeltaNetSerialize(Ar, Map, bOutSuccess, PreviousInputValues[I]);
+				InputValues[I].DeltaNetSerialize(Ar, Map, bOutSuccess, PreviousInputValues[I], InputQuantizationType);
 			}
 			else
 			{
 				bOutSuccess = false;
 				//Fail case.
-				InputValues[I].DeltaNetSerialize(Ar, Map, bOutSuccess, InputValues[I]);
+				InputValues[I].DeltaNetSerialize(Ar, Map, bOutSuccess, InputValues[I], InputQuantizationType);
 			}
 		
 		}
 	}
 	else
 	{
-		VehicleInputs.Container.Serialize(Ar, Map, bOutSuccess);
+		VehicleInputs.Container.Serialize(Ar, Map, bOutSuccess, InputQuantizationType);
 	}
 	
 	return bOutSuccess;
diff --git a/Engine/Plugins/Experimental/ChaosModularVehicle/Source/ChaosModularVehicleEngine/Private/ChaosModularVehicle/InputProducer.cpp b/Engine/Plugins/Experimental/ChaosModularVehicle/Source/ChaosModularVehicleEngine/Private/ChaosModularVehicle/InputProducer.cpp
--- a/Engine/Plugins/Experimental/ChaosModularVehicle/Source/ChaosModularVehicleEngine/Private/ChaosModularVehicle/InputProducer.cpp
+++ b/Engine/Plugins/Experimental/ChaosModularVehicle/Source/ChaosModularVehicleEngine/Private/ChaosModularVehicle/InputProducer.cpp
@@ -4,8 +4,10 @@
 
 // Default input producer
 
-void UVehicleDefaultInputProducer::InitializeContainer(TArray<FModuleInputSetup>& SetupData, FInputNameMap& NameMapOut)
+void UVehicleDefaultInputProducer::InitializeContainer(TArray<FModuleInputSetup>& SetupData, FInputNameMap& NameMapOut, EModuleInputQuantizationType InInputQuantizationType)
 {
+	Super::InitializeContainer(SetupData, NameMapOut, InInputQuantizationType);
+
 	MergedInput.Initialize(SetupData, NameMapOut);
 }
 
@@ -14,7 +16,7 @@ void UVehicleDefaultInputProducer::BufferInput(const FInputNameMap& InNameMap, c
 	//UE_LOG(LogTemp, Warning, TEXT("BufferInput - InName %s Val=%s"), *InName.ToString(), *InValue.ToString());
 
 	// inputs are merged here rather than buffered, since they would be merged anyway before use in ProduceInput
-	FInputInterface Inputs(InNameMap, MergedInput);
+	FInputInterface Inputs(InNameMap, MergedInput, InputQuantizationType);
 	Inputs.MergeValue(InName, InValue);
 }
 
@@ -34,8 +36,10 @@ void UVehicleDefaultInputProducer::ProduceInput(int32 PhysicsStep, int32 NumStep
 
 // Example Playback Input Producer
 
-void UVehiclePlaybackInputProducer::InitializeContainer(TArray<FModuleInputSetup>& SetupData, FInputNameMap& NameMapOut)
+void UVehiclePlaybackInputProducer::InitializeContainer(TArray<FModuleInputSetup>& SetupData, FInputNameMap& NameMapOut, EModuleInputQuantizationType InInputQuantizationType)
 {
+	Super::InitializeContainer(SetupData, NameMapOut, InInputQuantizationType);
+
 	int32 Seed = 123;
 	FRandomStream Random(Seed);
 	StartStep = 0;
@@ -51,7 +55,7 @@ void UVehiclePlaybackInputProducer::InitializeContainer(TArray<FModuleInputSetup
 		PlaybackBuffer.Emplace(InputsForFrame);
 
 		// change value of some intputs
-		FInputInterface Inputs(NameMapOut, PlaybackBuffer[I]);
+		FInputInterface Inputs(NameMapOut, PlaybackBuffer[I], InputQuantizationType);
 		Inputs.SetValue("Throttle", Random.FRand());
 		Inputs.SetValue("Steering", 1.0f - 2.0f * Random.FRand());
 	}
@@ -87,8 +91,10 @@ void UVehiclePlaybackInputProducer::ProduceInput(int32 PhysicsStep, int32 NumSte
 
 // Example Random Input Producer
 
-void UVehicleRandomInputProducer::InitializeContainer(TArray<FModuleInputSetup>& SetupData, FInputNameMap& NameMapOut)
+void UVehicleRandomInputProducer::InitializeContainer(TArray<FModuleInputSetup>& SetupData, FInputNameMap& NameMapOut, EModuleInputQuantizationType InInputQuantizationType)
 {
+	Super::InitializeContainer(SetupData, NameMapOut, InInputQuantizationType);
+
 	PlaybackContainer.Initialize(SetupData, NameMapOut);
 }
 
@@ -110,7 +116,7 @@ void UVehicleRandomInputProducer::ProduceInput(int32 PhysicsStep, int32 NumSteps
 		PlaybackContainer.ZeroValues();
 
 		// generate new random input
-		FInputInterface Inputs(InNameMap, PlaybackContainer);
+		FInputInterface Inputs(InNameMap, PlaybackContainer, InputQuantizationType);
 		Inputs.SetValue("Throttle", Random.FRand());
 		Inputs.SetValue("Steering", 1.0f - 2.0f * Random.FRand());
 	}
diff --git a/Engine/Plugins/Experimental/ChaosModularVehicle/Source/ChaosModularVehicleEngine/Private/ChaosModularVehicle/ModularVehicleBaseComponent.cpp b/Engine/Plugins/Experimental/ChaosModularVehicle/Source/ChaosModularVehicleEngine/Private/ChaosModularVehicle/ModularVehicleBaseComponent.cpp
--- a/Engine/Plugins/Experimental/ChaosModularVehicle/Source/ChaosModularVehicleEngine/Private/ChaosModularVehicle/ModularVehicleBaseComponent.cpp
+++ b/Engine/Plugins/Experimental/ChaosModularVehicle/Source/ChaosModularVehicleEngine/Private/ChaosModularVehicle/ModularVehicleBaseComponent.cpp
@@ -674,7 +674,7 @@ void UModularVehicleBaseComponent::BeginPlay()
 
 	if (InputProducer)
 	{
-		InputProducer->InitializeContainer(CombinedInputConfiguration, InputNameMap);
+		InputProducer->InitializeContainer(CombinedInputConfiguration, InputNameMap, InputQuantizationType);
 		if (InputProducer->GetTestInputBuffer())
 		{
 			// deliberate deep copy rather than copy pointer
@@ -693,6 +693,7 @@ void UModularVehicleBaseComponent::BeginPlay()
 
 	if (VehicleSimulationPT)
 	{
+		VehicleSimulationPT->SetInputQuantizationType(InputQuantizationType);
 		VehicleSimulationPT->SetInputMappings(InputNameMap);
 		VehicleSimulationPT->SetStateMappings(StateNameMap);
 	}
@@ -1595,31 +1596,31 @@ void UModularVehicleBaseComponent::SetInput(const FName& Name, const FVector& Va
 
 void UModularVehicleBaseComponent::SetState(const FName& InName, const bool InValue)
 {
-	FInputInterface Inputs(StateNameMap, StateInputContainer);
+	FInputInterface Inputs(StateNameMap, StateInputContainer, InputQuantizationType);
 	Inputs.SetBool(InName, InValue);
 }
 
 void UModularVehicleBaseComponent::SetState(const FName& InName, const int32 InValue)
 {
-	FInputInterface Inputs(StateNameMap, StateInputContainer);
+	FInputInterface Inputs(StateNameMap, StateInputContainer, InputQuantizationType);
 	Inputs.SetInteger(InName, InValue);
 }
 
 void UModularVehicleBaseComponent::SetState(const FName& InName, const double InValue)
 {
-	FInputInterface Inputs(StateNameMap, StateInputContainer);
+	FInputInterface Inputs(StateNameMap, StateInputContainer, InputQuantizationType);
 	Inputs.SetFloat(InName, InValue);
 }
 
 void UModularVehicleBaseComponent::SetState(const FName& InName, const FVector2D& InValue)
 {
-	FInputInterface Inputs(StateNameMap, StateInputContainer);
+	FInputInterface Inputs(StateNameMap, StateInputContainer, InputQuantizationType);
 	Inputs.SetVector2D(InName, InValue);
 }
 
 void UModularVehicleBaseComponent::SetState(const FName& InName, const FVector& InValue)
 {
-	FInputInterface Inputs(StateNameMap, StateInputContainer);
+	FInputInterface Inputs(StateNameMap, StateInputContainer, InputQuantizationType);
 	Inputs.SetVector(InName, InValue);
 }
 
diff --git a/Engine/Plugins/Experimental/ChaosModularVehicle/Source/ChaosModularVehicleEngine/Private/ChaosModularVehicle/ModularVehicleSimulationCU.cpp b/Engine/Plugins/Experimental/ChaosModularVehicle/Source/ChaosModularVehicleEngine/Private/ChaosModularVehicle/ModularVehicleSimulationCU.cpp
--- a/Engine/Plugins/Experimental/ChaosModularVehicle/Source/ChaosModularVehicleEngine/Private/ChaosModularVehicle/ModularVehicleSimulationCU.cpp
+++ b/Engine/Plugins/Experimental/ChaosModularVehicle/Source/ChaosModularVehicleEngine/Private/ChaosModularVehicle/ModularVehicleSimulationCU.cpp
@@ -115,10 +115,10 @@ void FModularVehicleSimulation::SimulateModuleTree(UWorld* InWorld, float DeltaS
 			}
 		}
 
-		FInputInterface InputInterface(InputNameMap, Container);
+		FInputInterface InputInterface(InputNameMap, Container, InputQuantizationType);
 
 		FModuleInputContainer StateInputContainer = InputData.PhysicsInputs.StateInputs.StateInputContainer;
-		FInputInterface StateInterface(StateNameMap, StateInputContainer);
+		FInputInterface StateInterface(StateNameMap, StateInputContainer, InputQuantizationType);
 
 		SimInputData.ControlInputs = &InputInterface;
 		SimInputData.StateInputs = &StateInterface;
diff --git a/Engine/Plugins/Experimental/ChaosModularVehicle/Source/ChaosModularVehicleEngine/Public/ChaosModularVehicle/InputProducer.h b/Engine/Plugins/Experimental/ChaosModularVehicle/Source/ChaosModularVehicleEngine/Public/ChaosModularVehicle/InputProducer.h
--- a/Engine/Plugins/Experimental/ChaosModularVehicle/Source/ChaosModularVehicleEngine/Public/ChaosModularVehicle/InputProducer.h
+++ b/Engine/Plugins/Experimental/ChaosModularVehicle/Source/ChaosModularVehicleEngine/Public/ChaosModularVehicle/InputProducer.h
@@ -16,7 +16,7 @@ class UVehicleDefaultInputProducer : public UVehicleInputProducerBase
 public:
 
 	/* initialize the input buffer container */
-	UE_API virtual void InitializeContainer(TArray<FModuleInputSetup>& SetupData, FInputNameMap& NameMapOut) override;
+	UE_API virtual void InitializeContainer(TArray<FModuleInputSetup>& SetupData, FInputNameMap& NameMapOut, EModuleInputQuantizationType InInputQuantizationType) override;
 
 	/** capture input at game thread frequency */
 	UE_API virtual void BufferInput(const FInputNameMap& InNameMap, const FName InName, const FModuleInputValue& InValue) override;
@@ -37,7 +37,7 @@ class UVehiclePlaybackInputProducer : public UVehicleInputProducerBase
 public:
 
 	/* initialize the input buffer containers */
-	UE_API virtual void InitializeContainer(TArray<FModuleInputSetup>& SetupData, FInputNameMap& NameMapOut) override;
+	UE_API virtual void InitializeContainer(TArray<FModuleInputSetup>& SetupData, FInputNameMap& NameMapOut, EModuleInputQuantizationType InInputQuantizationType) override;
 
 	/** capture input at game thread frequency */
 	UE_API virtual void BufferInput(const FInputNameMap& InNameMap, const FName InName, const FModuleInputValue& InValue) override;
@@ -66,7 +66,7 @@ class UVehicleRandomInputProducer : public UVehicleInputProducerBase
 public:
 
 	/* initialize the input buffer containers */
-	UE_API virtual void InitializeContainer(TArray<FModuleInputSetup>& SetupData, FInputNameMap& NameMapOut) override;
+	UE_API virtual void InitializeContainer(TArray<FModuleInputSetup>& SetupData, FInputNameMap& NameMapOut, EModuleInputQuantizationType InInputQuantizationType) override;
 
 	/** capture input at game thread frequency */
 	UE_API virtual void BufferInput(const FInputNameMap& InNameMap, const FName InName, const FModuleInputValue& InValue) override;
diff --git a/Engine/Plugins/Experimental/ChaosModularVehicle/Source/ChaosModularVehicleEngine/Public/ChaosModularVehicle/ModularVehicleBaseComponent.h b/Engine/Plugins/Experimental/ChaosModularVehicle/Source/ChaosModularVehicleEngine/Public/ChaosModularVehicle/ModularVehicleBaseComponent.h
--- a/Engine/Plugins/Experimental/ChaosModularVehicle/Source/ChaosModularVehicleEngine/Public/ChaosModularVehicle/ModularVehicleBaseComponent.h
+++ b/Engine/Plugins/Experimental/ChaosModularVehicle/Source/ChaosModularVehicleEngine/Public/ChaosModularVehicle/ModularVehicleBaseComponent.h
@@ -208,6 +208,9 @@ public:
 	UPROPERTY(EditAnywhere, Category = "Game|Components|ModularVehicle")
 	TSubclassOf<UVehicleInputProducerBase> InputProducerClass;
 
+	UPROPERTY(EditAnywhere, Category = "Game|Components|ModularVehicle")
+	EModuleInputQuantizationType InputQuantizationType = EModuleInputQuantizationType::Default_16Bits;
+	
 	// CONTROLS
 	// 
 	//void SetInput(const FName& Name, const FModuleInputValue& Value);
diff --git a/Engine/Plugins/Experimental/ChaosModularVehicle/Source/ChaosModularVehicleEngine/Public/ChaosModularVehicle/ModularVehicleSimulationCU.h b/Engine/Plugins/Experimental/ChaosModularVehicle/Source/ChaosModularVehicleEngine/Public/ChaosModularVehicle/ModularVehicleSimulationCU.h
--- a/Engine/Plugins/Experimental/ChaosModularVehicle/Source/ChaosModularVehicleEngine/Public/ChaosModularVehicle/ModularVehicleSimulationCU.h
+++ b/Engine/Plugins/Experimental/ChaosModularVehicle/Source/ChaosModularVehicleEngine/Public/ChaosModularVehicle/ModularVehicleSimulationCU.h
@@ -57,11 +57,19 @@ public:
 
 	UE_API void Initialize(TUniquePtr<Chaos::FSimModuleTree>& InSimModuleTree);
 	UE_API void Terminate();
+
+	void SetInputQuantizationType(EModuleInputQuantizationType InInputQuantizationType)
+	{ 
+		UE::TWriteScopeLock InputConfigLock(InputConfigurationLock);
+		InputQuantizationType = InInputQuantizationType; 
+	}	
+
 	void SetInputMappings(const FInputNameMap& InNameMap)
 	{ 
 		UE::TWriteScopeLock InputConfigLock(InputConfigurationLock);
 		InputNameMap = InNameMap; 
 	}
+
 	void SetStateMappings(const FInputNameMap& InNameMap)
 	{
 		UE::TWriteScopeLock InputConfigLock(InputConfigurationLock);
@@ -105,6 +113,7 @@ public:
 
 	/** Current control inputs that is being used on the PT */
 	FModularVehicleInputs VehicleInputs;
+	EModuleInputQuantizationType InputQuantizationType = EModuleInputQuantizationType::Default_16Bits;
 	FInputNameMap InputNameMap;
 	FInputNameMap StateNameMap;
 	FTransactionallySafeRWLock InputConfigurationLock;
diff --git a/Engine/Source/Programs/HeadlessChaos/Private/HeadlessChaosTestModularVehicles.cpp b/Engine/Source/Programs/HeadlessChaos/Private/HeadlessChaosTestModularVehicles.cpp
--- a/Engine/Source/Programs/HeadlessChaos/Private/HeadlessChaosTestModularVehicles.cpp
+++ b/Engine/Source/Programs/HeadlessChaos/Private/HeadlessChaosTestModularVehicles.cpp
@@ -70,7 +70,7 @@ namespace ChaosTest
 
 			ValueContainer.Initialize(InputSetupData, NameMap);
 
-			ControlInputs = MakeUnique<FInputInterface>(NameMap, ValueContainer);
+			ControlInputs = MakeUnique<FInputInterface>(NameMap, ValueContainer, EModuleInputQuantizationType::Default_16Bits);
 			Inputs.ControlInputs = ControlInputs.Get();
 		}
 
diff --git a/Engine/Source/Runtime/Experimental/ChaosVehicles/ChaosVehiclesCore/Private/SimModule/ModuleInput.cpp b/Engine/Source/Runtime/Experimental/ChaosVehicles/ChaosVehiclesCore/Private/SimModule/ModuleInput.cpp
--- a/Engine/Source/Runtime/Experimental/ChaosVehicles/ChaosVehiclesCore/Private/SimModule/ModuleInput.cpp
+++ b/Engine/Source/Runtime/Experimental/ChaosVehicles/ChaosVehiclesCore/Private/SimModule/ModuleInput.cpp
@@ -5,15 +5,64 @@
 #include "HAL/IConsoleManager.h"
 #include "Net/Core/NetBitArray.h"
 
+DEFINE_LOG_CATEGORY(LogModularInput);
+
 TArray<struct FModuleInputSetup>* FScopedModuleInputInitializer::InitSetupData = nullptr;
 
 namespace ChaosModularVehicleCVars
 {
 	bool bEnableInputSendingTypeInfo = true;
 	FAutoConsoleVariableRef EnableInputSendingTypeInfo(TEXT("p.ModularVehicle.EnableInputSendingTypeInfo"), bEnableInputSendingTypeInfo, TEXT("Enable/Disable Input sending type info. Default: true"));
+
+	bool bEnableCustomInputQuantization = false;
+	FAutoConsoleVariableRef EnableCustomInputQuantization(TEXT("p.ModularVehicle.EnableCustomInputQuantization"), bEnableCustomInputQuantization, TEXT("Enable/Disable Custom Input Quantization info. Default: false"));	
 };
 
-DEFINE_LOG_CATEGORY(LogModularInput);
+namespace UInt8Quantize 
+{
+	uint8 Compress(const double InValue)
+	{
+		const double Result = FMath::RoundToDouble(128.0 + FMath::FloatSelect(InValue, FMath::Min(127 * InValue, 127.0), FMath::Max(128 * InValue, -128.0)));
+		return static_cast<uint8>(Result);
+	}
+
+	double UnCompress(const uint8 InValue)
+	{
+		const double Result = (InValue < 128) ? (InValue - 128.0) / 128.0 : (InValue - 128.0) / 127.0;
+		return Result; 
+	}
+
+	void Quantize(double& InOutValue)
+	{
+		const uint8 CompressedValue = Compress(InOutValue);
+		InOutValue = UnCompress(CompressedValue);
+	}
+
+	bool WriteCompressed(const double Value, FArchive& Ar)
+	{
+		uint8 CompressedValue = Compress(Value);
+		Ar << CompressedValue;
+		return true;
+	}
+
+	bool ReadCompressed(double& Value, FArchive& Ar)
+	{
+		uint8 CompressedValue;
+		Ar << CompressedValue;
+		Value = UnCompress(CompressedValue);
+		return true;
+	}
+
+	bool Serialize(double& InOutValue, FArchive& Ar)
+	{
+		if (Ar.IsSaving())
+		{
+			return WriteCompressed(InOutValue, Ar);
+		}
+
+		return ReadCompressed(InOutValue, Ar);
+	}
+}
 
 bool FModuleInputValue::IsNonZero(float Tolerance) const
 {
@@ -59,7 +108,7 @@ float FModuleInputValue::GetMagnitude() const
 	case EModuleInputValueType::MBoolean:
 	case EModuleInputValueType::MInteger:
 		return ValueInt;
-	case EModuleInputValueType::MAxis1D:
+	case EModuleInputValueType::MAxis1D: 
 		return Value.X;
 	case EModuleInputValueType::MAxis2D:
 		return Value.Size2D();
@@ -115,11 +164,91 @@ void FModuleInputValue::SetMagnitude(float NewSize)
 	}
 }
 
-void FModuleInputValue::Serialize(FArchive& Ar, UPackageMap* Map, bool& bOutSuccess)
+void FModuleInputValue::Quantize(double& InOutValue, EModuleInputQuantizationType InInputQuantizationType)
+{
+	if (ChaosModularVehicleCVars::bEnableCustomInputQuantization == false)
+	{
+		ModularQuantize::QuantizeValue<1, 16>(InOutValue);
+		return;
+	}
+	
+	switch (InInputQuantizationType)
+	{
+		case EModuleInputQuantizationType::Default_16Bits:
+			ModularQuantize::QuantizeValue<1, 16>(InOutValue);
+			break;
+
+		case EModuleInputQuantizationType::Custom_8Bits:
+			UInt8Quantize::Quantize(InOutValue);
+			break;
+
+		default:
+			checkf(false, TEXT("Unsupported value type for input quantization type!"));
+	}
+}
+
+bool FModuleInputValue::SerializeQuantized(double& InOutValue, FArchive& Ar, EModuleInputQuantizationType InInputQuantizationType)
+{
+	if (ChaosModularVehicleCVars::bEnableCustomInputQuantization == false)
+	{
+		return ModularQuantize::SerializeFixedFloat<1, 16>(InOutValue, Ar);
+	}
+	
+	switch (InInputQuantizationType)
+	{
+		case EModuleInputQuantizationType::Default_16Bits:
+			return ModularQuantize::SerializeFixedFloat<1, 16>(InOutValue, Ar);
+
+		case EModuleInputQuantizationType::Custom_8Bits:
+			return UInt8Quantize::Serialize(InOutValue, Ar);
+			
+		default:
+			checkf(false, TEXT("Unsupported value type for input quantization type!"));
+			return false;
+	}
+}
+
+FModuleInputValue FModuleInputValue::ReturnQuantized(EModuleInputQuantizationType InInputQuantizationType) const
+{
+	FModuleInputValue OutValue = FModuleInputValue(ValueType, Value);
+
+	switch (ValueType)
+	{
+		case EModuleInputValueType::MBoolean:
+		case EModuleInputValueType::MInteger:
+			OutValue = FModuleInputValue(ValueType, ValueInt);
+			break;
+
+		case EModuleInputValueType::MAxis1D:
+			Quantize(OutValue.Value.X, InInputQuantizationType);
+			break;
+
+		case EModuleInputValueType::MAxis2D:
+			Quantize(OutValue.Value.X, InInputQuantizationType);
+			Quantize(OutValue.Value.Y, InInputQuantizationType);
+			break;
+
+		case EModuleInputValueType::MAxis3D:
+			Quantize(OutValue.Value.X, InInputQuantizationType);
+			Quantize(OutValue.Value.Y, InInputQuantizationType);
+			Quantize(OutValue.Value.Z, InInputQuantizationType);
+			break;
+
+		default:
+			checkf(false, TEXT("Unsupported value type for module input value!"));
+			break;
+	}
+
+	//UE_LOG(LogTemp, Warning, TEXT("Original %f vs Quantized %f"), Value.X, OutValue.Value.X);
+
+	return OutValue;
+}
+
+void FModuleInputValue::Serialize(FArchive& Ar, UPackageMap* Map, bool& bOutSuccess, EModuleInputQuantizationType InInputQuantizationType)
 {
 	// #TODO: send only value changes/deltas
 
-	bool bIsNonZero = Ar.IsLoading() ? false : IsQuantizedNonZero();
+	bool bIsNonZero = Ar.IsLoading() ? false : IsQuantizedNonZero(InInputQuantizationType);
 	Ar.SerializeBits(&bIsNonZero, 1);
 	Ar.SerializeBits(&ValueType, 3);
 	Ar.SerializeBits(&bApplyInputDecay, 1);
@@ -146,12 +275,13 @@ void FModuleInputValue::Serialize(FArchive& Ar, UPackageMap* Map, bool& bOutSucc
 			}
 			break;
 
+
 		case EModuleInputValueType::MAxis3D:
-			ModularQuantize::SerializeFixedFloat<1, 16>(Value.Z, Ar);
+			SerializeQuantized(Value.Z, Ar, InInputQuantizationType);
 		case EModuleInputValueType::MAxis2D:
-			ModularQuantize::SerializeFixedFloat<1, 16>(Value.Y, Ar);
+			SerializeQuantized(Value.Y, Ar, InInputQuantizationType);
 		case EModuleInputValueType::MAxis1D:
-			ModularQuantize::SerializeFixedFloat<1, 16>(Value.X, Ar);
+			SerializeQuantized(Value.X, Ar, InInputQuantizationType);
 			break;
 
 		default:
@@ -166,9 +296,9 @@ void FModuleInputValue::Serialize(FArchive& Ar, UPackageMap* Map, bool& bOutSucc
 	bOutSuccess = true;
 }
 
-void FModuleInputValue::NetSerialize(FArchive& Ar, UPackageMap* Map, bool& bOutSuccess)
+void FModuleInputValue::NetSerialize(FArchive& Ar, UPackageMap* Map, bool& bOutSuccess,  EModuleInputQuantizationType InInputQuantizationType)
 {
-	bool bIsNonZero = Ar.IsLoading() ? false : IsQuantizedNonZero();
+	bool bIsNonZero = Ar.IsLoading() ? false : IsQuantizedNonZero(InInputQuantizationType);
 	Ar.SerializeBits(&bIsNonZero, 1);
 	if (ChaosModularVehicleCVars::bEnableInputSendingTypeInfo)
 	{
@@ -198,11 +328,11 @@ void FModuleInputValue::NetSerialize(FArchive& Ar, UPackageMap* Map, bool& bOutS
 			break;
 
 		case EModuleInputValueType::MAxis3D:
-			ModularQuantize::SerializeFixedFloat<1, 16>(Value.Z, Ar);
+			SerializeQuantized(Value.Z, Ar, InInputQuantizationType);
 		case EModuleInputValueType::MAxis2D:
-			ModularQuantize::SerializeFixedFloat<1, 16>(Value.Y, Ar);
+			SerializeQuantized(Value.Y, Ar, InInputQuantizationType);
 		case EModuleInputValueType::MAxis1D:
-			ModularQuantize::SerializeFixedFloat<1, 16>(Value.X, Ar);
+			SerializeQuantized(Value.X, Ar, InInputQuantizationType);
 			break;
 
 		default:
@@ -217,7 +347,7 @@ void FModuleInputValue::NetSerialize(FArchive& Ar, UPackageMap* Map, bool& bOutS
 	bOutSuccess = true;
 }
 
-void FModuleInputValue::DeltaNetSerialize(FArchive& Ar, UPackageMap* Map, bool& bOutSuccess, const FModuleInputValue& PreviousInputValue)
+void FModuleInputValue::DeltaNetSerialize(FArchive& Ar, UPackageMap* Map, bool& bOutSuccess, const FModuleInputValue& PreviousInputValue, EModuleInputQuantizationType InInputQuantizationType)
 {
 	bool bHasSamePreviousValue = Value == PreviousInputValue.Value;
 	bool bHasSamePreviousType = ValueType == PreviousInputValue.ValueType;
@@ -239,7 +369,7 @@ void FModuleInputValue::DeltaNetSerialize(FArchive& Ar, UPackageMap* Map, bool&
 		return;
 	}
 
-	NetSerialize(Ar, Map, bOutSuccess);
+	NetSerialize(Ar, Map, bOutSuccess, InInputQuantizationType);
 }
 
 void FModuleInputValue::Lerp(const FModuleInputValue& Min, const FModuleInputValue& Max, float Alpha)
@@ -412,7 +542,7 @@ void FModuleInputContainer::ZeroValues()
 	}
 }
 
-void FModuleInputContainer::Serialize(FArchive& Ar, UPackageMap* Map, bool& bOutSuccess)
+void FModuleInputContainer::Serialize(FArchive& Ar, UPackageMap* Map, bool& bOutSuccess, EModuleInputQuantizationType InInputQuantizationType)
 {
 	using namespace UE::Net;
 
@@ -428,7 +558,7 @@ void FModuleInputContainer::Serialize(FArchive& Ar, UPackageMap* Map, bool& bOut
 
 	for (uint32 I = 0; I < Number; I++)
 	{
-		InputValues[I].Serialize(Ar, Map, bOutSuccess);
+		InputValues[I].Serialize(Ar, Map, bOutSuccess, InInputQuantizationType);
 	}
 
 	return;
@@ -486,7 +616,7 @@ void FInputInterface::SetValue(const FName& InName, const FModuleInputValue& InV
 	{
 		if (const int* Index = NameMap.Find(InName))
 		{
-			ValueContainer.SetValueAtIndex(*Index, InValue, Quantize);
+			ValueContainer.SetValueAtIndex(*Index, InValue, InputQuantizationType, Quantize);
 		}
 		else
 		{
@@ -501,7 +631,7 @@ void FInputInterface::MergeValue(const FName& InName, const FModuleInputValue& I
 	{
 		if (const int* Index = NameMap.Find(InName))
 		{
-			ValueContainer.MergeValueAtIndex(*Index, InValue);
+			ValueContainer.MergeValueAtIndex(*Index, InValue, InputQuantizationType);
 		}
 		else
 		{
diff --git a/Engine/Source/Runtime/Experimental/ChaosVehicles/ChaosVehiclesCore/Public/SimModule/ModuleInput.h b/Engine/Source/Runtime/Experimental/ChaosVehicles/ChaosVehiclesCore/Public/SimModule/ModuleInput.h
--- a/Engine/Source/Runtime/Experimental/ChaosVehicles/ChaosVehiclesCore/Public/SimModule/ModuleInput.h
+++ b/Engine/Source/Runtime/Experimental/ChaosVehicles/ChaosVehiclesCore/Public/SimModule/ModuleInput.h
@@ -11,6 +11,13 @@ DECLARE_LOG_CATEGORY_EXTERN(LogModularInput, Log, All);
 
 #define UE_API CHAOSVEHICLESCORE_API
 
+UENUM(BlueprintType)
+enum class EModuleInputQuantizationType : uint8
+{
+	Default_16Bits,
+	Custom_8Bits, 
+};
+
 UENUM(BlueprintType)
 enum class EModuleInputValueType : uint8
 {
@@ -104,7 +111,7 @@ namespace ModularQuantize
 	}
 
 	template<int32 MaxValue, uint32 NumBits, typename T UE_REQUIRES(std::is_floating_point_v<T>&& NumBits < 32)>
-bool WriteCompressedFloat(const T Value, FArchive& Ar)
+	bool WriteCompressedFloat(const T Value, FArchive& Ar)
 	{
 		using Details = ModularQuantize::TCompressedFloatDetails<MaxValue, NumBits>;
 
@@ -163,7 +170,6 @@ bool WriteCompressedFloat(const T Value, FArchive& Ar)
 		ModularQuantize::FromCompressedFloat<MaxValue, NumBits>(CompressedValue, QuantRight);
 		return FMath::IsNearlyEqual(QuantLeft, QuantRight);
 	}
-
 }
 
 USTRUCT(BlueprintType)
@@ -190,43 +196,12 @@ public:
 	FModuleInputValue(MAxis2D InValue) : Value(InValue.X, InValue.Y, 0.f), ValueInt(0), ValueType(EModuleInputValueType::MAxis2D) {}
 	FModuleInputValue(MAxis3D InValue) : Value(InValue), ValueInt(0), ValueType(EModuleInputValueType::MAxis3D) {}
 
-	FModuleInputValue ReturnQuantized() const
-	{
-		FModuleInputValue OutValue = FModuleInputValue(ValueType, Value);
-
-		switch (ValueType)
-		{
-			case EModuleInputValueType::MBoolean:
-			case EModuleInputValueType::MInteger:
-				OutValue = FModuleInputValue(ValueType, ValueInt);
-			break;
-
-			case EModuleInputValueType::MAxis1D:
-				ModularQuantize::QuantizeValue<1, 16>(OutValue.Value.X);
-			break;
-
-			case EModuleInputValueType::MAxis2D:
-				ModularQuantize::QuantizeValue<1, 16>(OutValue.Value.X);
-				ModularQuantize::QuantizeValue<1, 16>(OutValue.Value.Y);
-			break;
-
-			case EModuleInputValueType::MAxis3D:
-				ModularQuantize::QuantizeValue<1, 16>(OutValue.Value.X);
-				ModularQuantize::QuantizeValue<1, 16>(OutValue.Value.Y);
-				ModularQuantize::QuantizeValue<1, 16>(OutValue.Value.Z);
-			break;
-
-			default:
-				checkf(false, TEXT("Unsupported value type for module input value!"));
-				break;
-
-		}
-
-		//UE_LOG(LogTemp, Warning, TEXT("Original %f vs Quantized %f"), Value.X, OutValue.Value.X);
-
-		return OutValue;
-	}
+	FModuleInputValue ReturnQuantized(EModuleInputQuantizationType InInputQuantizationType) const;
 
+	static void Quantize(double& InOutValue, EModuleInputQuantizationType InInputQuantizationType);
+	
+	static bool SerializeQuantized(double& InOutValue, FArchive& Ar, EModuleInputQuantizationType InInputQuantizationType);
+	
 	// Build a specific type with an arbitrary Axis3D value
 	FModuleInputValue(EModuleInputValueType InValueType, MAxis3D InValue) : Value(InValue), ValueType(InValueType)
 	{
@@ -346,7 +321,7 @@ public:
 	// Read only index based value accessor, doesn't care about type. Expect 0 when accessing unused components.
 	float operator[](int32 Index) const { return Value[Index]; }
 
-	bool IsQuantizedNonZero(float Tolerance = KINDA_SMALL_NUMBER) const { return ReturnQuantized().IsNonZero(Tolerance); }
+	bool IsQuantizedNonZero(EModuleInputQuantizationType InInputQuantizationType, float Tolerance = KINDA_SMALL_NUMBER) const { return ReturnQuantized(InInputQuantizationType).IsNonZero(Tolerance); }
 	UE_API bool IsNonZero(float Tolerance = KINDA_SMALL_NUMBER) const;
 
 	// In-place type conversion
@@ -367,9 +342,9 @@ public:
 	UE_API int32 GetMagnitudeInt() const;
 
 	// Serialize values
-	UE_API void Serialize(FArchive& Ar, UPackageMap* Map, bool& bOutSuccess);
-	UE_API void NetSerialize(FArchive& Ar, UPackageMap* Map, bool& bOutSuccess);
-	UE_API void DeltaNetSerialize(FArchive& Ar, UPackageMap* Map, bool& bOutSuccess,const FModuleInputValue& PreviousInputValue);
+	UE_API void Serialize(FArchive& Ar, UPackageMap* Map, bool& bOutSuccess, EModuleInputQuantizationType InInputQuantizationType);
+	UE_API void NetSerialize(FArchive& Ar, UPackageMap* Map, bool& bOutSuccess, EModuleInputQuantizationType InInputQuantizationType);
+	UE_API void DeltaNetSerialize(FArchive& Ar, UPackageMap* Map, bool& bOutSuccess,const FModuleInputValue& PreviousInputValue, EModuleInputQuantizationType InInputQuantizationType);
 
 	UE_API void Lerp(const FModuleInputValue& Min, const FModuleInputValue& Max, float Alpha);
 
@@ -518,7 +493,6 @@ public:
 	UE_API virtual float CalcControlFunction(float InputValue);
 };
 
-
 USTRUCT(BlueprintType)
 struct FModuleInputSetup
 {
@@ -597,19 +571,25 @@ public:
 	}
 
 	int GetNumInputs() const { return InputValues.Num(); }
+	
 	FModuleInputValue GetValueAtIndex(int Index) const { return InputValues[Index]; }
-	void SetValueAtIndex(int Index, const FModuleInputValue& InValue, bool Quantize = true)
+	
+	void SetValueAtIndex(int Index, const FModuleInputValue& InValue, EModuleInputQuantizationType InInputQuantizationType, bool Quantize = true)
 	{
 		if (Quantize)
 		{
-			InputValues[Index] = InValue.ReturnQuantized();
+			InputValues[Index] = InValue.ReturnQuantized(InInputQuantizationType);
 		}
 		else
 		{
 			InputValues[Index] = InValue;
 		}
 	}
-	void MergeValueAtIndex(int Index, const FModuleInputValue& InValue) { InputValues[Index].Merge(InValue.ReturnQuantized()); }
+
+	void MergeValueAtIndex(int Index, const FModuleInputValue& InValue, EModuleInputQuantizationType InInputQuantizationType)
+	{
+		InputValues[Index].Merge(InValue.ReturnQuantized(InInputQuantizationType));
+	}
 
 	FModuleInputContainer& operator=(const FModuleInputContainer& Other)
 	{
@@ -629,7 +609,7 @@ public:
 
 	UE_API void ZeroValues();
 
-	UE_API void Serialize(FArchive& Ar, UPackageMap* Map, bool& bOutSuccess);
+	UE_API void Serialize(FArchive& Ar, UPackageMap* Map, bool& bOutSuccess, EModuleInputQuantizationType InInputQuantizationType);
 
 	UE_API int AddInput(EModuleInputValueType Type, TSubclassOf<UDefaultModularVehicleInputModifier>& InputModifierClass);
 
@@ -657,9 +637,10 @@ public:
 
 	using FInputNameMap = TMap<FName, int>;
 
-	FInputInterface(const FInputNameMap& InNameMap, FModuleInputContainer& InValueContainer)
+	FInputInterface(const FInputNameMap& InNameMap, FModuleInputContainer& InValueContainer, EModuleInputQuantizationType InInputQuantizationType)
 		: NameMap(InNameMap)
 		, ValueContainer(InValueContainer)
+		, InputQuantizationType(InInputQuantizationType)
 	{
 	}
 
@@ -730,6 +711,7 @@ public:
 
 	const FInputNameMap& NameMap;			// per vehicle
 	FModuleInputContainer& ValueContainer;	// per vehicle instance
+	EModuleInputQuantizationType InputQuantizationType;	
 };
 
 UCLASS(Abstract, BlueprintType, Blueprintable, EditInlineNew, MinimalAPI)
@@ -741,7 +723,10 @@ public:
 	using FInputNameMap = TMap<FName, int>;
 
 	/* initialize the input buffer container(s) */
-	virtual void InitializeContainer(TArray<FModuleInputSetup>& SetupData, FInputNameMap& NameMapOut) {}
+	virtual void InitializeContainer(TArray<FModuleInputSetup>& SetupData, FInputNameMap& NameMapOut, EModuleInputQuantizationType InInputQuantizationType) 
+	{
+		InputQuantizationType = InInputQuantizationType;
+	}
 
 	/** capture input at game thread frequency */
 	virtual void BufferInput(const FInputNameMap& InNameMap, const FName InName, const FModuleInputValue& InValue) {}
@@ -755,6 +740,7 @@ public:
 	/** Special case override for providing test input straight onto the physics thread */
 	virtual bool IsLoopingTestInputBuffer() { return false; }
 
+	EModuleInputQuantizationType InputQuantizationType;
 };
 
 #undef UE_API 
\ No newline at end of file
