Iris - Don't call NetRPC_GetDebugFunctionName. It breaks occasionally.

#rnx

diff --git a/Engine/Source/Runtime/Experimental/Iris/Core/Private/Iris/ReplicationSystem/NetBlob/NetRPC.cpp b/Engine/Source/Runtime/Experimental/Iris/Core/Private/Iris/ReplicationSystem/NetBlob/NetRPC.cpp
--- a/Engine/Source/Runtime/Experimental/Iris/Core/Private/Iris/ReplicationSystem/NetBlob/NetRPC.cpp
+++ b/Engine/Source/Runtime/Experimental/Iris/Core/Private/Iris/ReplicationSystem/NetBlob/NetRPC.cpp
@@ -469,14 +469,16 @@ void FNetRPC::CallFunction(FNetRPCCallContext& CallContext)
 	{
 		if (!TargetObjectReference.IsValid())
 		{
-			UE_LOG(LogIrisRpc, Error, TEXT("Rejected RPC %s due to missing object or function for object: %s."),
-				*(Function ? Function->GetName() : NetRPC_GetDebugFunctionName(ReplicationSystem, NetObjectReference, FunctionLocator)),
+			UE_LOG(LogIrisRpc, Error, TEXT("Rejected RPC (%u:%u) %s due to missing object or function for object: %s."),
+				FunctionLocator.DescriptorIndex, FunctionLocator.FunctionIndex,
+				(Function ? ToCStr(Function->GetName()) : TEXT("\'unknown\'")),
 				*NetRPC_GetDebugObjectRefName(Context, NetObjectReference, Object));
 		}
 		else
 		{
-			UE_LOG(LogIrisRpc, Error, TEXT("Rejected RPC %s due to missing object or function for subobject: %s of rootobject: %s"),
-				*(Function ? Function->GetName() : NetRPC_GetDebugFunctionName(ReplicationSystem, NetObjectReference, FunctionLocator)),
+			UE_LOG(LogIrisRpc, Error, TEXT("Rejected RPC (%u:%u) %s due to missing object or function for subobject: %s of rootobject: %s"),
+				FunctionLocator.DescriptorIndex, FunctionLocator.FunctionIndex,
+				(Function ? ToCStr(Function->GetName()) : TEXT("\'unknown\'")),
 				*NetRPC_GetDebugObjectRefName(Context, TargetObjectReference, Object), *NetRPC_GetDebugObjectRefName(Context, NetObjectReference));
 		}
 		return;
