Iris - FScriptInterfaceNetSerializer no longer requires an interface pointer for non-null objects. This should fix issues with blueprint implemented interfaces.

diff --git a/Engine/Source/Runtime/Experimental/Iris/Core/Private/Iris/Serialization/ObjectNetSerializer.cpp b/Engine/Source/Runtime/Experimental/Iris/Core/Private/Iris/Serialization/ObjectNetSerializer.cpp
--- a/Engine/Source/Runtime/Experimental/Iris/Core/Private/Iris/Serialization/ObjectNetSerializer.cpp
+++ b/Engine/Source/Runtime/Experimental/Iris/Core/Private/Iris/Serialization/ObjectNetSerializer.cpp
@@ -244,11 +244,6 @@ void FScriptInterfaceNetSerializer::Dequantize(FNetSerializationContext& Context
 	{
 		const FScriptInterfaceNetSerializerConfig& Config = *static_cast<const FScriptInterfaceNetSerializerConfig*>(Args.NetSerializerConfig);
 		void* Interface = Object->GetInterfaceAddress(Config.InterfaceClass);
-		if (Interface == nullptr)
-		{
-			Context.SetError(GNetError_InvalidValue);
-			return;
-		}
 
 		Target.SetObject(Object);
 		Target.SetInterface(Interface);
