Iris - ScriptInterfaceNetSerializer validation will now allow null interface pointers if expected.

#rnx

diff --git a/Engine/Source/Runtime/Experimental/Iris/Core/Private/Iris/Serialization/ObjectNetSerializer.cpp b/Engine/Source/Runtime/Experimental/Iris/Core/Private/Iris/Serialization/ObjectNetSerializer.cpp
--- a/Engine/Source/Runtime/Experimental/Iris/Core/Private/Iris/Serialization/ObjectNetSerializer.cpp
+++ b/Engine/Source/Runtime/Experimental/Iris/Core/Private/Iris/Serialization/ObjectNetSerializer.cpp
@@ -261,11 +261,6 @@ bool FScriptInterfaceNetSerializer::Validate(FNetSerializationContext& Context,
 	if (UObject* Object = Value.GetObject())
 	{
 		void* Interface = Value.GetInterface();
-		if (Interface == nullptr)
-		{
-			return false;
-		}
-
 		const FScriptInterfaceNetSerializerConfig& Config = *static_cast<const FScriptInterfaceNetSerializerConfig*>(Args.NetSerializerConfig);
 		if (Interface != Object->GetInterfaceAddress(Config.InterfaceClass))
 		{
