IAD - fix hang when running defrag due to overcounting read-handle refs.

diff --git a/Engine/Source/Runtime/Experimental/IoStore/OnDemand/Private/OnDemandInstallCache.cpp b/Engine/Source/Runtime/Experimental/IoStore/OnDemand/Private/OnDemandInstallCache.cpp
--- a/Engine/Source/Runtime/Experimental/IoStore/OnDemand/Private/OnDemandInstallCache.cpp
+++ b/Engine/Source/Runtime/Experimental/IoStore/OnDemand/Private/OnDemandInstallCache.cpp
@@ -356,8 +356,6 @@ FSharedFileOpenAsyncResult FCas::OpenAsyncRead(FCasBlockId BlockId)
 	{
 		if (FSharedAsyncFileHandle Handle = MaybeHandle->Pin(); Handle.IsValid())
 		{
-			BlockIds.FindOrAdd(BlockId, 0)++;
-
 			return MakeValue(MoveTemp(Handle));
 		}
 	}
