Dataflow : add a notification system for when variable override state is changed and invalid the corresponding variable nodes when the state changes

diff --git a/Engine/Plugins/Experimental/Dataflow/Source/DataflowEditor/Private/Dataflow/DataflowMembersWidget.cpp b/Engine/Plugins/Experimental/Dataflow/Source/DataflowEditor/Private/Dataflow/DataflowMembersWidget.cpp
--- a/Engine/Plugins/Experimental/Dataflow/Source/DataflowEditor/Private/Dataflow/DataflowMembersWidget.cpp
+++ b/Engine/Plugins/Experimental/Dataflow/Source/DataflowEditor/Private/Dataflow/DataflowMembersWidget.cpp
@@ -320,6 +320,7 @@ void SDataflowMembersWidget::Construct(const FArguments& InArgs, TSharedPtr<FDat
 
 	FCoreUObjectDelegates::OnObjectPropertyChanged.AddRaw(this, &SDataflowMembersWidget::OnObjectPropertyChanged);
 	FDataflowAssetDelegates::OnSubGraphsChanged.AddRaw(this, &SDataflowMembersWidget::OnSubGraphsChanged);
+	FDataflowAssetDelegates::OnVariablesOverrideStateChanged.AddRaw(this, &SDataflowMembersWidget::OnVariablesOverrideStateChanged);
 }
 
 void SDataflowMembersWidget::CreateVariableOverrideDetailView()
@@ -377,6 +378,7 @@ SDataflowMembersWidget::~SDataflowMembersWidget()
 {
 	FCoreUObjectDelegates::OnObjectPropertyChanged.RemoveAll(this);
 	FDataflowAssetDelegates::OnSubGraphsChanged.RemoveAll(this);
+	FDataflowAssetDelegates::OnVariablesOverrideStateChanged.RemoveAll(this); 
 }
 
 const TSharedPtr<SDataflowGraphEditor> SDataflowMembersWidget::GetGraphEditor() const
@@ -538,10 +540,19 @@ namespace DataflowMembersWidget::Private
 	}
 }
 
-void SDataflowMembersWidget::OverridesDetailsViewFinishedChangingProperties(const FPropertyChangedEvent& InPropertyChangedEvent)
+void SDataflowMembersWidget::OnVariablesOverrideStateChanged(const UDataflow* InDataflowAsset, FName InVariableName, bool bInNewOverrideState)
 {
-///InPropertyChangedEvent.Property->GetOwnerProperty
+	if (TStrongObjectPtr<UDataflow> DataflowAsset = DataflowAssetWeakPtr.Pin())
+	{
+		if (DataflowAsset.Get() == InDataflowAsset)
+		{
+			InvalidateVariableNode(*DataflowAsset, InVariableName);
+		}
+	}
+}
 
+void SDataflowMembersWidget::OverridesDetailsViewFinishedChangingProperties(const FPropertyChangedEvent& InPropertyChangedEvent)
+{
 	if (TStrongObjectPtr<UDataflow> DataflowAsset = DataflowAssetWeakPtr.Pin())
 	{
 		if (IDataflowInstanceInterface* InstanceInterface = GetDataflowInstanceInterface())
@@ -552,18 +563,23 @@ void SDataflowMembersWidget::OverridesDetailsViewFinishedChangingProperties(cons
 			InPropertyChangedEvent.GetArrayIndicesPerObject(0, PropertyNameStack);
 
 			const FName VariableName = DataflowMembersWidget::Private::ExtractVariableNameFromPropertyChangeEvent(InPropertyChangedEvent);
-			if (DataflowAsset->Variables.FindPropertyDescByName(VariableName))
+			InvalidateVariableNode(*DataflowAsset, VariableName);
+		}
+	}
+}
+
+void SDataflowMembersWidget::InvalidateVariableNode(const UDataflow& InDataflowAsset, FName InVariableName)
+{
+	if (InDataflowAsset.Variables.FindPropertyDescByName(InVariableName))
+	{
+		// invalidate all the get nodes all the nodes that match the variable name 
+		for (const TSharedPtr<FDataflowNode>& Node : InDataflowAsset.GetDataflow()->GetNodes())
+		{
+			if (FGetDataflowVariableNode* VariableNode = Node->AsType<FGetDataflowVariableNode>())
 			{
-				// invalidate all the get nodes all the nodes that match the variable name 
-				for (const TSharedPtr<FDataflowNode>& Node : DataflowAsset->Dataflow->GetNodes())
+				if (VariableNode->GetVariableName() == InVariableName)
 				{
-					if (FGetDataflowVariableNode* VariableNode = Node->AsType<FGetDataflowVariableNode>())
-					{
-						if (VariableNode->GetVariableName() == VariableName)
-						{
-							VariableNode->Invalidate();
-						}
-					}
+					VariableNode->Invalidate();
 				}
 			}
 		}
diff --git a/Engine/Plugins/Experimental/Dataflow/Source/DataflowEditor/Public/Dataflow/DataflowMembersWidget.h b/Engine/Plugins/Experimental/Dataflow/Source/DataflowEditor/Public/Dataflow/DataflowMembersWidget.h
--- a/Engine/Plugins/Experimental/Dataflow/Source/DataflowEditor/Public/Dataflow/DataflowMembersWidget.h
+++ b/Engine/Plugins/Experimental/Dataflow/Source/DataflowEditor/Public/Dataflow/DataflowMembersWidget.h
@@ -79,8 +79,11 @@ private:
 	virtual FReply OnKeyDown(const FGeometry& MyGeometry, const FKeyEvent& InKeyEvent);
 
 	void OnObjectPropertyChanged(UObject* InObject, FPropertyChangedEvent& InPropertyChangedEvent);
+	void OnVariablesOverrideStateChanged(const UDataflow* DataflowAsset, FName VariableName, bool bNewOverrideState);
 	void OnSubGraphsChanged(const UDataflow* InDataflowAsset, const FGuid& InSubGraphGuid, UE::Dataflow::ESubGraphChangedReason InReason);
 
+	void InvalidateVariableNode(const UDataflow& InDataflowAsset, FName VariableName);
+
 	void OnFilterTextChanged(const FText& InFilterText);
 	FText GetFilterText() const;
 	TSharedRef<SWidget> OnCreateWidgetForAction(FCreateWidgetForActionData* const InCreateData);
diff --git a/Engine/Source/Runtime/Experimental/Dataflow/Engine/Private/Dataflow/DataflowInstance.cpp b/Engine/Source/Runtime/Experimental/Dataflow/Engine/Private/Dataflow/DataflowInstance.cpp
--- a/Engine/Source/Runtime/Experimental/Dataflow/Engine/Private/Dataflow/DataflowInstance.cpp
+++ b/Engine/Source/Runtime/Experimental/Dataflow/Engine/Private/Dataflow/DataflowInstance.cpp
@@ -140,13 +140,33 @@ FDataflowVariableOverrides& FDataflowVariableOverrides::operator=(const FDataflo
 	return *this;
 }
 
-bool FDataflowVariableOverrides::OverrideVariableBool(FName VariableName, bool Value)
+void FDataflowVariableOverrides::SetVariableOverrideAndNotify(const FGuid& PropertyID, bool bOverrideState)
 {
-	if (const FPropertyBagPropertyDesc* Desc = Variables.FindPropertyDescByName(VariableName))
+	if (const FPropertyBagPropertyDesc* Desc = Variables.FindPropertyDescByID(PropertyID))
 	{
-		if (Variables.SetValueBool(*Desc, Value) == EPropertyBagResult::Success)
+		if (bOverrideState)
 		{
 			OverriddenVariableGuids.AddUnique(Desc->ID);
+		}
+		else
+		{
+			OverriddenVariableGuids.Remove(Desc->ID);
+		}
+		if (FDataflowAssetDelegates::OnVariablesOverrideStateChanged.IsBound())
+		{
+			const UDataflow* DataflowAsset = Owner ? Owner->GetDataflowAsset() : nullptr;
+			FDataflowAssetDelegates::OnVariablesOverrideStateChanged.Broadcast(DataflowAsset, Desc->Name, bOverrideState);
+		}
+	}
+}
+
+bool FDataflowVariableOverrides::OverrideVariableBool(FName VariableName, bool bValue)
+{
+	if (const FPropertyBagPropertyDesc* Desc = Variables.FindPropertyDescByName(VariableName))
+	{
+		if (Variables.SetValueBool(*Desc, bValue) == EPropertyBagResult::Success)
+		{
+			SetVariableOverrideAndNotify(Desc->ID, true);
 			return true;
 		}
 	}
@@ -162,7 +182,7 @@ bool FDataflowVariableOverrides::OverrideVariableBoolArray(FName VariableName, c
 			return ArrayRef.SetValueBool(Idx, Value);
 		};
 
-	return OverrideVariableArray(Variables, VariableName, Values, WriteBoolValue, OverriddenVariableGuids);
+	return OverrideVariableArrayAndNotify(Variables, VariableName, Values, WriteBoolValue);
 }
 
 bool FDataflowVariableOverrides::OverrideVariableInt(FName VariableName, int64 Value)
@@ -171,7 +191,7 @@ bool FDataflowVariableOverrides::OverrideVariableInt(FName VariableName, int64 V
 	{
 		if (Variables.SetValueInt64(*Desc, Value) == EPropertyBagResult::Success)
 		{
-			OverriddenVariableGuids.AddUnique(Desc->ID);
+			SetVariableOverrideAndNotify(Desc->ID, true);
 			return true;
 		}
 	}
@@ -187,7 +207,7 @@ bool FDataflowVariableOverrides::OverrideVariableInt32Array(FName VariableName,
 			return ArrayRef.SetValueInt32(Idx, Value);
 		};
 
-	return OverrideVariableArray(Variables, VariableName, Values, WriteInt32Value, OverriddenVariableGuids);
+	return OverrideVariableArrayAndNotify(Variables, VariableName, Values, WriteInt32Value);
 }
 
 bool FDataflowVariableOverrides::OverrideVariableInt64Array(FName VariableName, const TArray<int64>& Values)
@@ -199,7 +219,7 @@ bool FDataflowVariableOverrides::OverrideVariableInt64Array(FName VariableName,
 			return ArrayRef.SetValueInt64(Idx, Value);
 		};
 
-	return OverrideVariableArray(Variables, VariableName, Values, WriteInt64Value, OverriddenVariableGuids);
+	return OverrideVariableArrayAndNotify(Variables, VariableName, Values, WriteInt64Value);
 }
 
 bool FDataflowVariableOverrides::OverrideVariableFloat(FName VariableName, float Value)
@@ -208,7 +228,7 @@ bool FDataflowVariableOverrides::OverrideVariableFloat(FName VariableName, float
 	{
 		if (Variables.SetValueFloat(*Desc, Value) == EPropertyBagResult::Success)
 		{
-			OverriddenVariableGuids.AddUnique(Desc->ID);
+			SetVariableOverrideAndNotify(Desc->ID, true);
 			return true;
 		}
 	}
@@ -224,7 +244,7 @@ bool FDataflowVariableOverrides::OverrideVariableFloatArray(FName VariableName,
 			return ArrayRef.SetValueFloat(Idx, Value);
 		};
 
-	return OverrideVariableArray(Variables, VariableName, Values, WriteFloatValue, OverriddenVariableGuids);
+	return OverrideVariableArrayAndNotify(Variables, VariableName, Values, WriteFloatValue);
 }
 
 bool FDataflowVariableOverrides::OverrideVariableObject(FName VariableName, const UObject* Value)
@@ -233,7 +253,7 @@ bool FDataflowVariableOverrides::OverrideVariableObject(FName VariableName, cons
 	{
 		if (Variables.SetValueObject(*Desc, Value) == EPropertyBagResult::Success)
 		{
-			OverriddenVariableGuids.AddUnique(Desc->ID);
+			SetVariableOverrideAndNotify(Desc->ID, true);
 			return true;
 		}
 	}
@@ -249,7 +269,7 @@ bool FDataflowVariableOverrides::OverrideVariableObjectArray(FName VariableName,
 			return ArrayRef.SetValueObject(Idx, Value);
 		};
 
-	return OverrideVariableArray(Variables, VariableName, Values, WriteObjectValue, OverriddenVariableGuids);
+	return OverrideVariableArrayAndNotify(Variables, VariableName, Values, WriteObjectValue);
 }
 
 bool FDataflowVariableOverrides::OverrideVariableObjectArray(FName VariableName, const TArray<UObject*>& Values)
@@ -261,7 +281,7 @@ bool FDataflowVariableOverrides::OverrideVariableObjectArray(FName VariableName,
 			return ArrayRef.SetValueObject(Idx, Value);
 		};
 
-	return OverrideVariableArray(Variables, VariableName, Values, WriteObjectValue, OverriddenVariableGuids);
+	return OverrideVariableArrayAndNotify(Variables, VariableName, Values, WriteObjectValue);
 }
 
 bool FDataflowVariableOverrides::OverrideVariableName(FName VariableName, FName Value)
@@ -270,7 +290,7 @@ bool FDataflowVariableOverrides::OverrideVariableName(FName VariableName, FName
 	{
 		if (Variables.SetValueName(*Desc, Value) == EPropertyBagResult::Success)
 		{
-			OverriddenVariableGuids.AddUnique(Desc->ID);
+			SetVariableOverrideAndNotify(Desc->ID, true);
 			return true;
 		}
 	}
@@ -286,7 +306,7 @@ bool FDataflowVariableOverrides::OverrideVariableName(FName VariableName, const
 			return ArrayRef.SetValueName(Idx, Value);
 		};
 
-	return OverrideVariableArray(Variables, VariableName, Values, WriteNameValue, OverriddenVariableGuids);
+	return OverrideVariableArrayAndNotify(Variables, VariableName, Values, WriteNameValue);
 }
 
 bool FDataflowVariableOverrides::OverrideVariableString(FName VariableName, FString Value)
@@ -295,7 +315,7 @@ bool FDataflowVariableOverrides::OverrideVariableString(FName VariableName, FStr
 	{
 		if (Variables.SetValueString(*Desc, Value) == EPropertyBagResult::Success)
 		{
-			OverriddenVariableGuids.AddUnique(Desc->ID);
+			SetVariableOverrideAndNotify(Desc->ID, true);
 			return true;
 		}
 	}
@@ -311,7 +331,7 @@ bool FDataflowVariableOverrides::OverrideVariableString(FName VariableName, cons
 			return ArrayRef.SetValueString(Idx, Value);
 		};
 
-	return OverrideVariableArray(Variables, VariableName, Values, WriteStringValue, OverriddenVariableGuids);
+	return OverrideVariableArrayAndNotify(Variables, VariableName, Values, WriteStringValue);
 }
 
 const FInstancedPropertyBag* FDataflowVariableOverrides::GetDefaultVariablesFromAsset() const 
@@ -406,14 +426,7 @@ bool FDataflowVariableOverrides::IsVariableOverridden(FName VariableName) const
 
 void FDataflowVariableOverrides::SetVariableOverridden(const FGuid PropertyID, const bool bIsOverridden)
 {
-	if (bIsOverridden)
-	{
-		OverriddenVariableGuids.AddUnique(PropertyID);
-	}
-	else
-	{
-		OverriddenVariableGuids.Remove(PropertyID);
-	}
+	SetVariableOverrideAndNotify(PropertyID, bIsOverridden);
 	SyncVariables();
 }
 
diff --git a/Engine/Source/Runtime/Experimental/Dataflow/Engine/Private/Dataflow/DataflowObject.cpp b/Engine/Source/Runtime/Experimental/Dataflow/Engine/Private/Dataflow/DataflowObject.cpp
--- a/Engine/Source/Runtime/Experimental/Dataflow/Engine/Private/Dataflow/DataflowObject.cpp
+++ b/Engine/Source/Runtime/Experimental/Dataflow/Engine/Private/Dataflow/DataflowObject.cpp
@@ -17,6 +17,7 @@
 FDataflowAssetDelegates::FOnVariablesChanged FDataflowAssetDelegates::OnVariablesChanged;
 FDataflowAssetDelegates::FOnSubGraphsChanged FDataflowAssetDelegates::OnSubGraphsChanged;
 FDataflowAssetDelegates::FOnNodeInvalidated FDataflowAssetDelegates::OnNodeInvalidated;
+FDataflowAssetDelegates::FOnVariablesOverrideStateChanged FDataflowAssetDelegates::OnVariablesOverrideStateChanged;
 
 namespace UE::Dataflow::CVars
 {
diff --git a/Engine/Source/Runtime/Experimental/Dataflow/Engine/Public/Dataflow/DataflowInstance.h b/Engine/Source/Runtime/Experimental/Dataflow/Engine/Public/Dataflow/DataflowInstance.h
--- a/Engine/Source/Runtime/Experimental/Dataflow/Engine/Public/Dataflow/DataflowInstance.h
+++ b/Engine/Source/Runtime/Experimental/Dataflow/Engine/Public/Dataflow/DataflowInstance.h
@@ -17,7 +17,7 @@ namespace UE::Dataflow::InstanceUtils::Private
 {
 	// need to be in the header to be used by the templated structure override methods
 	template<typename T, typename TWriteElementFunction>
-	bool OverrideVariableArray(FInstancedPropertyBag& Variables, FName VariableName, const TArray<T>& Values, const TWriteElementFunction& WriteElementFunction, TArray<FGuid>& InOutOverrideGuids)
+	bool OverrideVariableArray(FInstancedPropertyBag& Variables, FName VariableName, const TArray<T>& Values, const TWriteElementFunction& WriteElementFunction, FGuid& OutGuid)
 	{
 		if (const FPropertyBagPropertyDesc* Desc = Variables.FindPropertyDescByName(VariableName))
 		{
@@ -32,11 +32,12 @@ namespace UE::Dataflow::InstanceUtils::Private
 				}
 				if (bSuccess)
 				{
-					InOutOverrideGuids.AddUnique(Desc->ID);
+					OutGuid = Desc->ID;
 				}
 				return bSuccess;
 			}
 		}
+		OutGuid.Invalidate();
 		return false;
 	}
 }
@@ -80,7 +81,7 @@ struct FDataflowVariableOverrides
 	DATAFLOWENGINE_API static FName GetVariablePropertyName();
 
 	/** override a dataflow boolean variable for this asset */
-	DATAFLOWENGINE_API bool OverrideVariableBool(FName VariableName, bool Value);
+	DATAFLOWENGINE_API bool OverrideVariableBool(FName VariableName, bool bValue);
 
 	/** override a dataflow boolean array variable for this asset */
 	DATAFLOWENGINE_API bool OverrideVariableBoolArray(FName VariableName, const TArray<bool>& Values);
@@ -145,7 +146,7 @@ struct FDataflowVariableOverrides
 				return ArrayRef.SetValueStruct(Idx, Value);
 			};
 
-		return OverrideVariableArray(Variables, VariableName, Values, WriteStructValue, OverriddenVariableGuids);
+		return OverrideVariableArrayAndNotify(Variables, VariableName, Values, WriteStructValue);
 	}
 
 #if WITH_EDITOR
@@ -170,6 +171,22 @@ private:
 	void RemoveOverridenVariablesNotInDataflowAsset();
 
 private:
+	void SetVariableOverrideAndNotify(const FGuid& PropertyID, bool bOverrideState);
+
+	template<typename T, typename TWriteElementFunction>
+	bool OverrideVariableArrayAndNotify(FInstancedPropertyBag& InVariables, FName InVariableName, const TArray<T>& InValues, const TWriteElementFunction& InWriteElementFunction)
+	{
+		FGuid PropertyGuid;
+		const bool bSuccess = UE::Dataflow::InstanceUtils::Private::OverrideVariableArray(InVariables, InVariableName, InValues, InWriteElementFunction, PropertyGuid);
+		if (bSuccess)
+		{
+			SetVariableOverrideAndNotify(PropertyGuid, true);
+		}
+		return bSuccess;
+	}
+
+private:
+
 	/** Variables to override */
 	UPROPERTY(EditAnywhere, Category = "Dataflow", meta = (FixedLayout, ShowOnlyInnerProperties))
 	FInstancedPropertyBag Variables;
diff --git a/Engine/Source/Runtime/Experimental/Dataflow/Engine/Public/Dataflow/DataflowObject.h b/Engine/Source/Runtime/Experimental/Dataflow/Engine/Public/Dataflow/DataflowObject.h
--- a/Engine/Source/Runtime/Experimental/Dataflow/Engine/Public/Dataflow/DataflowObject.h
+++ b/Engine/Source/Runtime/Experimental/Dataflow/Engine/Public/Dataflow/DataflowObject.h
@@ -47,6 +47,9 @@ struct FDataflowAssetDelegates
 	DECLARE_MULTICAST_DELEGATE_TwoParams(FOnVariablesChanged, const UDataflow*, FName);
 	static DATAFLOWENGINE_API FOnVariablesChanged OnVariablesChanged;
 
+	DECLARE_MULTICAST_DELEGATE_ThreeParams(FOnVariablesOverrideStateChanged, const UDataflow*, FName, bool);
+	static DATAFLOWENGINE_API FOnVariablesOverrideStateChanged OnVariablesOverrideStateChanged;
+
 	/** Called when sub-graphs are edited ( add, remove ) */
 	DECLARE_MULTICAST_DELEGATE_ThreeParams(FOnSubGraphsChanged, const UDataflow*, const FGuid&, UE::Dataflow::ESubGraphChangedReason);
 	static DATAFLOWENGINE_API FOnSubGraphsChanged OnSubGraphsChanged;
