Dataflow : make sure perf data is recorded even if not display

diff --git a/Engine/Source/Runtime/Experimental/Dataflow/Core/Private/Dataflow/DataflowNodeParameters.cpp b/Engine/Source/Runtime/Experimental/Dataflow/Core/Private/Dataflow/DataflowNodeParameters.cpp
--- a/Engine/Source/Runtime/Experimental/Dataflow/Core/Private/Dataflow/DataflowNodeParameters.cpp
+++ b/Engine/Source/Runtime/Experimental/Dataflow/Core/Private/Dataflow/DataflowNodeParameters.cpp
@@ -22,7 +22,7 @@ namespace UE::Dataflow
 
 	void FContextPerfData::Accumulate(const FDataflowConnection* Connection, uint64 TotalTime, uint64 ExternalTime)
 	{
-		if (Connection && bEnabled)
+		if (Connection)
 		{
 			FData& Data = DataPerOutput.FindOrAdd(Connection->GetGuid());
 			Data.ExclusiveTimeMs = FMath::Max(Data.ExclusiveTimeMs, (float)FPlatformTime::ToMilliseconds64(TotalTime - ExternalTime));
