Added console command iostore.VerifyCache to verify all installed chunk(s).

diff --git a/Engine/Source/Runtime/Experimental/IoStore/OnDemand/Private/OnDemandContentInstaller.cpp b/Engine/Source/Runtime/Experimental/IoStore/OnDemand/Private/OnDemandContentInstaller.cpp
--- a/Engine/Source/Runtime/Experimental/IoStore/OnDemand/Private/OnDemandContentInstaller.cpp
+++ b/Engine/Source/Runtime/Experimental/IoStore/OnDemand/Private/OnDemandContentInstaller.cpp
@@ -269,6 +269,16 @@ void FOnDemandContentInstaller::EnqueueDefragRequest(FOnDemandDefragArgs&& Args,
 	TryExecuteNextRequest();
 }
 
+void FOnDemandContentInstaller::EnqueueVerifyRequest(FOnDemandVerifyCacheCompleted&& OnCompleted)
+{
+	{
+		TUniqueLock Lock(Mutex);
+		RequestQueue.HeapPush(RequestAllocator.Construct(MoveTemp(OnCompleted)), RequestSortPredicate);
+	}
+
+	TryExecuteNextRequest();
+}
+
 void FOnDemandContentInstaller::CancelInstallRequest(FSharedInternalInstallRequest InstallRequest)
 {
 	InstallerPipe.Launch(
@@ -402,6 +412,11 @@ void FOnDemandContentInstaller::ExecuteRequest(FRequest& Request)
 			Installer.ExecuteDefragRequest(Request);
 		}
 
+		void operator()(FRequest::FVerify& Verify)
+		{
+			Installer.ExecuteVerifyRequest(Request);
+		}
+
 		FOnDemandContentInstaller&	Installer;
 		FRequest&					Request;
 	};
@@ -750,6 +765,29 @@ void FOnDemandContentInstaller::ExecuteDefragRequest(FRequest& Request)
 	CompleteDefragRequest(Request); 
 }
 
+void FOnDemandContentInstaller::ExecuteVerifyRequest(FRequest& Request)
+{
+	check(Request.IsVerify());
+	check(&Request == CurrentRequest);
+
+	UE_LOG(LogIoStoreOnDemand, Log, TEXT("Executing verify cache request"));	
+
+	if (Request.ErrorCode.load(std::memory_order_relaxed) == EIoErrorCode::Cancelled)
+	{
+		return CompleteVerifyRequest(Request);
+	}
+
+	const FIoStatus Status	= IoStore.InstallCache->Verify();
+	Request.ErrorCode		= Status.GetErrorCode();
+
+	if (!Status.IsOk())
+	{
+		Request.ErrorReason = Status.ToString();
+	}
+
+	CompleteVerifyRequest(Request); 
+}
+
 void FOnDemandContentInstaller::ProcessDownloadedChunk(
 	FRequest& Request,
 	FChunkHttpRequestHandle& HttpRequest,
@@ -1095,6 +1133,35 @@ void FOnDemandContentInstaller::CompleteDefragRequest(FRequest& Request)
 	}
 }
 
+void FOnDemandContentInstaller::CompleteVerifyRequest(FRequest& Request)
+{
+	const uint64 DurationCycles = FPlatformTime::Cycles64() - Request.StartTimeCycles; 
+
+	FOnDemandVerifyCacheResult VerifyResult;
+	VerifyResult.Status = Request.ErrorCode == EIoErrorCode::Ok ? FIoStatus::Ok : FIoStatus(Request.ErrorCode, Request.ErrorReason);
+	VerifyResult.DurationInSeconds	= FPlatformTime::ToSeconds64(DurationCycles);
+
+	UE_LOG(LogIoStoreOnDemand, Log, TEXT("Verify request completed, Result='%s', Duration=%d ms"),
+		GetIoErrorText(VerifyResult.Status.GetErrorCode()), int32(VerifyResult.DurationInSeconds * 1000));
+
+	FOnDemandVerifyCacheCompleted OnCompleted = MoveTemp(Request.AsVerify().OnCompleted);
+	{
+		UE::TUniqueLock Lock(Mutex);
+		if (&Request == CurrentRequest)
+		{
+			CurrentRequest = nullptr;
+		}
+		RequestAllocator.Destroy(&Request);
+	}
+
+	if (!OnCompleted)
+	{
+		return;
+	}
+
+	OnCompleted(MoveTemp(VerifyResult));
+}
+
 void FOnDemandContentInstaller::Shutdown()
 {
 	bShuttingDown					= true;
diff --git a/Engine/Source/Runtime/Experimental/IoStore/OnDemand/Private/OnDemandContentInstaller.h b/Engine/Source/Runtime/Experimental/IoStore/OnDemand/Private/OnDemandContentInstaller.h
--- a/Engine/Source/Runtime/Experimental/IoStore/OnDemand/Private/OnDemandContentInstaller.h
+++ b/Engine/Source/Runtime/Experimental/IoStore/OnDemand/Private/OnDemandContentInstaller.h
@@ -109,7 +109,15 @@ class FOnDemandContentInstaller
 			FOnDemandDefragCompleted	OnCompleted;
 		};
 
-		using FRequestVariant = TVariant<FEmptyVariantState, FInstall, FPurge, FDefrag>;
+		struct FVerify
+		{
+			FVerify(FOnDemandVerifyCacheCompleted&& InOnCompleted)
+				: OnCompleted(MoveTemp(InOnCompleted)) { }
+
+			FOnDemandVerifyCacheCompleted	OnCompleted;
+		};
+
+		using FRequestVariant = TVariant<FEmptyVariantState, FInstall, FPurge, FDefrag, FVerify>;
 
 		FRequest(FOnDemandInstallArgs&& Args, FOnDemandInstallCompleted&& OnCompleted, FOnDemandInstallProgressed&& OnProgress)
 		{
@@ -126,13 +134,20 @@ class FOnDemandContentInstaller
 			Variant.Emplace<FDefrag>(MoveTemp(Args), MoveTemp(OnCompleted));
 		}
 
-		bool			IsInstall() const	{ return Variant.IsType<FInstall>(); } 
-		bool			IsPurge() const		{ return Variant.IsType<FPurge>(); } 
-		bool			IsDefrag() const	{ return Variant.IsType<FDefrag>(); } 
+		FRequest(FOnDemandVerifyCacheCompleted&& OnCompleted)
+		{
+			Variant.Emplace<FVerify>(MoveTemp(OnCompleted));
+		}
+
+		bool			IsInstall() const	{ return Variant.IsType<FInstall>(); }
+		bool			IsPurge() const		{ return Variant.IsType<FPurge>(); }
+		bool			IsDefrag() const	{ return Variant.IsType<FDefrag>(); }
+		bool			IsVerify() const	{ return Variant.IsType<FVerify>(); }
 
 		FInstall&		AsInstall()			{ return Variant.Get<FInstall>(); }
 		FPurge&			AsPurge()			{ return Variant.Get<FPurge>(); }
 		FDefrag&		AsDefrag()			{ return Variant.Get<FDefrag>(); }
+		FVerify&		AsVerify()			{ return Variant.Get<FVerify>(); }
 
 		static uint32	NextSeqNo;
 
@@ -171,6 +186,7 @@ public:
 										FOnDemandInstallProgressed&& OnProgress);
 	void							EnqueuePurgeRequest(FOnDemandPurgeArgs&& Args, FOnDemandPurgeCompleted&& OnCompleted);
 	void							EnqueueDefragRequest(FOnDemandDefragArgs&& Args, FOnDemandDefragCompleted&& OnCompleted);
+	void							EnqueueVerifyRequest(FOnDemandVerifyCacheCompleted&& OnCompleted);
 	void							CancelInstallRequest(FSharedInternalInstallRequest InstallRequest);
 	void							UpdateInstallRequestPriority(FSharedInternalInstallRequest InstallRequest, int32 NewPriority);
 	void							ReportAnalytics(TArray<FAnalyticsEventAttribute>& OutAnalyticsArray) const;
@@ -182,6 +198,7 @@ private:
 	void							ExecuteInstallRequest(FRequest& Request, bool bRemoveAlreadyCachedChunks);
 	void							ExecutePurgeRequest(FRequest& Request);
 	void							ExecuteDefragRequest(FRequest& Request);
+	void							ExecuteVerifyRequest(FRequest& Request);
 	void							ProcessDownloadedChunk(
 										FRequest& Request,
 										FChunkHttpRequestHandle& HttpRequest,
@@ -192,6 +209,7 @@ private:
 	void							CompleteInstallRequest(FRequest& Request);
 	void							CompletePurgeRequest(FRequest& Request);
 	void							CompleteDefragRequest(FRequest& Request);
+	void							CompleteVerifyRequest(FRequest& Request);
 	void							Shutdown();
 
 	FOnDemandIoStore&				IoStore;
diff --git a/Engine/Source/Runtime/Experimental/IoStore/OnDemand/Private/OnDemandInstallCache.cpp b/Engine/Source/Runtime/Experimental/IoStore/OnDemand/Private/OnDemandInstallCache.cpp
--- a/Engine/Source/Runtime/Experimental/IoStore/OnDemand/Private/OnDemandInstallCache.cpp
+++ b/Engine/Source/Runtime/Experimental/IoStore/OnDemand/Private/OnDemandInstallCache.cpp
@@ -1337,6 +1337,7 @@ public:
 	virtual FIoStatus							Purge(TSet<FIoHash>&& ChunksToInstall) override;
 	virtual FIoStatus							PurgeAllUnreferenced(bool bDefrag, const uint64* BytesToPurge = nullptr) override;
 	virtual FIoStatus							DefragAll(const uint64* BytesToFree = nullptr) override;
+	virtual FIoStatus							Verify() override;
 	virtual FIoStatus							Flush() override;
 	virtual FOnDemandInstallCacheUsage			GetCacheUsage() override;
 
@@ -2052,6 +2053,201 @@ FIoStatus FOnDemandInstallCache::DefragAll(const uint64* BytesToFree /*= nullptr
 	return Defrag(Containers, ChunkEntryIndices, BlockInfo, BytesToFree);
 }
 
+FIoStatus FOnDemandInstallCache::Verify()
+{
+	struct FChunkLookup
+	{
+		TMap<FCasAddr, int32> AddrToIndex;
+	};
+
+	struct FCasAddrLocation
+	{
+		FCasAddr		Addr;
+		FCasLocation	Location;
+
+		bool operator<(const FCasAddrLocation& Other) const
+		{
+			if (Location.BlockId == Other.Location.BlockId)
+			{
+				return Location.BlockOffset < Other.Location.BlockOffset;
+			}
+			return Location.BlockId.Id < Other.Location.BlockId.Id;
+		}
+	};
+
+	TArray<FSharedOnDemandContainer>	Containers = IoStore.GetContainers(EOnDemandContainerFlags::InstallOnDemand);
+	TArray<FCasAddrLocation>			ChunkLocations;
+	TArray<FChunkLookup>				ChunkLookups;
+
+	{
+		TUniqueLock Lock(Cas.Mutex);
+		ChunkLocations.Reserve(Cas.Lookup.Num());
+		for (const TPair<FCasAddr, FCasLocation>& Kv : Cas.Lookup)
+		{
+			ChunkLocations.Add(FCasAddrLocation
+			{
+				.Addr		= Kv.Key,
+				.Location	= Kv.Value
+			});
+		}
+	}
+	ChunkLocations.Sort();
+
+	ChunkLookups.Reserve(Containers.Num());
+	for (int32 Idx = 0; Idx < Containers.Num(); ++Idx)
+	{
+		FSharedOnDemandContainer& Container = Containers[Idx];
+		FChunkLookup& Lookup				= ChunkLookups.AddDefaulted_GetRef();
+
+		Lookup.AddrToIndex.Reserve(Container->ChunkEntries.Num());
+		for (int32 EntryIndex = 0; const FOnDemandChunkEntry& Entry : Container->ChunkEntries)
+		{
+			const FCasAddr& Addr = *reinterpret_cast<const FCasAddr*>(&Entry.Hash);
+			Lookup.AddrToIndex.Add(Addr, EntryIndex++);
+		}
+	}
+
+	auto FindChunkEntry = [&Containers, &ChunkLookups](const FCasAddr& Addr, int32& OutContainerIndex) -> int32
+	{
+		OutContainerIndex = INDEX_NONE;
+		for (int32 Idx = 0; Idx < Containers.Num(); ++Idx)
+		{
+			FChunkLookup& Lookup = ChunkLookups[Idx];
+			if (const int32* EntryIndex = Lookup.AddrToIndex.Find(Addr))
+			{
+				OutContainerIndex = Idx;
+				return *EntryIndex;
+			}
+		}
+
+		return INDEX_NONE;
+	};
+
+	const int32	ChunkCount			= ChunkLocations.Num();
+	uint32		CorruptChunkCount	= 0;
+	uint32		MissingChunkCount	= 0;
+	uint32		ReadErrorCount		= 0;
+	uint64		TotalVerifiedBytes	= 0;
+	FIoBuffer	Chunk(1 << 20);
+
+	if (ChunkCount == 0)
+	{
+		UE_LOG(LogIoStoreOnDemand, Log, TEXT("Verify skipped, install cache is empty"));
+		return FIoStatus::Ok;
+	}
+
+	UE_LOG(LogIoStoreOnDemand, Log, TEXT("Verifying %d installed chunks..."), ChunkCount);
+	for (int32 ChunkIndex = 0; const FCasAddrLocation& ChunkLocation : ChunkLocations)
+	{
+		FSharedFileOpenResult OpenResult = Cas.OpenRead(ChunkLocation.Location.BlockId);
+		if (OpenResult.HasError())
+		{
+			UE_LOG(LogIoStoreOnDemand, Error, TEXT("Failed to open block %u for reading"), ChunkLocation.Location.BlockId.Id);
+
+			ReadErrorCount++;
+			ChunkIndex++;
+			continue;
+		}
+
+		int32 ContainerIndex	= INDEX_NONE;
+		int32 EntryIndex		= FindChunkEntry(ChunkLocation.Addr, ContainerIndex);
+
+		if (EntryIndex == INDEX_NONE)
+		{
+			UE_LOG(LogIoStoreOnDemand, Error, TEXT("Failed to find chunk entry for CAS adress '%s'"), *LexToString(ChunkLocation.Addr));
+
+			MissingChunkCount++;
+			ChunkIndex++;
+			continue;
+		}
+
+		const FSharedOnDemandContainer& Container	= Containers[ContainerIndex];
+		const FIoChunkId& ChunkId					= Container->ChunkIds[EntryIndex];
+		const FOnDemandChunkEntry& ChunkEntry		= Container->ChunkEntries[EntryIndex];
+		FSharedFileHandle FileHandle				= OpenResult.GetValue();
+		const int64 ChunkSize						= Align(int64(ChunkEntry.EncodedSize), FAES::AESBlockSize);
+		TotalVerifiedBytes							+= ChunkSize;
+
+		if (int64(Chunk.GetSize()) < ChunkSize)
+		{
+			Chunk = FIoBuffer(ChunkSize);
+		}
+
+		if (FileHandle->Seek(int64(ChunkLocation.Location.BlockOffset)) == false)
+		{
+			UE_LOG(LogIoStoreOnDemand, Error, TEXT("Chunk %d/%d SEEK FAILED, Container='%s', ChunkId='%s', ChunkSize=%lld, Hash='%s', Block=%u, BlockOffset=%u"),
+				ChunkIndex + 1, ChunkCount, *Container->Name, *LexToString(ChunkId), ChunkSize, *LexToString(ChunkEntry.Hash),
+				ChunkLocation.Location.BlockId.Id, ChunkLocation.Location.BlockOffset);
+
+			ReadErrorCount++;
+			ChunkIndex++;
+			continue;
+		}
+
+		if (FileHandle->Read(reinterpret_cast<uint8*>(Chunk.GetData()), ChunkSize) == false)
+		{
+			UE_LOG(LogIoStoreOnDemand, Error, TEXT("Chunk %d/%d READ FAILED, Container='%s', ChunkId='%s', ChunkSize=%lld, Hash='%s', Block=%u, BlockOffset=%u"),
+				ChunkIndex + 1, ChunkCount, *Container->Name, *LexToString(ChunkId), ChunkSize, *LexToString(ChunkEntry.Hash),
+				ChunkLocation.Location.BlockId.Id, ChunkLocation.Location.BlockOffset);
+
+			ReadErrorCount++;
+			ChunkIndex++;
+			continue;
+		}
+
+		const FIoHash ChunkHash = FIoHash::HashBuffer(Chunk.GetView().Left(ChunkSize));
+
+		if (ChunkHash == ChunkEntry.Hash)
+		{
+			UE_LOG(LogIoStoreOnDemand, VeryVerbose, TEXT("Chunk %d/%d OK, Container='%s', ChunkId='%s', ChunkSize=%lld, Hash='%s', Block=%u, BlockOffset=%u"),
+				ChunkIndex + 1, ChunkCount, *Container->Name, *LexToString(ChunkId), ChunkSize, *LexToString(ChunkEntry.Hash),
+				ChunkLocation.Location.BlockId.Id, ChunkLocation.Location.BlockOffset);
+		}
+		else
+		{
+			UE_LOG(LogIoStoreOnDemand, Error, TEXT("Chunk %d/%d CORRUPT, Container='%s', ChunkId='%s', ChunkSize=%lld, Hash='%s', ActualHash='%s', Block=%u, BlockOffset=%u"),
+				ChunkIndex + 1, ChunkCount, *Container->Name, *LexToString(ChunkId), ChunkSize, *LexToString(ChunkEntry.Hash), *LexToString(ChunkHash),
+				ChunkLocation.Location.BlockId.Id, ChunkLocation.Location.BlockOffset);
+
+			CorruptChunkCount++;
+		}
+
+		ChunkIndex++;
+	}
+
+	if (CorruptChunkCount > 0 || MissingChunkCount > 0 || ReadErrorCount > 0)
+	{
+		const FString Reason = FString::Printf(TEXT("Verify install cache failed, Corrupt=%u, Missing=%u, ReadErrors=%u"),
+			CorruptChunkCount, MissingChunkCount, ReadErrorCount);
+
+		if (CorruptChunkCount > 0 || ReadErrorCount > 0)
+		{
+			UE_LOG(LogIoStoreOnDemand, Error, TEXT("%s"), *Reason);
+		}
+		else
+		{
+			UE_LOG(LogIoStoreOnDemand, Warning, TEXT("%s"), *Reason);
+		}
+
+		if (CorruptChunkCount > 0)
+		{
+			return FIoStatus(EIoErrorCode::SignatureError);
+		}
+
+		if (ReadErrorCount > 0)
+		{
+			return FIoStatus(EIoErrorCode::ReadError);
+		}
+
+		return FIoStatus(EIoErrorCode::NotFound);
+	}
+
+	UE_LOG(LogIoStoreOnDemand, Log, TEXT("Successfully verified %d chunk(s) of total %.2lf MiB"),
+		ChunkCount, ToMiB(TotalVerifiedBytes));
+
+	return FIoStatus::Ok;
+}
+
 uint64 FOnDemandInstallCache::AddReferencesToBlocks(
 	const TArray<FSharedOnDemandContainer>& Containers,
 	const TArray<TBitArray<>>& ChunkEntryIndices,
diff --git a/Engine/Source/Runtime/Experimental/IoStore/OnDemand/Private/OnDemandInstallCache.h b/Engine/Source/Runtime/Experimental/IoStore/OnDemand/Private/OnDemandInstallCache.h
--- a/Engine/Source/Runtime/Experimental/IoStore/OnDemand/Private/OnDemandInstallCache.h
+++ b/Engine/Source/Runtime/Experimental/IoStore/OnDemand/Private/OnDemandInstallCache.h
@@ -32,6 +32,7 @@ public:
 	virtual FIoStatus							Purge(TSet<FIoHash>&& ChunksToInstall) = 0;
 	virtual FIoStatus							PurgeAllUnreferenced(bool bDefrag, const uint64* BytesToPurge = nullptr) = 0;
 	virtual FIoStatus							DefragAll(const uint64* BytesToFree = nullptr) = 0;
+	virtual FIoStatus							Verify() = 0;
 	virtual FIoStatus							Flush() = 0;
 	virtual FOnDemandInstallCacheUsage			GetCacheUsage() = 0;
 };
diff --git a/Engine/Source/Runtime/Experimental/IoStore/OnDemand/Private/OnDemandIoStore.cpp b/Engine/Source/Runtime/Experimental/IoStore/OnDemand/Private/OnDemandIoStore.cpp
--- a/Engine/Source/Runtime/Experimental/IoStore/OnDemand/Private/OnDemandIoStore.cpp
+++ b/Engine/Source/Runtime/Experimental/IoStore/OnDemand/Private/OnDemandIoStore.cpp
@@ -577,6 +577,12 @@ void FOnDemandIoStore::Defrag(FOnDemandDefragArgs&& Args, FOnDemandDefragComplet
 	Installer->EnqueueDefragRequest(MoveTemp(Args), MoveTemp(OnCompleted));
 }
 
+void FOnDemandIoStore::Verify(FOnDemandVerifyCacheCompleted&& OnCompleted)
+{
+	LLM_SCOPE_BYNAME(TEXT("FileSystem/OnDemandIoStore"));
+	Installer->EnqueueVerifyRequest(MoveTemp(OnCompleted));
+}
+
 FIoStatus FOnDemandIoStore::Unmount(FStringView MountId)
 {
 	LLM_SCOPE_BYNAME(TEXT("FileSystem/OnDemandIoStore"));
@@ -1477,6 +1483,24 @@ void FOnDemandIoStore::CreateContainersFromToc(
 	}
 }
 
+TArray<FSharedOnDemandContainer> FOnDemandIoStore::GetContainers(EOnDemandContainerFlags ContainerFlags) const
+{
+	TArray<FSharedOnDemandContainer> Out;
+	Out.Reserve(Containers.Num());
+	{
+		UE::TUniqueLock Lock(ContainerMutex);
+		for (const FSharedOnDemandContainer& Container : Containers)
+		{
+			if (ContainerFlags == EOnDemandContainerFlags::None || Container->HasAllFlags(ContainerFlags))
+			{
+				Out.Add(Container);
+			}
+		}
+	}
+
+	return Out;
+}
+
 FIoStatus FOnDemandIoStore::GetContainersAndPackagesForInstall(
 	FStringView MountId,
 	const TArray<FString>& TagSets,
diff --git a/Engine/Source/Runtime/Experimental/IoStore/OnDemand/Private/OnDemandIoStore.h b/Engine/Source/Runtime/Experimental/IoStore/OnDemand/Private/OnDemandIoStore.h
--- a/Engine/Source/Runtime/Experimental/IoStore/OnDemand/Private/OnDemandIoStore.h
+++ b/Engine/Source/Runtime/Experimental/IoStore/OnDemand/Private/OnDemandIoStore.h
@@ -323,6 +323,7 @@ public:
 											FOnDemandInstallProgressed&& OnProgress = nullptr) override;
 	virtual void						Purge(FOnDemandPurgeArgs&& Args, FOnDemandPurgeCompleted&& OnCompleted) override;
 	virtual void						Defrag(FOnDemandDefragArgs&& Args, FOnDemandDefragCompleted&& OnCompleted) override;
+	virtual void						Verify(FOnDemandVerifyCacheCompleted&& OnCompleted) override;
 	virtual FIoStatus					Unmount(FStringView MountId) override;
 	virtual TIoStatusOr<uint64> 		GetInstallSize(const FOnDemandGetInstallSizeArgs& Args) const override;
 	virtual FIoStatus					GetInstallSizesByMountId(const FOnDemandGetInstallSizeArgs& Args, TMap<FString, uint64>& OutSizesByMountId) const override;
@@ -334,6 +335,7 @@ public:
 	virtual void						CancelInstallRequest(FSharedInternalInstallRequest InstallRequest) override;
 	virtual void						UpdateInstallRequestPriority(FSharedInternalInstallRequest InstallRequest, int32 NewPriority) override;
 	virtual void						ReleaseContent(FOnDemandInternalContentHandle& ContentHandle) override;
+	TArray<FSharedOnDemandContainer>	GetContainers(EOnDemandContainerFlags ContainerFlags = EOnDemandContainerFlags::None) const;
 
 private:
 	friend class FOnDemandContentInstaller;
diff --git a/Engine/Source/Runtime/Experimental/IoStore/OnDemandCore/Private/IoStoreOnDemand.cpp b/Engine/Source/Runtime/Experimental/IoStore/OnDemandCore/Private/IoStoreOnDemand.cpp
--- a/Engine/Source/Runtime/Experimental/IoStore/OnDemandCore/Private/IoStoreOnDemand.cpp
+++ b/Engine/Source/Runtime/Experimental/IoStore/OnDemandCore/Private/IoStoreOnDemand.cpp
@@ -358,6 +358,38 @@ static void InstallPackage(const TArray<FString>& Args)
 		});
 }
 
+////////////////////////////////////////////////////////////////////////////////
+static void VerifyCache()
+{
+	UE::IoStore::IOnDemandIoStore* IoStore = UE::IoStore::TryGetOnDemandIoStore();
+	if (IoStore == nullptr)
+	{
+		UE_LOG(LogIoStoreOnDemand, Error, TEXT("Failed to load I/O store on-demand module."));
+		return;
+	}
+
+	IoStore->Verify([](FOnDemandVerifyCacheResult&& VerifyResult)
+	{
+		if (VerifyResult.Status.IsOk())
+		{
+			UE_LOG(LogIoStoreOnDemand, Log, TEXT("Install cache verified OK!"));
+		}
+		else
+		{
+			if (VerifyResult.Status.GetErrorCode() == EIoErrorCode::NotFound)
+			{
+				UE_LOG(LogIoStoreOnDemand, Warning, TEXT("Verify install cache failed, reason '%s'"),
+					*VerifyResult.Status.ToString());
+			}
+			else
+			{
+				UE_LOG(LogIoStoreOnDemand, Error, TEXT("Verify install cache failed, reason '%s'"),
+					*VerifyResult.Status.ToString());
+			}
+		}
+	});
+}
+
 ////////////////////////////////////////////////////////////////////////////////
 static FAutoConsoleCommand PurgeCacheCommand(
 	TEXT("iostore.PurgeInstallCache"),
@@ -392,6 +424,12 @@ static FAutoConsoleCommand InstallPackageCommand(
 	TEXT("<PackageName>"),
 	FConsoleCommandWithArgsDelegate::CreateStatic(InstallPackage));
 
+static FAutoConsoleCommand VerifyCacheCommand(
+	TEXT("iostore.VerifyCache"),
+	TEXT("Verifies the cache chunks against the mounted container TOCs"),
+	FConsoleCommandDelegate::CreateStatic(VerifyCache),
+	ECVF_Cheat);
+
 } // namespace Commands
 #endif // !UE_BUILD_SHIPPING
 
diff --git a/Engine/Source/Runtime/Experimental/IoStore/OnDemandCore/Public/IO/IoStoreOnDemand.h b/Engine/Source/Runtime/Experimental/IoStore/OnDemandCore/Public/IO/IoStoreOnDemand.h
--- a/Engine/Source/Runtime/Experimental/IoStore/OnDemandCore/Public/IO/IoStoreOnDemand.h
+++ b/Engine/Source/Runtime/Experimental/IoStore/OnDemandCore/Public/IO/IoStoreOnDemand.h
@@ -366,7 +366,7 @@ struct FOnDemandDefragResult
 	double DurationInSeconds = 0.0;
 };
 
-/** Purge completion callback */
+/** Defrag completion callback */
 using FOnDemandDefragCompleted = TUniqueFunction<void(FOnDemandDefragResult)>;
 
 /** Options for controlling the behavior of the install size request. */
@@ -448,6 +448,18 @@ struct FOnDemandCacheUsage
 	FOnDemandStreamingCacheUsage StreamingCache;
 };
 
+/** Result from verifying the install cache. */
+struct FOnDemandVerifyCacheResult
+{
+	/** The status of the cache verification request. */
+	FIoStatus Status;
+	/** Duration in seconds. */
+	double DurationInSeconds = 0.0;
+};
+
+/** Verify cache completion callback. */
+using FOnDemandVerifyCacheCompleted = TUniqueFunction<void(FOnDemandVerifyCacheResult)>;
+
 /** Interface for recording analytics over a given time period */
 class IAnalyticsRecording
 {
@@ -484,12 +496,12 @@ public:
 	 * Mount an on-demand container.
 	 * @param	Args		Mount arguments.
 	 * @param	OnCompleted	Completion callback.
-	 **/
+	 */
 	virtual void Mount(FOnDemandMountArgs&& Args, FOnDemandMountCompleted&& OnCompleted) = 0;
 	/**
 	 * Unmount all container(s) associated with the specified mount ID.
 	 * @param	MountId	The Mount ID passed to the mount call.
-	 **/
+	 */
 	virtual FIoStatus Unmount(FStringView MountId) = 0;
 	/**
 	 * Install content.
@@ -497,7 +509,7 @@ public:
 	 * @param	OnCompleted			Completion callback.
 	 * @param	OnProgress			Optional progress callback.
 	 * @param	CancellationToken	Optional cancellation token.
-	 **/
+	 */
 	virtual FOnDemandInstallRequest Install(
 			FOnDemandInstallArgs&& Args,
 			FOnDemandInstallCompleted&& OnCompleted,
@@ -506,24 +518,29 @@ public:
 	 * Purge the cache.
 	 * @param	Args		The purge arguments.
 	 * @param	OnComplete	Completion callback.
-	 **/
+	 */
 	virtual void Purge(FOnDemandPurgeArgs&& Args, FOnDemandPurgeCompleted&& OnCompleted) = 0;
 	/**
 	 * Defrag the cache.
 	 * @param	Args		The defrag arguments.
 	 * @param	OnComplete	Completion callback.
-	 **/
+	 */
 	virtual void Defrag(FOnDemandDefragArgs&& Args, FOnDemandDefragCompleted&& OnCompleted) = 0;
+	/**
+	 * Verify the install cache.
+	 * @param	OnComplete	Completion callback.
+	 */
+	virtual void Verify(FOnDemandVerifyCacheCompleted&& OnCompleted) = 0;
 	/**
 	 * Get the total space in bytes needed to install the specified content.
 	 * @param Args	Arguments for computing required install size.
-	 **/
+	 */
 	virtual TIoStatusOr<uint64> GetInstallSize(const FOnDemandGetInstallSizeArgs& Args) const = 0;
 	/**
 	 * Get the total space in bytes needed to install the specified content.
 	 * @param Args	Arguments for computing required install size.
 	 * @return		The total size by mount ID.
-	 **/
+	 */
 	virtual FIoStatus GetInstallSizesByMountId(const FOnDemandGetInstallSizeArgs& Args, TMap<FString, uint64>& OutSizesByMountId) const = 0;
 	/** Returns the total cache size in bytes. */
 	virtual FOnDemandCacheUsage GetCacheUsage(const FOnDemandGetCacheUsageArgs& Args) const = 0;
