Partial sleeping for skeletal meshes - Bug fix:
- Set solver body position to the COM of the sleeping particle
- Add todo comment about sleeping constraints entering the solver

diff --git a/Engine/Source/Runtime/Experimental/Chaos/Private/Chaos/Character/CharacterGroundConstraintContainerSolver.cpp b/Engine/Source/Runtime/Experimental/Chaos/Private/Chaos/Character/CharacterGroundConstraintContainerSolver.cpp
--- a/Engine/Source/Runtime/Experimental/Chaos/Private/Chaos/Character/CharacterGroundConstraintContainerSolver.cpp
+++ b/Engine/Source/Runtime/Experimental/Chaos/Private/Chaos/Character/CharacterGroundConstraintContainerSolver.cpp
@@ -42,6 +42,11 @@ namespace Chaos
 				FConstraintHandle* Constraint = IslandConstraint->GetConstraint();
 				
 				// Filter out sleeping constraints in any partially sleeping island
+				// @todo(chaos): This is not working correctly since IsSleeping() always returns false for this constraint type.
+				// As a result, the constraint is added to the solver, data is gathered and scattered, and the constraint correction is computed.
+				// Nevertheless, the two connected particles are considered kinematic and, thus, do not "feel" the effect of the constraint.
+				// Hence, the solve remains correct, but we do more work than needed.
+				// In the future, we should store sleep state of for this constraint type as well.
 				if (!Constraint->IsSleeping())
 				{
 					// We will only ever be given constraints from our container (asserts in non-shipping)
diff --git a/Engine/Source/Runtime/Experimental/Chaos/Private/Chaos/Evolution/SolverBodyContainer.cpp b/Engine/Source/Runtime/Experimental/Chaos/Private/Chaos/Evolution/SolverBodyContainer.cpp
--- a/Engine/Source/Runtime/Experimental/Chaos/Private/Chaos/Evolution/SolverBodyContainer.cpp
+++ b/Engine/Source/Runtime/Experimental/Chaos/Private/Chaos/Evolution/SolverBodyContainer.cpp
@@ -14,10 +14,7 @@ namespace Chaos
 		const FPBDRigidParticleHandle* RigidParticle = Particle->CastToRigidParticle();
 		const FKinematicGeometryParticleHandle* KinematicParticle = Particle->CastToKinematicParticle();
 
-		// NOTE: IsDynamic() is true whether the particle is awake or asleep.
-		// In the solver, we want sleeping particles to behave like kinematics
-		// since they do not experience any positional impulse and position correction.
-		if ((RigidParticle != nullptr) && RigidParticle->IsDynamic() && !RigidParticle->IsSleeping())
+		if ((RigidParticle != nullptr) && RigidParticle->IsDynamic())
 		{
 			// Dynamic particle
 			SolverBody.SetX(RigidParticle->XCom());
@@ -26,11 +23,25 @@ namespace Chaos
 			SolverBody.SetQ(RigidParticle->QCom());
 			SolverBody.SetCoM(RigidParticle->CenterOfMass());
 			SolverBody.SetRoM(RigidParticle->RotationOfMass());
-			SolverBody.SetV(RigidParticle->GetV());
-			SolverBody.SetW(RigidParticle->GetW());
-			SolverBody.SetInvM(RigidParticle->InvM());
-			SolverBody.SetInvILocal(RigidParticle->ConditionedInvI());
-			// Note: SetInvILocal also calculates InvI for dynamic
+			if (!RigidParticle->IsSleeping()) // Awake
+			{
+				SolverBody.SetV(RigidParticle->GetV());
+				SolverBody.SetW(RigidParticle->GetW());
+				SolverBody.SetInvM(RigidParticle->InvM());
+				SolverBody.SetInvILocal(RigidParticle->ConditionedInvI());
+				// Note: SetInvILocal also calculates InvI for dynamic
+			}
+			else // Sleeping
+			{
+				// In the solver, we want sleeping particles to behave like kinematics
+				// since they do not experience any positional impulse and position correction.
+				// Thus, we set velocity, mass, and inertia to zero while using the position and orientation of the rigid particle.
+				SolverBody.SetV(FVec3(0));
+				SolverBody.SetW(FVec3(0));
+				SolverBody.SetInvM(0);
+				SolverBody.SetInvILocal(FVec3(0));	// Note: does not set InvI for Kinematic
+				SolverBody.SetInvI(FMatrix33(0));
+			}
 		}
 		else
 		{
diff --git a/Engine/Source/Runtime/Experimental/Chaos/Private/Chaos/Joint/PBDJointContainerSolver.cpp b/Engine/Source/Runtime/Experimental/Chaos/Private/Chaos/Joint/PBDJointContainerSolver.cpp
--- a/Engine/Source/Runtime/Experimental/Chaos/Private/Chaos/Joint/PBDJointContainerSolver.cpp
+++ b/Engine/Source/Runtime/Experimental/Chaos/Private/Chaos/Joint/PBDJointContainerSolver.cpp
@@ -152,6 +152,11 @@ namespace Chaos
 				FConstraintHandle* Constraint = IslandConstraint->GetConstraint();
 
 				// Filter out sleeping constraints in any partially sleeping island
+				// @todo(chaos): This is not working correctly since IsSleeping() always returns false for this constraint type.
+				// As a result, the constraint is added to the solver, data is gathered and scattered, and the constraint correction is computed.
+				// Nevertheless, the two connected particles are considered kinematic and, thus, do not "feel" the effect of the constraint.
+				// Hence, the solve remains correct, but we do more work than needed.
+				// In the future, we should store sleep state of for this constraint type as well.
 				if (!Constraint->IsSleeping())
 				{
 					// We will only ever be given constraints from our container (asserts in non-shipping)
diff --git a/Engine/Source/Runtime/Experimental/Chaos/Public/Chaos/Evolution/IndexedConstraintContainer.h b/Engine/Source/Runtime/Experimental/Chaos/Public/Chaos/Evolution/IndexedConstraintContainer.h
--- a/Engine/Source/Runtime/Experimental/Chaos/Public/Chaos/Evolution/IndexedConstraintContainer.h
+++ b/Engine/Source/Runtime/Experimental/Chaos/Public/Chaos/Evolution/IndexedConstraintContainer.h
@@ -160,6 +160,11 @@ namespace Chaos
 				FConstraintHandle* Constraint = IslandConstraint->GetConstraint();
 
 				// Filter out sleeping constraints in any partially sleeping island
+				// @todo(chaos): This is not working correctly since IsSleeping() always returns false for this constraint type.
+				// As a result, the constraint is added to the solver, data is gathered and scattered, and the constraint correction is computed.
+				// Nevertheless, the two connected particles are considered kinematic and, thus, do not "feel" the effect of the constraint.
+				// Hence, the solve remains correct, but we do more work than needed.
+				// In the future, we should store sleep state of for this constraint type as well.
 				if (!Constraint->IsSleeping())
 				{
 					// We will only ever be given constraints from our container (asserts in non-shipping)
