Constraints: avoid re-entrant evaluations when flushing the evaluation graph

#rnx

diff --git a/Engine/Source/Runtime/Experimental/Animation/Constraints/Private/ConstraintsEvaluationGraph.cpp b/Engine/Source/Runtime/Experimental/Animation/Constraints/Private/ConstraintsEvaluationGraph.cpp
--- a/Engine/Source/Runtime/Experimental/Animation/Constraints/Private/ConstraintsEvaluationGraph.cpp
+++ b/Engine/Source/Runtime/Experimental/Animation/Constraints/Private/ConstraintsEvaluationGraph.cpp
@@ -19,7 +19,7 @@ static bool	bDebugGraph = false;
 static FAutoConsoleVariableRef CVarDebugEvaluationGraph(
 	TEXT("Constraints.DebugEvaluationGraph"),
 	bDebugGraph,
-	TEXT("Print debug info about constraitns evaluation graph.") );
+	TEXT("Print debug info about constraints evaluation graph.") );
 	
 }
 
@@ -191,7 +191,14 @@ void FConstraintsEvaluationGraph::Evaluate(FConstraintNode* InNode)
 		return;
 	}
 
+	if (InNode->bEvaluating)
+	{
+		return;
+	}
+	
 	// evaluate
+	TGuardValue<bool> ReentrantGuard(InNode->bEvaluating, true);
+	
 	if (Constraint->IsFullyActive() && InNode->ConstraintTick->IsTickFunctionRegistered() && InNode->ConstraintTick->IsTickFunctionEnabled())
 	{
 		Constraint->Evaluate();
diff --git a/Engine/Source/Runtime/Experimental/Animation/Constraints/Public/ConstraintsEvaluationGraph.h b/Engine/Source/Runtime/Experimental/Animation/Constraints/Public/ConstraintsEvaluationGraph.h
--- a/Engine/Source/Runtime/Experimental/Animation/Constraints/Public/ConstraintsEvaluationGraph.h
+++ b/Engine/Source/Runtime/Experimental/Animation/Constraints/Public/ConstraintsEvaluationGraph.h
@@ -28,8 +28,11 @@ struct FConstraintNode
 	TSet<uint32> Parents;
 	TSet<uint32> Children;
 
-	/* internal hierarchy representing dependencies at the nodes level (indices here represent the nodes' NodeIndex)*/
+	/* Whether this node needs to be evaluated on the next flush. */
 	bool bMarkedForEvaluation = false;
+	
+	/* Whether this node is currently being evaluated (used to avoid re-entrant evaluations). */
+	bool bEvaluating = false;
 };
 
 /**
