Constraints & new evaluation scheme
- mark the skeletal mesh for evaluation when force-ticking the parent
- added the ability to register post evaluation function to make sure that transform data is up-to-date after the animation is evaluated and when the bones transforms are finalized
- it also ensures that no data is requested while the animation evaluation is currently running on a worker thread

#rnx

diff --git a/Engine/Plugins/Animation/ControlRig/Source/ControlRig/Private/Constraints/ControlRigTransformableHandle.cpp b/Engine/Plugins/Animation/ControlRig/Source/ControlRig/Private/Constraints/ControlRigTransformableHandle.cpp
--- a/Engine/Plugins/Animation/ControlRig/Source/ControlRig/Private/Constraints/ControlRigTransformableHandle.cpp
+++ b/Engine/Plugins/Animation/ControlRig/Source/ControlRig/Private/Constraints/ControlRigTransformableHandle.cpp
@@ -127,17 +127,17 @@ void UTransformableControlHandle::PreEvaluate(const bool bTick) const
 		}
 	}
 
-	if (!bTick)
-	{
-		return ControlRig->Evaluate_AnyThread();
-	}
-	
 	if (TransformableHandleUtils::SkipTicking())
 	{
 		// TODO test UControlRigComponent with FAnimationEvaluationCache 
-		TransformableHandleUtils::EvaluateComponent( GetSkeletalMesh() );
+		TransformableHandleUtils::EvaluateComponent(GetSkeletalMesh(), GetEvaluationTask());
 		return;
 	}
+	
+	if (!bTick)
+	{
+		return ControlRig->Evaluate_AnyThread();
+	}
 
 	// else full tick
 	TickTarget();
@@ -238,11 +238,16 @@ FTransform UTransformableControlHandle::GetGlobalTransform() const
 		return FTransform::Identity;
 	}
 
+	const FTransform& ComponentTransform = BoundComponent->GetComponentTransform();
+	if (TransformableHandleUtils::SkipTicking())
+	{
+		TransformableHandleUtils::EvaluateComponent( GetSkeletalMesh(), GetEvaluationTask() );
+		return Cache * ComponentTransform;
+	}
+
 	const FRigElementKey& ControlKey = ControlElement->GetKey();
 	const URigHierarchy* Hierarchy = ControlRig->GetHierarchy();
 	const int32 CtrlIndex = Hierarchy->GetIndex(ControlKey);
-
-	const FTransform& ComponentTransform = BoundComponent->GetComponentTransform();
 	return Hierarchy->GetGlobalTransform(CtrlIndex) * ComponentTransform;
 }
 
@@ -474,6 +479,46 @@ FRigControlElement* UTransformableControlHandle::GetControlElement() const
 	return ControlRig->FindControl(ControlName);
 }
 
+const UE::Anim::FAnimationEvaluationTask& UTransformableControlHandle::GetEvaluationTask() const
+{
+	const FRigControlElement* Control = GetControlElement();
+	if (!Control)
+	{
+		EvaluationTask = UE::Anim::FAnimationEvaluationTask();
+		return EvaluationTask;
+	}
+
+	USkeletalMeshComponent* SkeletalMeshComponent = GetSkeletalMesh();
+	if (!SkeletalMeshComponent)
+	{
+		EvaluationTask = UE::Anim::FAnimationEvaluationTask();
+		return EvaluationTask;
+	}
+
+	if (!EvaluationTask.Guid.IsValid())
+	{
+		EvaluationTask.Guid = FGuid::NewGuid();
+	}
+
+	if (EvaluationTask.SkeletalMeshComponent != SkeletalMeshComponent)
+	{
+		EvaluationTask.SkeletalMeshComponent = SkeletalMeshComponent;
+
+		URigHierarchy* Hierarchy = ControlRig->GetHierarchy();
+		const int32 CtrlIndex = Hierarchy->GetIndex(Control->GetKey());
+		FTransform &CacheRef = Cache;
+		EvaluationTask.PostEvaluationFunction = [&CacheRef, CtrlIndex, WeakHierarchy = TWeakObjectPtr<URigHierarchy>(Hierarchy)]()
+		{
+			if (WeakHierarchy.IsValid())
+			{
+				CacheRef = WeakHierarchy->GetGlobalTransform(CtrlIndex);
+			}
+		};
+	}
+	
+	return EvaluationTask;
+}
+
 void UTransformableControlHandle::UnregisterDelegates() const
 {
 #if WITH_EDITOR
diff --git a/Engine/Plugins/Animation/ControlRig/Source/ControlRig/Public/Constraints/ControlRigTransformableHandle.h b/Engine/Plugins/Animation/ControlRig/Source/ControlRig/Public/Constraints/ControlRigTransformableHandle.h
--- a/Engine/Plugins/Animation/ControlRig/Source/ControlRig/Public/Constraints/ControlRigTransformableHandle.h
+++ b/Engine/Plugins/Animation/ControlRig/Source/ControlRig/Public/Constraints/ControlRigTransformableHandle.h
@@ -4,6 +4,7 @@
 
 #include "Transform/TransformableHandle.h"
 #include "Rigs/RigHierarchyDefines.h"
+#include "Transform/AnimationEvaluation.h"
 
 #include "ControlRigTransformableHandle.generated.h"
 
@@ -150,6 +151,11 @@ private:
 
 	/** Returns the control that this handle wraps. */
 	FRigControlElement* GetControlElement() const;
+
+	/** Returns an post-evaluation task that can be added to the animation evaluator. */
+	const UE::Anim::FAnimationEvaluationTask& GetEvaluationTask() const;
+	mutable UE::Anim::FAnimationEvaluationTask EvaluationTask;
+	mutable FTransform Cache = FTransform::Identity;
 };
 
 #undef UE_API
diff --git a/Engine/Source/Runtime/Experimental/Animation/Constraints/Private/Transform/AnimationEvaluation.cpp b/Engine/Source/Runtime/Experimental/Animation/Constraints/Private/Transform/AnimationEvaluation.cpp
--- a/Engine/Source/Runtime/Experimental/Animation/Constraints/Private/Transform/AnimationEvaluation.cpp
+++ b/Engine/Source/Runtime/Experimental/Animation/Constraints/Private/Transform/AnimationEvaluation.cpp
@@ -62,6 +62,8 @@ FAnimationEvaluator::~FAnimationEvaluator()
 		}
 		OnBoneTransformsFinalizedHandle.Reset();
 	}
+
+	PostEvaluationTasks.Reset();
 }
 	
 void FAnimationEvaluator::Update(const bool bRefreshBoneTransforms)
@@ -100,6 +102,17 @@ FTransform FAnimationEvaluator::GetGlobalTransform(const FName InSocketName) con
 	return RawSkeletalMeshComponent->GetSocketTransform(InSocketName);
 }
 
+void FAnimationEvaluator::AddPostEvaluationTask(const FAnimationEvaluationTask& InTask)
+{
+	if (USkeletalMeshComponent* RawSkeletalMeshComponent = SkeletalMeshComponent.Get())
+	{
+		if (InTask.IsValid(RawSkeletalMeshComponent))
+		{
+			PostEvaluationTasks.FindOrAdd(InTask.Guid, InTask);
+		}
+	}
+}
+	
 void FAnimationEvaluator::UpdateContext()
 {
 	if (!SkeletalMeshComponent.IsValid())
@@ -155,7 +168,13 @@ void FAnimationEvaluator::RefreshBoneTransforms()
 }
 
 void FAnimationEvaluator::EvaluateAnimation()
-{	
+{
+	USkeletalMeshComponent* RawSkeletalMeshComponent = SkeletalMeshComponent.Get();
+	if (!RawSkeletalMeshComponent)
+	{
+		return;
+	}
+	
 	// see USkeletalMeshComponent:DoInstancePreEvaluation()
 	{
 		if (Context.AnimInstance)
@@ -178,7 +197,7 @@ void FAnimationEvaluator::EvaluateAnimation()
 			bShouldTickAnimation = !Context.AnimInstance->GetUpdateCounter().HasEverBeenUpdated();
 		}
 
-		const bool bBPEnabled = !SkeletalMeshComponent->GetDisablePostProcessBlueprint() && SkeletalMeshComponent->ShouldEvaluatePostProcessAnimBP();
+		const bool bBPEnabled = !RawSkeletalMeshComponent->GetDisablePostProcessBlueprint() && RawSkeletalMeshComponent->ShouldEvaluatePostProcessAnimBP();
 		if (bBPEnabled && Context.PostProcessAnimInstance && !Context.PostProcessAnimInstance->NeedsUpdate())
 		{
 			bShouldTickAnimation |= !Context.PostProcessAnimInstance->GetUpdateCounter().HasEverBeenUpdated();
@@ -189,24 +208,24 @@ void FAnimationEvaluator::EvaluateAnimation()
 			// We bypass TickPose() and call TickAnimation directly, so URO doesn't intercept us.
 			static constexpr float DeltaTime = 0.0f;
 			static constexpr bool bNeedsValidRootMotion = false;
-			SkeletalMeshComponent->TickAnimation(DeltaTime, bNeedsValidRootMotion);
+			RawSkeletalMeshComponent->TickAnimation(DeltaTime, bNeedsValidRootMotion);
 		}
 	}
 
 	// see USkeletalMeshComponent::DoParallelEvaluationTasks_OnGameThread()
 	{
-		Context.ComponentSpaceTransforms = SkeletalMeshComponent->GetEditableComponentSpaceTransforms();
-		Context.CachedComponentSpaceTransforms = SkeletalMeshComponent->GetCachedComponentSpaceTransforms();
-		Context.BoneSpaceTransforms = SkeletalMeshComponent->GetBoneSpaceTransforms();
-		Context.Curve = SkeletalMeshComponent->GetAnimCurves();
-		Context.RootBoneTranslation = SkeletalMeshComponent->RootBoneTranslation;
+		Context.ComponentSpaceTransforms = RawSkeletalMeshComponent->GetEditableComponentSpaceTransforms();
+		Context.CachedComponentSpaceTransforms = RawSkeletalMeshComponent->GetCachedComponentSpaceTransforms();
+		Context.BoneSpaceTransforms = RawSkeletalMeshComponent->GetBoneSpaceTransforms();
+		Context.Curve = RawSkeletalMeshComponent->GetAnimCurves();
+		Context.RootBoneTranslation = RawSkeletalMeshComponent->RootBoneTranslation;
 
 		// note that don't use curves and custom attributes here 
 	}
 		
 	// see USkeletalMeshComponent::ParallelAnimationEvaluation()
 	{
-		SkeletalMeshComponent->PerformAnimationProcessing(
+		RawSkeletalMeshComponent->PerformAnimationProcessing(
 		   Context.SkeletalMesh,
 		   Context.AnimInstance,
 		   Context.bDoEvaluation,
@@ -217,6 +236,24 @@ void FAnimationEvaluator::EvaluateAnimation()
 		   Context.Curve,
 		   Context.CustomAttributes);
 	}
+
+	// call post evaluation tasks if any
+	for (auto PostEvaluationTaskIt = PostEvaluationTasks.CreateIterator(); PostEvaluationTaskIt; ++PostEvaluationTaskIt)
+	{
+		const FAnimationEvaluationTask& PostEvaluationTask = PostEvaluationTaskIt.Value();
+		if (!PostEvaluationTask.IsValid(RawSkeletalMeshComponent))
+		{
+			PostEvaluationTaskIt.RemoveCurrent();
+		}
+		else
+		{
+			if (bLogNewEvaluation)
+			{
+				UE_LOG(LogTemp, Warning, TEXT("EvaluateTask %s"), *PostEvaluationTask.Guid.ToString() );
+			}
+			PostEvaluationTask.PostEvaluationFunction();
+		}
+	}
 }
 
 void FAnimationEvaluator::BoneTransformsFinalized()
@@ -226,10 +263,12 @@ void FAnimationEvaluator::BoneTransformsFinalized()
 		return;
 	}
 
+	USkeletalMeshComponent* RawSkeletalMeshComponent = SkeletalMeshComponent.Get();
+
 	bool bFinalized = false;
 
 	// NOTE FAnimationEvaluator could store a TBitArray of bones being requested and only check them instead of checking all of them 
-	const TArrayView<const FTransform> BonesTransforms = SkeletalMeshComponent->GetComponentSpaceTransforms();
+	const TArrayView<const FTransform> BonesTransforms = RawSkeletalMeshComponent->GetComponentSpaceTransforms();
 	if (Context.ComponentSpaceTransforms.Num() == BonesTransforms.Num())
 	{
 		for (int32 Index = 0; Index < BonesTransforms.Num(); ++Index)
@@ -243,6 +282,24 @@ void FAnimationEvaluator::BoneTransformsFinalized()
 		}
 	}
 
+	// call post evaluation tasks if any
+	for (auto PostEvaluationTaskIt = PostEvaluationTasks.CreateIterator(); PostEvaluationTaskIt; ++PostEvaluationTaskIt)
+	{
+		const FAnimationEvaluationTask& PostEvaluationTask = PostEvaluationTaskIt.Value();
+		if (!PostEvaluationTask.IsValid(RawSkeletalMeshComponent))
+		{
+			PostEvaluationTaskIt.RemoveCurrent();
+		}
+		else
+		{
+			if (bLogNewEvaluation)
+			{
+				UE_LOG(LogTemp, Warning, TEXT("FinalizeTask %s"), *PostEvaluationTask.Guid.ToString() );
+			}
+			PostEvaluationTask.PostEvaluationFunction();
+		}
+	}
+	
 	if (bLogNewEvaluation && bFinalized)
 	{
 		UE_LOG(LogTemp, Warning, TEXT("Finalized %s's evaluator."), *GetComponentName(SkeletalMeshComponent.Get()));
@@ -320,6 +377,43 @@ const FAnimationEvaluator& FAnimationEvaluationCache::GetEvaluator(USkeletalMesh
 	return NewEvaluator;
 }
 
+const FAnimationEvaluator& FAnimationEvaluationCache::GetEvaluator(USkeletalMeshComponent* InSkeletalMeshComponent,
+	const FAnimationEvaluationTask& InTask)
+{
+	if (!InSkeletalMeshComponent)
+	{
+		static const FAnimationEvaluator Invalid(nullptr);
+		return Invalid;
+	}
+	
+	constexpr bool bRefreshBones = true;
+
+	if (FAnimationEvaluator* Found = PerSkeletalMeshEvaluator.Find(InSkeletalMeshComponent))
+	{
+		Found->AddPostEvaluationTask(InTask);
+		if (!Found->IsValid())
+		{
+			if (bLogNewEvaluation)
+			{
+				UE_LOG(LogTemp, Warning, TEXT("Update %s's evaluator for evaluation."), *GetComponentName(InSkeletalMeshComponent));
+			}
+			Found->Update(bRefreshBones);
+		}
+		return *Found;
+	}
+	
+	FAnimationEvaluator& NewEvaluator = PerSkeletalMeshEvaluator.Emplace(InSkeletalMeshComponent, InSkeletalMeshComponent);
+	NewEvaluator.AddPostEvaluationTask(InTask);
+	NewEvaluator.Update(bRefreshBones);
+
+	if (bLogNewEvaluation)
+	{
+		UE_LOG(LogTemp, Warning, TEXT("Create new evaluator for %s."), *GetComponentName(InSkeletalMeshComponent));
+	}
+	
+	return NewEvaluator;
+}
+
 void FAnimationEvaluationCache::RegisterNotifications()
 {
 	FConstraintsManagerController& Controller = FConstraintsManagerController::Get(nullptr);
diff --git a/Engine/Source/Runtime/Experimental/Animation/Constraints/Private/Transform/TransformConstraint.cpp b/Engine/Source/Runtime/Experimental/Animation/Constraints/Private/Transform/TransformConstraint.cpp
--- a/Engine/Source/Runtime/Experimental/Animation/Constraints/Private/Transform/TransformConstraint.cpp
+++ b/Engine/Source/Runtime/Experimental/Animation/Constraints/Private/Transform/TransformConstraint.cpp
@@ -434,9 +434,19 @@ void UTickableTransformConstraint::Evaluate(bool bTickHandlesAlso) const
 {
 	if (IsFullyActive())
 	{
-		if (bTickHandlesAlso && !TransformableHandleUtils::SkipTicking())
+		if (bTickHandlesAlso)
 		{
-			ParentTRSHandle->TickTarget();
+			if (TransformableHandleUtils::SkipTicking())
+			{
+				if (USkeletalMeshComponent* SkeletalMeshComponent = Cast<USkeletalMeshComponent>(ParentTRSHandle->GetTarget()))
+				{
+					TransformableHandleUtils::MarkComponentForEvaluation(SkeletalMeshComponent);
+				}
+			}
+			else
+			{
+				ParentTRSHandle->TickTarget();
+			}
 		}
 		Super::Evaluate();
 	}
diff --git a/Engine/Source/Runtime/Experimental/Animation/Constraints/Private/Transform/TransformableHandleUtils.cpp b/Engine/Source/Runtime/Experimental/Animation/Constraints/Private/Transform/TransformableHandleUtils.cpp
--- a/Engine/Source/Runtime/Experimental/Animation/Constraints/Private/Transform/TransformableHandleUtils.cpp
+++ b/Engine/Source/Runtime/Experimental/Animation/Constraints/Private/Transform/TransformableHandleUtils.cpp
@@ -111,5 +111,16 @@ const UE::Anim::FAnimationEvaluator& EvaluateComponent(USceneComponent* InSceneC
 	static const UE::Anim::FAnimationEvaluator Invalid(nullptr);
 	return Invalid;
 }
-	
+
+CONSTRAINTS_API const UE::Anim::FAnimationEvaluator& EvaluateComponent(USceneComponent* InSceneComponent, const UE::Anim::FAnimationEvaluationTask& InTask)
+{
+	if (USkeletalMeshComponent* SkeletalMeshComponent = Cast<USkeletalMeshComponent>(InSceneComponent))
+	{
+		return UE::Anim::FAnimationEvaluationCache::Get().GetEvaluator(SkeletalMeshComponent, InTask);
+	}
+
+	static const UE::Anim::FAnimationEvaluator Invalid(nullptr);
+	return Invalid;
+}
+
 }
\ No newline at end of file
diff --git a/Engine/Source/Runtime/Experimental/Animation/Constraints/Public/Transform/AnimationEvaluation.h b/Engine/Source/Runtime/Experimental/Animation/Constraints/Public/Transform/AnimationEvaluation.h
--- a/Engine/Source/Runtime/Experimental/Animation/Constraints/Public/Transform/AnimationEvaluation.h
+++ b/Engine/Source/Runtime/Experimental/Animation/Constraints/Public/Transform/AnimationEvaluation.h
@@ -11,6 +11,23 @@
 namespace UE::Anim
 {
 
+/**
+ * FAnimationEvaluationTask allows post-evaluation functions to be added to an evaluator, so that data not cached in FAnimationEvaluationContext
+ * can be handled by external mechanisms.
+*/
+	
+struct FAnimationEvaluationTask
+{
+	bool IsValid(const USkeletalMeshComponent* InSkeletalMeshComponent) const
+	{
+		return Guid.IsValid() && PostEvaluationFunction.IsSet() && InSkeletalMeshComponent && InSkeletalMeshComponent == SkeletalMeshComponent;
+	}
+
+	FGuid Guid;
+	TWeakObjectPtr<USkeletalMeshComponent> SkeletalMeshComponent = nullptr;
+	TFunction<void()> PostEvaluationFunction;
+};
+	
 /**
  * FAnimationEvaluator allows evaluating the animation of a skeletal mesh without having to tick it directly.
  * USkeletalMeshComponent::TickAnimation & USkeletalMeshComponent::RefreshBoneTransforms update internal data, in particular bones transforms,
@@ -34,6 +51,9 @@ struct FAnimationEvaluator
 	/** Returns InSocketName component space transform composed with the skeletal mesh component's global transform. */
 	FTransform GetGlobalTransform(const FName InSocketName) const;
 
+	/** Adds a post-evaluation task to this evaluator if it has not already been added. */
+	void AddPostEvaluationTask(const FAnimationEvaluationTask& InTask);
+	
 private:
 
 	FAnimationEvaluator() = default;
@@ -53,6 +73,7 @@ private:
 	TWeakObjectPtr<USkeletalMeshComponent> SkeletalMeshComponent = nullptr;
 	FAnimationEvaluationContext Context;
 	FDelegateHandle OnBoneTransformsFinalizedHandle;
+	TMap<FGuid, FAnimationEvaluationTask> PostEvaluationTasks;
 };
 
 /**
@@ -75,6 +96,12 @@ public:
 	 * Note that calling this function will actually evaluate the animation if the evaluator has been previously dirtied.
 	 */
 	const FAnimationEvaluator& GetEvaluator(USkeletalMeshComponent* InSkeletalMeshComponent);
+	
+	/**
+	 * Returns and up-to-date evaluator to be queried and adds the post-evaluation task if it has not already been added.
+	 * Note that calling this function will actually evaluate the animation and the post-evaluation tasks if the evaluator has been previously dirtied.
+	 */
+	const FAnimationEvaluator& GetEvaluator(USkeletalMeshComponent* InSkeletalMeshComponent, const FAnimationEvaluationTask& InTask);
 
 private:
 	FAnimationEvaluationCache() = default;
diff --git a/Engine/Source/Runtime/Experimental/Animation/Constraints/Public/Transform/TransformableHandleUtils.h b/Engine/Source/Runtime/Experimental/Animation/Constraints/Public/Transform/TransformableHandleUtils.h
--- a/Engine/Source/Runtime/Experimental/Animation/Constraints/Public/Transform/TransformableHandleUtils.h
+++ b/Engine/Source/Runtime/Experimental/Animation/Constraints/Public/Transform/TransformableHandleUtils.h
@@ -24,4 +24,9 @@ namespace TransformableHandleUtils
 	
 	/** Returns an updated version of InSceneComponent's animation evaluator. */
 	CONSTRAINTS_API const UE::Anim::FAnimationEvaluator& EvaluateComponent(USceneComponent* InSceneComponent);
+	
+	/** Returns an updated version of InSceneComponent's animation evaluator and adds the input post-evaluation task if not alread added. */
+	CONSTRAINTS_API const UE::Anim::FAnimationEvaluator& EvaluateComponent(
+		USceneComponent* InSceneComponent,
+		const UE::Anim::FAnimationEvaluationTask& InTask);
 }
