Dataflow : fix asserts on the Branch node when some inputs are not connected

diff --git a/Engine/Source/Runtime/Experimental/Dataflow/Core/Private/Dataflow/DataflowCoreNodes.cpp b/Engine/Source/Runtime/Experimental/Dataflow/Core/Private/Dataflow/DataflowCoreNodes.cpp
--- a/Engine/Source/Runtime/Experimental/Dataflow/Core/Private/Dataflow/DataflowCoreNodes.cpp
+++ b/Engine/Source/Runtime/Experimental/Dataflow/Core/Private/Dataflow/DataflowCoreNodes.cpp
@@ -59,8 +59,8 @@ void FDataflowBranchNode::Evaluate(UE::Dataflow::FContext& Context, const FDataf
 		}
 		else
 		{
-			// TODO : throw an invalid type error when context error is available 
-			// Context.Error(TEXT("Both True and False Inputs must be connected"));
+			// set a null value so that the connected system are getting a default value
+			Out->SetNullValue(Context);
 		}
 	}
 }
@@ -96,8 +96,8 @@ void FDataflowSelectNode::Evaluate(UE::Dataflow::FContext& Context, const FDataf
 			}
 			else
 			{
-				// TODO : throw an invalid type error when context error is available 
-				// Context.Error(TEXT("Both True and False Inputs must be connected"));
+				// set a null value so that the connected system are getting a default value
+				Out->SetNullValue(Context);
 			}
 		}
 	}
diff --git a/Engine/Source/Runtime/Experimental/Dataflow/Core/Public/Dataflow/DataflowCoreNodes.h b/Engine/Source/Runtime/Experimental/Dataflow/Core/Public/Dataflow/DataflowCoreNodes.h
--- a/Engine/Source/Runtime/Experimental/Dataflow/Core/Public/Dataflow/DataflowCoreNodes.h
+++ b/Engine/Source/Runtime/Experimental/Dataflow/Core/Public/Dataflow/DataflowCoreNodes.h
@@ -41,10 +41,10 @@ public:
 	virtual void Evaluate(UE::Dataflow::FContext& Context, const FDataflowOutput* Out) const;
 
 public:
-	UPROPERTY(meta = (DataflowInput, DisplayName = "TrueValue"))
+	UPROPERTY(meta = (DataflowInput, DataflowIntrinsic, DisplayName = "TrueValue"))
 	FDataflowAnyType TrueValue;
 
-	UPROPERTY(meta = (DataflowInput, DisplayName = "FalseValue"))
+	UPROPERTY(meta = (DataflowInput, DataflowIntrinsic, DisplayName = "FalseValue"))
 	FDataflowAnyType FalseValue;
 
 	UPROPERTY(EditAnywhere, Category="Condition", meta = (DataflowInput, DisplayName = "Condition"))
