Add new p.Chaos.MinParallelTaskSize cvar and CVarMinParallelTaskSize function to set a threshold of tasks to run in parallel, default value uses old behavior of always parallel.
Change chaos solver to check this before spawning and waiting on tasks.

diff --git a/Engine/Source/Runtime/Experimental/Chaos/Private/Chaos/Collision/CollisionConstraintAllocator.cpp b/Engine/Source/Runtime/Experimental/Chaos/Private/Chaos/Collision/CollisionConstraintAllocator.cpp
--- a/Engine/Source/Runtime/Experimental/Chaos/Private/Chaos/Collision/CollisionConstraintAllocator.cpp
+++ b/Engine/Source/Runtime/Experimental/Chaos/Private/Chaos/Collision/CollisionConstraintAllocator.cpp
@@ -91,7 +91,7 @@ namespace Chaos
 			}
 		}
 
-		void FCollisionConstraintAllocator::ProcessNewNonCcdConstraints(const TArray<FPBDCollisionConstraint*>& NewConstraints, TArray<TArray<FPBDCollisionConstraint*>>& CCDConstraintsBatch)
+		void FCollisionConstraintAllocator::ProcessNewNonCcdConstraints(const TArray<FPBDCollisionConstraint*>& NewConstraints, TArray<FPBDCollisionConstraint*>& AllCCDConstraints)
 		{
 			int32 NumTasks = !!GSingleThreadedPhysics ? 1 : FMath::Min(FTaskGraphInterface::Get().GetNumWorkerThreads(), Chaos::MaxNumWorkers);
 			NumTasks = FMath::Max(FMath::Min(NumTasks, NewConstraints.Num()), 1);
@@ -99,36 +99,62 @@ namespace Chaos
 			const int32 NumConstraints = NewConstraints.Num();
 			const int32 ConstraintByTask = FMath::Max(FMath::DivideAndRoundUp(NumConstraints, NumTasks), 1);
 			const int32 NumBatches = FMath::DivideAndRoundUp(NumConstraints, ConstraintByTask);
-			TArray<UE::Tasks::FTask> PendingTasks;
-			PendingTasks.Reserve(NumBatches);
 
-			CCDConstraintsBatch.SetNum(NumBatches);
-
-			for (int32 BatchIndex = 0; BatchIndex < NumBatches; BatchIndex++)
+			if (Chaos::ShouldExecuteTasksInParallel(NumBatches))
 			{
-				const int32 StartIndex = BatchIndex * ConstraintByTask;
-				int32 EndIndex = (BatchIndex + 1) * ConstraintByTask;
-				EndIndex = BatchIndex == NumBatches - 1 ? FMath::Min(NumConstraints, EndIndex) : EndIndex;
+				TArray<UE::Tasks::FTask> PendingTasks;
+				PendingTasks.Reserve(NumBatches);
 
-				TArray<FPBDCollisionConstraint*>& CCDConstraints = CCDConstraintsBatch[BatchIndex];
+				TArray<TArray<FPBDCollisionConstraint*>> CCDConstraintsBatch;
+				CCDConstraintsBatch.SetNum(NumBatches);
 
-				UE::Tasks::FTask PendingTask = UE::Tasks::Launch(UE_SOURCE_LOCATION, [this, StartIndex, EndIndex, &NewConstraints, &CCDConstraints]
+				for (int32 BatchIndex = 0; BatchIndex < NumBatches; BatchIndex++)
 				{
-					QUICK_SCOPE_CYCLE_COUNTER(STAT_Collisions_SubProcessNewConstraints);
-					for (int32 Index = StartIndex; Index < EndIndex; Index++)
+					const int32 StartIndex = BatchIndex * ConstraintByTask;
+					int32 EndIndex = (BatchIndex + 1) * ConstraintByTask;
+					EndIndex = BatchIndex == NumBatches - 1 ? FMath::Min(NumConstraints, EndIndex) : EndIndex;
+
+					TArray<FPBDCollisionConstraint*>& CCDConstraints = CCDConstraintsBatch[BatchIndex];
+
+					UE::Tasks::FTask PendingTask = UE::Tasks::Launch(UE_SOURCE_LOCATION, [this, StartIndex, EndIndex, &NewConstraints, &CCDConstraints]
 					{
-						FPBDCollisionConstraint* Constraint = NewConstraints[Index];
-						SetConstraintCookie(Constraint, Index);
-						// If the constraint uses CCD, keep it in another list so we don't have to search the full list for CCD contacts
-						if (Constraint->GetCCDEnabled())
+						QUICK_SCOPE_CYCLE_COUNTER(STAT_Collisions_SubProcessNewConstraints);
+						for (int32 Index = StartIndex; Index < EndIndex; Index++)
 						{
-							CCDConstraints.Add(Constraint);
+							FPBDCollisionConstraint* Constraint = NewConstraints[Index];
+							SetConstraintCookie(Constraint, Index);
+							// If the constraint uses CCD, keep it in another list so we don't have to search the full list for CCD contacts
+							if (Constraint->GetCCDEnabled())
+							{
+								CCDConstraints.Add(Constraint);
+							}
 						}
+					}, UE::Tasks::ETaskPriority::High);
+					PendingTasks.Add(PendingTask);
+				}
+				UE::Tasks::Wait(PendingTasks);
+
+				// Single Threaded and protect the insertion order for determinism 
+				for (int32 BatchIndex = 0; BatchIndex < CCDConstraintsBatch.Num(); BatchIndex++)
+				{
+					AllCCDConstraints.Append(CCDConstraintsBatch[BatchIndex]);
+				}
+			}
+			else
+			{
+				// Small enough task to run immediately and avoid parallel overhead
+				QUICK_SCOPE_CYCLE_COUNTER(STAT_Collisions_SubProcessNewConstraints);
+				for (int32 Index = 0; Index < NumConstraints; Index++)
+				{
+					FPBDCollisionConstraint* Constraint = NewConstraints[Index];
+					SetConstraintCookie(Constraint, Index);
+					// If the constraint uses CCD, keep it in another list so we don't have to search the full list for CCD contacts
+					if (Constraint->GetCCDEnabled())
+					{
+						AllCCDConstraints.Add(Constraint);
 					}
-				}, UE::Tasks::ETaskPriority::High);
-				PendingTasks.Add(PendingTask);
+				}
 			}
-			UE::Tasks::Wait(PendingTasks);
 		}
 
 		void FCollisionConstraintAllocator::ProcessNewCcdConstraints()
@@ -139,29 +165,43 @@ namespace Chaos
 			const int32 NumConstraints = ActiveCCDConstraints.Num();
 			const int32 ConstraintByTask = FMath::Max(FMath::DivideAndRoundUp(NumConstraints, NumTasks), 1);
 			const int32 NumBatches = FMath::DivideAndRoundUp(NumConstraints, ConstraintByTask);
-			TArray<UE::Tasks::FTask> PendingTasks;
-			PendingTasks.Reserve(NumBatches);
 
-			for (int32 BatchIndex = 0; BatchIndex < NumBatches; BatchIndex++)
+			if (Chaos::ShouldExecuteTasksInParallel(NumBatches))
 			{
-				const int32 StartIndex = BatchIndex * ConstraintByTask;
-				int32 EndIndex = (BatchIndex + 1) * ConstraintByTask;
-				EndIndex = BatchIndex == NumBatches - 1 ? FMath::Min(NumConstraints, EndIndex) : EndIndex;
+				TArray<UE::Tasks::FTask> PendingTasks;
+				PendingTasks.Reserve(NumBatches);
+
+				for (int32 BatchIndex = 0; BatchIndex < NumBatches; BatchIndex++)
+				{
+					const int32 StartIndex = BatchIndex * ConstraintByTask;
+					int32 EndIndex = (BatchIndex + 1) * ConstraintByTask;
+					EndIndex = BatchIndex == NumBatches - 1 ? FMath::Min(NumConstraints, EndIndex) : EndIndex;
 
-				UE::Tasks::FTask PendingTask = UE::Tasks::Launch(UE_SOURCE_LOCATION, [this, StartIndex, EndIndex]
+					UE::Tasks::FTask PendingTask = UE::Tasks::Launch(UE_SOURCE_LOCATION, [this, StartIndex, EndIndex]
 					{
 						QUICK_SCOPE_CYCLE_COUNTER(STAT_Collisions_SubCCDProcessNewConstraints);
 						for (int32 Index = StartIndex; Index < EndIndex; Index++)
 						{
 							FPBDCollisionConstraint* Constraint = ActiveCCDConstraints[Index];
 							Constraint->GetContainerCookie().CCDConstraintIndex = Index;
-
 						}
 					}, UE::Tasks::ETaskPriority::High);
-				PendingTasks.Add(PendingTask);
+					PendingTasks.Add(PendingTask);
+				}
+				ContextAllocators[0]->NewActiveConstraints.Reset();
+				UE::Tasks::Wait(PendingTasks);
+			}
+			else
+			{
+				// Small enough task to run immediately and avoid parallel overhead
+				QUICK_SCOPE_CYCLE_COUNTER(STAT_Collisions_SubCCDProcessNewConstraints);
+				for (int32 Index = 0; Index < NumConstraints; Index++)
+				{
+					FPBDCollisionConstraint* Constraint = ActiveCCDConstraints[Index];
+					Constraint->GetContainerCookie().CCDConstraintIndex = Index;
+				}
+				ContextAllocators[0]->NewActiveConstraints.Reset();
 			}
-			ContextAllocators[0]->NewActiveConstraints.Reset();
-			UE::Tasks::Wait(PendingTasks);
 		}
 
 		void FCollisionConstraintAllocator::ProcessNewConstraints()
@@ -174,28 +214,28 @@ namespace Chaos
 				const TArray<FPBDCollisionConstraint*>& NewConstraints = ContextAllocators[0]->NewActiveConstraints;
 				ActiveConstraints.Append(NewConstraints);
 
-				TArray<TArray<FPBDCollisionConstraint*>> CCDConstraintsBatch;
-				ProcessNewNonCcdConstraints(NewConstraints, CCDConstraintsBatch);
-
-				// Single Threaded and protect the insertion order for determinism 
-				for (int32 BatchIndex = 0; BatchIndex < CCDConstraintsBatch.Num(); BatchIndex++)
-				{
-					ActiveCCDConstraints.Append(CCDConstraintsBatch[BatchIndex]);
-				}
-
+				ProcessNewNonCcdConstraints(NewConstraints, ActiveCCDConstraints);
 				ProcessNewCcdConstraints();
 			}
 		}
 
 		void FCollisionConstraintAllocator::ProcessNewItems()
 		{
-			UE::Tasks::FTask PendingTask = UE::Tasks::Launch(UE_SOURCE_LOCATION, [this] 
+			if (Chaos::ShouldExecuteTasksInParallel(2))
 			{
-				ProcessNewMidPhases();
-			}, UE::Tasks::ETaskPriority::High);
+				UE::Tasks::FTask PendingTask = UE::Tasks::Launch(UE_SOURCE_LOCATION, [this] 
+				{
+					ProcessNewMidPhases();
+				}, UE::Tasks::ETaskPriority::High);
 
-			ProcessNewConstraints();
-			UE::Tasks::Wait(PendingTask);
+				ProcessNewConstraints();
+				UE::Tasks::Wait(PendingTask);
+			}
+			else
+			{
+				ProcessNewMidPhases();
+				ProcessNewConstraints();
+			}
 		}
 
 		void FCollisionConstraintAllocator::EndDetectCollisions()
diff --git a/Engine/Source/Runtime/Experimental/Chaos/Private/Chaos/Framework/Parallel.cpp b/Engine/Source/Runtime/Experimental/Chaos/Private/Chaos/Framework/Parallel.cpp
--- a/Engine/Source/Runtime/Experimental/Chaos/Private/Chaos/Framework/Parallel.cpp
+++ b/Engine/Source/Runtime/Experimental/Chaos/Private/Chaos/Framework/Parallel.cpp
@@ -13,6 +13,7 @@ namespace Chaos
 	CHAOS_API int32 MaxNumWorkers = 100;
 	CHAOS_API int32 SmallBatchSize = 10;
 	CHAOS_API int32 LargeBatchSize = 100;
+	CHAOS_API int32 MinParallelTaskSize = 0;
 	
 	CHAOS_API bool bDisablePhysicsParallelFor = false;
 	CHAOS_API bool bDisableParticleParallelFor = false;
@@ -26,6 +27,7 @@ namespace Chaos
 	FAutoConsoleVariableRef CVarMaxRangeBatchWorkers(TEXT("p.Chaos.MaxNumWorkers"), MaxNumWorkers, TEXT("Set the max number of workers for physics"));
 	FAutoConsoleVariableRef CVarSmallBatchSize(TEXT("p.Chaos.SmallBatchSize"), SmallBatchSize, TEXT("Small batch size for chaos parallel loops"));
 	FAutoConsoleVariableRef CVarLargeBatchSize(TEXT("p.Chaos.LargeBatchSize"), LargeBatchSize, TEXT("Large batch size for chaos parallel loops"));
+	FAutoConsoleVariableRef CVarMinParallelTaskSize(TEXT("p.Chaos.MinParallelTaskSize"), MinParallelTaskSize, TEXT("Minimum number of tasks required to start parallel task execution on worker threads, set this higher to encourage single threaded execution"));
 }
 
 void Chaos::InnerPhysicsParallelFor(int32 Num, TFunctionRef<void(int32)> InCallable, bool bForceSingleThreaded)
@@ -194,6 +196,26 @@ void Chaos::PhysicsParallelForWithContext(int32 InNum, TFunctionRef<int32 (int32
 	::ParallelForWithTaskContext(TEXT("PhysicsParallelForWithContext"), Contexts, InNum, MinBatchSize, InContextCreator, PassThrough, Flags);
 }
 
+bool Chaos::ShouldExecuteTasksInParallel(int32 InTaskNum)
+{
+	if (GSingleThreadedPhysics)
+	{
+		// Single threaded physics should not run in parallel
+		return false;
+	}
+	else if (InTaskNum >= MinParallelTaskSize)
+	{
+		// Meets threshold so run in parallel
+		return true;
+	}
+	else if (IsInGameThread())
+	{
+		// Game thread should always use workers
+		return true;
+	}
+
+	return false;
+}
 
 //class FRecursiveDivideTask
 //{
diff --git a/Engine/Source/Runtime/Experimental/Chaos/Private/Chaos/PBDRigidsEvolution.cpp b/Engine/Source/Runtime/Experimental/Chaos/Private/Chaos/PBDRigidsEvolution.cpp
--- a/Engine/Source/Runtime/Experimental/Chaos/Private/Chaos/PBDRigidsEvolution.cpp
+++ b/Engine/Source/Runtime/Experimental/Chaos/Private/Chaos/PBDRigidsEvolution.cpp
@@ -632,17 +632,19 @@ namespace Chaos
 		const int32 ParticleByContext = FMath::Max(FMath::DivideAndRoundUp(NumParticles, NumWorker), SmallBatchSize);
 		const int32 NumBatch = FMath::DivideAndRoundUp(NumParticles, ParticleByContext);
 
-		TArray<TArray<FPendingSpatialData>> NewPendingDataArray;
-		NewPendingDataArray.SetNum(NumBatch);
-		TArray<UE::Tasks::FTask> PendingTasks;
-		PendingTasks.Reserve(NumBatch);
-		for (int32 BatchIndex = 0; BatchIndex < NumBatch; BatchIndex++)
-		{
-			const int32 FirstIndex = BatchIndex * ParticleByContext;
-			int32 LastIndex = (BatchIndex + 1) * ParticleByContext;
-			LastIndex = BatchIndex == NumBatch - 1 ? FMath::Min(NumParticles, LastIndex) : LastIndex;
-			TArray<FPendingSpatialData>& NewPendingData = NewPendingDataArray[BatchIndex];
-			UE::Tasks::FTask PendingTask = UE::Tasks::Launch(UE_SOURCE_LOCATION, [this, &AccelerationStructure, FirstIndex, LastIndex, &NewPendingData]()
+		if (Chaos::ShouldExecuteTasksInParallel(NumBatch))
+		{
+			TArray<TArray<FPendingSpatialData>> NewPendingDataArray;
+			NewPendingDataArray.SetNum(NumBatch);
+			TArray<UE::Tasks::FTask> PendingTasks;
+			PendingTasks.Reserve(NumBatch);
+			for (int32 BatchIndex = 0; BatchIndex < NumBatch; BatchIndex++)
+			{
+				const int32 FirstIndex = BatchIndex * ParticleByContext;
+				int32 LastIndex = (BatchIndex + 1) * ParticleByContext;
+				LastIndex = BatchIndex == NumBatch - 1 ? FMath::Min(NumParticles, LastIndex) : LastIndex;
+				TArray<FPendingSpatialData>& NewPendingData = NewPendingDataArray[BatchIndex];
+				UE::Tasks::FTask PendingTask = UE::Tasks::Launch(UE_SOURCE_LOCATION, [this, &AccelerationStructure, FirstIndex, LastIndex, &NewPendingData]()
 				{
 					QUICK_SCOPE_CYCLE_COUNTER(STAT_PruneInternalPendingDataTask);
 					for (int32 Index = FirstIndex; Index < LastIndex; Index++)
@@ -657,16 +659,36 @@ namespace Chaos
 						}
 					}
 				});
-			PendingTasks.Add(PendingTask);
-		}
+				PendingTasks.Add(PendingTask);
+			}
 
-		UE::Tasks::Wait(PendingTasks);
-		// Once all threads are finished merge all data in the acceleration queue
-		// Maybe this could be multi threaded, but since it is really fast not sure if it is worth it.
-		InternalAccelerationQueue.PendingData.Reset();
-		for (int32 BatchIndex = 0; BatchIndex < NumBatch; BatchIndex++)
+			UE::Tasks::Wait(PendingTasks);
+			// Once all threads are finished merge all data in the acceleration queue
+			// Maybe this could be multi threaded, but since it is really fast not sure if it is worth it.
+			InternalAccelerationQueue.PendingData.Reset();
+			for (int32 BatchIndex = 0; BatchIndex < NumBatch; BatchIndex++)
+			{
+				InternalAccelerationQueue.PendingData.Append(NewPendingDataArray[BatchIndex]);
+			}
+		}
+		else
 		{
-			InternalAccelerationQueue.PendingData.Append(NewPendingDataArray[BatchIndex]);
+			// Small enough task to run immediately and avoid parallel overhead
+			QUICK_SCOPE_CYCLE_COUNTER(STAT_PruneInternalPendingDataTask);
+
+			InternalAccelerationQueue.PendingData.Reset();
+
+			for (int32 Index = 0; Index < NumParticles; Index++)
+			{
+				const FPendingSpatialData& PendingData = InternalAccelerationQueue.PendingData[Index];
+				FGeometryParticleHandle* UpdateParticle = PendingData.AccelerationHandle.GetGeometryParticleHandle_PhysicsThread();
+
+				if (PendingData.Operation == EPendingSpatialDataOperation::Delete ||
+					(UpdateParticle->HasBounds() && AccelerationStructure.NeedUpdateElementIn(UpdateParticle, UpdateParticle->WorldSpaceInflatedBounds(), PendingData.SpatialIdx)))
+				{
+					InternalAccelerationQueue.PendingData.Add(PendingData);
+				}
+			}
 		}
 	}
 
diff --git a/Engine/Source/Runtime/Experimental/Chaos/Public/Chaos/Collision/CollisionConstraintAllocator.h b/Engine/Source/Runtime/Experimental/Chaos/Public/Chaos/Collision/CollisionConstraintAllocator.h
--- a/Engine/Source/Runtime/Experimental/Chaos/Public/Chaos/Collision/CollisionConstraintAllocator.h
+++ b/Engine/Source/Runtime/Experimental/Chaos/Public/Chaos/Collision/CollisionConstraintAllocator.h
@@ -557,7 +557,7 @@ namespace Chaos
 
 			// Collect all the constraints activated in the collision tasks and register them (calls ActivateConstraintImp on each)
 			CHAOS_API void ProcessNewConstraints();
-			CHAOS_API void ProcessNewNonCcdConstraints(const TArray<FPBDCollisionConstraint*>& NewConstraints, TArray<TArray<FPBDCollisionConstraint*>>& CCDConstraintsBatch);
+			CHAOS_API void ProcessNewNonCcdConstraints(const TArray<FPBDCollisionConstraint*>& NewConstraints, TArray<FPBDCollisionConstraint*>& AllCCDConstraints);
 			CHAOS_API void ProcessNewCcdConstraints();
 
 			CHAOS_API void AddActiveConstraintImpl(FPBDCollisionConstraint* Constraint);
diff --git a/Engine/Source/Runtime/Experimental/Chaos/Public/Chaos/Framework/Parallel.h b/Engine/Source/Runtime/Experimental/Chaos/Public/Chaos/Framework/Parallel.h
--- a/Engine/Source/Runtime/Experimental/Chaos/Public/Chaos/Framework/Parallel.h
+++ b/Engine/Source/Runtime/Experimental/Chaos/Public/Chaos/Framework/Parallel.h
@@ -13,10 +13,17 @@ namespace Chaos
 	void CHAOS_API PhysicsParallelForWithContext(int32 InNum, TFunctionRef<int32 (int32, int32)> InContextCreator, TFunctionRef<void(int32, int32)> InCallable, bool bForceSingleThreaded = false);
 	//void CHAOS_API PhysicsParallelFor_RecursiveDivide(int32 InNum, TFunctionRef<void(int32)> InCallable, bool bForceSingleThreaded = false);
 
+	/**
+	 * Helper to call from task execution logic that does not use ParallelFor, to know if it should use parallel execution or not.
+	 * Returns true if the tasks should be run in parallel based on the current context.
+	 * @param InTaskNum The total number of tasks that need to run, will be compared to MinParallelTaskSize
+	 */
+	bool CHAOS_API ShouldExecuteTasksInParallel(int32 InTaskNum);
 
 	CHAOS_API extern int32 MaxNumWorkers;
 	CHAOS_API extern int32 SmallBatchSize;
 	CHAOS_API extern int32 LargeBatchSize;
+	CHAOS_API extern int32 MinParallelTaskSize;
 	CHAOS_API extern bool bDisablePhysicsParallelFor;
 	CHAOS_API extern bool bDisableParticleParallelFor;
 	CHAOS_API extern bool bDisableCollisionParallelFor;
