Fix crash when using new p.chaos.MinParallelTaskSize option I just added.

diff --git a/Engine/Source/Runtime/Experimental/Chaos/Private/Chaos/PBDRigidsEvolution.cpp b/Engine/Source/Runtime/Experimental/Chaos/Private/Chaos/PBDRigidsEvolution.cpp
--- a/Engine/Source/Runtime/Experimental/Chaos/Private/Chaos/PBDRigidsEvolution.cpp
+++ b/Engine/Source/Runtime/Experimental/Chaos/Private/Chaos/PBDRigidsEvolution.cpp
@@ -676,11 +676,13 @@ namespace Chaos
 			// Small enough task to run immediately and avoid parallel overhead
 			QUICK_SCOPE_CYCLE_COUNTER(STAT_PruneInternalPendingDataTask);
 
+			TArray<FPendingSpatialData> OldPendingData = MoveTemp(InternalAccelerationQueue.PendingData);
+
 			InternalAccelerationQueue.PendingData.Reset();
 
 			for (int32 Index = 0; Index < NumParticles; Index++)
 			{
-				const FPendingSpatialData& PendingData = InternalAccelerationQueue.PendingData[Index];
+				const FPendingSpatialData& PendingData = OldPendingData[Index];
 				FGeometryParticleHandle* UpdateParticle = PendingData.AccelerationHandle.GetGeometryParticleHandle_PhysicsThread();
 
 				if (PendingData.Operation == EPendingSpatialDataOperation::Delete ||
