Chaos Cloth:
-Supported auto reset strategy for extreme deformation on selected edges.
-Added SimulationResolveExtremeDeformationConfig to customize asset-based vertex selection and threshold value.
-Added FPBDExtremeDeformationConstraints to check for extreme deformation on selected edges.
-Added debug draw on extreme edge deformation.
-Added cvar: EnableResetIfExtremeEdgeDeformation, DebugDrawExtremlyDeformedEdges, OverrideExtremeDeformationEdgeRatioThreshold, GlobalExtremeDeformationEdgeRatioThreshold

diff --git a/Engine/Plugins/ChaosCloth/Source/ChaosCloth/Private/ChaosCloth/ChaosClothConstraints.cpp b/Engine/Plugins/ChaosCloth/Source/ChaosCloth/Private/ChaosCloth/ChaosClothConstraints.cpp
--- a/Engine/Plugins/ChaosCloth/Source/ChaosCloth/Private/ChaosCloth/ChaosClothConstraints.cpp
+++ b/Engine/Plugins/ChaosCloth/Source/ChaosCloth/Private/ChaosCloth/ChaosClothConstraints.cpp
@@ -798,6 +798,9 @@ void FClothConstraints::AddRules(
 	// Bending constraints
 	CreateBendingConstraints(ConfigProperties, WeightMaps, TriangleMesh, PatternData);
 
+	// Extreme deformation constraints
+	CreateExtremeDeformationConstraints(ConfigProperties, WeightMaps, VertexSets, TriangleMesh, PatternData);
+
 	// Area constraints
 	CreateAreaConstraints(ConfigProperties, WeightMaps, TriangleMesh, PatternData);
 
@@ -1185,6 +1188,43 @@ void FClothConstraints::CreateBendingConstraints(
 	}
 }
 
+void FClothConstraints::CreateExtremeDeformationConstraints(
+	const Softs::FCollectionPropertyConstFacade& ConfigProperties,
+	const TMap<FString, TConstArrayView<FRealSingle>>& WeightMaps,
+	const TMap<FString, const TSet<int32>*>& VertexSets,
+	const FTriangleMesh& TriangleMesh,
+	const FClothingPatternData* PatternData)
+{
+	if (Softs::FPBDExtremeDeformationConstraints::IsEnabled(ConfigProperties))
+	{
+		const TArray<Chaos::TVec2<int32>> CrossEdges = TriangleMesh.GetUniqueAdjacentPoints();
+
+		if (Evolution)
+		{
+			ExtremeDeformationConstraints = MakeShared<Softs::FPBDExtremeDeformationConstraints>(
+				Evolution->GetSoftBodyParticles(ParticleRangeId),
+				CrossEdges,
+				WeightMaps,
+				VertexSets,
+				ConfigProperties,
+				/*bTrimKinematicConstraints =*/ true);
+		}
+		else
+		{
+			ExtremeDeformationConstraints = MakeShared<Softs::FPBDExtremeDeformationConstraints>(
+				PBDEvolution->Particles(),
+				ParticleOffset,
+				NumParticles,
+				CrossEdges,
+				WeightMaps,
+				VertexSets,
+				ConfigProperties,
+				/*bTrimKinematicConstraints =*/ true);
+		}
+		//No ConstraintInits or ConstraintRules
+	}
+}
+
 void FClothConstraints::CreateAreaConstraints(
 	const Softs::FCollectionPropertyConstFacade& ConfigProperties,
 	const TMap<FString, TConstArrayView<FRealSingle>>& WeightMaps,
@@ -2450,6 +2490,10 @@ void FClothConstraints::Update(
 	{
 		XAnisoBendingElementConstraints->SetProperties(ConfigProperties, WeightMaps);
 	}
+	if (ExtremeDeformationConstraints)
+	{
+		ExtremeDeformationConstraints->SetProperties(ConfigProperties, WeightMaps);
+	}
 	if (AreaConstraints)
 	{
 		AreaConstraints->SetProperties(ConfigProperties, WeightMaps);
diff --git a/Engine/Plugins/ChaosCloth/Source/ChaosCloth/Private/ChaosCloth/ChaosClothVisualization.cpp b/Engine/Plugins/ChaosCloth/Source/ChaosCloth/Private/ChaosCloth/ChaosClothVisualization.cpp
--- a/Engine/Plugins/ChaosCloth/Source/ChaosCloth/Private/ChaosCloth/ChaosClothVisualization.cpp
+++ b/Engine/Plugins/ChaosCloth/Source/ChaosCloth/Private/ChaosCloth/ChaosClothVisualization.cpp
@@ -3119,6 +3119,42 @@ namespace Chaos
 		}
 	}
 
+	void FClothVisualizationNoGC::DrawExtremlyDeformedEdges(FPrimitiveDrawInterface* PDI) const
+	{
+		static const FLinearColor DynamicColor = FColor::Red;
+		static const FLinearColor KinematicColor = FColor::Purple;
+
+		for (const FClothingSimulationCloth* const Cloth : Solver->GetCloths())
+		{
+			const int32 ParticleRangeId = Cloth->GetParticleRangeId(Solver);
+			if (ParticleRangeId == INDEX_NONE)
+			{
+				continue;
+			}
+			// Elements are local indexed for new solver
+			const int32 Offset = Solver->IsLegacySolver() ? ParticleRangeId : 0;
+			const FClothConstraints& ClothConstraint = Solver->GetClothConstraints(ParticleRangeId);
+			if (const Softs::FPBDExtremeDeformationConstraints* const ExtremeDeformationConstraints = Solver->GetClothConstraints(ParticleRangeId).GetExtremeDeformationConstraints().Get())
+			{
+				const TConstArrayView<Softs::FSolverVec3> Positions = Cloth->GetParticlePositions(Solver);
+				const TConstArrayView<Softs::FSolverReal> InvMasses = Cloth->GetParticleInvMasses(Solver);
+				const TArray<TVec2<int32>> ExtremelyDeformedEdges = Solver->IsLegacySolver() ?
+					ExtremeDeformationConstraints->GetExtremelyDeformedEdges(Solver->GetParticleXs()) :
+					ExtremeDeformationConstraints->GetExtremelyDeformedEdges(Solver->GetParticleXsView(ParticleRangeId));
+				check(InvMasses.Num() == Positions.Num());
+				for (int32 EdgeIndex = 0; EdgeIndex < ExtremelyDeformedEdges.Num(); ++EdgeIndex)
+				{
+					const TVec2<int32>& Edge = ExtremelyDeformedEdges[EdgeIndex];
+					const FVector Pos0 = GetWorldPosition(*Solver, Positions[Edge[0] - Offset]); // TODO: Triangle Mesh shouldn't really be solver dependent (ie not use an offset)
+					const FVector Pos1 = GetWorldPosition(*Solver, Positions[Edge[1] - Offset]);
+					const bool bIsKinematic0 = (InvMasses[Edge[0] - Offset] == (Softs::FSolverReal)0.);
+					const bool bIsKinematic1 = (InvMasses[Edge[1] - Offset] == (Softs::FSolverReal)0.);
+					DrawLine(PDI, Pos0, Pos1, bIsKinematic0 && bIsKinematic1 ? KinematicColor : DynamicColor);
+				}
+			}
+		}
+	}
+
 #if WITH_EDITOR
 	TArray<FString> FClothVisualizationNoGC::GetAllWeightMapNames() const
 	{
diff --git a/Engine/Plugins/ChaosCloth/Source/ChaosCloth/Private/ChaosCloth/ChaosClothingSimulation.cpp b/Engine/Plugins/ChaosCloth/Source/ChaosCloth/Private/ChaosCloth/ChaosClothingSimulation.cpp
--- a/Engine/Plugins/ChaosCloth/Source/ChaosCloth/Private/ChaosCloth/ChaosClothingSimulation.cpp
+++ b/Engine/Plugins/ChaosCloth/Source/ChaosCloth/Private/ChaosCloth/ChaosClothingSimulation.cpp
@@ -97,6 +97,7 @@ namespace ClothingSimulationCVar
 	TAutoConsoleVariable<bool> DebugDrawElementIndices      (TEXT("p.ChaosCloth.DebugDrawElementIndices"      ), false, TEXT("Whether to debug draw the Chaos Cloth element indices"), ECVF_Cheat);
 	TAutoConsoleVariable<bool> DebugDrawClothClothConstraints    (TEXT("p.ChaosCloth.DebugDrawClothClothConstraints"    ), false, TEXT("Whether to debug draw the Chaos Cloth-Cloth constraints"), ECVF_Cheat);
 	TAutoConsoleVariable<bool> DebugDrawTeleportReset(TEXT("p.ChaosCloth.DebugDrawTeleportReset"), false, TEXT("Whether to debug draw the Chaos Cloth teleport/reset status"), ECVF_Cheat);
+	TAutoConsoleVariable<bool> DebugDrawExtremlyDeformedEdges(TEXT("p.ChaosCloth.DebugDrawExtremlyDeformedEdges"), false, TEXT("Whether to debug draw edges that are deformed beyond EdgeDeformationThreshold"), ECVF_Cheat);
 }
 #endif  // #if CHAOS_DEBUG_DRAW
 
@@ -605,6 +606,7 @@ void FClothingSimulation::Simulate(IClothingSimulationContext* InContext)
 	if (ClothingSimulationCVar::DebugDrawSelfCollision       .GetValueOnAnyThread()) { DebugDrawSelfCollision       (); }
 	if (ClothingSimulationCVar::DebugDrawSelfIntersection    .GetValueOnAnyThread()) { DebugDrawSelfIntersection    (); }
 	if (ClothingSimulationCVar::DebugDrawTeleportReset       .GetValueOnAnyThread()) { DebugDrawTeleportReset       (); }
+	if (ClothingSimulationCVar::DebugDrawExtremlyDeformedEdges.GetValueOnAnyThread()) { DebugDrawExtremlyDeformedEdges(); }
 #endif  // #if CHAOS_DEBUG_DRAW
 }
 
diff --git a/Engine/Plugins/ChaosCloth/Source/ChaosCloth/Private/ChaosCloth/ChaosClothingSimulationSolver.cpp b/Engine/Plugins/ChaosCloth/Source/ChaosCloth/Private/ChaosCloth/ChaosClothingSimulationSolver.cpp
--- a/Engine/Plugins/ChaosCloth/Source/ChaosCloth/Private/ChaosCloth/ChaosClothingSimulationSolver.cpp
+++ b/Engine/Plugins/ChaosCloth/Source/ChaosCloth/Private/ChaosCloth/ChaosClothingSimulationSolver.cpp
@@ -85,6 +85,9 @@ FAutoConsoleVariableRef CVarClothSolverMaxVelocity(TEXT("p.ChaosCloth.Solver.Max
 static bool bClothSolverBrokenAccelerationClampingToZero = false;
 FAutoConsoleVariableRef CVarClothSolverBrokenAccelerationClampingHack(TEXT("p.ChaosCloth.Solver.BrokenAccelerationClampingHack"), bClothSolverBrokenAccelerationClampingToZero, TEXT("Return to previous broken behavior where acceleration clamping would clamp to zero velocity. This cvar will be removed."));
 
+static bool bClothSolverEnableResetIfExtremeEdgeDeformation = true;
+FAutoConsoleVariableRef CVarClothSolverEnableResetIfExtremeEdgeDeformation(TEXT("p.ChaosCloth.Solver.EnableResetIfExtremeEdgeDeformation"), bClothSolverEnableResetIfExtremeEdgeDeformation, TEXT("Set to false to disable reset to animation positions if there's any extreme edge deformation."));
+
 namespace ClothingSimulationSolverDefault
 {
 	static const Softs::FSolverVec3 Gravity((Softs::FSolverReal)0., (Softs::FSolverReal)0., (Softs::FSolverReal)-980.665);  // cm/s^2
@@ -2215,6 +2218,8 @@ void FClothingSimulationSolver::Update(Softs::FSolverReal InDeltaTime)
 		SCOPE_CYCLE_COUNTER(STAT_ChaosClothSolverUpdatePostSolverStep);
 		SCOPE_CYCLE_COUNTER(STAT_ClothComputeNormals);
 		
+		ExtremeDeformationReset();
+
 		EventPostSolve.Broadcast(DeltaTime);
 
 		PhysicsParallelFor(Cloths.Num(), [this](int32 ClothIndex)
@@ -2228,6 +2233,46 @@ void FClothingSimulationSolver::Update(Softs::FSolverReal InDeltaTime)
 	OldLocalSpaceLocation = LocalSpaceLocation;
 }
 
+void FClothingSimulationSolver::ExtremeDeformationReset()
+{
+	if (bClothSolverEnableResetIfExtremeEdgeDeformation)
+	{
+		if (Evolution)
+		{
+			const TArray<uint32> ActiveGroups = Evolution->GetActiveGroupsArray();
+			PhysicsParallelFor(ActiveGroups.Num(), [this, &ActiveGroups](int32 ActiveGroupIndex)
+				{
+					const uint32 GroupId = ActiveGroups[ActiveGroupIndex];
+					const TSet<int32>& ActiveSoftBodies = Evolution->GetGroupActiveSoftBodies(GroupId);
+					for (const int32 SoftBodyId : ActiveSoftBodies)
+					{
+						if (const Softs::FPBDExtremeDeformationConstraints* const ExtremeDeformationConstraints = GetClothConstraints(SoftBodyId).GetExtremeDeformationConstraints().Get())
+						{
+							if (ExtremeDeformationConstraints->IsExtremelyDeformed(GetParticleXsView(SoftBodyId)))
+							{
+								ResetStartPose(SoftBodyId, Evolution->GetSoftBodyParticleNum(SoftBodyId));
+							}
+						}
+					}
+				});
+		}
+		else
+		{
+			const TPBDActiveView<Softs::FSolverParticles>& ParticlesActiveView = PBDEvolution->ParticlesActiveView();
+			ParticlesActiveView.RangeFor(
+				[this](Softs::FSolverParticles& Particles, int32 Offset, int32 Range)
+			{
+				if (const Softs::FPBDExtremeDeformationConstraints* const ExtremeDeformationConstraints = GetClothConstraints(Offset).GetExtremeDeformationConstraints().Get())
+				{
+					if (ExtremeDeformationConstraints->IsExtremelyDeformed(GetParticleXs()))
+					{
+						ResetStartPose(Offset, PBDEvolution->GetParticleRangeSize(Offset));
+					}
+				}
+			}, /*bForceSingleThreaded =*/ !bClothSolverParallelClothPostUpdate);
+		}
+	}
+}
 
 void FClothingSimulationSolver::UpdateFromCache(const FClothingSimulationCacheData& CacheData)
 {
diff --git a/Engine/Plugins/ChaosCloth/Source/ChaosCloth/Public/ChaosCloth/ChaosClothConstraints.h b/Engine/Plugins/ChaosCloth/Source/ChaosCloth/Public/ChaosCloth/ChaosClothConstraints.h
--- a/Engine/Plugins/ChaosCloth/Source/ChaosCloth/Public/ChaosCloth/ChaosClothConstraints.h
+++ b/Engine/Plugins/ChaosCloth/Source/ChaosCloth/Public/ChaosCloth/ChaosClothConstraints.h
@@ -10,6 +10,7 @@
 #include "Chaos/XPBDPlanarConstraints.h"
 #include "Chaos/Deformable/GaussSeidelMainConstraint.h"
 #include "Chaos/Deformable/GaussSeidelCorotatedCodimensionalConstraints.h"
+#include "Chaos/PBDExtremeDeformationConstraints.h"
 
 struct FManagedArrayCollection;
 
@@ -105,6 +106,7 @@ namespace Chaos
 		const TSharedPtr<Softs::FPBDBendingConstraints>& GetBendingElementConstraints() const { return BendingElementConstraints; }
 		const TSharedPtr<Softs::FXPBDBendingConstraints>& GetXBendingElementConstraints() const { return XBendingElementConstraints; }
 		const TSharedPtr<Softs::FXPBDAnisotropicBendingConstraints>& GetXAnisoBendingElementConstraints() const { return XAnisoBendingElementConstraints; }
+		const TSharedPtr<Softs::FPBDExtremeDeformationConstraints>& GetExtremeDeformationConstraints() const { return ExtremeDeformationConstraints; }
 		const TSharedPtr<Softs::FPBDAreaSpringConstraints>& GetAreaSpringConstraints() const { return AreaConstraints; }
 		const TSharedPtr<Softs::FXPBDAreaSpringConstraints>& GetXAreaSpringConstraints() const { return XAreaConstraints; }
 		const TSharedPtr<Softs::FPBDLongRangeConstraints>& GetLongRangeConstraints() const { return LongRangeConstraints; }
@@ -145,6 +147,12 @@ namespace Chaos
 			const TMap<FString, TConstArrayView<FRealSingle>>& WeightMaps,
 			const FTriangleMesh& TriangleMesh,
 			const FClothingPatternData* PatternData);
+		void CreateExtremeDeformationConstraints(
+			const Softs::FCollectionPropertyConstFacade& ConfigProperties,
+			const TMap<FString, TConstArrayView<FRealSingle>>& WeightMaps,
+			const TMap<FString, const TSet<int32>*>& VertexSets,
+			const FTriangleMesh& TriangleMesh,
+			const FClothingPatternData* PatternData);
 		void CreateAreaConstraints(
 			const Softs::FCollectionPropertyConstFacade& ConfigProperties,
 			const TMap<FString, TConstArrayView<FRealSingle>>& WeightMaps,
@@ -217,6 +225,7 @@ namespace Chaos
 		TSharedPtr<Softs::FGaussSeidelMainConstraint<Softs::FSolverReal, Softs::FSolverParticles>> GSMainConstraint;
 		TSharedPtr<Softs::FGaussSeidelCorotatedCodimensionalConstraints<Softs::FSolverReal, Softs::FSolverParticles>> GSCorotatedCodimensionalConstraint;
 		TSharedPtr<Softs::FMultiResConstraints> MultiResConstraints;
+		TSharedPtr<Softs::FPBDExtremeDeformationConstraints> ExtremeDeformationConstraints;
 		//~ Begin Force-based solver only constraints
 		Softs::FSolverVec3 SolverWindVelocity; // Set from solver and added to wind from the config
 		TSharedPtr<Softs::FVelocityAndPressureField> VelocityAndPressureField;
diff --git a/Engine/Plugins/ChaosCloth/Source/ChaosCloth/Public/ChaosCloth/ChaosClothVisualization.h b/Engine/Plugins/ChaosCloth/Source/ChaosCloth/Public/ChaosCloth/ChaosClothVisualization.h
--- a/Engine/Plugins/ChaosCloth/Source/ChaosCloth/Public/ChaosCloth/ChaosClothVisualization.h
+++ b/Engine/Plugins/ChaosCloth/Source/ChaosCloth/Public/ChaosCloth/ChaosClothVisualization.h
@@ -60,7 +60,8 @@ namespace Chaos
 		CHAOSCLOTH_API void DrawMultiResConstraint(FPrimitiveDrawInterface* PDI = nullptr) const;
 		CHAOSCLOTH_API void DrawClothClothConstraints(FPrimitiveDrawInterface* PDI = nullptr) const;
 		CHAOSCLOTH_API void DrawTeleportReset(FPrimitiveDrawInterface* PDI = nullptr) const;
-
+		CHAOSCLOTH_API void DrawExtremlyDeformedEdges(FPrimitiveDrawInterface* PDI = nullptr) const;
+	
 #else  // #if CHAOS_DEBUG_DRAW
 		void SetSolver(const ::Chaos::FClothingSimulationSolver* /*InSolver*/) {}
 
@@ -93,7 +94,7 @@ namespace Chaos
 		void DrawMultiResConstraint(FPrimitiveDrawInterface* /*PDI*/ = nullptr) const {}
 		void DrawClothClothConstraints(FPrimitiveDrawInterface* /*PDI*/ = nullptr) const {}
 		void DrawTeleportReset(FPrimitiveDrawInterface* /*PDI*/ = nullptr) const {}
-
+		void DrawExtremlyDeformedEdges(FPrimitiveDrawInterface* /*PDI*/ = nullptr) const {}
 #endif  // #if CHAOS_DEBUG_DRAW
 
 #if WITH_EDITOR && CHAOS_DEBUG_DRAW
diff --git a/Engine/Plugins/ChaosCloth/Source/ChaosCloth/Public/ChaosCloth/ChaosClothingSimulation.h b/Engine/Plugins/ChaosCloth/Source/ChaosCloth/Public/ChaosCloth/ChaosClothingSimulation.h
--- a/Engine/Plugins/ChaosCloth/Source/ChaosCloth/Public/ChaosCloth/ChaosClothingSimulation.h
+++ b/Engine/Plugins/ChaosCloth/Source/ChaosCloth/Public/ChaosCloth/ChaosClothingSimulation.h
@@ -115,6 +115,7 @@ namespace Chaos
 		void DebugDrawBounds(FPrimitiveDrawInterface* PDI = nullptr) const { Visualization.DrawBounds(PDI); }
 		void DebugDrawGravity(FPrimitiveDrawInterface* PDI = nullptr) const { Visualization.DrawGravity(PDI); }
 		void DebugDrawTeleportReset(FPrimitiveDrawInterface* PDI = nullptr) const { Visualization.DrawTeleportReset(PDI); }
+		void DebugDrawExtremlyDeformedEdges(FPrimitiveDrawInterface* PDI = nullptr) const { Visualization.DrawExtremlyDeformedEdges(PDI); }
 #endif  // #if CHAOS_DEBUG_DRAW
 
 	private:
diff --git a/Engine/Plugins/ChaosCloth/Source/ChaosCloth/Public/ChaosCloth/ChaosClothingSimulationSolver.h b/Engine/Plugins/ChaosCloth/Source/ChaosCloth/Public/ChaosCloth/ChaosClothingSimulationSolver.h
--- a/Engine/Plugins/ChaosCloth/Source/ChaosCloth/Public/ChaosCloth/ChaosClothingSimulationSolver.h
+++ b/Engine/Plugins/ChaosCloth/Source/ChaosCloth/Public/ChaosCloth/ChaosClothingSimulationSolver.h
@@ -96,6 +96,8 @@ namespace Chaos
 		/** Advance the simulation. */
 		CHAOSCLOTH_API void Update(FSolverReal InDeltaTime);
 
+		void ExtremeDeformationReset();
+
 		/** Return the last delta time used for advancing the simulation. */
 		FSolverReal GetDeltaTime() const { return DeltaTime; }
 
diff --git a/Engine/Plugins/ChaosClothAsset/Source/ChaosClothAssetEngine/Private/ChaosClothAsset/ClothSimulationProxy.cpp b/Engine/Plugins/ChaosClothAsset/Source/ChaosClothAssetEngine/Private/ChaosClothAsset/ClothSimulationProxy.cpp
--- a/Engine/Plugins/ChaosClothAsset/Source/ChaosClothAssetEngine/Private/ChaosClothAsset/ClothSimulationProxy.cpp
+++ b/Engine/Plugins/ChaosClothAsset/Source/ChaosClothAssetEngine/Private/ChaosClothAsset/ClothSimulationProxy.cpp
@@ -556,6 +556,7 @@ PRAGMA_ENABLE_DEPRECATION_WARNINGS
 		static const TConsoleVariableData<bool>* const DebugDrawElementIndicesCVar = IConsoleManager::Get().FindTConsoleVariableDataBool(TEXT("p.ChaosCloth.DebugDrawElementIndices"));
 		static const TConsoleVariableData<bool>* const DebugDrawClothClothConstraintsCVar = IConsoleManager::Get().FindTConsoleVariableDataBool(TEXT("p.ChaosCloth.DebugDrawClothClothConstraints"));
 		static const TConsoleVariableData<bool>* const DebugDrawTeleportResetCVar = IConsoleManager::Get().FindTConsoleVariableDataBool(TEXT("p.ChaosCloth.DebugDrawTeleportReset"));
+		static const TConsoleVariableData<bool>* const DebugDrawExtremlyDeformedEdgesCVar = IConsoleManager::Get().FindTConsoleVariableDataBool(TEXT("p.ChaosCloth.DebugDrawExtremlyDeformedEdges"));
 
 		if (DebugDrawLocalSpaceCVar && DebugDrawLocalSpaceCVar->GetValueOnAnyThread()) { Visualization->DrawLocalSpace(); }
 		if (DebugDrawBoundsCVar && DebugDrawBoundsCVar->GetValueOnAnyThread()) { Visualization->DrawBounds(); }
@@ -587,6 +588,10 @@ PRAGMA_ENABLE_DEPRECATION_WARNINGS
 		{
 			Visualization->DrawTeleportReset();
 		}
+		if (DebugDrawExtremlyDeformedEdgesCVar && DebugDrawExtremlyDeformedEdgesCVar->GetValueOnAnyThread())
+		{
+			Visualization->DrawExtremlyDeformedEdges();
+		}
 #endif  // #if CHAOS_DEBUG_DRAW
 	}
 
diff --git a/Engine/Plugins/ChaosClothAssetEditor/Source/ChaosClothAssetDataflowNodes/Private/ChaosClothAsset/ChaosClothAssetDataflowNodesModule.cpp b/Engine/Plugins/ChaosClothAssetEditor/Source/ChaosClothAssetDataflowNodes/Private/ChaosClothAsset/ChaosClothAssetDataflowNodesModule.cpp
--- a/Engine/Plugins/ChaosClothAssetEditor/Source/ChaosClothAssetDataflowNodes/Private/ChaosClothAsset/ChaosClothAssetDataflowNodesModule.cpp
+++ b/Engine/Plugins/ChaosClothAssetEditor/Source/ChaosClothAssetDataflowNodes/Private/ChaosClothAsset/ChaosClothAssetDataflowNodesModule.cpp
@@ -51,6 +51,7 @@
 #include "ChaosClothAsset/SimulationSelfCollisionSpheresConfigNode.h"
 #include "ChaosClothAsset/SimulationSolverConfigNode.h"
 #include "ChaosClothAsset/SimulationMultiResConfigNode.h"
+#include "ChaosClothAsset/SimulationResolveExtremeDeformationConfigNode.h"
 #include "ChaosClothAsset/SimulationStretchConfigNode.h"
 #include "ChaosClothAsset/SimulationStretchOverrideConfigNode.h"
 #include "ChaosClothAsset/SimulationVelocityScaleConfigNode.h"
@@ -152,6 +153,7 @@ namespace UE::Chaos::ClothAsset
 			DATAFLOW_NODE_REGISTER_CREATION_FACTORY(FChaosClothAssetSimulationXPBDBendingElementConfigNode);
 			DATAFLOW_NODE_REGISTER_CREATION_FACTORY(FChaosClothAssetSimulationXPBDBendingSpringConfigNode);
 			DATAFLOW_NODE_REGISTER_CREATION_FACTORY(FChaosClothAssetSimulationXPBDEdgeSpringConfigNode);
+			DATAFLOW_NODE_REGISTER_CREATION_FACTORY(FChaosClothAssetSimulationResolveExtremeDeformationConfigNode);
 			DATAFLOW_NODE_REGISTER_CREATION_FACTORY(FChaosClothAssetSkeletalMeshImportNode_v2);
 			DATAFLOW_NODE_REGISTER_CREATION_FACTORY(FChaosClothAssetSkinningBlendNode);
 			DATAFLOW_NODE_REGISTER_CREATION_FACTORY(FChaosClothAssetStaticMeshImportNode_v2);
diff --git a/Engine/Plugins/ChaosClothAssetEditor/Source/ChaosClothAssetDataflowNodes/Private/ChaosClothAsset/SimulationResolveExtremeDeformationConfigNode.cpp b/Engine/Plugins/ChaosClothAssetEditor/Source/ChaosClothAssetDataflowNodes/Private/ChaosClothAsset/SimulationResolveExtremeDeformationConfigNode.cpp
new file mode 100644
--- /dev/null
+++ b/Engine/Plugins/ChaosClothAssetEditor/Source/ChaosClothAssetDataflowNodes/Private/ChaosClothAsset/SimulationResolveExtremeDeformationConfigNode.cpp
@@ -0,0 +1,47 @@
+// Copyright Epic Games, Inc. All Rights Reserved.
+
+#include "ChaosClothAsset/SimulationResolveExtremeDeformationConfigNode.h"
+#include "ChaosClothAsset/ClothCollectionGroup.h"
+#include "ChaosClothAsset/ClothGeometryTools.h"
+#include "ChaosClothAsset/CollectionClothSelectionFacade.h"
+#include "Dataflow/DataflowInputOutput.h"
+
+#include UE_INLINE_GENERATED_CPP_BY_NAME(SimulationResolveExtremeDeformationConfigNode)
+
+FChaosClothAssetSimulationResolveExtremeDeformationConfigNode::FChaosClothAssetSimulationResolveExtremeDeformationConfigNode(const UE::Dataflow::FNodeParameters& InParam, FGuid InGuid)
+	: FChaosClothAssetSimulationBaseConfigNode(InParam, InGuid)
+{
+	RegisterCollectionConnections();
+	RegisterInputConnection(&InputSelection, GET_MEMBER_NAME_CHECKED(FChaosClothAssetConnectableIStringValue, StringValue));
+	RegisterOutputConnection(&ExtremeDeformationVertexSelection);
+}
+
+void FChaosClothAssetSimulationResolveExtremeDeformationConfigNode::AddProperties(FPropertyHelper& PropertyHelper) const
+{
+	PropertyHelper.SetProperty(this, &ExtremeDeformationEdgeRatioThreshold);
+	PropertyHelper.SetPropertyString(this, &ExtremeDeformationVertexSelection, {}, ECollectionPropertyFlags::None);
+}
+
+void FChaosClothAssetSimulationResolveExtremeDeformationConfigNode::EvaluateClothCollection(UE::Dataflow::FContext& Context, const TSharedRef<FManagedArrayCollection>& ClothCollection) const
+{
+	using namespace UE::Chaos::ClothAsset;
+	TSet<int32> SelectionSet;
+	const FName InSelectionName(GetValue(Context, &InputSelection.StringValue));
+	const FName OutSelectionName(ExtremeDeformationVertexSelection);
+	FCollectionClothSelectionFacade SelectionFacade(ClothCollection);
+	SelectionFacade.DefineSchema();
+	if (FClothGeometryTools::ConvertSelectionToNewGroupType(ClothCollection, InSelectionName, ClothCollectionGroup::SimVertices3D, SelectionSet))
+	{
+		SelectionFacade.FindOrAddSelectionSet(OutSelectionName, ClothCollectionGroup::SimVertices3D) = SelectionSet;
+	}
+}
+
+void FChaosClothAssetSimulationResolveExtremeDeformationConfigNode::Evaluate(UE::Dataflow::FContext& Context, const FDataflowOutput* Out) const
+{
+	Super::Evaluate(Context, Out);
+
+	if (Out->IsA(&ExtremeDeformationVertexSelection))
+	{
+		SetValue(Context, ExtremeDeformationVertexSelection, &ExtremeDeformationVertexSelection);
+	}
+}
\ No newline at end of file
diff --git a/Engine/Plugins/ChaosClothAssetEditor/Source/ChaosClothAssetDataflowNodes/Public/ChaosClothAsset/SimulationResolveExtremeDeformationConfigNode.h b/Engine/Plugins/ChaosClothAssetEditor/Source/ChaosClothAssetDataflowNodes/Public/ChaosClothAsset/SimulationResolveExtremeDeformationConfigNode.h
new file mode 100644
--- /dev/null
+++ b/Engine/Plugins/ChaosClothAssetEditor/Source/ChaosClothAssetDataflowNodes/Public/ChaosClothAsset/SimulationResolveExtremeDeformationConfigNode.h
@@ -0,0 +1,38 @@
+// Copyright Epic Games, Inc. All Rights Reserved.
+
+#pragma once 
+
+#include "ChaosClothAsset/SimulationBaseConfigNode.h"
+#include "SimulationResolveExtremeDeformationConfigNode.generated.h"
+
+/** Resolve extreme deformation properties configuration node. */
+USTRUCT(Meta = (DataflowCloth))
+struct FChaosClothAssetSimulationResolveExtremeDeformationConfigNode : public FChaosClothAssetSimulationBaseConfigNode
+{
+	GENERATED_USTRUCT_BODY()
+	DATAFLOW_NODE_DEFINE_INTERNAL(FChaosClothAssetSimulationResolveExtremeDeformationConfigNode, "SimulationResolveExtremeDeformationConfig", "Cloth", "Cloth Simulation Resolve ExtremeDeformation Config")
+
+public:
+	/** Vertex Selection to check for extreme deformations */
+	UPROPERTY(EditAnywhere, Category = "Input Selection")
+	FChaosClothAssetConnectableIStringValue InputSelection = { TEXT("InputSelection") };
+
+	/** Edges deformed above this threshold will trigger position reset. */
+	UPROPERTY(EditAnywhere, Category = "Edge Ratio Threshold Properties", meta = (UIMin = "0", UIMax = "100", ClampMin = "0", ClampMax = "1000", InteractorName = "ExtremeDeformationEdgeRatioThreshold"))
+	float ExtremeDeformationEdgeRatioThreshold = 5.f;
+
+	/**
+	 * The name of the extreme deformation vertex set generated by this node.
+	 * The name of this set cannot be changed and is only provided for further tweaking.
+	 */
+	UPROPERTY(VisibleAnywhere, Category = "Output Selection Name", Meta = (DataflowOutput))
+	FString ExtremeDeformationVertexSelection = GET_MEMBER_NAME_STRING_CHECKED(FChaosClothAssetSimulationResolveExtremeDeformationConfigNode, ExtremeDeformationVertexSelection);
+
+	FChaosClothAssetSimulationResolveExtremeDeformationConfigNode(const UE::Dataflow::FNodeParameters& InParam, FGuid InGuid = FGuid::NewGuid());
+
+private:
+	virtual void AddProperties(FPropertyHelper& PropertyHelper) const override;
+
+	virtual void EvaluateClothCollection(UE::Dataflow::FContext& Context, const TSharedRef<FManagedArrayCollection>& ClothCollection) const override;
+	virtual void Evaluate(UE::Dataflow::FContext& Context, const FDataflowOutput* Out) const override;
+};
diff --git a/Engine/Source/Runtime/Experimental/Chaos/Private/Chaos/PBDExtremeDeformationConstraints.cpp b/Engine/Source/Runtime/Experimental/Chaos/Private/Chaos/PBDExtremeDeformationConstraints.cpp
new file mode 100644
--- /dev/null
+++ b/Engine/Source/Runtime/Experimental/Chaos/Private/Chaos/PBDExtremeDeformationConstraints.cpp
@@ -0,0 +1,108 @@
+// Copyright Epic Games, Inc. All Rights Reserved.
+
+#include "Chaos/PBDExtremeDeformationConstraints.h"
+#include "HAL/IConsoleManager.h"
+
+namespace Chaos::Softs
+{
+
+namespace Private
+{
+	bool bOverrideExtremeDeformationEdgeRatioThreshold = false;
+	FAutoConsoleVariableRef CVarOverrideExtremeDeformationEdgeRatioThreshold(TEXT("p.ExtremeDeformationConstraints.OverrideExtremeDeformationEdgeRatioThreshold"), bOverrideExtremeDeformationEdgeRatioThreshold, TEXT("Override asset-based extreme deformation edge ratio threshold with GlobalExtremeDeformationEdgeRatioThreshold."));
+	float GlobalExtremeDeformationEdgeRatioThreshold = FLT_MAX;
+	FAutoConsoleVariableRef CVarGlobalExtremeDeformationEdgeRatioThreshold(TEXT("p.ExtremeDeformationConstraints.GlobalExtremeDeformationEdgeRatioThreshold"), GlobalExtremeDeformationEdgeRatioThreshold, TEXT("Global extreme deformation edge ratio threshold parameter to trigger position reset."));
+}
+
+FSolverReal FPBDExtremeDeformationConstraints::GetThreshold() const
+{
+	if (Private::bOverrideExtremeDeformationEdgeRatioThreshold)
+	{
+		return Private::GlobalExtremeDeformationEdgeRatioThreshold;
+	}
+	return ExtremeDeformationThreshold;
+}
+
+bool FPBDExtremeDeformationConstraints::IsExtremelyDeformed(TConstArrayView<Softs::FSolverVec3> Positions) const
+{
+	const FSolverReal Threshold = GetThreshold();
+	for (int32 ConstraintIndex = 0; ConstraintIndex < Constraints.Num(); ++ConstraintIndex)
+	{
+		const auto& Constraint = Constraints[ConstraintIndex];
+		const int32 i1 = Constraint[0];
+		const int32 i2 = Constraint[1];
+		const FSolverVec3& P1 = Positions[i1];
+		const FSolverVec3& P2 = Positions[i2];
+		const FSolverReal EdgeLengthRatio = (P1 - P2).Size() / Dists[ConstraintIndex];
+		if (EdgeLengthRatio > Threshold)
+		{
+			return true;
+		}
+	}
+	return false;
+}
+
+bool FPBDExtremeDeformationConstraints::IsExtremelyDeformed(TConstArrayView<Softs::FSolverVec3> Positions, const Softs::FSolverVec3* const RefencePositions) const
+{
+	const FSolverReal Threshold = GetThreshold();
+	for (int32 ConstraintIndex = 0; ConstraintIndex < Constraints.Num(); ++ConstraintIndex)
+	{
+		const auto& Constraint = Constraints[ConstraintIndex];
+		const int32 i1 = Constraint[0];
+		const int32 i2 = Constraint[1];
+		const FSolverVec3& P1 = Positions[i1];
+		const FSolverVec3& P2 = Positions[i2];
+		const FSolverVec3& RefP1 = RefencePositions[i1];
+		const FSolverVec3& RefP2 = RefencePositions[i2];
+		const FSolverReal EdgeLengthRatio = (P1 - P2).Size() / (RefP1 - RefP2).Size();
+		if (EdgeLengthRatio > Threshold)
+		{
+			return true;
+		}
+	}
+	return false;
+}
+
+TArray<TVec2<int32>> FPBDExtremeDeformationConstraints::GetExtremelyDeformedEdges(TConstArrayView<Softs::FSolverVec3> Positions) const
+{
+	const FSolverReal Threshold = GetThreshold();
+	TArray<TVec2<int32>> OutEdges;
+	for (int32 ConstraintIndex = 0; ConstraintIndex < Constraints.Num(); ++ConstraintIndex)
+	{
+		const TVec2<int32>& Constraint = Constraints[ConstraintIndex];
+		const int32 i1 = Constraint[0];
+		const int32 i2 = Constraint[1];
+		const FSolverVec3& P1 = Positions[i1];
+		const FSolverVec3& P2 = Positions[i2];
+		const FSolverReal EdgeLengthRatio = (P1 - P2).Size() / Dists[ConstraintIndex];
+		if (EdgeLengthRatio > Threshold)
+		{
+			OutEdges.Add(Constraint);
+		}
+	}
+	return OutEdges;
+}
+
+TArray<TVec2<int32>> FPBDExtremeDeformationConstraints::GetExtremelyDeformedEdges(TConstArrayView<Softs::FSolverVec3> Positions, const Softs::FSolverVec3* const RefencePositions) const
+{
+	const FSolverReal Threshold = GetThreshold();
+	TArray<TVec2<int32>> OutEdges;
+	for (int32 ConstraintIndex = 0; ConstraintIndex < Constraints.Num(); ++ConstraintIndex)
+	{
+		const TVec2<int32>& Constraint = Constraints[ConstraintIndex];
+		const int32 i1 = Constraint[0];
+		const int32 i2 = Constraint[1];
+		const FSolverVec3& P1 = Positions[i1];
+		const FSolverVec3& P2 = Positions[i2];
+		const FSolverVec3& RefP1 = RefencePositions[i1];
+		const FSolverVec3& RefP2 = RefencePositions[i2];
+		const FSolverReal EdgeLengthRatio = (P1 - P2).Size() / (RefP1 - RefP2).Size();
+		if (EdgeLengthRatio > Threshold)
+		{
+			OutEdges.Add(Constraint);
+		}
+	}
+	return OutEdges;
+}
+
+} // End namespace Chaos::Softs
\ No newline at end of file
diff --git a/Engine/Source/Runtime/Experimental/Chaos/Public/Chaos/PBDExtremeDeformationConstraints.h b/Engine/Source/Runtime/Experimental/Chaos/Public/Chaos/PBDExtremeDeformationConstraints.h
new file mode 100644
--- /dev/null
+++ b/Engine/Source/Runtime/Experimental/Chaos/Public/Chaos/PBDExtremeDeformationConstraints.h
@@ -0,0 +1,168 @@
+// Copyright Epic Games, Inc. All Rights Reserved.
+#pragma once
+
+#include "Chaos/CollectionPropertyFacade.h"
+#include "Chaos/PBDSoftsSolverParticles.h"
+#include "Chaos/SoftsSolverParticlesRange.h"
+#include "Chaos/Utilities.h"
+#include "Containers/Array.h"
+
+namespace Chaos::Softs
+{
+
+class FPBDExtremeDeformationConstraints
+{
+public:
+
+	static bool IsEnabled(const FCollectionPropertyConstFacade& PropertyCollection)
+	{
+		return IsExtremeDeformationVertexSelectionEnabled(PropertyCollection, false) && 
+			IsExtremeDeformationEdgeRatioThresholdEnabled(PropertyCollection, false);
+	}
+
+	template<int32 Valence UE_REQUIRES(Valence >= 2 && Valence <= 4)>
+	FPBDExtremeDeformationConstraints(
+		const FSolverParticlesRange& Particles,
+		const TArray<TVector<int32, Valence>>& InConstraints,
+		const TMap<FString, TConstArrayView<FRealSingle>>& WeightMaps,
+		const TMap<FString, const TSet<int32>*>& VertexSets,
+		const FCollectionPropertyConstFacade& PropertyCollection,
+		bool bTrimKinematicConstraints = false)
+		: VertexSet(VertexSets.FindRef(GetExtremeDeformationVertexSelectionString(PropertyCollection, ExtremeDeformationVertexSelectionName.ToString()), nullptr))
+		, Constraints(TrimConstraints(InConstraints,
+			[&Particles, bTrimKinematicConstraints, this](int32 Index0, int32 Index1)
+			{
+				return (bTrimKinematicConstraints && Particles.InvM(Index0) == (FSolverReal)0. && Particles.InvM(Index1) == (FSolverReal)0.)
+					|| (VertexSet && (!VertexSet->Contains(Index0) || !VertexSet->Contains(Index1)));
+			}))
+		, ParticleOffset(0)
+		, ParticleCount(Particles.GetRangeSize())
+		, ExtremeDeformationThreshold(GetExtremeDeformationEdgeRatioThreshold(PropertyCollection, FLT_MAX))
+		, ExtremeDeformationEdgeRatioThresholdIndex(PropertyCollection)
+	{
+		// Update distances
+		Dists.Reset(Constraints.Num());
+		for (const TVec2<int32>& Constraint : Constraints)
+		{
+			const FSolverVec3& P0 = Particles.X(Constraint[0]);
+			const FSolverVec3& P1 = Particles.X(Constraint[1]);
+			Dists.Add((P1 - P0).Size());
+		}
+	}
+
+	template<int32 Valence UE_REQUIRES(Valence >= 2 && Valence <= 4)>
+	FPBDExtremeDeformationConstraints(
+		const FSolverParticles& Particles,
+		int32 InParticleOffset,
+		int32 InParticleCount,
+		const TArray<TVector<int32, Valence>>& InConstraints,
+		const TMap<FString, TConstArrayView<FRealSingle>>& WeightMaps,
+		const TMap<FString, const TSet<int32>*>& VertexSets,
+		const FCollectionPropertyConstFacade& PropertyCollection,
+		bool bTrimKinematicConstraints = false)
+		: VertexSet(VertexSets.FindRef(GetExtremeDeformationVertexSelectionString(PropertyCollection, ExtremeDeformationVertexSelectionName.ToString()), nullptr))
+		, Constraints(TrimConstraints(InConstraints,
+			[&Particles, bTrimKinematicConstraints, this](int32 Index0, int32 Index1)
+			{
+				return (bTrimKinematicConstraints && Particles.InvM(Index0) == (FSolverReal)0. && Particles.InvM(Index1) == (FSolverReal)0.)
+					|| (VertexSet && (!VertexSet->Contains(Index0) || !VertexSet->Contains(Index1)));
+			}))
+		, ParticleOffset(InParticleOffset)
+		, ParticleCount(InParticleCount)
+		, ExtremeDeformationThreshold(GetExtremeDeformationEdgeRatioThreshold(PropertyCollection, FLT_MAX))
+		, ExtremeDeformationEdgeRatioThresholdIndex(PropertyCollection)
+	{
+		// Update distances
+		Dists.Reset(Constraints.Num());
+		for (const TVec2<int32>& Constraint : Constraints)
+		{
+			const FSolverVec3& P0 = Particles.GetX(Constraint[0]);
+			const FSolverVec3& P1 = Particles.GetX(Constraint[1]);
+			Dists.Add((P1 - P0).Size());
+		}
+	}
+
+	virtual ~FPBDExtremeDeformationConstraints()
+	{}
+
+	void SetProperties(const FCollectionPropertyConstFacade& PropertyCollection,
+		const TMap<FString, TConstArrayView<FRealSingle>>& WeightMaps)
+	{
+		if (IsExtremeDeformationEdgeRatioThresholdMutable(PropertyCollection))
+		{
+			ExtremeDeformationThreshold = GetExtremeDeformationEdgeRatioThreshold(PropertyCollection, FLT_MAX);
+		}
+	}
+
+	const TArray<TVec2<int32>>& GetConstraints() const { return Constraints; }
+
+	CHAOS_API FSolverReal GetThreshold() const;
+	
+	/**
+	* Returns if edges (after being pruned during initialization) are deformed above threshold compared to the rest positions.
+	* @param Positions The solver particle positions, e.g. ClothingSimulationSolver.GetParticleXs().
+	* @param Threshold edge ratio threshold to trigger detection.
+	*/
+	CHAOS_API bool IsExtremelyDeformed(TConstArrayView<Softs::FSolverVec3> Positions) const;
+
+	/**
+	* Returns if edges (after being pruned during initialization) are deformed above threshold compared to the reference positions.
+	* @param Positions The solver particle positions, e.g. ClothingSimulationSolver.GetParticleXs().
+	* @param ReferencePositions The reference particle positions, e.g. ClothingSimulationSolver.GetAnimationPositions(ParticleRangeId) if using the animation pose as reference.
+	* @param Threshold edge ratio threshold to trigger detection.
+	*/
+	CHAOS_API bool IsExtremelyDeformed(TConstArrayView<Softs::FSolverVec3> Positions, const Softs::FSolverVec3* const RefencePositions) const;
+
+	/**
+	* Returns all edges (after being pruned during initialization) that are deformed above threshold compared to the rest positions.
+	* @param Positions The solver particle positions, e.g. ClothingSimulationSolver.GetParticleXs().
+	* @param Threshold edge ratio threshold to trigger detection.
+	*/
+	CHAOS_API TArray<TVec2<int32>> GetExtremelyDeformedEdges(TConstArrayView<Softs::FSolverVec3> Positions) const;
+
+	/**
+	* Returns edges that are deformed above threshold compared to reference positions.
+	* @param Positions The solver particle positions, e.g. ClothingSimulationSolver.GetParticleXs().
+	* @param ReferencePositions The reference particle positions, e.g. ClothingSimulationSolver.GetAnimationPositions(ParticleRangeId) if using the animation pose as reference.
+	* @param Threshold edge ratio threshold to trigger detection.
+	*/
+	CHAOS_API TArray<TVec2<int32>> GetExtremelyDeformedEdges(TConstArrayView<Softs::FSolverVec3> Positions, const Softs::FSolverVec3* const RefencePositions) const;
+
+private:
+	template<int32 Valence, typename Predicate UE_REQUIRES(Valence >= 2 && Valence <= 4)>
+	static TArray<TVector<int32, 2>> TrimConstraints(const TArray<TVector<int32, Valence>>& InConstraints, Predicate TrimPredicate)
+	{
+		TSet<TVec2<int32>> TrimmedConstraints;
+		TrimmedConstraints.Reserve(Valence == 2 ? InConstraints.Num() : InConstraints.Num() * Chaos::Utilities::NChooseR(Valence, 2));
+
+		for (const TVector<int32, Valence>& ConstraintV : InConstraints)
+		{
+			for (int32 i = 0; i < Valence - 1; ++i)
+			{
+				for (int32 j = i + 1; j < Valence; ++j)
+				{
+					const int32 IndexI = ConstraintV[i];
+					const int32 IndexJ = ConstraintV[j];
+
+					if (!TrimPredicate(IndexI, IndexJ))
+					{
+						TrimmedConstraints.Add(IndexI <= IndexJ ? TVec2<int32>(IndexI, IndexJ) : TVec2<int32>(IndexJ, IndexI));
+					}
+				}
+			}
+		}
+		return TrimmedConstraints.Array();
+	}
+
+	const TSet<int32>* VertexSet;
+	TArray<TVec2<int32>> Constraints;
+	const int32 ParticleOffset;
+	const int32 ParticleCount;
+	FSolverReal ExtremeDeformationThreshold;
+	TArray<FSolverReal> Dists;
+
+	UE_CHAOS_DECLARE_INDEXLESS_PROPERTYCOLLECTION_NAME(ExtremeDeformationVertexSelection, bool);  // Selection set name string property, the bool value is not actually used
+	UE_CHAOS_DECLARE_PROPERTYCOLLECTION_NAME(ExtremeDeformationEdgeRatioThreshold, float);
+};
+
+}  // End namespace Chaos::Softs
