Dataflow : Add context asset store  for creating asset during evaluation to be later saved as permanent assets

diff --git a/Engine/Source/Runtime/Experimental/Dataflow/Core/Private/Dataflow/DataflowContextAssetStore.cpp b/Engine/Source/Runtime/Experimental/Dataflow/Core/Private/Dataflow/DataflowContextAssetStore.cpp
new file mode 100644
--- /dev/null
+++ b/Engine/Source/Runtime/Experimental/Dataflow/Core/Private/Dataflow/DataflowContextAssetStore.cpp
@@ -0,0 +1,56 @@
+// Copyright Epic Games, Inc. All Rights Reserved.
+
+#include "Dataflow/DataflowContextAssetStore.h"
+
+#include "UObject/Package.h"
+#include "Misc/Paths.h"
+
+namespace UE::Dataflow
+{
+	UObject* FContextAssetStore::AddAsset(const FString& AssetPath, const UClass* AssetClass)
+	{
+		if (AssetClass)
+		{
+			FStaticConstructObjectParameters Params(AssetClass);
+			Params.Outer = GetTransientPackage();
+			Params.SetFlags = RF_Public;
+			if (TObjectPtr<UObject> NewObject = StaticConstructObject_Internal(Params))
+			{
+				Assets.Add(AssetPath, NewObject);
+				return NewObject;
+			}
+		}
+		return nullptr;
+	}
+
+	UObject* FContextAssetStore::CommitAsset(const FString& AssetPath)
+	{
+		const TObjectPtr<UObject> TransientAsset = Assets.FindRef(AssetPath);
+		if (TransientAsset)
+		{
+			FString UniqueAssetPath{ AssetPath };
+
+			// make sure we have a unique name 
+			if (FindPackage(nullptr, *UniqueAssetPath))
+			{
+				UniqueAssetPath = MakeUniqueObjectName(nullptr, UPackage::StaticClass(), FName(UniqueAssetPath)).ToString();
+			}
+			// create the corresponding package
+			if (UPackage* const NewAssetPackage = CreatePackage(*UniqueAssetPath))
+			{
+				const FString AssetName = FPaths::GetBaseFilename(UniqueAssetPath);
+				if (UObject* ObjectToCommit = DuplicateObject_Internal(TransientAsset->GetClass(), TransientAsset, NewAssetPackage, FName(AssetName)))
+				{
+					ObjectToCommit->MarkPackageDirty();
+					return ObjectToCommit;
+				}
+			}
+		}
+		return nullptr;
+	}
+
+	void FContextAssetStore::ClearAssets()
+	{
+		Assets.Reset();
+	}
+}
\ No newline at end of file
diff --git a/Engine/Source/Runtime/Experimental/Dataflow/Core/Private/Dataflow/DataflowNodeParameters.cpp b/Engine/Source/Runtime/Experimental/Dataflow/Core/Private/Dataflow/DataflowNodeParameters.cpp
--- a/Engine/Source/Runtime/Experimental/Dataflow/Core/Private/Dataflow/DataflowNodeParameters.cpp
+++ b/Engine/Source/Runtime/Experimental/Dataflow/Core/Private/Dataflow/DataflowNodeParameters.cpp
@@ -557,6 +557,21 @@ namespace UE::Dataflow
 		}
 	}
 
+	UObject* FContext::AddAsset(const FString& AssetPath, const UClass* AssetClass)
+	{
+		return AssetStore.AddAsset(AssetPath, AssetClass);
+	}
+
+	UObject* FContext::CommitAsset(const FString& AssetPath)
+	{
+		return AssetStore.CommitAsset(AssetPath);
+	}
+
+	void FContext::ClearAssets()
+	{
+		AssetStore.ClearAssets();
+	}
+
 	void FContextCache::Serialize(FArchive& Ar)
 	{
 		if (Ar.IsSaving())
diff --git a/Engine/Source/Runtime/Experimental/Dataflow/Core/Public/Dataflow/DataflowContextAssetStore.h b/Engine/Source/Runtime/Experimental/Dataflow/Core/Public/Dataflow/DataflowContextAssetStore.h
new file mode 100644
--- /dev/null
+++ b/Engine/Source/Runtime/Experimental/Dataflow/Core/Public/Dataflow/DataflowContextAssetStore.h
@@ -0,0 +1,76 @@
+// Copyright Epic Games, Inc. All Rights Reserved.
+
+#pragma once
+
+#include "UObject/ObjectPtr.h"
+#include "Containers/StringFwd.h"
+
+class UObject;
+
+namespace UE::Dataflow
+{
+	/** interface to be implemented by the context that expose the asset store API */
+	struct IContextAssetStoreInterface
+	{
+	public:
+		/**
+		* Typed version of the AddAsset virtual method
+		* @see AddAsset
+		*/
+		template <typename T>
+		inline T* AddAssetTyped(const FString& AssetPath)
+		{
+			return Cast<T>(AddAsset(AssetPath, T::StaticClass()));
+		}
+
+		/**
+		* Create an new asset to be commited later ( when the terminal nodes are evaluated )
+		* The return UObject is an object own by the transient package
+		* When commited, the asset will be duplicated to a final package matching AssetPath
+		*/
+		virtual UObject* AddAsset(const FString& AssetPath, const UClass* AssetClass) = 0;
+
+		/**
+		* Commit the asset matching the AssetPath
+		* If such asset was added previously this will duplicate it and save it to a package matching AssetPath
+		* @return the newly created persistent asset corresponding to AssetPath
+		*/
+		virtual UObject* CommitAsset(const FString& AssetPath) = 0;
+
+		/**
+		* Clear all assets from the store
+		*/
+		virtual void ClearAssets() = 0;
+
+	};
+
+	/** asset store to managed creation and storage of assets during the evaluation of the Dataflow graph */
+	struct FContextAssetStore
+	{
+	public:
+		/**
+		* Create an new asset to be commited later ( when the terminal nodes are evaluated )
+		* The return UObject is an object own by the transient package
+		* When commited, the asset will be duplicated to a final package matching AssetPath
+		*/
+		DATAFLOWCORE_API UObject* AddAsset(const FString& AssetPath, const UClass* AssetClass);
+
+		/**
+		* Commit the asset matching the AssetPath 
+		* If such asset was added previously this will duplicate it and save it to a package matching AssetPath
+		* @return the newly created persistent asset corresponding to AssetPath
+		*/
+		DATAFLOWCORE_API UObject* CommitAsset(const FString& AssetPath);
+
+		/**
+		* Clear all assets from the store 
+		*/
+		DATAFLOWCORE_API void ClearAssets();
+
+	private:
+		TMap<FString, TObjectPtr<UObject>> Assets;
+	};
+
+}
+
+
diff --git a/Engine/Source/Runtime/Experimental/Dataflow/Core/Public/Dataflow/DataflowNodeParameters.h b/Engine/Source/Runtime/Experimental/Dataflow/Core/Public/Dataflow/DataflowNodeParameters.h
--- a/Engine/Source/Runtime/Experimental/Dataflow/Core/Public/Dataflow/DataflowNodeParameters.h
+++ b/Engine/Source/Runtime/Experimental/Dataflow/Core/Public/Dataflow/DataflowNodeParameters.h
@@ -8,6 +8,7 @@
 #include "Templates/UniquePtr.h"
 #include "HAL/CriticalSection.h"
 #include "Dataflow/DataflowContextCache.h"
+#include "Dataflow/DataflowContextAssetStore.h"
 #include "Dataflow/DataflowContextEvaluator.h"
 #include "Serialization/Archive.h"
 
@@ -126,7 +127,7 @@ namespace UE::Dataflow
 	/*
 	* Dataflow context base class
 	*/
-	class FContext: public IContextCacheStore
+	class FContext: public IContextCacheStore, public IContextAssetStoreInterface
 	{
 	protected:
 		FContext(FContext&&) = delete;
@@ -318,6 +319,12 @@ namespace UE::Dataflow
 		}
 		// End - IContextCacheStore
 
+		// Begin - IContextAssetStoreInterface
+		DATAFLOWCORE_API virtual UObject* AddAsset(const FString& AssetPath, const UClass* AssetClass) override;
+		DATAFLOWCORE_API virtual UObject* CommitAsset(const FString& AssetPath) override;
+		DATAFLOWCORE_API virtual void ClearAssets() override;
+		// End - IContextAssetStoreInterface
+
 		DATAFLOWCORE_API FTimestamp GetTimestamp(FContextCacheKey Key) const;
 
 		DATAFLOWCORE_API void PushToCallstack(const FDataflowConnection* Connection);
@@ -375,6 +382,9 @@ namespace UE::Dataflow
 		TSet<const FDataflowNode*> NodesWithWarning;
 		TSet<const FDataflowNode*> NodesWithError;
 		TSet<const FDataflowNode*> NodesFailed;
+
+		/** used to store the dependent asset created during the evaluation of the graph */
+		FContextAssetStore AssetStore;
 	};
 
 	struct FContextScopedCallstack
diff --git a/Engine/Source/Runtime/Experimental/Dataflow/Engine/Private/Dataflow/DataflowSubGraphNodes.cpp b/Engine/Source/Runtime/Experimental/Dataflow/Engine/Private/Dataflow/DataflowSubGraphNodes.cpp
--- a/Engine/Source/Runtime/Experimental/Dataflow/Engine/Private/Dataflow/DataflowSubGraphNodes.cpp
+++ b/Engine/Source/Runtime/Experimental/Dataflow/Engine/Private/Dataflow/DataflowSubGraphNodes.cpp
@@ -124,6 +124,21 @@ namespace UE::Dataflow::SubGraphNodes::Private
 		void SetCurrentArrayIndex(int32 Index) { CurrentArrayIndex = Index; }
 		int32 GetCurrentArrayIndex() const { return CurrentArrayIndex; }
 
+		virtual UObject* AddAsset(const FString& AssetPath, const UClass* AssetClass) override
+		{
+			return ParentContext.AddAsset(AssetPath, AssetClass);
+		}
+
+		virtual UObject* CommitAsset(const FString& AssetPath) override
+		{
+			return ParentContext.CommitAsset(AssetPath);
+		}
+
+		virtual void ClearAssets() override
+		{
+			ParentContext.ClearAssets();
+		}
+
 	private:
 		int32 CurrentArrayIndex = 0;
 		FContext& ParentContext;
