Don't wait forever when trying to delete a CAS block.

#rnx

diff --git a/Engine/Source/Runtime/Core/Private/IO/IoDispatcher.cpp b/Engine/Source/Runtime/Core/Private/IO/IoDispatcher.cpp
--- a/Engine/Source/Runtime/Core/Private/IO/IoDispatcher.cpp
+++ b/Engine/Source/Runtime/Core/Private/IO/IoDispatcher.cpp
@@ -58,7 +58,8 @@ static const TCHAR* const GetIoErrorText_ErrorCodeTextArray[] =
 	TEXT("Pending Encryption Key"),
 	TEXT("Disabled"),
 	TEXT("Not Installed"),
-	TEXT("Pending Host Group")
+	TEXT("Pending Host Group"),
+	TEXT("Timeout")
 };
 
 const TCHAR* const* GetIoErrorText_ErrorCodeText = GetIoErrorText_ErrorCodeTextArray;
diff --git a/Engine/Source/Runtime/Core/Public/IO/IoStatus.h b/Engine/Source/Runtime/Core/Public/IO/IoStatus.h
--- a/Engine/Source/Runtime/Core/Public/IO/IoStatus.h
+++ b/Engine/Source/Runtime/Core/Public/IO/IoStatus.h
@@ -35,6 +35,7 @@ enum class EIoErrorCode
 	Disabled,
 	NotInstalled,
 	PendingHostGroup,
+	Timeout,
 	Last
 };
 
diff --git a/Engine/Source/Runtime/Experimental/IoStore/OnDemand/Private/OnDemandInstallCache.cpp b/Engine/Source/Runtime/Experimental/IoStore/OnDemand/Private/OnDemandInstallCache.cpp
--- a/Engine/Source/Runtime/Experimental/IoStore/OnDemand/Private/OnDemandInstallCache.cpp
+++ b/Engine/Source/Runtime/Experimental/IoStore/OnDemand/Private/OnDemandInstallCache.cpp
@@ -155,6 +155,8 @@ using FCasBlockInfoMap = TMap<FCasBlockId, FCasBlockInfo>;
 ///////////////////////////////////////////////////////////////////////////////
 struct FCas
 {
+	static constexpr uint32		DeleteBlockMaxWaitTimeMs = 10000;
+
 	using FLookup				= TMap<FCasAddr, FCasLocation>;
 	using FReadHandles			= TMap<FCasBlockId, FWeakAsyncFileHandle>;
 	using FLastAccess			= TMap<FCasBlockId, int64>;
@@ -268,6 +270,9 @@ FIoStatus FCas::DeleteBlock(FCasBlockId BlockId, TArray<FCasAddr>& OutAddrs)
 	UE::TDynamicUniqueLock Lock(Mutex, UE::FDeferLock());
 
 	// Wait for pending reads to flush before deleting block
+	uint32			StartTimeCycles	= FPlatformTime::Cycles();
+	const uint32	WaitTimeMs		= 1000;
+
 	for (;;)
 	{
 		Lock.Lock();
@@ -276,7 +281,16 @@ FIoStatus FCas::DeleteBlock(FCasBlockId BlockId, TArray<FCasAddr>& OutAddrs)
 		if (RequestCount)
 		{
 			Lock.Unlock();
-			BlockReadsDoneEvent->Wait();
+
+			if (FPlatformTime::ToMilliseconds(FPlatformTime::Cycles() - StartTimeCycles) > DeleteBlockMaxWaitTimeMs)
+			{
+				return FIoStatusBuilder(EIoErrorCode::Timeout)
+					<< TEXT("Timed out waiting for pending read(s) while deleting CAS block '")
+					<< Filename
+					<< TEXT("'"); 
+			}
+
+			BlockReadsDoneEvent->Wait(WaitTimeMs);
 		}
 		else
 		{
