IAD - fix AddReferencesToBlocks overcounting TotalReferencedBytes

diff --git a/Engine/Source/Runtime/Experimental/IoStore/OnDemand/Private/OnDemandInstallCache.cpp b/Engine/Source/Runtime/Experimental/IoStore/OnDemand/Private/OnDemandInstallCache.cpp
--- a/Engine/Source/Runtime/Experimental/IoStore/OnDemand/Private/OnDemandInstallCache.cpp
+++ b/Engine/Source/Runtime/Experimental/IoStore/OnDemand/Private/OnDemandInstallCache.cpp
@@ -2280,10 +2280,14 @@ uint64 FOnDemandInstallCache::AddReferencesToBlocks(
 			const bool bToInstall = ChunksToInstall.Contains(Entry.Hash);
 			const bool bIsReferenced = IsReferenced[EntryIndex];
 
+			if (bIsReferenced)
+			{
+				OutTotalReferencedBytes += Entry.EncodedSize;
+			}
+
 			FCasBlockInfo* BlockInfo = nullptr;
 			if (bToInstall || bIsReferenced)
 			{
-				OutTotalReferencedBytes += Entry.EncodedSize;
 				if (FCasLocation Loc = Cas.FindChunk(Entry.Hash); Loc.IsValid())
 				{
 					BlockInfo = BlockInfoMap.Find(Loc.BlockId);
