Iris
* Set the stalled must be mapped reference timeout back to 30secs (was set 9secs recently)
* Don't log entire replication state for the stalled object

#rnx

diff --git a/Engine/Source/Runtime/Experimental/Iris/Core/Private/Iris/ReplicationSystem/ObjectReplicationBridge.cpp b/Engine/Source/Runtime/Experimental/Iris/Core/Private/Iris/ReplicationSystem/ObjectReplicationBridge.cpp
--- a/Engine/Source/Runtime/Experimental/Iris/Core/Private/Iris/ReplicationSystem/ObjectReplicationBridge.cpp
+++ b/Engine/Source/Runtime/Experimental/Iris/Core/Private/Iris/ReplicationSystem/ObjectReplicationBridge.cpp
@@ -2577,7 +2577,7 @@ void UObjectReplicationBridge::OnErrorWithNetRefHandleReported(UE::Net::ENetRefH
 			return;
 		}
 
-		UE_LOG_BRIDGEID(LogIris, Error, TEXT("OnErrorWithNetRefHandleReported: %s from client %s. Problematic object was %s"), LexToString(ErrorType), *PrintConnectionInfo(ConnectionId), *NetRefHandleManager->PrintObjectFromIndex(ObjectInternalIndex));
+		UE_LOG_BRIDGEID(LogIris, Error, TEXT("OnErrorWithNetRefHandleReported: %s for %s. Client: %s"), LexToString(ErrorType), *NetRefHandleManager->PrintObjectFromIndex(ObjectInternalIndex), *PrintConnectionInfo(ConnectionId));
 	}
 
 	// Make a custom printer for the specific error
@@ -2600,7 +2600,7 @@ void UObjectReplicationBridge::OnErrorWithNetRefHandleReported(UE::Net::ENetRefH
 	}
 
 	// All all internal info we have on the blocked object
-	PrintDebugInfoForNetRefHandle(RefHandle, ConnectionId, UObjectReplicationBridge::EPrintDebugInfoTraits::NoProtocolState);
+	//PrintDebugInfoForNetRefHandle(RefHandle, ConnectionId, UObjectReplicationBridge::EPrintDebugInfoTraits::NoProtocolState);
 }
 
 TArray<uint32> UObjectReplicationBridge::FindConnectionsFromArgs(const TArray<FString>& Args) const
diff --git a/Engine/Source/Runtime/Experimental/Iris/Core/Private/Iris/ReplicationSystem/ReplicationReader.cpp b/Engine/Source/Runtime/Experimental/Iris/Core/Private/Iris/ReplicationSystem/ReplicationReader.cpp
--- a/Engine/Source/Runtime/Experimental/Iris/Core/Private/Iris/ReplicationSystem/ReplicationReader.cpp
+++ b/Engine/Source/Runtime/Experimental/Iris/Core/Private/Iris/ReplicationSystem/ReplicationReader.cpp
@@ -113,7 +113,7 @@ static FAutoConsoleVariableRef CVarQueuedBatchTimeoutWarningInterval(
 	QueuedBatchTimeoutWarningInterval,
 	TEXT("Number of ticks between checks for queuedbatches timeout."));
 
-static float QueuedBatchTimeoutSeconds = 9.0f;
+static float QueuedBatchTimeoutSeconds = 30.f;
 static FAutoConsoleVariableRef CVarQueuedBatchTimeoutSeconds(
 	TEXT("net.QueuedBatchTimeoutSeconds"),
 	QueuedBatchTimeoutSeconds,
