Dataflow:
- Added CenterColor and CenterSize to proximity debug draw
- Added sphere coloring to SphereCovering debug draw

diff --git a/Engine/Plugins/Experimental/GeometryCollectionPlugin/Source/GeometryCollectionNodes/Private/Dataflow/GeometryCollectionClusteringNodes.cpp b/Engine/Plugins/Experimental/GeometryCollectionPlugin/Source/GeometryCollectionNodes/Private/Dataflow/GeometryCollectionClusteringNodes.cpp
--- a/Engine/Plugins/Experimental/GeometryCollectionPlugin/Source/GeometryCollectionNodes/Private/Dataflow/GeometryCollectionClusteringNodes.cpp
+++ b/Engine/Plugins/Experimental/GeometryCollectionPlugin/Source/GeometryCollectionNodes/Private/Dataflow/GeometryCollectionClusteringNodes.cpp
@@ -36,6 +36,7 @@
 #include "Dataflow/DataflowRenderingViewMode.h"
 #endif
 #include "Dataflow/DataflowEngineUtil.h"
+#include "Dataflow/GeometryCollectionUtils.h"
 
 #include UE_INLINE_GENERATED_CPP_BY_NAME(GeometryCollectionClusteringNodes)
 
@@ -166,41 +167,14 @@ void FAutoClusterDataflowNode::DebugDraw(UE::Dataflow::FContext& Context, IDataf
 		{
 			const FManagedArrayCollection& OutCollection = Output->GetValue(Context, Collection);
 
-			DataflowRenderingInterface.SetLineWidth(LineWidthMultiplier);
-			DataflowRenderingInterface.SetWireframe(true);
-			DataflowRenderingInterface.SetWorldPriority();
-			DataflowRenderingInterface.SetColor(Color);
-
-			if (OutCollection.HasAttribute("Proximity", FGeometryCollection::GeometryGroup) &&
-				OutCollection.HasAttribute("BoundingBox", FGeometryCollection::GeometryGroup))
-			{
-				const TManagedArray<TSet<int32>>& Proximity = OutCollection.GetAttribute<TSet<int32>>("Proximity", FGeometryCollection::GeometryGroup);
-				const TManagedArray<FBox>& BoundingBox = OutCollection.GetAttribute<FBox>("BoundingBox", FGeometryCollection::GeometryGroup);
-
-				const int32 NumGeometry = OutCollection.NumElements(FGeometryCollection::GeometryGroup);
-				for (int32 Idx = 0; Idx < NumGeometry; ++Idx)
-				{
-					if (Proximity[Idx].Num() > 0)
-					{
-						const FVector Center = BoundingBox[Idx].GetCenter();
-
-						if (bRandomizeColor)
-						{
-							DataflowRenderingInterface.SetColor(UE::Dataflow::Color::GetRandomColor(ColorRandomSeed + 17, Idx));
-						}
-
-						DataflowRenderingInterface.DrawSphere(Center, 1.f);
-
-						for (int32 IdxOther : Proximity[Idx])
-						{
-							const FVector CenterOther = BoundingBox[IdxOther].GetCenter();
-
-							DataflowRenderingInterface.DrawSphere(CenterOther, 1.f);
-							DataflowRenderingInterface.DrawLine(Center, CenterOther);
-						}
-					}
-				}
-			}
+			UE::Dataflow::Utils::DebugDrawProximity(DataflowRenderingInterface,
+				OutCollection,
+				Color,
+				LineWidthMultiplier,
+				CenterSize,
+				CenterColor,
+				bRandomizeColor,
+				ColorRandomSeed);
 		}
 	}
 }
diff --git a/Engine/Plugins/Experimental/GeometryCollectionPlugin/Source/GeometryCollectionNodes/Private/Dataflow/GeometryCollectionNodes.cpp b/Engine/Plugins/Experimental/GeometryCollectionPlugin/Source/GeometryCollectionNodes/Private/Dataflow/GeometryCollectionNodes.cpp
--- a/Engine/Plugins/Experimental/GeometryCollectionPlugin/Source/GeometryCollectionNodes/Private/Dataflow/GeometryCollectionNodes.cpp
+++ b/Engine/Plugins/Experimental/GeometryCollectionPlugin/Source/GeometryCollectionNodes/Private/Dataflow/GeometryCollectionNodes.cpp
@@ -43,6 +43,7 @@
 #if WITH_EDITOR
 #include "Dataflow/DataflowRenderingViewMode.h"
 #endif
+#include "Dataflow/GeometryCollectionUtils.h"
 
 #include UE_INLINE_GENERATED_CPP_BY_NAME(GeometryCollectionNodes)
 
@@ -1008,41 +1009,14 @@ void FProximityDataflowNode::DebugDraw(UE::Dataflow::FContext& Context, IDataflo
 		{
 			const FManagedArrayCollection& OutCollection = Output->GetValue(Context, Collection);
 
-			DataflowRenderingInterface.SetLineWidth(LineWidthMultiplier);
-			DataflowRenderingInterface.SetWireframe(true);
-			DataflowRenderingInterface.SetWorldPriority();
-			DataflowRenderingInterface.SetColor(Color);
-
-			if (OutCollection.HasAttribute("Proximity", FGeometryCollection::GeometryGroup) &&
-				OutCollection.HasAttribute("BoundingBox", FGeometryCollection::GeometryGroup))
-			{
-				const TManagedArray<TSet<int32>>& Proximity = OutCollection.GetAttribute<TSet<int32>>("Proximity", FGeometryCollection::GeometryGroup);
-				const TManagedArray<FBox>& BoundingBox = OutCollection.GetAttribute<FBox>("BoundingBox", FGeometryCollection::GeometryGroup);
-
-				const int32 NumGeometry = OutCollection.NumElements(FGeometryCollection::GeometryGroup);
-				for (int32 Idx = 0; Idx < NumGeometry; ++Idx)
-				{
-					if (Proximity[Idx].Num() > 0)
-					{
-						const FVector Center = BoundingBox[Idx].GetCenter();
-
-						if (bRandomizeColor)
-						{
-							DataflowRenderingInterface.SetColor(UE::Dataflow::Color::GetRandomColor(ColorRandomSeed + 17, Idx));
-						}
-
-						DataflowRenderingInterface.DrawSphere(Center, 1.f);
-
-						for (int32 IdxOther : Proximity[Idx])
-						{
-							const FVector CenterOther = BoundingBox[IdxOther].GetCenter();
-
-							DataflowRenderingInterface.DrawSphere(CenterOther, 1.f);
-							DataflowRenderingInterface.DrawLine(Center, CenterOther);
-						}
-					}
-				}
-			}
+			UE::Dataflow::Utils::DebugDrawProximity(DataflowRenderingInterface,
+				OutCollection,
+				Color,
+				LineWidthMultiplier,
+				CenterSize,
+				CenterColor,
+				bRandomizeColor,
+				ColorRandomSeed);
 		}
 	}
 }
diff --git a/Engine/Plugins/Experimental/GeometryCollectionPlugin/Source/GeometryCollectionNodes/Private/Dataflow/GeometryCollectionUtilityNodes.cpp b/Engine/Plugins/Experimental/GeometryCollectionPlugin/Source/GeometryCollectionNodes/Private/Dataflow/GeometryCollectionUtilityNodes.cpp
--- a/Engine/Plugins/Experimental/GeometryCollectionPlugin/Source/GeometryCollectionNodes/Private/Dataflow/GeometryCollectionUtilityNodes.cpp
+++ b/Engine/Plugins/Experimental/GeometryCollectionPlugin/Source/GeometryCollectionNodes/Private/Dataflow/GeometryCollectionUtilityNodes.cpp
@@ -151,14 +151,38 @@ namespace UE::Dataflow::Convex
 				N = NumSpheres / CMaxNumberOfSpheres;
 			}
 
+			float MinRadius = FLT_MAX, MaxRadius = -FLT_MAX;
 			for (int32 Idx = 0; Idx < NumSpheres; ++Idx)
 			{
 				if (Idx % N == 0)
 				{
-					if (SphereCoveringDebugDrawRenderSettings.bRandomizeColor)
+					const float Radius = OutSpheres.Spheres.GetRadius(Idx);
+
+					if (Radius < MinRadius)
+					{
+						MinRadius = Radius;
+					}
+					else if (Radius > MaxRadius)
+					{
+							MaxRadius = Radius;
+					}
+				}
+			}
+
+			for (int32 Idx = 0; Idx < NumSpheres; ++Idx)
+			{
+				if (Idx % N == 0)
+				{
+					if (SphereCoveringDebugDrawRenderSettings.ColorMethod == EDataflowSphereCoveringColorMethod::Random)
 					{
 						DataflowRenderingInterface.SetColor(UE::Dataflow::Convex::GetRandomColor(SphereCoveringDebugDrawRenderSettings.ColorRandomSeed + 7, Idx));
 					}
+					else if (SphereCoveringDebugDrawRenderSettings.ColorMethod == EDataflowSphereCoveringColorMethod::ColorByRadius)
+					{
+						float Progress = (OutSpheres.Spheres.GetRadius(Idx) - MinRadius) / (MaxRadius - MinRadius);
+						FLinearColor Color = FLinearColor::LerpUsingHSV(SphereCoveringDebugDrawRenderSettings.ColorA, SphereCoveringDebugDrawRenderSettings.ColorB, Progress);
+						DataflowRenderingInterface.SetColor(Color);
+					}
 
 					DataflowRenderingInterface.DrawSphere(OutSpheres.Spheres.GetCenter(Idx), OutSpheres.Spheres.GetRadius(Idx));
 				}
diff --git a/Engine/Plugins/Experimental/GeometryCollectionPlugin/Source/GeometryCollectionNodes/Private/Dataflow/GeometryCollectionUtils.cpp b/Engine/Plugins/Experimental/GeometryCollectionPlugin/Source/GeometryCollectionNodes/Private/Dataflow/GeometryCollectionUtils.cpp
new file mode 100644
--- /dev/null
+++ b/Engine/Plugins/Experimental/GeometryCollectionPlugin/Source/GeometryCollectionNodes/Private/Dataflow/GeometryCollectionUtils.cpp
@@ -0,0 +1,70 @@
+// Copyright Epic Games, Inc. All Rights Reserved.
+
+#include "Dataflow/GeometryCollectionUtils.h"
+
+namespace UE::Dataflow
+{
+	namespace Utils
+	{
+#if WITH_EDITOR
+		void DebugDrawProximity(IDataflowDebugDrawInterface& DataflowRenderingInterface,
+			const FManagedArrayCollection& Collection,
+			const FLinearColor Color,
+			const float LineWidthMultiplier,
+			const float CenterSize,
+			const FLinearColor CenterColor,
+			const bool bRandomizeColor,
+			const int32 ColorRandomSeed)
+		{
+			DataflowRenderingInterface.SetLineWidth(LineWidthMultiplier);
+			DataflowRenderingInterface.SetWireframe(true);
+			DataflowRenderingInterface.SetWorldPriority();
+			DataflowRenderingInterface.SetPointSize(CenterSize);
+
+			if (Collection.HasAttribute("Proximity", FGeometryCollection::GeometryGroup) &&
+				Collection.HasAttribute("BoundingBox", FGeometryCollection::GeometryGroup))
+			{
+				const TManagedArray<TSet<int32>>& Proximity = Collection.GetAttribute<TSet<int32>>("Proximity", FGeometryCollection::GeometryGroup);
+				const TManagedArray<FBox>& BoundingBox = Collection.GetAttribute<FBox>("BoundingBox", FGeometryCollection::GeometryGroup);
+
+				const int32 NumGeometry = Collection.NumElements(FGeometryCollection::GeometryGroup);
+				TArray<FVector> Centers;
+				for (int32 Idx = 0; Idx < NumGeometry; ++Idx)
+				{
+					if (Proximity[Idx].Num() > 0)
+					{
+						const FVector Center = BoundingBox[Idx].GetCenter();
+
+						Centers.Add(Center);
+
+						if (bRandomizeColor)
+						{
+							DataflowRenderingInterface.SetColor(UE::Dataflow::Color::GetRandomColor(ColorRandomSeed + 17, Idx));
+						}
+						else
+						{
+							DataflowRenderingInterface.SetColor(Color);
+						}
+
+						for (int32 IdxOther : Proximity[Idx])
+						{
+							const FVector CenterOther = BoundingBox[IdxOther].GetCenter();
+
+							DataflowRenderingInterface.DrawLine(Center, CenterOther);
+						}
+
+						DataflowRenderingInterface.SetPointSize(CenterSize);
+						DataflowRenderingInterface.SetColor(CenterColor);
+						for (FVector& Point : Centers)
+						{
+							DataflowRenderingInterface.DrawPoint(Point);
+						}
+
+						Centers.Empty();
+					}
+				}
+			}
+		}
+#endif
+	}
+}
diff --git a/Engine/Plugins/Experimental/GeometryCollectionPlugin/Source/GeometryCollectionNodes/Public/Dataflow/GeometryCollectionClusteringNodes.h b/Engine/Plugins/Experimental/GeometryCollectionPlugin/Source/GeometryCollectionNodes/Public/Dataflow/GeometryCollectionClusteringNodes.h
--- a/Engine/Plugins/Experimental/GeometryCollectionPlugin/Source/GeometryCollectionNodes/Public/Dataflow/GeometryCollectionClusteringNodes.h
+++ b/Engine/Plugins/Experimental/GeometryCollectionPlugin/Source/GeometryCollectionNodes/Public/Dataflow/GeometryCollectionClusteringNodes.h
@@ -113,6 +113,12 @@ private:
 	UPROPERTY(EditAnywhere, Category = "Debug Draw", meta = (ClampMin = "0.1", ClampMax = "10.0"));
 	float LineWidthMultiplier = 2.f;
 
+	UPROPERTY(EditAnywhere, Category = "Debug Draw - Sphere Covering");
+	FLinearColor CenterColor = FLinearColor::Blue;
+
+	UPROPERTY(EditAnywhere, Category = "Debug Draw - Sphere Covering");
+	float CenterSize = 12.f;
+
 	/** Randomize color per connection */
 	UPROPERTY(EditAnywhere, Category = "Debug Draw")
 	bool bRandomizeColor = false;
diff --git a/Engine/Plugins/Experimental/GeometryCollectionPlugin/Source/GeometryCollectionNodes/Public/Dataflow/GeometryCollectionNodes.h b/Engine/Plugins/Experimental/GeometryCollectionPlugin/Source/GeometryCollectionNodes/Public/Dataflow/GeometryCollectionNodes.h
--- a/Engine/Plugins/Experimental/GeometryCollectionPlugin/Source/GeometryCollectionNodes/Public/Dataflow/GeometryCollectionNodes.h
+++ b/Engine/Plugins/Experimental/GeometryCollectionPlugin/Source/GeometryCollectionNodes/Public/Dataflow/GeometryCollectionNodes.h
@@ -1295,6 +1295,12 @@ private:
 	UPROPERTY(EditAnywhere, Category = "Debug Draw", meta = (ClampMin = "0.1", ClampMax = "10.0"));
 	float LineWidthMultiplier = 2.f;
 
+	UPROPERTY(EditAnywhere, Category = "Debug Draw - Sphere Covering");
+	FLinearColor CenterColor = FLinearColor::Blue;
+
+	UPROPERTY(EditAnywhere, Category = "Debug Draw - Sphere Covering");
+	float CenterSize = 12.f;
+
 	/** Randomize color per connection */
 	UPROPERTY(EditAnywhere, Category = "Debug Draw")
 	bool bRandomizeColor = false;
diff --git a/Engine/Plugins/Experimental/GeometryCollectionPlugin/Source/GeometryCollectionNodes/Public/Dataflow/GeometryCollectionUtils.h b/Engine/Plugins/Experimental/GeometryCollectionPlugin/Source/GeometryCollectionNodes/Public/Dataflow/GeometryCollectionUtils.h
new file mode 100644
--- /dev/null
+++ b/Engine/Plugins/Experimental/GeometryCollectionPlugin/Source/GeometryCollectionNodes/Public/Dataflow/GeometryCollectionUtils.h
@@ -0,0 +1,32 @@
+// Copyright Epic Games, Inc. All Rights Reserved.
+
+#pragma once 
+
+#include "Dataflow/DataflowEngine.h"
+#include "Dataflow/DataflowNode.h"
+#include "GeometryCollection/ManagedArrayCollection.h"
+#include "Dataflow/DataflowDebugDrawInterface.h"
+#include "Dataflow/DataflowDebugDraw.h"
+#if WITH_EDITOR
+#include "Dataflow/DataflowRenderingViewMode.h"
+#endif
+#include "Dataflow/DataflowEngineUtil.h"
+
+namespace UE::Dataflow
+{
+	namespace Utils
+	{
+#if WITH_EDITOR
+		GEOMETRYCOLLECTIONNODES_API void DebugDrawProximity(IDataflowDebugDrawInterface& DataflowRenderingInterface,
+			const FManagedArrayCollection& Collection,
+			const FLinearColor Color,
+			const float LineWidthMultiplier,
+			const float CenterSize,
+			const FLinearColor CenterColor,
+			const bool bRandomizeColor,
+			const int32 ColorRandomSeed);
+#endif
+	}
+}
+
+
diff --git a/Engine/Source/Runtime/Experimental/Dataflow/Engine/Public/Dataflow/DataflowDebugDraw.h b/Engine/Source/Runtime/Experimental/Dataflow/Engine/Public/Dataflow/DataflowDebugDraw.h
--- a/Engine/Source/Runtime/Experimental/Dataflow/Engine/Public/Dataflow/DataflowDebugDraw.h
+++ b/Engine/Source/Runtime/Experimental/Dataflow/Engine/Public/Dataflow/DataflowDebugDraw.h
@@ -107,6 +107,14 @@ struct FDataflowNodeDebugDrawSettings
 /**
  * SphereCovering DebugDraw basic common settings
  */
+UENUM(BlueprintType)
+enum class EDataflowSphereCoveringColorMethod : uint8
+{
+	Single UMETA(DisplayName = "Single Color"),
+	ColorByRadius UMETA(DisplayName = "Color by Radius"),
+	Random UMETA(DisplayName = "Random Colors"),
+};
+
 USTRUCT()
 struct FDataflowNodeSphereCoveringDebugDrawSettings
 {
@@ -122,19 +130,24 @@ struct FDataflowNodeSphereCoveringDebugDrawSettings
 	UPROPERTY(EditAnywhere, Category = "Debug Draw - Sphere Covering", meta = (EditCondition = "RenderType==EDataflowDebugDrawRenderType::Shaded", EditConditionHides));
 	bool bTranslucent = true;
 
-	UPROPERTY(EditAnywhere, Category = "Debug Draw - Sphere Covering");
-	FLinearColor Color = FLinearColor::Red;
-
 	UPROPERTY(EditAnywhere, Category = "Debug Draw - Sphere Covering", meta = (ClampMin = "0.1", ClampMax = "10.0"));
 	float LineWidthMultiplier = .25f;
 
-	/** Randomize color per sphere */
-	UPROPERTY(EditAnywhere, Category = "Debug Draw - Sphere Covering")
-	bool bRandomizeColor = false;
+	UPROPERTY(EditAnywhere, Category = "Debug Draw - Sphere Covering");
+	EDataflowSphereCoveringColorMethod ColorMethod = EDataflowSphereCoveringColorMethod::Single;
+
+	UPROPERTY(EditAnywhere, Category = "Debug Draw - Sphere Covering", meta = (EditCondition = "ColorMethod == EDataflowSphereCoveringColorMethod::Single", EditConditionHides));
+	FLinearColor Color = FLinearColor::Red;
 
 	/** Random seed */
-	UPROPERTY(EditAnywhere, Category = "Debug Draw - Sphere Covering", meta = (ClampMin = "0", EditCondition = "bRandomizeColor==true", EditConditionHides))
+	UPROPERTY(EditAnywhere, Category = "Debug Draw - Sphere Covering", meta = (ClampMin = "0", EditCondition = "ColorMethod == EDataflowSphereCoveringColorMethod::Random", EditConditionHides))
 	int32 ColorRandomSeed = 0;
+
+	UPROPERTY(EditAnywhere, Category = "Debug Draw - Sphere Covering", meta = (EditCondition = "ColorMethod == EDataflowSphereCoveringColorMethod::ColorByRadius", EditConditionHides));
+	FLinearColor ColorA = FLinearColor::Red;
+
+	UPROPERTY(EditAnywhere, Category = "Debug Draw - Sphere Covering", meta = (EditCondition = "ColorMethod == EDataflowSphereCoveringColorMethod::ColorByRadius", EditConditionHides));
+	FLinearColor ColorB = FLinearColor::Blue;
 };
 
 static void SetWireframeRender(IDataflowDebugDrawInterface& DataflowRenderingInterface)
