Dataflow : fix connection logic for connection with dependent types

diff --git a/Engine/Source/Runtime/Experimental/Dataflow/Core/Private/Dataflow/DataflowNode.cpp b/Engine/Source/Runtime/Experimental/Dataflow/Core/Private/Dataflow/DataflowNode.cpp
--- a/Engine/Source/Runtime/Experimental/Dataflow/Core/Private/Dataflow/DataflowNode.cpp
+++ b/Engine/Source/Runtime/Experimental/Dataflow/Core/Private/Dataflow/DataflowNode.cpp
@@ -1634,10 +1634,11 @@ bool FDataflowNode::IsAnytypeDependencyConnected(FName DependencyGroup, const FD
 	// first check if any connection is connected 
 	bool bIsAnyDependencyConnected = false;
 	ForEachConnection(
-		[&bIsAnyDependencyConnected, &IgnoreConnection](const FDataflowConnection* Connection)
+		[&bIsAnyDependencyConnected, &IgnoreConnection, &DependencyGroup](const FDataflowConnection* Connection)
 		{
 			// ignore the connection that is actively being notified
 			if (Connection
+				&& (Connection->GetTypeDependencyGroup() == DependencyGroup)
 				&& (IgnoreConnection == nullptr || Connection != IgnoreConnection)
 				&& Connection->IsConnected()
 				)
