Dataflow : make sure the asset store play nice with the garbage collection

diff --git a/Engine/Source/Runtime/Experimental/Dataflow/Core/Private/Dataflow/DataflowContextAssetStore.cpp b/Engine/Source/Runtime/Experimental/Dataflow/Core/Private/Dataflow/DataflowContextAssetStore.cpp
--- a/Engine/Source/Runtime/Experimental/Dataflow/Core/Private/Dataflow/DataflowContextAssetStore.cpp
+++ b/Engine/Source/Runtime/Experimental/Dataflow/Core/Private/Dataflow/DataflowContextAssetStore.cpp
@@ -53,4 +53,10 @@ namespace UE::Dataflow
 	{
 		Assets.Reset();
 	}
+
+	void FContextAssetStore::AddReferencedObjects(FReferenceCollector& Collector) 
+	{ 
+		//for (Assets)
+		Collector.AddReferencedObjects(Assets);
+	}
 }
\ No newline at end of file
diff --git a/Engine/Source/Runtime/Experimental/Dataflow/Core/Public/Dataflow/DataflowContextAssetStore.h b/Engine/Source/Runtime/Experimental/Dataflow/Core/Public/Dataflow/DataflowContextAssetStore.h
--- a/Engine/Source/Runtime/Experimental/Dataflow/Core/Public/Dataflow/DataflowContextAssetStore.h
+++ b/Engine/Source/Runtime/Experimental/Dataflow/Core/Public/Dataflow/DataflowContextAssetStore.h
@@ -4,6 +4,7 @@
 
 #include "UObject/ObjectPtr.h"
 #include "Containers/StringFwd.h"
+#include "UObject/GCObject.h"
 
 class UObject;
 
@@ -45,7 +46,7 @@ namespace UE::Dataflow
 	};
 
 	/** asset store to managed creation and storage of assets during the evaluation of the Dataflow graph */
-	struct FContextAssetStore
+	struct FContextAssetStore: public FGCObject
 	{
 	public:
 		/**
@@ -67,6 +68,11 @@ namespace UE::Dataflow
 		*/
 		DATAFLOWCORE_API void ClearAssets();
 
+		//~ Begin FGCObject interface
+		DATAFLOWCORE_API virtual void AddReferencedObjects(FReferenceCollector& Collector) override;
+		DATAFLOWCORE_API virtual FString GetReferencerName() const override { return TEXT("FContextAssetStore"); }
+		//~ End FGCObject interface
+
 	private:
 		TMap<FString, TObjectPtr<UObject>> Assets;
 	};
