Iris support for serializing object references as remote object references (UNetDriver::IsUsingRemoteObjectReferences()). Main changes:
- Pass the bUseRemoteObjectRefs flag from the NetDriver through the ReplicationSystem and into FInternalNetSerializationContext so that the object net serializer can use it.
- Object serializer quantized state is now a union-like type that stores an FNetObjectReference or an FRemoteObjectReference quantized state. The remote reference state is allocated dynamically to save memory when not in use.
- Object net serializer now writes a bit indicating whether the object is an FNetObjectReference or an FRemoteObjectReference.
- Refactored FRemoteObjectReference net serializer & quantized state to make it easier to work with in the main object serializer.
- Adjusted various Iris internals/reference collection to account for the new object quantized state.
- Increased size of some other serializer quantized states that contained object refs to account for the new size.
- Split serializers for raw UObject* and TObjectPtr. We don't want to serialize the raw pointers as remote refs.
- Add a unit test fixture to allow passing the UseRemoteReference flag to the replication system & unit tests that use it.

diff --git a/Engine/Plugins/Runtime/GameplayAbilities/Source/GameplayAbilities/Private/Serialization/GameplayEffectContextNetSerializer.cpp b/Engine/Plugins/Runtime/GameplayAbilities/Source/GameplayAbilities/Private/Serialization/GameplayEffectContextNetSerializer.cpp
--- a/Engine/Plugins/Runtime/GameplayAbilities/Source/GameplayAbilities/Private/Serialization/GameplayEffectContextNetSerializer.cpp
+++ b/Engine/Plugins/Runtime/GameplayAbilities/Source/GameplayAbilities/Private/Serialization/GameplayEffectContextNetSerializer.cpp
@@ -50,7 +50,7 @@ struct FGameplayEffectContextNetSerializer
 #if UE_WITH_REMOTE_OBJECT_HANDLE
 		alignas(16) uint8 EffectContext[176];
 #else
-		alignas(16) uint8 EffectContext[128];
+		alignas(16) uint8 EffectContext[152];
 #endif
 		// 320 should be large enough and makes struct a multiple of 16 bytes
 		alignas(16) uint8 HitResult[320];
diff --git a/Engine/Plugins/Runtime/GameplayAbilities/Source/GameplayAbilities/Public/Serialization/GameplayEffectContextNetSerializer.h b/Engine/Plugins/Runtime/GameplayAbilities/Source/GameplayAbilities/Public/Serialization/GameplayEffectContextNetSerializer.h
--- a/Engine/Plugins/Runtime/GameplayAbilities/Source/GameplayAbilities/Public/Serialization/GameplayEffectContextNetSerializer.h
+++ b/Engine/Plugins/Runtime/GameplayAbilities/Source/GameplayAbilities/Public/Serialization/GameplayEffectContextNetSerializer.h
@@ -17,7 +17,7 @@ namespace UE::Net
 #if UE_WITH_REMOTE_OBJECT_HANDLE
 constexpr SIZE_T GetGameplayEffectContextNetSerializerSafeQuantizedSize() { return 512; }
 #else
-constexpr SIZE_T GetGameplayEffectContextNetSerializerSafeQuantizedSize() { return 464; }
+constexpr SIZE_T GetGameplayEffectContextNetSerializerSafeQuantizedSize() { return 496; }
 #endif
 
 UE_NET_DECLARE_SERIALIZER(FGameplayEffectContextNetSerializer, GAMEPLAYABILITIES_API);
diff --git a/Engine/Plugins/Runtime/MultiServerReplication/Source/MultiServerReplication/Private/MultiServerBeaconHost.cpp b/Engine/Plugins/Runtime/MultiServerReplication/Source/MultiServerReplication/Private/MultiServerBeaconHost.cpp
--- a/Engine/Plugins/Runtime/MultiServerReplication/Source/MultiServerReplication/Private/MultiServerBeaconHost.cpp
+++ b/Engine/Plugins/Runtime/MultiServerReplication/Source/MultiServerReplication/Private/MultiServerBeaconHost.cpp
@@ -35,6 +35,9 @@ bool AMultiServerBeaconHost::InitHost()
 		{
 			ensureMsgf(NetDriver->IsA<UMultiServerNetDriver>(), TEXT("Multi-server beacon NetDriver should be a subclass of UMultiServerNetDriver to function correctly. Check the NetDriverDefinition for MultiServerNetDriver."));
 
+			// This flag has to be set before InitListen and the ReplicationSystem is created if using Iris
+			NetDriver->SetUsingRemoteObjectReferences(!!UE_WITH_REMOTE_OBJECT_HANDLE && GMultiServerAllowRemoteObjectReferences);
+
 			FString Error;
 			if (NetDriver->InitListen(this, URL, bReuseAddressAndPort, Error))
 			{
@@ -44,7 +47,6 @@ bool AMultiServerBeaconHost::InitHost()
 				NetDriver->InitialConnectTimeout = BeaconConnectionInitialTimeout;
 				NetDriver->ConnectionTimeout = BeaconConnectionTimeout;
 				NetDriver->SetReplicateTransactionally(false);
-				NetDriver->SetUsingRemoteObjectReferences(!!UE_WITH_REMOTE_OBJECT_HANDLE && GMultiServerAllowRemoteObjectReferences);
 				return true;
 			}
 			else
diff --git a/Engine/Plugins/Runtime/ReplicationSystemTestPlugin/Source/Private/Tests/ReplicationSystem/ReplicationSystemConfigOverrideTestFixture.cpp b/Engine/Plugins/Runtime/ReplicationSystemTestPlugin/Source/Private/Tests/ReplicationSystem/ReplicationSystemConfigOverrideTestFixture.cpp
new file mode 100644
--- /dev/null
+++ b/Engine/Plugins/Runtime/ReplicationSystemTestPlugin/Source/Private/Tests/ReplicationSystem/ReplicationSystemConfigOverrideTestFixture.cpp
@@ -0,0 +1,88 @@
+// Copyright Epic Games, Inc. All Rights Reserved.
+
+#include "Tests/ReplicationSystem/ReplicationSystemConfigOverrideTestFixture.h"
+#include "Iris/Core/IrisLog.h"
+
+namespace UE::Net
+{
+
+void FReplicationSystemConfigOverrideRPCTestFixture::SetUp()
+{
+	UNetBlobHandlerDefinitions* BlobHandlerDefinitions = GetMutableDefault<UNetBlobHandlerDefinitions>();
+	OriginalHandlerDefinitions = BlobHandlerDefinitions->ReadWriteHandlerDefinitions();
+
+	HandlerDefinitions.Add({ TEXT("NetRPCHandler") });
+	HandlerDefinitions.Add({ TEXT("PartialNetObjectAttachmentHandler") });
+	HandlerDefinitions.Add({ TEXT("NetObjectBlobHandler") });
+
+	BlobHandlerDefinitions->ReadWriteHandlerDefinitions() = HandlerDefinitions;
+
+	FNetworkAutomationTestSuiteFixture::SetUp();
+
+	// Fake what we normally get from config
+	DataStreamUtil.SetUp();
+	DataStreamUtil.AddDataStreamDefinition(TEXT("NetToken"), TEXT("/Script/IrisCore.NetTokenDataStream"));
+	DataStreamUtil.AddDataStreamDefinition(TEXT("Replication"), TEXT("/Script/IrisCore.ReplicationDataStream"));
+	DataStreamUtil.FixupDefinitions();
+
+	NetTokenDataStoreUtil.SetUp();
+}
+
+void FReplicationSystemConfigOverrideRPCTestFixture::TearDown()
+{
+	for (TArray<FReplicationSystemTestClient*> TempClients = Clients; FReplicationSystemTestClient * Client : TempClients)
+	{
+		delete Client;
+	}
+	delete Server;
+	DataStreamUtil.TearDown();
+	NetTokenDataStoreUtil.TearDown();
+
+	FNetworkAutomationTestSuiteFixture::TearDown();
+
+	UNetBlobHandlerDefinitions* BlobHandlerDefinitions = GetMutableDefault<UNetBlobHandlerDefinitions>();
+	BlobHandlerDefinitions->ReadWriteHandlerDefinitions() = OriginalHandlerDefinitions;
+}
+
+void FReplicationSystemConfigOverrideRPCTestFixture::CreateServer(const FReplicationSystemTestNode::FReplicationSystemParamsOverride& ParamsOverride)
+{
+	checkf(Server == nullptr, TEXT("Server already created"));
+	Server = new FReplicationSystemTestServer(FReplicationSystemTestNode::DelaySetup);
+	Server->Setup(true, GetName(), &ParamsOverride);
+}
+
+FReplicationSystemTestClient* FReplicationSystemConfigOverrideRPCTestFixture::CreateClient(const FReplicationSystemTestNode::FReplicationSystemParamsOverride& ParamsOverride)
+{
+	if (!ensureMsgf(Server, TEXT("Must create a server before a client.")))
+	{
+		return nullptr;
+	}
+
+	FReplicationSystemTestClient* Client = new FReplicationSystemTestClient(FReplicationSystemTestNode::DelaySetup);
+	Client->Setup(false, GetName(), &ParamsOverride);
+
+	Clients.Add(Client);
+
+	// The client needs a connection
+	Client->LocalConnectionId = Client->AddConnection();
+
+	// Auto connect to server
+	Client->ConnectionIdOnServer = Server->AddConnection();
+
+	return Client;
+}
+
+void FReplicationSystemConfigOverrideRPCTestFixture::DestroyClient(FReplicationSystemTestClient* Client)
+{
+	if (!Clients.Remove(Client))
+	{
+		UE_LOG(LogIris, Warning, TEXT("Unable to find FReplicationSystemTestClient %p for destroy. NOT destroying."), Client);
+		return;
+	}
+
+	Server->RemoveConnection(Client->ConnectionIdOnServer);
+
+	delete Client;
+}
+
+}
diff --git a/Engine/Plugins/Runtime/ReplicationSystemTestPlugin/Source/Private/Tests/ReplicationSystem/ReplicationSystemConfigOverrideTestFixture.h b/Engine/Plugins/Runtime/ReplicationSystemTestPlugin/Source/Private/Tests/ReplicationSystem/ReplicationSystemConfigOverrideTestFixture.h
new file mode 100644
--- /dev/null
+++ b/Engine/Plugins/Runtime/ReplicationSystemTestPlugin/Source/Private/Tests/ReplicationSystem/ReplicationSystemConfigOverrideTestFixture.h
@@ -0,0 +1,48 @@
+// Copyright Epic Games, Inc. All Rights Reserved.
+
+#pragma once
+
+#include "Tests/ReplicationSystem/ReplicationSystemServerClientTestFixture.h"
+#include "Iris/ReplicationSystem/NetBlob/NetBlobHandlerDefinitions.h"
+
+namespace UE::Net
+{
+
+/**
+ * A test fixture that allows overriding the replication system params and supports RPCs.
+ * Unlike FReplicationSystemServerClientTestFixture, does not create a server automatically in SetUp()
+ * to give test cases a chance to set up the override params. Call CreateServer() to create a server.
+ */
+class FReplicationSystemConfigOverrideRPCTestFixture : public FNetworkAutomationTestSuiteFixture
+{
+public:
+	FReplicationSystemConfigOverrideRPCTestFixture() : FNetworkAutomationTestSuiteFixture() {}
+
+protected:
+	enum : bool
+	{
+		DoNotDeliverPacket = false,
+		DeliverPacket = true,
+	};
+
+	virtual void SetUp() override;
+	virtual void TearDown() override;
+
+	void CreateServer(const FReplicationSystemTestNode::FReplicationSystemParamsOverride& ParamsOverride);
+	FReplicationSystemTestClient* CreateClient(const FReplicationSystemTestNode::FReplicationSystemParamsOverride& ParamsOverride);
+	void DestroyClient(FReplicationSystemTestClient* Client);
+
+	FDataStreamTestUtil DataStreamUtil;
+	FNetTokenDataStoreTestUtil NetTokenDataStoreUtil;
+	FReplicationSystemTestServer* Server = nullptr;
+	TArray<FReplicationSystemTestClient*> Clients;
+
+	FReplicationSystemTestNode::FReplicationSystemParamsOverride OverrideServerConfig;
+	FReplicationSystemTestNode::FReplicationSystemParamsOverride OverrideClientConfig;
+
+private:
+	TArray<FNetBlobHandlerDefinition> OriginalHandlerDefinitions;
+	TArray<FNetBlobHandlerDefinition> HandlerDefinitions;
+};
+
+}
diff --git a/Engine/Plugins/Runtime/ReplicationSystemTestPlugin/Source/Private/Tests/ReplicationSystem/ReplicationSystemServerClientTestFixture.cpp b/Engine/Plugins/Runtime/ReplicationSystemTestPlugin/Source/Private/Tests/ReplicationSystem/ReplicationSystemServerClientTestFixture.cpp
--- a/Engine/Plugins/Runtime/ReplicationSystemTestPlugin/Source/Private/Tests/ReplicationSystem/ReplicationSystemServerClientTestFixture.cpp
+++ b/Engine/Plugins/Runtime/ReplicationSystemTestPlugin/Source/Private/Tests/ReplicationSystem/ReplicationSystemServerClientTestFixture.cpp
@@ -119,7 +119,7 @@ FReplicationSystemTestNode::FReplicationSystemTestNode(FReplicationSystemTestNod
 	// Do nothing until asked to start up
 }
 
-void FReplicationSystemTestNode::Setup(bool bIsServer, const TCHAR* Name, FReplicationSystemTestNode::FReplicationSystemParamsOverride* ParamsOverride)
+void FReplicationSystemTestNode::Setup(bool bIsServer, const TCHAR* Name, const FReplicationSystemTestNode::FReplicationSystemParamsOverride* ParamsOverride)
 {
 	// Init NetTokenStore
 	{
@@ -154,6 +154,7 @@ void FReplicationSystemTestNode::Setup(bool bIsServer, const TCHAR* Name, FRepli
 		Params.MaxReplicatedObjectCount = ParamsOverride->MaxReplicatedObjectCount > 0 ? ParamsOverride->MaxReplicatedObjectCount : Params.MaxReplicatedObjectCount;
 		Params.InitialNetObjectListCount = ParamsOverride->InitialNetObjectListCount > 0 ? ParamsOverride->InitialNetObjectListCount : Params.InitialNetObjectListCount;
 		Params.NetObjectListGrowCount = ParamsOverride->NetObjectListGrowCount > 0 ? ParamsOverride->NetObjectListGrowCount : Params.NetObjectListGrowCount;
+		Params.bUseRemoteObjectReferences = ParamsOverride->bUseRemoteObjectReferences;
 	}
 
 	LOG_SCOPE_VERBOSITY_OVERRIDE(LogIris, ELogVerbosity::Error);
diff --git a/Engine/Plugins/Runtime/ReplicationSystemTestPlugin/Source/Private/Tests/ReplicationSystem/ReplicationSystemServerClientTestFixture.h b/Engine/Plugins/Runtime/ReplicationSystemTestPlugin/Source/Private/Tests/ReplicationSystem/ReplicationSystemServerClientTestFixture.h
--- a/Engine/Plugins/Runtime/ReplicationSystemTestPlugin/Source/Private/Tests/ReplicationSystem/ReplicationSystemServerClientTestFixture.h
+++ b/Engine/Plugins/Runtime/ReplicationSystemTestPlugin/Source/Private/Tests/ReplicationSystem/ReplicationSystemServerClientTestFixture.h
@@ -99,6 +99,7 @@ public:
 		uint32 MaxReplicatedObjectCount = 0;
 		uint32 InitialNetObjectListCount = 0;
 		uint32 NetObjectListGrowCount = 0;
+		bool bUseRemoteObjectReferences = false;
 	};
 
 	enum EDelaySetup { DelaySetup=0 };
@@ -108,7 +109,7 @@ public:
 	explicit FReplicationSystemTestNode(FReplicationSystemTestNode::EDelaySetup);
 	virtual ~FReplicationSystemTestNode();
 
-	void Setup(bool bIsServer, const TCHAR* Name, FReplicationSystemTestNode::FReplicationSystemParamsOverride* ParamsOverride=nullptr);
+	void Setup(bool bIsServer, const TCHAR* Name, const FReplicationSystemTestNode::FReplicationSystemParamsOverride* ParamsOverride=nullptr);
 
 	template<typename T>
 	T* CreateObject()
diff --git a/Engine/Plugins/Runtime/ReplicationSystemTestPlugin/Source/Private/Tests/ReplicationSystem/TestReplicationSystem.cpp b/Engine/Plugins/Runtime/ReplicationSystemTestPlugin/Source/Private/Tests/ReplicationSystem/TestReplicationSystem.cpp
--- a/Engine/Plugins/Runtime/ReplicationSystemTestPlugin/Source/Private/Tests/ReplicationSystem/TestReplicationSystem.cpp
+++ b/Engine/Plugins/Runtime/ReplicationSystemTestPlugin/Source/Private/Tests/ReplicationSystem/TestReplicationSystem.cpp
@@ -12,6 +12,7 @@
 #include "Logging/LogScopedVerbosityOverride.h"
 #include "Misc/ScopeExit.h"
 #include "Net/Core/NetToken/NetToken.h"
+#include "UObject/CoreNetTypes.h"
 
 namespace UE::Net::Private
 {
diff --git a/Engine/Plugins/Runtime/ReplicationSystemTestPlugin/Source/Private/Tests/Serialization/TestObjectsAsRemoteReferencesSerialization.cpp b/Engine/Plugins/Runtime/ReplicationSystemTestPlugin/Source/Private/Tests/Serialization/TestObjectsAsRemoteReferencesSerialization.cpp
new file mode 100644
--- /dev/null
+++ b/Engine/Plugins/Runtime/ReplicationSystemTestPlugin/Source/Private/Tests/Serialization/TestObjectsAsRemoteReferencesSerialization.cpp
@@ -0,0 +1,138 @@
+// Copyright Epic Games, Inc. All Rights Reserved.
+
+#include "TestObjectsAsRemoteReferencesSerialization.h"
+#include "Tests/ReplicationSystem/ReplicationSystemConfigOverrideTestFixture.h"
+#include "Tests/Serialization/TestRemoteObjectReferenceNetSerializer.h"
+#include "UObject/UObjectGlobals.h"
+#include "UObject/Package.h"
+
+void UTestObjectWithReferencesAsRemote::GetLifetimeReplicatedProps(TArray<class FLifetimeProperty>& OutLifetimeProps) const
+{
+	Super::GetLifetimeReplicatedProps(OutLifetimeProps);
+}
+
+void UTestObjectWithReferencesAsRemote::ClientRPCWithRawPointer_Implementation(UObject* RawPointer)
+{
+	LastReceivedRawPointer = RawPointer;
+}
+
+#if UE_WITH_REMOTE_OBJECT_HANDLE
+
+namespace UE::Net::Private
+{
+
+UE_NET_TEST_FIXTURE(FReplicationSystemConfigOverrideRPCTestFixture, TestObjectsAsRemoteReferences)
+{
+	FReplicationSystemTestNode::FReplicationSystemParamsOverride Params;
+	Params.bUseRemoteObjectReferences = true;
+
+	// Create the server
+	CreateServer(Params);
+
+	// Add a client
+	FReplicationSystemTestClient* Client = CreateClient(Params);
+
+	// Spawn objects on server
+	UTestObjectWithReferencesAsRemote* TestObject = Server->CreateObject<UTestObjectWithReferencesAsRemote>();
+	UReplicatedTestObject* ReferencedObject = Server->CreateObject<UReplicatedTestObject>();
+
+	TestObject->bIsServerObject = true;
+	TestObject->ReplicationSystem = Server->GetReplicationSystem();
+	Server->ReplicationSystem->SetOwningNetConnection(TestObject->NetRefHandle, Client->ConnectionIdOnServer);
+
+	TestObject->ReplicatedTObjectPtr = ReferencedObject;
+	TestObject->ReplicatedWeakObjectPtr = ReferencedObject;
+	TestObject->ClientRPCWithRawPointer(ReferencedObject);
+
+	Server->NetUpdate();
+	Server->SendAndDeliverTo(Client, true);
+	Server->PostSendUpdate();
+
+	UTestObjectWithReferencesAsRemote* ClientTestObject = Client->GetObjectAs<UTestObjectWithReferencesAsRemote>(TestObject->NetRefHandle);
+	UReplicatedTestObject* ClientReferencedObject = Client->GetObjectAs<UReplicatedTestObject>(ReferencedObject->NetRefHandle);
+
+	UE_NET_ASSERT_NE(ClientTestObject, nullptr);
+
+	// Verify that we received valid references
+	UE_NET_ASSERT_TRUE(ClientTestObject->ReplicatedTObjectPtr);
+	UE_NET_ASSERT_TRUE(ClientTestObject->ReplicatedWeakObjectPtr.IsValid());
+	UE_NET_ASSERT_TRUE(ClientTestObject->LastReceivedRawPointer);
+
+	// These are expected to resolve to to server object in this test case, since both the client and server
+	// exist in the same engine instance.
+	UE_NET_ASSERT_EQ(ClientTestObject->ReplicatedTObjectPtr, ReferencedObject);
+	UE_NET_ASSERT_EQ(ClientTestObject->ReplicatedWeakObjectPtr.Get(), ReferencedObject);
+
+	// The raw pointer RPC parameter doesn't go through the RemoteReference path and should point to the replicated
+	// object on the client.
+	UE_NET_ASSERT_EQ(ClientTestObject->LastReceivedRawPointer, ClientReferencedObject);
+}
+
+UE_NET_TEST_FIXTURE(FReplicationSystemConfigOverrideRPCTestFixture, TestObjectsAsRemoteReferencePaths)
+{
+	FReplicationSystemTestNode::FReplicationSystemParamsOverride Params;
+	Params.bUseRemoteObjectReferences = true;
+
+	// Create the server
+	CreateServer(Params);
+
+	// Add a client
+	FReplicationSystemTestClient* Client = CreateClient(Params);
+
+	// Don't want these to trigger in this test
+	FScopedRemoteDelegateOverride ScopedDelegateOverride;
+
+	// Spawn replicated object on server to do the replication
+	UTestObjectWithReferencesAsRemote* TestObject = Server->CreateObject<UTestObjectWithReferencesAsRemote>();
+
+	// Spawn a non-replicated object with a deterministic path, it will be the one referenced by the FRemoteObjectReferece
+	UTestNamedObject* ServerReferencedObject = NewObject<UTestNamedObject>(GetTransientPackage(), UTestNamedObject::StaticClass(), "TestNamedObject");
+
+	TestObject->bIsServerObject = true;
+	TestObject->ReplicationSystem = Server->GetReplicationSystem();
+	Server->ReplicationSystem->SetOwningNetConnection(TestObject->NetRefHandle, Client->ConnectionIdOnServer);
+
+	TestObject->ReplicatedTObjectPtr = ServerReferencedObject;
+	TestObject->ReplicatedWeakObjectPtr = ServerReferencedObject;
+	TestObject->ClientRPCWithRawPointer(ServerReferencedObject);
+
+	Server->NetUpdate();
+	Server->SendTo(Client);
+	Server->PostSendUpdate();
+
+	// Simulate having different instances of objects on the client & server, this forces the FRemoteObjectReference
+	// to look up the object by path on the client. We accomplish this by destroying, GCing, and recreating the object.
+	ServerReferencedObject->MarkAsGarbage();
+	CollectGarbage(GARBAGE_COLLECTION_KEEPFLAGS);
+	ServerReferencedObject = nullptr;
+
+	UTestNamedObject* ClientReferencedObject = NewObject<UTestNamedObject>(GetTransientPackage(), UTestNamedObject::StaticClass(), "TestNamedObject");
+
+	// Deliver to client now that the referenced object has been recreated
+	Server->DeliverTo(Client, true);
+
+	UTestObjectWithReferencesAsRemote* ClientTestObject = Client->GetObjectAs<UTestObjectWithReferencesAsRemote>(TestObject->NetRefHandle);
+
+	UE_NET_ASSERT_NE(ClientTestObject, nullptr);
+
+	// Client should be able to resolve the remote references since they resolve by path
+
+	// Check the TObjectPtr
+	UE_NET_ASSERT_TRUE(ClientTestObject->ReplicatedTObjectPtr);
+
+	UTestNamedObject* ClientResolvedTObjectProperty = Cast<UTestNamedObject>(ClientTestObject->ReplicatedTObjectPtr);
+	UE_NET_ASSERT_EQ(ClientResolvedTObjectProperty, ClientReferencedObject);
+
+	// Check the TWeakObjectPtr
+	UE_NET_ASSERT_TRUE(ClientTestObject->ReplicatedWeakObjectPtr.IsValid());
+
+	UTestNamedObject* ClientResolvedWeakObjectProperty = Cast<UTestNamedObject>(ClientTestObject->ReplicatedWeakObjectPtr);
+	UE_NET_ASSERT_EQ(ClientResolvedWeakObjectProperty, ClientReferencedObject);
+
+	// The raw pointer will be null because it intentionally doesn't go through the remote reference path
+	UE_NET_ASSERT_EQ(ClientTestObject->LastReceivedRawPointer, nullptr);
+}
+
+}
+
+#endif
\ No newline at end of file
diff --git a/Engine/Plugins/Runtime/ReplicationSystemTestPlugin/Source/Private/Tests/Serialization/TestObjectsAsRemoteReferencesSerialization.h b/Engine/Plugins/Runtime/ReplicationSystemTestPlugin/Source/Private/Tests/Serialization/TestObjectsAsRemoteReferencesSerialization.h
new file mode 100644
--- /dev/null
+++ b/Engine/Plugins/Runtime/ReplicationSystemTestPlugin/Source/Private/Tests/Serialization/TestObjectsAsRemoteReferencesSerialization.h
@@ -0,0 +1,26 @@
+// Copyright Epic Games, Inc. All Rights Reserved.
+
+#pragma once
+#include "Tests/ReplicationSystem/RPC/ReplicatedTestObjectWithRPC.h"
+#include "TestObjectsAsRemoteReferencesSerialization.generated.h"
+
+UCLASS()
+class UTestObjectWithReferencesAsRemote : public UTestReplicatedObjectWithRPC
+{
+	GENERATED_BODY()
+
+public:
+	void GetLifetimeReplicatedProps(TArray<class FLifetimeProperty>& OutLifetimeProps) const;
+
+	UPROPERTY(Replicated)
+	TObjectPtr<UObject> ReplicatedTObjectPtr;
+
+	UPROPERTY(Replicated)
+	TWeakObjectPtr<UObject> ReplicatedWeakObjectPtr;
+
+	// Raw object pointers are disallowed as member properties in some contexts, but allowed as RPC parameters
+	UFUNCTION(Client, Reliable)
+	void ClientRPCWithRawPointer(UObject* RawPointer);
+
+	UObject* LastReceivedRawPointer = nullptr;
+};
diff --git a/Engine/Plugins/Runtime/ReplicationSystemTestPlugin/Source/Private/Tests/Serialization/TestRemoteObjectReferenceNetSerializer.cpp b/Engine/Plugins/Runtime/ReplicationSystemTestPlugin/Source/Private/Tests/Serialization/TestRemoteObjectReferenceNetSerializer.cpp
--- a/Engine/Plugins/Runtime/ReplicationSystemTestPlugin/Source/Private/Tests/Serialization/TestRemoteObjectReferenceNetSerializer.cpp
+++ b/Engine/Plugins/Runtime/ReplicationSystemTestPlugin/Source/Private/Tests/Serialization/TestRemoteObjectReferenceNetSerializer.cpp
@@ -26,30 +26,6 @@ void UTestReplicatedObjectWithRemoteReference::RemoteRPCWithRemoteReferenceParam
 namespace UE::Net::Private
 {
 
-class FScopedRemoteDelegateOverride
-{
-public:
-	FScopedRemoteDelegateOverride()
-		: OldStoreRemoteObjectDataDelegate(UE::RemoteObject::Transfer::StoreRemoteObjectDataDelegate)
-		, OldRestoreRemoteObjectDataDelegate(UE::RemoteObject::Transfer::RestoreRemoteObjectDataDelegate)
-	{
-		// Just use empty delegates, we don't need them for the test
-		// Avoids crashing in FPlatformFile while it's only a stub for the test
-		UE::RemoteObject::Transfer::StoreRemoteObjectDataDelegate.BindLambda([](const UE::RemoteObject::Transfer::FMigrateSendParams&) {});
-		UE::RemoteObject::Transfer::RestoreRemoteObjectDataDelegate.BindLambda([](const FUObjectMigrationContext&) {});
-	}
-
-	~FScopedRemoteDelegateOverride()
-	{
-		UE::RemoteObject::Transfer::StoreRemoteObjectDataDelegate = OldStoreRemoteObjectDataDelegate;
-		UE::RemoteObject::Transfer::RestoreRemoteObjectDataDelegate = OldRestoreRemoteObjectDataDelegate;
-	}
-
-private:
-	TDelegate<void(const UE::RemoteObject::Transfer::FMigrateSendParams&)> OldStoreRemoteObjectDataDelegate;
-	TDelegate<void(const FUObjectMigrationContext&)> OldRestoreRemoteObjectDataDelegate;
-};
-
 UE_NET_TEST_FIXTURE(FRPCTestFixture, TestRemoteObjectReference)
 {
 	// Add a client
diff --git a/Engine/Plugins/Runtime/ReplicationSystemTestPlugin/Source/Private/Tests/Serialization/TestRemoteObjectReferenceNetSerializer.h b/Engine/Plugins/Runtime/ReplicationSystemTestPlugin/Source/Private/Tests/Serialization/TestRemoteObjectReferenceNetSerializer.h
--- a/Engine/Plugins/Runtime/ReplicationSystemTestPlugin/Source/Private/Tests/Serialization/TestRemoteObjectReferenceNetSerializer.h
+++ b/Engine/Plugins/Runtime/ReplicationSystemTestPlugin/Source/Private/Tests/Serialization/TestRemoteObjectReferenceNetSerializer.h
@@ -30,3 +30,32 @@ public:
 
 	FRemoteObjectReference LastReceivedRemoteReference;
 };
+
+namespace UE::Net::Private
+{
+
+// Overrides the StoreRemoteObjectDataDelegate and RestoreRemoteObjectDataDelegate to be no-ops for a scope
+// for test cases that don't need them, avoids crashing in FPlatformFile while it's only a stub for the test program
+class FScopedRemoteDelegateOverride
+{
+public:
+	FScopedRemoteDelegateOverride()
+		: OldStoreRemoteObjectDataDelegate(UE::RemoteObject::Transfer::StoreRemoteObjectDataDelegate)
+		, OldRestoreRemoteObjectDataDelegate(UE::RemoteObject::Transfer::RestoreRemoteObjectDataDelegate)
+	{
+		UE::RemoteObject::Transfer::StoreRemoteObjectDataDelegate.BindLambda([](const UE::RemoteObject::Transfer::FMigrateSendParams&) {});
+		UE::RemoteObject::Transfer::RestoreRemoteObjectDataDelegate.BindLambda([](const FUObjectMigrationContext&) {});
+	}
+
+	~FScopedRemoteDelegateOverride()
+	{
+		UE::RemoteObject::Transfer::StoreRemoteObjectDataDelegate = OldStoreRemoteObjectDataDelegate;
+		UE::RemoteObject::Transfer::RestoreRemoteObjectDataDelegate = OldRestoreRemoteObjectDataDelegate;
+	}
+
+private:
+	TDelegate<void(const UE::RemoteObject::Transfer::FMigrateSendParams&)> OldStoreRemoteObjectDataDelegate;
+	TDelegate<void(const FUObjectMigrationContext&)> OldRestoreRemoteObjectDataDelegate;
+};
+
+}
\ No newline at end of file
diff --git a/Engine/Source/Runtime/CoreUObject/Private/UObject/Class.cpp b/Engine/Source/Runtime/CoreUObject/Private/UObject/Class.cpp
--- a/Engine/Source/Runtime/CoreUObject/Private/UObject/Class.cpp
+++ b/Engine/Source/Runtime/CoreUObject/Private/UObject/Class.cpp
@@ -8001,6 +8001,12 @@ UScriptStruct* TBaseStructure<FRemoteObjectTables>::Get()
 	return ScriptStruct;
 }
 
+UScriptStruct* TBaseStructure<FRemoteObjectPathName>::Get()
+{
+	static UScriptStruct* ScriptStruct = StaticGetBaseStructureInternal(TEXT("RemoteObjectPathName"));
+	return ScriptStruct;
+}
+
 UScriptStruct* TBaseStructure<FPackedRemoteObjectPathName>::Get()
 {
 	static UScriptStruct* ScriptStruct = StaticGetBaseStructureInternal(TEXT("PackedRemoteObjectPathName"));
diff --git a/Engine/Source/Runtime/CoreUObject/Public/UObject/Class.h b/Engine/Source/Runtime/CoreUObject/Public/UObject/Class.h
--- a/Engine/Source/Runtime/CoreUObject/Public/UObject/Class.h
+++ b/Engine/Source/Runtime/CoreUObject/Public/UObject/Class.h
@@ -4613,6 +4613,12 @@ template<> struct TBaseStructure<FRemoteObjectTables>
 	static COREUOBJECT_API UScriptStruct* Get();
 };
 
+struct FRemoteObjectPathName;
+template<> struct TBaseStructure<FRemoteObjectPathName>
+{
+	static COREUOBJECT_API UScriptStruct* Get();
+};
+
 struct FPackedRemoteObjectPathName;
 template<> struct TBaseStructure<FPackedRemoteObjectPathName>
 {
diff --git a/Engine/Source/Runtime/CoreUObject/Public/UObject/RemoteObjectTypes.h b/Engine/Source/Runtime/CoreUObject/Public/UObject/RemoteObjectTypes.h
--- a/Engine/Source/Runtime/CoreUObject/Public/UObject/RemoteObjectTypes.h
+++ b/Engine/Source/Runtime/CoreUObject/Public/UObject/RemoteObjectTypes.h
@@ -12,6 +12,11 @@ class UObjectBase;
 class UObject;
 class FArchive;
 
+namespace UE::Net
+{
+	struct FRemoteObjectReferenceNetSerializer;
+}
+
 constexpr int32 REMOTE_OBJECT_SERVER_ID_BIT_SIZE = 10;
 constexpr int32 REMOTE_OBJECT_SERIAL_NUMBER_BIT_SIZE = 53;
 
@@ -38,6 +43,9 @@ struct FRemoteServerId
 	friend struct FRemoteObjectId;
 
 	FRemoteServerId() = default;
+	FRemoteServerId(const FRemoteServerId&) = default;
+	FRemoteServerId& operator=(const FRemoteServerId&) = default;
+
 	explicit FRemoteServerId(ERemoteServerIdConstants InId)
 		: Id((uint32)InId)
 	{
@@ -74,10 +82,6 @@ struct FRemoteServerId
 	{
 		return Id <= Other.Id;
 	}
-	void operator = (FRemoteServerId Other)
-	{
-		Id = Other.Id;
-	}
 	bool IsAsset() const
 	{
 		return (Id == (uint32)ERemoteServerIdConstants::Asset);
@@ -192,6 +196,7 @@ public:
 	COREUOBJECT_API static FRemoteObjectId Generate(UObjectBase* InObject, const TCHAR* InName, EInternalObjectFlags InInitialFlags = EInternalObjectFlags::None);
 
 	friend COREUOBJECT_API FArchive& operator<<(FArchive& Ar, FRemoteObjectId& Id);
+	friend struct UE::Net::FRemoteObjectReferenceNetSerializer;
 };
 
 FORCEINLINE uint32 GetTypeHash(const FRemoteObjectId& ObjectId)
diff --git a/Engine/Source/Runtime/Engine/Private/NetDriver.cpp b/Engine/Source/Runtime/Engine/Private/NetDriver.cpp
--- a/Engine/Source/Runtime/Engine/Private/NetDriver.cpp
+++ b/Engine/Source/Runtime/Engine/Private/NetDriver.cpp
@@ -7333,6 +7333,7 @@ void UNetDriver::CreateReplicationSystem(bool bInitAsClient)
 		Params.ReplicationBridge = ReplicationBridge;
 		Params.bIsServer = !bInitAsClient;
 		Params.bAllowObjectReplication = !bInitAsClient;
+		Params.bUseRemoteObjectReferences = IsUsingRemoteObjectReferences();
 		Params.ForwardNetRPCCallDelegate.BindUObject(this, &UNetDriver::ForwardRemoteFunction);
 		UE::Net::Private::ApplyReplicationSystemConfig(RepSystemConfig, Params);
 		Params.NetTokenStore = NetTokenStore.Get();
diff --git a/Engine/Source/Runtime/Experimental/Iris/Core/Private/Iris/ReplicationState/DefaultPropertyNetSerializerInfos.cpp b/Engine/Source/Runtime/Experimental/Iris/Core/Private/Iris/ReplicationState/DefaultPropertyNetSerializerInfos.cpp
--- a/Engine/Source/Runtime/Experimental/Iris/Core/Private/Iris/ReplicationState/DefaultPropertyNetSerializerInfos.cpp
+++ b/Engine/Source/Runtime/Experimental/Iris/Core/Private/Iris/ReplicationState/DefaultPropertyNetSerializerInfos.cpp
@@ -34,7 +34,48 @@ static const FName PropertyNetSerializerRegistry_NAME_SoftObjectPath("SoftObject
 static const FName PropertyNetSerializerRegistry_NAME_SoftClassPath("SoftClassPath");
 static const FName PropertyNetSerializerRegistry_NAME_RemoteObjectReference("RemoteObjectReference");
 
-UE_NET_IMPLEMENT_SIMPLE_NETSERIALIZER_INFO(FObjectProperty, FObjectNetSerializer);
+// Info struct for raw object pointers, can be used in RPC parameters
+struct FObjectPropertyNetSerializerInfo : public FPropertyNetSerializerInfo
+{
+	virtual const FFieldClass* GetPropertyTypeClass() const override
+	{
+		return FObjectProperty::StaticClass();
+	}
+
+	virtual bool IsSupported(const FProperty* Property) const
+	{
+		// Only use this info for raw UObject pointers.
+		return !Property->HasAnyPropertyFlags(CPF_TObjectPtr);
+	}
+
+	virtual const FNetSerializer* GetNetSerializer(const FProperty* Property) const override
+	{
+		return &UE_NET_GET_SERIALIZER(FObjectNetSerializer);
+	}
+};
+UE_NET_IMPLEMENT_NETSERIALIZER_INFO(FObjectPropertyNetSerializerInfo);
+
+// Info struct for TObjectPtr properties
+struct FObjectPtrPropertyNetSerializerInfo : public FPropertyNetSerializerInfo
+{
+	virtual const FFieldClass* GetPropertyTypeClass() const override
+	{
+		return FObjectProperty::StaticClass();
+}
+
+	virtual bool IsSupported(const FProperty* Property) const
+	{
+		// Only use this info for TObjectPtr properties.
+		return Property->HasAnyPropertyFlags(CPF_TObjectPtr);
+	}
+
+	virtual const FNetSerializer* GetNetSerializer(const FProperty* Property) const override
+	{
+		return &UE_NET_GET_SERIALIZER(FObjectPtrNetSerializer);
+	}
+};
+UE_NET_IMPLEMENT_NETSERIALIZER_INFO(FObjectPtrPropertyNetSerializerInfo);
+
 UE_NET_IMPLEMENT_SIMPLE_NETSERIALIZER_INFO(FWeakObjectProperty, FWeakObjectNetSerializer);
 UE_NET_IMPLEMENT_SIMPLE_NETSERIALIZER_INFO(FSoftObjectProperty, FSoftObjectNetSerializer);
 UE_NET_IMPLEMENT_SIMPLE_NETSERIALIZER_INFO(FSoftClassProperty, FSoftObjectNetSerializer);
@@ -189,7 +230,8 @@ void RegisterDefaultPropertyNetSerializerInfos()
 	UE_NET_REGISTER_NETSERIALIZER_INFO(FDoubleProperty);
 
 	// Object and field types
-	UE_NET_REGISTER_NETSERIALIZER_INFO(FObjectProperty);
+	UE_NET_REGISTER_NETSERIALIZER_INFO(FObjectPropertyNetSerializerInfo);
+	UE_NET_REGISTER_NETSERIALIZER_INFO(FObjectPtrPropertyNetSerializerInfo);
 	UE_NET_REGISTER_NETSERIALIZER_INFO(FWeakObjectProperty);
 	UE_NET_REGISTER_NETSERIALIZER_INFO(FScriptInterfacePropertyNetSerializerInfo);
 	UE_NET_REGISTER_NETSERIALIZER_INFO(FSoftObjectProperty);
diff --git a/Engine/Source/Runtime/Experimental/Iris/Core/Private/Iris/ReplicationSystem/ObjectReferenceCache.cpp b/Engine/Source/Runtime/Experimental/Iris/Core/Private/Iris/ReplicationSystem/ObjectReferenceCache.cpp
--- a/Engine/Source/Runtime/Experimental/Iris/Core/Private/Iris/ReplicationSystem/ObjectReferenceCache.cpp
+++ b/Engine/Source/Runtime/Experimental/Iris/Core/Private/Iris/ReplicationSystem/ObjectReferenceCache.cpp
@@ -16,6 +16,7 @@
 #include "Iris/Serialization/NetExportContext.h"
 #include "Iris/Serialization/NetSerializationContext.h"
 #include "Iris/Serialization/ObjectNetSerializer.h"
+#include "Iris/Serialization/QuantizedObjectReference.h"
 #include "Iris/Stats/NetStatsContext.h"
 #include "Misc/CoreMiscDefines.h"
 #include "Misc/PackageName.h"
@@ -1622,6 +1623,22 @@ void FObjectReferenceCache::AddPendingExports(FNetSerializationContext& Context,
 	}
 }
 
+void FObjectReferenceCache::AddPendingExports(FNetSerializationContext& Context, TArrayView<const FObjectNetSerializerQuantizedReferenceStorage> ExportsView) const
+{
+	if (FNetExportContext* ExportContext = Context.GetExportContext())
+	{
+		for (const FObjectNetSerializerQuantizedReferenceStorage& Reference : ExportsView)
+		{
+			const FQuantizedObjectReference& InternalReference = *reinterpret_cast<const FQuantizedObjectReference*>(&Reference.Storage);
+
+			if (InternalReference.IsNetReference())
+			{
+				AddPendingExport(*ExportContext, InternalReference.NetReference);
+			}
+		}
+	}
+}
+
 FObjectReferenceCache::EWriteExportsResult FObjectReferenceCache::WritePendingExports(FNetSerializationContext& Context, FInternalNetRefIndex ObjectIndex)
 {	
 	FNetBitStreamWriter& Writer = *Context.GetBitStreamWriter();
diff --git a/Engine/Source/Runtime/Experimental/Iris/Core/Private/Iris/ReplicationSystem/ObjectReferenceCache.h b/Engine/Source/Runtime/Experimental/Iris/Core/Private/Iris/ReplicationSystem/ObjectReferenceCache.h
--- a/Engine/Source/Runtime/Experimental/Iris/Core/Private/Iris/ReplicationSystem/ObjectReferenceCache.h
+++ b/Engine/Source/Runtime/Experimental/Iris/Core/Private/Iris/ReplicationSystem/ObjectReferenceCache.h
@@ -16,6 +16,7 @@ namespace UE::Net
 	class FStringTokenStore;
 	class FNetTokenStore;
 	class FNetSerializationContext;
+	struct FObjectNetSerializerQuantizedReferenceStorage;
 	namespace Private
 	{
 		class FNetRefHandleManager;
@@ -108,6 +109,9 @@ public:
 	// Add exports to the set of pending exports for the current batch being written
 	void AddPendingExports(FNetSerializationContext& Context, TArrayView<const FNetObjectReference> ExportsView) const;
 
+	// Add exports to the set of pending exports for the current batch being written
+	void AddPendingExports(FNetSerializationContext& Context, TArrayView<const FObjectNetSerializerQuantizedReferenceStorage> ExportsView) const;
+
 	// Add export to the set of pending exports for the current batch being written
 	void AddPendingExport(FNetExportContext& ExportContext, const FNetObjectReference& Reference) const;
 
diff --git a/Engine/Source/Runtime/Experimental/Iris/Core/Private/Iris/ReplicationSystem/ReplicationOperationsInternal.cpp b/Engine/Source/Runtime/Experimental/Iris/Core/Private/Iris/ReplicationSystem/ReplicationOperationsInternal.cpp
--- a/Engine/Source/Runtime/Experimental/Iris/Core/Private/Iris/ReplicationSystem/ReplicationOperationsInternal.cpp
+++ b/Engine/Source/Runtime/Experimental/Iris/Core/Private/Iris/ReplicationSystem/ReplicationOperationsInternal.cpp
@@ -24,6 +24,7 @@
 #include "Iris/Serialization/InternalNetSerializationContext.h"
 #include "Iris/Serialization/NetSerializers.h"
 #include "Iris/Serialization/InternalNetSerializers.h"
+#include "Iris/Serialization/QuantizedObjectReference.h"
 #include "Math/NumericLimits.h"
 #include "HAL/IConsoleManager.h"
 #include "Iris/Core/IrisProfiler.h"
@@ -265,7 +266,10 @@ void FReplicationStateOperationsInternal::CollectReferences(FNetSerializationCon
 		// Direct reference located in the buffer, we just need the offset to the NetObjectReference
 		if (MemberReferenceDescriptor.Info.ResolveType != FNetReferenceInfo::EResolveType::Invalid)
 		{
-			Collector.Add(MemberReferenceDescriptor.Info, *reinterpret_cast<const FNetObjectReference*>(SrcInternalBuffer + MemberReferenceDescriptor.Offset), OuterChangeMaskInfo);
+			const FQuantizedObjectReference& QuantizedReference = *reinterpret_cast<const FQuantizedObjectReference*>(SrcInternalBuffer + MemberReferenceDescriptor.Offset);
+			
+			// Only collect FNetObjectReferences, if it's a remote reference it doesn't need to be collected
+			Collector.Add(MemberReferenceDescriptor.Info, QuantizedReference, OuterChangeMaskInfo);
 		}
 		else
 		{
@@ -327,7 +331,9 @@ void FReplicationStateOperationsInternal::CollectReferencesWithMask(FNetSerializ
 		// Direct reference located in the buffer, we just need the offset to the NetObjectReference
 		if (MemberReferenceDescriptor.Info.ResolveType != FNetReferenceInfo::EResolveType::Invalid)
 		{
-			Collector.Add(MemberReferenceDescriptor.Info, *reinterpret_cast<const FNetObjectReference*>(SrcInternalBuffer + MemberReferenceDescriptor.Offset), LocalChangeMaskInfo);
+			const FQuantizedObjectReference& QuantizedReference = *reinterpret_cast<const FQuantizedObjectReference*>(SrcInternalBuffer + MemberReferenceDescriptor.Offset);
+
+			Collector.Add(MemberReferenceDescriptor.Info, QuantizedReference, LocalChangeMaskInfo);
 		}
 		else
 		{
diff --git a/Engine/Source/Runtime/Experimental/Iris/Core/Private/Iris/ReplicationSystem/ReplicationSystem.cpp b/Engine/Source/Runtime/Experimental/Iris/Core/Private/Iris/ReplicationSystem/ReplicationSystem.cpp
--- a/Engine/Source/Runtime/Experimental/Iris/Core/Private/Iris/ReplicationSystem/ReplicationSystem.cpp
+++ b/Engine/Source/Runtime/Experimental/Iris/Core/Private/Iris/ReplicationSystem/ReplicationSystem.cpp
@@ -91,6 +91,7 @@ public:
 			.MaxReplicatedObjectCount = Params.MaxReplicatedObjectCount,
 			.NetChunkedArrayCount = Params.PreAllocatedMemoryBuffersObjectCount,
 			.MaxReplicationWriterObjectCount = Params.MaxReplicationWriterObjectCount,
+			.bUseRemoteObjectReferences = Params.bUseRemoteObjectReferences,
 		}))
 	{
 		MetricNameTotalRootObjects = TEXT("TotalSubObjects");
@@ -2060,6 +2061,11 @@ UE::Net::FReplicationSystemDelegates& UReplicationSystem::GetDelegates()
 	return Impl->Delegates;
 }
 
+bool UReplicationSystem::IsUsingRemoteObjectReferences() const
+{
+	return Impl->ReplicationSystemInternal.GetInitParams().bUseRemoteObjectReferences;
+}
+
 #pragma region ReplicationSystemFactory
 
 namespace UE::Net
diff --git a/Engine/Source/Runtime/Experimental/Iris/Core/Private/Iris/ReplicationSystem/ReplicationSystemInternal.h b/Engine/Source/Runtime/Experimental/Iris/Core/Private/Iris/ReplicationSystem/ReplicationSystemInternal.h
--- a/Engine/Source/Runtime/Experimental/Iris/Core/Private/Iris/ReplicationSystem/ReplicationSystemInternal.h
+++ b/Engine/Source/Runtime/Experimental/Iris/Core/Private/Iris/ReplicationSystem/ReplicationSystemInternal.h
@@ -37,6 +37,7 @@ struct FReplicationSystemInternalInitParams
 	uint32 MaxReplicatedObjectCount;
 	uint32 NetChunkedArrayCount;
 	uint32 MaxReplicationWriterObjectCount;
+	bool bUseRemoteObjectReferences;
 };
 
 class FReplicationSystemInternal
diff --git a/Engine/Source/Runtime/Experimental/Iris/Core/Private/Iris/Serialization/FieldPathNetSerializer.cpp b/Engine/Source/Runtime/Experimental/Iris/Core/Private/Iris/Serialization/FieldPathNetSerializer.cpp
--- a/Engine/Source/Runtime/Experimental/Iris/Core/Private/Iris/Serialization/FieldPathNetSerializer.cpp
+++ b/Engine/Source/Runtime/Experimental/Iris/Core/Private/Iris/Serialization/FieldPathNetSerializer.cpp
@@ -39,7 +39,7 @@ struct FFieldPathNetSerializer
 	// Types
 	struct FQuantizedType
 	{
-		alignas(8) uint8 QuantizedStruct[32];
+		alignas(8) uint8 QuantizedStruct[40];
 	};
 
 	typedef FFieldPathForNetSerializer SourceType;
diff --git a/Engine/Source/Runtime/Experimental/Iris/Core/Private/Iris/Serialization/InstancedStructNetSerializer.cpp b/Engine/Source/Runtime/Experimental/Iris/Core/Private/Iris/Serialization/InstancedStructNetSerializer.cpp
--- a/Engine/Source/Runtime/Experimental/Iris/Core/Private/Iris/Serialization/InstancedStructNetSerializer.cpp
+++ b/Engine/Source/Runtime/Experimental/Iris/Core/Private/Iris/Serialization/InstancedStructNetSerializer.cpp
@@ -16,6 +16,7 @@
 #include "Iris/Serialization/NetSerializerArrayStorage.h"
 #include "Iris/Serialization/NetSerializerDelegates.h"
 #include "Iris/Serialization/NetSerializers.h"
+#include "Iris/Serialization/QuantizedObjectReference.h"
 #include "Templates/IsPODType.h"
 
 FInstancedStructNetSerializerConfig::FInstancedStructNetSerializerConfig()
@@ -37,7 +38,7 @@ const FName NetError_InstancedStructNetSerializer_InvalidStructType("Invalid str
 struct FFInstancedStructNetSerializerQuantizedData
 {
 	FNetSerializerAlignedStorage StructData;
-	FNetObjectReference StructType;
+	FQuantizedObjectReference StructType;
 
 	// Not serialized. Fully qualified path. For ReplicationStateDescriptor lookup, validation etc. 
 	FName StructName;
@@ -172,7 +173,7 @@ void FInstancedStructNetSerializer::Deserialize(FNetSerializationContext& Contex
 	// Was the instanced struct valid on the sending side?
 	if (Reader->ReadBool())
 	{
-		FNetObjectReference StructType;
+		FQuantizedObjectReference StructType;
 		const UScriptStruct* Struct = nullptr;
 
 		// Deserialize struct type
@@ -207,7 +208,7 @@ void FInstancedStructNetSerializer::Deserialize(FNetSerializationContext& Contex
 		}
 
 		// Deserialize struct data
-		if (ensureMsgf(Struct != nullptr, TEXT("Unable to find struct using NetObjectReference %s"), ToCStr(StructType.ToString())))
+		if (ensureMsgf(Struct != nullptr, TEXT("Unable to find struct using FQuantizedObjectReference %s"), ToCStr(StructType.ToString())))
 		{
 			FStructNetSerializerConfig StructConfig;
 			StructConfig.StateDescriptor = Config->DescriptorCache.FindOrAddDescriptor(Struct);
@@ -350,7 +351,7 @@ void FInstancedStructNetSerializer::Dequantize(FNetSerializationContext& Context
 		}
 	}
 	
-	if (ensureMsgf(Struct != nullptr, TEXT("Unable to find struct using NetObjectReference %s"), ToCStr(Source.StructType.ToString())))
+	if (ensureMsgf(Struct != nullptr, TEXT("Unable to find struct using FQuantizedObjectReference %s"), ToCStr(Source.StructType.ToString())))
 	{
 		// Re-initialize if the type changes.
 		if (Struct != Target.GetScriptStruct())
diff --git a/Engine/Source/Runtime/Experimental/Iris/Core/Private/Iris/Serialization/InternalNetSerializationContext.cpp b/Engine/Source/Runtime/Experimental/Iris/Core/Private/Iris/Serialization/InternalNetSerializationContext.cpp
--- a/Engine/Source/Runtime/Experimental/Iris/Core/Private/Iris/Serialization/InternalNetSerializationContext.cpp
+++ b/Engine/Source/Runtime/Experimental/Iris/Core/Private/Iris/Serialization/InternalNetSerializationContext.cpp
@@ -14,6 +14,7 @@ FInternalNetSerializationContext::FInternalNetSerializationContext(UReplicationS
 , PackageMap(InReplicationSystem->GetReplicationSystemInternal()->GetIrisObjectReferencePackageMap())
 , bDowngradeAutonomousProxyRole(0)
 , bInlineObjectReferenceExports(0)
+, bSerializeObjectReferencesAsRemoteIds(InReplicationSystem->IsUsingRemoteObjectReferences() ? 1 : 0)
 {
 }
 
diff --git a/Engine/Source/Runtime/Experimental/Iris/Core/Private/Iris/Serialization/InternalNetSerializationContext.h b/Engine/Source/Runtime/Experimental/Iris/Core/Private/Iris/Serialization/InternalNetSerializationContext.h
--- a/Engine/Source/Runtime/Experimental/Iris/Core/Private/Iris/Serialization/InternalNetSerializationContext.h
+++ b/Engine/Source/Runtime/Experimental/Iris/Core/Private/Iris/Serialization/InternalNetSerializationContext.h
@@ -51,6 +51,10 @@ public:
 
 	// Allow References to be inlined in serialized state
 	uint32 bInlineObjectReferenceExports : 1;
+
+	// If true, object references will be quantized & serialized as FRemoteObjectIds instead of FNetObjectReferences.
+	// Requires UE_WITH_REMOTE_OBJECT_HANDLE.
+	uint32 bSerializeObjectReferencesAsRemoteIds : 1;
 };
 
 inline FInternalNetSerializationContext::FInternalNetSerializationContext()
@@ -59,6 +63,7 @@ inline FInternalNetSerializationContext::FInternalNetSerializationContext()
 , PackageMap(nullptr)
 , bDowngradeAutonomousProxyRole(0)
 , bInlineObjectReferenceExports(0)
+, bSerializeObjectReferencesAsRemoteIds(0)
 {}
 
 // Scope used when we actually want to export references.
diff --git a/Engine/Source/Runtime/Experimental/Iris/Core/Private/Iris/Serialization/InternalNetSerializerUtils.cpp b/Engine/Source/Runtime/Experimental/Iris/Core/Private/Iris/Serialization/InternalNetSerializerUtils.cpp
--- a/Engine/Source/Runtime/Experimental/Iris/Core/Private/Iris/Serialization/InternalNetSerializerUtils.cpp
+++ b/Engine/Source/Runtime/Experimental/Iris/Core/Private/Iris/Serialization/InternalNetSerializerUtils.cpp
@@ -9,6 +9,7 @@ bool IsObjectReferenceNetSerializer(const FNetSerializer* NetSerializer)
 {
 	bool bIsObjectReferenceNetSerializer = false;
 	bIsObjectReferenceNetSerializer = bIsObjectReferenceNetSerializer || (NetSerializer == &UE_NET_GET_SERIALIZER(FObjectNetSerializer));
+	bIsObjectReferenceNetSerializer = bIsObjectReferenceNetSerializer || (NetSerializer == &UE_NET_GET_SERIALIZER(FObjectPtrNetSerializer));
 	bIsObjectReferenceNetSerializer = bIsObjectReferenceNetSerializer || (NetSerializer == &UE_NET_GET_SERIALIZER(FWeakObjectNetSerializer));
 	bIsObjectReferenceNetSerializer = bIsObjectReferenceNetSerializer || (NetSerializer == &UE_NET_GET_SERIALIZER(FScriptInterfaceNetSerializer));
 	return bIsObjectReferenceNetSerializer;
diff --git a/Engine/Source/Runtime/Experimental/Iris/Core/Private/Iris/Serialization/IrisPackageMapExportUtil.cpp b/Engine/Source/Runtime/Experimental/Iris/Core/Private/Iris/Serialization/IrisPackageMapExportUtil.cpp
--- a/Engine/Source/Runtime/Experimental/Iris/Core/Private/Iris/Serialization/IrisPackageMapExportUtil.cpp
+++ b/Engine/Source/Runtime/Experimental/Iris/Core/Private/Iris/Serialization/IrisPackageMapExportUtil.cpp
@@ -34,7 +34,7 @@ void FIrisPackageMapExportsUtil::Serialize(FNetSerializationContext& Context, co
 		{
 			UE::Net::WritePackedUint32(Writer, NumReferences);		
 			FObjectNetSerializerConfig ObjectSerializerConfig;
-			for (const FNetObjectReference& Ref : MakeArrayView(Value.ObjectReferenceStorage.GetData(), NumReferences))
+			for (const FObjectNetSerializerQuantizedReferenceStorage& Ref : MakeArrayView(Value.ObjectReferenceStorage.GetData(), NumReferences))
 			{
 				FNetSerializeArgs ObjectArgs;
 				ObjectArgs.NetSerializerConfig = &ObjectSerializerConfig;
@@ -96,7 +96,7 @@ void FIrisPackageMapExportsUtil::Deserialize(FNetSerializationContext& Context,
 			Value.ObjectReferenceStorage.AdjustSize(Context, NumReferences);
 
 			FObjectNetSerializerConfig ObjectSerializerConfig;
-			for (FNetObjectReference& Ref : MakeArrayView(Value.ObjectReferenceStorage.GetData(), Value.ObjectReferenceStorage.Num()))
+			for (FObjectNetSerializerQuantizedReferenceStorage& Ref : MakeArrayView(Value.ObjectReferenceStorage.GetData(), Value.ObjectReferenceStorage.Num()))
 			{
 				FNetDeserializeArgs ObjectArgs;
 				ObjectArgs.NetSerializerConfig = &ObjectSerializerConfig;
@@ -158,7 +158,7 @@ void FIrisPackageMapExportsUtil::Quantize(FNetSerializationContext& Context, con
 		{
 			FObjectNetSerializerConfig ObjectNetSerializerConfig;
 			const TObjectPtr<UObject>* SourceReferences = ObjectReferences.GetData();
-			FNetObjectReference* TargetReferences = Value.ObjectReferenceStorage.GetData();
+			FObjectNetSerializerQuantizedReferenceStorage* TargetReferences = Value.ObjectReferenceStorage.GetData();
 			for (uint32 ReferenceIndex = 0; ReferenceIndex < NumObjectReferences; ++ReferenceIndex)
 			{
 				FNetQuantizeArgs ObjectArgs;
@@ -228,7 +228,7 @@ void FIrisPackageMapExportsUtil::Dequantize(FNetSerializationContext& Context, c
 			ObjectReferences.SetNumUninitialized(NumObjectReferences);
 
 			FObjectNetSerializerConfig ObjectNetSerializerConfig;
-			const FNetObjectReference* SourceReferences = Source.ObjectReferenceStorage.GetData();
+			const FObjectNetSerializerQuantizedReferenceStorage* SourceReferences = Source.ObjectReferenceStorage.GetData();
 			TObjectPtr<UObject>* TargetReferences = ObjectReferences.GetData();
 			for (uint32 ReferenceIndex = 0; ReferenceIndex < NumObjectReferences; ++ReferenceIndex)
 			{
@@ -305,9 +305,11 @@ void FIrisPackageMapExportsUtil::FreeDynamicState(FNetSerializationContext& Cont
 void FIrisPackageMapExportsUtil::CollectNetReferences(FNetSerializationContext& Context, const QuantizedType& Value, const FNetSerializerChangeMaskParam& ChangeMaskInfo, FNetReferenceCollector& Collector)
 {
 	const FNetReferenceInfo ReferenceInfo(FNetReferenceInfo::EResolveType::ResolveOnClient);	
-	for (const FNetObjectReference& Ref : MakeArrayView(Value.ObjectReferenceStorage.GetData(), Value.ObjectReferenceStorage.Num()))
+	for (const FObjectNetSerializerQuantizedReferenceStorage& Ref : MakeArrayView(Value.ObjectReferenceStorage.GetData(), Value.ObjectReferenceStorage.Num()))
 	{
-		Collector.Add(ReferenceInfo, Ref, ChangeMaskInfo);
+		const FQuantizedObjectReference& InternalReference = *reinterpret_cast<const FQuantizedObjectReference*>(&Ref.Storage);
+
+		Collector.Add(ReferenceInfo, InternalReference, ChangeMaskInfo);
 	}
 }
 
diff --git a/Engine/Source/Runtime/Experimental/Iris/Core/Private/Iris/Serialization/NetReferenceCollector.cpp b/Engine/Source/Runtime/Experimental/Iris/Core/Private/Iris/Serialization/NetReferenceCollector.cpp
new file mode 100644
--- /dev/null
+++ b/Engine/Source/Runtime/Experimental/Iris/Core/Private/Iris/Serialization/NetReferenceCollector.cpp
@@ -0,0 +1,38 @@
+// Copyright Epic Games, Inc. All Rights Reserved.
+
+#include "Iris/Serialization/NetReferenceCollector.h"
+
+#include "Iris/Serialization/ObjectNetSerializer.h"
+#include "Iris/Serialization/QuantizedObjectReference.h"
+
+namespace UE::Net
+{
+
+void FNetReferenceCollector::Add(const FNetReferenceInfo& ReferenceInfo, const FQuantizedObjectReference& Reference, const FNetSerializerChangeMaskParam& ChangeMaskInfo)
+{
+	if (Reference.IsRemoteReference())
+	{
+		return;
+	}
+
+	const FNetObjectReference& NetReference = Reference.NetReference;
+
+	if (!NetReference.IsValid() && !EnumHasAnyFlags(Traits, ENetReferenceCollectorTraits::IncludeInvalidReferences))
+	{
+		return;
+	}
+
+	if (!NetReference.CanBeExported() && EnumHasAnyFlags(Traits, ENetReferenceCollectorTraits::OnlyCollectReferencesThatCanBeExported))
+	{
+		return;
+	}
+
+	FReferenceInfo RefInfo;
+	RefInfo.Info = ReferenceInfo;
+	RefInfo.Reference = NetReference;
+	RefInfo.ChangeMaskInfo = ChangeMaskInfo;
+
+	ReferenceInfos.Add(MoveTemp(RefInfo));
+}
+
+}
diff --git a/Engine/Source/Runtime/Experimental/Iris/Core/Private/Iris/Serialization/ObjectNetSerializer.cpp b/Engine/Source/Runtime/Experimental/Iris/Core/Private/Iris/Serialization/ObjectNetSerializer.cpp
--- a/Engine/Source/Runtime/Experimental/Iris/Core/Private/Iris/Serialization/ObjectNetSerializer.cpp
+++ b/Engine/Source/Runtime/Experimental/Iris/Core/Private/Iris/Serialization/ObjectNetSerializer.cpp
@@ -5,16 +5,25 @@
 #include "Iris/Core/IrisLog.h"
 #include "Net/Core/Trace/NetTrace.h"
 #include "Iris/Serialization/InternalNetSerializer.h"
+#include "Iris/Serialization/InternalNetSerializationContext.h"
 #include "Iris/Serialization/NetBitStreamReader.h"
 #include "Iris/Serialization/NetBitStreamWriter.h"
-#include "Iris/Serialization/InternalNetSerializationContext.h"
+#include "Iris/Serialization/QuantizedObjectReference.h"
+#include "Iris/Serialization/QuantizedRemoteObjectReference.h"
+#include "Iris/Serialization/RemoteObjectReferenceNetSerializer.h"
 #include "Iris/ReplicationSystem/NetRefHandleManager.h"
 #include "Iris/ReplicationSystem/ObjectReferenceCache.h"
 #include "Iris/Serialization/NetBitStreamUtil.h"
+#include "UObject/RemoteObjectTransfer.h"
+#include "UObject/RemoteObjectTypes.h"
 
 namespace UE::Net
 {
 
+// Validate size & alignment of public quantized storage type
+static_assert(sizeof(FObjectNetSerializerQuantizedReferenceStorage) >= sizeof(FQuantizedObjectReference), "FObjectNetSerializerQuantizedReferenceStorage::Storage must be large enough to store FQuantizedObjectReference");
+static_assert(alignof(FObjectNetSerializerQuantizedReferenceStorage) >= alignof(FQuantizedObjectReference), "FObjectNetSerializerQuantizedReferenceStorage::Storage must be aligned to store FQuantizedObjectReference");
+
 void WriteNetRefHandle(FNetSerializationContext& Context, const FNetRefHandle Handle)
 {
 	FNetBitStreamWriter* Writer = Context.GetBitStreamWriter();
@@ -90,11 +99,19 @@ static void WriteNetObjectReference(FNetSerializationContext& Context, const FNe
 	}
 }
 
+static void AllocateRemoteReferenceStorage(FQuantizedObjectReference& QuantizedReference, FInternalNetSerializationContext& InternalContext)
+{
+	QuantizedReference.RemoteReferencePtr = static_cast<FQuantizedRemoteObjectReference*>(InternalContext.Alloc(sizeof(FQuantizedRemoteObjectReference), alignof(FQuantizedRemoteObjectReference)));
+	FMemory::Memset(QuantizedReference.RemoteReferencePtr, 0, sizeof(FQuantizedRemoteObjectReference));
+}
+
 template<typename T>
 struct FObjectNetSerializerBase
 {
 	typedef T SourceType;
-	typedef FNetObjectReference QuantizedType;
+	typedef FQuantizedObjectReference QuantizedType;
+
+	static constexpr bool bHasDynamicState = true;
 
 	static void Serialize(FNetSerializationContext&, const FNetSerializeArgs& Args);
 	static void Deserialize(FNetSerializationContext&, const FNetDeserializeArgs& Args);
@@ -104,17 +121,41 @@ struct FObjectNetSerializerBase
 
 	static bool IsEqual(FNetSerializationContext&, const FNetIsEqualArgs& Args);
 
+	static void CloneDynamicState(FNetSerializationContext&, const FNetCloneDynamicStateArgs&);
+	static void FreeDynamicState(FNetSerializationContext&, const FNetFreeDynamicStateArgs&);
+
 protected:
-	static UObject* GetValue(UObject* Value) { return Value; }
-	static void SetValue(UObject*& Dst, UObject* Value) { Dst = Value; }
 
-	static UObject* GetValue(const TWeakObjectPtr<UObject>& Value) { return Value.Get(); }
-	static void SetValue(TWeakObjectPtr<UObject>& Dst, UObject* Value) { Dst = TWeakObjectPtr<UObject>(Value); }
+	// Raw pointer getters/setters
+	static UObject* GetRawPtrValue(UObject* Value) { return Value; }
+	static TObjectPtr<UObject> GetObjectPtrValue(UObject* Value) { return Value; }
+
+	static void SetRawPtrValue(UObject*& Dst, UObject* Value) { Dst = Value; }
+
+	// TObjectPtr getters/setters
+	static UObject* GetRawPtrValue(TObjectPtr<UObject> Value) { return Value; }
+	static TObjectPtr<UObject> GetObjectPtrValue(TObjectPtr<UObject> Value) { return Value; }
+
+	static void SetRawPtrValue(TObjectPtr<UObject>& Dst, UObject* Value) { Dst = Value; }
+	static void SetObjectPtrValue(TObjectPtr<UObject>& Dst, const FRemoteObjectReference& Value) { Dst = TObjectPtr<UObject>(Value.ToObjectPtr()); }
+
+	// Weak pointer getters/setters
+	static UObject* GetRawPtrValue(const TWeakObjectPtr<UObject>& Value) { return Value.Get(); }
+	static TWeakObjectPtr<UObject> GetObjectPtrValue(const TWeakObjectPtr<UObject>& Value) { return Value; }
+
+	static void SetRawPtrValue(TWeakObjectPtr<UObject>& Dst, UObject* Value) { Dst = TWeakObjectPtr<UObject>(Value); }
+#if UE_WITH_REMOTE_OBJECT_HANDLE
+	static void SetObjectPtrValue(TWeakObjectPtr<UObject>& Dst, const FRemoteObjectReference& Value) { Dst = TWeakObjectPtr<UObject>(Value.GetRemoteId()); }
+#endif
 
 	// For FScriptInterface we have a custom Dequantize so we only provide a value getter.
-	static UObject* GetValue(const FScriptInterface& Value) { return Value.GetObject(); }
+	static UObject* GetRawPtrValue(const FScriptInterface& Value) { return Value.GetObject(); }
+	static TObjectPtr<UObject> GetObjectPtrValue(FScriptInterface& Value) { return Value.GetObjectRef(); }
 
 	static UObject* ResolveObjectReference(FNetSerializationContext&, const FNetObjectReference&);
+
+	inline static const FNetSerializer* RemoteObjectReferenceNetSerializer = &UE_NET_GET_SERIALIZER(FRemoteObjectReferenceNetSerializer);
+	inline static const FNetSerializerConfig* RemoteObjectReferenceNetSerializerConfig = UE_NET_GET_SERIALIZER_DEFAULT_CONFIG(FRemoteObjectReferenceNetSerializer);
 };
 
 }
@@ -132,6 +173,16 @@ struct FObjectNetSerializer : public Private::FObjectNetSerializerBase<UObject*>
 
 UE_NET_IMPLEMENT_SERIALIZER_INTERNAL(FObjectNetSerializer);
 
+struct FObjectPtrNetSerializer : public Private::FObjectNetSerializerBase<TObjectPtr<UObject>>
+{
+	static const uint32 Version = 0;
+	typedef FObjectPtrNetSerializerConfig ConfigType;
+
+	inline static const ConfigType DefaultConfig;
+};
+
+UE_NET_IMPLEMENT_SERIALIZER_INTERNAL(FObjectPtrNetSerializer);
+
 struct FWeakObjectNetSerializer : public Private::FObjectNetSerializerBase<TWeakObjectPtr<UObject>>
 {
 	static const uint32 Version = 0;
@@ -154,6 +205,27 @@ struct FScriptInterfaceNetSerializer : public Private::FObjectNetSerializerBase<
 
 UE_NET_IMPLEMENT_SERIALIZER_INTERNAL(FScriptInterfaceNetSerializer);
 
+// Use a simple traits class to distinguish between raw UObject pointers and TObjectPtrs in the serializer implementation.
+// Raw pointers are never serialized as remote references. Using traits reduces boilerplate code that would be needed
+// in two separate, but very similar, net serializer implementations.
+template<class T>
+struct TObjectNetSerializerTraits
+{
+	// If true, serializer is allowed to use the remote reference path for the template parameter type.
+	// UE_WITH_REMOTE_OBJECT_HANDLE is required to use remote references.
+#if UE_WITH_REMOTE_OBJECT_HANDLE
+	static constexpr bool bCanSerializeRemoteReference = true;
+#else
+	static constexpr bool bCanSerializeRemoteReference = false;
+#endif
+};
+
+template<>
+struct TObjectNetSerializerTraits<UObject*>
+{
+	static constexpr bool bCanSerializeRemoteReference = false;
+};
+
 }
 
 namespace UE::Net::Private
@@ -164,7 +236,29 @@ void FObjectNetSerializerBase<T>::Serialize(FNetSerializationContext& Context, c
 {
 	const QuantizedType& Source = *reinterpret_cast<const QuantizedType*>(Args.Source);
 
-	WriteNetObjectReference(Context, Source);
+	FNetBitStreamWriter* Writer = Context.GetBitStreamWriter();
+
+	if (!TObjectNetSerializerTraits<T>::bCanSerializeRemoteReference && Source.IsRemoteReference())
+	{
+		UE_LOG(LogIris, Error, TEXT("FObjectNetSerializerBase::Serialize: traits disallow remote references but trying to serialize a remote reference: %s"), *Source.ToString());
+		Context.SetError(GNetError_InvalidNetHandle);
+		return;
+	}
+
+	if (Source.IsNetReference())
+	{
+		Writer->WriteBool(true);
+		WriteNetObjectReference(Context, Source.NetReference);
+	}
+	else
+	{
+		Writer->WriteBool(false);
+
+		FNetSerializeArgs SerializeArgs = Args;
+		SerializeArgs.NetSerializerConfig = RemoteObjectReferenceNetSerializerConfig;
+		SerializeArgs.Source = NetSerializerValuePointer(Source.RemoteReferencePtr);
+		RemoteObjectReferenceNetSerializer->Serialize(Context, SerializeArgs);
+	}
 }
 
 template<typename T>
@@ -172,22 +266,71 @@ void FObjectNetSerializerBase<T>::Deserialize(FNetSerializationContext& Context,
 {
 	QuantizedType& Target = *reinterpret_cast<QuantizedType*>(Args.Target);
 
-	// Read full ref for now
-	ReadNetObjectReference(Context, Target);
+	FNetBitStreamReader* Reader = Context.GetBitStreamReader();
+
+	const bool bIsNetObjectReference = Reader->ReadBool();
+
+	if (!TObjectNetSerializerTraits<T>::bCanSerializeRemoteReference && !bIsNetObjectReference)
+	{
+		UE_LOG(LogIris, Error, TEXT("FObjectNetSerializerBase::Deserialize: traits disallow remote references but trying to deserialize a remote reference"));
+		Context.SetError(GNetError_InvalidNetHandle);
+		return;
+	}
+
+	if (bIsNetObjectReference)
+	{
+		FNetObjectReference Ref;
+		// Read full ref for now
+		ReadNetObjectReference(Context, Ref);
+		Target.SetNetReference(Context, Args, Ref);
+	}
+	else
+	{
+		// It's a remote object reference
+		if (!Target.RemoteReferencePtr)
+		{
+			AllocateRemoteReferenceStorage(Target, *Context.GetInternalContext());
+		}
+
+		FNetDeserializeArgs InternalArgs = Args;
+		InternalArgs.NetSerializerConfig = RemoteObjectReferenceNetSerializerConfig;
+		InternalArgs.Target = NetSerializerValuePointer(Target.RemoteReferencePtr);
+		RemoteObjectReferenceNetSerializer->Deserialize(Context, InternalArgs);
+	}
 }
 
 template<typename T>
 void FObjectNetSerializerBase<T>::Quantize(FNetSerializationContext& Context, const FNetQuantizeArgs& Args)
 {
 	QuantizedType& Target = *reinterpret_cast<QuantizedType*>(Args.Target);
-	UObject* Source = GetValue(*reinterpret_cast<T*>(Args.Source));
+	T& Source = *reinterpret_cast<T*>(Args.Source);
 
 	// Notice that we quantize to default here if initializing for default state, this is due to the fact that
 	// object references cannot be quantized for default state as they will use locally assigned ids which will differ between server and client
 	// $TODO: Jira: UE-221750 It should now be possible to implement support for this, need to track if stored quantized data is local or not, as long as dequantized
 	// state is correct it should work.
+
 	const FInternalNetSerializationContext* InternalContext = Context.GetInternalContext();
-	Target = Context.IsInitializingDefaultState() ? QuantizedType() : InternalContext->ObjectReferenceCache->GetOrCreateObjectReference(Source);
+	if (TObjectNetSerializerTraits<T>::bCanSerializeRemoteReference && InternalContext->bSerializeObjectReferencesAsRemoteIds)
+	{
+		if (!Target.RemoteReferencePtr)
+		{
+			AllocateRemoteReferenceStorage(Target, *Context.GetInternalContext());
+		}
+
+		FRemoteObjectReference SourceReference = Context.IsInitializingDefaultState() ? FRemoteObjectReference() : FRemoteObjectReference(GetObjectPtrValue(Source));
+
+		FNetQuantizeArgs InternalArgs = Args;
+		InternalArgs.NetSerializerConfig = RemoteObjectReferenceNetSerializerConfig;
+		InternalArgs.Source = NetSerializerValuePointer(&SourceReference);
+		InternalArgs.Target = NetSerializerValuePointer(Target.RemoteReferencePtr);
+		RemoteObjectReferenceNetSerializer->Quantize(Context, InternalArgs);
+	}
+	else
+	{
+		UObject* SourceObject = GetRawPtrValue(Source);
+		Target.SetNetReference(Context, Args, Context.IsInitializingDefaultState() ? FNetObjectReference() : InternalContext->ObjectReferenceCache->GetOrCreateObjectReference(SourceObject));
+	}
 }
 
 template<typename T>
@@ -196,8 +339,32 @@ void FObjectNetSerializerBase<T>::Dequantize(FNetSerializationContext& Context,
 	const QuantizedType& Source = *reinterpret_cast<const QuantizedType*>(Args.Source);
 	T& Target = *reinterpret_cast<T*>(Args.Target);
 	
-	UObject* DequantizedObject = ResolveObjectReference(Context, Source);
-	SetValue(Target, DequantizedObject);
+	if (Source.IsNetReference())
+	{
+		UObject* DequantizedObject = ResolveObjectReference(Context, Source.NetReference);
+		SetRawPtrValue(Target, DequantizedObject);
+		return;
+	}
+
+	// Source is a remote object reference
+	if constexpr (TObjectNetSerializerTraits<T>::bCanSerializeRemoteReference)
+	{
+		FRemoteObjectReference DequantizedRemoteReference;
+
+		FNetDequantizeArgs InternalArgs = Args;
+		InternalArgs.NetSerializerConfig = RemoteObjectReferenceNetSerializerConfig;
+		InternalArgs.Source = NetSerializerValuePointer(Source.RemoteReferencePtr);
+		InternalArgs.Target = NetSerializerValuePointer(&DequantizedRemoteReference);
+		RemoteObjectReferenceNetSerializer->Dequantize(Context, InternalArgs);
+
+		SetObjectPtrValue(Target, DequantizedRemoteReference);
+	}
+	else
+	{
+		UE_LOG(LogIris, Error, TEXT("FObjectNetSerializerBase::Dequantize: source is a remote ref but target type doesn't support remote refs."));
+		Context.SetError(GNetError_InvalidNetHandle);
+		return;
+	}
 }
 
 template<typename T>
@@ -219,6 +386,33 @@ bool FObjectNetSerializerBase<T>::IsEqual(FNetSerializationContext& Context, con
 	}
 }
 
+template<typename T>
+void FObjectNetSerializerBase<T>::CloneDynamicState(FNetSerializationContext& Context, const FNetCloneDynamicStateArgs& Args)
+{
+	QuantizedType& SourceValue = *reinterpret_cast<QuantizedType*>(Args.Source);
+	QuantizedType& TargetValue = *reinterpret_cast<QuantizedType*>(Args.Target);
+
+	if (SourceValue.IsRemoteReference())
+	{
+		// Need to deep copy source to target. Target memory is in undefined state, so allocate new target dynamic state
+		AllocateRemoteReferenceStorage(TargetValue, *Context.GetInternalContext());
+
+		FNetCloneDynamicStateArgs InternalArgs = Args;
+		InternalArgs.NetSerializerConfig = RemoteObjectReferenceNetSerializerConfig;
+		InternalArgs.Source = NetSerializerValuePointer(SourceValue.RemoteReferencePtr);
+		InternalArgs.Target = NetSerializerValuePointer(TargetValue.RemoteReferencePtr);
+		RemoteObjectReferenceNetSerializer->CloneDynamicState(Context, InternalArgs);
+	}
+}
+
+template<typename T>
+void FObjectNetSerializerBase<T>::FreeDynamicState(FNetSerializationContext& Context, const FNetFreeDynamicStateArgs& Args)
+{
+	QuantizedType& SourceValue = *reinterpret_cast<QuantizedType*>(Args.Source);
+
+	SourceValue.FreeRemoteReference(Context, Args);
+}
+
 template<typename T>
 UObject* FObjectNetSerializerBase<T>::ResolveObjectReference(FNetSerializationContext& Context, const FNetObjectReference& NetObjectReference)
 {
@@ -239,7 +433,25 @@ void FScriptInterfaceNetSerializer::Dequantize(FNetSerializationContext& Context
 	const QuantizedType& Source = *reinterpret_cast<const QuantizedType*>(Args.Source);
 	SourceType& Target = *reinterpret_cast<SourceType*>(Args.Target);
 
-	UObject* Object = ResolveObjectReference(Context, Source);
+	UObject* Object = nullptr;
+	
+	if (Source.IsNetReference())
+	{
+		Object = ResolveObjectReference(Context, Source.NetReference);
+	}
+	else
+	{
+		FRemoteObjectReference DequantizedRemoteReference;
+
+		FNetDequantizeArgs InternalArgs = Args;
+		InternalArgs.NetSerializerConfig = RemoteObjectReferenceNetSerializerConfig;
+		InternalArgs.Source = NetSerializerValuePointer(Source.RemoteReferencePtr);
+		InternalArgs.Target = NetSerializerValuePointer(&DequantizedRemoteReference);
+		RemoteObjectReferenceNetSerializer->Dequantize(Context, InternalArgs);
+
+		Object = DequantizedRemoteReference.Resolve();
+	}
+
 	if (Object != nullptr)
 	{
 		const FScriptInterfaceNetSerializerConfig& Config = *static_cast<const FScriptInterfaceNetSerializerConfig*>(Args.NetSerializerConfig);
diff --git a/Engine/Source/Runtime/Experimental/Iris/Core/Private/Iris/Serialization/QuantizedObjectReference.cpp b/Engine/Source/Runtime/Experimental/Iris/Core/Private/Iris/Serialization/QuantizedObjectReference.cpp
new file mode 100644
--- /dev/null
+++ b/Engine/Source/Runtime/Experimental/Iris/Core/Private/Iris/Serialization/QuantizedObjectReference.cpp
@@ -0,0 +1,68 @@
+// Copyright Epic Games, Inc. All Rights Reserved.
+
+#include "Iris/Serialization/QuantizedObjectReference.h"
+
+#include "Iris/Serialization/InternalNetSerializationContext.h"
+#include "Iris/Serialization/NetSerializer.h"
+#include "Iris/Serialization/QuantizedRemoteObjectReference.h"
+#include "Iris/Serialization/RemoteObjectReferenceNetSerializer.h"
+
+namespace UE::Net
+{
+
+bool FQuantizedObjectReference::operator==(const FQuantizedObjectReference& RHS) const
+{
+	if (!IsRemoteReference() && !RHS.IsRemoteReference())
+	{
+		return NetReference == RHS.NetReference;
+	}
+	else if (IsRemoteReference() && RHS.IsRemoteReference())
+	{
+		return *RemoteReferencePtr == *RHS.RemoteReferencePtr;
+	}
+
+	return false;
+}
+
+bool FQuantizedObjectReference::IsValid() const
+{
+	if (IsRemoteReference())
+	{
+		return RemoteReferencePtr->IsValid();
+	}
+	else
+	{
+		return NetReference.IsValid();
+	}
+}
+
+void FQuantizedObjectReference::FreeRemoteReference(FNetSerializationContext& Context, const FNetSerializerBaseArgs& Args)
+{
+	if (RemoteReferencePtr)
+	{
+		FNetFreeDynamicStateArgs InternalArgs;
+		InternalArgs.ChangeMaskInfo = Args.ChangeMaskInfo;
+		InternalArgs.Version = Args.Version;
+		InternalArgs.NetSerializerConfig = UE_NET_GET_SERIALIZER_DEFAULT_CONFIG(FRemoteObjectReferenceNetSerializer);
+		InternalArgs.Source = NetSerializerValuePointer(RemoteReferencePtr);
+		UE_NET_GET_SERIALIZER(FRemoteObjectReferenceNetSerializer).FreeDynamicState(Context, InternalArgs);
+
+		Context.GetInternalContext()->Free(RemoteReferencePtr);
+		RemoteReferencePtr = nullptr;
+	}
+}
+
+FString FQuantizedObjectReference::ToString() const
+{
+	if (IsNetReference())
+	{
+		return NetReference.ToString();
+	}
+	else
+	{
+		const FRemoteObjectId TempId = FRemoteObjectId::CreateFromInt(RemoteReferencePtr->ObjectId);
+		return FString::Printf(TEXT("Remote: %s"), *TempId.ToString());
+	}
+}
+
+}
diff --git a/Engine/Source/Runtime/Experimental/Iris/Core/Private/Iris/Serialization/QuantizedObjectReference.h b/Engine/Source/Runtime/Experimental/Iris/Core/Private/Iris/Serialization/QuantizedObjectReference.h
new file mode 100644
--- /dev/null
+++ b/Engine/Source/Runtime/Experimental/Iris/Core/Private/Iris/Serialization/QuantizedObjectReference.h
@@ -0,0 +1,52 @@
+// Copyright Epic Games, Inc. All Rights Reserved.
+#pragma once
+
+// Note: This is intended for internal use and should not be included or used outside of NetSerializers having to deal with object references.
+ 
+#include "Iris/Core/NetObjectReference.h"
+#include "Templates/IsPODType.h"
+
+namespace UE::Net
+{
+
+class FNetSerializationContext;
+struct FNetSerializerBaseArgs;
+
+struct FQuantizedRemoteObjectReference;
+
+// The quantized state can hold either an FNetObjectReference or a FQuantizedRemoteObjectReference.
+// The FQuantizedRemoteObjectReference is stored as dynamic state in order to reduce required size
+// when it's not used.
+// If RemoteReferencePtr is null, this state is using the FNetObjectReference.
+// If RemoteReferencePtr is non-null, this state is using the remote reference.
+struct FQuantizedObjectReference
+{
+	// Assign an FNetObjectReference to this state. Called from within serializer functions.
+	void SetNetReference(FNetSerializationContext& Context, const FNetSerializerBaseArgs& Args, const FNetObjectReference& InNetReference)
+	{
+		FreeRemoteReference(Context, Args);
+		NetReference = InNetReference;
+	}
+
+	bool IsNetReference() const
+	{
+		return !IsRemoteReference();
+	}
+
+	bool IsRemoteReference() const
+	{
+		return RemoteReferencePtr != nullptr;
+	}
+
+	bool operator==(const FQuantizedObjectReference& RHS) const;
+	bool IsValid() const;
+	void FreeRemoteReference(FNetSerializationContext& Context, const FNetSerializerBaseArgs& Args);
+	FString ToString() const;
+
+	FNetObjectReference NetReference;
+	FQuantizedRemoteObjectReference* RemoteReferencePtr = nullptr;
+};
+
+}
+
+template <> struct TIsPODType<UE::Net::FQuantizedObjectReference> { enum { Value = true }; };
diff --git a/Engine/Source/Runtime/Experimental/Iris/Core/Private/Iris/Serialization/QuantizedRemoteObjectReference.h b/Engine/Source/Runtime/Experimental/Iris/Core/Private/Iris/Serialization/QuantizedRemoteObjectReference.h
new file mode 100644
--- /dev/null
+++ b/Engine/Source/Runtime/Experimental/Iris/Core/Private/Iris/Serialization/QuantizedRemoteObjectReference.h
@@ -0,0 +1,28 @@
+// Copyright Epic Games, Inc. All Rights Reserved.
+#pragma once
+
+#include "HAL/Platform.h"
+
+namespace UE::Net
+{
+
+// Quantized/POD state for FRemoteObjectReference
+struct FQuantizedRemoteObjectReference
+{
+	uint64 ObjectId;
+	uint16 ServerId;
+	alignas(8) uint8 QuantizedPathNameStruct[32];
+
+	bool IsValid() const
+	{
+		return ObjectId != 0 && ServerId != 0;
+	}
+
+	bool operator==(const FQuantizedRemoteObjectReference& Other) const
+	{
+		// Equal only if the object itself is the same
+		return ObjectId == Other.ObjectId;
+	}
+};
+
+}
diff --git a/Engine/Source/Runtime/Experimental/Iris/Core/Private/Iris/Serialization/RemoteObjectReferenceNetSerializer.cpp b/Engine/Source/Runtime/Experimental/Iris/Core/Private/Iris/Serialization/RemoteObjectReferenceNetSerializer.cpp
--- a/Engine/Source/Runtime/Experimental/Iris/Core/Private/Iris/Serialization/RemoteObjectReferenceNetSerializer.cpp
+++ b/Engine/Source/Runtime/Experimental/Iris/Core/Private/Iris/Serialization/RemoteObjectReferenceNetSerializer.cpp
@@ -4,6 +4,7 @@
 #include "Iris/ReplicationState/ReplicationStateDescriptorBuilder.h"
 #include "Iris/Serialization/NetSerializerDelegates.h"
 #include "Iris/Serialization/NetSerializers.h"
+#include "Iris/Serialization/QuantizedRemoteObjectReference.h"
 #include "UObject/RemoteObjectTransfer.h"
 
 namespace UE::Net
@@ -20,17 +21,14 @@ struct FRemoteObjectReferenceNetSerializer
 	static constexpr bool bHasCustomNetReference = true;
 
 	// Types
-	struct FQuantizedType
-	{
-		alignas(8) uint8 QuantizedStruct[48];
-	};
-
 	typedef FRemoteObjectReference SourceType;
-	typedef FQuantizedType QuantizedType;
+	typedef FQuantizedRemoteObjectReference QuantizedType;
 	typedef FRemoteObjectReferenceNetSerializerConfig ConfigType;
 
 	static const ConfigType DefaultConfig;
 
+	static_assert(sizeof(QuantizedType::ServerId) * 8 >= REMOTE_OBJECT_SERVER_ID_BIT_SIZE, "Quantized ServerId is not large enough to store maximum server ID");
+
 	static void Serialize(FNetSerializationContext&, const FNetSerializeArgs& Args);
 	static void Deserialize(FNetSerializationContext&, const FNetDeserializeArgs& Args);
 
@@ -55,92 +53,143 @@ private:
 		virtual void OnPostFreezeNetSerializerRegistry() override;
 	};
 
-	static void RemoteObjectReferenceToRemoteObjectReferenceNetSerializationHelper(const SourceType& RemoteReference, FRemoteObjectReferenceNetSerializationHelper& OutStruct);
-
 	static FRemoteObjectReferenceNetSerializer::FNetSerializerRegistryDelegates NetSerializerRegistryDelegates;
-	static FStructNetSerializerConfig StructNetSerializerConfig;
-	static const FNetSerializer* StructNetSerializer;
+	static FStructNetSerializerConfig RemoteObjectPathNameNetSerializerConfig;
+	static const FNetSerializer* RemoteObjectPathNameNetSerializer;
 };
 UE_NET_IMPLEMENT_SERIALIZER_INTERNAL(FRemoteObjectReferenceNetSerializer);
 
 const FRemoteObjectReferenceNetSerializer::ConfigType FRemoteObjectReferenceNetSerializer::DefaultConfig;
 
 FRemoteObjectReferenceNetSerializer::FNetSerializerRegistryDelegates FRemoteObjectReferenceNetSerializer::NetSerializerRegistryDelegates;
-FStructNetSerializerConfig FRemoteObjectReferenceNetSerializer::StructNetSerializerConfig;
-const FNetSerializer* FRemoteObjectReferenceNetSerializer::StructNetSerializer = &UE_NET_GET_SERIALIZER(FStructNetSerializer);
+FStructNetSerializerConfig FRemoteObjectReferenceNetSerializer::RemoteObjectPathNameNetSerializerConfig;
+const FNetSerializer* FRemoteObjectReferenceNetSerializer::RemoteObjectPathNameNetSerializer = &UE_NET_GET_SERIALIZER(FStructNetSerializer);
 
 void FRemoteObjectReferenceNetSerializer::Serialize(FNetSerializationContext& Context, const FNetSerializeArgs& Args)
 {
+	const QuantizedType& Value = *reinterpret_cast<const QuantizedType*>(Args.Source);
+
+	FNetBitStreamWriter* Writer = Context.GetBitStreamWriter();
+
+	static_assert(sizeof(Value.ObjectId) == 8, "Size of ObjectId expected to be 8 bytes");
+	Writer->WriteBits(static_cast<uint32>(Value.ObjectId), 32U);
+	Writer->WriteBits(static_cast<uint32>(Value.ObjectId >> 32U), 32U);
+
+	Writer->WriteBits(Value.ServerId, REMOTE_OBJECT_SERVER_ID_BIT_SIZE);
+
+	// Forward serialization of the FRemoteObjectPathName to struct serializer
 	FNetSerializeArgs InternalArgs = Args;
-	InternalArgs.NetSerializerConfig = &StructNetSerializerConfig;
-	return StructNetSerializer->Serialize(Context, InternalArgs);
+	InternalArgs.NetSerializerConfig = &RemoteObjectPathNameNetSerializerConfig;
+	InternalArgs.Source = NetSerializerValuePointer(&Value.QuantizedPathNameStruct[0]);
+	RemoteObjectPathNameNetSerializer->Serialize(Context, InternalArgs);
 }
 
 void FRemoteObjectReferenceNetSerializer::Deserialize(FNetSerializationContext& Context, const FNetDeserializeArgs& Args)
 {
+	QuantizedType& Value = *reinterpret_cast<QuantizedType*>(Args.Target);
+
+	FNetBitStreamReader* Reader = Context.GetBitStreamReader();
+
+	Value.ObjectId = Reader->ReadBits(32U);
+	Value.ObjectId |= static_cast<uint64>(Reader->ReadBits(32U)) << 32U;
+
+	Value.ServerId = static_cast<uint16>(Reader->ReadBits(REMOTE_OBJECT_SERVER_ID_BIT_SIZE));
+
 	FNetDeserializeArgs InternalArgs = Args;
-	InternalArgs.NetSerializerConfig = &StructNetSerializerConfig;
-	return StructNetSerializer->Deserialize(Context, InternalArgs);
+	InternalArgs.NetSerializerConfig = &RemoteObjectPathNameNetSerializerConfig;
+	InternalArgs.Target = NetSerializerValuePointer(&Value.QuantizedPathNameStruct[0]);
+	RemoteObjectPathNameNetSerializer->Deserialize(Context, InternalArgs);
 }
 
 void FRemoteObjectReferenceNetSerializer::SerializeDelta(FNetSerializationContext& Context, const FNetSerializeDeltaArgs& Args)
 {
+	const QuantizedType& Value = *reinterpret_cast<const QuantizedType*>(Args.Source);
+
+	FNetBitStreamWriter* Writer = Context.GetBitStreamWriter();
+
+	static_assert(sizeof(Value.ObjectId) == 8, "Size of ObjectId expected to be 8 bytes");
+	Writer->WriteBits(static_cast<uint32>(Value.ObjectId), 32U);
+	Writer->WriteBits(static_cast<uint32>(Value.ObjectId >> 32U), 32U);
+
+	Writer->WriteBits(Value.ServerId, REMOTE_OBJECT_SERVER_ID_BIT_SIZE);
+
+	// Forward serialization of the FRemoteObjectPathName to struct serializer
 	FNetSerializeDeltaArgs InternalArgs = Args;
-	InternalArgs.NetSerializerConfig = &StructNetSerializerConfig;
-	return StructNetSerializer->SerializeDelta(Context, InternalArgs);
+	InternalArgs.NetSerializerConfig = &RemoteObjectPathNameNetSerializerConfig;
+	InternalArgs.Source = NetSerializerValuePointer(&Value.QuantizedPathNameStruct[0]);
+	RemoteObjectPathNameNetSerializer->SerializeDelta(Context, InternalArgs);
 }
 
 void FRemoteObjectReferenceNetSerializer::DeserializeDelta(FNetSerializationContext& Context, const FNetDeserializeDeltaArgs& Args)
 {
+	QuantizedType& Value = *reinterpret_cast<QuantizedType*>(Args.Target);
+
+	FNetBitStreamReader* Reader = Context.GetBitStreamReader();
+
+	Value.ObjectId = Reader->ReadBits(32U);
+	Value.ObjectId |= static_cast<uint64>(Reader->ReadBits(32U)) << 32U;
+
+	Value.ServerId = static_cast<uint16>(Reader->ReadBits(REMOTE_OBJECT_SERVER_ID_BIT_SIZE));
+
 	FNetDeserializeDeltaArgs InternalArgs = Args;
-	InternalArgs.NetSerializerConfig = &StructNetSerializerConfig;
-	return StructNetSerializer->DeserializeDelta(Context, InternalArgs);
+	InternalArgs.NetSerializerConfig = &RemoteObjectPathNameNetSerializerConfig;
+	InternalArgs.Target = NetSerializerValuePointer(&Value.QuantizedPathNameStruct[0]);
+	RemoteObjectPathNameNetSerializer->DeserializeDelta(Context, InternalArgs);
 }
 
 void FRemoteObjectReferenceNetSerializer::Quantize(FNetSerializationContext& Context, const FNetQuantizeArgs& Args)
 {
 	const SourceType& SourceValue = *reinterpret_cast<const SourceType*>(Args.Source);
+	QuantizedType& TargetValue = *reinterpret_cast<QuantizedType*>(Args.Target);
 
-	FRemoteObjectReferenceNetSerializationHelper IntermediateValue;
-	RemoteObjectReferenceToRemoteObjectReferenceNetSerializationHelper(SourceValue, IntermediateValue);
+	TargetValue.ObjectId = SourceValue.GetRemoteId().GetIdNumber();
+	TargetValue.ServerId = static_cast<uint16>(SourceValue.GetSharingServerId().GetIdNumber());
 
-	if (IntermediateValue.ObjectId.IsValid())
+	FRemoteObjectPathName PathName;
+	if (SourceValue.GetRemoteId().IsValid())
 	{
-		if (UObject* ExistingObject = StaticFindObjectFastInternal(IntermediateValue.ObjectId))
+		if (UObject* ExistingObject = StaticFindObjectFastInternal(SourceValue.GetRemoteId()))
 		{
 			UE::RemoteObject::Transfer::RegisterSharedObject(ExistingObject);
-			IntermediateValue.Path = FRemoteObjectPathName(ExistingObject);
+			PathName = FRemoteObjectPathName(ExistingObject);
 		}
 	}
 
 	FNetQuantizeArgs InternalArgs = Args;
-	InternalArgs.Source = NetSerializerValuePointer(&IntermediateValue);
-	InternalArgs.NetSerializerConfig = &StructNetSerializerConfig;
-	return StructNetSerializer->Quantize(Context, InternalArgs);
+	InternalArgs.Source = NetSerializerValuePointer(&PathName);
+	InternalArgs.Target = NetSerializerValuePointer(&TargetValue.QuantizedPathNameStruct[0]);
+	InternalArgs.NetSerializerConfig = &RemoteObjectPathNameNetSerializerConfig;
+	return RemoteObjectPathNameNetSerializer->Quantize(Context, InternalArgs);
 }
 
 void FRemoteObjectReferenceNetSerializer::Dequantize(FNetSerializationContext& Context, const FNetDequantizeArgs& Args)
 {
 	SourceType& TargetValue = *reinterpret_cast<SourceType*>(Args.Target);
+	QuantizedType& QuantizedValue = *reinterpret_cast<QuantizedType*>(Args.Source);
+
+	FRemoteObjectId ObjectId;
+	ObjectId.Id = QuantizedValue.ObjectId;
 
-	FRemoteObjectReferenceNetSerializationHelper IntermediateValue;
+	FRemoteObjectPathName PathName;
 
 	FNetDequantizeArgs InternalArgs = Args;
-	InternalArgs.Target = NetSerializerValuePointer(&IntermediateValue);
-	InternalArgs.NetSerializerConfig = &StructNetSerializerConfig;
+	InternalArgs.Source = NetSerializerValuePointer(&QuantizedValue.QuantizedPathNameStruct[0]);
+	InternalArgs.Target = NetSerializerValuePointer(&PathName);
+	InternalArgs.NetSerializerConfig = &RemoteObjectPathNameNetSerializerConfig;
 
-	StructNetSerializer->Dequantize(Context, InternalArgs);
+	RemoteObjectPathNameNetSerializer->Dequantize(Context, InternalArgs);
 
-	TargetValue.NetDequantize(IntermediateValue.ObjectId, IntermediateValue.ServerId, IntermediateValue.Path);
+	TargetValue.NetDequantize(ObjectId, FRemoteServerId(QuantizedValue.ServerId), PathName);
 }
 
 bool FRemoteObjectReferenceNetSerializer::IsEqual(FNetSerializationContext& Context, const FNetIsEqualArgs& Args)
 {
 	if (Args.bStateIsQuantized)
 	{
-		FNetIsEqualArgs InternalArgs = Args;
-		InternalArgs.NetSerializerConfig = &StructNetSerializerConfig;
-		return StructNetSerializer->IsEqual(Context, InternalArgs);
+		const QuantizedType& QuantizedValue0 = *reinterpret_cast<const QuantizedType*>(Args.Source0);
+		const QuantizedType& QuantizedValue1 = *reinterpret_cast<const QuantizedType*>(Args.Source1);
+
+		return QuantizedValue0 == QuantizedValue1;
 	}
 	else
 	{
@@ -156,66 +205,51 @@ bool FRemoteObjectReferenceNetSerializer::Validate(FNetSerializationContext& Con
 	const SourceType& SourceValue = *reinterpret_cast<const SourceType*>(Args.Source);
 
 	// The remote reference's ObjectId and ServerId should either both be valid or both be invalid
-	if (SourceValue.GetRemoteId().IsValid() != SourceValue.GetSharingServerId().IsValid())
-	{
-		return false;
-	}
-
-	FRemoteObjectReferenceNetSerializationHelper IntermediateValue;
-	RemoteObjectReferenceToRemoteObjectReferenceNetSerializationHelper(SourceValue, IntermediateValue);
-
-	FNetValidateArgs InternalArgs = Args;
-	InternalArgs.Source = NetSerializerValuePointer(&IntermediateValue);
-	InternalArgs.NetSerializerConfig = &StructNetSerializerConfig;
-
-	return StructNetSerializer->Validate(Context, InternalArgs);
+	return SourceValue.GetRemoteId().IsValid() == SourceValue.GetSharingServerId().IsValid();
 }
 
 void FRemoteObjectReferenceNetSerializer::CollectNetReferences(FNetSerializationContext& Context, const FNetCollectReferencesArgs& Args)
 {
-	using namespace Private;
-
-	const FReplicationStateDescriptor* Descriptor = StructNetSerializerConfig.StateDescriptor;
-	if (Descriptor && EnumHasAnyFlags(Descriptor->Traits, EReplicationStateTraits::HasObjectReference))
-	{
-		FNetCollectReferencesArgs InternalArgs = Args;
-		InternalArgs.NetSerializerConfig = &StructNetSerializerConfig;
-		StructNetSerializer->CollectNetReferences(Context, InternalArgs);
-	}
+	// Remote references don't need to go through this collection process.
 }
 
 void FRemoteObjectReferenceNetSerializer::CloneDynamicState(FNetSerializationContext& Context, const FNetCloneDynamicStateArgs& Args)
 {
+	QuantizedType& SourceValue = *reinterpret_cast<QuantizedType*>(Args.Source);
+	QuantizedType& TargetValue = *reinterpret_cast<QuantizedType*>(Args.Target);
+
+	TargetValue.ObjectId = SourceValue.ObjectId;
+	TargetValue.ServerId = SourceValue.ServerId;
+
 	FNetCloneDynamicStateArgs InternalArgs = Args;
-	InternalArgs.NetSerializerConfig = &StructNetSerializerConfig;
-	return StructNetSerializer->CloneDynamicState(Context, InternalArgs);
+	InternalArgs.NetSerializerConfig = &RemoteObjectPathNameNetSerializerConfig;
+	InternalArgs.Source = NetSerializerValuePointer(&SourceValue.QuantizedPathNameStruct[0]);
+	InternalArgs.Target = NetSerializerValuePointer(&TargetValue.QuantizedPathNameStruct[0]);
+	return RemoteObjectPathNameNetSerializer->CloneDynamicState(Context, InternalArgs);
 }
 
 void FRemoteObjectReferenceNetSerializer::FreeDynamicState(FNetSerializationContext& Context, const FNetFreeDynamicStateArgs& Args)
 {
-	FNetFreeDynamicStateArgs InternalArgs = Args;
-	InternalArgs.NetSerializerConfig = &StructNetSerializerConfig;
-	return StructNetSerializer->FreeDynamicState(Context, InternalArgs);
-}
+	QuantizedType& SourceValue = *reinterpret_cast<QuantizedType*>(Args.Source);
 
-void FRemoteObjectReferenceNetSerializer::RemoteObjectReferenceToRemoteObjectReferenceNetSerializationHelper(const SourceType& RemoteObjectReference, FRemoteObjectReferenceNetSerializationHelper& OutStruct)
-{
-	OutStruct.ObjectId = RemoteObjectReference.ObjectId;
-	OutStruct.ServerId = RemoteObjectReference.ServerId;
+	FNetFreeDynamicStateArgs InternalArgs = Args;
+	InternalArgs.NetSerializerConfig = &RemoteObjectPathNameNetSerializerConfig;
+	InternalArgs.Source = NetSerializerValuePointer(&SourceValue.QuantizedPathNameStruct[0]);
+	return RemoteObjectPathNameNetSerializer->FreeDynamicState(Context, InternalArgs);
 }
 
 void FRemoteObjectReferenceNetSerializer::FNetSerializerRegistryDelegates::OnPostFreezeNetSerializerRegistry()
 {
-	UStruct* Struct = FRemoteObjectReferenceNetSerializationHelper::StaticStruct();
-	StructNetSerializerConfig.StateDescriptor = FReplicationStateDescriptorBuilder::CreateDescriptorForStruct(Struct);
-	const FReplicationStateDescriptor* Descriptor = StructNetSerializerConfig.StateDescriptor.GetReference();
+	UStruct* Struct = TBaseStructure<FRemoteObjectPathName>::Get();
+	RemoteObjectPathNameNetSerializerConfig.StateDescriptor = FReplicationStateDescriptorBuilder::CreateDescriptorForStruct(Struct);
+	const FReplicationStateDescriptor* Descriptor = RemoteObjectPathNameNetSerializerConfig.StateDescriptor.GetReference();
 	check(Descriptor != nullptr);
 
 	// Validate our assumptions regarding quantized state size and alignment.
-	static_assert(offsetof(FQuantizedType, QuantizedStruct) == 0U, "Expected buffer for struct to be first member of FQuantizedType.");
-	if (sizeof(FQuantizedType::QuantizedStruct) < Descriptor->InternalSize || alignof(FQuantizedType) < Descriptor->InternalAlignment)
+	static_assert(offsetof(QuantizedType, QuantizedPathNameStruct) == 16U, "Expected buffer for struct to be at offset 16 of QuantizedType.");
+	if (sizeof(QuantizedType::QuantizedPathNameStruct) < Descriptor->InternalSize || alignof(QuantizedType) < Descriptor->InternalAlignment)
 	{
-		LowLevelFatalError(TEXT("FQuantizedType::QuantizedStruct for FRemoteObjectReferenceNetSerializer has size %u and alignment %u but requires size %u and alignment %u."), uint32(sizeof(FQuantizedType::QuantizedStruct)), uint32(alignof(FQuantizedType)), uint32(Descriptor->InternalSize), uint32(Descriptor->InternalAlignment)); 
+		LowLevelFatalError(TEXT("QuantizedType::QuantizedStruct for FRemoteObjectReferenceNetSerializer has size %u and alignment %u but requires size %u and alignment %u."), uint32(sizeof(QuantizedType::QuantizedPathNameStruct)), uint32(alignof(QuantizedType)), uint32(Descriptor->InternalSize), uint32(Descriptor->InternalAlignment));
 	}
 }
 
diff --git a/Engine/Source/Runtime/Experimental/Iris/Core/Private/Iris/Serialization/RemoteObjectReferenceNetSerializer.h b/Engine/Source/Runtime/Experimental/Iris/Core/Private/Iris/Serialization/RemoteObjectReferenceNetSerializer.h
--- a/Engine/Source/Runtime/Experimental/Iris/Core/Private/Iris/Serialization/RemoteObjectReferenceNetSerializer.h
+++ b/Engine/Source/Runtime/Experimental/Iris/Core/Private/Iris/Serialization/RemoteObjectReferenceNetSerializer.h
@@ -3,7 +3,6 @@
 #pragma once
 
 #include "Iris/Serialization/InternalNetSerializer.h"
-#include "UObject/RemoteObjectPathName.h"
 #include "RemoteObjectReferenceNetSerializer.generated.h"
 
 USTRUCT()
@@ -12,24 +11,7 @@ struct FRemoteObjectReferenceNetSerializerConfig : public FNetSerializerConfig
 	GENERATED_BODY()
 };
 
-USTRUCT()
-struct FRemoteObjectReferenceNetSerializationHelper
-{
-	GENERATED_BODY()
-
-	UPROPERTY()
-	FRemoteObjectId ObjectId;
-
-	UPROPERTY()
-	FRemoteServerId ServerId;
-
-	UPROPERTY()
-	FRemoteObjectPathName Path;
-};
-
 namespace UE::Net
 {
-
-UE_NET_DECLARE_SERIALIZER_INTERNAL(FRemoteObjectReferenceNetSerializer, IRISCORE_API);
-
+	UE_NET_DECLARE_SERIALIZER_INTERNAL(FRemoteObjectReferenceNetSerializer, IRISCORE_API);
 }
diff --git a/Engine/Source/Runtime/Experimental/Iris/Core/Private/Iris/Serialization/SoftObjectNetSerializers.cpp b/Engine/Source/Runtime/Experimental/Iris/Core/Private/Iris/Serialization/SoftObjectNetSerializers.cpp
--- a/Engine/Source/Runtime/Experimental/Iris/Core/Private/Iris/Serialization/SoftObjectNetSerializers.cpp
+++ b/Engine/Source/Runtime/Experimental/Iris/Core/Private/Iris/Serialization/SoftObjectNetSerializers.cpp
@@ -5,6 +5,7 @@
 #include "Iris/Serialization/InternalNetSerializer.h"
 #include "Iris/Serialization/NetReferenceCollector.h"
 #include "Iris/Serialization/ObjectNetSerializer.h"
+#include "Iris/Serialization/QuantizedObjectReference.h"
 #include "Iris/Serialization/StringNetSerializerUtils.h"
 #include "Templates/IsPODType.h"
 #include "UObject/SoftObjectPtr.h"
@@ -16,7 +17,7 @@ namespace UE::Net
 struct FFSoftObjectNetSerializerQuantizedData
 {
 	Private::FStringNetSerializerBase::QuantizedType Path;
-	FNetObjectReference ObjectReference;
+	FQuantizedObjectReference ObjectReference;
 	bool bIsObject;
 	uint8 Padding[7];
 };
diff --git a/Engine/Source/Runtime/Experimental/Iris/Core/Public/Iris/ReplicationSystem/ReplicationSystem.h b/Engine/Source/Runtime/Experimental/Iris/Core/Public/Iris/ReplicationSystem/ReplicationSystem.h
--- a/Engine/Source/Runtime/Experimental/Iris/Core/Public/Iris/ReplicationSystem/ReplicationSystem.h
+++ b/Engine/Source/Runtime/Experimental/Iris/Core/Public/Iris/ReplicationSystem/ReplicationSystem.h
@@ -127,6 +127,9 @@ public:
 		 */
 		bool bAllowObjectReplication = false;
 
+		/** When using remote objects with UE_WITH_REMOTE_OBJECT_HANDLE, determines whether or not TObjectPtr or TWeakObjectPtr properties are serialized as FRemoteObjectReferences */
+		bool bUseRemoteObjectReferences = false;
+
 		/** Delegate that receives every RPC executed locally. */
 		UE::Net::FForwardNetRPCCallDelegate ForwardNetRPCCallDelegate;
 
@@ -146,6 +149,9 @@ public:
 	/** @return Is this system configured to replicate object properties. */
 	bool AllowObjectReplication() { return bAllowObjectReplication; }
 
+	/** When using remote objects with UE_WITH_REMOTE_OBJECT_HANDLE, whether or not TObjectPtr or TWeakObjectPtr properties should be serialized as FRemoteObjectReferences */
+	IRISCORE_API bool IsUsingRemoteObjectReferences() const;
+
 	/**
 	 * Update all internal systems, such as filtering, dirty tracking, prioritization, etc.
 	 * Also copies all replicated data into the internal Iris protocols so they are ready to be sent to clients
diff --git a/Engine/Source/Runtime/Experimental/Iris/Core/Public/Iris/Serialization/IrisPackageMapExportUtil.h b/Engine/Source/Runtime/Experimental/Iris/Core/Public/Iris/Serialization/IrisPackageMapExportUtil.h
--- a/Engine/Source/Runtime/Experimental/Iris/Core/Public/Iris/Serialization/IrisPackageMapExportUtil.h
+++ b/Engine/Source/Runtime/Experimental/Iris/Core/Public/Iris/Serialization/IrisPackageMapExportUtil.h
@@ -8,6 +8,7 @@
 #include "Iris/Core/NetObjectReference.h"
 #include "Net/Core/NetToken/NetToken.h"
 #include "Iris/Serialization/NetSerializerArrayStorage.h"
+#include "Iris/Serialization/ObjectNetSerializer.h"
 #include "Iris/Serialization/StringNetSerializers.h"
 
 namespace UE::Net
@@ -18,7 +19,7 @@ class FNetReferenceCollector;
 struct FIrisPackageMapExportsQuantizedType
 {
 	static constexpr uint32 MaxInlinedObjectRefs = 4;
-	typedef FNetSerializerArrayStorage<FNetObjectReference, AllocationPolicies::TInlinedElementAllocationPolicy<MaxInlinedObjectRefs>> FObjectReferenceStorage;
+	typedef FNetSerializerArrayStorage<FObjectNetSerializerQuantizedReferenceStorage, AllocationPolicies::TInlinedElementAllocationPolicy<MaxInlinedObjectRefs>> FObjectReferenceStorage;
 
 	struct FQuantizedName
 	{
diff --git a/Engine/Source/Runtime/Experimental/Iris/Core/Public/Iris/Serialization/NetReferenceCollector.h b/Engine/Source/Runtime/Experimental/Iris/Core/Public/Iris/Serialization/NetReferenceCollector.h
--- a/Engine/Source/Runtime/Experimental/Iris/Core/Public/Iris/Serialization/NetReferenceCollector.h
+++ b/Engine/Source/Runtime/Experimental/Iris/Core/Public/Iris/Serialization/NetReferenceCollector.h
@@ -11,6 +11,8 @@
 namespace UE::Net
 {
 
+struct FQuantizedObjectReference;
+
 enum class ENetReferenceCollectorTraits : uint32
 {
 	None								= 0U,
@@ -35,7 +37,7 @@ public:
 	FNetReferenceCollector() : Traits(ENetReferenceCollectorTraits::None) {}
 	explicit FNetReferenceCollector(ENetReferenceCollectorTraits InTraits) : Traits(InTraits) {}
 	
-	inline void Add(const FNetReferenceInfo& ReferenceInfo, const FNetObjectReference& Reference, const FNetSerializerChangeMaskParam& ChangeMaskInfo);
+	void Add(const FNetReferenceInfo& ReferenceInfo, const FQuantizedObjectReference& Reference, const FNetSerializerChangeMaskParam& ChangeMaskInfo);
 
 	const FReferenceInfoArray& GetCollectedReferences() const { return ReferenceInfos; }
 
@@ -46,24 +48,4 @@ private:
 	const ENetReferenceCollectorTraits Traits;
 };
 
-void FNetReferenceCollector::Add(const FNetReferenceInfo& ReferenceInfo, const FNetObjectReference& Reference, const FNetSerializerChangeMaskParam& ChangeMaskInfo)
-{
-	if (!Reference.IsValid() && !EnumHasAnyFlags(Traits, ENetReferenceCollectorTraits::IncludeInvalidReferences))
-	{
-		return;
-	}
-
-	if (!Reference.CanBeExported() && EnumHasAnyFlags(Traits, ENetReferenceCollectorTraits::OnlyCollectReferencesThatCanBeExported))
-	{
-		return;
-	}
-
-	FReferenceInfo RefInfo;
-	RefInfo.Info = ReferenceInfo;
-	RefInfo.Reference = Reference;
-	RefInfo.ChangeMaskInfo = ChangeMaskInfo;
-
-	ReferenceInfos.Add(MoveTemp(RefInfo));
-}
-
 }
diff --git a/Engine/Source/Runtime/Experimental/Iris/Core/Public/Iris/Serialization/ObjectNetSerializer.h b/Engine/Source/Runtime/Experimental/Iris/Core/Public/Iris/Serialization/ObjectNetSerializer.h
--- a/Engine/Source/Runtime/Experimental/Iris/Core/Public/Iris/Serialization/ObjectNetSerializer.h
+++ b/Engine/Source/Runtime/Experimental/Iris/Core/Public/Iris/Serialization/ObjectNetSerializer.h
@@ -19,6 +19,12 @@ struct FObjectNetSerializerConfig : public FNetSerializerConfig
 	GENERATED_BODY()
 };
 
+USTRUCT()
+struct FObjectPtrNetSerializerConfig : public FNetSerializerConfig
+{
+	GENERATED_BODY()
+};
+
 USTRUCT()
 struct FWeakObjectNetSerializerConfig : public FNetSerializerConfig
 {
@@ -44,8 +50,15 @@ IRISCORE_API void ReadFullNetObjectReference(FNetSerializationContext& Context,
 IRISCORE_API void WriteFullNetObjectReference(FNetSerializationContext& Context, const FNetObjectReference& Reference);
 
 UE_NET_DECLARE_SERIALIZER(FObjectNetSerializer, IRISCORE_API);
+UE_NET_DECLARE_SERIALIZER(FObjectPtrNetSerializer, IRISCORE_API);
 UE_NET_DECLARE_SERIALIZER(FWeakObjectNetSerializer, IRISCORE_API);
 UE_NET_DECLARE_SERIALIZER(FScriptInterfaceNetSerializer, IRISCORE_API);
 
-}
+// Public type that can be used when storage for a quantized reference is needed. This should only be needed
+// in specific cases of forwarding net serializers.
+struct FObjectNetSerializerQuantizedReferenceStorage
+{
+	alignas(8) uint8 Storage[24];
+};
 
+}
