Dataflow : adjust the asset store API

diff --git a/Engine/Source/Runtime/Experimental/Dataflow/Core/Private/Dataflow/DataflowContextAssetStore.cpp b/Engine/Source/Runtime/Experimental/Dataflow/Core/Private/Dataflow/DataflowContextAssetStore.cpp
--- a/Engine/Source/Runtime/Experimental/Dataflow/Core/Private/Dataflow/DataflowContextAssetStore.cpp
+++ b/Engine/Source/Runtime/Experimental/Dataflow/Core/Private/Dataflow/DataflowContextAssetStore.cpp
@@ -7,7 +7,7 @@
 
 namespace UE::Dataflow
 {
-	UObject* FContextAssetStore::AddAsset(const FString& AssetPath, const UClass* AssetClass)
+	UObject* FContextAssetStore::AddAsset(const FString& PersistentAssetPath, const UClass* AssetClass)
 	{
 		if (AssetClass)
 		{
@@ -16,19 +16,42 @@ namespace UE::Dataflow
 			Params.SetFlags = RF_Public;
 			if (TObjectPtr<UObject> NewObject = StaticConstructObject_Internal(Params))
 			{
-				Assets.Add(AssetPath, NewObject);
+				AssetData.Emplace(
+					FAssetData
+					{
+						.PersistentPath = PersistentAssetPath,
+						.TransientPath = NewObject->GetPathName(),
+						.TransientAsset = NewObject
+					});
 				return NewObject;
 			}
 		}
 		return nullptr;
 	}
 
-	UObject* FContextAssetStore::CommitAsset(const FString& AssetPath)
+	UObject* FContextAssetStore::CommitAsset(const FString& TransientAssetPath)
 	{
-		const TObjectPtr<UObject> TransientAsset = Assets.FindRef(AssetPath);
-		if (TransientAsset)
+		if (const FAssetData* FoundAssetData = FindByTransientPath(TransientAssetPath))
 		{
-			FString UniqueAssetPath{ AssetPath };
+			return CommitAsset(*FoundAssetData);
+		}
+		return nullptr;
+	}
+
+	UObject* FContextAssetStore::CommitAsset(const UObject* TransientAsset)
+	{
+		if (const FAssetData* FoundAssetData = FindByTransientAsset(TransientAsset))
+		{
+			return CommitAsset(*FoundAssetData);
+		}
+		return nullptr;
+	}
+
+	UObject* FContextAssetStore::CommitAsset(const FAssetData& InAssetData)
+	{
+		if (InAssetData.TransientAsset)
+		{
+			FString UniqueAssetPath{ InAssetData.PersistentPath };
 
 			// make sure we have a unique name 
 			if (FindPackage(nullptr, *UniqueAssetPath))
@@ -39,7 +62,7 @@ namespace UE::Dataflow
 			if (UPackage* const NewAssetPackage = CreatePackage(*UniqueAssetPath))
 			{
 				const FString AssetName = FPaths::GetBaseFilename(UniqueAssetPath);
-				if (UObject* ObjectToCommit = DuplicateObject_Internal(TransientAsset->GetClass(), TransientAsset, NewAssetPackage, FName(AssetName)))
+				if (UObject* ObjectToCommit = DuplicateObject_Internal(InAssetData.TransientAsset->GetClass(), InAssetData.TransientAsset, NewAssetPackage, FName(AssetName)))
 				{
 					ObjectToCommit->MarkPackageDirty();
 					return ObjectToCommit;
@@ -51,12 +74,30 @@ namespace UE::Dataflow
 
 	void FContextAssetStore::ClearAssets()
 	{
-		Assets.Reset();
+		AssetData.Empty();
+	}
+
+	const FContextAssetStore::FAssetData* FContextAssetStore::FindByTransientPath(const FString& TransientAssetPath) const
+	{
+		return AssetData.FindByPredicate([&TransientAssetPath](const FAssetData& Data)
+			{
+				return (Data.TransientPath == TransientAssetPath);
+			});
+	}
+
+	const FContextAssetStore::FAssetData* FContextAssetStore::FindByTransientAsset(const UObject* TransientAsset) const
+	{
+		return AssetData.FindByPredicate([&TransientAsset](const FAssetData& Data)
+			{
+				return (Data.TransientAsset == TransientAsset);
+			});
 	}
 
 	void FContextAssetStore::AddReferencedObjects(FReferenceCollector& Collector) 
 	{ 
-		//for (Assets)
-		Collector.AddReferencedObjects(Assets);
+		for (FAssetData& Data : AssetData)
+		{
+			Collector.AddReferencedObject(Data.TransientAsset);
+		}
 	}
 }
\ No newline at end of file
diff --git a/Engine/Source/Runtime/Experimental/Dataflow/Core/Public/Dataflow/DataflowContextAssetStore.h b/Engine/Source/Runtime/Experimental/Dataflow/Core/Public/Dataflow/DataflowContextAssetStore.h
--- a/Engine/Source/Runtime/Experimental/Dataflow/Core/Public/Dataflow/DataflowContextAssetStore.h
+++ b/Engine/Source/Runtime/Experimental/Dataflow/Core/Public/Dataflow/DataflowContextAssetStore.h
@@ -26,17 +26,17 @@ namespace UE::Dataflow
 
 		/**
 		* Create an new asset to be commited later ( when the terminal nodes are evaluated )
-		* The return UObject is an object own by the transient package
-		* When commited, the asset will be duplicated to a final package matching AssetPath
+		* When commited, the asset will be duplicated to a final package matching PersistentAssetPath
+		* @return an UObject owned by the transient package
 		*/
-		virtual UObject* AddAsset(const FString& AssetPath, const UClass* AssetClass) = 0;
+		virtual UObject* AddAsset(const FString& PersistentAssetPath, const UClass* AssetClass) = 0;
 
 		/**
-		* Commit the asset matching the AssetPath
-		* If such asset was added previously this will duplicate it and save it to a package matching AssetPath
-		* @return the newly created persistent asset corresponding to AssetPath
+		* Commit the asset matching the TransientAssetPath
+		* If such asset was added previously this will duplicate it and save it to a persistent package
+		* @return the newly created persistent UObject asset
 		*/
-		virtual UObject* CommitAsset(const FString& AssetPath) = 0;
+		virtual UObject* CommitAsset(const FString& TransientAssetPath) = 0;
 
 		/**
 		* Clear all assets from the store
@@ -51,18 +51,25 @@ namespace UE::Dataflow
 	public:
 		/**
 		* Create an new asset to be commited later ( when the terminal nodes are evaluated )
-		* The return UObject is an object own by the transient package
-		* When commited, the asset will be duplicated to a final package matching AssetPath
+		* When commited, the asset will be duplicated to a final package matching PersistentAssetPath
+		* @return an UObject owned by the transient package
 		*/
-		DATAFLOWCORE_API UObject* AddAsset(const FString& AssetPath, const UClass* AssetClass);
+		DATAFLOWCORE_API UObject* AddAsset(const FString& PersistentAssetPath, const UClass* AssetClass);
 
 		/**
-		* Commit the asset matching the AssetPath 
-		* If such asset was added previously this will duplicate it and save it to a package matching AssetPath
-		* @return the newly created persistent asset corresponding to AssetPath
+		* Commit the asset matching the TransientAssetPath 
+		* If such asset was added previously this will duplicate it and save it to a persistent package
+		* @return the newly created persistent UObject asset
 		*/
-		DATAFLOWCORE_API UObject* CommitAsset(const FString& AssetPath);
+		DATAFLOWCORE_API UObject* CommitAsset(const FString& TransientAssetPath);
 
+		/**
+		* Commit the asset matching the TransientAsset
+		* If such asset was added previously this will duplicate it and save it to a package matching AssetPath
+		* @return the newly created persistent UObject asset
+		*/
+		DATAFLOWCORE_API UObject* CommitAsset(const UObject* TransientAsset);
+		
 		/**
 		* Clear all assets from the store 
 		*/
@@ -74,7 +81,19 @@ namespace UE::Dataflow
 		//~ End FGCObject interface
 
 	private:
-		TMap<FString, TObjectPtr<UObject>> Assets;
+		struct FAssetData
+		{
+			FString				PersistentPath;
+			FString				TransientPath;
+			TObjectPtr<UObject> TransientAsset = nullptr;
+		};
+
+		const FAssetData* FindByTransientPath(const FString& TransientAssetPath) const;
+		const FAssetData* FindByTransientAsset(const UObject* TransientAsset) const;
+
+		UObject* CommitAsset(const FAssetData& AssetData);
+
+		TArray<FAssetData> AssetData;
 	};
 
 }
