Adding FIndexedCacheStorageManager: this is the new public API for managing indexed cache storage. IIndexedCacheStorage is not exposed as part of the public API anymore.
Configuration of ICS entries is done with

[IndexedCacheStorage]
+Storage=(Name="CacheStorageName",Index=CacheStorageIndex,Mount="CacheStorageMount")

We can then retrieve the cache index with GetStorageIndex(CacheStorageName) and perform Create/Destroy/Mount/Unmount without having to manage indices / mount names manually.

This also ensures we can't start leaking cache storages by mistake.

diff --git a/Engine/Source/Runtime/Experimental/IndexedCacheStorage/Private/IndexedCacheStorage.cpp b/Engine/Source/Runtime/Experimental/IndexedCacheStorage/Private/IndexedCacheStorage.cpp
--- a/Engine/Source/Runtime/Experimental/IndexedCacheStorage/Private/IndexedCacheStorage.cpp
+++ b/Engine/Source/Runtime/Experimental/IndexedCacheStorage/Private/IndexedCacheStorage.cpp
@@ -20,7 +20,7 @@ public:
 		return 0;
 	}
 
-	virtual bool CreateCacheStorage(uint64 InRequestNumberOfBytes, int32& OutCacheStorageIndex) override
+	virtual bool CreateCacheStorage(uint64 InRequestNumberOfBytes, int32 CacheStorageIndex) override
 	{
 		return false;
 	}
diff --git a/Engine/Source/Runtime/Experimental/IndexedCacheStorage/Private/IndexedCacheStorage.h b/Engine/Source/Runtime/Experimental/IndexedCacheStorage/Private/IndexedCacheStorage.h
new file mode 100644
--- /dev/null
+++ b/Engine/Source/Runtime/Experimental/IndexedCacheStorage/Private/IndexedCacheStorage.h
@@ -0,0 +1,109 @@
+// Copyright Epic Games, Inc. All Rights Reserved.
+
+#pragma once
+
+#include "CoreTypes.h"
+#include "Containers/Array.h"
+
+#include "Modules/ModuleInterface.h"
+#include "Modules/ModuleManager.h"
+
+namespace Experimental
+{
+
+class IIndexedCacheStorage;
+
+class IIndexedCacheStorageModule : public IModuleInterface
+{
+public:
+	static IIndexedCacheStorageModule& Get() 
+	{
+		return FModuleManager::Get().LoadModuleChecked<IIndexedCacheStorageModule>("IndexedCacheStorage");
+	}
+
+	IIndexedCacheStorage& GetOrCreateIndexedCacheStorage();
+
+};
+
+
+/**
+  * Interface for managing cache storages: some platforms may store user data in different cache storages that may be created / destroyed
+  * This allows the title to manage its disk space with a finer granularity. 
+  * The data put in these caches need to be re-generated/re-obtained if needed.
+  */
+class IIndexedCacheStorage
+{
+public:
+	static IIndexedCacheStorage& Get();
+
+	/**
+	  * Check if indexed cache storage is available on the current platform.
+	  * 
+	  * @return True is available, false otherwise.
+	  */
+	virtual bool SupportsIndexedCacheStorage() = 0;
+
+	/**
+	  * Retrieve the number of indexed cache entries available on the platform.
+	  * 
+	  * @return number of indexed cache entries available on the platform.
+	  */
+	virtual uint32 GetNumIndexedCacheStorages() = 0;
+
+	/**
+	  * Requests a new indexed cache entry, capable of holding up to RequestNumberOfBytes.
+	  * 
+ 	  * @param InRequestNumberOfBytes number of bytes this entry needs to hold
+ 	  * @param CacheStorageIndex - the index of the entry to create
+	  * @return True on success, otherwise false
+	  */
+	virtual bool CreateCacheStorage(uint64 InRequestNumberOfBytes, int32 CacheStorageIndex) = 0;
+
+	/**
+	  * Destroy an indexed cache entry.
+	  * 
+ 	  * @param InCacheStorageIndex - the index of the entry
+	  * @return True on success, otherwise false
+	  */
+	virtual bool DestroyCacheStorage(int32 InCacheStorageIndex) = 0;
+
+	/**
+	  * Enumerate existing cache entries.
+	  * 
+ 	  * @param OutCacheStorageIndices - indices of existing cache storages
+	  */
+	virtual void EnumerateCacheStorages(TArray<int32>& OutCacheStorageIndices) = 0;
+	
+	/**
+	  * Queries cache storage.
+	  * 
+ 	  * @param InCacheStorageIndex - the index of the cache
+ 	  * @param OutCacheStorageSize - How large the cache storage is
+	  */
+	virtual void GetCacheStorageInfo(int32 InCacheStorageIndex, uint64& OutCacheStorageSize) = 0;
+	
+	/**
+	  * Gets the platform mount path for a cached storage based on its generic Name
+	  * 
+ 	  * @param OutPlatformMountPath - the platform specific mount path
+ 	  * @param InMountName - the generic mount name
+	  */
+	virtual void GetCacheStorageMountPath(FString& OutPlatformMountPath, const FString& InMountName) = 0;
+
+	/**
+	  * Mounts a cache storage so file operations can be done
+	  * 
+ 	  * @param OutPlatformMountPath - the platform specific mount path. Same name as the one returned by GetCacheStorageMountPath
+ 	  * @param InCacheStorageIndex - the index of the cache
+ 	  * @param InMountName - the generic mount name
+	  */
+	virtual bool MountCacheStorage(FString& OutPlatformMountPath, int32 InCacheStorageIndex, const FString& InMountName) = 0;
+
+	/**
+	  * Unmounts a cache storage
+ 	  * @param InMountName - the generic mount name
+	  */
+	virtual void UnmountCacheStorage(const FString& InMountName) = 0;
+};
+
+}
diff --git a/Engine/Source/Runtime/Experimental/IndexedCacheStorage/Private/IndexedCacheStorageManager.cpp b/Engine/Source/Runtime/Experimental/IndexedCacheStorage/Private/IndexedCacheStorageManager.cpp
new file mode 100644
--- /dev/null
+++ b/Engine/Source/Runtime/Experimental/IndexedCacheStorage/Private/IndexedCacheStorageManager.cpp
@@ -0,0 +1,267 @@
+// Copyright Epic Games, Inc. All Rights Reserved.
+
+#include "IndexedCacheStorageManager.h"
+#include "HAL/PlatformMisc.h"
+#include "Misc/ConfigCacheIni.h"
+#include "Misc/ScopeRWLock.h"
+#include "IndexedCacheStorage.h"
+
+DECLARE_LOG_CATEGORY_EXTERN(LogIndexedCacheStorageManager, Display, All);
+DEFINE_LOG_CATEGORY(LogIndexedCacheStorageManager);
+
+namespace Experimental
+{
+
+FIndexedCacheStorageManager& FIndexedCacheStorageManager::Get()
+{
+	static FIndexedCacheStorageManager GIndexedCacheStorageManager;
+	UE_CALL_ONCE([&]
+	{
+		GIndexedCacheStorageManager.Initialize();
+	}
+	);
+	
+	return GIndexedCacheStorageManager;
+}
+
+void FIndexedCacheStorageManager::Initialize()
+{
+	uint32 NumIndexedCacheEntries = IIndexedCacheStorage::Get().GetNumIndexedCacheStorages();
+	IndexedCacheEntriesMetaData.SetNum(NumIndexedCacheEntries);
+
+	// Takes on the pattern
+	// (Name="CacheStorageName",Index=CacheStorageIndex,Mount="CacheStorageMount")
+	TArray<FString> StorageConfigs;
+	GConfig->GetArray(TEXT("IndexedCacheStorage"), TEXT("Storage"), StorageConfigs, GEngineIni);
+
+	for (const FString& Category : StorageConfigs)
+	{
+		FString TrimmedCategory = Category;
+		TrimmedCategory.TrimStartAndEndInline();
+		if (TrimmedCategory.Left(1) == TEXT("("))
+		{
+			TrimmedCategory.RightChopInline(1, EAllowShrinking::No);
+		}
+		if (TrimmedCategory.Right(1) == TEXT(")"))
+		{
+			TrimmedCategory.LeftChopInline(1, EAllowShrinking::No);
+		}
+
+		// Find all custom chunks and parse
+		const TCHAR* PropertyName = TEXT("Name=");
+		const TCHAR* PropertyIndex = TEXT("Index=");
+		const TCHAR* PropertyMount = TEXT("Mount=");
+		FString StorageName;
+		int32 StorageIndex;
+		FString StorageMount;
+
+		if (FParse::Value(*TrimmedCategory, PropertyName, StorageName) &&
+			FParse::Value(*TrimmedCategory, PropertyIndex, StorageIndex) &&
+			FParse::Value(*TrimmedCategory, PropertyMount, StorageMount))
+		{
+			if (StorageName.Len() == 0)
+			{
+				UE_LOG(LogIndexedCacheStorageManager, Error, TEXT("Found empty Name in [IndexedCacheStorage]:Storage"));
+				continue;
+			}
+
+			if (StorageIndex <= 0 || StorageIndex >= (int32)NumIndexedCacheEntries)
+			{
+				UE_LOG(LogIndexedCacheStorageManager, Error, TEXT("Found invalid Index in [IndexedCacheStorage]:Storage: %d, max allowed = %d"), StorageIndex, NumIndexedCacheEntries-1);
+				continue;
+			}
+
+			if (StorageMount.Len() == 0)
+			{
+				UE_LOG(LogIndexedCacheStorageManager, Error, TEXT("Found empty Mount in [IndexedCacheStorage]:Storage"));
+				continue;
+			}
+
+			StorageName.ReplaceInline(TEXT("\""), TEXT(""));
+			StorageMount.ReplaceInline(TEXT("\""), TEXT(""));
+
+			int32& ExistingStorageIndex = StorageNameToStorageIndex.FindOrAdd(StorageName, 0);
+			if (ExistingStorageIndex > 0)
+			{
+				UE_LOG(LogIndexedCacheStorageManager, Error, TEXT("Found an existing entry with name %s in [IndexedCacheStorage]:Storage"), *StorageName);
+				continue;
+			}
+			
+			FCacheStorageMetaData& CacheStorageMetaData = IndexedCacheEntriesMetaData[StorageIndex];
+
+			if (!CacheStorageMetaData.MountName.IsEmpty())
+			{
+				UE_LOG(LogIndexedCacheStorageManager, Error, TEXT("Found an existing entry with at index %d for name %s in [IndexedCacheStorage]:Storage"), StorageIndex, *StorageName);
+				continue;
+			}
+
+			ExistingStorageIndex = StorageIndex;
+			CacheStorageMetaData.MountName = StorageMount;
+		}
+	}
+
+	TArray<int32> CacheStorageIndices;
+	IIndexedCacheStorage::Get().EnumerateCacheStorages(CacheStorageIndices);
+	TSet<int32> CacheStorageIndicesSet;
+	for (int32 CacheStorageIndex : CacheStorageIndices)
+	{
+		IndexedCacheEntriesMetaData[CacheStorageIndex].bCacheExists = true;
+	}
+	bHasRegisteredEntries = (CacheStorageIndices.Num() > 0);
+
+}
+
+bool FIndexedCacheStorageManager::SupportsIndexedCacheStorage()
+{
+	return bHasRegisteredEntries && IIndexedCacheStorage::Get().SupportsIndexedCacheStorage();
+}
+
+int32 FIndexedCacheStorageManager::GetStorageIndex(const FString& StorageName)
+{
+	FRWScopeLock _(IndexedCacheEntriesMetaDataLock, SLT_ReadOnly);
+
+	int32* StorageIndex = StorageNameToStorageIndex.Find(StorageName);
+
+	return StorageIndex ? *StorageIndex : 0;
+}
+
+void FIndexedCacheStorageManager::EnumerateCacheStorages(TArray<FString>& OutCacheStorageNames)
+{
+	FRWScopeLock _(IndexedCacheEntriesMetaDataLock, SLT_ReadOnly);
+
+	OutCacheStorageNames.Reserve(StorageNameToStorageIndex.Num());
+	for (const TPair<FString, int32>& StorageNameToStorageIndexIter : StorageNameToStorageIndex)
+	{
+		OutCacheStorageNames.Add(StorageNameToStorageIndexIter.Key);
+	}
+}
+
+
+bool FIndexedCacheStorageManager::CreateCacheStorage(uint64 RequestNumberOfBytes, int32 CacheIndex)
+{
+	FRWScopeLock _(IndexedCacheEntriesMetaDataLock, SLT_Write);
+	if (CacheIndex <= 0 && !ensure(IndexedCacheEntriesMetaData.IsValidIndex(CacheIndex)))
+	{
+		return false;
+	}
+
+	if (IndexedCacheEntriesMetaData[CacheIndex].MountName.IsEmpty())
+	{
+		return false;
+	}
+
+	ensure(IndexedCacheEntriesMetaData[CacheIndex].MountRefCount == 0);
+	IndexedCacheEntriesMetaData[CacheIndex].bCacheExists = IIndexedCacheStorage::Get().CreateCacheStorage(RequestNumberOfBytes, CacheIndex);
+	return IndexedCacheEntriesMetaData[CacheIndex].bCacheExists;
+}
+
+void FIndexedCacheStorageManager::DestroyCacheStorage(int32 CacheIndex)
+{
+	FRWScopeLock _(IndexedCacheEntriesMetaDataLock, SLT_Write);
+
+	if (CacheIndex <= 0 && !ensure(IndexedCacheEntriesMetaData.IsValidIndex(CacheIndex)))
+	{
+		return;
+	}
+
+	if (IndexedCacheEntriesMetaData[CacheIndex].MountName.IsEmpty() || !IndexedCacheEntriesMetaData[CacheIndex].bCacheExists)
+	{
+		return;
+	}
+
+	ensure(IndexedCacheEntriesMetaData[CacheIndex].MountRefCount == 0);
+	IIndexedCacheStorage::Get().DestroyCacheStorage(CacheIndex);
+	IndexedCacheEntriesMetaData[CacheIndex].bCacheExists = false;
+}
+
+uint64 FIndexedCacheStorageManager::GetCacheStorageCapacity(int32 CacheIndex)
+{
+	FRWScopeLock _(IndexedCacheEntriesMetaDataLock, SLT_ReadOnly);
+
+	if (CacheIndex <= 0 && !ensure(IndexedCacheEntriesMetaData.IsValidIndex(CacheIndex)))
+	{
+		return 0;
+	}
+
+	if (IndexedCacheEntriesMetaData[CacheIndex].MountName.IsEmpty() || !IndexedCacheEntriesMetaData[CacheIndex].bCacheExists)
+	{
+		return 0;
+	}
+
+	uint64 ExistingSpace = 0;
+	IIndexedCacheStorage::Get().GetCacheStorageInfo(CacheIndex, ExistingSpace);
+	return ExistingSpace;
+}
+
+FString FIndexedCacheStorageManager::MountCacheStorage(int32 CacheIndex)
+{
+	FRWScopeLock _(IndexedCacheEntriesMetaDataLock, SLT_Write);
+
+	FString ActualMountName;
+	if (!IndexedCacheEntriesMetaData.IsValidIndex(CacheIndex) || !IndexedCacheEntriesMetaData[CacheIndex].bCacheExists)
+	{
+		return ActualMountName;
+	}
+
+	check(IndexedCacheEntriesMetaData[CacheIndex].MountRefCount >= 0);
+
+	FString MountedPath;
+	if (IndexedCacheEntriesMetaData[CacheIndex].MountRefCount == 0)
+	{
+		bool bMounted = IIndexedCacheStorage::Get().MountCacheStorage(MountedPath, CacheIndex, IndexedCacheEntriesMetaData[CacheIndex].MountName);
+		if (!bMounted)
+		{
+			return ActualMountName;
+		}
+	}
+	else
+	{
+		IIndexedCacheStorage::Get().GetCacheStorageMountPath(MountedPath, IndexedCacheEntriesMetaData[CacheIndex].MountName);
+	}
+
+	check(!MountedPath.IsEmpty());
+
+	++IndexedCacheEntriesMetaData[CacheIndex].MountRefCount;
+
+	return MountedPath;
+}
+
+void FIndexedCacheStorageManager::UnmountCacheStorage(int32 CacheIndex)
+{
+	FRWScopeLock _(IndexedCacheEntriesMetaDataLock, SLT_Write);
+
+	FString ActualMountName;
+	if (!IndexedCacheEntriesMetaData.IsValidIndex(CacheIndex) || IndexedCacheEntriesMetaData[CacheIndex].MountRefCount == 0)
+	{
+		return;
+	}
+
+	--IndexedCacheEntriesMetaData[CacheIndex].MountRefCount;
+	if (IndexedCacheEntriesMetaData[CacheIndex].MountRefCount == 0)
+	{
+		IIndexedCacheStorage::Get().UnmountCacheStorage(*IndexedCacheEntriesMetaData[CacheIndex].MountName);
+	}
+}
+
+
+FString FIndexedCacheStorageManager::GetMountName(int32 CacheIndex)
+{
+	FRWScopeLock _(IndexedCacheEntriesMetaDataLock, SLT_ReadOnly);
+	if (!IndexedCacheEntriesMetaData.IsValidIndex(CacheIndex))
+	{
+		return FString();
+	}
+
+	return IndexedCacheEntriesMetaData[CacheIndex].MountName;
+
+}
+
+FString FIndexedCacheStorageManager::GetMountPath(int32 CacheIndex)
+{
+	FString MountName = GetMountName(CacheIndex);
+	FString BundleMountPath; 
+	IIndexedCacheStorage::Get().GetCacheStorageMountPath(BundleMountPath, MountName);
+	return BundleMountPath;
+}
+
+}
\ No newline at end of file
diff --git a/Engine/Source/Runtime/Experimental/IndexedCacheStorage/Public/IndexedCacheStorage.h b/Engine/Source/Runtime/Experimental/IndexedCacheStorage/Public/IndexedCacheStorage.h
deleted file mode 100644
diff --git a/Engine/Source/Runtime/Experimental/IndexedCacheStorage/Public/IndexedCacheStorageManager.h b/Engine/Source/Runtime/Experimental/IndexedCacheStorage/Public/IndexedCacheStorageManager.h
new file mode 100644
--- /dev/null
+++ b/Engine/Source/Runtime/Experimental/IndexedCacheStorage/Public/IndexedCacheStorageManager.h
@@ -0,0 +1,92 @@
+// Copyright Epic Games, Inc. All Rights Reserved.
+
+#pragma once
+
+#include "Containers/Array.h"
+#include "Containers/Map.h"
+#include "HAL/Platform.h"
+#include "HAL/CriticalSection.h"
+#include "Containers/UnrealString.h"
+
+#define UE_API INDEXEDCACHESTORAGE_API
+
+namespace Experimental
+{
+
+// Manager to handle ref-counting of ICS entries
+class FIndexedCacheStorageManager
+{
+public:
+
+	UE_API static FIndexedCacheStorageManager& Get();
+
+	// Returns whether ICS is supported
+	UE_API bool SupportsIndexedCacheStorage();
+
+	// Returns the cache index registered to that storage name. Registration happens in the config section:
+	// [IndexedCacheStorage]
+	// (Name="StorageName",Index=StorageIndex,Mount="StorageMount")
+	UE_API int32 GetStorageIndex(const FString& CacheStorageName);
+
+	// returns all registered storage names. The assoaciated cache storages might not have been created
+	UE_API void EnumerateCacheStorages(TArray<FString>& OutCacheStorageNames);
+
+	// Create a new storage at the specified index. Will destroy the previous one if it had a different capacity
+	UE_API bool CreateCacheStorage(uint64 RequestNumberOfBytes, int32 CacheIndex);
+	
+	// Destroy the specified cache storage
+	UE_API void DestroyCacheStorage(int32 CacheIndex);
+	
+	// Returns 0 if no cache currently exists
+	UE_API uint64 GetCacheStorageCapacity(int32 CacheIndex);
+
+	// Returns the path of the mounted cache
+	UE_API FString MountCacheStorage(int32 CacheIndex);
+
+	// Requests unmounting for CacheIndex. Actually unmounts if # of UnmountCacheEntry == # of MountCacheStorage for the same index
+	UE_API void UnmountCacheStorage(int32 CacheIndex);
+
+	// Returns mount name in cache index is mounted, empty if not mounted
+	UE_API FString GetMountName(int32 CacheIndex);
+
+	// Helper to retrieve the mount path. Equivalent to GetCacheStorageMountPath(GetMountName)
+	UE_API FString GetMountPath(int32 CacheIndex);
+
+private:
+
+	void Initialize();
+
+	FRWLock IndexedCacheEntriesMetaDataLock;
+
+	struct FCacheStorageMetaData
+	{
+		FString MountName;
+		int32 MountRefCount = 0;
+		bool bCacheExists = false;
+	};
+	TArray<FCacheStorageMetaData> IndexedCacheEntriesMetaData;
+	TMap<FString, int32> StorageNameToStorageIndex;
+	bool bHasRegisteredEntries = false;
+};
+
+struct FIndexedCacheStorageScopeMount
+{
+	FIndexedCacheStorageScopeMount(int32 InMountPoint)
+	: MountPoint(InMountPoint)
+	{
+		MountPath = Experimental::FIndexedCacheStorageManager::Get().MountCacheStorage(MountPoint);
+	}
+
+	~FIndexedCacheStorageScopeMount()
+	{
+		Experimental::FIndexedCacheStorageManager::Get().UnmountCacheStorage(MountPoint);
+	}
+
+	FString MountPath;
+	const int32 MountPoint = 0;
+};
+
+
+}
+
+#undef UE_API
