Set the relative test URL when creating on-demand host groups.

#rnx

diff --git a/Engine/Source/Runtime/Experimental/IoStore/OnDemand/Private/OnDemandIoStore.cpp b/Engine/Source/Runtime/Experimental/IoStore/OnDemand/Private/OnDemandIoStore.cpp
--- a/Engine/Source/Runtime/Experimental/IoStore/OnDemand/Private/OnDemandIoStore.cpp
+++ b/Engine/Source/Runtime/Experimental/IoStore/OnDemand/Private/OnDemandIoStore.cpp
@@ -1116,10 +1116,15 @@ FIoStatus FOnDemandIoStore::TickMountRequest(FMountRequest& MountRequest)
 
 			if (Container->HostGroup.IsResolved() == false && Args.HostGroup.IsEmpty() == false)
 			{
-				TIoStatusOr<FIASHostGroup> NewHostGroup = FHostGroupManager::Get().Register(Container->HostGroupName, Args.HostGroup.Hosts());
+				TIoStatusOr<FIASHostGroup> NewHostGroup = FHostGroupManager::Get().Register(Container->HostGroupName, Container->GetTestUrl());
 				if (NewHostGroup.IsOk())
 				{
 					Container->HostGroup = NewHostGroup.ConsumeValueOrDie();
+					if (FIoStatus Status = Container->HostGroup.Resolve(Args.HostGroup.Hosts()); !Status.IsOk())
+					{
+						return Status;
+					}
+
 					{
 						const FOnDemandHostGroup& HostGroup = Container->HostGroup.GetUnderlyingHostGroup();
 						TConstArrayView<FAnsiString> Hosts	= HostGroup.Hosts();
diff --git a/Engine/Source/Runtime/Experimental/IoStore/OnDemand/Private/OnDemandIoStore.h b/Engine/Source/Runtime/Experimental/IoStore/OnDemand/Private/OnDemandIoStore.h
--- a/Engine/Source/Runtime/Experimental/IoStore/OnDemand/Private/OnDemandIoStore.h
+++ b/Engine/Source/Runtime/Experimental/IoStore/OnDemand/Private/OnDemandIoStore.h
@@ -96,6 +96,7 @@ struct FOnDemandContainer
 	inline FOnDemandChunkEntryReferences&	FindOrAddChunkEntryReferences(const FOnDemandInternalContentHandle& ContentHandle);
 	inline TBitArray<>						GetReferencedChunkEntries() const;
 	inline bool								IsReferenced(int32 ChunkEntryIndex) const;
+	inline FAnsiString						GetTestUrl();
 	bool									HasAnyFlags(EOnDemandContainerFlags Contains) const { return EnumHasAnyFlags(Flags, Contains); }
 	bool									HasAllFlags(EOnDemandContainerFlags Contains) const { return EnumHasAllFlags(Flags, Contains); }
 
@@ -189,6 +190,24 @@ bool FOnDemandContainer::IsReferenced(int32 ChunkEntryIndex) const
 	return false;
 }
 
+FAnsiString FOnDemandContainer::GetTestUrl()
+{
+	if (ChunkEntries.IsEmpty())
+	{
+		return FAnsiString();
+	}
+
+	TAnsiStringBuilder<41> HashString;
+	HashString << ChunkEntries[0].Hash;
+
+	TAnsiStringBuilder<256> Url;
+	Url << "/" << ChunksDirectory
+		<< "/" << HashString.ToView().Left(2)
+		<< "/" << HashString << ANSITEXTVIEW(".iochunk");
+
+	return Url.ToString();
+}
+
 ///////////////////////////////////////////////////////////////////////////////
 struct FOnDemandChunkInfo
 {
