Minor include fixes.

#rnx

diff --git a/Engine/Source/Runtime/Experimental/IoStore/HttpClient/Internal/IO/Http/LaneTrace.h b/Engine/Source/Runtime/Experimental/IoStore/HttpClient/Internal/IO/Http/LaneTrace.h
--- a/Engine/Source/Runtime/Experimental/IoStore/HttpClient/Internal/IO/Http/LaneTrace.h
+++ b/Engine/Source/Runtime/Experimental/IoStore/HttpClient/Internal/IO/Http/LaneTrace.h
@@ -2,10 +2,12 @@
 
 #pragma once
 
+#if !defined(NO_UE_INCLUDES)
 #include <Trace/Trace.h>
 #include <Containers/StringView.h>
 #include <ProfilingDebugging/CpuProfilerTrace.h>
 #include <Templates/UnrealTemplate.h>
+#endif
 
 ////////////////////////////////////////////////////////////////////////////////
 #if !defined(UE_LANETRACE_ENABLED)
diff --git a/Engine/Source/Runtime/Experimental/IoStore/HttpClient/Private/Http.cpp b/Engine/Source/Runtime/Experimental/IoStore/HttpClient/Private/Http.cpp
--- a/Engine/Source/Runtime/Experimental/IoStore/HttpClient/Private/Http.cpp
+++ b/Engine/Source/Runtime/Experimental/IoStore/HttpClient/Private/Http.cpp
@@ -1,13 +1,13 @@
 // Copyright Epic Games, Inc. All Rights Reserved.
 
 #include "IO/Http/Client.h"
+#include "IO/Http/LaneTrace.h"
 
 #if !defined(NO_UE_INCLUDES)
 #include <Containers/Array.h>
 #include <HAL/CriticalSection.h>
 #include <HAL/IConsoleManager.h>
 #include <IO/IoBuffer.h>
-#include <IO/Http/LaneTrace.h>
 #include <Math/UnrealMathUtility.h>
 #include <Memory/MemoryView.h>
 #include <Misc/ByteSwap.h>
diff --git a/Engine/Source/Runtime/Experimental/IoStore/HttpClient/Private/Test.inl b/Engine/Source/Runtime/Experimental/IoStore/HttpClient/Private/Test.inl
--- a/Engine/Source/Runtime/Experimental/IoStore/HttpClient/Private/Test.inl
+++ b/Engine/Source/Runtime/Experimental/IoStore/HttpClient/Private/Test.inl
@@ -2,15 +2,17 @@
 
 #pragma once
 
+#if !defined(NO_UE_INCLUDES)
 #include <HAL/FileManager.h>
 #include <Misc/Paths.h>
+#endif
 
 namespace UE::IoStore::HTTP
 {
 
 // {{{1 test ...................................................................
 
-#if !(UE_BUILD_SHIPPING|UE_BUILD_TEST)
+#if !(UE_BUILD_SHIPPING|UE_BUILD_TEST) && !defined(NO_UE_INCLUDES)
 
 ////////////////////////////////////////////////////////////////////////////////
 static void MiscTest()
