Generate test chunks from 1KiB to 32MiB for speed testing when uploading on-demand I/O store containers.

#rnx

diff --git a/Engine/Source/Runtime/Experimental/IoStore/OnDemand/Private/Tool/Upload.cpp b/Engine/Source/Runtime/Experimental/IoStore/OnDemand/Private/Tool/Upload.cpp
--- a/Engine/Source/Runtime/Experimental/IoStore/OnDemand/Private/Tool/Upload.cpp
+++ b/Engine/Source/Runtime/Experimental/IoStore/OnDemand/Private/Tool/Upload.cpp
@@ -856,6 +856,43 @@ static TIoStatusOr<FUploadResult> UploadContainerFiles(
 		WriteConfigFile(UploadParams, UploadResult, KeyChain);
 	}
 
+	// Generate test chunks from 1KiB - 32MiB
+	{
+		TStringBuilder<256> Key;
+		for (int32 Shft = 0; Shft < 16; ++Shft)
+		{
+			Key.Reset();
+			const int32 ChunkId		= (1 << Shft);
+			const int32 ChunkSize	= ChunkId << 10;
+			Key << ChunksRelativePath << TEXT("/") << ChunkId << TEXT("kib") << TEXT(".iochunk");
+
+			const FS3HeadObjectResponse HeadResponse = Client.HeadObject(FS3HeadObjectRequest{UploadParams.Bucket, Key.ToString()});
+			if (HeadResponse.IsOk())
+			{
+				continue;
+			}
+
+			FIoBuffer Buffer(ChunkSize);
+			ANSICHAR* BufferContent = reinterpret_cast<ANSICHAR*>(Buffer.GetData());
+			FAnsiString SizeText = FAnsiString::Printf("%d", ChunkSize);
+			for (int32 Idx = 0; Idx < Buffer.GetSize(); ++Idx)
+			{
+				BufferContent[Idx] = SizeText[Idx % SizeText.Len()];
+			}
+
+			const FS3PutObjectResponse Response = Client.TryPutObject(
+					FS3PutObjectRequest{UploadParams.Bucket, Key.ToString(), Buffer.GetView()}); 
+			if (Response.IsOk())
+			{
+				UE_LOG(LogIas, Display, TEXT("Uploaded test chunk '%s'"), Key.ToString());
+			}
+			else
+			{
+				UE_LOG(LogIas, Display, TEXT("Failed to uploaded test chunk '%s'"), Key.ToString());
+			}
+		}
+	}
+
 	{
 		const double Duration = FPlatformTime::Seconds() - StartTime;
 		
