Add a cvar 'iax.InvalidUrlChance' which introduces a chance (0-100%) that a iax http request url will be invalid and should trigger a 404 or 403 response, depending on the cdn.

#rnx

- This cvar is compiled out in shipping builds as it's only for testing purposes.

diff --git a/Engine/Source/Runtime/Experimental/IoStore/OnDemand/Private/OnDemandHttpThread.cpp b/Engine/Source/Runtime/Experimental/IoStore/OnDemand/Private/OnDemandHttpThread.cpp
--- a/Engine/Source/Runtime/Experimental/IoStore/OnDemand/Private/OnDemandHttpThread.cpp
+++ b/Engine/Source/Runtime/Experimental/IoStore/OnDemand/Private/OnDemandHttpThread.cpp
@@ -19,6 +19,9 @@
 #define UE_IOSTORE_ONDEMAND_NO_HTTP_THREAD 0
 #endif
 
+// When enabled the cvar 'iax.InvalidUrlChance' will be enabled and allow us to simulate invalid urls
+#define UE_ALLOW_INVALID_URL_DEBUGGING !UE_BUILD_SHIPPING
+
 // TODO - Maybe do LexToString/LexFromString (dupe code in UnrealEngine.cpp ThreadPriorityToString)
 static const TCHAR* LexToString(EThreadPriority Priority)
 {
@@ -120,6 +123,15 @@ static FAutoConsoleVariableRef CVar_GIasHttpPollTimeoutMs(
 	TEXT("Http tick poll timeout in milliseconds")
 );
 
+#if UE_ALLOW_INVALID_URL_DEBUGGING
+float GIaxInvalidUrlChance = 0.0;
+static FAutoConsoleVariableRef CVar_GIaxInvalidUrlChance(
+	TEXT("iax.InvalidUrlChance"),
+	GIaxInvalidUrlChance,
+	TEXT("Chance that a url for a GET request will be invalid (0-100%)")
+);
+#endif // UE_ALLOW_INVALID_URL_DEBUGGING
+
 const int32 GOnDemandBackendThreadPriorities[5] =
 {
 	EThreadPriority::TPri_Lowest,
@@ -347,6 +359,14 @@ void FOnDemandHttpThread::TickRequests()
 					TAnsiStringBuilder<128> ChunkUrl;
 					HttpRequest->GetUrl(ChunkUrl);
 
+#if UE_ALLOW_INVALID_URL_DEBUGGING
+					// Avoid the rand call if there is no chance
+					if (GIaxInvalidUrlChance > 0.0 && (FMath::FRand() * 100.0f) < GIaxInvalidUrlChance)
+					{
+						ChunkUrl << "-DebugInvalidUrl";
+					}
+#endif // UE_ALLOW_INVALID_URL_DEBUGGING
+
 					HttpClient->Get(HttpRequest->HostGroup().GetUnderlyingHostGroup(), ChunkUrl, HttpRequest->ChunkRange,
 						[this, HttpRequest, &NumConcurrentRequests]
 						(FMultiEndpointHttpClientResponse&& HttpResponse)
@@ -482,3 +502,5 @@ void FOnDemandHttpThread::DestroyRequest(FHttpRequest* Request)
 }
 
 } //namespace UE::IoStore
+
+#undef UE_ALLOW_INVALID_URL_DEBUGGING
