Improve logging when making an IAX latency test

#rnx

- Recently we were debugging an issue where the latency test was failing and it was noted that the old logging did show which CDN was being tested but not where in the CDN.
- The logging will now show the full path used if the latency test fails.

diff --git a/Engine/Source/Runtime/Experimental/IoStore/OnDemand/Private/LatencyTesting.cpp b/Engine/Source/Runtime/Experimental/IoStore/OnDemand/Private/LatencyTesting.cpp
--- a/Engine/Source/Runtime/Experimental/IoStore/OnDemand/Private/LatencyTesting.cpp
+++ b/Engine/Source/Runtime/Experimental/IoStore/OnDemand/Private/LatencyTesting.cpp
@@ -48,7 +48,7 @@ void LatencyTest(FAnsiStringView InUrl, FAnsiStringView InPath, uint32 TimeOutMs
 				if (Status.GetId() == FTicketStatus::EId::Error)
 				{
 					FTicketStatus::FError Error = Status.GetError();
-					UE_LOGFMT(LogIas, Warning, "LatencyTest Error: {Url} {ErrorReason} ({ErrorCode})", ConnectionDesc, Error.Reason, Error.Code);
+					UE_LOGFMT(LogIas, Warning, "LatencyTest Error: 'HEAD {Url}{Path}' {ErrorReason} ({ErrorCode})", InUrl, AnsiPath, Error.Reason, Error.Code);
 					return;
 				}
 				else if (Status.GetId() != FTicketStatus::EId::Response)
@@ -65,7 +65,7 @@ void LatencyTest(FAnsiStringView InUrl, FAnsiStringView InPath, uint32 TimeOutMs
 				}
 				else
 				{
-					UE_LOGFMT(LogIas, Warning, "LatencyTest Failed: {Url} HTTP response ({ResponseCode})", ConnectionDesc, Response.GetStatusCode());
+					UE_LOGFMT(LogIas, Warning, "LatencyTest Failed: 'HEAD {Url}{Path}' HTTP response ({ResponseCode})", InUrl, AnsiPath, Response.GetStatusCode());
 				}
 			});
 
