Implement UNetTokenDataStream::BeginWrite so that the data stream manager can skip sending empty packets if no streams have data to write.

diff --git a/Engine/Source/Runtime/Experimental/Iris/Core/Private/Iris/ReplicationSystem/NetTokenDataStream.cpp b/Engine/Source/Runtime/Experimental/Iris/Core/Private/Iris/ReplicationSystem/NetTokenDataStream.cpp
--- a/Engine/Source/Runtime/Experimental/Iris/Core/Private/Iris/ReplicationSystem/NetTokenDataStream.cpp
+++ b/Engine/Source/Runtime/Experimental/Iris/Core/Private/Iris/ReplicationSystem/NetTokenDataStream.cpp
@@ -83,6 +83,16 @@ void UNetTokenDataStream::AddNetTokenForExplicitExport(UE::Net::FNetToken NetTok
 	NetTokensPendingExport.Add(NetToken);
 }
 
+UDataStream::EWriteResult UNetTokenDataStream::BeginWrite(const FBeginWriteParameters& Params)
+{
+	if (!ensure(NetTokenStore) || NetTokensPendingExport.Num() == 0)
+	{
+		return EWriteResult::NoData;
+	}
+
+	return EWriteResult::HasMoreData;
+}
+
 UDataStream::EWriteResult UNetTokenDataStream::WriteData(UE::Net::FNetSerializationContext& Context, FDataStreamRecord const*& OutRecord)
 {
 	using namespace UE::Net;
diff --git a/Engine/Source/Runtime/Experimental/Iris/Core/Private/Iris/ReplicationSystem/NetTokenDataStream.h b/Engine/Source/Runtime/Experimental/Iris/Core/Private/Iris/ReplicationSystem/NetTokenDataStream.h
--- a/Engine/Source/Runtime/Experimental/Iris/Core/Private/Iris/ReplicationSystem/NetTokenDataStream.h
+++ b/Engine/Source/Runtime/Experimental/Iris/Core/Private/Iris/ReplicationSystem/NetTokenDataStream.h
@@ -37,6 +37,7 @@ private:
 
 	// UDataStream interface
 	virtual void Init(const UDataStream::FInitParameters& Params) override;
+	virtual EWriteResult BeginWrite(const FBeginWriteParameters& Params) override;
 	virtual EWriteResult WriteData(UE::Net::FNetSerializationContext& Context, FDataStreamRecord const*& OutRecord) override;
 	virtual void ReadData(UE::Net::FNetSerializationContext& Context) override;
 	virtual void ProcessPacketDeliveryStatus(UE::Net::EPacketDeliveryStatus Status, FDataStreamRecord const* Record) override;
