Add ReferenceSpace local damping option to Chaos Cloth.

diff --git a/Engine/Plugins/ChaosCloth/Source/ChaosCloth/Private/ChaosCloth/ChaosClothingSimulationSolver.cpp b/Engine/Plugins/ChaosCloth/Source/ChaosCloth/Private/ChaosCloth/ChaosClothingSimulationSolver.cpp
--- a/Engine/Plugins/ChaosCloth/Source/ChaosCloth/Private/ChaosCloth/ChaosClothingSimulationSolver.cpp
+++ b/Engine/Plugins/ChaosCloth/Source/ChaosCloth/Private/ChaosCloth/ChaosClothingSimulationSolver.cpp
@@ -204,6 +204,8 @@ FClothingSimulationSolver::FClothingSimulationSolver(FClothingSimulationConfig*
 		Evolution->AddGroupArray(&PreSimulationTransforms);
 		Evolution->AddGroupArray(&FictitiousAngularVelocities);
 		Evolution->AddGroupArray(&ReferenceSpaceLocations);
+		Evolution->AddGroupArray(&AppliedReferenceSpaceVelocities);
+		Evolution->AddGroupArray(&AppliedReferenceSpaceAngularVelocities);
 
 		Evolution->AddParticleArray(&Normals);
 		Evolution->AddParticleArray(&OldAnimationPositions);
@@ -295,6 +297,8 @@ FClothingSimulationSolver::FClothingSimulationSolver(FClothingSimulationConfig*
 		PBDEvolution->AddArray(&PreSimulationTransforms);
 		PBDEvolution->AddArray(&FictitiousAngularVelocities);
 		PBDEvolution->AddArray(&ReferenceSpaceLocations);
+		PBDEvolution->AddArray(&AppliedReferenceSpaceVelocities);
+		PBDEvolution->AddArray(&AppliedReferenceSpaceAngularVelocities);
 
 		PBDEvolution->Particles().AddArray(&Normals);
 		PBDEvolution->Particles().AddArray(&OldAnimationPositions);
@@ -1308,6 +1312,8 @@ void FClothingSimulationSolver::SetReferenceVelocityScale(uint32 GroupId,
 		* AppliedFictitiousAngularScale;
 	FictitiousAngularVelocities[GroupId] = DeltaTime > (Softs::FSolverReal)0.f ? Softs::FSolverVec3(FictitiousAngularDisplacement) / DeltaTime : Softs::FSolverVec3(0.f);
 	ReferenceSpaceLocations[GroupId] = LocalSpaceScaleInv * (ReferenceSpaceTransform.GetLocation() - LocalSpaceLocation);
+	AppliedReferenceSpaceVelocities[GroupId] = DeltaTime > (Softs::FSolverReal)0.f ? FSolverVec3(LocalSpaceScaleInv * InOutAppliedReferenceVelocity) : Softs::FSolverVec3(0.f);
+	AppliedReferenceSpaceAngularVelocities[GroupId] = DeltaTime > (Softs::FSolverReal)0.f ? FSolverVec3(InOutAppliedReferenceAngularVelocity) : Softs::FSolverVec3(0.f);
 }
 
 Softs::FSolverReal FClothingSimulationSolver::SetParticleMassPerArea(Softs::FSolverParticlesRange& Particles, const FTriangleMesh& Mesh)
@@ -1425,9 +1431,9 @@ void FClothingSimulationSolver::SetProperties(int32 ParticleRangeId, const Softs
 {
 	if (Evolution)
 	{
-		Evolution->SetSoftBodyProperties(ParticleRangeId, InPropertyCollection, WeightMaps);
-
 		const uint32 GroupId = Evolution->GetSoftBodyGroupId(ParticleRangeId);
+		Evolution->SetSoftBodyProperties(ParticleRangeId, InPropertyCollection, WeightMaps, ReferenceSpaceLocations[GroupId], AppliedReferenceSpaceVelocities[GroupId], AppliedReferenceSpaceAngularVelocities[GroupId]);
+
 		// Set properties to constraints that come from the solver (e.g., solver-level gravity, wind)
 		GetClothConstraints(ParticleRangeId).UpdateFromSolver(
 			Gravity, bIsClothGravityOverrideEnabled,
diff --git a/Engine/Plugins/ChaosCloth/Source/ChaosCloth/Public/ChaosCloth/ChaosClothingSimulationSolver.h b/Engine/Plugins/ChaosCloth/Source/ChaosCloth/Public/ChaosCloth/ChaosClothingSimulationSolver.h
--- a/Engine/Plugins/ChaosCloth/Source/ChaosCloth/Public/ChaosCloth/ChaosClothingSimulationSolver.h
+++ b/Engine/Plugins/ChaosCloth/Source/ChaosCloth/Public/ChaosCloth/ChaosClothingSimulationSolver.h
@@ -471,7 +471,9 @@ namespace Chaos
 		// Simulation group attributes
 		TArrayCollectionArray<Softs::FSolverRigidTransform3> PreSimulationTransforms;  // Allow a different frame of reference for each cloth groups
 		TArrayCollectionArray<Softs::FSolverVec3> FictitiousAngularVelocities;  // Relative angular velocity of the reference bone. Depends on the fictitious angular scale factor when using PBDEvolution (is applied later for Evolution)
-		TArrayCollectionArray<Softs::FSolverVec3> ReferenceSpaceLocations;  // Center of rotations for fictitious forces in local coordinate to the simulation space location
+		TArrayCollectionArray<Softs::FSolverVec3> ReferenceSpaceLocations;  // Center of rotations for fictitious forces and local velocity damping in local coordinate to the simulation space location
+		TArrayCollectionArray<Softs::FSolverVec3> AppliedReferenceSpaceVelocities; // Used for local velocity damping.
+		TArrayCollectionArray<Softs::FSolverVec3> AppliedReferenceSpaceAngularVelocities; // Used for local velocity damping.
 
 		// Particle attributes
 		TArrayCollectionArray<Softs::FSolverVec3> Normals;
diff --git a/Engine/Plugins/ChaosClothAsset/Source/ChaosClothAssetEngine/Private/ChaosClothAsset/ClothAssetInteractor.cpp b/Engine/Plugins/ChaosClothAsset/Source/ChaosClothAssetEngine/Private/ChaosClothAsset/ClothAssetInteractor.cpp
--- a/Engine/Plugins/ChaosClothAsset/Source/ChaosClothAssetEngine/Private/ChaosClothAsset/ClothAssetInteractor.cpp
+++ b/Engine/Plugins/ChaosClothAsset/Source/ChaosClothAssetEngine/Private/ChaosClothAsset/ClothAssetInteractor.cpp
@@ -106,6 +106,14 @@ namespace UE::Chaos::ClothAsset::Private
 		TEXT("XPBDAreaSpringStiffness"),
 		TEXT("AreaSpringStiffness"),
 	};
+
+	// SimulationDampingConfigNode
+	static const TArray<FString> LocalDampingCoefficientAliases = // For backwards compatibility--single LocalDampingCoefficient has been split into linear and angular parts.
+	{
+		TEXT("LocalDampingCoefficient"),
+		TEXT("LocalDampingLinearCoefficient"),
+		TEXT("LocalDampingAngularCoefficient"),
+	};
 	
 
 	static const TMap<FString, TArray<FString>> Aliases =
@@ -128,6 +136,7 @@ namespace UE::Chaos::ClothAsset::Private
 		{TEXT("StretchWarpScale"), StretchWarpScaleAliases},
 		{TEXT("StretchWeftScale"), StretchWeftScaleAliases},
 		{TEXT("AreaStiffness"), AreaStiffnessAliases},
+		{TEXT("LocalDampingCoefficient"), LocalDampingCoefficientAliases},
 	};
 
 
diff --git a/Engine/Plugins/ChaosClothAssetEditor/Source/ChaosClothAssetDataflowNodes/Private/ChaosClothAsset/SimulationDampingConfigNode.cpp b/Engine/Plugins/ChaosClothAssetEditor/Source/ChaosClothAssetDataflowNodes/Private/ChaosClothAsset/SimulationDampingConfigNode.cpp
--- a/Engine/Plugins/ChaosClothAssetEditor/Source/ChaosClothAssetDataflowNodes/Private/ChaosClothAsset/SimulationDampingConfigNode.cpp
+++ b/Engine/Plugins/ChaosClothAssetEditor/Source/ChaosClothAssetDataflowNodes/Private/ChaosClothAsset/SimulationDampingConfigNode.cpp
@@ -17,7 +17,9 @@ FChaosClothAssetSimulationDampingConfigNode::FChaosClothAssetSimulationDampingCo
 void FChaosClothAssetSimulationDampingConfigNode::AddProperties(FPropertyHelper& PropertyHelper) const
 {
 	PropertyHelper.SetPropertyWeighted(TEXT("DampingCoefficient"), DampingCoefficientWeighted);
-	PropertyHelper.SetProperty(this, &LocalDampingCoefficient);
+	PropertyHelper.SetPropertyEnum(this, &LocalDampingSpace);
+	PropertyHelper.SetProperty(this, &LocalDampingLinearCoefficient);
+	PropertyHelper.SetProperty(this, &LocalDampingAngularCoefficient);
 }
 
 void FChaosClothAssetSimulationDampingConfigNode::Serialize(FArchive& Ar)
@@ -28,4 +30,9 @@ void FChaosClothAssetSimulationDampingConfigNode::Serialize(FArchive& Ar)
 		DampingCoefficientWeighted.Low = DampingCoefficientWeighted.High = DampingCoefficient_DEPRECATED;
 		DampingCoefficient_DEPRECATED = DeprecatedDampingCoefficientValue;
 	}
+	if (LocalDampingCoefficient_DEPRECATED != DeprecatedDampingCoefficientValue)
+	{
+		LocalDampingLinearCoefficient = LocalDampingAngularCoefficient = LocalDampingCoefficient_DEPRECATED;
+		LocalDampingCoefficient_DEPRECATED = DeprecatedDampingCoefficientValue;
+	}
 }
diff --git a/Engine/Plugins/ChaosClothAssetEditor/Source/ChaosClothAssetDataflowNodes/Public/ChaosClothAsset/SimulationDampingConfigNode.h b/Engine/Plugins/ChaosClothAssetEditor/Source/ChaosClothAssetDataflowNodes/Public/ChaosClothAsset/SimulationDampingConfigNode.h
--- a/Engine/Plugins/ChaosClothAssetEditor/Source/ChaosClothAssetDataflowNodes/Public/ChaosClothAsset/SimulationDampingConfigNode.h
+++ b/Engine/Plugins/ChaosClothAssetEditor/Source/ChaosClothAssetDataflowNodes/Public/ChaosClothAsset/SimulationDampingConfigNode.h
@@ -3,9 +3,10 @@
 #pragma once 
 
 #include "ChaosClothAsset/SimulationBaseConfigNode.h"
+#include "Chaos/SoftsSimulationSpace.h"
 #include "SimulationDampingConfigNode.generated.h"
 
-/** Damping properties configuration node. */
+/** Damping properties configuration node.*/
 USTRUCT(Meta = (DataflowCloth))
 struct FChaosClothAssetSimulationDampingConfigNode : public FChaosClothAssetSimulationBaseConfigNode
 {
@@ -21,13 +22,28 @@ public:
 	FChaosClothAssetWeightedValue DampingCoefficientWeighted = {true, 0.01f, 0.01f, TEXT("DampingCoefficient")};
 
 	/**
-	 * The amount of local damping applied to the cloth velocities.
-	 * This type of damping only damps individual deviations of the particles velocities from the global motion.
+	* The space in which local damping is calculated. Center of mass adds the expense of calculating the center of mass.
+	*/
+	UPROPERTY(EditAnywhere, Category = "Damping Properties", Meta = (InteractorName = "LocalDampingSpace"))
+	EChaosSoftsLocalDampingSpace LocalDampingSpace = EChaosSoftsLocalDampingSpace::CenterOfMass;
+
+	/**
+	 * The amount of local linear damping applied to the cloth velocities.
+	 * This type of damping only damps individual deviations of the particles velocities from the global linear motion.
+	 * It makes the cloth deformations more cohesive and reduces jitter without affecting the overall movement.
+	 * It can also produce synchronization artifacts where part of the cloth mesh are disconnected (which might well be desirable, or not).
+	 */
+	UPROPERTY(EditAnywhere, Category = "Damping Properties", Meta = (UIMin = "0", UIMax = "1", ClampMin = "0", ClampMax = "1", InteractorName = "LocalDampingLinearCoefficient"))
+	float LocalDampingLinearCoefficient = 0.f;
+
+	/**
+	 * The amount of local angular damping applied to the cloth velocities.
+	 * This type of damping only damps individual deviations of the particles velocities from the global angular motion.
 	 * It makes the cloth deformations more cohesive and reduces jitter without affecting the overall movement.
-	 * It can also produce synchronization artifacts where part of the cloth mesh are disconnected (which might well be desirable, or not), and is more expensive than global damping.
+	 * It can also produce synchronization artifacts where part of the cloth mesh are disconnected (which might well be desirable, or not).
 	 */
-	UPROPERTY(EditAnywhere, Category = "Damping Properties", Meta = (UIMin = "0", UIMax = "1", ClampMin = "0", ClampMax = "1", InteractorName = "LocalDampingCoefficient"))
-	float LocalDampingCoefficient = 0.f;
+	UPROPERTY(EditAnywhere, Category = "Damping Properties", Meta = (EditCondition = "bEnableLocalDampingAngular", UIMin = "0", UIMax = "1", ClampMin = "0", ClampMax = "1", InteractorName = "LocalDampingAngularCoefficient"))
+	float LocalDampingAngularCoefficient = 0.f;
 
 	FChaosClothAssetSimulationDampingConfigNode(const UE::Dataflow::FNodeParameters& InParam, FGuid InGuid = FGuid::NewGuid());
 
@@ -40,6 +56,9 @@ private:
 	static constexpr float DeprecatedDampingCoefficientValue = -1.f; // This is outside the settable range when it wasn't deprecated.
 	UPROPERTY()
 	float DampingCoefficient_DEPRECATED = DeprecatedDampingCoefficientValue;
+
+	UPROPERTY()
+	float LocalDampingCoefficient_DEPRECATED = DeprecatedDampingCoefficientValue;
 #endif 
 
 };
diff --git a/Engine/Source/Runtime/Experimental/Chaos/Private/Chaos/PerParticleDampVelocity.cpp b/Engine/Source/Runtime/Experimental/Chaos/Private/Chaos/PerParticleDampVelocity.cpp
--- a/Engine/Source/Runtime/Experimental/Chaos/Private/Chaos/PerParticleDampVelocity.cpp
+++ b/Engine/Source/Runtime/Experimental/Chaos/Private/Chaos/PerParticleDampVelocity.cpp
@@ -24,77 +24,145 @@ static FAutoConsoleVariableRef CVarChaosDampVelocityISPCEnabled(TEXT("p.Chaos.Da
 
 namespace Chaos::Softs {
 
-void FPerParticleDampVelocity::UpdatePositionBasedState(const FSolverParticles& Particles, const int32 Offset, const int32 Range)
-{
-	const FSolverParticlesRange ParticlesRange(const_cast<FSolverParticles*>(&Particles), Offset, Range - Offset);
-	return UpdatePositionBasedState(ParticlesRange);
-}
-
-void FPerParticleDampVelocity::UpdatePositionBasedState(const FSolverParticlesRange& Particles)
-{
-#if INTEL_ISPC
-	if (bRealTypeCompatibleWithISPC && bChaos_DampVelocity_ISPC_Enabled)
+	void FPerParticleDampVelocity::SetProperties(
+		const FCollectionPropertyConstFacade& PropertyCollection,
+		const FSolverVec3& ReferenceLocation,
+		const FSolverVec3& ReferenceVelocity,
+		const FSolverVec3& ReferenceAngularVelocity)
 	{
-		ispc::UpdatePositionBasedState(
-			(ispc::FVector3f&)Xcm,
-			(ispc::FVector3f&)Vcm,
-			(ispc::FVector3f&)Omega,
-			(const ispc::FVector3f*)Particles.XArray().GetData(),
-			(const ispc::FVector3f*)Particles.GetV().GetData(),
-			Particles.GetM().GetData(),
-			Particles.GetInvM().GetData(),
-			0,
-			Particles.GetRangeSize());
-	}
-	else
-#endif
-	{
-		Xcm = FSolverVec3(0.);
-		Vcm = FSolverVec3(0.);
-		FSolverReal Mcm = (FSolverReal)0.;
+		constexpr FSolverReal DefaultPropertyCoefficient = 0.f;
+		if (!bPropertyIndicesInitialized)
+		{
+			LocalDampingCoefficientIndex = FLocalDampingCoefficientIndex(PropertyCollection);
+			LocalDampingLinearCoefficientIndex = FLocalDampingLinearCoefficientIndex(PropertyCollection);
+			LocalDampingAngularCoefficientIndex = FLocalDampingAngularCoefficientIndex(PropertyCollection);
+			LocalDampingSpaceIndex = FLocalDampingSpaceIndex(PropertyCollection);
 
-		for (int32 Index = 0; Index < Particles.GetRangeSize(); ++Index)
+			if (LocalDampingLinearCoefficientIndex != INDEX_NONE && LocalDampingAngularCoefficientIndex != INDEX_NONE)
+			{
+				LinearCoefficient = GetLocalDampingLinearCoefficient(PropertyCollection);
+				AngularCoefficient = GetLocalDampingAngularCoefficient(PropertyCollection);
+			}
+			else
+			{
+				const float LocalDampingCoefficient = GetLocalDampingCoefficient(PropertyCollection, DefaultPropertyCoefficient);
+				LinearCoefficient = AngularCoefficient = LocalDampingCoefficient;
+			}
+			bPropertyIndicesInitialized = true;
+			return;
+		}
+
+		if (LocalDampingLinearCoefficientIndex != INDEX_NONE && LocalDampingAngularCoefficientIndex != INDEX_NONE)
 		{
-			if (Particles.InvM(Index) == (FSolverReal)0.)
+			if (IsLocalDampingLinearCoefficientMutable(PropertyCollection))
+			{
+				LinearCoefficient = GetLocalDampingLinearCoefficient(PropertyCollection);
+			}
+			if (IsLocalDampingAngularCoefficientMutable(PropertyCollection))
 			{
-				continue;
+				AngularCoefficient = GetLocalDampingAngularCoefficient(PropertyCollection);
+			}
+		}
+		else
+		{
+			if (IsLocalDampingCoefficientMutable(PropertyCollection))
+			{
+				const float LocalDampingCoefficient = GetLocalDampingCoefficient(PropertyCollection);
+				LinearCoefficient = AngularCoefficient = LocalDampingCoefficient;
 			}
 
-			Xcm += Particles.GetX(Index) * Particles.M(Index);
-			Vcm += Particles.V(Index) * Particles.M(Index);
-			Mcm += Particles.M(Index);
 		}
 
-		if (Mcm != (FSolverReal)0.)
+		if (IsLocalDampingSpaceMutable(PropertyCollection))
 		{
-			Xcm /= Mcm;
-			Vcm /= Mcm;
+			LocalDampingSpace = (EChaosSoftsLocalDampingSpace)(GetLocalDampingSpace(PropertyCollection));
 		}
 
-		FSolverVec3 L = FSolverVec3(0.);
-		FSolverMatrix33 I(0.);
-		for (int32 Index = 0; Index < Particles.GetRangeSize(); ++Index)
+		if (LocalDampingSpace == EChaosSoftsLocalDampingSpace::ReferenceBoneSpace)
 		{
-			if (Particles.InvM(Index) == (FSolverReal)0.)
+			Xcm = ReferenceLocation;
+			Vcm = ReferenceVelocity;
+			Omega = ReferenceAngularVelocity;
+		}
+	}
+
+	void FPerParticleDampVelocity::UpdatePositionBasedState(const FSolverParticles& Particles, const int32 Offset, const int32 Range)
+	{
+		if (LocalDampingSpace == EChaosSoftsLocalDampingSpace::CenterOfMass)
+		{
+			const FSolverParticlesRange ParticlesRange(const_cast<FSolverParticles*>(&Particles), Offset, Range - Offset);
+			return UpdatePositionBasedState(ParticlesRange);
+		}
+	}
+
+	void FPerParticleDampVelocity::UpdatePositionBasedState(const FSolverParticlesRange& Particles)
+	{
+		if (LocalDampingSpace == EChaosSoftsLocalDampingSpace::CenterOfMass)
+		{
+#if INTEL_ISPC
+			if (bRealTypeCompatibleWithISPC && bChaos_DampVelocity_ISPC_Enabled)
 			{
-				continue;
+				ispc::UpdatePositionBasedState(
+					(ispc::FVector3f&)Xcm,
+					(ispc::FVector3f&)Vcm,
+					(ispc::FVector3f&)Omega,
+					(const ispc::FVector3f*)Particles.XArray().GetData(),
+					(const ispc::FVector3f*)Particles.GetV().GetData(),
+					Particles.GetM().GetData(),
+					Particles.GetInvM().GetData(),
+					0,
+					Particles.GetRangeSize());
 			}
+			else
+#endif
+			{
+				Xcm = FSolverVec3(0.);
+				Vcm = FSolverVec3(0.);
+				FSolverReal Mcm = (FSolverReal)0.;
 
-			const FSolverVec3 V = Particles.GetX(Index) - Xcm;
-			L += FSolverVec3::CrossProduct(V, Particles.M(Index) * Particles.V(Index));
-			const FSolverMatrix33 M(0, V[2], -V[1], -V[2], 0, V[0], V[1], -V[0], 0);
-			I += M.GetTransposed() * M * Particles.M(Index);
-		}
+				for (int32 Index = 0; Index < Particles.GetRangeSize(); ++Index)
+				{
+					if (Particles.InvM(Index) == (FSolverReal)0.)
+					{
+						continue;
+					}
+
+					Xcm += Particles.GetX(Index) * Particles.M(Index);
+					Vcm += Particles.V(Index) * Particles.M(Index);
+					Mcm += Particles.M(Index);
+				}
+
+				if (Mcm != (FSolverReal)0.)
+				{
+					Xcm /= Mcm;
+					Vcm /= Mcm;
+				}
 
-		const FSolverReal Det = I.Determinant();
-		Omega = Det < (FSolverReal)UE_SMALL_NUMBER || !FMath::IsFinite(Det) ?
-			FSolverVec3(0.) :
+				FSolverVec3 L = FSolverVec3(0.);
+				FSolverMatrix33 I(0.);
+				for (int32 Index = 0; Index < Particles.GetRangeSize(); ++Index)
+				{
+					if (Particles.InvM(Index) == (FSolverReal)0.)
+					{
+						continue;
+					}
+
+					const FSolverVec3 V = Particles.GetX(Index) - Xcm;
+					L += FSolverVec3::CrossProduct(V, Particles.M(Index) * Particles.V(Index));
+					const FSolverMatrix33 M(0, V[2], -V[1], -V[2], 0, V[0], V[1], -V[0], 0);
+					I += M.GetTransposed() * M * Particles.M(Index);
+				}
+
+				const FSolverReal Det = I.Determinant();
+				Omega = Det < (FSolverReal)UE_SMALL_NUMBER || !FMath::IsFinite(Det) ?
+					FSolverVec3(0.) :
 #if COMPILE_WITHOUT_UNREAL_SUPPORT
-			FSolverRigidTransform3(I).InverseTransformVector(L);
+					FSolverRigidTransform3(I).InverseTransformVector(L);
 #else
-			I.InverseTransformVector(L); // Calls FMatrix::InverseFast(), which tests against SMALL_NUMBER
+					I.InverseTransformVector(L); // Calls FMatrix::InverseFast(), which tests against SMALL_NUMBER
 #endif
+			}
+		}
 	}
-}
 
 }  // End namespace Chaos::Softs
diff --git a/Engine/Source/Runtime/Experimental/Chaos/Private/Chaos/SoftsEvolution.cpp b/Engine/Source/Runtime/Experimental/Chaos/Private/Chaos/SoftsEvolution.cpp
--- a/Engine/Source/Runtime/Experimental/Chaos/Private/Chaos/SoftsEvolution.cpp
+++ b/Engine/Source/Runtime/Experimental/Chaos/Private/Chaos/SoftsEvolution.cpp
@@ -1,6 +1,5 @@
 // Copyright Epic Games, Inc. All Rights Reserved.
 #include "Chaos/SoftsEvolution.h"
-#include "Chaos/PerParticleDampVelocity.h"
 #include "Chaos/SoftsEvolutionLinearSystem.h"
 #include "Chaos/PBDFlatWeightMap.h"
 
@@ -22,6 +21,7 @@ namespace EvolutionSoftBodyDefault
 {
 constexpr FSolverReal GlobalDamping = (FSolverReal)0.01;
 constexpr FSolverReal LocalDamping = (FSolverReal)0.;
+constexpr EChaosSoftsLocalDampingSpace LocalDampingSpace = EChaosSoftsLocalDampingSpace::CenterOfMass;
 constexpr bool bUsePerParticleDamping = false;
 }
 
@@ -242,7 +242,7 @@ int32 FEvolution::AddSoftBody(uint32 GroupId, int32 NumParticles, bool bEnable)
 	SoftBodies.Status[SoftBodyId] = FSoftBodies::EStatus::Inactive;
 	SoftBodies.GroupId[SoftBodyId] = GroupId;
 	SoftBodies.GlobalDampings[SoftBodyId] = Private::EvolutionSoftBodyDefault::GlobalDamping;
-	SoftBodies.LocalDampings[SoftBodyId] = Private::EvolutionSoftBodyDefault::LocalDamping;
+	SoftBodies.LocalDamping[SoftBodyId] = FPerParticleDampVelocity(Private::EvolutionSoftBodyDefault::LocalDamping);
 	SoftBodies.UsePerParticleDamping[SoftBodyId] = Private::EvolutionSoftBodyDefault::bUsePerParticleDamping;
 	SoftBodies.PreSubstepParallelInits[SoftBodyId].Reset();
 	SoftBodies.PBDExternalForceRules[SoftBodyId].Reset();
@@ -276,9 +276,12 @@ void FEvolution::RemoveSoftBody(int32 SoftBodyId)
 
 
 void FEvolution::SetSoftBodyProperties(int32 SoftBodyId, const FCollectionPropertyConstFacade& PropertyCollection,
-	const TMap<FString, TConstArrayView<FRealSingle>>& WeightMaps)
+	const TMap<FString, TConstArrayView<FRealSingle>>& WeightMaps,
+	const FSolverVec3& ReferenceSpaceLocation,
+	const FSolverVec3& ReferenceSpaceVelocity,
+	const FSolverVec3& ReferenceSpaceAngularVelocity)
 {
-	SoftBodies.LocalDampings[SoftBodyId] = GetLocalDampingCoefficient(PropertyCollection, Private::EvolutionSoftBodyDefault::LocalDamping);
+	SoftBodies.LocalDamping[SoftBodyId].SetProperties(PropertyCollection, ReferenceSpaceLocation, ReferenceSpaceVelocity, ReferenceSpaceAngularVelocity);
 
 	const FSolverVec2 GlobalDamping = GetWeightedFloatDampingCoefficient(PropertyCollection, Private::EvolutionSoftBodyDefault::GlobalDamping);
 	TConstArrayView<FRealSingle> WeightMap = WeightMaps.FindRef(GetDampingCoefficientString(PropertyCollection, DampingCoefficientName.ToString()));
@@ -455,13 +458,14 @@ void FEvolution::AdvanceOneTimeStepInternal(const FSolverReal Dt, const int32 Ti
 		Context.Init(SolverMode, Dt, NumPBDIterations, bDoNewtonUpdate ? NumNewtonIterations : 0);
 
 		// PreSubstepParallelInits + LocalDamping
-		const bool bDoLocalDamping = SoftBodies.LocalDampings[SoftBodyId] > (FSolverReal)0.;
-		FPerParticleDampVelocity DampVelocityRule(SoftBodies.LocalDampings[SoftBodyId]);
+		FPerParticleDampVelocity& DampVelocityRule = SoftBodies.LocalDamping[SoftBodyId];
+		const bool bDoLocalDampingPositionUpdate = DampVelocityRule.RequiresUpdatePositionBasedState();
+		const bool bDoLocalDamping = DampVelocityRule.IsEnabled();
 		{
 			TRACE_CPUPROFILER_EVENT_SCOPE(FSoftsEvolution_AdvanceOneTimeStepInternal_PreSubstepParallelInits);
 
 			TConstArrayView<ParallelInitFunc> PreSubstepParallelInits(SoftBodies.PreSubstepParallelInits[SoftBodyId]);
-			const int32 NumPreSubstepInits = PreSubstepParallelInits.Num() + (bDoLocalDamping ? 1 : 0);
+			const int32 NumPreSubstepInits = PreSubstepParallelInits.Num() + (bDoLocalDampingPositionUpdate ? 1 : 0);
 			PhysicsParallelFor(NumPreSubstepInits,
 				[this, &PreSubstepParallelInits, &DampVelocityRule, Dt, SoftBodyId, SolverMode](int32 FuncIdx)
 			{
diff --git a/Engine/Source/Runtime/Experimental/Chaos/Public/Chaos/PerParticleDampVelocity.h b/Engine/Source/Runtime/Experimental/Chaos/Public/Chaos/PerParticleDampVelocity.h
--- a/Engine/Source/Runtime/Experimental/Chaos/Public/Chaos/PerParticleDampVelocity.h
+++ b/Engine/Source/Runtime/Experimental/Chaos/Public/Chaos/PerParticleDampVelocity.h
@@ -2,6 +2,8 @@
 #pragma once
 
 #include "Chaos/PBDSoftsEvolutionFwd.h"
+#include "Chaos/CollectionPropertyFacade.h"
+#include "Chaos/SoftsSimulationSpace.h"
 #include "Chaos/Vector.h"
 #include "Chaos/PBDSoftsSolverParticles.h"
 
@@ -11,10 +13,17 @@ namespace Chaos::Softs
 class FPerParticleDampVelocity final
 {
 public:
+
 	FPerParticleDampVelocity(const FSolverReal InCoefficient = (FSolverReal)0.01)
-		: Coefficient(InCoefficient)
+		: LinearCoefficient(InCoefficient)
+		, AngularCoefficient(InCoefficient)
+		, LocalDampingCoefficientIndex(ForceInit)
+		, LocalDampingLinearCoefficientIndex(ForceInit)
+		, LocalDampingAngularCoefficientIndex(ForceInit)
+		, LocalDampingSpaceIndex(ForceInit)
 	{
 	}
+
 	~FPerParticleDampVelocity() {}
 
 	void UpdatePositionBasedState(const FSolverParticles& Particles, const int32 Offset, const int32 Range);
@@ -29,13 +38,37 @@ public:
 	void Apply(const FSolverVec3& X, FSolverVec3& V) const
 	{
 		const FSolverVec3 R = X - Xcm;
-		const FSolverVec3 Dv = Vcm - V + FSolverVec3::CrossProduct(R, Omega);
-		V += Coefficient * Dv;
+		const FSolverVec3 Dv = LinearCoefficient * (Vcm - V) + AngularCoefficient * FSolverVec3::CrossProduct(R, Omega);
+		V += Dv;
+	}
+
+	bool IsEnabled() const
+	{
+		return LinearCoefficient > 0.f || AngularCoefficient > 0.f;
 	}
 
+	bool RequiresUpdatePositionBasedState() const
+	{
+		return LocalDampingSpace == EChaosSoftsLocalDampingSpace::CenterOfMass;
+	}
+
+	void SetProperties(
+		const FCollectionPropertyConstFacade& PropertyCollection,
+		const FSolverVec3& ReferenceLocation,
+		const FSolverVec3& ReferenceVelocity,
+		const FSolverVec3& ReferenceAngularVelocity);
+
 private:
-	FSolverReal Coefficient;
+	EChaosSoftsLocalDampingSpace LocalDampingSpace = EChaosSoftsLocalDampingSpace::CenterOfMass;
+	bool bPropertyIndicesInitialized = false;
+	FSolverReal LinearCoefficient;
+	FSolverReal AngularCoefficient;
 	FSolverVec3 Xcm, Vcm, Omega;
+
+	UE_CHAOS_DECLARE_PROPERTYCOLLECTION_NAME(LocalDampingCoefficient, float);
+	UE_CHAOS_DECLARE_PROPERTYCOLLECTION_NAME(LocalDampingLinearCoefficient, float);
+	UE_CHAOS_DECLARE_PROPERTYCOLLECTION_NAME(LocalDampingAngularCoefficient, float);
+	UE_CHAOS_DECLARE_PROPERTYCOLLECTION_NAME(LocalDampingSpace, int32);
 };
 
 }  // End namespace Chaos::Softs
diff --git a/Engine/Source/Runtime/Experimental/Chaos/Public/Chaos/SoftsEvolution.h b/Engine/Source/Runtime/Experimental/Chaos/Public/Chaos/SoftsEvolution.h
--- a/Engine/Source/Runtime/Experimental/Chaos/Public/Chaos/SoftsEvolution.h
+++ b/Engine/Source/Runtime/Experimental/Chaos/Public/Chaos/SoftsEvolution.h
@@ -5,6 +5,7 @@
 #include "Chaos/ArrayCollection.h"
 #include "Chaos/PBDSoftsEvolutionFwd.h"
 #include "Chaos/PBDSoftsSolverParticles.h"
+#include "Chaos/PerParticleDampVelocity.h"
 #include "Chaos/SoftsEvolutionLinearSystem.h"
 #include "Chaos/SoftsSolverParticlesRange.h"
 #include "Chaos/SoftsSolverCollisionParticles.h"
@@ -12,11 +13,11 @@
 #include "Chaos/VelocityField.h"
 #include "Misc/EnumClassFlags.h"
 
+enum struct EChaosSoftsLocalDampingSpace : uint8;
+
 namespace Chaos::Softs
 {
 
-
-
 enum struct ESolverMode : uint8
 {
 	None = 0,
@@ -125,7 +126,17 @@ public:
 	int32 GetSoftBodyParticleNum(int32 SoftBodyId) const { return SoftBodies.ParticleRanges[SoftBodyId].GetRangeSize(); }
 	int32 GetSoftBodyGroupId(int32 SoftBodyId) const { return SoftBodies.GroupId[SoftBodyId]; }
 	CHAOS_API void SetSoftBodyProperties(int32 SoftBodyId, const FCollectionPropertyConstFacade& PropertyCollection,
-		const TMap<FString, TConstArrayView<FRealSingle>>& WeightMaps);
+		const TMap<FString, TConstArrayView<FRealSingle>>& WeightMaps,
+		const FSolverVec3& ReferenceSpaceLocation,
+		const FSolverVec3& ReferenceSpaceVelocity,
+		const FSolverVec3& ReferenceSpaceAngularVelocity
+	);
+	UE_DEPRECATED(5.6, "Use version with ReferenceSpace parameters. Otherwise, LocalSpace damping with ReferenceBone space will not work correctly.")
+	void SetSoftBodyProperties(int32 SoftBodyId, const FCollectionPropertyConstFacade& PropertyCollection,
+			const TMap<FString, TConstArrayView<FRealSingle>>& WeightMaps)
+	{
+		SetSoftBodyProperties(SoftBodyId, PropertyCollection, WeightMaps, FSolverVec3(0.f), FSolverVec3(0.f), FSolverVec3(0.f));
+	}
 	// Activating/deactivating Soft Bodies in different groups is threadsafe. Activations within a group is not threadsafe.
 	CHAOS_API void ActivateSoftBody(int32 SoftBodyId, bool bActivate);
 	bool IsSoftBodyActive(int32 SoftBodyId) const { return SoftBodies.Status[SoftBodyId] == FSoftBodies::EStatus::Active; }
@@ -329,7 +340,7 @@ private:
 			TArrayCollection::AddArray(&GroupId);
 			TArrayCollection::AddArray(&ParticleRanges);
 			TArrayCollection::AddArray(&GlobalDampings);
-			TArrayCollection::AddArray(&LocalDampings);
+			TArrayCollection::AddArray(&LocalDamping);
 			TArrayCollection::AddArray(&UsePerParticleDamping);
 			TArrayCollection::AddArray(&LinearSystems);
 
@@ -368,7 +379,7 @@ private:
 		TArrayCollectionArray<uint32> GroupId;
 		TArrayCollectionArray<FSolverParticlesRange> ParticleRanges;
 		TArrayCollectionArray<FSolverReal> GlobalDampings;
-		TArrayCollectionArray<FSolverReal> LocalDampings;
+		TArrayCollectionArray<FPerParticleDampVelocity> LocalDamping;
 		TArrayCollectionArray<bool> UsePerParticleDamping;
 		TArrayCollectionArray<FEvolutionLinearSystem> LinearSystems;
 
diff --git a/Engine/Source/Runtime/Experimental/Chaos/Public/Chaos/SoftsSimulationSpace.h b/Engine/Source/Runtime/Experimental/Chaos/Public/Chaos/SoftsSimulationSpace.h
--- a/Engine/Source/Runtime/Experimental/Chaos/Public/Chaos/SoftsSimulationSpace.h
+++ b/Engine/Source/Runtime/Experimental/Chaos/Public/Chaos/SoftsSimulationSpace.h
@@ -15,4 +15,15 @@ enum struct EChaosSoftsSimulationSpace : uint8
 
 	// Top level bone for each cloth.
 	ReferenceBoneSpace,
+};
+
+UENUM()
+enum struct EChaosSoftsLocalDampingSpace : uint8
+{
+	// Calculated CoM (expensive)
+	CenterOfMass = 0,
+
+	// Top level bone for each cloth.
+	ReferenceBoneSpace,
+
 };
\ No newline at end of file
