fix crash on invalid auto-instance mesh access in geometry collection conversion

diff --git a/Engine/Source/Runtime/Experimental/GeometryCollectionEngine/Private/GeometryCollection/GeometryCollectionEngineConversion.cpp b/Engine/Source/Runtime/Experimental/GeometryCollectionEngine/Private/GeometryCollection/GeometryCollectionEngineConversion.cpp
--- a/Engine/Source/Runtime/Experimental/GeometryCollectionEngine/Private/GeometryCollection/GeometryCollectionEngineConversion.cpp
+++ b/Engine/Source/Runtime/Experimental/GeometryCollectionEngine/Private/GeometryCollection/GeometryCollectionEngineConversion.cpp
@@ -1101,7 +1101,7 @@ void FGeometryCollectionEngineConversion::AppendGeometryCollectionInstancedMeshe
 				int32 TargetInstancedMeshIndex = INDEX_NONE;
 
 				const int32 SourceAutoInstanceIndex = SourceInstancedMeshFacade.GetIndex(SourceTransformIndex);
-				if (SourceAutoInstanceIndex != INDEX_NONE)
+				if (SourceGeometryCollectionObject->AutoInstanceMeshes.IsValidIndex(SourceAutoInstanceIndex))
 				{
 					const FGeometryCollectionAutoInstanceMesh& SourceAutoInstanceMesh = SourceGeometryCollectionObject->GetAutoInstanceMesh(SourceAutoInstanceIndex);
 					TargetInstancedMeshIndex = TargetGeometryCollectionObject->FindOrAddAutoInstanceMesh(SourceAutoInstanceMesh);
