Geometry collection : fix automatic root proxies not using the right transform relative to the root bone of the collection
- also fixed the dataflow related workflows
- fixed root proxy custom renderer to properly remove existing static mesh before registering new ones

diff --git a/Engine/Plugins/Experimental/ChaosEditor/Source/FractureEditor/Private/FractureToolGenerators.cpp b/Engine/Plugins/Experimental/ChaosEditor/Source/FractureEditor/Private/FractureToolGenerators.cpp
--- a/Engine/Plugins/Experimental/ChaosEditor/Source/FractureEditor/Private/FractureToolGenerators.cpp
+++ b/Engine/Plugins/Experimental/ChaosEditor/Source/FractureEditor/Private/FractureToolGenerators.cpp
@@ -1010,18 +1010,7 @@ AGeometryCollectionActor* UFractureToolGenerateAsset::ConvertActorsToGeometryCol
 	if (Params.bUseRootProxies)
 	{
 		FracturedGeometryCollection->CustomRendererType = UGeometryCollectionRootProxyRenderer::StaticClass();
-		FracturedGeometryCollection->RootProxyData.ProxyMeshes.Reserve(FracturedGeometryCollection->GeometrySource.Num());
-		FracturedGeometryCollection->RootProxyData.MeshTransforms.Reserve(FracturedGeometryCollection->GeometrySource.Num());
-
-		for (const FGeometryCollectionSource& Source : FracturedGeometryCollection->GeometrySource)
-		{
-			UObject* SourceMesh = Source.SourceGeometryObject.TryLoad();
-			if (UStaticMesh* SourceStaticMesh = Cast<UStaticMesh>(SourceMesh))
-			{
-				FracturedGeometryCollection->RootProxyData.ProxyMeshes.Add(SourceStaticMesh);
-				FracturedGeometryCollection->RootProxyData.MeshTransforms.Add(FTransform3f(Source.LocalTransform));
-			}
-		}
+		FracturedGeometryCollection->SetRootProxiesFromGeometrySources();
 	}
 
 	FracturedGeometryCollection->InitializeMaterials(bAddInternalMaterials);
diff --git a/Engine/Plugins/Experimental/GeometryCollectionPlugin/Source/GeometryCollectionNodes/Private/Dataflow/GeometryCollectionStaticMeshToCollectionNode.cpp b/Engine/Plugins/Experimental/GeometryCollectionPlugin/Source/GeometryCollectionNodes/Private/Dataflow/GeometryCollectionStaticMeshToCollectionNode.cpp
--- a/Engine/Plugins/Experimental/GeometryCollectionPlugin/Source/GeometryCollectionNodes/Private/Dataflow/GeometryCollectionStaticMeshToCollectionNode.cpp
+++ b/Engine/Plugins/Experimental/GeometryCollectionPlugin/Source/GeometryCollectionNodes/Private/Dataflow/GeometryCollectionStaticMeshToCollectionNode.cpp
@@ -10,6 +10,7 @@
 #include "GeometryCollection/GeometryCollectionEngineConversion.h"
 #include "GeometryCollection/GeometryCollectionObject.h"
 #include "GeometryCollection/ManagedArrayCollection.h"
+#include "GeometryCollection/Facades/CollectionTransformFacade.h"
 
 namespace UE::Dataflow
 {
diff --git a/Engine/Source/Runtime/Experimental/Chaos/Private/GeometryCollection/Facades/CollectionTransformFacade.cpp b/Engine/Source/Runtime/Experimental/Chaos/Private/GeometryCollection/Facades/CollectionTransformFacade.cpp
--- a/Engine/Source/Runtime/Experimental/Chaos/Private/GeometryCollection/Facades/CollectionTransformFacade.cpp
+++ b/Engine/Source/Runtime/Experimental/Chaos/Private/GeometryCollection/Facades/CollectionTransformFacade.cpp
@@ -67,6 +67,20 @@ namespace GeometryCollection::Facades
 		return Chaos::Facades::FCollectionHierarchyFacade::GetRootIndices(ParentAttribute);
 	}
 
+	FTransform FCollectionTransformFacade::GetRootTransform() const
+	{
+		const TArray<int32> RootIndices = GetRootIndices();
+		if (RootIndices.Num() > 0)
+		{
+			const int32 RootIndex = RootIndices[0];
+			if (TransformAttribute.IsValid())
+			{
+				return FTransform(TransformAttribute[RootIndex]);
+			}
+		}
+		return FTransform::Identity;
+	}
+
 	TArray<FTransform> FCollectionTransformFacade::ComputeCollectionSpaceTransforms() const
 	{
 		TArray<FTransform> OutTransforms;
diff --git a/Engine/Source/Runtime/Experimental/Chaos/Public/GeometryCollection/Facades/CollectionTransformFacade.h b/Engine/Source/Runtime/Experimental/Chaos/Public/GeometryCollection/Facades/CollectionTransformFacade.h
--- a/Engine/Source/Runtime/Experimental/Chaos/Public/GeometryCollection/Facades/CollectionTransformFacade.h
+++ b/Engine/Source/Runtime/Experimental/Chaos/Public/GeometryCollection/Facades/CollectionTransformFacade.h
@@ -33,6 +33,9 @@ namespace GeometryCollection::Facades
 		/** Gets the root index */
 		CHAOS_API TArray<int32> GetRootIndices() const;
 
+		/** Gets the main root transform */
+		CHAOS_API FTransform GetRootTransform() const;
+
 		/**
 		* Returns the parent indices from the collection. Null if not initialized.
 		*/
diff --git a/Engine/Source/Runtime/Experimental/GeometryCollectionEngine/Private/GeometryCollection/GeometryCollectionObject.cpp b/Engine/Source/Runtime/Experimental/GeometryCollectionEngine/Private/GeometryCollection/GeometryCollectionObject.cpp
--- a/Engine/Source/Runtime/Experimental/GeometryCollectionEngine/Private/GeometryCollection/GeometryCollectionObject.cpp
+++ b/Engine/Source/Runtime/Experimental/GeometryCollectionEngine/Private/GeometryCollection/GeometryCollectionObject.cpp
@@ -1524,6 +1524,36 @@ void  UGeometryCollection::PropagateTransformUpdateToComponents() const
 	}
 }
 
+void UGeometryCollection::SetRootProxiesFromGeometrySources()
+{
+	RootProxyData.ProxyMeshes.Reset();
+	RootProxyData.MeshTransforms.Reset();
+
+	// make sure the root index is properly cached
+	if (RootIndex == INDEX_NONE)
+	{
+		UpdateRootIndex();
+	}
+
+	FTransform3f RootTransform = FTransform3f::Identity;
+	if (GeometryCollection && GeometryCollection->Transform.IsValidIndex(RootIndex))
+	{
+		RootTransform = GeometryCollection->Transform[RootIndex];
+	}
+
+	for (const FGeometryCollectionSource& Source : GeometrySource)
+	{
+		UObject* SourceMesh = Source.SourceGeometryObject.TryLoad();
+		if (UStaticMesh* SourceStaticMesh = Cast<UStaticMesh>(SourceMesh))
+		{
+			RootProxyData.ProxyMeshes.Add(SourceStaticMesh);
+
+			const FTransform3f MeshTransform = FTransform3f(Source.LocalTransform).GetRelativeTransform(RootTransform);
+			RootProxyData.MeshTransforms.Add(MeshTransform);
+		}
+	}
+}
+
 TSharedPtr<FGeometryCollection, ESPMode::ThreadSafe> UGeometryCollection::GenerateMinimalGeometryCollection() const
 {
 	TMap<FName, TSet<FName>> SkipList;
diff --git a/Engine/Source/Runtime/Experimental/GeometryCollectionEngine/Private/GeometryCollection/GeometryCollectionRootProxyRenderer.cpp b/Engine/Source/Runtime/Experimental/GeometryCollectionEngine/Private/GeometryCollection/GeometryCollectionRootProxyRenderer.cpp
--- a/Engine/Source/Runtime/Experimental/GeometryCollectionEngine/Private/GeometryCollection/GeometryCollectionRootProxyRenderer.cpp
+++ b/Engine/Source/Runtime/Experimental/GeometryCollectionEngine/Private/GeometryCollection/GeometryCollectionRootProxyRenderer.cpp
@@ -62,7 +62,7 @@ void UGeometryCollectionRootProxyRenderer::CreateRootProxyComponents(UGeometryCo
 	{
 		return;
 	}
-
+	ClearRootProxyComponents();
 	StaticMeshComponents.Reserve(GeometryCollection->RootProxyData.ProxyMeshes.Num());
 
 	for (UStaticMesh* ProxyMesh : GeometryCollection->RootProxyData.ProxyMeshes)
diff --git a/Engine/Source/Runtime/Experimental/GeometryCollectionEngine/Public/GeometryCollection/GeometryCollectionObject.h b/Engine/Source/Runtime/Experimental/GeometryCollectionEngine/Public/GeometryCollection/GeometryCollectionObject.h
--- a/Engine/Source/Runtime/Experimental/GeometryCollectionEngine/Public/GeometryCollection/GeometryCollectionObject.h
+++ b/Engine/Source/Runtime/Experimental/GeometryCollectionEngine/Public/GeometryCollection/GeometryCollectionObject.h
@@ -499,11 +499,15 @@ public:
 
 	/** Propogate the fact that transform have been updated to all components */
 	GEOMETRYCOLLECTIONENGINE_API void PropagateTransformUpdateToComponents() const;
+
+	/** Create root proxies from geometryt sources */
+	GEOMETRYCOLLECTIONENGINE_API void SetRootProxiesFromGeometrySources();
+
 #endif
 
 	GEOMETRYCOLLECTIONENGINE_API void InitResources();
 	GEOMETRYCOLLECTIONENGINE_API void ReleaseResources();
-
+	
 	/** Fills params struct with parameters used for precomputing content. */
 	GEOMETRYCOLLECTIONENGINE_API void GetSharedSimulationParams(FSharedSimulationParameters& OutParams) const;
 
