Fixed assumption over the amount of content expected

#rnx

diff --git a/Engine/Source/Runtime/Experimental/IoStore/HttpClient/Private/Loop.inl b/Engine/Source/Runtime/Experimental/IoStore/HttpClient/Private/Loop.inl
--- a/Engine/Source/Runtime/Experimental/IoStore/HttpClient/Private/Loop.inl
+++ b/Engine/Source/Runtime/Experimental/IoStore/HttpClient/Private/Loop.inl
@@ -175,7 +175,7 @@ static FOutcome DoRecvMessage(FActivity* Activity, FHttpPeer& Peer)
 	int32 Count = 3;
 	bool bChunked = false;
 	bool IsKeepAlive = true;
-	int32 ContentLength = -1;
+	int32 ContentLength = Activity->NoContent ? 0 : -1;
 	EnumerateHeaders(
 		Headers,
 		[&ContentLength, &bChunked, &IsKeepAlive, &Count] (FAnsiStringView Name, FAnsiStringView Value)
diff --git a/Engine/Source/Runtime/Experimental/IoStore/HttpClient/Private/Test.inl b/Engine/Source/Runtime/Experimental/IoStore/HttpClient/Private/Test.inl
--- a/Engine/Source/Runtime/Experimental/IoStore/HttpClient/Private/Test.inl
+++ b/Engine/Source/Runtime/Experimental/IoStore/HttpClient/Private/Test.inl
@@ -255,12 +255,6 @@ static void TlsTest()
 		WaitForLoopIdle();
 	}
 
-	if (false) {
-		FRequest Request = Loop.Request("HEAD", "[REDACTED_PROFILE_URL]/clink/releases/download/1.0.0a4/clink-1.0.0a4.zip");
-		Loop.Send(MoveTemp(Request), OkSink);
-		WaitForLoopIdle();
-	}
-
 	{
 		FCertRoots NotACert(FMemoryView("493", 3));
 		check(NotACert.IsValid() == false);
@@ -666,6 +660,23 @@ static void HttpTest(const ANSICHAR* TestHost, FCertRootsRef VerifyCert)
 		WaitForLoopIdle();
 	}
 
+	// head barage
+	{
+		FConnectionPool::FParams Params;
+		Params.SetHostFromUrl(BuildUrl());
+		Params.VerifyCert = VerifyCert;
+		Params.ConnectionCount = 1;
+		FConnectionPool Pool(Params);
+		for (int32 i = 0; i < 61; ++i)
+		{
+			TAnsiStringBuilder<16> Path;
+			Path << "/data?head";
+			FRequest Request = Loop.Request("HEAD", Path, Pool);
+			Loop.Send(MoveTemp(Request), NullSink);
+		}
+		WaitForLoopIdle();
+	}
+
 	// fatal timeout
 	for (int32 i = 0; i < 14; ++i)
 	{
diff --git a/Engine/Source/Runtime/Experimental/IoStore/HttpClient/Private/TestServer.py b/Engine/Source/Runtime/Experimental/IoStore/HttpClient/Private/TestServer.py
--- a/Engine/Source/Runtime/Experimental/IoStore/HttpClient/Private/TestServer.py
+++ b/Engine/Source/Runtime/Experimental/IoStore/HttpClient/Private/TestServer.py
@@ -52,7 +52,6 @@ def proxy_impl(client, httpd):
 
     # extract request
     method, path, proto = line.split(b" ")
-    assert method == b"GET"
 
     if b"?" in path:
         _, query = path.split(b"?", 1)
@@ -69,11 +68,12 @@ def proxy_impl(client, httpd):
     # get response
     line, headers = intercept(httpd)
     close = close or (headers.get("Connection", "").lower() == "close")
-    content_len = int(headers.get("Content-Length", "0"))
 
     # get data to retransmit
     data = make_preamble(line, headers)
-    data += httpd.read(content_len)
+    if method != b"HEAD":
+        content_len = int(headers.get("Content-Length", "0"))
+        data += httpd.read(content_len)
     data_size = len(data)
 
     if stall:
@@ -295,7 +295,8 @@ class Handler(http.server.BaseHTTPRequestHandler):
             return
 
         self.send_response(200)
-        self.send_header("Content-Length", 0)
+        if random.random() < 0.75:
+            self.send_header("Content-Length", 0)
         self.end_headers()
 
     def end_headers(self):
