Added CIntegral concept.

diff --git a/Engine/Source/Runtime/Core/Public/Clang/ClangPlatformMath.h b/Engine/Source/Runtime/Core/Public/Clang/ClangPlatformMath.h
--- a/Engine/Source/Runtime/Core/Public/Clang/ClangPlatformMath.h
+++ b/Engine/Source/Runtime/Core/Public/Clang/ClangPlatformMath.h
@@ -8,6 +8,7 @@
 
 #ifdef __clang__
 
+#include "Concepts/Integral.h"
 #include "GenericPlatform/GenericPlatformMath.h"
 
 // HEADER_UNIT_UNSUPPORTED - Clang not supporting header units
@@ -128,10 +129,7 @@ struct FClangPlatformMath : public FGenericPlatformMath
 	 * If there was overflow, it returns false, and OutResult may or may not be written.
 	 * If there wasn't overflow, it returns true, and the result of the addition is written to OutResult.
 	 */
-	template <
-		typename IntType
-		UE_REQUIRES(std::is_integral_v<IntType>)
-	>
+	template <CIntegral IntType>
 	static FORCEINLINE bool AddAndCheckForOverflow(IntType A, IntType B, IntType& OutResult)
 	{
 		return !__builtin_add_overflow(A, B, &OutResult);
@@ -142,10 +140,7 @@ struct FClangPlatformMath : public FGenericPlatformMath
 	 * If there was overflow, it returns false, and OutResult may or may not be written.
 	 * If there wasn't overflow, it returns true, and the result of the subtraction is written to OutResult.
 	 */
-	template <
-		typename IntType
-		UE_REQUIRES(std::is_integral_v<IntType>)
-	>
+	template <CIntegral IntType>
 	static FORCEINLINE bool SubtractAndCheckForOverflow(IntType A, IntType B, IntType& OutResult)
 	{
 		return !__builtin_sub_overflow(A, B, &OutResult);
@@ -156,10 +151,7 @@ struct FClangPlatformMath : public FGenericPlatformMath
 	 * If there was overflow, it returns false, and OutResult may or may not be written.
 	 * If there wasn't overflow, it returns true, and the result of the multiplication is written to OutResult.
 	 */
-	template <
-		typename IntType
-		UE_REQUIRES(std::is_integral_v<IntType>)
-	>
+	template <CIntegral IntType>
 	static FORCEINLINE bool MultiplyAndCheckForOverflow(IntType A, IntType B, IntType& OutResult)
 	{
 		return !__builtin_mul_overflow(A, B, &OutResult);
diff --git a/Engine/Source/Runtime/Core/Public/Concepts/Integral.h b/Engine/Source/Runtime/Core/Public/Concepts/Integral.h
new file mode 100644
--- /dev/null
+++ b/Engine/Source/Runtime/Core/Public/Concepts/Integral.h
@@ -0,0 +1,11 @@
+// Copyright Epic Games, Inc. All Rights Reserved.
+
+#pragma once
+
+#include <type_traits>
+
+/**
+ * Concept which describes an integral type.  We use this instead of std::integral because <concepts> isn't a well supported header yet.
+ */
+template <typename T>
+concept CIntegral = std::is_integral_v<T>;
diff --git a/Engine/Source/Runtime/Core/Public/Containers/Array.h b/Engine/Source/Runtime/Core/Public/Containers/Array.h
--- a/Engine/Source/Runtime/Core/Public/Containers/Array.h
+++ b/Engine/Source/Runtime/Core/Public/Containers/Array.h
@@ -2111,10 +2111,7 @@ public:
 	 * @param Count (Optional) Number of elements to remove. Default is 1.
 	 * @param AllowShrinking (Optional) Tells if this call can shrink array if suitable after remove. Default is yes.
 	 */
-	template <
-		typename CountType
-		UE_REQUIRES(std::is_integral_v<CountType>)
-	>
+	template <CIntegral CountType>
 	FORCEINLINE void RemoveAt(SizeType Index, CountType Count, EAllowShrinking AllowShrinking = EAllowShrinking::Default)
 	{
 		static_assert(!std::is_same_v<CountType, bool>, "TArray::RemoveAt: unexpected bool passed as the Count argument");
@@ -2220,10 +2217,7 @@ public:
 	 * @param AllowShrinking (Optional) Tells if this call can shrink array if
 	 *                        suitable after remove. Default is yes.
 	 */
-	template <
-		typename CountType
-		UE_REQUIRES(std::is_integral_v<CountType>)
-	>
+	template <CIntegral CountType>
 	FORCEINLINE void RemoveAtSwap(SizeType Index, CountType Count, EAllowShrinking AllowShrinking = EAllowShrinking::Default)
 	{
 		static_assert(!std::is_same_v<CountType, bool>, "TArray::RemoveAtSwap: unexpected bool passed as the Count argument");
diff --git a/Engine/Source/Runtime/Core/Public/GenericPlatform/GenericPlatformMath.h b/Engine/Source/Runtime/Core/Public/GenericPlatform/GenericPlatformMath.h
--- a/Engine/Source/Runtime/Core/Public/GenericPlatform/GenericPlatformMath.h
+++ b/Engine/Source/Runtime/Core/Public/GenericPlatform/GenericPlatformMath.h
@@ -8,6 +8,7 @@
 #pragma once
 
 #include "CoreTypes.h"
+#include "Concepts/Integral.h"
 #include "Containers/ContainersFwd.h"
 #include "HAL/PlatformCrt.h"
 #include "Templates/AndOrNot.h"
@@ -1181,10 +1182,7 @@ struct FGenericPlatformMath
 	 * If there was overflow, it returns false, and OutResult may or may not be written.
 	 * If there wasn't overflow, it returns true, and the result of the addition is written to OutResult.
 	 */
-	template <
-		typename IntType
-		UE_REQUIRES(std::is_integral_v<IntType>)
-	>
+	template <CIntegral IntType>
 	[[nodiscard]] static FORCEINLINE bool AddAndCheckForOverflow(IntType A, IntType B, IntType& OutResult)
 	{
 		// This follows Hacker's Delight, Chapter 2-12
@@ -1231,10 +1229,7 @@ struct FGenericPlatformMath
 	 * If there was overflow, it returns false, and OutResult may or may not be written.
 	 * If there wasn't overflow, it returns true, and the result of the subtraction is written to OutResult.
 	 */
-	template <
-		typename IntType
-		UE_REQUIRES(std::is_integral_v<IntType>)
-	>
+	template <CIntegral IntType>
 	[[nodiscard]] static FORCEINLINE bool SubtractAndCheckForOverflow(IntType A, IntType B, IntType& OutResult)
 	{
 		// This follows Hacker's Delight, Chapter 2-12
@@ -1275,10 +1270,7 @@ struct FGenericPlatformMath
 	 * If there was overflow, it returns false, and OutResult may or may not be written.
 	 * If there wasn't overflow, it returns true, and the result of the multiplication is written to OutResult.
 	 */
-	template <
-		typename IntType
-		UE_REQUIRES(std::is_integral_v<IntType>)
-	>
+	template <CIntegral IntType>
 	[[nodiscard]] static FORCEINLINE bool MultiplyAndCheckForOverflow(IntType A, IntType B, IntType& OutResult)
 	{
 		// Handle the case where the second factor is 0 specially (why will become clear in a minute).
diff --git a/Engine/Source/Runtime/Core/Public/Internationalization/FastDecimalFormat.h b/Engine/Source/Runtime/Core/Public/Internationalization/FastDecimalFormat.h
--- a/Engine/Source/Runtime/Core/Public/Internationalization/FastDecimalFormat.h
+++ b/Engine/Source/Runtime/Core/Public/Internationalization/FastDecimalFormat.h
@@ -2,6 +2,7 @@
 
 #pragma once
 
+#include "Concepts/Integral.h"
 #include "Containers/UnrealString.h"
 #include "CoreTypes.h"
 #include "Internationalization/Text.h"
@@ -77,10 +78,7 @@ struct FDecimalNumberIntegralLimits
 	uint64 NumericLimitMax;
 	bool bIsNumericSigned;
 
-	template <
-		typename IntegralType
-		UE_REQUIRES(std::is_integral_v<IntegralType>)
-	>
+	template <CIntegral IntegralType>
 	static FDecimalNumberIntegralLimits FromNumericLimits()
 	{
 		return FDecimalNumberIntegralLimits(TNumericLimits<IntegralType>::Lowest(), TNumericLimits<IntegralType>::Max(), std::is_signed_v<IntegralType>);
diff --git a/Engine/Source/Runtime/Core/Public/Math/GuardedInt.h b/Engine/Source/Runtime/Core/Public/Math/GuardedInt.h
--- a/Engine/Source/Runtime/Core/Public/Math/GuardedInt.h
+++ b/Engine/Source/Runtime/Core/Public/Math/GuardedInt.h
@@ -3,6 +3,7 @@
 #pragma once
 
 #include "CoreTypes.h"
+#include "Concepts/Integral.h"
 #include "HAL/PlatformMath.h"
 #include "Misc/AssertionMacros.h"
 #include "Templates/Requires.h"
@@ -50,7 +51,7 @@
  * As-is, the current approach is not the fastest, but it's not presently intended to be used in contexts
  * where speed of arithmetic operations is a major concern.
  */
-template<typename IntType>
+template<CIntegral IntType>
 class TGuardedInt
 {
 private:
@@ -69,10 +70,7 @@ public:
 	[[nodiscard]] TGuardedInt() = default;
 
 	/** Construct a TGuardedInt from a regular signed integer value. If it's out of range, it results in an invalid value. */
-	template<
-		typename ValueType
-		UE_REQUIRES(std::is_integral_v<ValueType>)
-	>
+	template<CIntegral ValueType>
 	[[nodiscard]] explicit TGuardedInt(ValueType InValue)
 		: Value((IntType)InValue), bIsValid(false)
 	{
@@ -99,7 +97,7 @@ public:
 	}
 
 	/** @return The value if valid, DefaultValue otherwise. */
-	[[nodiscard]]  IntType Get(const IntType DefaultValue) const
+	[[nodiscard]] IntType Get(const IntType DefaultValue) const
 	{
 		return IsValid() ? Value : DefaultValue;
 	}
diff --git a/Engine/Source/Runtime/Core/Public/Math/UnrealMathUtility.h b/Engine/Source/Runtime/Core/Public/Math/UnrealMathUtility.h
--- a/Engine/Source/Runtime/Core/Public/Math/UnrealMathUtility.h
+++ b/Engine/Source/Runtime/Core/Public/Math/UnrealMathUtility.h
@@ -3,6 +3,7 @@
 #pragma once
 
 #include "CoreTypes.h"
+#include "Concepts/Integral.h"
 #include "Misc/AssertionMacros.h"
 #include "HAL/PlatformMath.h"
 #include "Math/MathFwd.h"
@@ -534,13 +535,10 @@ public:
 	 * Converts an integral type to a nearest less or equal integer.
 	 * Unlike std::floor, it returns an IntegralType.
 	 */
-	template <
-		typename IntegralType
-		UE_REQUIRES(std::is_integral_v<IntegralType>)
-	>
+	template <CIntegral IntegralType>
 	[[nodiscard]] static constexpr FORCEINLINE IntegralType Floor(IntegralType I)
 	{
-	    return I;
+		return I;
 	}
 
 
@@ -659,10 +657,7 @@ public:
 	}
 
 	/** Wraps X to be between Min and Max, exclusive. Will never return Max. */
-	template <
-		typename T
-		UE_REQUIRES(std::is_integral_v<T>)
-	>
+	template <CIntegral T>
 	[[nodiscard]] static constexpr FORCEINLINE T WrapExclusive(const T X, const T Min, const T Max)
 	{
 		// Use unsigned type allow for large ranges which don't overflow on subtraction.
diff --git a/Engine/Source/Runtime/Experimental/IoStore/OnDemand/Private/Statistics.cpp b/Engine/Source/Runtime/Experimental/IoStore/OnDemand/Private/Statistics.cpp
--- a/Engine/Source/Runtime/Experimental/IoStore/OnDemand/Private/Statistics.cpp
+++ b/Engine/Source/Runtime/Experimental/IoStore/OnDemand/Private/Statistics.cpp
@@ -871,9 +871,9 @@ private:
 	/** This wrapper class allows us to treat raw integer types and FCounter types  as the same thing */
 	//template<typename CounterType, CounterType& Counter, int32 INDEX = -1>
 	template<typename CounterType, const CounterType& Counter, int32 INDEX = -1>
-	struct FTrackedValue
+	struct TTrackedValue
 	{
-		FTrackedValue()
+		TTrackedValue()
 		{
 			Value = GetCurrentValue();
 		}
@@ -897,23 +897,16 @@ private:
 		}
 
 	private:
-
-		template <
-			typename U = CounterType
-			UE_REQUIRES(!std::is_integral_v<U>)
-		>
-		int64 GetCurrentValue() const
-		{
-			return Counter.Get();
-		}
-
-		template <
-			typename U = CounterType
-			UE_REQUIRES(std::is_integral_v<U>)
-		>
 		int64 GetCurrentValue() const
 		{
-			return Counter;
+			if constexpr (std::is_integral_v<CounterType>)
+			{
+				return Counter;
+			}
+			else
+			{
+				return Counter.Get();
+			}
 		}
 
 		int64 Value;
@@ -923,26 +916,26 @@ private:
 	/** Http Stats */
 	struct FHttpStats
 	{
-		FTrackedValue<FCounterInt, GHttpErrorCount[0]>				ErrorCount;
-		FTrackedValue<FCounterAtomicInt, GHttpDecodeErrorCount[0]>	DecodeErrorCount;
-		FTrackedValue<FCounterInt, GHttpRetryCount[0]>				RetryCount;
-		FTrackedValue<FCounterInt, GHttpGetCount[0]>				GetCount;
-		FTrackedValue<FCounterInt, GHttpDownloadedBytes[0]>			DownloadedBytes;
-		FTrackedValue<int64, GHttpDurationMsSum[0]>					TotalDuration;
+		TTrackedValue<FCounterInt, GHttpErrorCount[0]>				ErrorCount;
+		TTrackedValue<FCounterAtomicInt, GHttpDecodeErrorCount[0]>	DecodeErrorCount;
+		TTrackedValue<FCounterInt, GHttpRetryCount[0]>				RetryCount;
+		TTrackedValue<FCounterInt, GHttpGetCount[0]>				GetCount;
+		TTrackedValue<FCounterInt, GHttpDownloadedBytes[0]>			DownloadedBytes;
+		TTrackedValue<int64, GHttpDurationMsSum[0]>					TotalDuration;
 		
 	} Http;
 
 	/** Cache Stats */
 	struct FCacheStats
 	{
-		FTrackedValue<FCounterAtomicInt, GCacheErrorCount>			ErrorCount;
-		FTrackedValue<FCounterAtomicInt, GCacheDecodeErrorCount>	DecodeErrorCount;
-		FTrackedValue<FCounterAtomicInt, GCacheGetCount>			GetCount;
-		FTrackedValue<FCounterAtomicInt, GCachePutCount>			PutCount;
-
-		FTrackedValue<FCounterAtomicInt, GCacheWrittenBytes>		WrittenBytes;
-		FTrackedValue<FCounterAtomicInt, GCacheReadBytes>			ReadBytes;
-		FTrackedValue<FCounterAtomicInt, GCacheRejectBytes>			RejectBytes;
+		TTrackedValue<FCounterAtomicInt, GCacheErrorCount>			ErrorCount;
+		TTrackedValue<FCounterAtomicInt, GCacheDecodeErrorCount>	DecodeErrorCount;
+		TTrackedValue<FCounterAtomicInt, GCacheGetCount>			GetCount;
+		TTrackedValue<FCounterAtomicInt, GCachePutCount>			PutCount;
+
+		TTrackedValue<FCounterAtomicInt, GCacheWrittenBytes>		WrittenBytes;
+		TTrackedValue<FCounterAtomicInt, GCacheReadBytes>			ReadBytes;
+		TTrackedValue<FCounterAtomicInt, GCacheRejectBytes>			RejectBytes;
 	} Cache;
 
 	const FBackendStatus& BackendStatus;
diff --git a/Engine/Source/Runtime/SlateCore/Public/Containers/ObservableArray.h b/Engine/Source/Runtime/SlateCore/Public/Containers/ObservableArray.h
--- a/Engine/Source/Runtime/SlateCore/Public/Containers/ObservableArray.h
+++ b/Engine/Source/Runtime/SlateCore/Public/Containers/ObservableArray.h
@@ -2,6 +2,7 @@
 
 #pragma once
 
+#include "Concepts/Integral.h"
 #include "Containers/Array.h"
 #include "Containers/ArrayView.h"
 #include "Delegates/Delegate.h"
@@ -340,10 +341,7 @@ public:
 		Array.RemoveAt(Index, AllowShrinking);
 		ArrayChangedDelegate.Broadcast(ObservableArrayChangedArgsType::MakeRemoveAction({ &RemovedElement, 1 }, Index));
 	}
-	template <
-		typename CountType
-		UE_REQUIRES(std::is_integral_v<CountType>)
-	>
+	template <CIntegral CountType>
 	void RemoveAt(SizeType Index, CountType NumToRemove, EAllowShrinking AllowShrinking = EAllowShrinking::Default)
 	{
 		static_assert(!std::is_same_v<CountType, bool>, "TObservableArray::RemoveAt: unexpected bool passed as the Count argument");
@@ -388,10 +386,7 @@ public:
 		Array.RemoveAtSwap(Index, AllowShrinking);
 		ArrayChangedDelegate.Broadcast(ObservableArrayChangedArgsType::MakeRemoveSwapAction({ &RemovedElement, 1 }, Index, PreviousNum - SwapAmount));
 	}
-	template <
-		typename CountType
-		UE_REQUIRES(std::is_integral_v<CountType>)
-	>
+	template <CIntegral CountType>
 	void RemoveAtSwap(SizeType Index, CountType NumToRemove, EAllowShrinking AllowShrinking = EAllowShrinking::Default)
 	{
 		static_assert(!std::is_same_v<CountType, bool>, "TObservableArray::RemoveAtSwap: unexpected bool passed as the Count argument");
