Geometry collection : Fix collision group not properly being set on partial clusters when they get created
this cause the partial cluster collision group to be set to INT_MAX and to miss collision with the original collision group of the geometry collection
Bug was introduced from CL 11257801

diff --git a/Engine/Source/Runtime/Experimental/Chaos/Private/Chaos/PBDRigidClustering.cpp b/Engine/Source/Runtime/Experimental/Chaos/Private/Chaos/PBDRigidClustering.cpp
--- a/Engine/Source/Runtime/Experimental/Chaos/Private/Chaos/PBDRigidClustering.cpp
+++ b/Engine/Source/Runtime/Experimental/Chaos/Private/Chaos/PBDRigidClustering.cpp
@@ -618,6 +618,7 @@ namespace Chaos
 		}
 		bool bClusterIsOneWayInteraction = true;
 		bool bClusterIsMACD = false;
+
 		for (FPBDRigidParticleHandle* Child : ChildrenArray)
 		{
 			if (FPBDRigidClusteredParticleHandle* ClusteredChild = Child->CastToClustered())
@@ -633,8 +634,11 @@ namespace Chaos
 					TopLevelClusterParentsStrained.Remove(ClusteredChild);
 				}
 
-				ClusteredChild->SetCollisionImpulses(FMath::Max(NewParticle->CollisionImpulses(), ClusteredChild->CollisionImpulses()));
-				Child->SetCollisionGroup(FMath::Min(NewParticle->CollisionGroup(), Child->CollisionGroup()));
+				NewParticle->SetCollisionImpulses(FMath::Max(NewParticle->CollisionImpulses(), ClusteredChild->CollisionImpulses()));
+
+				const int32 NewCG = NewParticle->CollisionGroup();
+				const int32 ChildCG = ClusteredChild->CollisionGroup();
+				NewParticle->SetCollisionGroup(NewCG < ChildCG ? NewCG : ChildCG);
 
 				bClusterIsOneWayInteraction &= Child->OneWayInteraction();
 				bClusterIsMACD |= Child->MACDEnabled();
