Geometry collection : make sure we have at least one material to make sure we don't get weird rendering issues when none are assigned

diff --git a/Engine/Source/Runtime/Experimental/GeometryCollectionEngine/Private/GeometryCollection/GeometryCollectionObject.cpp b/Engine/Source/Runtime/Experimental/GeometryCollectionEngine/Private/GeometryCollection/GeometryCollectionObject.cpp
--- a/Engine/Source/Runtime/Experimental/GeometryCollectionEngine/Private/GeometryCollection/GeometryCollectionObject.cpp
+++ b/Engine/Source/Runtime/Experimental/GeometryCollectionEngine/Private/GeometryCollection/GeometryCollectionObject.cpp
@@ -718,6 +718,11 @@ void UGeometryCollection::InitializeMaterials(bool bHasLegacyInternalMaterialsPa
 {
 	Modify();
 
+	if (Materials.IsEmpty())
+	{
+		Materials.Add({}); // add an empty material is none exist
+	}
+
 	// Initialize the BoneSelectedMaterial separate from the materials on the collection
 	UMaterialInterface* BoneSelectedMaterial = GetBoneSelectedMaterial();
 
