FOnDemandInstallCache::Defrag - Add missing call to FOnDemandInstallCacheStats::OnDefrag.

diff --git a/Engine/Source/Runtime/Experimental/IoStore/OnDemand/Private/OnDemandInstallCache.cpp b/Engine/Source/Runtime/Experimental/IoStore/OnDemand/Private/OnDemandInstallCache.cpp
--- a/Engine/Source/Runtime/Experimental/IoStore/OnDemand/Private/OnDemandInstallCache.cpp
+++ b/Engine/Source/Runtime/Experimental/IoStore/OnDemand/Private/OnDemandInstallCache.cpp
@@ -2390,6 +2390,7 @@ FIoStatus FOnDemandInstallCache::Defrag(
 		const FString ErrorMsg = FString::Printf(TEXT("Cache size is greater than disk quota - Cannot Defragment!, MaxCacheSize=%.2lf MiB, TotalCachedBytes=%.2lf MiB"), 
 			ToMiB(MaxCacheSize), ToMiB(TotalCachedBytes));
 		UE_LOG(LogIoStoreOnDemand, Error, TEXT("%s"), *ErrorMsg);
+		FOnDemandInstallCacheStats::OnDefrag(EIoErrorCode::WriteError, 0);
 		return FIoStatusBuilder(EIoErrorCode::WriteError) << ErrorMsg;
 	}
 
