Iris
* Added ability for dependent objects to be relevant only if their parent is also relevant to the connection

#rnx

diff --git a/Engine/Source/Runtime/Experimental/Iris/Core/Private/Iris/ReplicationSystem/Filtering/ReplicationFiltering.cpp b/Engine/Source/Runtime/Experimental/Iris/Core/Private/Iris/ReplicationSystem/Filtering/ReplicationFiltering.cpp
--- a/Engine/Source/Runtime/Experimental/Iris/Core/Private/Iris/ReplicationSystem/Filtering/ReplicationFiltering.cpp
+++ b/Engine/Source/Runtime/Experimental/Iris/Core/Private/Iris/ReplicationSystem/Filtering/ReplicationFiltering.cpp
@@ -161,7 +161,7 @@ private:
 	const FNetRefHandleManager* NetRefHandleManager = nullptr;
 };
 
-// Helper class instance to avoid logging redundant depewndent objects more than once per object class.
+// Helper class instance to avoid logging redundant dependent objects more than once per object class.
 #if !UE_BUILD_SHIPPING
 static FIrisLogOnceTracker ReplicationFiltering_MootDependentObjectTracker;
 #endif
@@ -1351,6 +1351,44 @@ void FReplicationFiltering::PreUpdateDynamicFiltering()
 	}
 }
 
+void FReplicationFiltering::UpdateParentFilterDependent(uint32 DependentIndex, const FNetBitArrayView DependentsWithParentFilter,  FNetBitArrayView OutConnectionObjectsInScope, bool bRecursive) const
+{
+	const bool bStartedRelevant = OutConnectionObjectsInScope.IsBitSet(DependentIndex);
+	if (!bStartedRelevant)
+	{
+		// No need to check it's parents, the object must be relevant by itself first.
+		return;
+	}
+
+	const TArrayView<const FInternalNetRefIndex> Parents = NetRefHandleManager->GetDependentObjectParents(DependentIndex);
+
+	for (FInternalNetRefIndex ParentIndex : Parents)
+	{
+		const bool bIsParentDependent = DependentsWithParentFilter.IsBitSet(ParentIndex);
+					
+		// If the parent is a dependent AND relevant check his status first
+		//TODO: Optimize this by flagging which parents were processed already or not
+		if (bIsParentDependent && OutConnectionObjectsInScope.IsBitSet(ParentIndex))
+		{
+			constexpr bool bIsRecursiveCall = true;
+			UpdateParentFilterDependent(ParentIndex, DependentsWithParentFilter, OutConnectionObjectsInScope, bIsRecursiveCall);
+		}
+
+		const bool bIsParentRelevant = OutConnectionObjectsInScope.IsBitSet(ParentIndex);
+		OutConnectionObjectsInScope.AndBitValue(DependentIndex, bIsParentRelevant);
+	}
+
+	// Check if we are still relevant, if not clear our subobjects
+	const bool bIsStillRelevant = OutConnectionObjectsInScope.IsBitSet(DependentIndex);
+	if (!bRecursive && !bIsStillRelevant)
+	{
+		for (const FInternalNetRefIndex SubObjectIndex : NetRefHandleManager->GetSubObjects(DependentIndex))
+		{
+			OutConnectionObjectsInScope.ClearBit(SubObjectIndex);
+		}
+	}
+};
+
 /**
  * Per connection we will combine the effects of connection/owner and group filtering before passing
  * the objects with a specific dynamic filter set. The filter can then decide which objects should
@@ -1384,8 +1422,10 @@ void FReplicationFiltering::UpdateDynamicFiltering()
 
 	const FNetBitArrayView DependentObjectsView = NetRefHandleManager->GetDependentObjectInternalIndices();
 	const uint32* DependentObjectsData = DependentObjectsView.GetDataChecked(WordCountForObjectBitArrays);
-	
 
+	const FNetBitArrayView DependentsWithParentFilter = NetRefHandleManager->GetDependentObjectsWithParentFilter();
+	const uint32* DependentsWithParentFilterData = DependentsWithParentFilter.GetDataChecked(WordCountForObjectBitArrays);
+	
 	uint32* ConnectionIds = static_cast<uint32*>(FMemory_Alloca(ValidConnections.GetNumBits() * sizeof(uint32)));
 	uint32 ConnectionCount = 0;
 	ValidConnections.ForAllSetBits([ConnectionIds, &ConnectionCount](uint32 Bit) { ConnectionIds[ConnectionCount++] = Bit; });
@@ -1454,7 +1494,8 @@ void FReplicationFiltering::UpdateDynamicFiltering()
 			for (uint32 WordIt = 0, WordEndIt = WordCountForObjectBitArrays; WordIt != WordEndIt; ++WordIt)
 			{
 				const uint32 SubObjects = SubObjectsData[WordIt];
-				const uint32 DependentObjects = DependentObjectsData[WordIt];
+				// Set of dependents objects that don't have the parent filter trait
+				const uint32 DependentObjects = (DependentObjectsData[WordIt] & ~DependentsWithParentFilterData[WordIt]);
 				const uint32 ObjectsRequiringUpdate = ObjectsRequiringDynamicFilterUpdateData[WordIt];
 				const uint32 PrevFilteredOutObjects = DynamicFilteredOutObjectsData[WordIt];
 				// Mask off group included objects in this loops so they're accounted for when updating hysteresis. A no longer group included object should also be subject to hysteresis.
@@ -1623,6 +1664,17 @@ void FReplicationFiltering::UpdateDynamicFiltering()
 			}
 		}
 
+		// Update dependent objects that can only be relevant if their parent also is
+		{
+			IRIS_PROFILER_SCOPE(FReplicationFiltering_UpdateDependentsWithParentFilter);
+			FNetBitArrayView ConnectionObjectsInScope = MakeNetBitArrayView(ConnectionInfo.ObjectsInScope);
+			FNetBitArrayView::ForAllSetBits(DependentsWithParentFilter, ConnectionObjectsInScope, FNetBitArrayView::AndOp, [&](uint32 DependentIndex)
+			{
+				constexpr bool bIsNotRecursive = false;
+				this->UpdateParentFilterDependent(DependentIndex, DependentsWithParentFilter, ConnectionObjectsInScope, bIsNotRecursive);
+			});
+		}
+
 		// The scope for the connection is now fully updated, apart from disabled dependent objects.
 		// If any object that has a dependency that isn't filtered out we must enable the dependent object.
 		if (HysteresisState.Mode == EHysteresisProcessingMode::Enabled)
diff --git a/Engine/Source/Runtime/Experimental/Iris/Core/Private/Iris/ReplicationSystem/Filtering/ReplicationFiltering.h b/Engine/Source/Runtime/Experimental/Iris/Core/Private/Iris/ReplicationSystem/Filtering/ReplicationFiltering.h
--- a/Engine/Source/Runtime/Experimental/Iris/Core/Private/Iris/ReplicationSystem/Filtering/ReplicationFiltering.h
+++ b/Engine/Source/Runtime/Experimental/Iris/Core/Private/Iris/ReplicationSystem/Filtering/ReplicationFiltering.h
@@ -241,6 +241,8 @@ private:
 
 	bool HasDynamicFilters() const;
 
+	void UpdateParentFilterDependent(uint32 DependentIndex, const FNetBitArrayView DependentsWithParentFilter, FNetBitArrayView OutConnectionObjectsInScope, bool bIsRecursive) const;
+
 	// Helper to update and reset group exclusion filter effects if objects are removed from a filter or after a filter status change, returns true if the group filter was changed
 	bool ClearGroupExclusionFilterEffectsForObject(uint32 ObjectIndex, uint32 ConnectionId);
 	// Helper to update and reset group inclusion filter effects if objects are removed from a filter or after a filter status change, returns true if the group filter was changed
diff --git a/Engine/Source/Runtime/Experimental/Iris/Core/Private/Iris/ReplicationSystem/NetDependencyData.h b/Engine/Source/Runtime/Experimental/Iris/Core/Private/Iris/ReplicationSystem/NetDependencyData.h
--- a/Engine/Source/Runtime/Experimental/Iris/Core/Private/Iris/ReplicationSystem/NetDependencyData.h
+++ b/Engine/Source/Runtime/Experimental/Iris/Core/Private/Iris/ReplicationSystem/NetDependencyData.h
@@ -24,10 +24,19 @@ struct FChildSubObjectsInfo
 	uint32 NumSubObjects = 0U;
 };
 
+enum class EDependentFilterTrait : uint8
+{
+	/** Allow the dependent object to be relevant even if the parent isn't*/
+	None,
+	/** The dependent object can only be relevant if the parent also is */
+	ParentMustBeRelevant,
+};
+
 struct FDependentObjectInfo
 {
 	FInternalNetRefIndex NetRefIndex = 0U;
 	EDependentObjectSchedulingHint SchedulingHint = EDependentObjectSchedulingHint::Default;
+	EDependentFilterTrait ParentFilterTrait = EDependentFilterTrait::None;
 };
 
 class FNetDependencyData
diff --git a/Engine/Source/Runtime/Experimental/Iris/Core/Private/Iris/ReplicationSystem/NetRefHandleManager.cpp b/Engine/Source/Runtime/Experimental/Iris/Core/Private/Iris/ReplicationSystem/NetRefHandleManager.cpp
--- a/Engine/Source/Runtime/Experimental/Iris/Core/Private/Iris/ReplicationSystem/NetRefHandleManager.cpp
+++ b/Engine/Source/Runtime/Experimental/Iris/Core/Private/Iris/ReplicationSystem/NetRefHandleManager.cpp
@@ -68,6 +68,7 @@ void FNetRefHandleManager::Init(const FInitParams& InitParams)
 		InitNetBitArray(&AssignedInternalIndices);
 		InitNetBitArray(&SubObjectInternalIndices);
 		InitNetBitArray(&DependentObjectInternalIndices);
+		InitNetBitArray(&DependentObjectsWithParentFilter);
 		InitNetBitArray(&ObjectsWithDependentObjectsInternalIndices);
 		InitNetBitArray(&DestroyedStartupObjectInternalIndices);
 		InitNetBitArray(&WantToBeDormantInternalIndices);
@@ -792,7 +793,7 @@ FInternalNetRefIndex FNetRefHandleManager::GetRootObjectInternalIndexOfSubObject
 	return SubObjectIndex != InvalidInternalIndex ? ReplicatedObjectData[SubObjectIndex].SubObjectRootIndex : InvalidInternalIndex;
 }
 
-bool FNetRefHandleManager::AddDependentObject(FNetRefHandle ParentRefHandle, FNetRefHandle DependentObjectRefHandle, EDependentObjectSchedulingHint SchedulingHint)
+bool FNetRefHandleManager::AddDependentObject(FNetRefHandle ParentRefHandle, FNetRefHandle DependentObjectRefHandle, EDependentObjectSchedulingHint SchedulingHint, EDependentFilterTrait FilterTrait)
 {
 	check(ParentRefHandle != DependentObjectRefHandle);
 
@@ -843,6 +844,7 @@ bool FNetRefHandleManager::AddDependentObject(FNetRefHandle ParentRefHandle, FNe
 	FDependentObjectInfo DependentObjectInfo;
 	DependentObjectInfo.NetRefIndex = DependentObjectInternalIndex;
 	DependentObjectInfo.SchedulingHint = SchedulingHint;
+	DependentObjectInfo.ParentFilterTrait = FilterTrait;
 	ParentDependentObjectsArray.Add(DependentObjectInfo);
 
 	// Add parent to dependent's list
@@ -853,6 +855,7 @@ bool FNetRefHandleManager::AddDependentObject(FNetRefHandle ParentRefHandle, FNe
 	ParentObjectData.bHasDependentObjects = true;
 	ObjectsWithDependentObjectsInternalIndices.SetBit(ParentInternalIndex);
 	DependentObjectInternalIndices.SetBit(DependentObjectInternalIndex);
+	DependentObjectsWithParentFilter.SetBitValue(DependentObjectInternalIndex, FilterTrait == EDependentFilterTrait::ParentMustBeRelevant);
 
 	return true;
 }
@@ -879,6 +882,7 @@ void FNetRefHandleManager::InternalRemoveDependentObject(FInternalNetRefIndex Pa
 				FReplicatedObjectData& DependentObjectData = GetReplicatedObjectDataNoCheck(DependentInternalIndex);
 				DependentObjectData.bIsDependentObject = false;
 				DependentObjectInternalIndices.ClearBit(DependentInternalIndex);
+				DependentObjectsWithParentFilter.ClearBit(DependentInternalIndex);
 			}
 		}
 	}
@@ -935,6 +939,7 @@ void FNetRefHandleManager::InternalRemoveDependentObject(FInternalNetRefIndex De
 	DependentObjectData.bHasDependentObjects = false;
 	ObjectsWithDependentObjectsInternalIndices.ClearBit(DependentInternalIndex);
 	DependentObjectInternalIndices.ClearBit(DependentInternalIndex);
+	DependentObjectsWithParentFilter.ClearBit(DependentInternalIndex);
 }
 
 void FNetRefHandleManager::RemoveDependentObject(FNetRefHandle ParentHandle, FNetRefHandle DependentHandle)
diff --git a/Engine/Source/Runtime/Experimental/Iris/Core/Private/Iris/ReplicationSystem/NetRefHandleManager.h b/Engine/Source/Runtime/Experimental/Iris/Core/Private/Iris/ReplicationSystem/NetRefHandleManager.h
--- a/Engine/Source/Runtime/Experimental/Iris/Core/Private/Iris/ReplicationSystem/NetRefHandleManager.h
+++ b/Engine/Source/Runtime/Experimental/Iris/Core/Private/Iris/ReplicationSystem/NetRefHandleManager.h
@@ -256,7 +256,12 @@ public:
 	// DependentObjects
 	const FNetBitArrayView GetDependentObjectInternalIndices() const { return MakeNetBitArrayView(DependentObjectInternalIndices); }
 	const FNetBitArrayView GetObjectsWithDependentObjectsInternalIndices() const { return MakeNetBitArrayView(ObjectsWithDependentObjectsInternalIndices); }
-	bool AddDependentObject(FNetRefHandle ParentHandle, FNetRefHandle DependentHandle, EDependentObjectSchedulingHint SchedulingHint);
+	const FNetBitArrayView GetDependentObjectsWithParentFilter() const
+	{
+		return MakeNetBitArrayView(DependentObjectsWithParentFilter);
+	}
+
+	bool AddDependentObject(FNetRefHandle ParentHandle, FNetRefHandle DependentHandle, EDependentObjectSchedulingHint SchedulingHint, EDependentFilterTrait FilterTrait = EDependentFilterTrait::None);
 	void RemoveDependentObject(FNetRefHandle ParentHandle, FNetRefHandle DependentHandle);
 
 	// Remove DependentHandles from all dependent object tracking
@@ -458,6 +463,9 @@ private:
 	// Bitset containing all internal indices that are dependentObjects
 	FNetBitArray DependentObjectInternalIndices;
 
+	// Dependent objects that need their parent to be relevant too
+	FNetBitArray DependentObjectsWithParentFilter;
+
 	// Bitset containing all internal indices that have dependentObjects
 	FNetBitArray ObjectsWithDependentObjectsInternalIndices;
 
diff --git a/Engine/Source/Runtime/Experimental/Iris/Core/Private/Iris/ReplicationSystem/ObjectReplicationBridge.cpp b/Engine/Source/Runtime/Experimental/Iris/Core/Private/Iris/ReplicationSystem/ObjectReplicationBridge.cpp
--- a/Engine/Source/Runtime/Experimental/Iris/Core/Private/Iris/ReplicationSystem/ObjectReplicationBridge.cpp
+++ b/Engine/Source/Runtime/Experimental/Iris/Core/Private/Iris/ReplicationSystem/ObjectReplicationBridge.cpp
@@ -725,6 +725,28 @@ void UObjectReplicationBridge::RemoveDependentObject(FNetRefHandle ParentHandle,
 	}
 }
 
+void UObjectReplicationBridge::AddDependentObjectWithParentRelevancy(FNetRefHandle Parent, FNetRefHandle DependentObject)
+{
+	using namespace UE::Net;
+	using namespace UE::Net::Private;
+
+	FReplicationSystemInternal* ReplicationSystemInternal = GetReplicationSystem()->GetReplicationSystemInternal();
+	FNetRefHandleManager& LocalNetRefHandleManager = ReplicationSystemInternal->GetNetRefHandleManager();
+
+	if (LocalNetRefHandleManager.AddDependentObject(Parent, DependentObject, EDependentObjectSchedulingHint::Default, EDependentFilterTrait::ParentMustBeRelevant))
+	{
+		FReplicationFiltering& Filtering = ReplicationSystemInternal->GetFiltering();
+		const FInternalNetRefIndex DependentInternalIndex = LocalNetRefHandleManager.GetInternalIndex(DependentObject);
+		Filtering.NotifyAddedDependentObject(DependentInternalIndex);
+
+		UE_LOG_OBJECTREPLICATIONBRIDGE(Verbose, TEXT("AddDependentObjectWithParentRelevancy Added dependent object %s to parent %s"), *PrintObjectFromNetRefHandle(DependentObject), *PrintObjectFromNetRefHandle(Parent));
+	}
+	else
+	{
+		UE_LOG_OBJECTREPLICATIONBRIDGE(Warning , TEXT("AddDependentObjectWithParentRelevancy Failed to add dependent object %s to parent %s"), *PrintObjectFromNetRefHandle(DependentObject), *PrintObjectFromNetRefHandle(Parent));
+	}
+}
+
 bool UObjectReplicationBridge::WriteNetRefHandleCreationInfo(FReplicationBridgeSerializationContext& Context, FNetRefHandle Handle)
 {
 	using namespace UE::Net;
diff --git a/Engine/Source/Runtime/Experimental/Iris/Core/Public/Iris/ReplicationSystem/ObjectReplicationBridge.h b/Engine/Source/Runtime/Experimental/Iris/Core/Public/Iris/ReplicationSystem/ObjectReplicationBridge.h
--- a/Engine/Source/Runtime/Experimental/Iris/Core/Public/Iris/ReplicationSystem/ObjectReplicationBridge.h
+++ b/Engine/Source/Runtime/Experimental/Iris/Core/Public/Iris/ReplicationSystem/ObjectReplicationBridge.h
@@ -147,6 +147,9 @@ public:
 	/** Remove dependent object from parent. The dependent object will function as a standard standalone replicated object. */
 	IRISCORE_API void RemoveDependentObject(FNetRefHandle Parent, FNetRefHandle DependentObject);
 
+	/** Make an object dependent but it can only be relevant if it's parent is also relevant to the connection.  */
+	IRISCORE_API void AddDependentObjectWithParentRelevancy(FNetRefHandle Parent, FNetRefHandle DependentObject);
+
 	// Dormancy support
 
 	/** Set whether object should go dormant. If dormancy is enabled any dirty state will be replicated first. */
