Fixed bug with on-demand install priorities being unsigned instead of signed integers.

#rnx

diff --git a/Engine/Source/Runtime/Experimental/IoStore/OnDemand/Private/OnDemandContentInstaller.cpp b/Engine/Source/Runtime/Experimental/IoStore/OnDemand/Private/OnDemandContentInstaller.cpp
--- a/Engine/Source/Runtime/Experimental/IoStore/OnDemand/Private/OnDemandContentInstaller.cpp
+++ b/Engine/Source/Runtime/Experimental/IoStore/OnDemand/Private/OnDemandContentInstaller.cpp
@@ -335,8 +335,8 @@ void FOnDemandContentInstaller::UpdateInstallRequestPriority(FSharedInternalInst
 			FRequest&			Request	= *reinterpret_cast<FRequest*>(InstallRequest->InstallerRequest);
 			FRequest::FInstall& Install	= Request.AsInstall();
 
-			UE_LOG(LogIoStoreOnDemand, Log, TEXT("Updating install request priority, ContentHandle=(%s), CurrentPriority=%d, NewPriority=%d"),
-				*LexToString(Install.Args.ContentHandle), Request.Priority, NewPriority);
+			UE_LOG(LogIoStoreOnDemand, Log, TEXT("Updating install request priority, SeqNo=%u, Priority=%d, NewPriority=%d, ContentHandle=(%s)"),
+				Request.SeqNo, Request.Priority, NewPriority, *LexToString(Install.Args.ContentHandle));
 
 			Request.Priority = NewPriority;
 
@@ -431,10 +431,11 @@ void FOnDemandContentInstaller::ProcessInstallRequest(FRequest& Request)
 
 	TRACE_CPUPROFILER_EVENT_SCOPE(FOnDemandContentInstaller::ProcessInstallRequest);
 
-	FRequest::FInstall& InstallRequest = Request.AsInstall();
+	FRequest::FInstall& InstallRequest	= Request.AsInstall();
+	Request.Priority					= InstallRequest.Args.Priority;
 
-	UE_LOG(LogIoStoreOnDemand, Log, TEXT("Processing install request, ContentHandle=(%s)"),
-		*LexToString(InstallRequest.Args.ContentHandle));
+	UE_LOG(LogIoStoreOnDemand, Log, TEXT("Processing install request, SeqNo=%u, Priority=%d, ContentHandle=(%s)"),
+		Request.SeqNo, Request.Priority, *LexToString(InstallRequest.Args.ContentHandle));
 
 	if (InstallRequest.Args.ContentHandle.IsValid() == false)
 	{
@@ -577,7 +578,6 @@ void FOnDemandContentInstaller::ProcessInstallRequest(FRequest& Request)
 	InstallRequest.Progress.TotalContentSize	= TotalContentSize;
 	InstallRequest.Progress.TotalInstallSize	= TotalInstallSize;
 	InstallRequest.Progress.CurrentInstallSize	= 0;
-	Request.Priority							= InstallRequest.Args.Priority;
 
 	if (InstallRequest.HttpRequestHandles.IsEmpty())
 	{
@@ -621,6 +621,9 @@ void FOnDemandContentInstaller::ExecuteInstallRequest(FRequest& Request, bool bR
 	FRequest::FInstall& InstallRequest = Request.AsInstall();
 	check(InstallRequest.HttpRequestHandles.IsEmpty() == false);
 
+	UE_LOG(LogIoStoreOnDemand, Log, TEXT("Executing install request, SeqNo=%u, Priority=%d, ContentHandle=(%s)"),
+		Request.SeqNo, Request.Priority, *LexToString(InstallRequest.Args.ContentHandle));
+
 	if (Request.ErrorCode.load(std::memory_order_relaxed) == EIoErrorCode::Cancelled)
 	{
 		return CompleteInstallRequest(Request);
@@ -805,9 +808,9 @@ void FOnDemandContentInstaller::ProcessDownloadedChunk(
 
 	InstallRequest.Progress.CurrentInstallSize	+= ChunkEntry.EncodedSize;
 
-	UE_LOG(LogIoStoreOnDemand, Verbose, TEXT("Install progress %.2lf/%.2lf MiB, ContentHandle=(%s), ChunkId='%s', ChunkSize=%.2lf KiB, HttpStatus=%u"),
+	UE_LOG(LogIoStoreOnDemand, Verbose, TEXT("Install progress %.2lf/%.2lf MiB, SeqNo=%u, Priority=%d, ContentHandle=(%s), ChunkId='%s', ChunkSize=%.2lf KiB, HttpStatus=%u"),
 		double(InstallRequest.Progress.CurrentInstallSize) / 1024.0 / 1024.0, double(InstallRequest.Progress.TotalInstallSize) / 1024.0 / 1024.0,
-		*LexToString(InstallRequest.Args.ContentHandle), *LexToString(ChunkId), double(ChunkEntry.EncodedSize) / 1024.0, HttpStatusCode);
+		Request.SeqNo, Request.Priority, *LexToString(InstallRequest.Args.ContentHandle), *LexToString(ChunkId), double(ChunkEntry.EncodedSize) / 1024.0, HttpStatusCode);
 
 	if (Request.ErrorCode == EIoErrorCode::Unknown)
 	{
@@ -959,8 +962,8 @@ void FOnDemandContentInstaller::CompleteInstallRequest(FRequest& Request)
 		DurationCycles,
 		Request.ErrorCode);
 
-	UE_LOG(LogIoStoreOnDemand, Log, TEXT("Install request completed, Result='%s', ContentHandle=(%s), ContentSize=%.2lf MiB, InstallSize=%.2lf MiB, CacheHitRatio=%d %%, Duration=%d ms"),
-		GetIoErrorText(InstallResult.Status.GetErrorCode()), *LexToString(InstallRequest.Args.ContentHandle), double(InstallResult.Progress.TotalContentSize) / 1024.0 / 1024.0,
+	UE_LOG(LogIoStoreOnDemand, Log, TEXT("Install request completed, Result='%s', SeqNo=%u, Priority=%d, ContentHandle=(%s), ContentSize=%.2lf MiB, InstallSize=%.2lf MiB, CacheHitRatio=%d%%, Duration=%dms"),
+		GetIoErrorText(InstallResult.Status.GetErrorCode()), Request.SeqNo, Request.Priority, *LexToString(InstallRequest.Args.ContentHandle), double(InstallResult.Progress.TotalContentSize) / 1024.0 / 1024.0,
 		double(InstallResult.Progress.TotalInstallSize) / 1024.0 / 1024.0, int32(CacheHitRatio * 100), int32(InstallResult.DurationInSeconds * 1000));
 
 	const bool bCallbackOnGameThread			= EnumHasAnyFlags(InstallRequest.Args.Options, EOnDemandInstallOptions::CallbackOnGameThread);
diff --git a/Engine/Source/Runtime/Experimental/IoStore/OnDemand/Private/OnDemandContentInstaller.h b/Engine/Source/Runtime/Experimental/IoStore/OnDemand/Private/OnDemandContentInstaller.h
--- a/Engine/Source/Runtime/Experimental/IoStore/OnDemand/Private/OnDemandContentInstaller.h
+++ b/Engine/Source/Runtime/Experimental/IoStore/OnDemand/Private/OnDemandContentInstaller.h
@@ -152,7 +152,7 @@ class FOnDemandContentInstaller
 		static uint32	NextSeqNo;
 
 		uint32						SeqNo = NextSeqNo++;
-		uint32						Priority = 0;
+		int32						Priority = 0;
 		uint64						StartTimeCycles = FPlatformTime::Cycles64();
 		FString						ErrorReason;
 		std::atomic<EIoErrorCode>	ErrorCode{EIoErrorCode::Unknown};
