Fix potential assert if an IAX request fails.

#rnx

- It is only valid to call FTicketStatus::GetPerf if it's Id is EId::Content as we won't have proper values if the request failed or was cancelled.

diff --git a/Engine/Source/Runtime/Experimental/IoStore/OnDemand/Private/OnDemandHttpClient.cpp b/Engine/Source/Runtime/Experimental/IoStore/OnDemand/Private/OnDemandHttpClient.cpp
--- a/Engine/Source/Runtime/Experimental/IoStore/OnDemand/Private/OnDemandHttpClient.cpp
+++ b/Engine/Source/Runtime/Experimental/IoStore/OnDemand/Private/OnDemandHttpClient.cpp
@@ -305,7 +305,6 @@ void FMultiEndpointHttpClient::IssueRequest(FRequest&& Request)
 					FMultiEndpointHttpClientResponse Response
 					{
 						.Body					= MoveTemp(Request.Body),
-						.Sample					= TicketStatus.GetPerf().GetSample(),
 						.DurationMilliseconds	= DurationMs,
 						.StatusCode				= Request.StatusCode,
 						.RetryCount				= Request.RetryCount,
@@ -313,6 +312,11 @@ void FMultiEndpointHttpClient::IssueRequest(FRequest&& Request)
 						.CDNCacheStatus			= Request.CDNCacheStatus
 					};
 
+					if (TicketStatus.GetId() == FTicketStatus::EId::Content)
+					{
+						Response.Sample = TicketStatus.GetPerf().GetSample();
+					}
+
 					if (bError)
 					{
 						Response.Reason = TicketStatus.GetError().Reason;
