Iris

* Moved PollFrequency getter from EngineBridge to NetObjectFactory

#rnx

diff --git a/Engine/Source/Runtime/Engine/Private/Net/Iris/ReplicationSystem/EngineReplicationBridge.cpp b/Engine/Source/Runtime/Engine/Private/Net/Iris/ReplicationSystem/EngineReplicationBridge.cpp
--- a/Engine/Source/Runtime/Engine/Private/Net/Iris/ReplicationSystem/EngineReplicationBridge.cpp
+++ b/Engine/Source/Runtime/Engine/Private/Net/Iris/ReplicationSystem/EngineReplicationBridge.cpp
@@ -712,18 +712,6 @@ bool UEngineReplicationBridge::CanCreateDestructionInfo() const
 	return false;
 }
 
-float UEngineReplicationBridge::GetPollFrequencyOfRootObject(const UObject* ReplicatedObject) const
-{
-	if (const AActor* ReplicatedActor = Cast<AActor>(ReplicatedObject))
-	{
-		float PollFrequency = ReplicatedActor->GetNetUpdateFrequency();
-		GetClassPollFrequency(ReplicatedActor->GetClass(), PollFrequency);
-		return PollFrequency;
-	}
-	//TODO: Ask NetObjectFactory for this value instead
-	return 100.0f;
-}
-
 bool UEngineReplicationBridge::IsClassReplicatedByDefault(const UClass* Class) const
 {
 	if (Class && Class->IsChildOf(AActor::StaticClass()))
diff --git a/Engine/Source/Runtime/Engine/Private/Net/Iris/ReplicationSystem/NetActorFactory.cpp b/Engine/Source/Runtime/Engine/Private/Net/Iris/ReplicationSystem/NetActorFactory.cpp
--- a/Engine/Source/Runtime/Engine/Private/Net/Iris/ReplicationSystem/NetActorFactory.cpp
+++ b/Engine/Source/Runtime/Engine/Private/Net/Iris/ReplicationSystem/NetActorFactory.cpp
@@ -558,6 +558,20 @@ void UNetActorFactory::GetWorldInfo(const FWorldInfoContext& Context, FWorldInfo
 	}
 }
 
+float UNetActorFactory::GetPollFrequency(UE::Net::FNetRefHandle RootObjectHandle, UObject* RootObjectInstance)
+{
+	using namespace UE::Net;
+
+	const AActor* Actor = Cast<AActor>(RootObjectInstance);
+	if (!Actor)
+	{
+		// Return the CDO default value
+		return AActor::StaticClass()->GetDefaultObject<AActor>()->GetNetUpdateFrequency();
+	}
+
+	return Actor->GetNetUpdateFrequency();
+}
+
  namespace UE::Net
  {
 
diff --git a/Engine/Source/Runtime/Engine/Private/Net/Iris/ReplicationSystem/NetSubObjectFactory.cpp b/Engine/Source/Runtime/Engine/Private/Net/Iris/ReplicationSystem/NetSubObjectFactory.cpp
--- a/Engine/Source/Runtime/Engine/Private/Net/Iris/ReplicationSystem/NetSubObjectFactory.cpp
+++ b/Engine/Source/Runtime/Engine/Private/Net/Iris/ReplicationSystem/NetSubObjectFactory.cpp
@@ -244,6 +244,12 @@ void UNetSubObjectFactory::GetWorldInfo(const FWorldInfoContext& Context, FWorld
 	ensureMsgf(false, TEXT("UNetSubObjectFactory::GetWorldInfo called but subobjects should never support this. Instance: %s, NetRefHandle: %s"), *GetNameSafe(Context.Instance), *Bridge->PrintObjectFromNetRefHandle(Context.Handle));
 }
 
+float UNetSubObjectFactory::GetPollFrequency(UE::Net::FNetRefHandle RootObjectHandle, UObject* RootObjectInstance)
+{
+	ensureMsgf(false, TEXT("UNetSubObjectFactory::GetPollFrequency called but subobjects should never support this. Instance: %s, NetRefHandle: %s"), *GetNameSafe(RootObjectInstance), *Bridge->PrintObjectFromNetRefHandle(RootObjectHandle));
+	return 100.0f;
+}
+
 //------------------------------------------------------------------------
 
 namespace UE::Net
diff --git a/Engine/Source/Runtime/Engine/Public/Net/Iris/ReplicationSystem/EngineReplicationBridge.h b/Engine/Source/Runtime/Engine/Public/Net/Iris/ReplicationSystem/EngineReplicationBridge.h
--- a/Engine/Source/Runtime/Engine/Public/Net/Iris/ReplicationSystem/EngineReplicationBridge.h
+++ b/Engine/Source/Runtime/Engine/Public/Net/Iris/ReplicationSystem/EngineReplicationBridge.h
@@ -141,7 +141,6 @@ protected:
 	virtual void OnProtocolMismatchDetected(FNetRefHandle ObjectHandle) override;
 	virtual void OnProtocolMismatchReported(FNetRefHandle RefHandle, uint32 ConnectionId) override;
 	virtual bool CanCreateDestructionInfo() const override;
-	virtual float GetPollFrequencyOfRootObject(const UObject* ReplicatedObject) const override;
 
 	/** Returns true if the class is derived from Actor and its CDO has set bReplicates. */
 	virtual bool IsClassReplicatedByDefault(const UClass* Class) const;
diff --git a/Engine/Source/Runtime/Engine/Public/Net/Iris/ReplicationSystem/NetActorFactory.h b/Engine/Source/Runtime/Engine/Public/Net/Iris/ReplicationSystem/NetActorFactory.h
--- a/Engine/Source/Runtime/Engine/Public/Net/Iris/ReplicationSystem/NetActorFactory.h
+++ b/Engine/Source/Runtime/Engine/Public/Net/Iris/ReplicationSystem/NetActorFactory.h
@@ -130,6 +130,8 @@ public:
 
 	virtual void GetWorldInfo(const FWorldInfoContext& Context, FWorldInfoData& OutData) override;
 
+	virtual float GetPollFrequency(UE::Net::FNetRefHandle RootObjectHandle, UObject* RootObjectInstance) override;
+
 protected:
 
 	virtual TUniquePtr<UE::Net::FNetObjectCreationHeader> CreateAndFillHeader(UE::Net::FNetRefHandle Handle) override;
diff --git a/Engine/Source/Runtime/Engine/Public/Net/Iris/ReplicationSystem/NetSubObjectFactory.h b/Engine/Source/Runtime/Engine/Public/Net/Iris/ReplicationSystem/NetSubObjectFactory.h
--- a/Engine/Source/Runtime/Engine/Public/Net/Iris/ReplicationSystem/NetSubObjectFactory.h
+++ b/Engine/Source/Runtime/Engine/Public/Net/Iris/ReplicationSystem/NetSubObjectFactory.h
@@ -35,6 +35,8 @@ public:
 
 	virtual void GetWorldInfo(const FWorldInfoContext& Context, FWorldInfoData& OutData) override;
 
+	virtual float GetPollFrequency(UE::Net::FNetRefHandle RootObjectHandle, UObject* RootObjectInstance) override;
+
 protected:
 
 	virtual TUniquePtr<UE::Net::FNetObjectCreationHeader> CreateAndFillHeader(UE::Net::FNetRefHandle Handle) override;
diff --git a/Engine/Source/Runtime/Experimental/Iris/Core/Private/Iris/ReplicationSystem/ObjectReplicationBridge.cpp b/Engine/Source/Runtime/Experimental/Iris/Core/Private/Iris/ReplicationSystem/ObjectReplicationBridge.cpp
--- a/Engine/Source/Runtime/Experimental/Iris/Core/Private/Iris/ReplicationSystem/ObjectReplicationBridge.cpp
+++ b/Engine/Source/Runtime/Experimental/Iris/Core/Private/Iris/ReplicationSystem/ObjectReplicationBridge.cpp
@@ -2393,10 +2393,23 @@ uint8 UObjectReplicationBridge::ConvertPollFrequencyIntoFrames(float PollFrequen
 	return static_cast<uint8>(FMath::Clamp<uint32>(FramesBetweenUpdatesForObject, 0U, UE::Net::Private::FObjectPollFrequencyLimiter::GetMaxPollingFrames()));
 }
 
-float UObjectReplicationBridge::GetPollFrequencyOfRootObject(const UObject* ReplicatedObject) const
+float UObjectReplicationBridge::GetPollFrequencyOfRootObject(UE::Net::Private::FInternalNetRefIndex RootObjectIndex) const
 {
-	float PollFrequency = 0.0f;
-	GetClassPollFrequency(ReplicatedObject->GetClass(), PollFrequency);
+	using namespace UE::Net;
+	using namespace UE::Net::Private;
+
+	UObject* RootObjectInstance = NetRefHandleManager->GetReplicatedObjectInstance(RootObjectIndex);
+	FNetRefHandle RootObjectHandle = NetRefHandleManager->GetNetRefHandleFromInternalIndex(RootObjectIndex);
+	
+	FNetObjectFactoryId FactoryId = NetRefHandleManager->GetReplicatedObjectDataNoCheck(RootObjectIndex).NetFactoryId;
+	UNetObjectFactory* Factory = GetNetFactory(FactoryId);
+
+	float PollFrequency = Factory->GetPollFrequency(RootObjectHandle, RootObjectInstance);
+
+	if (RootObjectInstance)
+	{
+		GetClassPollFrequency(RootObjectInstance->GetClass(), PollFrequency);
+	}
 
 	return PollFrequency;
 }
@@ -2411,25 +2424,23 @@ void UObjectReplicationBridge::ReinitPollFrequency()
 
 	auto UpdatePollFrequency = [this, &LocalNetRefHandleManager](uint32 RootObjectIndex)
 	{
-		if (UObject* RootObjectInstance = LocalNetRefHandleManager.GetReplicatedObjectInstance(RootObjectIndex))
-		{
-			const float PollFrequency = GetPollFrequencyOfRootObject(RootObjectInstance);
-			const uint8 PollFramePeriod = ConvertPollFrequencyIntoFrames(PollFrequency);
+		const float PollFrequency = GetPollFrequencyOfRootObject(RootObjectIndex);
+		const uint8 PollFramePeriod = ConvertPollFrequencyIntoFrames(PollFrequency);
 
-			PollFrequencyLimiter->SetPollFramePeriod(RootObjectIndex, PollFramePeriod);
+		PollFrequencyLimiter->SetPollFramePeriod(RootObjectIndex, PollFramePeriod);
 
-			// Make sure the subobjects are polled the same frame as the root object.
-			for (const FInternalNetRefIndex SubObjectIndex : LocalNetRefHandleManager.GetSubObjects(RootObjectIndex))
-			{
-				PollFrequencyLimiter->SetPollWithObject(RootObjectIndex, SubObjectIndex);
-			}
+		// Make sure the subobjects are polled the same frame as the root object.
+		for (const FInternalNetRefIndex SubObjectIndex : LocalNetRefHandleManager.GetSubObjects(RootObjectIndex))
+		{
+			PollFrequencyLimiter->SetPollWithObject(RootObjectIndex, SubObjectIndex);
 		}
 	};
 
-	const FNetBitArrayView RootObjects = LocalNetRefHandleManager.GetGlobalScopableInternalIndices();
+	// Update frequency on all root objects
+	const FNetBitArrayView ReplicatedObjects = LocalNetRefHandleManager.GetGlobalScopableInternalIndices();
 	const FNetBitArrayView SubObjects = MakeNetBitArrayView(LocalNetRefHandleManager.GetSubObjectInternalIndices());
 
-	FNetBitArrayView::ForAllSetBits(RootObjects, SubObjects, FNetBitArrayView::AndNotOp, UpdatePollFrequency);
+	FNetBitArrayView::ForAllSetBits(ReplicatedObjects, SubObjects, FNetBitArrayView::AndNotOp, UpdatePollFrequency);
 }
 
 void UObjectReplicationBridge::SetPollFrequency(FNetRefHandle RefHandle, float PollFrequency)
diff --git a/Engine/Source/Runtime/Experimental/Iris/Core/Public/Iris/ReplicationSystem/NetObjectFactory.h b/Engine/Source/Runtime/Experimental/Iris/Core/Public/Iris/ReplicationSystem/NetObjectFactory.h
--- a/Engine/Source/Runtime/Experimental/Iris/Core/Public/Iris/ReplicationSystem/NetObjectFactory.h
+++ b/Engine/Source/Runtime/Experimental/Iris/Core/Public/Iris/ReplicationSystem/NetObjectFactory.h
@@ -170,6 +170,9 @@ public:
 	 */
 	virtual void GetWorldInfo(const FWorldInfoContext& Context, FWorldInfoData& OutData) PURE_VIRTUAL(UNetObjectFactory::GetWorldInfo,);
 
+	/** Return the poll frequency of a root object managed by this factory. */
+	virtual float GetPollFrequency(UE::Net::FNetRefHandle RootObjectHandle, UObject* RootObjectInstance) PURE_VIRTUAL(UNetObjectFactory::GetPollFrequency, return 100.f;);
+
 public:
 
 	/** Result of the instantiate request */
diff --git a/Engine/Source/Runtime/Experimental/Iris/Core/Public/Iris/ReplicationSystem/ObjectReplicationBridge.h b/Engine/Source/Runtime/Experimental/Iris/Core/Public/Iris/ReplicationSystem/ObjectReplicationBridge.h
--- a/Engine/Source/Runtime/Experimental/Iris/Core/Public/Iris/ReplicationSystem/ObjectReplicationBridge.h
+++ b/Engine/Source/Runtime/Experimental/Iris/Core/Public/Iris/ReplicationSystem/ObjectReplicationBridge.h
@@ -331,9 +331,6 @@ protected:
 	/** Transform a poll frequency (in updates per second) to an equivalent number of frames. */
 	IRISCORE_API uint8 ConvertPollFrequencyIntoFrames(float PollFrequency) const;
 
-	/** Returns the poll frequency of a specific root object. */
-	IRISCORE_API virtual float GetPollFrequencyOfRootObject(const UObject* ReplicatedObject) const;
-
 	/** Re-initialize the poll frequency of all replicated root objects. */
 	IRISCORE_API void ReinitPollFrequency();
 
@@ -413,6 +410,9 @@ private:
 	/** Tries to load the classes used in poll period overrides. */
 	void FindClassesInPollPeriodOverrides();
 
+	/** Returns the poll frequency of a specific root object. */
+	float GetPollFrequencyOfRootObject(UE::Net::Private::FInternalNetRefIndex RootObjectIndex) const;
+
 	/** Does all that is necessary to replicate any given UObject. Returns the object's handle if successful  */
 	UE::Net::FNetRefHandle StartReplicatingNetObject(UObject* Instance, UE::Net::EReplicationFragmentTraits Traits, UE::Net::FNetObjectFactoryId NetFactoryId);
 
