Add statistic to FOnDemandInstallCacheStats to track oldest block age.

diff --git a/Engine/Source/Runtime/Experimental/IoStore/OnDemand/Private/OnDemandInstallCache.cpp b/Engine/Source/Runtime/Experimental/IoStore/OnDemand/Private/OnDemandInstallCache.cpp
--- a/Engine/Source/Runtime/Experimental/IoStore/OnDemand/Private/OnDemandInstallCache.cpp
+++ b/Engine/Source/Runtime/Experimental/IoStore/OnDemand/Private/OnDemandInstallCache.cpp
@@ -1906,6 +1906,7 @@ FIoStatus FOnDemandInstallCache::Purge(TSet<FIoHash>&& ChunksToInstall)
 	uint64 ReferencedBytes = 0;
 	uint64 FragmentedBytes = 0;
 	uint64 TotalReferencedBlockBytes = 0;
+	int64 OldestBlockAccess = FDateTime::MaxValue().GetTicks();
 
 	const uint64 TotalUncachedBytes = AddReferencesToBlocks(Containers, ChunkEntryIndices, ChunksToInstall, BlockInfo, ReferencedBytes);
 	for (const TPair<FCasBlockId, FCasBlockInfo>& Kv : BlockInfo)
@@ -1919,9 +1920,14 @@ FIoStatus FOnDemandInstallCache::Purge(TSet<FIoHash>&& ChunksToInstall)
 		{
 			TotalReferencedBlockBytes += Info.FileSize;
 		}
+		if (Info.LastAccess < OldestBlockAccess)
+		{
+			OldestBlockAccess = Info.LastAccess;
+		}
 	}
 
-	FOnDemandInstallCacheStats::OnCacheUsage(MaxCacheSize, TotalCachedBytes, TotalReferencedBlockBytes, ReferencedBytes, FragmentedBytes);
+	FOnDemandInstallCacheStats::OnCacheUsage(
+		MaxCacheSize, TotalCachedBytes, TotalReferencedBlockBytes, ReferencedBytes, FragmentedBytes, OldestBlockAccess);
 
 	const uint64 TotalRequiredBytes = TotalCachedBytes + TotalUncachedBytes;
 	if (TotalRequiredBytes <= MaxCacheSize)
diff --git a/Engine/Source/Runtime/Experimental/IoStore/OnDemand/Private/Statistics.cpp b/Engine/Source/Runtime/Experimental/IoStore/OnDemand/Private/Statistics.cpp
--- a/Engine/Source/Runtime/Experimental/IoStore/OnDemand/Private/Statistics.cpp
+++ b/Engine/Source/Runtime/Experimental/IoStore/OnDemand/Private/Statistics.cpp
@@ -469,6 +469,7 @@ struct FInstallCacheAnalytics
 	uint64 MaxReferencedBlockSize = 0;
 	uint64 MaxReferencedSize = 0;
 	uint64 MaxFragmentedSize = 0;
+	int64 OldestBlockAccess = FDateTime::MaxValue().GetTicks();
 };
 FMutex FInstallCacheAnalytics::Mutex;
 static FInstallCacheAnalytics InstallCacheAnalytics;
@@ -1177,6 +1178,10 @@ void FOnDemandContentInstallerStats::ReportAnalytics(TArray<FAnalyticsEventAttri
 		InstallCacheAnalytics			= FInstallCacheAnalytics();
 	}
 
+	const FDateTime Now = FDateTime::UtcNow();
+	const FDateTime OldestBlockAccess(CurrentInstallCacheAnalytics.OldestBlockAccess);
+	const FTimespan OldestBlockAge = (Now >= OldestBlockAccess) ? (Now - OldestBlockAccess) : FTimespan::MaxValue();
+
 	AppendAnalyticsEventAttributeArray(OutAnalyticsArray
 		, TEXT("IadTotalInstallCount"), CurrentInstallerAnalytics.InstallCount
 		, TEXT("IadTotalInstallErrorCount"), CurrentInstallerAnalytics.InstallErrorCount 
@@ -1200,6 +1205,7 @@ void FOnDemandContentInstallerStats::ReportAnalytics(TArray<FAnalyticsEventAttri
 		, TEXT("IadInstallInstallCacheMaxReferencedBlockSize"), CurrentInstallCacheAnalytics.MaxReferencedBlockSize
 		, TEXT("IadInstallInstallCacheMaxReferencedSize"), CurrentInstallCacheAnalytics.MaxReferencedSize
 		, TEXT("IadInstallInstallCacheMaxFragmentedSize"), CurrentInstallCacheAnalytics.MaxFragmentedSize
+		, TEXT("IadInstallInstallCacheOldestBlockAgeMinutes"), OldestBlockAge.GetTotalMinutes()
 	);
 }
 
@@ -1269,7 +1275,8 @@ void FOnDemandInstallCacheStats::OnCacheUsage(
 	uint64 CacheSize,
 	uint64 ReferencedBlockSize,
 	uint64 ReferencedSize,
-	uint64 FragmentedSize)
+	uint64 FragmentedSize,
+	int64 OldestBlockAccess)
 {
 	UE::TUniqueLock Lock(InstallCacheAnalytics);
 	InstallCacheAnalytics.MaxCacheSize				= MaxCacheSize; 
@@ -1277,6 +1284,7 @@ void FOnDemandInstallCacheStats::OnCacheUsage(
 	InstallCacheAnalytics.MaxReferencedBlockSize	= FMath::Max(InstallCacheAnalytics.MaxReferencedBlockSize, ReferencedBlockSize);
 	InstallCacheAnalytics.MaxReferencedSize			= FMath::Max(InstallCacheAnalytics.MaxReferencedSize, ReferencedSize);
 	InstallCacheAnalytics.MaxFragmentedSize			= FMath::Max(InstallCacheAnalytics.MaxFragmentedSize, FragmentedSize);
+	InstallCacheAnalytics.OldestBlockAccess			= FMath::Min(InstallCacheAnalytics.OldestBlockAccess, OldestBlockAccess);
 }
 
 } // namespace UE::IoStore
diff --git a/Engine/Source/Runtime/Experimental/IoStore/OnDemand/Private/Statistics.h b/Engine/Source/Runtime/Experimental/IoStore/OnDemand/Private/Statistics.h
--- a/Engine/Source/Runtime/Experimental/IoStore/OnDemand/Private/Statistics.h
+++ b/Engine/Source/Runtime/Experimental/IoStore/OnDemand/Private/Statistics.h
@@ -123,7 +123,8 @@ public:
 		uint64 CacheSize,
 		uint64 ReferencedBlockSize,
 		uint64 ReferencedSize,
-		uint64 FragmentedSize) IAS_STATISTICS_IMPL()
+		uint64 FragmentedSize,
+		int64 OldestBlockAccess) IAS_STATISTICS_IMPL()
 };
 
 #undef IAS_STATISTICS_IMPL
