Iris
* Quick fix for crash in ReinitPollFrequency when invalid root objects are part of the scopeable objec tlist

#rnx

diff --git a/Engine/Source/Runtime/Experimental/Iris/Core/Private/Iris/ReplicationSystem/ObjectReplicationBridge.cpp b/Engine/Source/Runtime/Experimental/Iris/Core/Private/Iris/ReplicationSystem/ObjectReplicationBridge.cpp
--- a/Engine/Source/Runtime/Experimental/Iris/Core/Private/Iris/ReplicationSystem/ObjectReplicationBridge.cpp
+++ b/Engine/Source/Runtime/Experimental/Iris/Core/Private/Iris/ReplicationSystem/ObjectReplicationBridge.cpp
@@ -2402,6 +2402,11 @@ float UObjectReplicationBridge::GetPollFrequencyOfRootObject(UE::Net::Private::F
 	FNetRefHandle RootObjectHandle = NetRefHandleManager->GetNetRefHandleFromInternalIndex(RootObjectIndex);
 	
 	FNetObjectFactoryId FactoryId = NetRefHandleManager->GetReplicatedObjectDataNoCheck(RootObjectIndex).NetFactoryId;
+	if (FactoryId == UE::Net::InvalidNetObjectFactoryId)
+	{
+		// TODO: Investigate how scopeable objects are Invalid
+		return 100.f;
+	}
 	UNetObjectFactory* Factory = GetNetFactory(FactoryId);
 
 	float PollFrequency = Factory->GetPollFrequency(RootObjectHandle, RootObjectInstance);
