Geometry collection : fix issue preventing sleeping geometry collection from being affected by velocities
#okfrogithub public

diff --git a/Engine/Source/Runtime/Experimental/Chaos/Private/PhysicsProxy/GeometryCollectionPhysicsProxy.cpp b/Engine/Source/Runtime/Experimental/Chaos/Private/PhysicsProxy/GeometryCollectionPhysicsProxy.cpp
--- a/Engine/Source/Runtime/Experimental/Chaos/Private/PhysicsProxy/GeometryCollectionPhysicsProxy.cpp
+++ b/Engine/Source/Runtime/Experimental/Chaos/Private/PhysicsProxy/GeometryCollectionPhysicsProxy.cpp
@@ -2784,9 +2784,13 @@ void FGeometryCollectionPhysicsProxy::ApplyBreakingLinearVelocity_External(FGeom
 			if (Chaos::FPBDRigidClusteredParticleHandle* ClusteredHandle = FindClusteredParticleHandleByItemIndex_Internal(ItemIndex))
 			{
 				ApplyToBreakingChildren_Internal(RBDSolver->GetEvolution()->GetRigidClustering(), *ClusteredHandle,
-					[LinearVelocity](Chaos::FPBDRigidClusteredParticleHandle* ClusteredChildHandle)
+					[LinearVelocity, RBDSolver](Chaos::FPBDRigidClusteredParticleHandle* ClusteredChildHandle)
 					{
 						ClusteredChildHandle->SetV(ClusteredChildHandle->GetV() + LinearVelocity);
+						if (!ClusteredChildHandle->GetV().IsNearlyZero() && !ClusteredChildHandle->Disabled())
+						{
+							RBDSolver->GetEvolution()->WakeParticle(ClusteredChildHandle);
+						}
 					});
 			}
 		});
@@ -2804,9 +2808,13 @@ void FGeometryCollectionPhysicsProxy::ApplyBreakingAngularVelocity_External(FGeo
 			if (Chaos::FPBDRigidClusteredParticleHandle* ClusteredHandle = FindClusteredParticleHandleByItemIndex_Internal(ItemIndex))
 			{
 				ApplyToBreakingChildren_Internal(RBDSolver->GetEvolution()->GetRigidClustering(), *ClusteredHandle,
-					[AngularVelocity](Chaos::FPBDRigidClusteredParticleHandle* ClusteredChildHandle)
+					[AngularVelocity, RBDSolver](Chaos::FPBDRigidClusteredParticleHandle* ClusteredChildHandle)
 					{
 						ClusteredChildHandle->SetW(ClusteredChildHandle->GetW() + AngularVelocity);
+						if (!ClusteredChildHandle->GetW().IsNearlyZero() && !ClusteredChildHandle->Disabled())
+						{
+							RBDSolver->GetEvolution()->WakeParticle(ClusteredChildHandle);
+						}
 					});
 			}
 		});
@@ -2819,11 +2827,15 @@ void FGeometryCollectionPhysicsProxy::ApplyLinearVelocity_External(FGeometryColl
 
 	if (Chaos::FPhysicsSolver* RBDSolver = GetSolver<Chaos::FPhysicsSolver>())
 	{
-		RBDSolver->EnqueueCommandImmediate([this, ItemIndex, LinearVelocity]()
+		RBDSolver->EnqueueCommandImmediate([this, ItemIndex, LinearVelocity, RBDSolver]()
 		{
 			if (Chaos::FPBDRigidClusteredParticleHandle* ClusteredHandle = FindClusteredParticleHandleByItemIndex_Internal(ItemIndex))
 			{
 				ClusteredHandle->SetV(ClusteredHandle->GetV() + LinearVelocity);
+				if (!ClusteredHandle->GetV().IsNearlyZero() && !ClusteredHandle->Disabled())
+				{
+					RBDSolver->GetEvolution()->WakeParticle(ClusteredHandle);
+				}
 			}
 		});
 	}
@@ -2835,11 +2847,15 @@ void FGeometryCollectionPhysicsProxy::ApplyAngularVelocity_External(FGeometryCol
 
 	if (Chaos::FPhysicsSolver* RBDSolver = GetSolver<Chaos::FPhysicsSolver>())
 	{
-		RBDSolver->EnqueueCommandImmediate([this, ItemIndex, AngularVelocity]()
+		RBDSolver->EnqueueCommandImmediate([this, ItemIndex, AngularVelocity, RBDSolver]()
 		{
 			if (Chaos::FPBDRigidClusteredParticleHandle* ClusteredHandle = FindClusteredParticleHandleByItemIndex_Internal(ItemIndex))
 			{
 				ClusteredHandle->SetW(ClusteredHandle->GetW() + AngularVelocity);
+				if (!ClusteredHandle->GetW().IsNearlyZero() && !ClusteredHandle->Disabled())
+				{
+					RBDSolver->GetEvolution()->WakeParticle(ClusteredHandle);
+				}
 			}
 		});
 	}
