Fixed bug when defragging the on-demand install cache.

- updated the initial journal replay to set the block location and not just use FindOrAdd since defrag will add multiple locations for the same chunk in the journal
- updated defrag to use the correct disk size when reading encoded chunks

#rnx

diff --git a/Engine/Source/Runtime/Experimental/IoStore/OnDemand/Private/OnDemandInstallCache.cpp b/Engine/Source/Runtime/Experimental/IoStore/OnDemand/Private/OnDemandInstallCache.cpp
--- a/Engine/Source/Runtime/Experimental/IoStore/OnDemand/Private/OnDemandInstallCache.cpp
+++ b/Engine/Source/Runtime/Experimental/IoStore/OnDemand/Private/OnDemandInstallCache.cpp
@@ -1455,7 +1455,8 @@ FOnDemandInstallCache::FOnDemandInstallCache(const FOnDemandInstallCacheConfig&
 			const FCasJournal::FEntry::FChunkLocation& ChunkLocation = JournalEntry.ChunkLocation;
 			if (ChunkLocation.CasLocation.IsValid())
 			{
-				Cas.Lookup.FindOrAdd(ChunkLocation.CasAddr, ChunkLocation.CasLocation);
+				FCasLocation& Loc = Cas.Lookup.FindOrAdd(ChunkLocation.CasAddr);
+				Loc = ChunkLocation.CasLocation;
 			}
 			else
 			{
@@ -2547,7 +2548,8 @@ FIoStatus FOnDemandInstallCache::Defrag(
 			{
 				FileHandle->Seek(ReffedChunk.BlockOffset);
 
-				FIoBuffer Buffer(ReffedChunk.EncodedSize);
+				const int64 ChunkDiskSize = Align(int64(ReffedChunk.EncodedSize), FAES::AESBlockSize);
+				FIoBuffer Buffer(ChunkDiskSize);
 				const bool bOk = FileHandle->Read(Buffer.GetData(), Buffer.GetSize());
 				if (!bOk)
 				{
@@ -2555,6 +2557,22 @@ FIoStatus FOnDemandInstallCache::Defrag(
 					return EIoErrorCode::ReadError;
 				}
 
+				const FIoHash ChunkHash = FIoHash::HashBuffer(Buffer.GetView());
+				if (ChunkHash != ReffedChunk.Hash)
+				{
+					UE_LOG(LogIoStoreOnDemand, Error, TEXT("Found chunk with invalid hash while defragging block, BlockId=%u, BlockOffset=%u"),
+						DefragBlock.BlockId.Id, ReffedChunk.BlockOffset);
+
+					if (FIoStatus Status = FlushPendingChunks(DefragPendingChunks); Status.IsOk() == false)
+					{
+						FOnDemandInstallCacheStats::OnDefrag(Status.GetErrorCode(), 0);
+						return Status;
+					}
+	
+					FOnDemandInstallCacheStats::OnDefrag(EIoErrorCode::SignatureError, 0);
+					return EIoErrorCode::SignatureError;
+				}
+
 				if (DefragPendingChunks.TotalSize > FPendingChunks::MaxPendingBytes)
 				{
 					if (FIoStatus Status = FlushPendingChunks(DefragPendingChunks, DefragBlock.LastAccess); Status.IsOk() == false)
