Minor fix of IAD analytics event names.

#rnx

diff --git a/Engine/Source/Runtime/Experimental/IoStore/OnDemand/Private/Statistics.cpp b/Engine/Source/Runtime/Experimental/IoStore/OnDemand/Private/Statistics.cpp
--- a/Engine/Source/Runtime/Experimental/IoStore/OnDemand/Private/Statistics.cpp
+++ b/Engine/Source/Runtime/Experimental/IoStore/OnDemand/Private/Statistics.cpp
@@ -1198,14 +1198,14 @@ void FOnDemandContentInstallerStats::ReportAnalytics(TArray<FAnalyticsEventAttri
 		, TEXT("IadInstallCachePurgeErrorCount"), CurrentInstallCacheAnalytics.PurgeErrorCount
 		, TEXT("IadInstallCacheDefragCount"), CurrentInstallCacheAnalytics.DefragCount 
 		, TEXT("IadInstallCacheDefragErrorCount"), CurrentInstallCacheAnalytics.DefragErrorCount 
-		, TEXT("IadInstallInstallCacheJournalCommitCount"), CurrentInstallCacheAnalytics.JournalCommitCount 
-		, TEXT("IadInstallInstallCacheJournalCommitErrorCount"), CurrentInstallCacheAnalytics.JournalCommitErrorCount
-		, TEXT("IadInstallInstallCacheMaxSize"), CurrentInstallCacheAnalytics.MaxCacheSize 
-		, TEXT("IadInstallInstallCacheMaxUsageSize"), CurrentInstallCacheAnalytics.MaxCacheUsageSize 
-		, TEXT("IadInstallInstallCacheMaxReferencedBlockSize"), CurrentInstallCacheAnalytics.MaxReferencedBlockSize
-		, TEXT("IadInstallInstallCacheMaxReferencedSize"), CurrentInstallCacheAnalytics.MaxReferencedSize
-		, TEXT("IadInstallInstallCacheMaxFragmentedSize"), CurrentInstallCacheAnalytics.MaxFragmentedSize
-		, TEXT("IadInstallInstallCacheOldestBlockAgeMinutes"), OldestBlockAge.GetTotalMinutes()
+		, TEXT("IadInstallCacheJournalCommitCount"), CurrentInstallCacheAnalytics.JournalCommitCount 
+		, TEXT("IadInstallCacheJournalCommitErrorCount"), CurrentInstallCacheAnalytics.JournalCommitErrorCount
+		, TEXT("IadInstallCacheMaxSize"), CurrentInstallCacheAnalytics.MaxCacheSize 
+		, TEXT("IadInstallCacheMaxUsageSize"), CurrentInstallCacheAnalytics.MaxCacheUsageSize 
+		, TEXT("IadInstallCacheMaxReferencedBlockSize"), CurrentInstallCacheAnalytics.MaxReferencedBlockSize
+		, TEXT("IadInstallCacheMaxReferencedSize"), CurrentInstallCacheAnalytics.MaxReferencedSize
+		, TEXT("IadInstallCacheMaxFragmentedSize"), CurrentInstallCacheAnalytics.MaxFragmentedSize
+		, TEXT("IadInstallCacheOldestBlockAgeMinutes"), OldestBlockAge.GetTotalMinutes()
 	);
 }
 
