Fixed bug when creating CAS snapshots and more robust replacing of journal file

#rnx

diff --git a/Engine/Source/Runtime/Experimental/IoStore/OnDemand/Private/OnDemandInstallCache.cpp b/Engine/Source/Runtime/Experimental/IoStore/OnDemand/Private/OnDemandInstallCache.cpp
--- a/Engine/Source/Runtime/Experimental/IoStore/OnDemand/Private/OnDemandInstallCache.cpp
+++ b/Engine/Source/Runtime/Experimental/IoStore/OnDemand/Private/OnDemandInstallCache.cpp
@@ -1018,8 +1018,9 @@ struct FCasSnapshot
 	using FChunkLocation = TPair<FCasAddr, FCasLocation>;
 
 	static TIoStatusOr<FCasSnapshot>	FromJournal(const FString& JournalFile);
-	static TIoStatusOr<FCasSnapshot>	Load(const FString& SnapShotFile, int64* OutFileSize = nullptr);
+	static TIoStatusOr<FCasSnapshot>	Load(const FString& SnapshotFile, int64* OutFileSize = nullptr);
 	static TIoStatusOr<int64>			Save(const FCasSnapshot& Snapshot, const FString& SnapshotFile);
+	static TIoStatusOr<int64>			TryCreateAndResetJournal(const FString& SnapshotFile, const FString& JournalFile); 
 
 	TArray<FBlock>						Blocks;
 	TArray<FChunkLocation>				ChunkLocations;
@@ -1065,7 +1066,8 @@ TIoStatusOr<FCasSnapshot> FCasSnapshot::FromJournal(const FString& JournalFile)
 				const FCasJournal::FEntry::FChunkLocation& ChunkLocation = JournalEntry.ChunkLocation;
 				if (ChunkLocation.CasLocation.IsValid())
 				{
-					CasLookup.FindOrAdd(ChunkLocation.CasAddr, ChunkLocation.CasLocation);
+					FCasLocation& Loc = CasLookup.FindOrAdd(ChunkLocation.CasAddr);
+					Loc = ChunkLocation.CasLocation;
 				}
 				else
 				{
@@ -1224,6 +1226,60 @@ TIoStatusOr<FCasSnapshot> FCasSnapshot::Load(const FString& SnapshotFile, int64*
 	return Snapshot;
 }
 
+TIoStatusOr<int64> FCasSnapshot::TryCreateAndResetJournal(const FString& SnapshotFile, const FString& JournalFile)
+{
+	IFileManager& Ifm = IFileManager::Get();
+
+	const int64 JournalFileSize = Ifm.FileSize(*JournalFile);
+	if (JournalFileSize < 0)
+	{
+		return FIoStatus(EIoErrorCode::NotFound);
+	}
+
+	// Load the snapshot from the journal
+	TIoStatusOr<FCasSnapshot> SnapshotStatus = FCasSnapshot::FromJournal(JournalFile);
+	if (SnapshotStatus.IsOk() == false)
+	{
+		return SnapshotStatus.Status();
+	}
+
+	// Save the snapshot
+	int64 SnapshotSize		= -1;
+	FCasSnapshot Snapshot	= SnapshotStatus.ConsumeValueOrDie();
+	if (TIoStatusOr<int64> Status = FCasSnapshot::Save(Snapshot, SnapshotFile); Status.IsOk())
+	{
+		SnapshotSize = Status.ConsumeValueOrDie();
+	}
+	else
+	{
+		return Status.Status();
+	}
+
+	// Try create a new empty journal 
+	const FString TmpJournalFile = FPaths::ChangeExtension(JournalFile, TEXT(".jrntmp"));
+	if (FIoStatus Status = FCasJournal::Create(TmpJournalFile); Status.IsOk() == false)
+	{
+		if (Ifm.Delete(*SnapshotFile) == false)
+		{
+			UE_LOG(LogIoStoreOnDemand, Error, TEXT("Failed to delete CAS snapshot '%s'"), *SnapshotFile); 
+		}
+
+		return Status;
+	}
+
+	if (Ifm.Move(*JournalFile , *TmpJournalFile) == false)
+	{
+		if (Ifm.Delete(*SnapshotFile) == false)
+		{
+			UE_LOG(LogIoStoreOnDemand, Error, TEXT("Failed to delete CAS snapshot '%s'"), *SnapshotFile); 
+		}
+
+		return FIoStatus(EIoErrorCode::WriteError);
+	}
+
+	return SnapshotSize;
+}
+
 ///////////////////////////////////////////////////////////////////////////////
 class FOnDemandInstallCache final
 	: public IOnDemandInstallCache 
@@ -1591,41 +1647,25 @@ void FOnDemandInstallCache::Shutdown()
 
 	if (FIoStatus Status = FCasJournal::Commit(MoveTemp(Transaction)); !Status.IsOk())
 	{
-		UE_LOG(LogIoStoreOnDemand, Error, TEXT("Failed to commit journal entries to '%s'"), *JournalFile);
-		//TODO: Delete journal and start over? 
+		UE_LOG(LogIoStoreOnDemand, Error, TEXT("Failed to update CAS journal '%s' with block timestamp(s), reason '%s'"),
+			*JournalFile, *Status.ToString());
 	}
 
-	IPlatformFile& Ipf = FPlatformFileManager::Get().GetPlatformFile();
-	if (int64 JournalSize = Ipf.FileSize(*JournalFile); JournalSize > int64(MaxJournalSize))
+	IFileManager& Ifm = IFileManager::Get();
+	const FString JournalFilename = GetJournalFilename();
+	if (Ifm.FileSize(*JournalFile) > int64(MaxJournalSize))
 	{
-		TIoStatusOr<FCasSnapshot> SnapshotStatus = FCasSnapshot::FromJournal(JournalFile);
-		if (SnapshotStatus.IsOk() == false)
-		{
-			UE_LOG(LogIoStoreOnDemand, Error, TEXT("Failed to create CAS snapshot from journal '%s', reason '%s'"),
-				*JournalFile, *SnapshotStatus.Status().ToString());
-			return;
-		}
-
-		FCasSnapshot Snapshot = SnapshotStatus.ConsumeValueOrDie();
-		const FString SnapshotFile = GetSnapshotFilename();
-		TIoStatusOr<int64> SaveStatus = FCasSnapshot::Save(Snapshot, SnapshotFile);
-		if (SaveStatus.IsOk() == false)
+		const FString SnapshotFilename = GetSnapshotFilename();
+		TIoStatusOr<int64> SnapshotStatus = FCasSnapshot::TryCreateAndResetJournal(SnapshotFilename, JournalFilename);
+		if (SnapshotStatus.IsOk())
 		{
-			UE_LOG(LogIoStoreOnDemand, Error, TEXT("Failed to save CAS snapshot '%s', reason '%s'"),
-				*SnapshotFile, *SaveStatus.Status().ToString());
-			return;
+			UE_LOG(LogIoStoreOnDemand, Log, TEXT("Saved CAS snapshot '%s' %.2lf KiB"), *SnapshotFilename, ToKiB(SnapshotStatus.ValueOrDie()));
 		}
-
-		if (Ipf.DeleteFile(*JournalFile) == false)
+		else
 		{
-			UE_LOG(LogIoStoreOnDemand, Error, TEXT("Failed to delete CAS journal '%s'"), *JournalFile);
-			if (Ipf.DeleteFile(*SnapshotFile) == false)
-			{
-				UE_LOG(LogIoStoreOnDemand, Error, TEXT("Failed to delete CAS snapshot '%s'"), *JournalFile);
-			}
-			return;
+			UE_LOG(LogIoStoreOnDemand, Error, TEXT("Failed to create CAS snapshot from journal '%s', reason '%s'"),
+				*JournalFile, *SnapshotStatus.Status().ToString());
 		}
-		UE_LOG(LogIoStoreOnDemand, Log, TEXT("Saved CAS snapshot '%s' %.2lf KiB"), *SnapshotFile, ToKiB(JournalSize));
 	}
 }
 
