Startup hardening of the on-demand install cache.

#rnx

diff --git a/Engine/Source/Runtime/Experimental/IoStore/OnDemand/Private/OnDemandInstallCache.cpp b/Engine/Source/Runtime/Experimental/IoStore/OnDemand/Private/OnDemandInstallCache.cpp
--- a/Engine/Source/Runtime/Experimental/IoStore/OnDemand/Private/OnDemandInstallCache.cpp
+++ b/Engine/Source/Runtime/Experimental/IoStore/OnDemand/Private/OnDemandInstallCache.cpp
@@ -162,7 +162,7 @@ struct FCas
 	using FLastAccess			= TMap<FCasBlockId, int64>;
 	using FBlockIdHandleCounts	= TMap<FCasBlockId, int32>;
 
-	FIoStatus					Initialize(FStringView Directory);
+	FIoStatus					Initialize(FStringView Directory, bool bDeleteExisting = false);
 	FCasLocation				FindChunk(const FIoHash& Hash) const;
 	FCasBlockId					CreateBlock();
 	FIoStatus					DeleteBlock(FCasBlockId BlockId, TArray<FCasAddr>& OutAddrs);
@@ -189,7 +189,7 @@ struct FCas
 };
 
 ///////////////////////////////////////////////////////////////////////////////
-FIoStatus FCas::Initialize(FStringView Directory)
+FIoStatus FCas::Initialize(FStringView Directory, bool bDeleteExisting)
 {
 	RootDirectory = Directory;
 
@@ -201,6 +201,21 @@ FIoStatus FCas::Initialize(FStringView Directory)
 	FPathViews::Append(Path, RootDirectory, TEXT("blocks"));
 
 	IFileManager& Ifm = IFileManager::Get();
+
+	if (bDeleteExisting)
+	{
+		bool bRequireExists = false;
+		const bool bTree	= true;
+
+		if (Ifm.DeleteDirectory(Path.ToString(), bRequireExists, bTree) == false)
+		{
+			return FIoStatusBuilder(EIoErrorCode::WriteError)
+				<< TEXT("Failed to delete CAS blocks directory '")
+				<< Path.ToString()
+				<< TEXT("'");
+		}
+	}
+
 	if (Ifm.DirectoryExists(Path.ToString()) == false)
 	{
 		const bool bTree = true;
@@ -1412,6 +1427,8 @@ public:
 	virtual FOnDemandInstallCacheUsage			GetCacheUsage() override;
 
 private:
+	FIoStatus									Reinitialize();
+	FIoStatus									InitialVerify();
 	uint64										AddReferencesToBlocks(
 													const TArray<FSharedOnDemandContainer>& Containers, 
 													const TArray<TBitArray<>>& ChunkEntryIndices,
@@ -1471,10 +1488,11 @@ FOnDemandInstallCache::FOnDemandInstallCache(const FOnDemandInstallCacheConfig&
 	FIoStatus Status = Cas.Initialize(CacheDirectory);
 	if (Status.IsOk() == false)
 	{
-		UE_LOG(LogIoStoreOnDemand, Error, TEXT("Failed to initialize install cache"));
+		UE_LOG(LogIoStoreOnDemand, Error, TEXT("Failed to initialize install cache, reason '%s'"), *Status.ToString());
 		return;
 	}
 
+	// Try read the journal snapshot
 	{
 		const FString SnapshotFile = GetSnapshotFilename();
 		int64 SnapshotSize = -1;
@@ -1501,6 +1519,7 @@ FOnDemandInstallCache::FOnDemandInstallCache(const FOnDemandInstallCacheConfig&
 		}
 	}
 
+	// Replay the journal 
 	const FString JournalFile = GetJournalFilename();
 	Status = FCasJournal::Replay(JournalFile, [this](const FCasJournal::FEntry& JournalEntry)
 	{
@@ -1544,80 +1563,43 @@ FOnDemandInstallCache::FOnDemandInstallCache(const FOnDemandInstallCacheConfig&
 		};
 	});
 
-	// Verify the current state with the cached content on disk 
-	// TODO: Add checksums etc
-	TArray<FCasAddr> RemovedChunks;
-	if (FIoStatus Verify = Cas.Verify(RemovedChunks); !Verify.IsOk())
+	// Initializing the cache for the first time
+	if (Status.GetErrorCode() == EIoErrorCode::NotFound)
 	{
-		FOnDemandInstallCacheStats::OnCasVerificationError(RemovedChunks.Num());
+		if (Status = FCasJournal::Create(JournalFile); Status.IsOk())
+		{
+			UE_LOG(LogIoStoreOnDemand, Log, TEXT("Created CAS journal '%s'"), *JournalFile);
 
-		// Try to recover if the CAS blocks on disk doesn't match
-		FCasJournal::FTransaction Transaction = FCasJournal::Begin(GetJournalFilename());
-		for (const FCasAddr& Addr : RemovedChunks)
+			// Make sure that there are no existing blocks when starting from an empty cache
+			const bool bDeleteExisting = true;
+			Status = Cas.Initialize(CacheDirectory, bDeleteExisting);
+		}
+		else
 		{
-			Transaction.ChunkLocation(FCasLocation::Invalid, Addr);
+			UE_LOG(LogIoStoreOnDemand, Error, TEXT("Failed to create CAS journal '%s'"), *JournalFile);
 		}
-		Status = FCasJournal::Commit(MoveTemp(Transaction));
 	}
 
-	// Check that Cas is in-budget and shrink if needed
+	// Verify the current state of the cache
 	if (Status.IsOk())
 	{
-		FCasBlockInfoMap	BlockInfo;
-		uint64				TotalSize = Cas.GetBlockInfo(BlockInfo);
-		if (TotalSize > MaxCacheSize)
-		{
-			const uint64 TotalBytesToPurge = TotalSize - MaxCacheSize;
-
-			UE_LOG(LogIoStoreOnDemand, Warning, TEXT("Cache size is greater than disk quota - Purging install cache, MaxCacheSize=%.2lf MiB, TotalSize=%.2lf MiB, TotalBytesToPurge=%.2lf MiB"),
-				ToMiB(MaxCacheSize), ToMiB(TotalSize), ToMiB(TotalBytesToPurge));
-			
-			uint64 TotalPurgedBytes = 0;
-			FIoStatus PurgeStatus = Purge(BlockInfo, TotalBytesToPurge, TotalPurgedBytes);
-
-			if (PurgeStatus.IsOk() && TotalPurgedBytes >= TotalBytesToPurge)
-			{
-				UE_LOG(LogIoStoreOnDemand, Warning, TEXT("Successfully purged %.2lf MiB from install cache"), ToMiB(TotalPurgedBytes));
-			}
-			else
-			{
-				UE_LOG(LogIoStoreOnDemand, Error, TEXT("Failed to purge %.2lf MiB from install cache. Actually purged %.2lf MiB from install cache"),
-					ToMiB(TotalBytesToPurge), ToMiB(TotalPurgedBytes));
-
-				Status = FIoStatusBuilder(EIoErrorCode::WriteError) << FString::Printf(TEXT("Failed to purge overbudget cache (%s)"), *PurgeStatus.ToString());
-			}
-		}
+		Status = InitialVerify();
 	}
 
-	Cas.Compact();
-
+	// Try to reinitialize the cache if something has gone wrong 
 	if (Status.IsOk() == false)
 	{
-		if (Status.GetErrorCode() != EIoErrorCode::NotFound)
-		{
-			FOnDemandInstallCacheStats::OnStartupError(Status.GetErrorCode());
-
-			UE_LOG(LogIoStoreOnDemand, Warning, TEXT("Failed to replay install cache journal file '%s, reason '%s'"),
-				*Status.ToString(), *JournalFile);
-
-			UE_LOG(LogIoStoreOnDemand, Log, TEXT("Deleting installed content and reinitializing cache"));
-			IFileManager::Get().DeleteDirectory(*CacheDirectory);
-			if (Status = Cas.Initialize(CacheDirectory); Status.IsOk() == false)
-			{
-				UE_LOG(LogIoStoreOnDemand, Error, TEXT("Failed to initialize install cache, reason '%s'"),
-					*Status.ToString());
-				return;
-			}
-		}
+		Status = Reinitialize();
+	}
 
-		if (Status = FCasJournal::Create(JournalFile); Status.IsOk())
-		{
-			UE_LOG(LogIoStoreOnDemand, Log, TEXT("Created CAS journal '%s'"), *JournalFile);
-		}
-		else
-		{
-			UE_LOG(LogIoStoreOnDemand, Error, TEXT("Failed to create CAS journal '%s'"), *JournalFile);
-		}
+	if (Status.IsOk())
+	{
+		Cas.Compact();
+	}
+	else
+	{
+		UE_LOG(LogIoStoreOnDemand, Error, TEXT("Failed to initialize install cache, reason '%s'"),
+			*Status.ToString());
 	}
 }
 
@@ -2309,6 +2291,98 @@ FIoStatus FOnDemandInstallCache::Verify()
 	return FIoStatus::Ok;
 }
 
+FIoStatus FOnDemandInstallCache::Reinitialize()
+{
+	UE_LOG(LogIoStoreOnDemand, Log, TEXT("Recreating install cache in directory '%s'"), *CacheDirectory);
+
+	IFileManager& Ifm	= IFileManager::Get();
+	const bool bTree	= true;
+
+	if (Ifm.DeleteDirectory(*CacheDirectory, false, bTree) == false)
+	{
+		return FIoStatusBuilder(EIoErrorCode::WriteError)
+			<< TEXT("Failed to delete directory '")
+			<< CacheDirectory
+			<< TEXT("'");
+	}
+
+	if (Ifm.MakeDirectory(*CacheDirectory, bTree) == false)
+	{
+		return FIoStatusBuilder(EIoErrorCode::WriteError)
+			<< TEXT("Failed to create directory '")
+			<< CacheDirectory
+			<< TEXT("'");
+	}
+
+	if (FIoStatus Status = Cas.Initialize(CacheDirectory); Status.IsOk() == false)
+	{
+		return Status;
+	}
+
+	const FString JournalFile = GetJournalFilename();
+	if (FIoStatus Status = FCasJournal::Create(JournalFile); Status.IsOk() == false)
+	{
+		return Status;
+	}
+
+	UE_LOG(LogIoStoreOnDemand, Log, TEXT("Created CAS journal '%s'"), *JournalFile);
+	return FIoStatus::Ok;
+}
+
+FIoStatus FOnDemandInstallCache::InitialVerify()
+{
+	// Check if the cache is over budget
+	{
+		FCasBlockInfoMap	BlockInfo;
+		uint64				CacheSize = Cas.GetBlockInfo(BlockInfo);
+
+		if (CacheSize > MaxCacheSize)
+		{
+			const uint64	TotalBytesToPurge = CacheSize - MaxCacheSize;
+			uint64			TotalPurgedBytes = 0;
+
+			UE_LOG(LogIoStoreOnDemand, Warning, TEXT("Cache size is greater than disk quota - Purging install cache, MaxCacheSize=%.2lf MiB, TotalSize=%.2lf MiB, TotalBytesToPurge=%.2lf MiB"),
+				ToMiB(MaxCacheSize), ToMiB(CacheSize), ToMiB(TotalBytesToPurge));
+
+			FIoStatus PurgeStatus = Purge(BlockInfo, TotalBytesToPurge, TotalPurgedBytes);
+			if (PurgeStatus.IsOk() && TotalPurgedBytes >= TotalBytesToPurge)
+			{
+				UE_LOG(LogIoStoreOnDemand, Warning, TEXT("Successfully purged %.2lf MiB from install cache"), ToMiB(TotalPurgedBytes));
+			}
+			else
+			{
+				UE_LOG(LogIoStoreOnDemand, Error, TEXT("Failed to purge %.2lf MiB from install cache. Actually purged %.2lf MiB from install cache"),
+					ToMiB(TotalBytesToPurge), ToMiB(TotalPurgedBytes));
+
+				return FIoStatusBuilder(EIoErrorCode::WriteError) << FString::Printf(TEXT("Failed to purge overbudget cache (%s)"), *PurgeStatus.ToString());
+			}
+		}
+	}
+
+	// Verify the blocks on disk with the current state of the CAS
+	{
+		TArray<FCasAddr> RemovedChunks;
+		if (FIoStatus Verify = Cas.Verify(RemovedChunks); Verify.IsOk() == false)
+		{
+			FOnDemandInstallCacheStats::OnCasVerificationError(RemovedChunks.Num());
+
+			// Remove all entries that doesn't have a valid cache block 
+			FCasJournal::FTransaction Transaction = FCasJournal::Begin(GetJournalFilename());
+			for (const FCasAddr& Addr : RemovedChunks)
+			{
+				Transaction.ChunkLocation(FCasLocation::Invalid, Addr);
+			}
+
+			if (FIoStatus Status = FCasJournal::Commit(MoveTemp(Transaction)); Status.IsOk() == false)
+			{
+				return Status;
+			}
+		}
+	}
+
+	return FIoStatus::Ok;
+}
+
 uint64 FOnDemandInstallCache::AddReferencesToBlocks(
 	const TArray<FSharedOnDemandContainer>& Containers,
 	const TArray<TBitArray<>>& ChunkEntryIndices,
