Iris
* Removed LogIrisBridge and use LogIris category instead
* Merged XXXReplicationBridge(ID) headers to all use RepBrige(ID) instead

#rnx

diff --git a/Engine/Source/Runtime/Engine/Private/Net/Iris/ReplicationSystem/EngineReplicationBridge.cpp b/Engine/Source/Runtime/Engine/Private/Net/Iris/ReplicationSystem/EngineReplicationBridge.cpp
--- a/Engine/Source/Runtime/Engine/Private/Net/Iris/ReplicationSystem/EngineReplicationBridge.cpp
+++ b/Engine/Source/Runtime/Engine/Private/Net/Iris/ReplicationSystem/EngineReplicationBridge.cpp
@@ -62,8 +62,6 @@
 
 #include <limits>
 
-#define UE_LOG_ACTORREPLICATIONBRIDGE(Category, Format, ...)	UE_LOG(LogIrisBridge, Category, TEXT("ActorReplicationBridge(%u)::") Format, GetReplicationSystem()->GetId(), ##__VA_ARGS__)
-
 extern bool GDefaultUseSubObjectReplicationList;
 
 namespace UE::Net::Private
@@ -203,56 +201,56 @@ UE::Net::FNetRefHandle UEngineReplicationBridge::StartReplicatingActor(AActor* A
 	const bool bIsNetActor = ULevel::IsNetActor(Actor);
 	if (!bIsNetActor)
 	{
-		UE_LOG_ACTORREPLICATIONBRIDGE(VeryVerbose, TEXT("Actor %s doesn't have a NetRole."), ToCStr(GetFullNameSafe(Actor)));
+		UE_LOG_BRIDGEID(LogIris, VeryVerbose, TEXT("Actor %s doesn't have a NetRole."), ToCStr(GetFullNameSafe(Actor)));
 		return FNetRefHandle::GetInvalid();
 	}
 
 	if (Actor->GetLocalRole() != ROLE_Authority)
 	{
-		UE_LOG_ACTORREPLICATIONBRIDGE(VeryVerbose, TEXT("Actor %s NetRole isn't Authority."), ToCStr(GetFullNameSafe(Actor)));
+		UE_LOG_BRIDGEID(LogIris, VeryVerbose, TEXT("Actor %s NetRole isn't Authority."), ToCStr(GetFullNameSafe(Actor)));
 		return FNetRefHandle::GetInvalid();
 	}
 
 	if (Actor->IsActorBeingDestroyed() || !IsValid(Actor) || Actor->IsUnreachable())
 	{
-		UE_LOG_ACTORREPLICATIONBRIDGE(Verbose, TEXT("Actor %s is being destroyed or unreachable and can't be replicated."), ToCStr(GetFullNameSafe(Actor)));
+		UE_LOG_BRIDGEID(LogIris, Verbose, TEXT("Actor %s is being destroyed or unreachable and can't be replicated."), ToCStr(GetFullNameSafe(Actor)));
 		return FNetRefHandle::GetInvalid();
 	}
 
 	if (!Actor->GetIsReplicated())
 	{
-		UE_LOG_ACTORREPLICATIONBRIDGE(Verbose, TEXT("Actor %s is not supposed to be replicated."), ToCStr(GetFullNameSafe(Actor)));
+		UE_LOG_BRIDGEID(LogIris, Verbose, TEXT("Actor %s is not supposed to be replicated."), ToCStr(GetFullNameSafe(Actor)));
 		return FNetRefHandle::GetInvalid();
 	}
 
 	if (Actor->GetTearOff())
 	{
-		UE_LOG_ACTORREPLICATIONBRIDGE(Verbose, TEXT("Actor %s is torn off and should not be replicated."), ToCStr(GetFullNameSafe(Actor)));
+		UE_LOG_BRIDGEID(LogIris, Verbose, TEXT("Actor %s is torn off and should not be replicated."), ToCStr(GetFullNameSafe(Actor)));
 		return FNetRefHandle::GetInvalid();
 	}
 
 	if (!Actor->IsActorInitialized())
 	{
-		UE_LOG_ACTORREPLICATIONBRIDGE(Warning, TEXT("Actor %s is not initialized and won't be replicated."), ToCStr(GetFullNameSafe(Actor)));
+		UE_LOG_BRIDGEID(LogIris, Warning, TEXT("Actor %s is not initialized and won't be replicated."), ToCStr(GetFullNameSafe(Actor)));
 		return FNetRefHandle::GetInvalid();
 	}
 
 	if (!ensureMsgf(Actor->HasActorBegunPlay() || Actor->IsActorBeginningPlay(), TEXT("Actor %s hasn't begun or isn't beginning play in BeginReplication."), ToCStr(GetFullNameSafe(Actor))))
 	{
-		UE_LOG_ACTORREPLICATIONBRIDGE(Error, TEXT("Actor %s hasn't begun or isn't beginning play in BeginReplication."), ToCStr(GetFullNameSafe(Actor)));
+		UE_LOG_BRIDGEID(LogIris, Error, TEXT("Actor %s hasn't begun or isn't beginning play in BeginReplication."), ToCStr(GetFullNameSafe(Actor)));
 		return FNetRefHandle::GetInvalid();
 	}
 
 	if (!NetDriver)
 	{
-		UE_LOG_ACTORREPLICATIONBRIDGE(VeryVerbose, TEXT("There's no NetDriver so nothing can be replicated."));
+		UE_LOG_BRIDGEID(LogIris, VeryVerbose, TEXT("There's no NetDriver so nothing can be replicated."));
 		return FNetRefHandle::GetInvalid();
 	}
 
 
 	if (!NetDriver->ShouldReplicateActor(Actor))
 	{
-		UE_LOG_ACTORREPLICATIONBRIDGE(VeryVerbose, TEXT("Actor %s doesn't want to replicate with NetDriver %s."), ToCStr(GetFullNameSafe(Actor)), ToCStr(NetDriver->GetName()));
+		UE_LOG_BRIDGEID(LogIris, VeryVerbose, TEXT("Actor %s doesn't want to replicate with NetDriver %s."), ToCStr(GetFullNameSafe(Actor)), ToCStr(NetDriver->GetName()));
 		return FNetRefHandle::GetInvalid();
 	}
 
@@ -273,7 +271,7 @@ UE::Net::FNetRefHandle UEngineReplicationBridge::StartReplicatingActor(AActor* A
 	{
 		// Ensure the first time to get attention!
 		ensureMsgf(false, TEXT("Actor %s does not replicate subobjects using the registered SubObjectsLists, SubObjects will not replicate properly"), ToCStr(GetFullNameSafe(Actor)));
-		UE_LOG_ACTORREPLICATIONBRIDGE(Warning, TEXT("Actor %s does not replicate subobjects using the registered SubObjectsLists, SubObjects will not replicate properly"), ToCStr(GetFullNameSafe(Actor)));
+		UE_LOG_BRIDGEID(LogIris, Warning, TEXT("Actor %s does not replicate subobjects using the registered SubObjectsLists, SubObjects will not replicate properly"), ToCStr(GetFullNameSafe(Actor)));
 	}
 
 	// Create handles for the registered fragments
@@ -322,7 +320,7 @@ UE::Net::FNetRefHandle UEngineReplicationBridge::StartReplicatingActor(AActor* A
 		return FNetRefHandle::GetInvalid();
 	}
 
-    UE_CLOG(Actor->bAlwaysRelevant, LogIrisBridge, Verbose, TEXT("StartReplicatingActor of AlwaysRelevant actor %s"), *PrintObjectFromNetRefHandle(ActorRefHandle));
+    UE_CLOG(Actor->bAlwaysRelevant, LogIris, Verbose, TEXT("StartReplicatingActor of AlwaysRelevant actor %s"), *PrintObjectFromNetRefHandle(ActorRefHandle));
 
 	// Set owning connection filtering if actor is only relevant to owner
 	{
@@ -422,7 +420,7 @@ UE::Net::FNetRefHandle UEngineReplicationBridge::StartReplicatingComponent(FNetR
 
 		if (!SubObject->IsUsingRegisteredSubObjectList())
 		{
-			UE_LOG_ACTORREPLICATIONBRIDGE(Warning, TEXT("ActorComponent %s does not replicate subobjects using the registered SubObjectsLists, SubObjects might not replicate properly."), ToCStr(GetFullNameSafe(SubObject)));
+			UE_LOG_BRIDGEID(LogIris, Warning, TEXT("ActorComponent %s does not replicate subobjects using the registered SubObjectsLists, SubObjects might not replicate properly."), ToCStr(GetFullNameSafe(SubObject)));
 		}
 
 		if (RepComponentInfo == nullptr || RepComponentInfo->NetCondition == ELifetimeCondition::COND_Never)
@@ -486,7 +484,7 @@ void UEngineReplicationBridge::StopReplicatingActor(AActor* Actor, EEndPlayReaso
 		return;
 	}
 
-	UE_LOG(LogIrisBridge, Verbose, TEXT("StopReplicatingActor %s. Reason %s "), *PrintObjectFromNetRefHandle(RefHandle), *UEnum::GetValueAsString(TEXT("Engine.EEndPlayReason"), EndPlayReason));
+	UE_LOG(LogIris, Verbose, TEXT("StopReplicatingActor %s. Reason %s "), *PrintObjectFromNetRefHandle(RefHandle), *UEnum::GetValueAsString(TEXT("Engine.EEndPlayReason"), EndPlayReason));
 	ensureMsgf(IsValid(Actor), TEXT("StopReplicatingActor called on invalid actor tied to handle: %s."), *PrintObjectFromNetRefHandle(RefHandle));
 
 	const bool bIsServer = GetReplicationSystem()->IsServer();
@@ -542,7 +540,7 @@ void UEngineReplicationBridge::StopReplicatingComponent(UActorComponent* ActorCo
 	FNetRefHandle ComponentHandle = GetReplicatedRefHandle(ActorComponent, EGetRefHandleFlags::EvenIfGarbage);
 	if (ComponentHandle.IsValid())
 	{
-		UE_LOG(LogIrisBridge, Verbose, TEXT("StopReplicatingComponent for %s %s."), *GetNameSafe(ActorComponent), *ComponentHandle.ToString());
+		UE_LOG(LogIris, Verbose, TEXT("StopReplicatingComponent for %s %s."), *GetNameSafe(ActorComponent), *ComponentHandle.ToString());
 		ensureMsgf(IsValid(ActorComponent), TEXT("StopReplicatingActor called on invalid component tied to handle: %s."), *ComponentHandle.ToString());
 		
 		StopReplicatingNetRefHandle(ComponentHandle, EndReplicationFlags);
@@ -780,7 +778,7 @@ void UEngineReplicationBridge::OnProtocolMismatchReported(FNetRefHandle RefHandl
 			if (UNetConnection* ClientConnection = NetDriver->GetConnectionByHandle(UE::Net::FConnectionHandle(ConnectionId)))
 			{
 				FString ErrorMsg = FString::Printf(TEXT("Protocol mismatch: %s:%s. Class: %s"), *RefHandle.ToString(), *GetNameSafe(ReplicatedObject), *GetNameSafe(ObjectClass));
-				UE_LOG(LogIrisBridge, Error, TEXT("%s: Closing connection due to: %s"), ToCStr(ClientConnection->Describe()), ToCStr(ErrorMsg));
+				UE_LOG(LogIris, Error, TEXT("%s: Closing connection due to: %s"), ToCStr(ClientConnection->Describe()), ToCStr(ErrorMsg));
 				{
 					UE::Net::FNetCloseResult CloseReason = ENetCloseResult::IrisProtocolMismatch;
 					ClientConnection->SendCloseReason(MoveTemp(CloseReason));
@@ -809,7 +807,7 @@ void UEngineReplicationBridge::SendErrorWithNetRefHandle(UE::Net::ENetRefHandleE
 		}
 		else
 		{
-			UE_LOG(LogIrisBridge, Error, TEXT("UEngineReplicationBridge::SendErrorWithNetRefHandle could not find Connection for id:%u"), ConnectionId);
+			UE_LOG(LogIris, Error, TEXT("UEngineReplicationBridge::SendErrorWithNetRefHandle could not find Connection for id:%u"), ConnectionId);
 		}
 	}
 }
@@ -833,7 +831,7 @@ void UEngineReplicationBridge::SendErrorWithNetRefHandle(UE::Net::ENetRefHandleE
 		}
 		else
 		{
-			UE_LOG(LogIrisBridge, Error, TEXT("UEngineReplicationBridge::SendErrorWithNetRefHandle could not find Connection for id:%u"), ConnectionId);
+			UE_LOG(LogIris, Error, TEXT("UEngineReplicationBridge::SendErrorWithNetRefHandle could not find Connection for id:%u"), ConnectionId);
 		}
 	}
 }
@@ -845,7 +843,7 @@ void UEngineReplicationBridge::ActorChangedLevel(const AActor* Actor, const ULev
 		return;
 	}
 
-	UE_LOG_ACTORREPLICATIONBRIDGE(Verbose, TEXT("ActorChangedLevel: Actor %s from PreviousLevel %s"), *GetFullNameSafe(Actor), *GetFullNameSafe(PreviousLevel));
+	UE_LOG_BRIDGEID(LogIris, Verbose, TEXT("ActorChangedLevel: Actor %s from PreviousLevel %s"), *GetFullNameSafe(Actor), *GetFullNameSafe(PreviousLevel));
 
 	if (!UE::Net::Private::ShouldIncludeRootObjectInLevelGroups(Actor))
 	{
@@ -864,12 +862,12 @@ void UEngineReplicationBridge::ActorChangedLevel(const AActor* Actor, const ULev
 		UE::Net::FNetObjectGroupHandle PreviousLevelGroup = GetLevelGroup(PreviousLevel);
 		if (GetReplicationSystem()->IsInGroup(PreviousLevelGroup, ActorRefHandle))
 		{
-			UE_LOG_ACTORREPLICATIONBRIDGE(Verbose, TEXT("ActorChangedLevel: removing %s from GroupIndex: %u PreviousLevel: %s"), *ActorRefHandle.ToString(), PreviousLevelGroup.GetGroupIndex(), ToCStr(PreviousLevelPackageName.ToString()));
+			UE_LOG_BRIDGEID(LogIris, Verbose, TEXT("ActorChangedLevel: removing %s from GroupIndex: %u PreviousLevel: %s"), *ActorRefHandle.ToString(), PreviousLevelGroup.GetGroupIndex(), ToCStr(PreviousLevelPackageName.ToString()));
 			GetReplicationSystem()->RemoveFromGroup(PreviousLevelGroup, ActorRefHandle);
 		}
 		else
 		{
-			UE_LOG_ACTORREPLICATIONBRIDGE(Warning, TEXT("ActorChangedLevel: %s not found in GroupIndex: %u PreviousLevel: %s"), *ActorRefHandle.ToString(), PreviousLevelGroup.GetGroupIndex(), ToCStr(PreviousLevelPackageName.ToString()));
+			UE_LOG_BRIDGEID(LogIris, Warning, TEXT("ActorChangedLevel: %s not found in GroupIndex: %u PreviousLevel: %s"), *ActorRefHandle.ToString(), PreviousLevelGroup.GetGroupIndex(), ToCStr(PreviousLevelPackageName.ToString()));
 		}
 	}
 
@@ -907,7 +905,7 @@ void UEngineReplicationBridge::AddRootObjectToLevelGroup(const UObject* RootObje
 		{
 			LevelGroup = CreateLevelGroup(Level, PackageName);
 
-			UE_LOG_ACTORREPLICATIONBRIDGE(Verbose, TEXT("Created new GroupIndex: %u for Level: %s"), LevelGroup.GetGroupIndex(), ToCStr(PackageName.ToString()));
+			UE_LOG_BRIDGEID(LogIris, Verbose, TEXT("Created new GroupIndex: %u for Level: %s"), LevelGroup.GetGroupIndex(), ToCStr(PackageName.ToString()));
 
 			// Update the filtering status of the group based on current level visibility for all connections
 			NetDriver->UpdateGroupFilterStatusForLevel(Level, LevelGroup);
@@ -916,7 +914,7 @@ void UEngineReplicationBridge::AddRootObjectToLevelGroup(const UObject* RootObje
 		const UE::Net::FNetRefHandle RootObjectRefHandle = GetReplicatedRefHandle(RootObject);
 
 		// Add object to group
-		UE_LOG_ACTORREPLICATIONBRIDGE(Verbose, TEXT("Added %s to GroupIndex: %u Level: %s"), *RootObjectRefHandle.ToString(), LevelGroup.GetGroupIndex(), ToCStr(PackageName.ToString()));
+		UE_LOG_BRIDGEID(LogIris, Verbose, TEXT("Added %s to GroupIndex: %u Level: %s"), *RootObjectRefHandle.ToString(), LevelGroup.GetGroupIndex(), ToCStr(PackageName.ToString()));
 		GetReplicationSystem()->AddToGroup(LevelGroup, RootObjectRefHandle);
 	}
 }
diff --git a/Engine/Source/Runtime/Experimental/Iris/Core/Private/Iris/ReplicationSystem/ObjectReplicationBridge.cpp b/Engine/Source/Runtime/Experimental/Iris/Core/Private/Iris/ReplicationSystem/ObjectReplicationBridge.cpp
--- a/Engine/Source/Runtime/Experimental/Iris/Core/Private/Iris/ReplicationSystem/ObjectReplicationBridge.cpp
+++ b/Engine/Source/Runtime/Experimental/Iris/Core/Private/Iris/ReplicationSystem/ObjectReplicationBridge.cpp
@@ -56,7 +56,7 @@ namespace UE::Net::Private::ObjectBridgeDebugging
 
 DEFINE_LOG_CATEGORY(LogIrisFilterConfig)
 
-#define UE_LOG_OBJECTREPLICATIONBRIDGE(Category, Format, ...)  UE_LOG(LogIrisBridge, Category, TEXT("ObjectReplicationBridge(%u)::") Format, GetReplicationSystem()->GetId(), ##__VA_ARGS__)
+#define UE_LOG_OBJECTREPLICATIONBRIDGE(Category, Format, ...)  UE_LOG(LogIris, Category, TEXT("RepBridge(%u)::") Format, GetReplicationSystem()->GetId(), ##__VA_ARGS__)
 
 static bool bUseFrequencyBasedPolling = true;
 static FAutoConsoleVariableRef CVarUseFrequencyBasedPolling(
@@ -560,7 +560,7 @@ void UObjectReplicationBridge::AssignDynamicFilter(UObject* Instance, const FRoo
 		{
 			FilterHandle = ReplicationSystem->GetFilterHandle(Params.ExplicitDynamicFilterName);
 			
-			UE_CLOG(FilterHandle == InvalidNetObjectFilterHandle, LogIrisBridge, Error, TEXT("Could not assign explicit dynamic filter to %s. No filters named %s exist"), *GetPathNameSafe(Instance), *Params.ExplicitDynamicFilterName.ToString() );
+			UE_CLOG(FilterHandle == InvalidNetObjectFilterHandle, LogIris, Error, TEXT("Could not assign explicit dynamic filter to %s. No filters named %s exist"), *GetPathNameSafe(Instance), *Params.ExplicitDynamicFilterName.ToString() );
 			ensure(FilterHandle != InvalidNetObjectFilterHandle);
 		}
 	}
@@ -757,7 +757,7 @@ bool UObjectReplicationBridge::WriteNetRefHandleCreationInfo(FReplicationBridgeS
 	{
 		const FNetObjectCreationHeader* Header = CachedHeader->Get();
 		UNetObjectFactory* Factory = GetNetFactory(Header->GetNetFactoryId());
-		UE_LOG_BRIDGEID(LogIrisBridge, VeryVerbose, TEXT("WriteNetRefHandleCreationInfo (cached) for %s : Factory: %s (id:%u) Header: %s"), *Handle.ToString(), *GetNameSafe(Factory), Header->GetNetFactoryId(), *Header->ToString());
+		UE_LOG_BRIDGEID(LogIris, VeryVerbose, TEXT("WriteNetRefHandleCreationInfo (cached) for %s : Factory: %s (id:%u) Header: %s"), *Handle.ToString(), *GetNameSafe(Factory), Header->GetNetFactoryId(), *Header->ToString());
 		return Factory->WriteHeader(Handle, Context.SerializationContext, Header);
 	}
 	
@@ -773,7 +773,7 @@ bool UObjectReplicationBridge::WriteNetRefHandleCreationInfo(FReplicationBridgeS
 	const FInternalNetRefIndex ObjectIndex = NetRefHandleManager->GetInternalIndex(Handle);
 	if (ObjectIndex == FNetRefHandleManager::InvalidInternalIndex)
 	{
-		UE_LOG_BRIDGEID(LogIrisBridge, Error, TEXT("WriteNetRefHandleCreationInfo: Cannot write creationinfo for %s: No internal index assigned to it"), *Handle.ToString());
+		UE_LOG_BRIDGEID(LogIris, Error, TEXT("WriteNetRefHandleCreationInfo: Cannot write creationinfo for %s: No internal index assigned to it"), *Handle.ToString());
 		return false;
 	}
 
@@ -786,11 +786,11 @@ bool UObjectReplicationBridge::WriteNetRefHandleCreationInfo(FReplicationBridgeS
 
 	if (!Header.IsValid())
 	{
-		UE_LOG_BRIDGEID(LogIrisBridge, Error, TEXT("WriteNetRefHandleCreationInfo: CreateHeader via %s failed for %s"), *GetNameSafe(Factory), *Handle.ToString());
+		UE_LOG_BRIDGEID(LogIris, Error, TEXT("WriteNetRefHandleCreationInfo: CreateHeader via %s failed for %s"), *GetNameSafe(Factory), *Handle.ToString());
 		return false;
 	}
 
-	UE_LOG_BRIDGEID(LogIrisBridge, VeryVerbose, TEXT("WriteNetRefHandleCreationInfo (new) for %s : Factory: %s (id:%u) Header: %s"), *Handle.ToString(), *GetNameSafe(Factory), FactoryId, *Header->ToString());
+	UE_LOG_BRIDGEID(LogIris, VeryVerbose, TEXT("WriteNetRefHandleCreationInfo (new) for %s : Factory: %s (id:%u) Header: %s"), *Handle.ToString(), *GetNameSafe(Factory), FactoryId, *Header->ToString());
 
 	return Factory->WriteHeader(Handle, Context.SerializationContext, Header.Get());
 }
@@ -809,7 +809,7 @@ bool UObjectReplicationBridge::CacheNetRefHandleCreationInfo(FNetRefHandle Handl
 	const FInternalNetRefIndex ObjectIndex = NetRefHandleManager->GetInternalIndex(Handle);
 	if (ObjectIndex == FNetRefHandleManager::InvalidInternalIndex)
 	{
-		UE_LOG_BRIDGEID(LogIrisBridge, Error, TEXT("CacheNetRefHandleCreationInfo cannot cache creation info for %s: no internal index assigned to it"), *Handle.ToString());
+		UE_LOG_BRIDGEID(LogIris, Error, TEXT("CacheNetRefHandleCreationInfo cannot cache creation info for %s: no internal index assigned to it"), *Handle.ToString());
 		ensureMsgf(false, TEXT("CacheNetRefHandleCreationInfo cannot cache creation info for %s: no internal index assigned to it"), *Handle.ToString());
 		return false;
 	}
@@ -951,15 +951,15 @@ FReplicationBridgeCreateNetRefHandleResult UObjectReplicationBridge::CreateNetRe
 	FFragmentRegistrationContext FragmentRegistrationContext(GetReplicationStateDescriptorRegistry(), GetReplicationSystem(), EReplicationFragmentTraits::CanReceive);
 	FReplicationProtocolManager* ProtocolManager = GetReplicationProtocolManager();
 
-	if (UE_LOG_ACTIVE(LogIrisBridge, Verbose))
+	if (UE_LOG_ACTIVE(LogIris, Verbose))
 	{
 		if (RootObjectOfSubObject.IsValid())
 		{
-			UE_LOG_BRIDGEID(LogIrisBridge, Verbose, TEXT("CreateNetRefHandleFromRemote: SubObject: %s of RootObject: %s using header: %s"), *WantedNetHandle.ToString(), *RootObjectOfSubObject.ToString(), ToCStr(Header->ToString()));
+			UE_LOG_BRIDGEID(LogIris, Verbose, TEXT("CreateNetRefHandleFromRemote: SubObject: %s of RootObject: %s using header: %s"), *WantedNetHandle.ToString(), *RootObjectOfSubObject.ToString(), ToCStr(Header->ToString()));
 		}
 		else
 		{
-			UE_LOG_BRIDGEID(LogIrisBridge, Verbose, TEXT("CreateNetRefHandleFromRemote: RootObject: %s using header: %s"), *WantedNetHandle.ToString(), ToCStr(Header->ToString()));
+			UE_LOG_BRIDGEID(LogIris, Verbose, TEXT("CreateNetRefHandleFromRemote: RootObject: %s using header: %s"), *WantedNetHandle.ToString(), ToCStr(Header->ToString()));
 		}
 	}
 	
@@ -977,15 +977,15 @@ FReplicationBridgeCreateNetRefHandleResult UObjectReplicationBridge::CreateNetRe
 	// Failed ?
 	if (!Result.Instance)
 	{
-		if (UE_LOG_ACTIVE(LogIrisBridge, Warning) && !bSuppressCreateInstanceFailedEnsure)
+		if (UE_LOG_ACTIVE(LogIris, Warning) && !bSuppressCreateInstanceFailedEnsure)
 		{
 			if (RootObjectOfSubObject.IsValid())
 			{
-				UE_LOG_BRIDGEID(LogIrisBridge, Warning, TEXT("CreateNetRefHandleFromRemote: Failed to instantiate SubObject NetHandle: %s of %s using header: %s"), *WantedNetHandle.ToString(), *PrintObjectFromNetRefHandle(RootObjectOfSubObject), ToCStr(Header->ToString()));
+				UE_LOG_BRIDGEID(LogIris, Warning, TEXT("CreateNetRefHandleFromRemote: Failed to instantiate SubObject NetHandle: %s of %s using header: %s"), *WantedNetHandle.ToString(), *PrintObjectFromNetRefHandle(RootObjectOfSubObject), ToCStr(Header->ToString()));
 			}
 			else
 			{
-				UE_LOG_BRIDGEID(LogIrisBridge, Warning, TEXT("CreateNetRefHandleFromRemote: Failed to instantiate RootObject NetHandle: %s using header: %s"), *WantedNetHandle.ToString(), ToCStr(Header->ToString()));
+				UE_LOG_BRIDGEID(LogIris, Warning, TEXT("CreateNetRefHandleFromRemote: Failed to instantiate RootObject NetHandle: %s using header: %s"), *WantedNetHandle.ToString(), ToCStr(Header->ToString()));
 			}
 		}
 		
@@ -1167,7 +1167,7 @@ void UObjectReplicationBridge::PruneStaleObjects()
 			{
 				const FReplicationProtocol* Protocol = ObjectData.Protocol;
 				const FNetDebugName* DebugName = Protocol ? Protocol->DebugName : nullptr;
-				UE_LOG(LogIrisBridge, Warning, TEXT("UObjectReplicationBridge::PruneStaleObjects ObjectInstance replicated as: %s of Type named:%s has been destroyed without notifying the ReplicationSystem %u"), *ObjectData.RefHandle.ToString(), ToCStr(DebugName), ObjectData.RefHandle.GetReplicationSystemId());
+				UE_LOG(LogIris, Warning, TEXT("UObjectReplicationBridge::PruneStaleObjects ObjectInstance replicated as: %s of Type named:%s has been destroyed without notifying the ReplicationSystem %u"), *ObjectData.RefHandle.ToString(), ToCStr(DebugName), ObjectData.RefHandle.GetReplicationSystemId());
 
 				// If the instance protocol is bound, then this is an error and we cannot safely cleanup as unbinding abound instance protocol will modify bound states
 				if (ensureMsgf(!EnumHasAnyFlags(ObjectData.InstanceProtocol->InstanceTraits, EReplicationInstanceProtocolTraits::IsBound), TEXT("UObjectReplicationBridge::PruneStaleObjects Bound ObjectInstance replicated as: %s has been destroyed without notifying the ReplicationSystem."), *ObjectData.RefHandle.ToString()))
@@ -1559,7 +1559,7 @@ void UObjectReplicationBridge::UpdateRootObjectWorldInfo(const FUpdateWorldInfoC
 
 	// Check if we need to warn about invalid world locations.
 	// Note that FWorldLocations will store the invalid location but only return a valid clamped position from it's API
-	if (UE_LOG_ACTIVE(LogIrisBridge, Warning) && GInvalidLocationWarnings != 0)
+	if (UE_LOG_ACTIVE(LogIris, Warning) && GInvalidLocationWarnings != 0)
 	{
 		const bool bIsValidLocation = UpdateContext.WorldLocations.IsValidLocation(WorldInfo.WorldLocation);
 		
@@ -1584,7 +1584,7 @@ void UObjectReplicationBridge::UpdateRootObjectWorldInfo(const FUpdateWorldInfoC
 			// If it's the first warning or it's set to warn on every update
 			if (!bAlreadyWarned || GInvalidLocationWarnings == 3)
 			{
-				UE_LOG_BRIDGEID(LogIrisBridge, Warning, TEXT("RootObject %s has a location %s that is out of bounds."), 
+				UE_LOG_BRIDGEID(LogIris, Warning, TEXT("RootObject %s has a location %s that is out of bounds."), 
 					ToCStr(NetRefHandleManager->PrintObjectFromIndex(UpdateContext.ObjectIndex)), *WorldInfo.WorldLocation.ToCompactString());
 				bAlreadyWarned = true;
 			}
diff --git a/Engine/Source/Runtime/Experimental/Iris/Core/Private/Iris/ReplicationSystem/ObjectReplicationBridgeDebugging.cpp b/Engine/Source/Runtime/Experimental/Iris/Core/Private/Iris/ReplicationSystem/ObjectReplicationBridgeDebugging.cpp
--- a/Engine/Source/Runtime/Experimental/Iris/Core/Private/Iris/ReplicationSystem/ObjectReplicationBridgeDebugging.cpp
+++ b/Engine/Source/Runtime/Experimental/Iris/Core/Private/Iris/ReplicationSystem/ObjectReplicationBridgeDebugging.cpp
@@ -284,7 +284,7 @@ void LogRootObjectList(FLogContext& LogContext, bool bLogSubObjects)
 
 	for (const FRootObjectData& RootObject : LogContext.RootObjectArray)
 	{
-		UE_LOG(LogIrisBridge, Display, TEXT("%s %s"), *PrintNetObject(NetRefHandleManager, RootObject.ObjectIndex), LogContext.OptionalObjectPrint ? *LogContext.OptionalObjectPrint(RootObject.ObjectIndex) : TEXT(""));
+		UE_LOG(LogIris, Display, TEXT("%s %s"), *PrintNetObject(NetRefHandleManager, RootObject.ObjectIndex), LogContext.OptionalObjectPrint ? *LogContext.OptionalObjectPrint(RootObject.ObjectIndex) : TEXT(""));
 
 		LogContext.NumRootObjects++;
 
@@ -293,7 +293,7 @@ void LogRootObjectList(FLogContext& LogContext, bool bLogSubObjects)
 			TArrayView<const FInternalNetRefIndex> SubObjects = NetRefHandleManager->GetSubObjects(RootObject.ObjectIndex);
 			for (FInternalNetRefIndex SubObjectIndex : SubObjects)
 			{
-				UE_LOG(LogIrisBridge, Display, TEXT("\t%s %s"), *PrintNetObject(NetRefHandleManager, SubObjectIndex), LogContext.OptionalObjectPrint ? *LogContext.OptionalObjectPrint(SubObjectIndex) : TEXT(""));
+				UE_LOG(LogIris, Display, TEXT("\t%s %s"), *PrintNetObject(NetRefHandleManager, SubObjectIndex), LogContext.OptionalObjectPrint ? *LogContext.OptionalObjectPrint(SubObjectIndex) : TEXT(""));
 
 				LogContext.NumSubObjects++;
 			}
@@ -368,14 +368,14 @@ FAutoConsoleCommand ObjectBridgePrintDynamicFilterConfig(
 	UReplicationSystem* RepSystem = FindReplicationSystemFromArg(Args);
 	if (!RepSystem)
 	{
-		UE_LOG(LogIrisBridge, Error, TEXT("Could not find ReplicationSystem."));
+		UE_LOG(LogIris, Error, TEXT("Could not find ReplicationSystem."));
 		return;
 	}
 
 	UObjectReplicationBridge* ObjectBridge = CastChecked<UObjectReplicationBridge>(RepSystem->GetReplicationBridge());
 	if (!ObjectBridge)
 	{
-		UE_LOG(LogIrisBridge, Error, TEXT("Could not find ObjectReplicationBridge."));
+		UE_LOG(LogIris, Error, TEXT("Could not find ObjectReplicationBridge."));
 		return;
 	}
 
@@ -393,14 +393,14 @@ FAutoConsoleCommand ObjectBridgePrintDynamicFilterConfigIssues(
 	UReplicationSystem* RepSystem = FindReplicationSystemFromArg(Args);
 	if (!RepSystem)
 	{
-		UE_LOG(LogIrisBridge, Error, TEXT("Could not find ReplicationSystem."));
+		UE_LOG(LogIris, Error, TEXT("Could not find ReplicationSystem."));
 		return;
 	}
 
 	UObjectReplicationBridge* ObjectBridge = CastChecked<UObjectReplicationBridge>(RepSystem->GetReplicationBridge());
 	if (!ObjectBridge)
 	{
-		UE_LOG(LogIrisBridge, Error, TEXT("Could not find ObjectReplicationBridge."));
+		UE_LOG(LogIris, Error, TEXT("Could not find ObjectReplicationBridge."));
 		return;
 	}
 
@@ -637,8 +637,8 @@ void UObjectReplicationBridge::PrintReplicatedObjects(uint32 ArgTraits) const
 	using namespace UE::Net::Private;
 	using namespace UE::Net::Private::ObjectBridgeDebugging;
 
-	UE_LOG(LogIrisBridge, Display, TEXT("################ Start Printing ALL Replicated Objects ################"));
-	UE_LOG(LogIrisBridge, Display, TEXT(""));
+	UE_LOG(LogIris, Display, TEXT("################ Start Printing ALL Replicated Objects ################"));
+	UE_LOG(LogIris, Display, TEXT(""));
 
 	uint32 TotalRootObjects = 0;
 	uint32 TotalSubObjects = 0;
@@ -665,9 +665,9 @@ void UObjectReplicationBridge::PrintReplicatedObjects(uint32 ArgTraits) const
 	FLogContext LogContext = {.NetRefHandleManager = NetRefHandleManager, .RootObjectArray = RootObjectArray, .OptionalObjectPrint = PrintClassOrProtocol};
 	LogViaTrait(LogContext, (EPrintTraits)ArgTraits, EPrintTraits::Default);
 
-	UE_LOG(LogIrisBridge, Display, TEXT(""));
-	UE_LOG(LogIrisBridge, Display, TEXT("Printed %u root objects and %u sub objects"), LogContext.NumRootObjects, LogContext.NumSubObjects);
-	UE_LOG(LogIrisBridge, Display, TEXT("################ Stop Printing ALL Replicated Objects ################"));
+	UE_LOG(LogIris, Display, TEXT(""));
+	UE_LOG(LogIris, Display, TEXT("Printed %u root objects and %u sub objects"), LogContext.NumRootObjects, LogContext.NumSubObjects);
+	UE_LOG(LogIris, Display, TEXT("################ Stop Printing ALL Replicated Objects ################"));
 }
 
 //-----------------------------------------------
@@ -694,8 +694,8 @@ void UObjectReplicationBridge::PrintRelevantObjects(uint32 ArgTraits) const
 	using namespace UE::Net::Private;
 	using namespace UE::Net::Private::ObjectBridgeDebugging;
 
-	UE_LOG(LogIrisBridge, Display, TEXT("################ Start Printing Relevant Objects ################"));
-	UE_LOG(LogIrisBridge, Display, TEXT(""));
+	UE_LOG(LogIris, Display, TEXT("################ Start Printing Relevant Objects ################"));
+	UE_LOG(LogIris, Display, TEXT(""));
 
 	FNetBitArray RootObjects;
 	RootObjects.Init(NetRefHandleManager->GetCurrentMaxInternalNetRefIndex());
@@ -711,9 +711,9 @@ void UObjectReplicationBridge::PrintRelevantObjects(uint32 ArgTraits) const
 	FLogContext LogContext = {.NetRefHandleManager = NetRefHandleManager, .RootObjectArray = RootObjectArray };
 	LogViaTrait(LogContext, (EPrintTraits)ArgTraits, EPrintTraits::Default);
 
-	UE_LOG(LogIrisBridge, Display, TEXT(""));
-	UE_LOG(LogIrisBridge, Display, TEXT("Printed %u root objects and %u sub objects"), LogContext.NumRootObjects, LogContext.NumSubObjects);
-	UE_LOG(LogIrisBridge, Display, TEXT("################ Stop Printing Relevant Objects ################"));
+	UE_LOG(LogIris, Display, TEXT(""));
+	UE_LOG(LogIris, Display, TEXT("Printed %u root objects and %u sub objects"), LogContext.NumRootObjects, LogContext.NumSubObjects);
+	UE_LOG(LogIris, Display, TEXT("################ Stop Printing Relevant Objects ################"));
 }
 
 //-----------------------------------------------
@@ -742,8 +742,8 @@ void UObjectReplicationBridge::PrintAlwaysRelevantObjects(uint32 ArgTraits) cons
 
 	FReplicationSystemInternal* ReplicationSystemInternal = GetReplicationSystem()->GetReplicationSystemInternal();
 
-	UE_LOG(LogIrisBridge, Display, TEXT("################ Start Printing Always Relevant Objects ################"));
-	UE_LOG(LogIrisBridge, Display, TEXT(""));
+	UE_LOG(LogIris, Display, TEXT("################ Start Printing Always Relevant Objects ################"));
+	UE_LOG(LogIris, Display, TEXT(""));
 
 	FNetBitArray AlwaysRelevantList;
 	AlwaysRelevantList.Init(NetRefHandleManager->GetCurrentMaxInternalNetRefIndex());
@@ -779,9 +779,9 @@ void UObjectReplicationBridge::PrintAlwaysRelevantObjects(uint32 ArgTraits) cons
 	FLogContext LogContext = {.NetRefHandleManager=NetRefHandleManager, .RootObjectArray=AlwaysRelevantObjects };
 	LogViaTrait(LogContext, (EPrintTraits)ArgTraits, EPrintTraits::Default);
 
-	UE_LOG(LogIrisBridge, Display, TEXT(""));
-	UE_LOG(LogIrisBridge, Display, TEXT("Printed %u root objects and %u subobjects"), LogContext.NumRootObjects, LogContext.NumSubObjects);
-	UE_LOG(LogIrisBridge, Display, TEXT("################ Stop Printing Always Relevant Objects ################"));
+	UE_LOG(LogIris, Display, TEXT(""));
+	UE_LOG(LogIris, Display, TEXT("Printed %u root objects and %u subobjects"), LogContext.NumRootObjects, LogContext.NumSubObjects);
+	UE_LOG(LogIris, Display, TEXT("################ Stop Printing Always Relevant Objects ################"));
 }
 
 //-----------------------------------------------
@@ -842,8 +842,8 @@ void UObjectReplicationBridge::PrintRelevantObjectsForConnections(const TArray<F
 		}
 	}
 
-	UE_LOG(LogIrisBridge, Display, TEXT("################ Start Printing Relevant Objects of %d Connections ################"), ConnectionsToPrint.CountSetBits());
-	UE_LOG(LogIrisBridge, Display, TEXT(""));
+	UE_LOG(LogIris, Display, TEXT("################ Start Printing Relevant Objects of %d Connections ################"), ConnectionsToPrint.CountSetBits());
+	UE_LOG(LogIris, Display, TEXT(""));
 
 	const bool bWithFilterInfo = nullptr != Args.FindByPredicate([](const FString& Str) { return Str.Contains(TEXT("WithFilter")); });
 
@@ -858,9 +858,9 @@ void UObjectReplicationBridge::PrintRelevantObjectsForConnections(const TArray<F
 			ViewLocs += FString::Printf(TEXT("%s "), *UserView.Pos.ToCompactString());
 		}
 
-		UE_LOG(LogIrisBridge, Display, TEXT(""));
-		UE_LOG(LogIrisBridge, Display, TEXT("###### Begin Relevant list of Connection:%u ViewPos:%s Named: %s ######"), ConnectionId, *ViewLocs, *PrintConnectionInfo(ConnectionId));
-		UE_LOG(LogIrisBridge, Display, TEXT(""));
+		UE_LOG(LogIris, Display, TEXT(""));
+		UE_LOG(LogIris, Display, TEXT("###### Begin Relevant list of Connection:%u ViewPos:%s Named: %s ######"), ConnectionId, *ViewLocs, *PrintConnectionInfo(ConnectionId));
+		UE_LOG(LogIris, Display, TEXT(""));
 
 		FNetBitArray RootObjects;
 		RootObjects.Init(NetRefHandleManager->GetCurrentMaxInternalNetRefIndex());
@@ -881,13 +881,13 @@ void UObjectReplicationBridge::PrintRelevantObjectsForConnections(const TArray<F
 		FLogContext LogContext = {.NetRefHandleManager=NetRefHandleManager, .RootObjectArray=RelevantObjects, .OptionalObjectPrint=AddFilterInfo};
 		LogViaTrait(LogContext, ArgTraits, EPrintTraits::Default);
 
-		UE_LOG(LogIrisBridge, Display, TEXT(""));
-		UE_LOG(LogIrisBridge, Display, TEXT("###### Stop Relevant list of Connection:%u | Total: %u root objects relevant ######"), ConnectionId, LogContext.NumRootObjects);
-		UE_LOG(LogIrisBridge, Display, TEXT(""));
+		UE_LOG(LogIris, Display, TEXT(""));
+		UE_LOG(LogIris, Display, TEXT("###### Stop Relevant list of Connection:%u | Total: %u root objects relevant ######"), ConnectionId, LogContext.NumRootObjects);
+		UE_LOG(LogIris, Display, TEXT(""));
 	});
 
-	UE_LOG(LogIrisBridge, Display, TEXT(""));
-	UE_LOG(LogIrisBridge, Display, TEXT("################ Stop Printing Relevant Objects of %d Connections ################"), ConnectionsToPrint.CountSetBits());
+	UE_LOG(LogIris, Display, TEXT(""));
+	UE_LOG(LogIris, Display, TEXT("################ Stop Printing Relevant Objects of %d Connections ################"), ConnectionsToPrint.CountSetBits());
 }
 
 //-----------------------------------------------
@@ -1035,8 +1035,8 @@ void UObjectReplicationBridge::PrintNetCullDistances(const TArray<FString>& Args
 		});
 	}
 
-	UE_LOG(LogIrisBridge, Display, TEXT("################ Start Printing NetCullDistance Values ################"));
-	UE_LOG(LogIrisBridge, Display, TEXT(""));
+	UE_LOG(LogIris, Display, TEXT("################ Start Printing NetCullDistance Values ################"));
+	UE_LOG(LogIris, Display, TEXT(""));
 
 	int32 NumClassesPrinted = 0;
 	for (auto ClassIt = ClassCullDistanceMap.CreateIterator(); ClassIt; ++ClassIt)
@@ -1044,12 +1044,12 @@ void UObjectReplicationBridge::PrintNetCullDistances(const TArray<FString>& Args
 		FCullDistanceInfo& Info = ClassIt.Value();
 		UClass* Class = Info.Class;
 
-		UE_LOG(LogIrisBridge, Display, TEXT("MostCommon NetCullDistance: %f | Class: %s | Instances: %u | Overrides: %u"), Info.FindMostUsedCullDistance(), *Info.Class->GetName(), Info.NumTotal, Info.NumOverrides);
+		UE_LOG(LogIris, Display, TEXT("MostCommon NetCullDistance: %f | Class: %s | Instances: %u | Overrides: %u"), Info.FindMostUsedCullDistance(), *Info.Class->GetName(), Info.NumTotal, Info.NumOverrides);
 
 		Info.UniqueCullDistances.KeySort([](const float& lhs, const float& rhs){ return lhs >= rhs; });
 		for (auto DivergentIt = Info.UniqueCullDistances.CreateConstIterator(); DivergentIt; ++DivergentIt)
 		{
-			UE_LOG(LogIrisBridge, Display, TEXT("\tNetCullDistance: %f | UseCount: %d/%d (%.2f%%)"), DivergentIt.Key(), DivergentIt.Value(), Info.NumTotal,((float)DivergentIt.Value()/(float)Info.NumTotal)*100.f);
+			UE_LOG(LogIris, Display, TEXT("\tNetCullDistance: %f | UseCount: %d/%d (%.2f%%)"), DivergentIt.Key(), DivergentIt.Value(), Info.NumTotal,((float)DivergentIt.Value()/(float)Info.NumTotal)*100.f);
 		}
 
 		if (++NumClassesPrinted == NumClassesToPrint)
@@ -1058,8 +1058,8 @@ void UObjectReplicationBridge::PrintNetCullDistances(const TArray<FString>& Args
 		}
 	}
 	
-	UE_LOG(LogIrisBridge, Display, TEXT(""));
-	UE_LOG(LogIrisBridge, Display, TEXT("################ Stop Printing NetCullDistance Values ################"));
+	UE_LOG(LogIris, Display, TEXT(""));
+	UE_LOG(LogIris, Display, TEXT("################ Stop Printing NetCullDistance Values ################"));
 }
 
 //-----------------------------------------------
@@ -1074,14 +1074,14 @@ FAutoConsoleCommand ObjectBridgePrintPushBasedStatuses(
 	UReplicationSystem* RepSystem = FindReplicationSystemFromArg(Args);
 	if (!RepSystem)
 	{
-		UE_LOG(LogIrisBridge, Error, TEXT("Could not find ReplicationSystem."));
+		UE_LOG(LogIris, Error, TEXT("Could not find ReplicationSystem."));
 		return;
 	}
 
 	UObjectReplicationBridge* ObjectBridge = RepSystem->GetReplicationBridgeAs<UObjectReplicationBridge>();
 	if (!ObjectBridge)
 	{
-		UE_LOG(LogIrisBridge, Error, TEXT("Could not find ObjectReplicationBridge."));
+		UE_LOG(LogIris, Error, TEXT("Could not find ObjectReplicationBridge."));
 		return;
 	}
 
@@ -1096,7 +1096,7 @@ void UObjectReplicationBridge::PrintPushBasedStatuses() const
 	FReplicationProtocolManager* ProtocolManager = GetReplicationProtocolManager();
 	if (!ProtocolManager)
 	{
-		UE_LOG(LogIrisBridge, Error, TEXT("Could not find ReplicationProtocolManager."));
+		UE_LOG(LogIris, Error, TEXT("Could not find ReplicationProtocolManager."));
 		return;
 	}
 
@@ -1141,15 +1141,15 @@ void UObjectReplicationBridge::PrintPushBasedStatuses() const
 		return A.Class->GetName() < B.Class->GetName();
 	});
 
-	UE_LOG(LogIrisBridge, Display, TEXT("################ Start Printing Push-Based Statuses ################"));
-	UE_LOG(LogIrisBridge, Display, TEXT(""));
+	UE_LOG(LogIris, Display, TEXT("################ Start Printing Push-Based Statuses ################"));
+	UE_LOG(LogIris, Display, TEXT(""));
 
 	for (const FPushBasedInfo& Info : PushBasedInfos)
 	{
-		UE_LOG(LogIrisBridge, Display, TEXT("%s (RefCount: %d) (PushBased: %d)"), ToCStr(Info.Class->GetName()), Info.RefCount, (int32)Info.bIsFullyPushBased);
+		UE_LOG(LogIris, Display, TEXT("%s (RefCount: %d) (PushBased: %d)"), ToCStr(Info.Class->GetName()), Info.RefCount, (int32)Info.bIsFullyPushBased);
 		if (!Info.bIsFullyPushBased)
 		{
-			UE_LOG(LogIrisBridge, Display, TEXT("\tPrinting properties that aren't push-based:"));
+			UE_LOG(LogIris, Display, TEXT("\tPrinting properties that aren't push-based:"));
 
 			TArray<FLifetimeProperty> LifetimeProps;
 			LifetimeProps.Reserve(Info.Class->ClassReps.Num());
@@ -1160,14 +1160,14 @@ void UObjectReplicationBridge::PrintPushBasedStatuses() const
 				{
 					const FRepRecord& RepRecord = Info.Class->ClassReps[LifetimeProp.RepIndex];
 					const FProperty* Prop = RepRecord.Property;
-					UE_LOG(LogIrisBridge, Display, TEXT("\t\t%s"), ToCStr(Prop->GetPathName()));
+					UE_LOG(LogIris, Display, TEXT("\t\t%s"), ToCStr(Prop->GetPathName()));
 				}
 			}
 		}
 	}
 
-	UE_LOG(LogIrisBridge, Display, TEXT(""));
-	UE_LOG(LogIrisBridge, Display, TEXT("################ Stop Printing Push-Based Statuses ################"));
+	UE_LOG(LogIris, Display, TEXT(""));
+	UE_LOG(LogIris, Display, TEXT("################ Stop Printing Push-Based Statuses ################"));
 }
 
 //-----------------------------------------------
@@ -1183,14 +1183,14 @@ FConsoleCommandWithArgsDelegate::CreateLambda([](const TArray< FString >& Args)
 	UReplicationSystem* RepSystem = FindReplicationSystemFromArg(Args);
 	if (!RepSystem)
 	{
-		UE_LOG(LogIrisBridge, Error, TEXT("Could not find ReplicationSystem."));
+		UE_LOG(LogIris, Error, TEXT("Could not find ReplicationSystem."));
 		return;
 	}
 
 	UObjectReplicationBridge* ObjectBridge = RepSystem->GetReplicationBridgeAs<UObjectReplicationBridge>();
 	if (!ObjectBridge)
 	{
-		UE_LOG(LogIrisBridge, Error, TEXT("Could not find ObjectReplicationBridge."));
+		UE_LOG(LogIris, Error, TEXT("Could not find ObjectReplicationBridge."));
 		return;
 	}
 
diff --git a/Engine/Source/Runtime/Experimental/Iris/Core/Private/Iris/ReplicationSystem/ReplicationBridge.cpp b/Engine/Source/Runtime/Experimental/Iris/Core/Private/Iris/ReplicationSystem/ReplicationBridge.cpp
--- a/Engine/Source/Runtime/Experimental/Iris/Core/Private/Iris/ReplicationSystem/ReplicationBridge.cpp
+++ b/Engine/Source/Runtime/Experimental/Iris/Core/Private/Iris/ReplicationSystem/ReplicationBridge.cpp
@@ -28,9 +28,7 @@
 #include "Iris/Serialization/InternalNetSerializationContext.h"
 #include "ReplicationFragmentInternal.h"
 
-DEFINE_LOG_CATEGORY(LogIrisBridge)
-
-#define UE_LOG_REPLICATIONBRIDGE(Category, Format, ...)  UE_LOG(LogIrisBridge, Category, TEXT("ReplicationBridge(%u)::") Format, ReplicationSystem->GetId(), ##__VA_ARGS__)
+#define UE_LOG_REPLICATIONBRIDGE(Category, Format, ...)  UE_LOG(LogIris, Category, TEXT("RepBridge(%u)::") Format, ReplicationSystem->GetId(), ##__VA_ARGS__)
 
 static bool bEnableFlushReliableRPCOnDestroy = true;
 static FAutoConsoleVariableRef CVarEnableFlushReliableRPCOnDestroy(
@@ -171,7 +169,7 @@ void UReplicationBridge::DetachSubObjectInstancesFromRemote(FNetRefHandle OwnerH
 				}
 				else
 				{
-					UE_LOG(LogIrisBridge, Warning, TEXT("Detaching Subobject %s with reason DoNotDestroy even though it is dynamic"), *SubObjectHandle.ToString());
+					UE_LOG(LogIris, Warning, TEXT("Detaching Subobject %s with reason DoNotDestroy even though it is dynamic"), *SubObjectHandle.ToString());
 				}
 			}
 
@@ -265,7 +263,7 @@ bool UReplicationBridge::CallWriteNetRefHandleDestructionInfo(FReplicationBridge
 	const FDestructionInfo* Info = StaticObjectsPendingDestroy.Find(Handle);
 	if (ensure(Info))
 	{
-		UE_LOG(LogIrisBridge, VeryVerbose, TEXT("CallWriteNetRefHandleDestructionInfo on %s | %s | NetFactoryId: %u"), *PrintObjectFromNetRefHandle(Handle), *Info->StaticRef.ToString(), Info->NetFactoryId);
+		UE_LOG(LogIris, VeryVerbose, TEXT("CallWriteNetRefHandleDestructionInfo on %s | %s | NetFactoryId: %u"), *PrintObjectFromNetRefHandle(Handle), *Info->StaticRef.ToString(), Info->NetFactoryId);
 		// Write destruction info
 		WriteFullNetObjectReference(Context.SerializationContext, Info->StaticRef);
 		Context.SerializationContext.GetBitStreamWriter()->WriteBits(Info->NetFactoryId, FNetObjectFactoryRegistry::GetMaxBits());
@@ -324,7 +322,7 @@ void UReplicationBridge::ReadAndExecuteDestructionInfoFromRemote(FReplicationBri
 			NetRefHandleManager->GetReplicatedObjectDataNoCheck(InternalReplicationIndex).bPendingEndReplication = 1U;
 		}
 
-		UE_LOG(LogIrisBridge, VeryVerbose, TEXT("ReadAndExecuteDestructionInfoFromRemote on %s | %s | NetFactoryId: %u"), *ReferenceToDestroy.ToString(), *ReferenceToDestroy.GetRefHandle().ToString(), NetFactoryId);
+		UE_LOG(LogIris, VeryVerbose, TEXT("ReadAndExecuteDestructionInfoFromRemote on %s | %s | NetFactoryId: %u"), *ReferenceToDestroy.ToString(), *ReferenceToDestroy.GetRefHandle().ToString(), NetFactoryId);
 
 		constexpr EReplicationBridgeDestroyInstanceReason DestroyReason = EReplicationBridgeDestroyInstanceReason::Destroy;
 		const EReplicationBridgeDestroyInstanceFlags DestroyFlags = IsAllowedToDestroyInstance(Instance) ? EReplicationBridgeDestroyInstanceFlags::AllowDestroyInstanceFromRemote : EReplicationBridgeDestroyInstanceFlags::None;
@@ -632,7 +630,7 @@ void UReplicationBridge::StopReplicatingNetRefHandle(FNetRefHandle Handle, EEndR
 	// This allows us the apply any received updates before we cut off this object
 	if (IsInReceiveUpdate())
 	{
-		UE_LOG(LogIrisBridge, Verbose, TEXT("Delayed request to StopReplicating %s (flags: %s) because it was called while inside ReceiveUpdate"), *NetRefHandleManager->PrintObjectFromNetRefHandle(Handle), *LexToString(EndReplicationFlags));
+		UE_LOG(LogIris, Verbose, TEXT("Delayed request to StopReplicating %s (flags: %s) because it was called while inside ReceiveUpdate"), *NetRefHandleManager->PrintObjectFromNetRefHandle(Handle), *LexToString(EndReplicationFlags));
 		
 		// Detect if we have diverging EndReplicationFlags for the same netobject
 #if DO_ENSURE
@@ -1017,7 +1015,7 @@ UE::Net::FNetRefHandle UReplicationBridge::StoreDestructionInfo(FNetRefHandle Ha
 	if (!FNetObjectFactoryRegistry::IsValidFactoryId(Parameters.NetFactoryId))
 	{
 		ensureMsgf(false, TEXT("StoreDestructionInfo received invalid NetFactoryId: %u for %s"), Parameters.NetFactoryId, *PrintObjectFromNetRefHandle(Handle));
-		UE_LOG(LogIrisBridge, Error, TEXT("StoreDestructionInfo received invalid NetFactoryID: %u for %s"), Parameters.NetFactoryId, *PrintObjectFromNetRefHandle(Handle));
+		UE_LOG(LogIris, Error, TEXT("StoreDestructionInfo received invalid NetFactoryID: %u for %s"), Parameters.NetFactoryId, *PrintObjectFromNetRefHandle(Handle));
 
 		return FNetRefHandle::GetInvalid();
 	}
diff --git a/Engine/Source/Runtime/Experimental/Iris/Core/Public/Iris/ReplicationSystem/ReplicationBridge.h b/Engine/Source/Runtime/Experimental/Iris/Core/Public/Iris/ReplicationSystem/ReplicationBridge.h
--- a/Engine/Source/Runtime/Experimental/Iris/Core/Public/Iris/ReplicationSystem/ReplicationBridge.h
+++ b/Engine/Source/Runtime/Experimental/Iris/Core/Public/Iris/ReplicationSystem/ReplicationBridge.h
@@ -55,13 +55,7 @@ namespace UE::Net
 	typedef TArray<FNetDependencyInfo, TInlineAllocator<32> > FNetDependencyInfoArray;
 }
 
-#if UE_BUILD_SHIPPING
-	IRISCORE_API DECLARE_LOG_CATEGORY_EXTERN(LogIrisBridge, All, Display);
-#else
-	IRISCORE_API DECLARE_LOG_CATEGORY_EXTERN(LogIrisBridge, Log, All);
-#endif
-
-#define UE_LOG_BRIDGEID(Category, Verbosity, Format, ...)  UE_LOG(Category, Verbosity, TEXT("ReplicationBridge(%u)::") Format, GetReplicationSystem()->GetId(), ##__VA_ARGS__)
+#define UE_LOG_BRIDGEID(Category, Verbosity, Format, ...)  UE_LOG(Category, Verbosity, TEXT("RepBridge(%u)::") Format, GetReplicationSystem()->GetId(), ##__VA_ARGS__)
 
 struct FReplicationBridgeSerializationContext
 {
