Optimize Chaos FCollectionPropertyMutableFacade Copy.

For typical cloth asset, this drops Proxy Initialize Config (part of GT tick) from ~90us to ~25us in PIE. On console the total Cloth Component Tick goes from ~100us to ~60us.

Tested locally in PIE and on a cooked build.

diff --git a/Engine/Plugins/ChaosCloth/Source/ChaosCloth/Private/ChaosCloth/ChaosClothingSimulationConfig.cpp b/Engine/Plugins/ChaosCloth/Source/ChaosCloth/Private/ChaosCloth/ChaosClothingSimulationConfig.cpp
--- a/Engine/Plugins/ChaosCloth/Source/ChaosCloth/Private/ChaosCloth/ChaosClothingSimulationConfig.cpp
+++ b/Engine/Plugins/ChaosCloth/Source/ChaosCloth/Private/ChaosCloth/ChaosClothingSimulationConfig.cpp
@@ -249,13 +249,19 @@ namespace Chaos
 
 	void FClothingSimulationConfig::Initialize(const TArray<TSharedPtr<const FManagedArrayCollection>>& InPropertyCollections)
 	{
-		PropertyCollections.Reset(InPropertyCollections.Num());
-		Properties.Reset(InPropertyCollections.Num());
+		PropertyCollections.SetNum(InPropertyCollections.Num());
+		Properties.SetNum(InPropertyCollections.Num());
 
-		for (const TSharedPtr<const FManagedArrayCollection>& InPropertyCollection : InPropertyCollections)
+		for (int32 Index = 0; Index < InPropertyCollections.Num(); ++Index)
 		{
-			TSharedPtr<FManagedArrayCollection>& PropertyCollection = PropertyCollections.Add_GetRef(MakeShared<FManagedArrayCollection>());
-			TUniquePtr<Softs::FCollectionPropertyMutableFacade>& Property = Properties.Add_GetRef(MakeUnique<Softs::FCollectionPropertyMutableFacade>(PropertyCollection));
+			const TSharedPtr<const FManagedArrayCollection>& InPropertyCollection = InPropertyCollections[Index];
+			TSharedPtr<FManagedArrayCollection>& PropertyCollection = PropertyCollections[Index];
+			TUniquePtr<Softs::FCollectionPropertyMutableFacade>& Property = Properties[Index];
+			if (!PropertyCollection || !Property)
+			{
+				PropertyCollection = MakeShared<FManagedArrayCollection>();
+				Property = MakeUnique<Softs::FCollectionPropertyMutableFacade>(PropertyCollection);
+			}
 			Property->Copy(*InPropertyCollection);
 		}
 	}
diff --git a/Engine/Source/Runtime/Experimental/Chaos/Private/Chaos/CollectionPropertyFacade.cpp b/Engine/Source/Runtime/Experimental/Chaos/Private/Chaos/CollectionPropertyFacade.cpp
--- a/Engine/Source/Runtime/Experimental/Chaos/Private/Chaos/CollectionPropertyFacade.cpp
+++ b/Engine/Source/Runtime/Experimental/Chaos/Private/Chaos/CollectionPropertyFacade.cpp
@@ -49,6 +49,7 @@ namespace Chaos::Softs
 	
 	void FCollectionPropertyConstFacade::RebuildKeyIndices()
 	{
+		TRACE_CPUPROFILER_EVENT_SCOPE(FCollectionPropertyConstFacade_RebuildKeyIndices);
 		// Create a fast access search map (although it might only be faster for a large enough number of properties)
 		const int32 NumKeys = KeyArray.Num();
 		KeyIndices.Empty(NumKeys);
@@ -178,6 +179,7 @@ namespace Chaos::Softs
 
 	void FCollectionPropertyMutableFacade::DefineSchema()
 	{
+		TRACE_CPUPROFILER_EVENT_SCOPE(FCollectionPropertyMutableFacade_DefineSchema);
 		GetManagedArrayCollection()->AddAttribute<FString>(PropertyFacadeNames::KeyName, PropertyFacadeNames::PropertyGroup);
 		GetManagedArrayCollection()->AddAttribute<FVector3f>(PropertyFacadeNames::LowValueName, PropertyFacadeNames::PropertyGroup);
 		GetManagedArrayCollection()->AddAttribute<FVector3f>(PropertyFacadeNames::HighValueName, PropertyFacadeNames::PropertyGroup);
@@ -257,15 +259,49 @@ namespace Chaos::Softs
 
 	void FCollectionPropertyMutableFacade::Copy(const FManagedArrayCollection& InManagedArrayCollection)
 	{
-		GetManagedArrayCollection()->Reset();
-		DefineSchema();
+		TRACE_CPUPROFILER_EVENT_SCOPE(FCollectionPropertyMutableFacade_Copy);
 
-		TArray<FName> GroupsToSkip = InManagedArrayCollection.GroupNames();
-		GroupsToSkip.RemoveSingleSwap(PropertyFacadeNames::PropertyGroup);
-
-		InManagedArrayCollection.CopyTo(GetManagedArrayCollection().Get(), GroupsToSkip);
-		UpdateArrays();
-		RebuildKeyIndices();
+		if (InManagedArrayCollection.HasAttribute(PropertyFacadeNames::KeyName, PropertyFacadeNames::PropertyGroup) &&
+			InManagedArrayCollection.HasAttribute(PropertyFacadeNames::LowValueName, PropertyFacadeNames::PropertyGroup) &&
+			InManagedArrayCollection.HasAttribute(PropertyFacadeNames::HighValueName, PropertyFacadeNames::PropertyGroup) &&
+			InManagedArrayCollection.HasAttribute(PropertyFacadeNames::StringValueName, PropertyFacadeNames::PropertyGroup) &&
+			InManagedArrayCollection.HasAttribute(PropertyFacadeNames::FlagsName, PropertyFacadeNames::PropertyGroup))
+		{
+			// Compare keys
+			bool bKeyIndicesMatch = false;
+			const TManagedArray<FString>* InKeyManagedArray = InManagedArrayCollection.FindAttributeTyped<FString>(PropertyFacadeNames::KeyName, PropertyFacadeNames::PropertyGroup);
+			const TManagedArray<FString>* OrigKeyManagedArray = GetManagedArrayCollection()->FindAttributeTyped<FString>(PropertyFacadeNames::KeyName, PropertyFacadeNames::PropertyGroup);
+			if (InKeyManagedArray && OrigKeyManagedArray && InKeyManagedArray->Num() == OrigKeyManagedArray->Num())
+			{
+				TRACE_CPUPROFILER_EVENT_SCOPE(FCollectionPropertyMutableFacade_Copy_CompareKeys);
+				bKeyIndicesMatch = true;
+				for (int32 Index = 0; Index < OrigKeyManagedArray->Num(); ++Index)
+				{
+					if (!(*InKeyManagedArray)[Index].Equals((*OrigKeyManagedArray)[Index], ESearchCase::CaseSensitive))
+					{
+						bKeyIndicesMatch = false;
+						break;
+					}
+				}
+			}
+			if (!bKeyIndicesMatch)
+			{
+				GetManagedArrayCollection()->CopyAttribute(InManagedArrayCollection, PropertyFacadeNames::KeyName, PropertyFacadeNames::PropertyGroup);
+			}
+			GetManagedArrayCollection()->CopyAttribute(InManagedArrayCollection, PropertyFacadeNames::LowValueName, PropertyFacadeNames::PropertyGroup);
+			GetManagedArrayCollection()->CopyAttribute(InManagedArrayCollection, PropertyFacadeNames::HighValueName, PropertyFacadeNames::PropertyGroup);
+			GetManagedArrayCollection()->CopyAttribute(InManagedArrayCollection, PropertyFacadeNames::StringValueName, PropertyFacadeNames::PropertyGroup);
+			GetManagedArrayCollection()->CopyAttribute(InManagedArrayCollection, PropertyFacadeNames::FlagsName, PropertyFacadeNames::PropertyGroup);
+			UpdateArrays();
+			if (!bKeyIndicesMatch)
+			{
+				RebuildKeyIndices();
+			}
+		}
+		else
+		{
+			GetManagedArrayCollection()->Reset();
+		}
 	}
 
 	void FCollectionPropertyMutableFacade::Update(const TSharedPtr<const FManagedArrayCollection>& InManagedArrayCollection, ECollectionPropertyUpdateFlags UpdateFlags)
