Fixed on-demand install cache hit ratio trace counter.

#rnx

diff --git a/Engine/Source/Runtime/Experimental/IoStore/OnDemand/Private/Statistics.cpp b/Engine/Source/Runtime/Experimental/IoStore/OnDemand/Private/Statistics.cpp
--- a/Engine/Source/Runtime/Experimental/IoStore/OnDemand/Private/Statistics.cpp
+++ b/Engine/Source/Runtime/Experimental/IoStore/OnDemand/Private/Statistics.cpp
@@ -1141,7 +1141,7 @@ void FOnDemandContentInstallerStats::OnRequestCompleted(
 		InstallerTraceCounters.InstallDurationMs.Increment(FPlatformTime::ToMilliseconds64(DurationCycles));
 		InstallerTraceCounters.CacheHitRatio.Increment(CacheHitRatio);
 		InstallerTraceCounters.AvgInstallDurationMs.Set(InstallerTraceCounters.InstallDurationMs.GetMean());
-		InstallerTraceCounters.AvgCacheHitRatio.Set(InstallerTraceCounters.InstallDurationMs.GetMean());
+		InstallerTraceCounters.AvgCacheHitRatio.Set(InstallerTraceCounters.CacheHitRatio.GetMean());
 	}
 #endif // COUNTERSTRACE_ENABLED
 
