Strategic consolidation of related code to be behind a common interface.

#rnx

diff --git a/Engine/Source/Runtime/Experimental/IoStore/HttpClient/Private/Activity.inl b/Engine/Source/Runtime/Experimental/IoStore/HttpClient/Private/Activity.inl
--- a/Engine/Source/Runtime/Experimental/IoStore/HttpClient/Private/Activity.inl
+++ b/Engine/Source/Runtime/Experimental/IoStore/HttpClient/Private/Activity.inl
@@ -7,6 +7,76 @@ namespace UE::IoStore::HTTP
 
 // {{{1 activity ...............................................................
 
+////////////////////////////////////////////////////////////////////////////////
+class FRequestCloak
+	: public FBuffer
+{
+public:
+	using FBuffer::FBuffer;
+
+	void			Begin(FAnsiStringView Method, FAnsiStringView Path);
+	void			AddHeader(FAnsiStringView Key, FAnsiStringView Value);
+	void			End();
+	FAnsiStringView	GetMethod() const;
+	void			EnumerateHeaders(FResponse::FHeaderSink Sink) const;
+
+private:
+	FRequestCloak&	operator << (FAnsiStringView Value);
+	uint16			HeaderLeft;
+	int16			HeaderRight = -1;
+	int8			MethodLength = -1;
+};
+
+////////////////////////////////////////////////////////////////////////////////
+FAnsiStringView FRequestCloak::GetMethod() const
+{
+	check(MethodLength > 0);
+	return FAnsiStringView(GetData(), MethodLength);
+}
+
+////////////////////////////////////////////////////////////////////////////////
+void FRequestCloak::EnumerateHeaders(FResponse::FHeaderSink Sink) const
+{
+	check(HeaderRight >= 2);
+	int32 HeaderLength = HeaderRight - HeaderLeft;
+	FAnsiStringView Headers(GetData() + HeaderLeft, HeaderLength);
+	UE::IoStore::HTTP::EnumerateHeaders(Headers, Sink);
+}
+
+////////////////////////////////////////////////////////////////////////////////
+void FRequestCloak::Begin(FAnsiStringView Method, FAnsiStringView Path)
+{
+	check(MethodLength < 0);
+	*this << Method << " " << Path << " HTTP/1.1" "\r\n";
+	MethodLength = int8(Method.Len());
+	HeaderLeft = uint16(GetSize());
+}
+
+////////////////////////////////////////////////////////////////////////////////
+void FRequestCloak::AddHeader(FAnsiStringView Key, FAnsiStringView Value)
+{
+	*this << Key << ":" << Value << "\r\n";
+}
+
+////////////////////////////////////////////////////////////////////////////////
+void FRequestCloak::End()
+{
+	*this << "\r\n";
+	HeaderRight = int16(GetSize());
+}
+
+////////////////////////////////////////////////////////////////////////////////
+FRequestCloak& FRequestCloak::operator << (FAnsiStringView Value)
+{
+	uint32 Length = uint32(Value.Len());
+	FBuffer::FMutableSection Section = GetMutableFree(Length);
+	::memcpy(Section.Data, Value.GetData(), Length);
+	AdvanceUsed(Length);
+	return *this;
+}
+
+
+
 #if IAS_HTTP_WITH_PERF
 
 ////////////////////////////////////////////////////////////////////////////////
@@ -104,7 +174,7 @@ struct alignas(16) FActivity
 	UPTRINT				SinkParam;
 	FTicketSink			Sink;
 	FResponseInternal	Response;
-	FBuffer				Buffer;
+	FRequestCloak		Buffer;
 };
 
 ////////////////////////////////////////////////////////////////////////////////
@@ -172,7 +242,7 @@ static FActivity* Activity_Alloc(uint32 BufferSize)
 
 	auto* Scratch = (char*)(Activity + 1);
 	uint32 ScratchSize = BufferSize;
-	Activity->Buffer = FBuffer(Scratch, ScratchSize);
+	Activity->Buffer = decltype(Activity->Buffer)(Scratch, ScratchSize);
 
 	return Activity;
 }
diff --git a/Engine/Source/Runtime/Experimental/IoStore/HttpClient/Private/Api.inl b/Engine/Source/Runtime/Experimental/IoStore/HttpClient/Private/Api.inl
--- a/Engine/Source/Runtime/Experimental/IoStore/HttpClient/Private/Api.inl
+++ b/Engine/Source/Runtime/Experimental/IoStore/HttpClient/Private/Api.inl
@@ -52,7 +52,7 @@ FRequest&& FRequest::Accept(FAnsiStringView MimeType)
 FRequest&& FRequest::Header(FAnsiStringView Key, FAnsiStringView Value)
 {
 	checkSlow(Ptr->State == FActivity::EState::Build);
-	FMessageBuilder(Ptr->Buffer) << Key << ": " << Value << "\r\n";
+	Ptr->Buffer.AddHeader(Key, Value);
 	return MoveTemp(*this);
 }
 
diff --git a/Engine/Source/Runtime/Experimental/IoStore/HttpClient/Private/Loop.inl b/Engine/Source/Runtime/Experimental/IoStore/HttpClient/Private/Loop.inl
--- a/Engine/Source/Runtime/Experimental/IoStore/HttpClient/Private/Loop.inl
+++ b/Engine/Source/Runtime/Experimental/IoStore/HttpClient/Private/Loop.inl
@@ -1250,15 +1250,14 @@ FRequest FEventLoop::FImpl::Request(
 
 	Activity->NoContent = (Method == "HEAD");
 
-	FMessageBuilder Builder(Activity->Buffer);
-
-	Builder << Method << " " << Path << " HTTP/1.1" "\r\n"
-		"Host: " << Activity->Host->GetHostName() << "\r\n";
+	auto& Buffer = Activity->Buffer;
+	Buffer.Begin(Method, Path);
+	Buffer.AddHeader("Host", Activity->Host->GetHostName());
 
 	// HTTP/1.1 is persistent by default thus "Connection" header isn't required
 	if (!Activity->IsKeepAlive)
 	{
-		Builder << "Connection: close\r\n";
+		Buffer.AddHeader("Connection", "close");
 	}
 
 	FRequest Ret;
@@ -1271,7 +1270,7 @@ FTicket FEventLoop::FImpl::Send(FActivity* Activity)
 {
 	Trace(Activity, ETrace::RequestBegin);
 
-	FMessageBuilder(Activity->Buffer) << "\r\n";
+	Activity->Buffer.End();
 	Activity_ChangeState(Activity, FActivity::EState::Send);
 
 	uint64 Slot;
@@ -1617,16 +1616,7 @@ bool FEventLoop::Redirect(const FTicketStatus& Status, FTicketSink& OuterSink)
 	const auto& Activity = (FActivity&)Response; // todo: yuk
 
 	// Original method should remain unchanged
-	const char* Data = Activity.Buffer.GetData();
-	FAnsiStringView Method;
-	for (uint32 i = 0; i < 5; ++i)
-	{
-		if (Data[i] <= ' ')
-		{
-			Method = FAnsiStringView(Data, i);
-			break;
-		}
-	}
+	FAnsiStringView Method = Activity.Buffer.GetMethod();
 	check(!Method.IsEmpty());
 
 	FRequest ForwardRequest;
@@ -1658,21 +1648,7 @@ bool FEventLoop::Redirect(const FTicketStatus& Status, FTicketSink& OuterSink)
 
 	// Transfer original request headers
 	check(Activity.State == FActivity::EState::RecvMessage);
-	uint32 Length = 0;
-	for (const char* End = Data + Activity.StateParam; Data + Length < End; ++Length)
-	{
-		if (Data[Length] != '\n')
-		{
-			continue;
-		}
-
-		Data = Data + Length + 1;
-		Length = uint32(ptrdiff_t(End - Data));
-		break;
-	}
-
-	FAnsiStringView OriginalHeaders(Data, Length);
-	EnumerateHeaders(OriginalHeaders, [&ForwardRequest] (FAnsiStringView Name, FAnsiStringView Value)
+	Activity.Buffer.EnumerateHeaders([&ForwardRequest] (FAnsiStringView Name, FAnsiStringView Value)
 	{
 		if (Name != "Host" && Name != "Connection")
 		{
