Deleted ununsed class.

#rnx

diff --git a/Engine/Source/Runtime/Experimental/IoStore/HttpClient/Private/Misc.inl b/Engine/Source/Runtime/Experimental/IoStore/HttpClient/Private/Misc.inl
--- a/Engine/Source/Runtime/Experimental/IoStore/HttpClient/Private/Misc.inl
+++ b/Engine/Source/Runtime/Experimental/IoStore/HttpClient/Private/Misc.inl
@@ -617,35 +617,6 @@ void FBuffer::Extend(uint32 AtLeast, uint32 PageSize)
 
 
 
-////////////////////////////////////////////////////////////////////////////////
-class FMessageBuilder
-{
-public:
-						FMessageBuilder(FBuffer& Ruffeb);
-	FMessageBuilder&	operator << (FAnsiStringView Lhs);
-
-private:
-	FBuffer&		Buffer;
-};
-
-////////////////////////////////////////////////////////////////////////////////
-FMessageBuilder::FMessageBuilder(FBuffer& Ruffeb)
-: Buffer(Ruffeb)
-{
-}
-
-////////////////////////////////////////////////////////////////////////////////
-FMessageBuilder& FMessageBuilder::operator << (FAnsiStringView Lhs)
-{
-	uint32 Length = uint32(Lhs.Len());
-	FBuffer::FMutableSection Section = Buffer.GetMutableFree(Length);
-	::memcpy(Section.Data, Lhs.GetData(), Length);
-	Buffer.AdvanceUsed(Length);
-	return *this;
-}
-
-
-
 // {{{1 throttler ..............................................................
 
 ////////////////////////////////////////////////////////////////////////////////
