Removed unused fold markers

#rnx

diff --git a/Engine/Source/Runtime/Experimental/IoStore/HttpClient/Private/Activity.inl b/Engine/Source/Runtime/Experimental/IoStore/HttpClient/Private/Activity.inl
--- a/Engine/Source/Runtime/Experimental/IoStore/HttpClient/Private/Activity.inl
+++ b/Engine/Source/Runtime/Experimental/IoStore/HttpClient/Private/Activity.inl
@@ -5,8 +5,6 @@
 namespace UE::IoStore::HTTP
 {
 
-// {{{1 activity ...............................................................
-
 ////////////////////////////////////////////////////////////////////////////////
 class FRequestCloak
 	: public FBuffer
diff --git a/Engine/Source/Runtime/Experimental/IoStore/HttpClient/Private/ConnectionPool.inl b/Engine/Source/Runtime/Experimental/IoStore/HttpClient/Private/ConnectionPool.inl
--- a/Engine/Source/Runtime/Experimental/IoStore/HttpClient/Private/ConnectionPool.inl
+++ b/Engine/Source/Runtime/Experimental/IoStore/HttpClient/Private/ConnectionPool.inl
@@ -5,8 +5,6 @@
 namespace UE::IoStore::HTTP
 {
 
-// {{{1 connection-pool ........................................................
-
 ////////////////////////////////////////////////////////////////////////////////
 class FHost
 {
diff --git a/Engine/Source/Runtime/Experimental/IoStore/HttpClient/Private/Socks.inl b/Engine/Source/Runtime/Experimental/IoStore/HttpClient/Private/Socks.inl
--- a/Engine/Source/Runtime/Experimental/IoStore/HttpClient/Private/Socks.inl
+++ b/Engine/Source/Runtime/Experimental/IoStore/HttpClient/Private/Socks.inl
@@ -5,8 +5,6 @@
 namespace UE::IoStore::HTTP
 {
 
-// {{{1 socks ..................................................................
-
 #if !UE_BUILD_SHIPPING
 
 ///////////////////////////////////////////////////////////////////////////////
diff --git a/Engine/Source/Runtime/Experimental/IoStore/HttpClient/Private/Test.inl b/Engine/Source/Runtime/Experimental/IoStore/HttpClient/Private/Test.inl
--- a/Engine/Source/Runtime/Experimental/IoStore/HttpClient/Private/Test.inl
+++ b/Engine/Source/Runtime/Experimental/IoStore/HttpClient/Private/Test.inl
@@ -10,8 +10,6 @@
 namespace UE::IoStore::HTTP
 {
 
-// {{{1 test ...................................................................
-
 #if !(UE_BUILD_SHIPPING|UE_BUILD_TEST) && !defined(NO_UE_INCLUDES)
 
 ////////////////////////////////////////////////////////////////////////////////
