Don't count cancelled operations as errors in on-demand content installer analytics.

diff --git a/Engine/Source/Runtime/Experimental/IoStore/OnDemand/Private/Statistics.cpp b/Engine/Source/Runtime/Experimental/IoStore/OnDemand/Private/Statistics.cpp
--- a/Engine/Source/Runtime/Experimental/IoStore/OnDemand/Private/Statistics.cpp
+++ b/Engine/Source/Runtime/Experimental/IoStore/OnDemand/Private/Statistics.cpp
@@ -1147,7 +1147,7 @@ void FOnDemandContentInstallerStats::OnRequestCompleted(
 
 	{
 		TUniqueLock Lock(InstallerAnalytics);
-		if (ErrorCode != EIoErrorCode::Ok)
+		if (ErrorCode != EIoErrorCode::Ok && ErrorCode != EIoErrorCode::Cancelled)
 		{
 			InstallerAnalytics.InstallErrorCount++;
 		}
@@ -1218,6 +1218,11 @@ void FOnDemandInstallCacheStats::OnStartupError(EIoErrorCode ErrorCode)
 
 void FOnDemandInstallCacheStats::OnFlush(EIoErrorCode ErrorCode, int64 ByteCount)
 {
+	if (ErrorCode == EIoErrorCode::Cancelled)
+	{
+		return;
+	}
+
 	UE::TUniqueLock Lock(InstallCacheAnalytics);
 	if (ErrorCode != EIoErrorCode::Ok)
 	{
@@ -1229,6 +1234,11 @@ void FOnDemandInstallCacheStats::OnFlush(EIoErrorCode ErrorCode, int64 ByteCount
 
 void FOnDemandInstallCacheStats::OnJournalCommit(EIoErrorCode ErrorCode, int64 ByteCount)
 {
+	if (ErrorCode == EIoErrorCode::Cancelled)
+	{
+		return;
+	}
+
 	UE::TUniqueLock Lock(InstallCacheAnalytics);
 	if (ErrorCode != EIoErrorCode::Ok)
 	{
@@ -1250,6 +1260,11 @@ void FOnDemandInstallCacheStats::OnPurge(
 	uint64 BytesToPurge,
 	uint64 PurgedBytes)
 {
+	if (ErrorCode == EIoErrorCode::Cancelled)
+	{
+		return;
+	}
+
 	UE::TUniqueLock Lock(InstallCacheAnalytics);
 	if (ErrorCode != EIoErrorCode::Ok)
 	{
@@ -1261,6 +1276,11 @@ void FOnDemandInstallCacheStats::OnPurge(
 
 void FOnDemandInstallCacheStats::OnDefrag(EIoErrorCode ErrorCode, uint64 FragmentedBytes)
 {
+	if (ErrorCode == EIoErrorCode::Cancelled)
+	{
+		return;
+	}
+
 	UE::TUniqueLock Lock(InstallCacheAnalytics);
 	if (ErrorCode != EIoErrorCode::Ok)
 	{
