Iris - Fixed bug that allowed adding subobjects to torn off root objects.

diff --git a/Engine/Plugins/Runtime/ReplicationSystemTestPlugin/Source/Private/Tests/ReplicationSystem/TestReplicationSystem.cpp b/Engine/Plugins/Runtime/ReplicationSystemTestPlugin/Source/Private/Tests/ReplicationSystem/TestReplicationSystem.cpp
--- a/Engine/Plugins/Runtime/ReplicationSystemTestPlugin/Source/Private/Tests/ReplicationSystem/TestReplicationSystem.cpp
+++ b/Engine/Plugins/Runtime/ReplicationSystemTestPlugin/Source/Private/Tests/ReplicationSystem/TestReplicationSystem.cpp
@@ -1629,6 +1629,43 @@ UE_NET_TEST_FIXTURE(FReplicationSystemServerClientTestFixture, TestImmediateTear
 	UE_NET_ASSERT_TRUE(Cast<UTestReplicatedIrisObject>(Client->GetReplicationBridge()->GetReplicatedObject(ServerSubObject->NetRefHandle)) == nullptr);
 }
 
+// Test creating a subobject after the root is already torn off
+UE_NET_TEST_FIXTURE(FReplicationSystemServerClientTestFixture, TestCreatingSubobjectAfterRootObjectTearOff)
+{
+	//LOG_SCOPE_VERBOSITY_OVERRIDE(LogIrisBridge, ELogVerbosity::Verbose);
+
+	// Add a client
+	FReplicationSystemTestClient* Client = CreateClient();
+
+	// Spawn object with subobject on server
+	UTestReplicatedIrisObject* ServerObject = Server->CreateObject();
+	UTestReplicatedIrisObject* ServerSubObject0 = Server->CreateSubObject(ServerObject->NetRefHandle);
+
+	// Send and deliver packet
+	Server->UpdateAndSend({Client});
+
+	// TearOff the object
+	Server->ReplicationSystem->TearOffNextUpdate(ServerObject->NetRefHandle);
+
+	// Send and deliver packet
+	Server->NetUpdate();
+	Server->PostSendUpdate();
+
+	{
+		FTestEnsureScope EnsureScope;
+
+		// Add a new subobject. This shouldn't become a replicated subobject of the already torn off root object.
+		UTestReplicatedIrisObject* ServerSubObject1 = Server->CreateSubObject(ServerObject->NetRefHandle);
+		UE_NET_ASSERT_EQ(EnsureScope.GetCount(), 0);
+		// The subobject should not start replicating.
+		UE_NET_ASSERT_FALSE(ServerSubObject1->NetRefHandle.IsValid());
+
+		// Send and deliver packet
+		Server->UpdateAndSend({ Client });
+		UE_NET_ASSERT_EQ(EnsureScope.GetCount(), 0);
+	}
+}
+
 // Test dropped creation of subobject dirties owner
 UE_NET_TEST_FIXTURE(FReplicationSystemServerClientTestFixture, TestDroppedCreationForSubobjectDirtiesOwner)
 {
diff --git a/Engine/Source/Runtime/Engine/Private/Net/Iris/ReplicationSystem/EngineReplicationBridge.cpp b/Engine/Source/Runtime/Engine/Private/Net/Iris/ReplicationSystem/EngineReplicationBridge.cpp
--- a/Engine/Source/Runtime/Engine/Private/Net/Iris/ReplicationSystem/EngineReplicationBridge.cpp
+++ b/Engine/Source/Runtime/Engine/Private/Net/Iris/ReplicationSystem/EngineReplicationBridge.cpp
@@ -435,7 +435,6 @@ UE::Net::FNetRefHandle UEngineReplicationBridge::StartReplicatingComponent(FNetR
 
 	if (!ReplicatedComponentHandle.IsValid())
 	{
-		ensureMsgf(false, TEXT("Failed to create or find NetRefHandle for ActorComponent Named %s"), ToCStr(SubObject->GetName()));
 		return FNetRefHandle::GetInvalid();
 	}
 
diff --git a/Engine/Source/Runtime/Experimental/Iris/Core/Private/Iris/ReplicationSystem/ObjectReplicationBridge.cpp b/Engine/Source/Runtime/Experimental/Iris/Core/Private/Iris/ReplicationSystem/ObjectReplicationBridge.cpp
--- a/Engine/Source/Runtime/Experimental/Iris/Core/Private/Iris/ReplicationSystem/ObjectReplicationBridge.cpp
+++ b/Engine/Source/Runtime/Experimental/Iris/Core/Private/Iris/ReplicationSystem/ObjectReplicationBridge.cpp
@@ -578,7 +578,7 @@ void UObjectReplicationBridge::AssignDynamicFilter(UObject* Instance, const FRoo
 	
 }
 
-UE::Net::FNetRefHandle UObjectReplicationBridge::StartReplicatingSubObject(UObject* Instance, const FSubObjectReplicationParams& Params , UE::Net::FNetObjectFactoryId NetFactoryId)
+UE::Net::FNetRefHandle UObjectReplicationBridge::StartReplicatingSubObject(UObject* Instance, const FSubObjectReplicationParams& Params, UE::Net::FNetObjectFactoryId NetFactoryId)
 {
 	LLM_SCOPE_BYTAG(IrisState);
 
@@ -588,11 +588,23 @@ UE::Net::FNetRefHandle UObjectReplicationBridge::StartReplicatingSubObject(UObje
 	const FReplicationSystemInternal* ReplicationSystemInternal = GetReplicationSystem()->GetReplicationSystemInternal();
 	const FNetRefHandleManager& LocalNetRefHandleManager = ReplicationSystemInternal->GetNetRefHandleManager();
 
-	checkf(IsReplicatedHandle(Params.RootObjectHandle), TEXT("Owner %s (%s) must be replicated for subobject %s to replicate."), 
-		*GetNameSafe(LocalNetRefHandleManager.GetReplicatedObjectInstance(LocalNetRefHandleManager.GetInternalIndex(Params.RootObjectHandle))), *Params.RootObjectHandle.ToString(), *GetNameSafe(Instance));
+	const FInternalNetRefIndex RootObjectIndex = LocalNetRefHandleManager.GetInternalIndex(Params.RootObjectHandle);
 
-	checkf(!Instance->HasAnyFlags(RF_ArchetypeObject | RF_ClassDefaultObject), TEXT("Iris cannot replicate subobject %s owned by %s because it's an %s"), 
-		*GetNameSafe(Instance), *GetNameSafe(LocalNetRefHandleManager.GetReplicatedObjectInstance(LocalNetRefHandleManager.GetInternalIndex(Params.RootObjectHandle))), Instance->HasAnyFlags(RF_ArchetypeObject)?TEXT("Archetype"):TEXT("DefaultObject"));
+	checkf(!Instance->IsTemplate(), TEXT("Iris cannot replicate subobject %s owned by %s because it's an %s"), 
+		*GetNameSafe(Instance), *GetNameSafe(LocalNetRefHandleManager.GetReplicatedObjectInstance(RootObjectIndex)), Instance->HasAnyFlags(RF_ArchetypeObject)?TEXT("Archetype"):TEXT("DefaultObject"));
+
+	if (!ensureMsgf(IsReplicatedHandle(Params.RootObjectHandle), TEXT("Owner %s (%s) must be replicated for subobject %s to replicate."), 
+		*GetNameSafe(LocalNetRefHandleManager.GetReplicatedObjectInstance(RootObjectIndex)), *Params.RootObjectHandle.ToString(), *GetNameSafe(Instance)))
+	{
+		return FNetRefHandle::GetInvalid();
+	}
+
+	if (!LocalNetRefHandleManager.IsScopableIndex(RootObjectIndex))
+	{
+		UE_LOG_OBJECTREPLICATIONBRIDGE(Warning, TEXT("Owner %s (%s) must be in scope for subobject %s to replicate."), 
+			*GetNameSafe(LocalNetRefHandleManager.GetReplicatedObjectInstance(RootObjectIndex)), *Params.RootObjectHandle.ToString(), *GetNameSafe(Instance))
+		return FNetRefHandle::GetInvalid();
+	}
 
 	// Pre-registering subobjects is not supported
 	FNetRefHandle PotentiallyPreRegisteredHandle = GetObjectReferenceCache()->GetObjectReferenceHandleFromObject(Instance);
diff --git a/Engine/Source/Runtime/Experimental/Iris/Core/Private/Iris/ReplicationSystem/ReplicationOperationsInternal.cpp b/Engine/Source/Runtime/Experimental/Iris/Core/Private/Iris/ReplicationSystem/ReplicationOperationsInternal.cpp
--- a/Engine/Source/Runtime/Experimental/Iris/Core/Private/Iris/ReplicationSystem/ReplicationOperationsInternal.cpp
+++ b/Engine/Source/Runtime/Experimental/Iris/Core/Private/Iris/ReplicationSystem/ReplicationOperationsInternal.cpp
@@ -116,7 +116,7 @@ uint32 FReplicationInstanceOperationsInternal::QuantizeObjectStateData(FNetBitSt
 			{
 				const bool bIsOwnerScopable = NetRefHandleManager.IsScopableIndex(SubObjectOwnerIndex);
 				// Dependent objects should not ensure if the owner isn't scopable. Subobjects pending tear off is ok too.
-				ensureAlwaysMsgf(bIsOwnerScopable || Object.IsDependentObject() || Object.bTearOff, TEXT("SubObject ( InternaIndex: %u ) with NetRefHandle (Id=%" UINT64_FMT ") is trying to dirty parent ( InternalIndex: %u ) not in scope."), InternalIndex, Object.RefHandle.GetId(), SubObjectOwnerIndex);
+				ensureMsgf(bIsOwnerScopable || Object.IsDependentObject() || Object.bTearOff, TEXT("SubObject ( InternaIndex: %u ) with NetRefHandle (Id=%" UINT64_FMT ") is trying to dirty parent ( InternalIndex: %u ) not in scope."), InternalIndex, Object.RefHandle.GetId(), SubObjectOwnerIndex);
 				if (bIsOwnerScopable)
 				{
 					// Do we want to control this separately for subobjects? Or should they respect the setting on the owner?
