Deleted unused local variable

#rnx

diff --git a/Engine/Source/Runtime/Experimental/IoStore/HttpClient/Private/Socket.inl b/Engine/Source/Runtime/Experimental/IoStore/HttpClient/Private/Socket.inl
--- a/Engine/Source/Runtime/Experimental/IoStore/HttpClient/Private/Socket.inl
+++ b/Engine/Source/Runtime/Experimental/IoStore/HttpClient/Private/Socket.inl
@@ -614,7 +614,6 @@ static void Trace(UPTRINT Socket, ETrace Action, const FOutcome* Outcome)
 
 	if (Outcome == nullptr)
 	{
-		uint32 ScopeId = (Action == ETrace::Send) ? SendScopeId : RecvScopeId;
 		LaneTrace_Change(Lane, RecvScopeId);
 		return;
 	}
