Deleted unused file-static functions.

#rnx

diff --git a/Engine/Source/Runtime/Experimental/IoStore/HttpClient/Private/Peer.inl b/Engine/Source/Runtime/Experimental/IoStore/HttpClient/Private/Peer.inl
--- a/Engine/Source/Runtime/Experimental/IoStore/HttpClient/Private/Peer.inl
+++ b/Engine/Source/Runtime/Experimental/IoStore/HttpClient/Private/Peer.inl
@@ -345,15 +345,11 @@ struct SSL;
 static void		Ssl_ContextDestroy(...)		{}
 static UPTRINT	Ssl_ContextCreate(...)		{ return 0; }
 static int32	Ssl_ContextCertNum(...)		{ return 0; }
-static int32	Ssl_BioWrite(...)			{ return -1; }
-static int32	Ssl_BioRead(...)			{ return -1; }
-static long		Ssl_BioControl(...)			{ return 0; }
 static SSL*		Ssl_Create(...)				{ return nullptr; }
 static void		Ssl_Destroy(...)			{}
 static void		Ssl_AssociatePeer(...)		{}
 static void		Ssl_SetupAlpn(...)			{}
 static int32	Ssl_GetProtocolVersion(...)	{ return -1; }
-static FOutcome	Ssl_GetOutcome(...)			{ return FOutcome::Error("!impl"); }
 static FOutcome	Ssl_Handshake(...)			{ return FOutcome::Error("!impl"); }
 static FOutcome	Ssl_Write(...)				{ return FOutcome::Error("!impl"); }
 static FOutcome	Ssl_Read(...)				{ return FOutcome::Error("!impl"); }
