Modular vehicle: Pass simulation output data back to verse

#okforversepublic

diff --git a/Engine/Source/Runtime/Experimental/ChaosVehicles/ChaosVehiclesCore/Private/SimModule/SimulationModuleBase.cpp b/Engine/Source/Runtime/Experimental/ChaosVehicles/ChaosVehiclesCore/Private/SimModule/SimulationModuleBase.cpp
--- a/Engine/Source/Runtime/Experimental/ChaosVehicles/ChaosVehiclesCore/Private/SimModule/SimulationModuleBase.cpp
+++ b/Engine/Source/Runtime/Experimental/ChaosVehicles/ChaosVehiclesCore/Private/SimModule/SimulationModuleBase.cpp
@@ -146,6 +146,7 @@ void FSimOutputData::FillOutputState(const ISimulationModuleBase* SimModule)
 	{
 		AnimationData = SimModule->GetAnimationData();
 		AnimationData.AnimationLocOffset += SimModule->GetAnimationOffset();
+		ModuleGuid = SimModule->GetGuid();
 	}
 
 #if !(UE_BUILD_SHIPPING || UE_BUILD_TEST)
diff --git a/Engine/Source/Runtime/Experimental/ChaosVehicles/ChaosVehiclesCore/Public/SimModule/SimulationModuleBase.h b/Engine/Source/Runtime/Experimental/ChaosVehicles/ChaosVehiclesCore/Public/SimModule/SimulationModuleBase.h
--- a/Engine/Source/Runtime/Experimental/ChaosVehicles/ChaosVehiclesCore/Public/SimModule/SimulationModuleBase.h
+++ b/Engine/Source/Runtime/Experimental/ChaosVehicles/ChaosVehiclesCore/Public/SimModule/SimulationModuleBase.h
@@ -309,7 +309,7 @@ namespace Chaos
 		{}
 		virtual ~ISimulationModuleBase() {}
 
-		const int GetGuid() { return Guid; }
+		int GetGuid() const { return Guid; }
 		void SetGuid(int GuidIn) { Guid = GuidIn; }
 		/**
 		* Get the friendly name for this module, primarily for logging & debugging module tree
@@ -564,6 +564,7 @@ namespace Chaos
 
 		bool bEnabled = true;
 		FSimModuleAnimationData AnimationData;
+		int ModuleGuid = -1;
 
 #if !(UE_BUILD_SHIPPING || UE_BUILD_TEST)
 		virtual FString ToString() { return FString(); }
