FOnDemandInstallCache
- Fix AddReferencesToBlocks not using aligned disk size for chunks.
- Defrag now consolidates small blocks even if they aren't fragmented.

diff --git a/Engine/Source/Runtime/Experimental/IoStore/OnDemand/Private/OnDemandInstallCache.cpp b/Engine/Source/Runtime/Experimental/IoStore/OnDemand/Private/OnDemandInstallCache.cpp
--- a/Engine/Source/Runtime/Experimental/IoStore/OnDemand/Private/OnDemandInstallCache.cpp
+++ b/Engine/Source/Runtime/Experimental/IoStore/OnDemand/Private/OnDemandInstallCache.cpp
@@ -184,7 +184,8 @@ struct FCas
 	FLastAccess				LastAccess;
 	FReadHandles			ReadHandles;
 	FEventRef				BlockReadsDoneEvent;
-	uint32					MaxBlockSize = 32 << 20; //TODO: Make configurable	
+	const uint32			MaxBlockSize = 32 << 20; //TODO: Make configurable	
+	const uint32			MinBlockSize = MaxBlockSize >> 1; //TODO: Make configurable	
 	mutable UE::FMutex		Mutex;
 };
 
@@ -2401,9 +2402,11 @@ uint64 FOnDemandInstallCache::AddReferencesToBlocks(
 			const bool bToInstall = ChunksToInstall.Contains(Entry.Hash);
 			const bool bIsReferenced = IsReferenced[EntryIndex];
 
+			uint64 ChunkDiskSize = Align(int64(Entry.EncodedSize), FAES::AESBlockSize);
+
 			if (bIsReferenced)
 			{
-				OutTotalReferencedBytes += Entry.EncodedSize;
+				OutTotalReferencedBytes += ChunkDiskSize;
 			}
 
 			FCasBlockInfo* BlockInfo = nullptr;
@@ -2427,11 +2430,11 @@ uint64 FOnDemandInstallCache::AddReferencesToBlocks(
 
 			if (BlockInfo)
 			{
-				BlockInfo->RefSize += Entry.EncodedSize;
+				BlockInfo->RefSize += ChunkDiskSize;
 			}
 			else if (bToInstall)
 			{
-				TotalUncachedBytes += Entry.EncodedSize;
+				TotalUncachedBytes += ChunkDiskSize;
 			}
 
 			EntryIndex++;
@@ -2501,6 +2504,11 @@ FIoStatus FOnDemandInstallCache::Defrag(
 	FCasBlockInfoMap& BlockInfo, 
 	const uint64* TotalBytesToFree /*= nullptr*/)
 {
+	if (TotalBytesToFree && *TotalBytesToFree == 0)
+	{
+		return FIoStatus::Ok;
+	}
+
 	const uint64 TotalCachedBytes = Algo::TransformAccumulate(BlockInfo,
 		[](const TPair<FCasBlockId, FCasBlockInfo>& Kv) { return Kv.Value.FileSize; },
 		uint64(0));
@@ -2554,8 +2562,9 @@ FIoStatus FOnDemandInstallCache::Defrag(
 			const FCasBlockId BlockId = Kv.Key;
 			const FCasBlockInfo& Info = Kv.Value;
 
-			if (Info.RefSize < Info.FileSize)
+			if (!bPossibleToFreeBytes && Info.RefSize < Info.FileSize)
 			{
+				// Block is fragmented
 				FragmentedBytes += (Info.FileSize - Info.RefSize);
 				TotalBlockSize += Info.FileSize;
 
@@ -2564,12 +2573,23 @@ FIoStatus FOnDemandInstallCache::Defrag(
 
 				BlocksToDefrag.Add(FDefragBlock{ .BlockId = BlockId, .LastAccess = Info.LastAccess });
 
-				if (FreedBlockBytes >= NewBlockBytes && FreedBlockBytes - NewBlockBytes > *TotalBytesToFree)
+				if (FreedBlockBytes >= NewBlockBytes && FreedBlockBytes - NewBlockBytes >= *TotalBytesToFree)
 				{
 					bPossibleToFreeBytes = true;
-					break;
 				}
 			}
+			else if (Info.FileSize < Cas.MinBlockSize)
+			{
+				// Block is too small whether or not its fragmented
+				if (ensure(Info.RefSize <= Info.FileSize))
+				{
+					FragmentedBytes += (Info.FileSize - Info.RefSize);
+				}
+				
+				TotalBlockSize += Info.FileSize;
+
+				BlocksToDefrag.Add(FDefragBlock{ .BlockId = BlockId, .LastAccess = Info.LastAccess });
+			}
 		}
 
 		if (!bPossibleToFreeBytes)
@@ -2588,9 +2608,22 @@ FIoStatus FOnDemandInstallCache::Defrag(
 
 			if (Info.RefSize < Info.FileSize)
 			{
+				// Block is fragmented
 				FragmentedBytes += (Info.FileSize - Info.RefSize);
 				TotalBlockSize += Info.FileSize;
 
+				BlocksToDefrag.Add(FDefragBlock{ .BlockId = BlockId, .LastAccess = Info.LastAccess });
+			}
+			else if (Info.FileSize < Cas.MinBlockSize)
+			{
+				// Block is too small whether or not its fragmented
+				if (ensure(Info.RefSize <= Info.FileSize))
+				{
+					FragmentedBytes += (Info.FileSize - Info.RefSize);
+				}
+
+				TotalBlockSize += Info.FileSize;
+
 				BlocksToDefrag.Add(FDefragBlock{ .BlockId = BlockId, .LastAccess = Info.LastAccess });
 			}
 		}
