Iris
* Added CSVEvent on clients when replication as been blocked for more then 30secs
* Trigger the CSVEvent event 10secs later if the blocked is still not resolved.

#rnx

diff --git a/Engine/Source/Runtime/Experimental/Iris/Core/Private/Iris/Core/IrisProfiler.cpp b/Engine/Source/Runtime/Experimental/Iris/Core/Private/Iris/Core/IrisProfiler.cpp
--- a/Engine/Source/Runtime/Experimental/Iris/Core/Private/Iris/Core/IrisProfiler.cpp
+++ b/Engine/Source/Runtime/Experimental/Iris/Core/Private/Iris/Core/IrisProfiler.cpp
@@ -7,6 +7,7 @@ CSV_DEFINE_CATEGORY(IrisClient, IRIS_CLIENT_PROFILER_ENABLE);
 CSV_DEFINE_CATEGORY(IrisClientDetailObjectCreate, IRIS_CLIENT_PROFILER_ENABLE);
 CSV_DEFINE_CATEGORY(IrisClientDetailRepNotify, IRIS_CLIENT_PROFILER_ENABLE);
 CSV_DEFINE_CATEGORY(IrisClientDetailRPC, IRIS_CLIENT_PROFILER_ENABLE);
+CSV_DEFINE_CATEGORY(IrisClientBlockedByAsyncLoading, IRIS_CLIENT_PROFILER_ENABLE);
 
 // Detailed CSV stats are disabled in shipping by default.
 #ifndef IRIS_CLIENT_PROFILER_DETAILED
@@ -71,6 +72,11 @@ void FClientProfiler::RecordRPC(FName RPCName)
 #endif
 }
 
+void FClientProfiler::RecordBlockedReplication(const TCHAR* BlockedObject, int32 NumBlockedAssets, float BlockedTime)
+{
+	FCsvProfiler::RecordEventf(CSV_CATEGORY_INDEX(IrisClientBlockedByAsyncLoading), TEXT("ReplicatedObject: %s has been blocked waiting for async loading of %d assets for %f secs"), BlockedObject, NumBlockedAssets, BlockedTime);
+}
+
 bool FClientProfiler::IsCapturing()
 {
 	return FCsvProfiler::Get()->IsCapturing();
diff --git a/Engine/Source/Runtime/Experimental/Iris/Core/Private/Iris/ReplicationSystem/PendingBatchData.h b/Engine/Source/Runtime/Experimental/Iris/Core/Private/Iris/ReplicationSystem/PendingBatchData.h
--- a/Engine/Source/Runtime/Experimental/Iris/Core/Private/Iris/ReplicationSystem/PendingBatchData.h
+++ b/Engine/Source/Runtime/Experimental/Iris/Core/Private/Iris/ReplicationSystem/PendingBatchData.h
@@ -50,6 +50,9 @@ struct FPendingBatchData
 
 	// When we started to accumulate pending batch data for this object
 	uint64 PendingBatchStartCycles = 0;
+
+	// At what time should we warn about being blocked too long
+	double NextWarningTimeout = 0.0;
 	
 	// Incremented every time we try to process queued batches, reset each time we output warning.
 	int32 PendingBatchTryProcessCount = 0;
diff --git a/Engine/Source/Runtime/Experimental/Iris/Core/Private/Iris/ReplicationSystem/ReplicationReader.cpp b/Engine/Source/Runtime/Experimental/Iris/Core/Private/Iris/ReplicationSystem/ReplicationReader.cpp
--- a/Engine/Source/Runtime/Experimental/Iris/Core/Private/Iris/ReplicationSystem/ReplicationReader.cpp
+++ b/Engine/Source/Runtime/Experimental/Iris/Core/Private/Iris/ReplicationSystem/ReplicationReader.cpp
@@ -109,19 +109,25 @@ static FAutoConsoleVariableRef CvarResolvedObjectsDispatchDebugging(
 
 static int32 QueuedBatchTimeoutWarningInterval = 30;
 static FAutoConsoleVariableRef CVarQueuedBatchTimeoutWarningInterval(
-	TEXT("net.QueuedBatchTimeoutWarningInterval"),
+	TEXT("net.Iris.QueuedBatch.TimeoutWarningInterval"),
 	QueuedBatchTimeoutWarningInterval,
 	TEXT("Number of ticks between checks for queuedbatches timeout."));
 
 static float QueuedBatchTimeoutSeconds = 30.f;
 static FAutoConsoleVariableRef CVarQueuedBatchTimeoutSeconds(
-	TEXT("net.QueuedBatchTimeoutSeconds"),
+	TEXT("net.Iris.QueuedBatch.TimeoutSeconds"),
 	QueuedBatchTimeoutSeconds,
 	TEXT("Time in seconds to wait for queued batches to flush before logging a warning."));
 
+static float QueuedBatchWarningRepeatTime = 10.f;
+static FAutoConsoleVariableRef CVarQueuedBatchWarningRepeatTime(
+	TEXT("net.Iris.QueuedBatch.WarningRepeatTime"),
+	QueuedBatchWarningRepeatTime,
+	TEXT("How long to wait before warning again."));
+
 static int32 MaxMustBeMappedHandleArray = 16;
 static FAutoConsoleVariableRef CVarMaxMustBeMappedHandleArray(
-	TEXT("net.Iris.MaxMustBeMappedHandleArray"),
+	TEXT("net.Iris.QueuedBatch.MaxMustBeMappedHandleArray"),
 	MaxMustBeMappedHandleArray,
 	TEXT("Maximum number of unresolved references that will be included with the QueuedBatch timeout error message."));
 
@@ -762,6 +768,7 @@ FPendingBatchData* FReplicationReader::UpdateUnresolvedMustBeMappedReferences(FN
 			Batch = &PendingBatches.PendingBatches.AddDefaulted_GetRef();
 			Batch->Handle = InHandle;
 			Batch->PendingBatchStartCycles = FPlatformTime::Cycles64();
+			Batch->NextWarningTimeout = FPlatformTime::ToSeconds64(Batch->PendingBatchStartCycles) + QueuedBatchTimeoutSeconds;
 		}
 
 		// Update
@@ -2759,17 +2766,26 @@ void FReplicationReader::ProcessQueuedBatches()
 		}
 		else
 		{
-			// Warn if we have been waiting for too long
+			// Increment counter
+			++PendingBatchData.PendingBatchTryProcessCount;
+
+			// Check how long we've been blocked only every 30 ticks
 			if ((QueuedBatchTimeoutWarningInterval > 0) && (PendingBatchData.PendingBatchTryProcessCount > QueuedBatchTimeoutWarningInterval) && FPlatformProperties::RequiresCookedData())
 			{
-				const double QueuedBatchPendingTime = FPlatformTime::ToSeconds64(FPlatformTime::Cycles64() - PendingBatchData.PendingBatchStartCycles);
-				if (QueuedBatchPendingTime > QueuedBatchTimeoutSeconds)
+				const double CurrentTime = FPlatformTime::Seconds();
+
+				// Log a warning and tell the server about the blockage every X seconds (first after 30secs, again every 10secs after)
+				if (CurrentTime >= PendingBatchData.NextWarningTimeout)
 				{
-					UE_LOG(LogIris, Warning, TEXT("FReplicationWriter::ProcessQueuedBatches: Queued batches for longer than normal. %s, Queued: %i, PendingNetRefHandleResolves: %i, BatchPendingTime: %f")
-						, *NetRefHandleManager->PrintObjectFromNetRefHandle(PendingBatchData.Handle)
+					const double BlockedTime = FPlatformTime::ToSeconds64(FPlatformTime::Cycles64() - PendingBatchData.PendingBatchStartCycles);
+
+					const FString BlockedObjectName = *NetRefHandleManager->PrintObjectFromNetRefHandle(PendingBatchData.Handle);
+
+					UE_LOG(LogIris, Warning, TEXT("FReplicationWriter::ProcessQueuedBatches: Replication is blocked for longer than normal for: %s, QueuedBunches: %i, PendingNetRefHandleResolves: %i, BlockedTime: %f")
+						, *BlockedObjectName
 						, PendingBatchData.QueuedDataChunks.Num()
 						, PendingBatchData.PendingMustBeMappedReferences.Num()
-						, QueuedBatchPendingTime);
+						, BlockedTime);
 
 					TArray<FNetRefHandle> PendingMustBeMapped;
 					PendingMustBeMapped.Reserve(MaxMustBeMappedHandleArray);
@@ -2783,16 +2799,17 @@ void FReplicationReader::ProcessQueuedBatches()
 						}
 					}
 
+#if IRIS_CLIENT_PROFILER_ENABLE
+					FClientProfiler::RecordBlockedReplication(*BlockedObjectName, PendingBatchData.PendingMustBeMappedReferences.Num(), (float)BlockedTime);
+#endif
+
 					ReplicationBridge->SendErrorWithNetRefHandle(UE::Net::ENetRefHandleError::BlockedByMustBeMapped, PendingBatchData.Handle, Parameters.ConnectionId, PendingMustBeMapped);
 
-					// Reset counter
-					PendingBatchData.PendingBatchTryProcessCount = 0;
+					PendingBatchData.NextWarningTimeout = CurrentTime + QueuedBatchWarningRepeatTime;
 				}
-			}
-			else
-			{
-				// Increment counter
-				++PendingBatchData.PendingBatchTryProcessCount;
+
+				// Reset counter
+				PendingBatchData.PendingBatchTryProcessCount = 0;
 			}
 
 			++BatchIt;
diff --git a/Engine/Source/Runtime/Experimental/Iris/Core/Private/Iris/ReplicationSystem/ReplicationWriter.cpp b/Engine/Source/Runtime/Experimental/Iris/Core/Private/Iris/ReplicationSystem/ReplicationWriter.cpp
--- a/Engine/Source/Runtime/Experimental/Iris/Core/Private/Iris/ReplicationSystem/ReplicationWriter.cpp
+++ b/Engine/Source/Runtime/Experimental/Iris/Core/Private/Iris/ReplicationSystem/ReplicationWriter.cpp
@@ -565,7 +565,8 @@ void FReplicationWriter::StartReplication(uint32 InternalIndex)
 		if (ensure(RootObjectInternalIndex != FNetRefHandleManager::InvalidInternalIndex))
 		{
 			FReplicationInfo& OwnerInfo = ReplicatedObjects[RootObjectInternalIndex];
-			if (ensureMsgf(OwnerInfo.GetState() != EReplicatedObjectState::Invalid && OwnerInfo.GetState() < EReplicatedObjectState::PendingDestroy, TEXT("Unsupported state %s"), LexToString(OwnerInfo.GetState())))
+			if (ensureMsgf(OwnerInfo.GetState() != EReplicatedObjectState::Invalid && OwnerInfo.GetState() < EReplicatedObjectState::PendingDestroy, 
+							TEXT("Unsupported state %s on root object of: %s"), LexToString(OwnerInfo.GetState()), *NetRefHandleManager->PrintObjectFromIndex(InternalIndex)))
 			{
 				ObjectsWithDirtyChanges.SetBit(RootObjectInternalIndex);
 				OwnerInfo.HasDirtySubObjects = 1U;
diff --git a/Engine/Source/Runtime/Experimental/Iris/Core/Public/Iris/Core/IrisProfiler.h b/Engine/Source/Runtime/Experimental/Iris/Core/Public/Iris/Core/IrisProfiler.h
--- a/Engine/Source/Runtime/Experimental/Iris/Core/Public/Iris/Core/IrisProfiler.h
+++ b/Engine/Source/Runtime/Experimental/Iris/Core/Public/Iris/Core/IrisProfiler.h
@@ -78,6 +78,8 @@ public:
 	static void RecordRepNotify(FName RepNotifyName);
 	static void RecordRPC(FName RPCName);
 
+	static void RecordBlockedReplication(const TCHAR* BlockedObject, int32 NumBlockedAssets, float BlockedTime);
+
 	/** Return true if capturing events. */
 	static bool IsCapturing();
 };
