Geometry collection : Make sure we keep a shared pointer of the game thread collection in the physics proxy
This prevents a physics proxy to access freed memory when its component has been garbaged collected while running its last physics update on the game thread

diff --git a/Engine/Source/Runtime/Experimental/Chaos/Private/PhysicsProxy/GeometryCollectionPhysicsProxy.cpp b/Engine/Source/Runtime/Experimental/Chaos/Private/PhysicsProxy/GeometryCollectionPhysicsProxy.cpp
--- a/Engine/Source/Runtime/Experimental/Chaos/Private/PhysicsProxy/GeometryCollectionPhysicsProxy.cpp
+++ b/Engine/Source/Runtime/Experimental/Chaos/Private/PhysicsProxy/GeometryCollectionPhysicsProxy.cpp
@@ -495,7 +495,7 @@ void PopulateSimulatedParticle(
 // FGeometryCollectionPhysicsProxy
 //==============================================================================
 
-
+PRAGMA_DISABLE_DEPRECATION_WARNINGS
 FGeometryCollectionPhysicsProxy::FGeometryCollectionPhysicsProxy(
 	UObject* InOwner,
 	FGeometryDynamicCollection& GameThreadCollectionIn,
@@ -530,7 +530,43 @@ FGeometryCollectionPhysicsProxy::FGeometryCollectionPhysicsProxy(
 	check(BufferMode == Chaos::EMultiBufferMode::TripleGuarded);
 	InterpolationData = MakeUnique<FProxyInterpolationBase>();
 }
+PRAGMA_ENABLE_DEPRECATION_WARNINGS
 
+FGeometryCollectionPhysicsProxy::FGeometryCollectionPhysicsProxy(
+	UObject* InOwner,
+	TSharedRef<FGeometryDynamicCollection> GameThreadCollectionSharedRefIn,
+	const FSimulationParameters& SimulationParameters,
+	FCollisionFilterData InSimFilter,
+	FCollisionFilterData InQueryFilter,
+	FGuid InCollectorGuid,
+	const Chaos::EMultiBufferMode BufferMode)
+	: Base(InOwner)
+	, Parameters(SimulationParameters)
+	, WorldTransform_External(FTransform::Identity)
+	, NumTransforms(INDEX_NONE)
+	, NumEffectiveParticles(INDEX_NONE)
+	, BaseParticleIndex(INDEX_NONE)
+	, CollisionParticlesPerObjectFraction(CollisionParticlesPerObjectFractionDefault)
+	, bIsGameThreadWorldTransformDirty(false)
+	, bHasBuiltGeometryOnPT(false)
+	, bHasBuiltGeometryOnGT(false)
+	, IsObjectDynamic(false)
+	, IsObjectLoading(true)
+	, IsObjectDeleting(false)
+	, bSkipChildToParentUpdateWhenInClusterUnion(false)
+	, SimFilter(InSimFilter)
+	, QueryFilter(InQueryFilter)
+	, PhysicsThreadCollection(Parameters.RestCollectionShared)
+	, GameThreadCollection(GameThreadCollectionSharedRefIn.Get())
+	, GameThreadCollectionSharedPtr(GameThreadCollectionSharedRefIn.ToSharedPtr())
+#if WITH_EDITORONLY_DATA
+	, CollectorGuid(InCollectorGuid)
+#endif
+{
+	// We rely on a guarded buffer.
+	check(BufferMode == Chaos::EMultiBufferMode::TripleGuarded);
+	InterpolationData = MakeUnique<FProxyInterpolationBase>();
+}
 
 FGeometryCollectionPhysicsProxy::~FGeometryCollectionPhysicsProxy()
 {}
@@ -4369,6 +4405,17 @@ bool FGeometryCollectionPhysicsProxy::PullFromPhysicsState(const Chaos::FDirtyGe
 		return false;
 	}
 
+	if (NumTransforms != GameThreadCollection.GetNumTransforms())
+	{
+		// At this point teh component may be gone and we have no way to get the name of the asset/component
+		UE_LOG(LogChaos, Error, 
+			TEXT("Attempt to update physics state on thegamethread where the number of transform (%d) does not match anymore the rest collection (%d) - skipping physics update"),
+			NumTransforms, 
+			GameThreadCollection.GetNumTransforms()
+		);
+		return false;
+	}
+
 	FProxyInterpolationBase* InterpData = GetInterpolationData();
 	if (Error)
 	{
@@ -4393,7 +4440,6 @@ bool FGeometryCollectionPhysicsProxy::PullFromPhysicsState(const Chaos::FDirtyGe
 	}
 	const bool bHasInterpolationData = InterpData != nullptr;
 
-	check(NumTransforms == GameThreadCollection.GetNumTransforms());
 	const bool bNeedInterpolation = (NextPullData != nullptr);
 
 	bool bIsCollectionDirty = false;
diff --git a/Engine/Source/Runtime/Experimental/Chaos/Public/PhysicsProxy/GeometryCollectionPhysicsProxy.h b/Engine/Source/Runtime/Experimental/Chaos/Public/PhysicsProxy/GeometryCollectionPhysicsProxy.h
--- a/Engine/Source/Runtime/Experimental/Chaos/Public/PhysicsProxy/GeometryCollectionPhysicsProxy.h
+++ b/Engine/Source/Runtime/Experimental/Chaos/Public/PhysicsProxy/GeometryCollectionPhysicsProxy.h
@@ -159,6 +159,7 @@ public:
 	 * \p InCacheSyncFunc callback invoked from \c PullFromPhysicsState().
 	 * \p InFinalSyncFunc callback invoked from \c SyncBeforeDestory().
 	 */
+	UE_DEPRECATED(5.6, "Use the constructor that use a shared ref  for the game thread collection")
 	CHAOS_API FGeometryCollectionPhysicsProxy(
 		UObject* InOwner, 
 		FGeometryDynamicCollection& GameThreadCollection,
@@ -167,6 +168,16 @@ public:
 		FCollisionFilterData InQueryFilter,
 		FGuid InCollectorGuid = FGuid::NewGuid(),
 		const Chaos::EMultiBufferMode BufferMode=Chaos::EMultiBufferMode::TripleGuarded);
+
+	CHAOS_API FGeometryCollectionPhysicsProxy(
+		UObject* InOwner,
+		TSharedRef<FGeometryDynamicCollection> GameThreadCollection,
+		const FSimulationParameters& SimulationParameters,
+		FCollisionFilterData InSimFilter,
+		FCollisionFilterData InQueryFilter,
+		FGuid InCollectorGuid = FGuid::NewGuid(),
+		const Chaos::EMultiBufferMode BufferMode = Chaos::EMultiBufferMode::TripleGuarded);
+
 	CHAOS_API virtual ~FGeometryCollectionPhysicsProxy();
 
 	/**
@@ -747,7 +758,8 @@ private:
 	// uses the data, but at the start of the simulation the PhysicsThreadCollection
 	// is a deep copy from the GameThreadCollection. 
 	FGeometryDynamicCollection PhysicsThreadCollection;
-	FGeometryDynamicCollection& GameThreadCollection;
+	FGeometryDynamicCollection& GameThreadCollection; // now initialize from the SharedRef passed in the constructor 
+	const TSharedPtr<FGeometryDynamicCollection> GameThreadCollectionSharedPtr; // shoul dnever be null
 
 	// Currently this is using triple buffers for game-physics and 
 	// physics-game thread communication, but not for any reason other than this 
diff --git a/Engine/Source/Runtime/Experimental/GeometryCollectionEngine/Private/GeometryCollection/GeometryCollectionComponent.cpp b/Engine/Source/Runtime/Experimental/GeometryCollectionEngine/Private/GeometryCollection/GeometryCollectionComponent.cpp
--- a/Engine/Source/Runtime/Experimental/GeometryCollectionEngine/Private/GeometryCollection/GeometryCollectionComponent.cpp
+++ b/Engine/Source/Runtime/Experimental/GeometryCollectionEngine/Private/GeometryCollection/GeometryCollectionComponent.cpp
@@ -3661,7 +3661,7 @@ void UGeometryCollectionComponent::ResetDynamicCollection()
 #endif
 	if (bCreateDynamicCollection && RestCollection && RestCollection->GetGeometryCollection())
 	{
-		DynamicCollection = MakeUnique<FGeometryDynamicCollection>(RestCollection->GetGeometryCollection());
+		DynamicCollection = MakeShared<FGeometryDynamicCollection>(RestCollection->GetGeometryCollection());
 
 		if (bStoreVelocities || bNotifyTrailing)
 		{
@@ -3822,6 +3822,8 @@ static FORCEINLINE_DEBUGGABLE int32 ComputeParticleLevel(Chaos::FPBDRigidCluster
 
 void UGeometryCollectionComponent::RegisterAndInitializePhysicsProxy()
 {
+	check(DynamicCollection != nullptr);
+
 	// CVar defined in BodyInstance but pertinent here as we will need to copy simplicials in the case that this is set.
 	// Original CVar is read-only so taking a static ptr here is fine as the value cannot be changed
 	static IConsoleVariable* AnalyticDisableCVar = IConsoleManager::Get().FindConsoleVariable(TEXT("p.IgnoreAnalyticCollisionsOverride"));
@@ -3940,7 +3942,7 @@ PRAGMA_ENABLE_DEPRECATION_WARNINGS
 		FRuntimeDataCollector::GetInstance().RemoveCollector(CollectorGuid);
 	}
 #endif
-	PhysicsProxy = new FGeometryCollectionPhysicsProxy(this, *DynamicCollection, SimulationParameters, InitialSimFilter, InitialQueryFilter, CollectorGuid);
+	PhysicsProxy = new FGeometryCollectionPhysicsProxy(this, DynamicCollection.ToSharedRef(), SimulationParameters, InitialSimFilter, InitialQueryFilter, CollectorGuid);
 
 	PhysicsProxy->SetPostPhysicsSyncCallback(
 		[WeakThis = MakeWeakObjectPtr(this)]()
diff --git a/Engine/Source/Runtime/Experimental/GeometryCollectionEngine/Public/GeometryCollection/GeometryCollectionComponent.h b/Engine/Source/Runtime/Experimental/GeometryCollectionEngine/Public/GeometryCollection/GeometryCollectionComponent.h
--- a/Engine/Source/Runtime/Experimental/GeometryCollectionEngine/Public/GeometryCollection/GeometryCollectionComponent.h
+++ b/Engine/Source/Runtime/Experimental/GeometryCollectionEngine/Public/GeometryCollection/GeometryCollectionComponent.h
@@ -1739,7 +1739,7 @@ private:
 	TArray<bool> EventsPlayed;
 
 	FGeometryCollectionPhysicsProxy* PhysicsProxy = nullptr;
-	TUniquePtr<FGeometryDynamicCollection> DynamicCollection;
+	TSharedPtr<FGeometryDynamicCollection> DynamicCollection;
 
 	// Temporary storage for body setup in order to initialise a dummy body instance
 	UPROPERTY(Transient)
