Added support for appending critical error entries in the on-demand install cache  journal.

- the install cache will be recreated when a critical error entry is found in the journal at startup
- updated Defrag to append critical error entry in the journal when running out of disk space and
if any of the chunk hashes doesn't match

#rnx

diff --git a/Engine/Source/Runtime/Experimental/IoStore/OnDemand/Private/OnDemandInstallCache.cpp b/Engine/Source/Runtime/Experimental/IoStore/OnDemand/Private/OnDemandInstallCache.cpp
--- a/Engine/Source/Runtime/Experimental/IoStore/OnDemand/Private/OnDemandInstallCache.cpp
+++ b/Engine/Source/Runtime/Experimental/IoStore/OnDemand/Private/OnDemandInstallCache.cpp
@@ -43,6 +43,10 @@
 #include "Tasks/Pipe.h"
 #endif // UE_ONDEMANDINSTALLCACHE_EXCLUSIVE_WRITE
 
+#ifndef UE_IAD_DEBUG_CONSOLE_CMDS
+#define UE_IAD_DEBUG_CONSOLE_CMDS (1 && !NO_CVARS && !UE_BUILD_SHIPPING)
+#endif
+
 namespace UE::IoStore
 {
 
@@ -613,6 +617,14 @@ struct FCasJournal
 		Latest = LatestPlusOne - 1
 	};
 
+	enum class EErrorCode : uint32
+	{
+		None					= 0,
+		Simulated				= 1,
+		DefragOutOfDiskSpace	= 2,
+		DefragHashMismatch		= 3
+	};
+
 	struct FHeader
 	{
 		static const inline uint8 MagicSequence[16] = {'C', 'A', 'S', 'J', 'O', 'U', 'R', 'N', 'A', 'L', 'H', 'E', 'A', 'D', 'E', 'R'};
@@ -643,7 +655,8 @@ struct FCasJournal
 			ChunkLocation,
 			BlockCreated,
 			BlockDeleted,
-			BlockAccess
+			BlockAccess,
+			CriticalError
 		};
 
 		struct FChunkLocation
@@ -665,10 +678,18 @@ struct FCasJournal
 		};
 		static_assert(sizeof(FBlockOperation) == 24);
 
+		struct FCriticalError
+		{
+			EType		Type = EType::CriticalError;
+			EErrorCode	ErrorCode = EErrorCode::None;
+		};
+		static_assert(sizeof(FBlockOperation) == 24);
+
 		union
 		{
 			FChunkLocation	ChunkLocation;
 			FBlockOperation	BlockOperation;
+			FCriticalError	CriticalError;
 		};
 
 		EType Type() const { return *reinterpret_cast<const EType*>(this); }
@@ -681,6 +702,7 @@ struct FCasJournal
 		void			BlockCreated(FCasBlockId BlockId);
 		void			BlockDeleted(FCasBlockId BlockId);
 		void			BlockAccess(FCasBlockId BlockId, int64 UtcTicks);
+		void			CriticalError(FCasJournal::EErrorCode ErrorCode);
 
 		FString			JournalFile;
 		TArray<FEntry>	Entries;
@@ -695,6 +717,24 @@ struct FCasJournal
 	static FIoStatus		Commit(FTransaction&& Transaction);
 };
 
+///////////////////////////////////////////////////////////////////////////////
+const TCHAR* GetErrorText(FCasJournal::EErrorCode ErrorCode)
+{
+	switch (ErrorCode)
+	{
+		case FCasJournal::EErrorCode::None:
+			return TEXT("None");
+		case FCasJournal::EErrorCode::Simulated:
+			return TEXT("Simulated error");
+		case FCasJournal::EErrorCode::DefragOutOfDiskSpace:
+			return TEXT("Defrag failed due to out of disk space");
+		case FCasJournal::EErrorCode::DefragHashMismatch:
+			return TEXT("Found corrupt chunk while defragging");
+	}
+
+	return TEXT("Unknown");
+}
+
 ///////////////////////////////////////////////////////////////////////////////
 bool FCasJournal::FHeader::IsValid() const
 {
@@ -793,6 +833,15 @@ FIoStatus FCasJournal::Replay(const FString& JournalFile, FEntryHandler&& Handle
 
 	for (const FEntry& Entry : Entries)
 	{
+		if (Entry.Type() == FEntry::EType::CriticalError)
+		{
+			const FEntry::FCriticalError& Error = Entry.CriticalError;
+			UE_LOG(LogIoStoreOnDemand, Warning, TEXT("Found critical error entry '%s' (%d) in journal '%s'"),
+				GetErrorText(Error.ErrorCode), Error.ErrorCode, *JournalFile);
+
+			return FIoStatus(EIoErrorCode::ReadError, FString(GetErrorText(Error.ErrorCode)));
+		}
+
 		Handler(Entry);
 	}
 
@@ -984,6 +1033,15 @@ void FCasJournal::FTransaction::BlockAccess(FCasBlockId BlockId, int64 UtcTicks)
 	};
 }
 
+void FCasJournal::FTransaction::CriticalError(FCasJournal::EErrorCode ErrorCode)
+{
+	Entries.AddZeroed_GetRef().CriticalError = FEntry::FCriticalError
+	{
+		.Type		= FEntry::EType::CriticalError,
+		.ErrorCode	= ErrorCode
+	};
+}
+
 ///////////////////////////////////////////////////////////////////////////////
 struct FCasSnapshot
 {
@@ -1117,6 +1175,11 @@ TIoStatusOr<FCasSnapshot> FCasSnapshot::FromJournal(const FString& JournalFile)
 			};
 		});
 
+	if (ReplayStatus.IsOk() == false)
+	{
+		return ReplayStatus;
+	}
+
 	FCasSnapshot Snapshot;
 	Snapshot.Blocks.Reserve(BlockIds.Num());
 	for (FCasBlockId BlockId : BlockIds)
@@ -1428,6 +1491,7 @@ public:
 	virtual FOnDemandInstallCacheUsage			GetCacheUsage() override;
 
 private:
+	void										RegisterConsoleCommands();
 	FIoStatus									Reinitialize();
 	FIoStatus									InitialVerify();
 	uint64										AddReferencesToBlocks(
@@ -1463,6 +1527,10 @@ private:
 #if UE_ONDEMANDINSTALLCACHE_EXCLUSIVE_WRITE
 	UE::Tasks::FPipe			ExclusivePipe{ UE_SOURCE_LOCATION };
 #endif // UE_ONDEMANDINSTALLCACHE_EXCLUSIVE_WRITE
+
+#if UE_IAD_DEBUG_CONSOLE_CMDS
+	TArray<IConsoleCommand*> ConsoleCommands;
+#endif // UE_IAS_DEBUG_CONSOLE_CMDS
 };
 
 ///////////////////////////////////////////////////////////////////////////////
@@ -1590,11 +1658,15 @@ FOnDemandInstallCache::FOnDemandInstallCache(const FOnDemandInstallCacheConfig&
 	// Try to reinitialize the cache if something has gone wrong 
 	if (Status.IsOk() == false)
 	{
+		UE_LOG(LogIoStoreOnDemand, Warning, TEXT("Deleting install cash and starting from scratch, reason '%s'"),
+			*Status.ToString());
 		Status = Reinitialize();
 	}
 
 	if (Status.IsOk())
 	{
+		UE_LOG(LogIoStoreOnDemand, Log, TEXT("Install cache Ok!"));
+		RegisterConsoleCommands();
 		Cas.Compact();
 	}
 	else
@@ -1650,6 +1722,13 @@ void FOnDemandInstallCache::Shutdown()
 				*JournalFile, *SnapshotStatus.Status().ToString());
 		}
 	}
+
+#if UE_IAD_DEBUG_CONSOLE_CMDS
+	for (IConsoleCommand* Cmd : ConsoleCommands)
+	{
+		IConsoleManager::Get().UnregisterConsoleObject(Cmd);
+	}
+#endif // UE_IAD_DEBUG_CONSOLE_CMDS
 }
 
 void FOnDemandInstallCache::ResolveIoRequests(FIoRequestList Requests, FIoRequestList& OutUnresolved)
@@ -2292,6 +2371,30 @@ FIoStatus FOnDemandInstallCache::Verify()
 	return FIoStatus::Ok;
 }
 
+void FOnDemandInstallCache::RegisterConsoleCommands()
+{
+#if UE_IAD_DEBUG_CONSOLE_CMDS
+	ConsoleCommands.Emplace(
+	IConsoleManager::Get().RegisterConsoleCommand(
+		TEXT("iostore.SimulateCriticalInstallCacheError"),
+		TEXT(""),
+		FConsoleCommandDelegate::CreateLambda([this]()
+		{
+			UE_LOG(LogIoStoreOnDemand, Log, TEXT("Simulating critical install cache error"));
+
+			FCasJournal::FTransaction Transaction = FCasJournal::Begin(GetJournalFilename());
+			Transaction.CriticalError(FCasJournal::EErrorCode::Simulated);
+			if (FIoStatus Status = FCasJournal::Commit(MoveTemp(Transaction)); Status.IsOk() == false)
+			{
+				UE_LOG(LogIoStoreOnDemand, Warning, TEXT("Failed to append critical error to journal, reason '%s'"),
+					*Status.ToString());
+			}
+		}),
+		ECVF_Default)
+	);
+#endif // UE_IAD_DEBUG_CONSOLE_CMDS
+}
+
 FIoStatus FOnDemandInstallCache::Reinitialize()
 {
 	UE_LOG(LogIoStoreOnDemand, Log, TEXT("Recreating install cache in directory '%s'"), *CacheDirectory);
@@ -2520,6 +2623,16 @@ FIoStatus FOnDemandInstallCache::Defrag(
 			ToMiB(MaxCacheSize), ToMiB(TotalCachedBytes));
 		UE_LOG(LogIoStoreOnDemand, Error, TEXT("%s"), *ErrorMsg);
 		FOnDemandInstallCacheStats::OnDefrag(EIoErrorCode::WriteError, 0);
+
+		// Append a critical error entry to clear the cache at next startup
+		FCasJournal::FTransaction Transaction = FCasJournal::Begin(GetJournalFilename());
+		Transaction.CriticalError(FCasJournal::EErrorCode::DefragOutOfDiskSpace);
+		if (FIoStatus Status = FCasJournal::Commit(MoveTemp(Transaction)); Status.IsOk() == false)
+		{
+			UE_LOG(LogIoStoreOnDemand, Error, TEXT("Failed to append critical error to journal, reason '%s'"),
+				*Status.ToString());
+		}
+
 		return FIoStatusBuilder(EIoErrorCode::WriteError) << ErrorMsg;
 	}
 
@@ -2710,6 +2823,15 @@ FIoStatus FOnDemandInstallCache::Defrag(
 					UE_LOG(LogIoStoreOnDemand, Error, TEXT("Found chunk with invalid hash while defragging block, BlockId=%u, BlockOffset=%u"),
 						DefragBlock.BlockId.Id, ReffedChunk.BlockOffset);
 
+					// Append a critical error entry to clear the cache at next startup
+					FCasJournal::FTransaction Transaction = FCasJournal::Begin(GetJournalFilename());
+					Transaction.CriticalError(FCasJournal::EErrorCode::DefragHashMismatch);
+					if (FIoStatus Status = FCasJournal::Commit(MoveTemp(Transaction)); Status.IsOk() == false)
+					{
+						UE_LOG(LogIoStoreOnDemand, Error, TEXT("Failed to append critical error to journal, reason '%s'"),
+							*Status.ToString());
+					}
+
 					if (FIoStatus Status = FlushPendingChunks(DefragPendingChunks); Status.IsOk() == false)
 					{
 						FOnDemandInstallCacheStats::OnDefrag(Status.GetErrorCode(), 0);
