Marked a constant if statement constexpr

#rnx

diff --git a/Engine/Source/Runtime/Experimental/IoStore/HttpClient/Private/Misc.inl b/Engine/Source/Runtime/Experimental/IoStore/HttpClient/Private/Misc.inl
--- a/Engine/Source/Runtime/Experimental/IoStore/HttpClient/Private/Misc.inl
+++ b/Engine/Source/Runtime/Experimental/IoStore/HttpClient/Private/Misc.inl
@@ -221,17 +221,19 @@ static int64 CrudeToInt(FAnsiStringView View)
 		uint32 Digit = c - '0';
 		if (Digit > 9u)
 		{
-			if (Base != 16)
+			if constexpr (Base != 16)
 			{
 				break;
 			}
-
-			Digit = (c | 0x20) - 'a';
-			if (Digit > uint32('f' - 'a'))
+			else
 			{
-				break;
+				Digit = (c | 0x20) - 'a';
+				if (Digit > uint32('f' - 'a'))
+				{
+					break;
+				}
+				Digit += 10;
 			}
-			Digit += 10;
 		}
 		Value *= Base;
 		Value += Digit;
