Added read error analytics for the on-demand install cache.

#rnx

diff --git a/Engine/Source/Runtime/Experimental/IoStore/OnDemand/Private/OnDemandInstallCache.cpp b/Engine/Source/Runtime/Experimental/IoStore/OnDemand/Private/OnDemandInstallCache.cpp
--- a/Engine/Source/Runtime/Experimental/IoStore/OnDemand/Private/OnDemandInstallCache.cpp
+++ b/Engine/Source/Runtime/Experimental/IoStore/OnDemand/Private/OnDemandInstallCache.cpp
@@ -1927,6 +1927,7 @@ bool FOnDemandInstallCache::Resolve(FIoRequestImpl* Request)
 		const FString Filename = Cas.GetBlockFilename(CasLoc.BlockId);
 		UE_LOG(LogIoStoreOnDemand, Error, TEXT("Failed to open CAS block '%s' for async reading, reason '%s'"), *Filename, *FileOpenResult.GetError().GetMessage());
 		TUniquePtr<FChunkRequest> Detached = FChunkRequest::Detach(*Request);
+		FOnDemandInstallCacheStats::OnReadCompleted(EIoErrorCode::FileOpenFailed);
 		CompleteRequest(Request, EIoErrorCode::FileOpenFailed);
 		return true;
 	}
@@ -3077,7 +3078,6 @@ FIoStatus FOnDemandInstallCache::FlushPendingChunksImpl(FPendingChunks& Chunks,
 
 void FOnDemandInstallCache::CompleteRequest(FIoRequestImpl* Request, EIoErrorCode Status)
 {
-	bool bSucceeded	= false;
 	if (Status == EIoErrorCode::Ok && !Request->IsCancelled())
 	{
 		FChunkRequest& ChunkRequest = FChunkRequest::GetRef(*Request);
@@ -3099,19 +3099,17 @@ void FOnDemandInstallCache::CompleteRequest(FIoRequestImpl* Request, EIoErrorCod
 			Request->CreateBuffer(ChunkRequest.RawSize);
 			FMutableMemoryView RawChunk = Request->GetBuffer().GetMutableView();
 
-			bSucceeded = FIoChunkEncoding::Decode(Params, EncodedChunk.GetView(), RawChunk);
-			UE_CLOG(!bSucceeded, LogIoStoreOnDemand, Error, TEXT("Failed to decode chunk, ChunkId='%s'"), *LexToString(Request->ChunkId));
+			if (FIoChunkEncoding::Decode(Params, EncodedChunk.GetView(), RawChunk) == false)
+			{
+				UE_LOG(LogIoStoreOnDemand, Error, TEXT("Failed to decode chunk, ChunkId='%s'"), *LexToString(Request->ChunkId));
+				Status = EIoErrorCode::CompressionError;
+			}
 		}
 	}
 
-	if (bSucceeded == false)
+	if (Status != EIoErrorCode::Ok)
 	{
-		if (Status == EIoErrorCode::Ok)
-		{
-			Status = EIoErrorCode::ReadError;
-		}
 		Request->SetLastBackendError(Status);
-
 		Request->SetResult(FIoBuffer());
 		TRACE_IOSTORE_BACKEND_REQUEST_FAILED(Request);
 	}
@@ -3123,6 +3121,7 @@ void FOnDemandInstallCache::CompleteRequest(FIoRequestImpl* Request, EIoErrorCod
 	{
 		UE::TUniqueLock Lock(Mutex);
 		CompletedRequests.AddTail(Request);
+		FOnDemandInstallCacheStats::OnReadCompleted(Status);
 	}
 
 	BackendContext->WakeUpDispatcherThreadDelegate.Execute();
diff --git a/Engine/Source/Runtime/Experimental/IoStore/OnDemand/Private/Statistics.cpp b/Engine/Source/Runtime/Experimental/IoStore/OnDemand/Private/Statistics.cpp
--- a/Engine/Source/Runtime/Experimental/IoStore/OnDemand/Private/Statistics.cpp
+++ b/Engine/Source/Runtime/Experimental/IoStore/OnDemand/Private/Statistics.cpp
@@ -464,6 +464,8 @@ struct FInstallCacheAnalytics
 	uint64 DefragErrorCount = 0;
 	uint64 JournalCommitCount = 0;
 	uint64 JournalCommitErrorCount = 0;
+	uint64 ReadCount = 0;
+	uint64 ReadErrorCount = 0;
 	uint64 MaxCacheSize = 0;
 	uint64 MaxCacheUsageSize = 0;
 	uint64 MaxReferencedBlockSize = 0;
@@ -1171,6 +1173,7 @@ void FOnDemandContentInstallerStats::ReportAnalytics(TArray<FAnalyticsEventAttri
 		CurrentInstallerAnalytics	= InstallerAnalytics;
 		InstallerAnalytics			= FInstallerAnalytics();
 	}
+
 	FInstallCacheAnalytics CurrentInstallCacheAnalytics;
 	{
 		TUniqueLock Lock(InstallCacheAnalytics);
@@ -1206,6 +1209,8 @@ void FOnDemandContentInstallerStats::ReportAnalytics(TArray<FAnalyticsEventAttri
 		, TEXT("IadInstallCacheMaxReferencedSize"), CurrentInstallCacheAnalytics.MaxReferencedSize
 		, TEXT("IadInstallCacheMaxFragmentedSize"), CurrentInstallCacheAnalytics.MaxFragmentedSize
 		, TEXT("IadInstallCacheOldestBlockAgeMinutes"), OldestBlockAge.GetTotalMinutes()
+		, TEXT("IadInstallCacheReadCount"), CurrentInstallCacheAnalytics.ReadCount
+		, TEXT("IadInstallCacheReadErrorCount"), CurrentInstallCacheAnalytics.ReadErrorCount
 	);
 }
 
@@ -1307,6 +1312,21 @@ void FOnDemandInstallCacheStats::OnCacheUsage(
 	InstallCacheAnalytics.OldestBlockAccess			= FMath::Min(InstallCacheAnalytics.OldestBlockAccess, OldestBlockAccess);
 }
 
+void FOnDemandInstallCacheStats::OnReadCompleted(EIoErrorCode ErrorCode)
+{
+	if (ErrorCode == EIoErrorCode::Cancelled)
+	{
+		return;
+	}
+
+	UE::TUniqueLock Lock(InstallCacheAnalytics);
+	if (ErrorCode != EIoErrorCode::Ok)
+	{
+		++InstallCacheAnalytics.ReadErrorCount;
+	}
+	++InstallCacheAnalytics.ReadCount;
+}
+
 } // namespace UE::IoStore
 
 #undef UE_ENABLE_ONSCREEN_STATISTICS
diff --git a/Engine/Source/Runtime/Experimental/IoStore/OnDemand/Private/Statistics.h b/Engine/Source/Runtime/Experimental/IoStore/OnDemand/Private/Statistics.h
--- a/Engine/Source/Runtime/Experimental/IoStore/OnDemand/Private/Statistics.h
+++ b/Engine/Source/Runtime/Experimental/IoStore/OnDemand/Private/Statistics.h
@@ -125,6 +125,7 @@ public:
 		uint64 ReferencedSize,
 		uint64 FragmentedSize,
 		int64 OldestBlockAccess) IAS_STATISTICS_IMPL()
+	static void OnReadCompleted(EIoErrorCode ErrorCode) IAS_STATISTICS_IMPL()
 };
 
 #undef IAS_STATISTICS_IMPL
