Constraint: fix for ConstraintID not being initialized due to functions that should not have been editor-only
- as the guid is null, only the first constraint was actually added in the constraints' subsystem (for package builds only)

#rnx

diff --git a/Engine/Source/Runtime/Experimental/Animation/Constraints/Private/ConstraintsManager.cpp b/Engine/Source/Runtime/Experimental/Animation/Constraints/Private/ConstraintsManager.cpp
--- a/Engine/Source/Runtime/Experimental/Animation/Constraints/Private/ConstraintsManager.cpp
+++ b/Engine/Source/Runtime/Experimental/Animation/Constraints/Private/ConstraintsManager.cpp
@@ -124,35 +124,21 @@ UTickableConstraint* UTickableConstraint::Duplicate(UObject* NewOuter) const
 	return DuplicateObject<UTickableConstraint>(this, NewOuter, GetFName());
 }
 
-#if WITH_EDITOR
-
-FString UTickableConstraint::GetLabel() const
-{
-	return UTickableConstraint::StaticClass()->GetName();
-}
-
-FString UTickableConstraint::GetFullLabel() const
-{
-	return GetLabel();
-}
-
-FString UTickableConstraint::GetTypeLabel() const
-{
-	return GetLabel();
-}
-
 void UTickableConstraint::PostDuplicate(bool bDuplicateForPIE)
 {
 	Super::PostDuplicate(bDuplicateForPIE);
-	if (bDuplicateForPIE == false) //doing actually copy then give unique id
+
+	if (!bDuplicateForPIE) //doing actually copy then give unique id
 	{
 		ConstraintID = FGuid::NewGuid();
 	}
 }
+
 void UTickableConstraint::PostLoad()
 {
 	Super::PostLoad();
-	if (ConstraintID.IsValid() == false)
+	
+	if (!ConstraintID.IsValid())
 	{
 		ConstraintID = FGuid::NewGuid();
 	}
@@ -161,12 +147,30 @@ void UTickableConstraint::PostLoad()
 void UTickableConstraint::PostInitProperties()
 {
 	Super::PostInitProperties();
-	if (HasAnyFlags(RF_ClassDefaultObject | RF_NeedLoad) == false)
+	
+	if (!HasAnyFlags(RF_ClassDefaultObject | RF_NeedLoad) && !ConstraintID.IsValid())
 	{
 		ConstraintID = FGuid::NewGuid();
 	}
 }
 
+#if WITH_EDITOR
+
+FString UTickableConstraint::GetLabel() const
+{
+	return UTickableConstraint::StaticClass()->GetName();
+}
+
+FString UTickableConstraint::GetFullLabel() const
+{
+	return GetLabel();
+}
+
+FString UTickableConstraint::GetTypeLabel() const
+{
+	return GetLabel();
+}
+
 void UTickableConstraint::PostEditChangeProperty(FPropertyChangedEvent& PropertyChangedEvent)
 {
 	Super::PostEditChangeProperty(PropertyChangedEvent);
diff --git a/Engine/Source/Runtime/Experimental/Animation/Constraints/Public/ConstraintsManager.h b/Engine/Source/Runtime/Experimental/Animation/Constraints/Public/ConstraintsManager.h
--- a/Engine/Source/Runtime/Experimental/Animation/Constraints/Public/ConstraintsManager.h
+++ b/Engine/Source/Runtime/Experimental/Animation/Constraints/Public/ConstraintsManager.h
@@ -113,6 +113,12 @@ public:
 	CONSTRAINTS_API FConstraintTickFunction& GetTickFunction(UWorld* InWorld);
 	CONSTRAINTS_API const FConstraintTickFunction& GetTickFunction(UWorld* InWorld) const;
 
+	// UObject interface
+	CONSTRAINTS_API virtual void PostLoad() override;
+	CONSTRAINTS_API virtual void PostDuplicate(bool bDuplicateForPIE) override;
+	CONSTRAINTS_API virtual void PostInitProperties() override;
+	// End of UObject interface
+	
 #if WITH_EDITOR
 	/** Returns the constraint's label used for UI. */
 	CONSTRAINTS_API virtual FString GetLabel() const;
@@ -121,14 +127,11 @@ public:
 	/** Returns the constraint's type label used for UI. */
 	CONSTRAINTS_API virtual FString GetTypeLabel() const;
 
-	// UObject interface
+	// UObject editor only interface
 	CONSTRAINTS_API virtual void PostEditChangeProperty(FPropertyChangedEvent& PropertyChangedEvent) override;
 	CONSTRAINTS_API virtual void PostEditUndo() override;
-	CONSTRAINTS_API virtual void PostLoad() override;
-	CONSTRAINTS_API virtual void PostDuplicate(bool bDuplicateForPIE) override;
-	CONSTRAINTS_API virtual void PostInitProperties() override;
-
-	// End of UObject interface
+	// End of UObject editor only interface
+	
 #endif
 	/** @todo documentation. */
 	UPROPERTY(EditAnywhere, BlueprintReadWrite, DisplayName="Active State", Category="Constraint")
@@ -171,7 +174,7 @@ public:
 	CONSTRAINTS_API virtual ~UConstraintsManager();
 	
 	//UObjects
-	CONSTRAINTS_API virtual void PostLoad() override;
+	CONSTRAINTS_API virtual void	PostLoad() override;
 
 	CONSTRAINTS_API void Init(UWorld* InWorld);
 	CONSTRAINTS_API void Clear(UWorld* World);
