Dataflow : Fixed dynamic connection not properly clearing their internal dynamic properties array

diff --git a/Engine/Source/Runtime/Experimental/Dataflow/Engine/Private/Dataflow/DataflowDynamicConnections.cpp b/Engine/Source/Runtime/Experimental/Dataflow/Engine/Private/Dataflow/DataflowDynamicConnections.cpp
--- a/Engine/Source/Runtime/Experimental/Dataflow/Engine/Private/Dataflow/DataflowDynamicConnections.cpp
+++ b/Engine/Source/Runtime/Experimental/Dataflow/Engine/Private/Dataflow/DataflowDynamicConnections.cpp
@@ -70,6 +70,7 @@ void FDataflowDynamicConnections::Refresh()
 	ClearDynamicConnections();
 
 	ConnectionNameToPropertyId.Reset();
+	DynamicProperties.Reset();
 
 	// Go through all the properties of the property bag recreate the output and reconnect them if needed
 	const TConnectionReference<FDataflowAllTypes> NullReference(nullptr);
