Destroy the install request after the completion callback has been triggered in order to prevent the progress callback to access an invalid request.

#rnx

diff --git a/Engine/Source/Runtime/Experimental/IoStore/OnDemand/Private/OnDemandContentInstaller.cpp b/Engine/Source/Runtime/Experimental/IoStore/OnDemand/Private/OnDemandContentInstaller.cpp
--- a/Engine/Source/Runtime/Experimental/IoStore/OnDemand/Private/OnDemandContentInstaller.cpp
+++ b/Engine/Source/Runtime/Experimental/IoStore/OnDemand/Private/OnDemandContentInstaller.cpp
@@ -966,45 +966,14 @@ void FOnDemandContentInstaller::CompleteInstallRequest(FRequest& Request)
 		GetIoErrorText(InstallResult.Status.GetErrorCode()), Request.SeqNo, Request.Priority, *LexToString(InstallRequest.Args.ContentHandle), double(InstallResult.Progress.TotalContentSize) / 1024.0 / 1024.0,
 		double(InstallResult.Progress.TotalInstallSize) / 1024.0 / 1024.0, int32(CacheHitRatio * 100), int32(InstallResult.DurationInSeconds * 1000));
 
-	const bool bCallbackOnGameThread			= EnumHasAnyFlags(InstallRequest.Args.Options, EOnDemandInstallOptions::CallbackOnGameThread);
-	FOnDemandInstallCompleted OnCompleted		= MoveTemp(InstallRequest.OnCompleted);
-	FSharedInternalInstallRequest ClientRequest = MoveTemp(InstallRequest.Request);
-
-	// TODO: Remove support for notifying progress on game thread
-	if (InstallRequest.bNotifyingProgressOnGameThread)
-	{
-		const double WaitTimeoutSeconds		= 5.0;
-		const double WarningTimeoutSeconds	= .5;
-		double WaitTimeSeconds				= 0.0;
-		double WaitStartTimeSeconds			= FPlatformTime::Seconds();
-
-		while (InstallRequest.bNotifyingProgressOnGameThread)
-		{
-			WaitTimeSeconds = FPlatformTime::Seconds() - WaitStartTimeSeconds;
-			if (WaitTimeSeconds > WaitTimeoutSeconds)
-			{
-				UE_LOG(LogIoStoreOnDemand, Error,
-					TEXT("Content installer exceeded timeout threshold of %.2f seconds waiting for progress callback to finish"),
-					WaitTimeoutSeconds);
-				break;
-			}
-			FPlatformProcess::SleepNoStats(0.017f);
-		}
-
-		UE_CLOG(WaitTimeSeconds > WarningTimeoutSeconds, LogIoStoreOnDemand, Warning,
-			TEXT("Content installer waited %.2lf seconds for progress callback to finish"),
-			WaitTimeSeconds);
-	}
-
 	{
 		UE::TUniqueLock Lock(Mutex);
-		ClientRequest->InstallerRequest = 0;
 
+		InstallRequest.Request->InstallerRequest = 0;
 		if (&Request == CurrentRequest)
 		{
 			CurrentRequest = nullptr;
 		}
-		RequestAllocator.Destroy(&Request);
 	}
 
 	TryExecuteNextRequest();
@@ -1023,24 +992,44 @@ void FOnDemandContentInstaller::CompleteInstallRequest(FRequest& Request)
 			break;
 	}
 
-	if (!OnCompleted)
+	if (!InstallRequest.OnCompleted)
 	{
-		ClientRequest->Status.store(RequestStatus);
+		InstallRequest.Request->Status.store(RequestStatus);
+		{
+			UE::TUniqueLock Lock(Mutex);
+			RequestAllocator.Destroy(&Request);
+		}
+		return;
 	}
-	else if (bCallbackOnGameThread)
+
+	const bool bCallbackOnGameThread = EnumHasAnyFlags(InstallRequest.Args.Options, EOnDemandInstallOptions::CallbackOnGameThread);
+	if (bCallbackOnGameThread)
 	{
 		ExecuteOnGameThread(
 			UE_SOURCE_LOCATION,
-			[ClientRequest, InstallResult = MoveTemp(InstallResult), RequestStatus, OnCompleted = MoveTemp(OnCompleted)]() mutable
+			[this, &Request, InstallResult = MoveTemp(InstallResult), RequestStatus]() mutable
 			{
+				FRequest::FInstall& InstallRequest		= Request.AsInstall();
+				FOnDemandInstallCompleted OnCompleted	= MoveTemp(InstallRequest.OnCompleted);
+
+				ensure(InstallRequest.bNotifyingProgressOnGameThread == false);
 				OnCompleted(MoveTemp(InstallResult));
-				ClientRequest->Status.store(RequestStatus);
+				InstallRequest.Request->Status.store(RequestStatus);
+				{
+					UE::TUniqueLock Lock(Mutex);
+					RequestAllocator.Destroy(&Request);
+				}
 			});
 	}
 	else
 	{
+		FOnDemandInstallCompleted OnCompleted = MoveTemp(InstallRequest.OnCompleted);
 		OnCompleted(MoveTemp(InstallResult));
-		ClientRequest->Status.store(RequestStatus);
+		InstallRequest.Request->Status.store(RequestStatus);
+		{
+			UE::TUniqueLock Lock(Mutex);
+			RequestAllocator.Destroy(&Request);
+		}
 	}
 }
 
