Added analytics for on-demand install cache startup errors.

#rnx

diff --git a/Engine/Source/Runtime/Experimental/IoStore/OnDemand/Private/OnDemandInstallCache.cpp b/Engine/Source/Runtime/Experimental/IoStore/OnDemand/Private/OnDemandInstallCache.cpp
--- a/Engine/Source/Runtime/Experimental/IoStore/OnDemand/Private/OnDemandInstallCache.cpp
+++ b/Engine/Source/Runtime/Experimental/IoStore/OnDemand/Private/OnDemandInstallCache.cpp
@@ -1492,7 +1492,7 @@ public:
 
 private:
 	void										RegisterConsoleCommands();
-	FIoStatus									Reinitialize();
+	FIoStatus									Reset();
 	FIoStatus									InitialVerify();
 	uint64										AddReferencesToBlocks(
 													const TArray<FSharedOnDemandContainer>& Containers, 
@@ -1655,12 +1655,14 @@ FOnDemandInstallCache::FOnDemandInstallCache(const FOnDemandInstallCacheConfig&
 		Status = InitialVerify();
 	}
 
-	// Try to reinitialize the cache if something has gone wrong 
+	// Try to reset the cache if something has gone wrong 
 	if (Status.IsOk() == false)
 	{
-		UE_LOG(LogIoStoreOnDemand, Warning, TEXT("Deleting install cash and starting from scratch, reason '%s'"),
+		UE_LOG(LogIoStoreOnDemand, Warning, TEXT("Resetting install cash, reason '%s'"),
 			*Status.ToString());
-		Status = Reinitialize();
+
+		FOnDemandInstallCacheStats::OnStartupError(Status.GetErrorCode());
+		Status = Reset();
 	}
 
 	if (Status.IsOk())
@@ -2396,9 +2398,9 @@ void FOnDemandInstallCache::RegisterConsoleCommands()
 #endif // UE_IAD_DEBUG_CONSOLE_CMDS
 }
 
-FIoStatus FOnDemandInstallCache::Reinitialize()
+FIoStatus FOnDemandInstallCache::Reset()
 {
-	UE_LOG(LogIoStoreOnDemand, Log, TEXT("Recreating install cache in directory '%s'"), *CacheDirectory);
+	UE_LOG(LogIoStoreOnDemand, Log, TEXT("Resetting install cache in directory '%s'"), *CacheDirectory);
 
 	IFileManager& Ifm	= IFileManager::Get();
 	const bool bTree	= true;
diff --git a/Engine/Source/Runtime/Experimental/IoStore/OnDemand/Private/Statistics.cpp b/Engine/Source/Runtime/Experimental/IoStore/OnDemand/Private/Statistics.cpp
--- a/Engine/Source/Runtime/Experimental/IoStore/OnDemand/Private/Statistics.cpp
+++ b/Engine/Source/Runtime/Experimental/IoStore/OnDemand/Private/Statistics.cpp
@@ -452,7 +452,6 @@ struct FInstallCacheAnalytics
 	static FMutex Mutex; 
 
 	uint64 VerificationRemovedBlockCount = 0;
-	uint64 StartupErrorCount = 0;
 	uint64 FlushCount = 0;
 	uint64 FlushErrorCount = 0;
 	uint64 FlushedBytes = 0;
@@ -472,6 +471,7 @@ struct FInstallCacheAnalytics
 	uint64 MaxReferencedSize = 0;
 	uint64 MaxFragmentedSize = 0;
 	int64 OldestBlockAccess = FDateTime::MaxValue().GetTicks();
+	uint32 StartupErrorCode = 0;
 };
 FMutex FInstallCacheAnalytics::Mutex;
 static FInstallCacheAnalytics InstallCacheAnalytics;
@@ -1192,7 +1192,7 @@ void FOnDemandContentInstallerStats::ReportAnalytics(TArray<FAnalyticsEventAttri
 		, TEXT("IadTotalInstallDurationMs"), CurrentInstallerAnalytics.TotalInstallDurationMs 
 		, TEXT("IadAvgInstallDurationMs"), (double(CurrentInstallerAnalytics.TotalInstallDurationMs) / double(CurrentInstallerAnalytics.InstallCount))
 		, TEXT("IadAvgCacheHitRatio"), (CurrentInstallerAnalytics.TotalCacheHitRatio / double(CurrentInstallerAnalytics.InstallCount))
-		, TEXT("IadInstallCacheStartupErrorCount"), CurrentInstallCacheAnalytics.StartupErrorCount
+		, TEXT("IadInstallCacheStartupErrorCode"), CurrentInstallCacheAnalytics.StartupErrorCode
 		, TEXT("IadInstallCacheVerificationRemovedBlockCount"), CurrentInstallCacheAnalytics.VerificationRemovedBlockCount 
 		, TEXT("IadInstallCacheFlushCount"), CurrentInstallCacheAnalytics.FlushCount 
 		, TEXT("IadInstallCacheFlushErrorCount"), CurrentInstallCacheAnalytics.FlushErrorCount 
@@ -1217,8 +1217,11 @@ void FOnDemandContentInstallerStats::ReportAnalytics(TArray<FAnalyticsEventAttri
 ////////////////////////////////////////////////////////////////////////////////
 void FOnDemandInstallCacheStats::OnStartupError(EIoErrorCode ErrorCode)
 {
-	UE::TUniqueLock Lock(InstallCacheAnalytics);
-	InstallCacheAnalytics.StartupErrorCount++;
+	if (ErrorCode != EIoErrorCode::Ok)
+	{
+		UE::TUniqueLock Lock(InstallCacheAnalytics);
+		InstallCacheAnalytics.StartupErrorCode = static_cast<uint32>(ErrorCode);
+	}
 }
 
 void FOnDemandInstallCacheStats::OnFlush(EIoErrorCode ErrorCode, int64 ByteCount)
