Fix Indexed Cache Storage when entries are registered, ICS is enabled but no entry was created before (ie EnumerateCacheStorages returns an empty array).
We want to enable the system if NumIndexedCacheEntries > 0 and we found entries in the ini file

diff --git a/Engine/Source/Runtime/Experimental/IndexedCacheStorage/Private/IndexedCacheStorageManager.cpp b/Engine/Source/Runtime/Experimental/IndexedCacheStorage/Private/IndexedCacheStorageManager.cpp
--- a/Engine/Source/Runtime/Experimental/IndexedCacheStorage/Private/IndexedCacheStorageManager.cpp
+++ b/Engine/Source/Runtime/Experimental/IndexedCacheStorage/Private/IndexedCacheStorageManager.cpp
@@ -34,6 +34,7 @@ void FIndexedCacheStorageManager::Initialize()
 	TArray<FString> StorageConfigs;
 	GConfig->GetArray(TEXT("IndexedCacheStorage"), TEXT("Storage"), StorageConfigs, GEngineIni);
 
+	uint32 RegisteredStorages = 0;
 	for (const FString& Category : StorageConfigs)
 	{
 		FString TrimmedCategory = Category;
@@ -97,6 +98,7 @@ void FIndexedCacheStorageManager::Initialize()
 
 			ExistingStorageIndex = StorageIndex;
 			CacheStorageMetaData.MountName = StorageMount;
+			++RegisteredStorages;
 		}
 	}
 
@@ -107,7 +109,7 @@ void FIndexedCacheStorageManager::Initialize()
 	{
 		IndexedCacheEntriesMetaData[CacheStorageIndex].bCacheExists = true;
 	}
-	bHasRegisteredEntries = (CacheStorageIndices.Num() > 0);
+	bHasRegisteredEntries = (RegisteredStorages > 0);
 
 }
 
@@ -264,4 +266,4 @@ FString FIndexedCacheStorageManager::GetMountPath(int32 CacheIndex)
 	return BundleMountPath;
 }
 
-}
\ No newline at end of file
+}
