Gracefully handle reaching the end of the bitstream even if it's unexpected. The behavior is controlled via cvar net.Iris.ReplicationReader.GracefullyHandleReachingEndOfBitstream which defaults to true.

diff --git a/Engine/Source/Runtime/Experimental/Iris/Core/Private/Iris/ReplicationSystem/ReplicationReader.cpp b/Engine/Source/Runtime/Experimental/Iris/Core/Private/Iris/ReplicationSystem/ReplicationReader.cpp
--- a/Engine/Source/Runtime/Experimental/Iris/Core/Private/Iris/ReplicationSystem/ReplicationReader.cpp
+++ b/Engine/Source/Runtime/Experimental/Iris/Core/Private/Iris/ReplicationSystem/ReplicationReader.cpp
@@ -131,6 +131,12 @@ static FAutoConsoleVariableRef CVarMaxMustBeMappedHandleArray(
 	MaxMustBeMappedHandleArray,
 	TEXT("Maximum number of unresolved references that will be included with the QueuedBatch timeout error message."));
 
+static bool bGracefullyHandleReachingEndOfBitstream = true;
+static FAutoConsoleVariableRef CvarGracefullyHandleReachingEndOfBitstream(
+		TEXT("net.Iris.ReplicationReader.GracefullyHandleReachingEndOfBitstream"),
+		bGracefullyHandleReachingEndOfBitstream,
+		TEXT("Gracefully handle reaching end of the ReplicationReader bitstream prior to reading as many object batches as expected. Default is true."));
+
 static const FName NetError_FailedToFindAttachmentQueue("Failed to find attachment queue");
 
 //------------------------------------------------------------------------
@@ -2486,14 +2492,23 @@ void FReplicationReader::ReadObjects(FNetSerializationContext& Context, uint32 O
 
 	FNetBitStreamReader& Reader = *Context.GetBitStreamReader();
 	
-	while (ObjectBatchCountToRead && !Context.HasErrorOrOverflow())
+	while (ObjectBatchCountToRead && !Context.HasErrorOrOverflow() && (!bGracefullyHandleReachingEndOfBitstream || (Reader.GetBitsLeft() > 0)))
 	{
 		ReadObjectBatch(Context, ReadObjectFlags);
 		--ObjectBatchCountToRead;
 	}
 
-	UE_CLOG(Context.HasErrorOrOverflow(), LogIris, Error, TEXT("Overflow: %c Error: %s Bit stream bits left: %u position: %u %s"), TEXT("YN")[Context.HasError()], ToCStr(Context.GetError().ToString()), Reader.GetBitsLeft(), Reader.GetPosBits(), *Context.PrintReadJournal())
-	ensure(!Context.HasErrorOrOverflow());
+	UE_CLOG(Context.HasErrorOrOverflow(), LogIris, Error, TEXT("Overflow: %c Error: %s Bit stream bits left: %u position: %u %s"), TEXT("YN")[Context.HasError()], ToCStr(Context.GetError().ToString()), Reader.GetBitsLeft(), Reader.GetPosBits(), *Context.PrintReadJournal());
+
+	if (ensure(!Context.HasErrorOrOverflow()))
+	{
+		if (bGracefullyHandleReachingEndOfBitstream)
+		{
+			// ObjectBatchCountToRead should be zero at this point otherwise there's a problem on the writing side.
+			UE_CLOG(ObjectBatchCountToRead > 0, LogIris, Error, TEXT("Reached end of bitstream prior to reading all expected objects. %u left to read"), ObjectBatchCountToRead);
+			ensure(ObjectBatchCountToRead == 0);
+		}
+	}
 }
 
 void FReplicationReader::ProcessHugeObjectAttachment(FNetSerializationContext& Context, const TRefCountPtr<FNetBlob>& Attachment)
