Consolidate message bounds determination closer to the buffer that backs it.

#rnx

diff --git a/Engine/Source/Runtime/Experimental/IoStore/HttpClient/Private/Activity.inl b/Engine/Source/Runtime/Experimental/IoStore/HttpClient/Private/Activity.inl
--- a/Engine/Source/Runtime/Experimental/IoStore/HttpClient/Private/Activity.inl
+++ b/Engine/Source/Runtime/Experimental/IoStore/HttpClient/Private/Activity.inl
@@ -18,7 +18,7 @@ public:
 	FAnsiStringView	GetMethod() const;
 	void			EnumerateHeaders(FResponse::FHeaderSink Sink) const;
 
-private:
+protected:
 	FRequestCloak&	operator << (FAnsiStringView Value);
 	uint16			HeaderLeft;
 	int16			HeaderRight = -1;
@@ -75,6 +75,87 @@ FRequestCloak& FRequestCloak::operator << (FAnsiStringView Value)
 
 
 
+////////////////////////////////////////////////////////////////////////////////
+class FResponseCloak
+	: public FRequestCloak
+{
+public:
+	using FRequestCloak::FRequestCloak;
+
+	FOutcome	TryRecvResponse(FHttpPeer& Peer);
+	const char*	TEMP_GetMessageRight() const;
+
+private:
+	int16		MessageRight = -1;
+};
+
+////////////////////////////////////////////////////////////////////////////////
+const char* FResponseCloak::TEMP_GetMessageRight() const
+{
+	return GetData() + MessageRight;
+}
+
+////////////////////////////////////////////////////////////////////////////////
+FOutcome FResponseCloak::TryRecvResponse(FHttpPeer& Peer)
+{
+	auto FindMessageTerminal = [] (const char* Cursor, int32 Length)
+	{
+		for (int32 i = 4; i <= Length; ++i)
+		{
+			uint32 Candidate;
+			::memcpy(&Candidate, Cursor + i - 4, sizeof(Candidate));
+			if (Candidate == 0x0a0d0a0d)
+			{
+				return i;
+			}
+
+			i += (Cursor[i - 1] > 0x0d) ? 3 : 0;
+		}
+
+		return -1;
+	};
+
+	while (true)
+	{
+		static const uint32 PageSize = 256;
+		static const uint32 MaxHeaderSize = 8 << 10;
+
+		auto [Dest, DestSize] = GetMutableFree(0, PageSize);
+		FOutcome Outcome = Peer.Recv(Dest, DestSize);
+		if (!Outcome.IsOk())
+		{
+			return Outcome;
+		}
+
+		int32 Result = Outcome.GetResult();
+		AdvanceUsed(Result);
+
+		// Rewind a little to cover cases where the terminal is fragmented across
+		// recv() calls
+		uint32 DestBias = 0;
+		if (Dest - 3 >= GetData() + HeaderRight)
+		{
+			Dest -= (DestBias = 3);
+		}
+
+		int32 MessageEnd = FindMessageTerminal(Dest, Result + DestBias);
+		if (MessageEnd < 0)
+		{
+			if (GetSize() - HeaderRight > MaxHeaderSize)
+			{
+				return FOutcome::Error("Headers have grown larger than expected");
+			}
+
+			continue;
+		}
+
+		MessageRight = int16(ptrdiff_t(Dest + MessageEnd - GetData()));
+		return FOutcome::Ok();
+	}
+}
+
+
+
 #if IAS_HTTP_WITH_PERF
 
 ////////////////////////////////////////////////////////////////////////////////
@@ -172,7 +253,7 @@ struct alignas(16) FActivity
 	UPTRINT				SinkParam;
 	FTicketSink			Sink;
 	FResponseInternal	Response;
-	FRequestCloak		Buffer;
+	FResponseCloak		Buffer;
 };
 
 ////////////////////////////////////////////////////////////////////////////////
diff --git a/Engine/Source/Runtime/Experimental/IoStore/HttpClient/Private/Loop.inl b/Engine/Source/Runtime/Experimental/IoStore/HttpClient/Private/Loop.inl
--- a/Engine/Source/Runtime/Experimental/IoStore/HttpClient/Private/Loop.inl
+++ b/Engine/Source/Runtime/Experimental/IoStore/HttpClient/Private/Loop.inl
@@ -103,55 +103,23 @@ static FOutcome DoRecvMessage(FActivity* Activity, FHttpPeer& Peer)
 	Activity->Stopwatch.RecvStart();
 #endif
 
-	FBuffer& Buffer = Activity->Buffer;
+	auto& Buffer = Activity->Buffer;
 
-	const char* MessageRight;
-	while (true)
+	FOutcome Outcome = Buffer.TryRecvResponse(Peer);
+	if (Outcome.IsError())
 	{
-		static const uint32 PageSize = 256;
-		auto [Dest, DestSize] = Buffer.GetMutableFree(0, PageSize);
-		FOutcome Outcome = Peer.Recv(Dest, DestSize);
-
-		if (Outcome.IsError())
-		{
-			Activity_SetError(Activity, Outcome);
-			return Outcome;
-		}
-
-		if (Outcome.IsWaiting())
-		{
-			return Outcome;
-		}
-
-		check(Outcome.IsOk());
-
-		int32 Result = Outcome.GetResult();
-		Buffer.AdvanceUsed(Result);
-
-		// Rewind a little to cover cases where the terminal is fragmented across
-		// recv() calls
-		uint32 DestBias = 0;
-		if (Dest - 3 >= Buffer.GetData() + Activity->StateParam)
-		{
-			Dest -= (DestBias = 3);
-		}
-
-		int32 MessageEnd = FindMessageTerminal(Dest, Result + DestBias);
-		if (MessageEnd < 0)
-		{
-			if (Buffer.GetSize() > (8 << 10))
-			{
-				Activity_SetError(Activity, "Headers have grown larger than expected");
-				return FOutcome::Error(Activity->ErrorReason);
-			}
-
-			continue;
-		}
+		Activity_SetError(Activity, Outcome);
+		return Outcome;
+	}
 
-		MessageRight = Dest + MessageEnd;
-		break;
+	if (Outcome.IsWaiting())
+	{
+		return Outcome;
 	}
 
+	check(Outcome.IsOk());
+	const char* MessageRight = Buffer.TEMP_GetMessageRight();
+
 	// Fill out the internal response object
 	FResponseInternal& Internal = Activity->Response;
 	const char* MessageData = Buffer.GetData() + Activity->StateParam;
