Unused function delete.

#rnx

diff --git a/Engine/Source/Runtime/Experimental/IoStore/HttpClient/Private/Misc.inl b/Engine/Source/Runtime/Experimental/IoStore/HttpClient/Private/Misc.inl
--- a/Engine/Source/Runtime/Experimental/IoStore/HttpClient/Private/Misc.inl
+++ b/Engine/Source/Runtime/Experimental/IoStore/HttpClient/Private/Misc.inl
@@ -188,24 +188,6 @@ static void EnumerateHeaders(FAnsiStringView Headers, LambdaType&& Lambda)
 	} while (Cursor < End);
 }
 
-////////////////////////////////////////////////////////////////////////////////
-static int32 FindMessageTerminal(const char* Data, uint32 Length)
-{
-	for (uint32 i = 4; i <= Length; ++i)
-	{
-		uint32 Candidate;
-		::memcpy(&Candidate, Data + i - 4, sizeof(Candidate));
-		if (Candidate == 0x0a0d0a0d)
-		{
-			return i;
-		}
-
-		i += (Data[i - 1] > 0x0d) ? 3 : 0;
-	}
-
-	return -1;
-}
-
 ////////////////////////////////////////////////////////////////////////////////
 template <uint32 Base=10>
 static int64 CrudeToInt(FAnsiStringView View)
diff --git a/Engine/Source/Runtime/Experimental/IoStore/HttpClient/Private/Test.inl b/Engine/Source/Runtime/Experimental/IoStore/HttpClient/Private/Test.inl
--- a/Engine/Source/Runtime/Experimental/IoStore/HttpClient/Private/Test.inl
+++ b/Engine/Source/Runtime/Experimental/IoStore/HttpClient/Private/Test.inl
@@ -16,6 +16,7 @@ namespace UE::IoStore::HTTP
 static void MiscTest()
 {
 #define CRLF "\r\n"
+#if 0
 	struct {
 		FAnsiStringView Input;
 		int32 Output;
@@ -31,6 +32,7 @@ static void MiscTest()
 	{
 		check(FindMessageTerminal(Input.GetData(), Input.Len()) == Output);
 	}
+#endif // 0
 
 	FMessageOffsets MsgOut;
 	check(ParseMessage("", MsgOut) == -1);
