Iris
- Changed check to log/ensure to not crash sever.

diff --git a/Engine/Source/Runtime/Experimental/Iris/Core/Private/Iris/ReplicationSystem/ObjectReplicationBridge.cpp b/Engine/Source/Runtime/Experimental/Iris/Core/Private/Iris/ReplicationSystem/ObjectReplicationBridge.cpp
--- a/Engine/Source/Runtime/Experimental/Iris/Core/Private/Iris/ReplicationSystem/ObjectReplicationBridge.cpp
+++ b/Engine/Source/Runtime/Experimental/Iris/Core/Private/Iris/ReplicationSystem/ObjectReplicationBridge.cpp
@@ -619,7 +619,13 @@ UE::Net::FNetRefHandle UObjectReplicationBridge::StartReplicatingSubObject(UObje
 	if (SubObjectRefHandle.IsValid())
 	{
 		// Verify that the existing object is a subobject of the owner
-		check(Params.RootObjectHandle == LocalNetRefHandleManager.GetRootObjectOfSubObject(SubObjectRefHandle));
+		const FNetRefHandle ExistingRootObjectHandle = LocalNetRefHandleManager.GetRootObjectOfSubObject(SubObjectRefHandle);
+		if (Params.RootObjectHandle != ExistingRootObjectHandle)
+		{
+			UE_LOG_OBJECTREPLICATIONBRIDGE(Error, TEXT("UObjectReplicationBridge::StartReplicatingSubObject attempted to replicate already replicating subobject (%s) with a different rootobject (%s)."),
+				*LocalNetRefHandleManager.PrintObjectFromNetRefHandle(SubObjectRefHandle), *LocalNetRefHandleManager.PrintObjectFromNetRefHandle(Params.RootObjectHandle));
+			ensureMsgf(false, TEXT("Attempted to replicate already replicating subobject with a different root"));
+		}
 		return SubObjectRefHandle;
 	}
 
