Chaos Cloth:
- Support saving pre-resized sim 3D positions for scaling 2D distances of edge/axial springs.
- ApplyResizing node now writes old positions and a boolean flag to use it

diff --git a/Engine/Plugins/ChaosCloth/Source/ChaosCloth/Private/ChaosCloth/ChaosClothConstraints.cpp b/Engine/Plugins/ChaosCloth/Source/ChaosCloth/Private/ChaosCloth/ChaosClothConstraints.cpp
--- a/Engine/Plugins/ChaosCloth/Source/ChaosCloth/Private/ChaosCloth/ChaosClothConstraints.cpp
+++ b/Engine/Plugins/ChaosCloth/Source/ChaosCloth/Private/ChaosCloth/ChaosClothConstraints.cpp
@@ -793,7 +793,7 @@ void FClothConstraints::AddRules(
 	CreateSelfCollisionConstraints(ConfigProperties, WeightMaps, VertexSets, FaceSets, FaceIntMaps, TriangleMesh);
 
 	// Edge constraints
-	CreateStretchConstraints(ConfigProperties, WeightMaps, TriangleMesh, PatternData);
+	CreateStretchConstraints(ConfigProperties, WeightMaps, TriangleMesh, PatternData, ManagedArrayCollection, MeshScale);
 
 	// Bending constraints
 	CreateBendingConstraints(ConfigProperties, WeightMaps, TriangleMesh, PatternData);
@@ -924,7 +924,9 @@ void FClothConstraints::CreateStretchConstraints(
 	const Softs::FCollectionPropertyConstFacade& ConfigProperties,
 	const TMap<FString, TConstArrayView<FRealSingle>>& WeightMaps,
 	const FTriangleMesh& TriangleMesh,
-	const FClothingPatternData* PatternData)
+	const FClothingPatternData* PatternData,
+	const TSharedPtr<const FManagedArrayCollection>& ManagedArrayCollection,
+	Softs::FSolverReal MeshScale)
 {
 	if (PatternData && PatternData->PatternPositions.Num() && Softs::FXPBDStretchBiasElementConstraints::IsEnabled(ConfigProperties))
 	{
@@ -1021,6 +1023,17 @@ void FClothConstraints::CreateStretchConstraints(
 	}
 	if (PatternData && PatternData->PatternPositions.Num() && Softs::FXPBDAnisotropicSpringConstraints::IsEnabled(ConfigProperties))
 	{
+		TArray<FVector3f> PreResizedPositions;
+		if (ManagedArrayCollection)
+		{
+			if (const TManagedArray<FVector3f>* PreResizedAttr =
+				ManagedArrayCollection->FindAttributeTyped<FVector3f>(
+					TEXT("PreResizedSimPosition3D"),
+					TEXT("SimVertices3D")))
+			{
+				PreResizedPositions = PreResizedAttr->GetConstArray();
+			}
+		}
 		if (Evolution)
 		{
 			XAnisoSpringConstraints = MakeShared<Softs::FXPBDAnisotropicSpringConstraints>(
@@ -1028,7 +1041,9 @@ void FClothConstraints::CreateStretchConstraints(
 				TriangleMesh,
 				PatternData->WeldedFaceVertexPatternPositions,
 				WeightMaps,
-				ConfigProperties);
+				ConfigProperties,
+				PreResizedPositions,
+				MeshScale);
 			++NumUpdateLinearSystemRules;
 		}
 		else
@@ -1040,7 +1055,9 @@ void FClothConstraints::CreateStretchConstraints(
 				TriangleMesh,
 				PatternData->WeldedFaceVertexPatternPositions,
 				WeightMaps,
-				ConfigProperties);
+				ConfigProperties,
+				PreResizedPositions,
+				MeshScale);
 		}
 
 		++NumConstraintInits;  // Uses init to update the property tables
diff --git a/Engine/Plugins/ChaosCloth/Source/ChaosCloth/Public/ChaosCloth/ChaosClothConstraints.h b/Engine/Plugins/ChaosCloth/Source/ChaosCloth/Public/ChaosCloth/ChaosClothConstraints.h
--- a/Engine/Plugins/ChaosCloth/Source/ChaosCloth/Public/ChaosCloth/ChaosClothConstraints.h
+++ b/Engine/Plugins/ChaosCloth/Source/ChaosCloth/Public/ChaosCloth/ChaosClothConstraints.h
@@ -141,7 +141,9 @@ namespace Chaos
 			const Softs::FCollectionPropertyConstFacade& ConfigProperties,
 			const TMap<FString, TConstArrayView<FRealSingle>>& WeightMaps,
 			const FTriangleMesh& TriangleMesh,
-			const FClothingPatternData* PatternData);
+			const FClothingPatternData* PatternData,
+			const TSharedPtr<const FManagedArrayCollection>& ManagedArrayCollection,
+			Softs::FSolverReal MeshScale);
 		void CreateBendingConstraints(
 			const Softs::FCollectionPropertyConstFacade& ConfigProperties,
 			const TMap<FString, TConstArrayView<FRealSingle>>& WeightMaps,
diff --git a/Engine/Plugins/ChaosClothAsset/Source/ChaosClothAsset/Private/ChaosClothAsset/ClothCollection.cpp b/Engine/Plugins/ChaosClothAsset/Source/ChaosClothAsset/Private/ChaosClothAsset/ClothCollection.cpp
--- a/Engine/Plugins/ChaosClothAsset/Source/ChaosClothAsset/Private/ChaosClothAsset/ClothCollection.cpp
+++ b/Engine/Plugins/ChaosClothAsset/Source/ChaosClothAsset/Private/ChaosClothAsset/ClothCollection.cpp
@@ -90,11 +90,12 @@ namespace UE::Chaos::ClothAsset::Private
 
 	// Sim Vertices 3D Group
 	// NOTE: if you add anything here, you need to implement how to merge it 
-	// in CollectionClothSeamsFacade. Otherwise, the data in the lowest vertex index 
+	// in CollectionClothSeamFacade. Otherwise, the data in the lowest vertex index 
 	// will survive and the other data will be lost whenever seams are added.
 	static const TArray<FName> SimVertices3DGroupAttributes =
 	{
 		ClothCollectionAttribute::SimPosition3D,
+		ClothCollectionAttribute::PreResizedSimPosition3D,
 		ClothCollectionAttribute::SimNormal,
 		ClothCollectionAttribute::SimBoneIndices,
 		ClothCollectionAttribute::SimBoneWeights,
@@ -243,6 +244,7 @@ namespace UE::Chaos::ClothAsset
 
 		// Sim Vertices 3D Group
 		SimPosition3D = ManagedArrayCollection->FindAttribute<FVector3f>(ClothCollectionAttribute::SimPosition3D, ClothCollectionGroup::SimVertices3D);
+		PreResizedSimPosition3D = ManagedArrayCollection->FindAttribute<FVector3f>(ClothCollectionAttribute::PreResizedSimPosition3D, ClothCollectionGroup::SimVertices3D);
 		SimNormal = ManagedArrayCollection->FindAttribute<FVector3f>(ClothCollectionAttribute::SimNormal, ClothCollectionGroup::SimVertices3D);
 		SimBoneIndices = ManagedArrayCollection->FindAttribute<TArray<int32>>(ClothCollectionAttribute::SimBoneIndices, ClothCollectionGroup::SimVertices3D);
 		SimBoneWeights = ManagedArrayCollection->FindAttribute<TArray<float>>(ClothCollectionAttribute::SimBoneWeights, ClothCollectionGroup::SimVertices3D);
@@ -300,7 +302,10 @@ namespace UE::Chaos::ClothAsset
 				// Sim Morph Target Vertices Group
 				SimMorphTargetPositionDelta &&
 				SimMorphTargetTangentZDelta &&
-				SimMorphTargetSimVertex3DIndex;
+				SimMorphTargetSimVertex3DIndex &&
+
+				// Sim Vertices 3D Group
+				PreResizedSimPosition3D;
 		}
 
 		return
@@ -371,7 +376,7 @@ namespace UE::Chaos::ClothAsset
 			SimVertex2DLookup &&
 			SeamStitchLookup &&
 			(!EnumHasAnyFlags(OptionalSchemas, EClothCollectionExtendedSchemas::Resizing) ||
-				SimCustomResizingBlend) &&
+				(SimCustomResizingBlend && PreResizedSimPosition3D)) &&
 
 			// Sim Morph Targets Group
 			(!EnumHasAnyFlags(OptionalSchemas, EClothCollectionExtendedSchemas::SimMorphTargets) ||
@@ -442,6 +447,9 @@ namespace UE::Chaos::ClothAsset
 			SimMorphTargetPositionDelta = &ManagedArrayCollection->AddAttribute<FVector3f>(ClothCollectionAttribute::SimMorphTargetPositionDelta, ClothCollectionGroup::SimMorphTargetVertices);
 			SimMorphTargetTangentZDelta = &ManagedArrayCollection->AddAttribute<FVector3f>(ClothCollectionAttribute::SimMorphTargetTangentZDelta, ClothCollectionGroup::SimMorphTargetVertices);
 			SimMorphTargetSimVertex3DIndex = &ManagedArrayCollection->AddAttribute<int32>(ClothCollectionAttribute::SimMorphTargetSimVertex3DIndex, ClothCollectionGroup::SimMorphTargetVertices);
+			
+			// Sim Vertices 3D Group
+			PreResizedSimPosition3D = &ManagedArrayCollection->AddAttribute<FVector3f>(ClothCollectionAttribute::PreResizedSimPosition3D, ClothCollectionGroup::SimVertices3D);
 
 			return;
 		}
@@ -558,6 +566,7 @@ namespace UE::Chaos::ClothAsset
 		{
 			// Sim Vertices 3D Group
 			SimCustomResizingBlend = &ManagedArrayCollection->AddAttribute<float>(ClothCollectionAttribute::SimCustomResizingBlend, ClothCollectionGroup::SimVertices3D);
+			PreResizedSimPosition3D = &ManagedArrayCollection->AddAttribute<FVector3f>(ClothCollectionAttribute::PreResizedSimPosition3D, ClothCollectionGroup::SimVertices3D);
 
 			// Render Vertices Group
 			RenderCustomResizingBlend = &ManagedArrayCollection->AddAttribute<float>(ClothCollectionAttribute::RenderCustomResizingBlend, ClothCollectionGroup::RenderVertices);
diff --git a/Engine/Plugins/ChaosClothAsset/Source/ChaosClothAsset/Private/ChaosClothAsset/ClothCollection.h b/Engine/Plugins/ChaosClothAsset/Source/ChaosClothAsset/Private/ChaosClothAsset/ClothCollection.h
--- a/Engine/Plugins/ChaosClothAsset/Source/ChaosClothAsset/Private/ChaosClothAsset/ClothCollection.h
+++ b/Engine/Plugins/ChaosClothAsset/Source/ChaosClothAsset/Private/ChaosClothAsset/ClothCollection.h
@@ -213,6 +213,7 @@ namespace UE::Chaos::ClothAsset
 
 		//~ Sim Vertices 3D Group
 		const TManagedArray<FVector3f>* GetSimPosition3D() const { return SimPosition3D; }
+		const TManagedArray<FVector3f>* GetPreResizedSimPosition3D() const { return PreResizedSimPosition3D; }
 		const TManagedArray<FVector3f>* GetSimNormal() const { return SimNormal; }
 		const TManagedArray<TArray<int32>>* GetSimBoneIndices() const { return SimBoneIndices; }
 		const TManagedArray<TArray<float>>* GetSimBoneWeights() const { return SimBoneWeights; }
@@ -311,6 +312,7 @@ namespace UE::Chaos::ClothAsset
 
 		//~ Sim Vertices 3D Group
 		TManagedArray<FVector3f>* GetSimPosition3D() { return SimPosition3D; }
+		TManagedArray<FVector3f>* GetPreResizedSimPosition3D() { return PreResizedSimPosition3D; }
 		TManagedArray<FVector3f>* GetSimNormal() { return SimNormal; }
 		TManagedArray<TArray<int32>>* GetSimBoneIndices() { return SimBoneIndices; }
 		TManagedArray<TArray<float>>* GetSimBoneWeights() { return SimBoneWeights; }
@@ -421,6 +423,7 @@ namespace UE::Chaos::ClothAsset
 		TManagedArray<TArray<int32>>* SimVertex2DLookup; // Lookup into corresponding 2D vertices
 		TManagedArray<TArray<int32>>* SeamStitchLookup; // Lookup into any seam stitches which weld this vertex
 		TManagedArray<float>* SimCustomResizingBlend; // Resizing weight map. How much to use group resizing result (1) vs rbf interpolation result (0) when doing sim mesh resizing
+		TManagedArray<FVector3f>* PreResizedSimPosition3D;
 
 		//~ Sim Morph Targets Group
 		TManagedArray<FString>* SimMorphTargetName;
diff --git a/Engine/Plugins/ChaosClothAsset/Source/ChaosClothAsset/Private/ChaosClothAsset/CollectionClothFacade.cpp b/Engine/Plugins/ChaosClothAsset/Source/ChaosClothAsset/Private/ChaosClothAsset/CollectionClothFacade.cpp
--- a/Engine/Plugins/ChaosClothAsset/Source/ChaosClothAsset/Private/ChaosClothAsset/CollectionClothFacade.cpp
+++ b/Engine/Plugins/ChaosClothAsset/Source/ChaosClothAsset/Private/ChaosClothAsset/CollectionClothFacade.cpp
@@ -109,6 +109,7 @@ namespace UE::Chaos::ClothAsset
 		ResultHash = HashCombineFast(ResultHash, ClothCollection->GetElementsTypeHash(ClothCollection->GetSimVertex2DLookup()));
 		ResultHash = HashCombineFast(ResultHash, ClothCollection->GetElementsTypeHash(ClothCollection->GetSeamStitchLookup()));
 		ResultHash = HashCombineFast(ResultHash, ClothCollection->GetElementsTypeHash(ClothCollection->GetSimCustomResizingBlend()));
+		ResultHash = HashCombineFast(ResultHash, ClothCollection->GetElementsTypeHash(ClothCollection->GetPreResizedSimPosition3D()));
 
 		//~ Sim Morph Targets Group
 		ResultHash = HashCombineFast(ResultHash, ClothCollection->GetElementsTypeHash(ClothCollection->GetSimMorphTargetName()));
@@ -254,6 +255,11 @@ namespace UE::Chaos::ClothAsset
 		return ClothCollection->GetElements(ClothCollection->GetSimPosition3D());
 	}
 
+	TConstArrayView<FVector3f> FCollectionClothConstFacade::GetPreResizedSimPosition3D() const
+	{
+		return ClothCollection->GetElements(ClothCollection->GetPreResizedSimPosition3D());
+	}
+
 	TConstArrayView<FVector3f> FCollectionClothConstFacade::GetSimNormal() const
 	{
 		return ClothCollection->GetElements(ClothCollection->GetSimNormal());
@@ -724,6 +730,7 @@ namespace UE::Chaos::ClothAsset
 		if (bOtherHasResizingData)
 		{
 			FClothCollection::CopyArrayViewData(GetSimCustomResizingBlend().Right(OtherNumSimVertices3D), Other.GetSimCustomResizingBlend());
+			FClothCollection::CopyArrayViewData(GetPreResizedSimPosition3D().Right(OtherNumSimVertices3D), Other.GetPreResizedSimPosition3D());
 		}
 
 		// Fabrics Group
@@ -852,6 +859,10 @@ namespace UE::Chaos::ClothAsset
 		{
 			GetSimMorphTarget(TargetIndex).Initialize(Other.GetSimMorphTarget(TargetIndex), 0);
 		}
+
+		// Sim Vertices 3D Group
+		GetClothCollection()->SetNumElements(Other.GetNumSimVertices3D(), ClothCollectionGroup::SimVertices3D);
+		FClothCollection::CopyArrayViewData(GetPreResizedSimPosition3D().Right(Other.GetNumSimVertices3D()), Other.GetPreResizedSimPosition3D());
 	}
 
 	void FCollectionClothFacade::SetPhysicsAssetPathName(const FString& PathName)
@@ -921,6 +932,11 @@ namespace UE::Chaos::ClothAsset
 		return GetClothCollection()->GetElements(GetClothCollection()->GetSimPosition3D());
 	}
 
+	TArrayView<FVector3f> FCollectionClothFacade::GetPreResizedSimPosition3D()
+	{
+		return GetClothCollection()->GetElements(GetClothCollection()->GetPreResizedSimPosition3D());
+	}
+
 	TArrayView<FVector3f> FCollectionClothFacade::GetSimNormal()
 	{
 		return GetClothCollection()->GetElements(GetClothCollection()->GetSimNormal());
diff --git a/Engine/Plugins/ChaosClothAsset/Source/ChaosClothAsset/Public/ChaosClothAsset/ClothCollectionAttribute.h b/Engine/Plugins/ChaosClothAsset/Source/ChaosClothAsset/Public/ChaosClothAsset/ClothCollectionAttribute.h
--- a/Engine/Plugins/ChaosClothAsset/Source/ChaosClothAsset/Public/ChaosClothAsset/ClothCollectionAttribute.h
+++ b/Engine/Plugins/ChaosClothAsset/Source/ChaosClothAsset/Public/ChaosClothAsset/ClothCollectionAttribute.h
@@ -66,6 +66,7 @@ namespace UE::Chaos::ClothAsset
 
 		// Sim Vertices 3D Group
 		inline const FName SimPosition3D(TEXT("SimPosition3D"));
+		inline const FName PreResizedSimPosition3D(TEXT("PreResizedSimPosition3D"));
 		inline const FName SimNormal(TEXT("SimNormal"));
 		inline const FName SimBoneIndices(TEXT("SimBoneIndices"));
 		inline const FName SimBoneWeights(TEXT("SimBoneWeights"));
diff --git a/Engine/Plugins/ChaosClothAsset/Source/ChaosClothAsset/Public/ChaosClothAsset/CollectionClothFacade.h b/Engine/Plugins/ChaosClothAsset/Source/ChaosClothAsset/Public/ChaosClothAsset/CollectionClothFacade.h
--- a/Engine/Plugins/ChaosClothAsset/Source/ChaosClothAsset/Public/ChaosClothAsset/CollectionClothFacade.h
+++ b/Engine/Plugins/ChaosClothAsset/Source/ChaosClothAsset/Public/ChaosClothAsset/CollectionClothFacade.h
@@ -96,6 +96,7 @@ namespace UE::Chaos::ClothAsset
 		/** Return the total number of 3D simulation vertices for this collection. */
 		UE_API int32 GetNumSimVertices3D() const;
 		UE_API TConstArrayView<FVector3f> GetSimPosition3D() const;
+		UE_API TConstArrayView<FVector3f> GetPreResizedSimPosition3D() const;
 		UE_API TConstArrayView<FVector3f> GetSimNormal() const;
 		UE_API TConstArrayView<TArray<int32>> GetSimBoneIndices() const;
 		UE_API TConstArrayView<TArray<float>> GetSimBoneWeights() const;
@@ -283,6 +284,7 @@ namespace UE::Chaos::ClothAsset
 
 		//~ Pattern Sim Vertices 3D Group
 		UE_API TArrayView<FVector3f> GetSimPosition3D();
+		UE_API TArrayView<FVector3f> GetPreResizedSimPosition3D();
 		UE_API TArrayView<FVector3f> GetSimNormal();
 		UE_API TArrayView<TArray<int32>> GetSimBoneIndices();
 		UE_API TArrayView<TArray<float>> GetSimBoneWeights();
diff --git a/Engine/Plugins/ChaosClothAssetEditor/Source/ChaosClothAssetDataflowNodes/Private/ChaosClothAsset/ApplyResizingNode.cpp b/Engine/Plugins/ChaosClothAssetEditor/Source/ChaosClothAssetDataflowNodes/Private/ChaosClothAsset/ApplyResizingNode.cpp
--- a/Engine/Plugins/ChaosClothAssetEditor/Source/ChaosClothAssetDataflowNodes/Private/ChaosClothAsset/ApplyResizingNode.cpp
+++ b/Engine/Plugins/ChaosClothAssetEditor/Source/ChaosClothAssetDataflowNodes/Private/ChaosClothAsset/ApplyResizingNode.cpp
@@ -1,6 +1,7 @@
 // Copyright Epic Games, Inc. All Rights Reserved.
 
 #include "ChaosClothAsset/ApplyResizingNode.h"
+#include "Chaos/CollectionPropertyFacade.h"
 #include "ChaosClothAsset/ClothGeometryTools.h"
 #include "ChaosClothAsset/CollectionClothFacade.h"
 #include "Dataflow/DataflowInputOutput.h"
@@ -12,7 +13,7 @@
 #define LOCTEXT_NAMESPACE "ChaosClothAssetApplyResizingNode"
 
 FChaosClothAssetApplyResizingNode::FChaosClothAssetApplyResizingNode(const UE::Dataflow::FNodeParameters& InParam, FGuid InGuid)
-	: FDataflowNode(InParam, InGuid)
+	: FChaosClothAssetSimulationBaseConfigNode(InParam, InGuid)
 {
 	RegisterInputConnection(&Collection);
 	RegisterOutputConnection(&Collection, &Collection);
@@ -30,79 +31,95 @@ FChaosClothAssetApplyResizingNode::FChaosClothAssetApplyResizingNode(const UE::D
 	RegisterInputConnection(&bSkipCustomRegionResizing)
 		.SetCanHidePin(true)
 		.SetPinIsHidden(true);
+	RegisterInputConnection(&bSavePreResizedSimPosition3D)
+		.SetCanHidePin(true)
+		.SetPinIsHidden(true);
 }
 
-void FChaosClothAssetApplyResizingNode::Evaluate(UE::Dataflow::FContext& Context, const FDataflowOutput* Out) const
+void FChaosClothAssetApplyResizingNode::EvaluateClothCollection(UE::Dataflow::FContext& Context, const TSharedRef<FManagedArrayCollection>& ClothCollection) const
 {
-	if (Out->IsA<FManagedArrayCollection>(&Collection))
+	using namespace UE::Chaos::ClothAsset;
+	FCollectionClothFacade ClothFacade(ClothCollection);
+	if (ClothFacade.IsValid())
 	{
-		using namespace UE::Chaos::ClothAsset;
-		// Evaluate in collection
-		FManagedArrayCollection InCollection = GetValue<FManagedArrayCollection>(Context, &Collection);
-		const TSharedRef<FManagedArrayCollection> ClothCollection = MakeShared<FManagedArrayCollection>(MoveTemp(InCollection));
-		FCollectionClothFacade ClothFacade(ClothCollection);
-		if (ClothFacade.IsValid())
+		if (TObjectPtr<const USkeletalMesh> InTargetSKM = GetValue(Context, &TargetSkeletalMesh))
 		{
-			if (TObjectPtr<const USkeletalMesh> InTargetSKM = GetValue(Context, &TargetSkeletalMesh))
+			const int32 InLODIndex = GetValue(Context, &SkeletalMeshLODIndex);
+			if (InTargetSKM->IsValidLODIndex(InLODIndex) && InTargetSKM->HasMeshDescription(InLODIndex))
 			{
-				const int32 InLODIndex = GetValue(Context, &SkeletalMeshLODIndex);
-				if (InTargetSKM->IsValidLODIndex(InLODIndex) && InTargetSKM->HasMeshDescription(InLODIndex))
+				if (const FMeshDescription* const TargetMeshDescription = InTargetSKM->GetMeshDescription(InLODIndex))
 				{
-					if (const FMeshDescription* const TargetMeshDescription = InTargetSKM->GetMeshDescription(InLODIndex))
+					const FMeshResizingRBFInterpolationData& InInterpolationData = GetValue(Context, &InterpolationData);
+					if (InInterpolationData.SampleIndices.Num())
 					{
-						const FMeshResizingRBFInterpolationData& InInterpolationData = GetValue(Context, &InterpolationData);
-						if (InInterpolationData.SampleIndices.Num())
+						if (!GetValue(Context, &bForceApplyToRenderMesh) && ClothFacade.HasValidSimulationData())
 						{
-							if (!GetValue(Context, &bForceApplyToRenderMesh) && ClothFacade.HasValidSimulationData())
+							if (bSavePreResizedSimPosition3D)
 							{
-								TArray<FMeshResizingCustomRegion> ResizingGroupData;
-								if (!GetValue(Context, &bSkipCustomRegionResizing))
+								if (!ClothFacade.IsValid(EClothCollectionExtendedSchemas::Resizing))
 								{
-									if (TObjectPtr<const USkeletalMesh> InSourceSKM = GetValue(Context, &SourceSkeletalMesh))
-									{
-										if (const FMeshDescription* const SourceMeshDescription = InTargetSKM->GetMeshDescription(InLODIndex))
-										{
-											FClothGeometryTools::GenerateSimMeshResizingGroupData(ClothCollection, *SourceMeshDescription, ResizingGroupData);
-										}
-									}
+									ClothFacade.DefineSchema(EClothCollectionExtendedSchemas::Resizing);
 								}
+								TArrayView<FVector3f> PreResizedView = ClothFacade.GetPreResizedSimPosition3D();
+								const TConstArrayView<FVector3f> SimPosView = ClothFacade.GetSimPosition3D();
 
-								UE::MeshResizing::FRBFInterpolation::DeformPoints(*TargetMeshDescription, InInterpolationData, ClothFacade.GetSimPosition3D(), ClothFacade.GetSimNormal());
-
-								if (!ResizingGroupData.IsEmpty())
+								if (ensure(PreResizedView.Num() == SimPosView.Num()))
 								{
-									FClothGeometryTools::ApplySimGroupResizing(ClothCollection, *TargetMeshDescription, InInterpolationData, ResizingGroupData);
+									for (int32 i = 0; i < SimPosView.Num(); ++i)
+									{
+										PreResizedView[i] = SimPosView[i];
+									}
 								}
 							}
-							else if (ClothFacade.HasValidRenderData())
+							TArray<FMeshResizingCustomRegion> ResizingGroupData;
+							if (!GetValue(Context, &bSkipCustomRegionResizing))
 							{
-								TArray<FMeshResizingCustomRegion> ResizingGroupData;
-								if (!GetValue(Context, &bSkipCustomRegionResizing))
+								if (TObjectPtr<const USkeletalMesh> InSourceSKM = GetValue(Context, &SourceSkeletalMesh))
 								{
-									if (TObjectPtr<const USkeletalMesh> InSourceSKM = GetValue(Context, &SourceSkeletalMesh))
+									if (const FMeshDescription* const SourceMeshDescription = InTargetSKM->GetMeshDescription(InLODIndex))
 									{
-										if (const FMeshDescription* const SourceMeshDescription = InTargetSKM->GetMeshDescription(InLODIndex))
-										{
-											FClothGeometryTools::GenerateRenderMeshResizingGroupData(ClothCollection, *SourceMeshDescription, ResizingGroupData);
-										}
+										FClothGeometryTools::GenerateSimMeshResizingGroupData(ClothCollection, *SourceMeshDescription, ResizingGroupData);
 									}
 								}
+							}
 
-								UE::MeshResizing::FRBFInterpolation::DeformPoints(*TargetMeshDescription, InInterpolationData, ClothFacade.GetRenderPosition(), ClothFacade.GetRenderNormal(), ClothFacade.GetRenderTangentU(), ClothFacade.GetRenderTangentV());
+							UE::MeshResizing::FRBFInterpolation::DeformPoints(*TargetMeshDescription, InInterpolationData, ClothFacade.GetSimPosition3D(), ClothFacade.GetSimNormal());
 
-								if (!ResizingGroupData.IsEmpty())
+							if (!ResizingGroupData.IsEmpty())
+							{
+								FClothGeometryTools::ApplySimGroupResizing(ClothCollection, *TargetMeshDescription, InInterpolationData, ResizingGroupData);
+							}
+						}
+						else if (ClothFacade.HasValidRenderData())
+						{
+							TArray<FMeshResizingCustomRegion> ResizingGroupData;
+							if (!GetValue(Context, &bSkipCustomRegionResizing))
+							{
+								if (TObjectPtr<const USkeletalMesh> InSourceSKM = GetValue(Context, &SourceSkeletalMesh))
 								{
-									FClothGeometryTools::ApplyRenderGroupResizing(ClothCollection, *TargetMeshDescription, InInterpolationData, ResizingGroupData);
+									if (const FMeshDescription* const SourceMeshDescription = InTargetSKM->GetMeshDescription(InLODIndex))
+									{
+										FClothGeometryTools::GenerateRenderMeshResizingGroupData(ClothCollection, *SourceMeshDescription, ResizingGroupData);
+									}
 								}
 							}
+
+							UE::MeshResizing::FRBFInterpolation::DeformPoints(*TargetMeshDescription, InInterpolationData, ClothFacade.GetRenderPosition(), ClothFacade.GetRenderNormal(), ClothFacade.GetRenderTangentU(), ClothFacade.GetRenderTangentV());
+
+							if (!ResizingGroupData.IsEmpty())
+							{
+								FClothGeometryTools::ApplyRenderGroupResizing(ClothCollection, *TargetMeshDescription, InInterpolationData, ResizingGroupData);
+							}
 						}
 					}
 				}
 			}
 		}
-
-		SetValue(Context, MoveTemp(*ClothCollection), &Collection);
 	}
 }
 
+void FChaosClothAssetApplyResizingNode::AddProperties(FPropertyHelper& PropertyHelper) const
+{
+	PropertyHelper.SetPropertyBool(this, &bSavePreResizedSimPosition3D);
+}
 #undef LOCTEXT_NAMESPACE
diff --git a/Engine/Plugins/ChaosClothAssetEditor/Source/ChaosClothAssetDataflowNodes/Public/ChaosClothAsset/ApplyResizingNode.h b/Engine/Plugins/ChaosClothAssetEditor/Source/ChaosClothAssetDataflowNodes/Public/ChaosClothAsset/ApplyResizingNode.h
--- a/Engine/Plugins/ChaosClothAssetEditor/Source/ChaosClothAssetDataflowNodes/Public/ChaosClothAsset/ApplyResizingNode.h
+++ b/Engine/Plugins/ChaosClothAssetEditor/Source/ChaosClothAssetDataflowNodes/Public/ChaosClothAsset/ApplyResizingNode.h
@@ -2,7 +2,7 @@
 
 #pragma once 
 
-#include "Dataflow/DataflowNode.h"
+#include "ChaosClothAsset/SimulationBaseConfigNode.h"
 #include "GeometryCollection/ManagedArrayCollection.h"
 #include "MeshResizing/RBFInterpolation.h"
 #include "ApplyResizingNode.generated.h"
@@ -11,7 +11,7 @@ class USkeletalMesh;
 
 /** Apply resizing for a given Target Mesh.*/
 USTRUCT(Meta = (DataflowCloth, Experimental))
-struct FChaosClothAssetApplyResizingNode final : public FDataflowNode
+struct FChaosClothAssetApplyResizingNode final : public FChaosClothAssetSimulationBaseConfigNode
 {
 	GENERATED_USTRUCT_BODY()
 	DATAFLOW_NODE_DEFINE_INTERNAL(FChaosClothAssetApplyResizingNode, "ApplyResizing", "Cloth", "Apply Cloth Outfit Resizing")
@@ -21,12 +21,9 @@ public:
 	FChaosClothAssetApplyResizingNode(const UE::Dataflow::FNodeParameters& InParam, FGuid InGuid = FGuid::NewGuid());
 
 private:
-	//~ Begin FDataflowNode interface
-	virtual void Evaluate(UE::Dataflow::FContext& Context, const FDataflowOutput* Out) const override;
-	//~ End FDataflowNode interface
+	virtual void EvaluateClothCollection(UE::Dataflow::FContext& Context, const TSharedRef<FManagedArrayCollection>& ClothCollection) const override;
 
-	UPROPERTY(Meta = (Dataflowinput, DataflowOutput, DataflowPassthrough = "Collection"))
-	FManagedArrayCollection Collection;
+	virtual void AddProperties(FPropertyHelper& PropertyHelper) const override;
 
 	/** The target mesh that corresponds with the SourceMesh used to generate the InterpolationData. Must have matching vertices with SourceMesh */
 	UPROPERTY(meta = (DataflowInput))
@@ -51,4 +48,12 @@ private:
 	/** Skip applying Custom Region Resizing data. */
 	UPROPERTY(EditAnywhere, Category = "Resizing", meta = (DataflowInput))
 	bool bSkipCustomRegionResizing = false;
+
+	/** Save pre-resized sim 3d vertices for scaling 2D rest length in XPBD anisotropic stretch constraints.
+	* IMPORTANT: Using Save Pre-Resized Sim Position 3D requires the following settings in SimulationStretchConfig node:
+	* Stretch Use 3d Rest Lengths: false
+	* Solver Type: XPBD
+	* Distribution Type: Anisotropic */
+	UPROPERTY(EditAnywhere, Category = "Resizing", meta = (DataflowInput, DisplayName = "Save Pre-Resized Sim Position 3D"))
+	bool bSavePreResizedSimPosition3D = false;
 };
\ No newline at end of file
diff --git a/Engine/Source/Runtime/Experimental/Chaos/Private/Chaos/XPBDAnisotropicSpringConstraints.cpp b/Engine/Source/Runtime/Experimental/Chaos/Private/Chaos/XPBDAnisotropicSpringConstraints.cpp
--- a/Engine/Source/Runtime/Experimental/Chaos/Private/Chaos/XPBDAnisotropicSpringConstraints.cpp
+++ b/Engine/Source/Runtime/Experimental/Chaos/Private/Chaos/XPBDAnisotropicSpringConstraints.cpp
@@ -48,6 +48,9 @@ FXPBDAnisotropicEdgeSpringConstraints::FXPBDAnisotropicEdgeSpringConstraints(
 	const FTriangleMesh& TriangleMesh,
 	const TArray<TVec3<FVec2f>>& FaceVertexPatternPositions,
 	bool bUse3dRestLengths,
+	bool SavePreResizedSimPosition3D,
+	const TArray<FVector3f>& PreResizedSimPosition3D,
+	FSolverReal MeshScale,
 	const TConstArrayView<FRealSingle>& StiffnessWarpMultipliers,
 	const TConstArrayView<FRealSingle>& StiffnessWeftMultipliers,
 	const TConstArrayView<FRealSingle>& StiffnessBiasMultipliers,
@@ -101,7 +104,7 @@ FXPBDAnisotropicEdgeSpringConstraints::FXPBDAnisotropicEdgeSpringConstraints(
 		ParticleCount)
 {
 	Init(); // Lambda initialization
-	InitFromPatternData(bUse3dRestLengths, FaceVertexPatternPositions, TriangleMesh);
+	InitFromPatternData(Particles.Size(), bUse3dRestLengths, FaceVertexPatternPositions, TriangleMesh, SavePreResizedSimPosition3D, PreResizedSimPosition3D, MeshScale);
 	InitColor(Particles);
 }
 
@@ -112,6 +115,9 @@ FXPBDAnisotropicEdgeSpringConstraints::FXPBDAnisotropicEdgeSpringConstraints(
 	const FTriangleMesh& TriangleMesh,
 	const TArray<TVec3<FVec2f>>& FaceVertexPatternPositions,
 	bool bUse3dRestLengths,
+	bool SavePreResizedSimPosition3D,
+	const TArray<FVector3f>& PreResizedSimPosition3D,
+	FSolverReal MeshScale,
 	const TConstArrayView<FRealSingle>& StiffnessWarpMultipliers,
 	const TConstArrayView<FRealSingle>& StiffnessWeftMultipliers,
 	const TConstArrayView<FRealSingle>& StiffnessBiasMultipliers,
@@ -167,11 +173,139 @@ FXPBDAnisotropicEdgeSpringConstraints::FXPBDAnisotropicEdgeSpringConstraints(
 		ParticleCount)
 {
 	Init(); // Lambda initialization
-	InitFromPatternData(bUse3dRestLengths, FaceVertexPatternPositions, TriangleMesh);
+	InitFromPatternData(Particles.Size(), bUse3dRestLengths, FaceVertexPatternPositions, TriangleMesh, SavePreResizedSimPosition3D, PreResizedSimPosition3D, MeshScale);
 	InitColor(Particles);
 }
 
-void FXPBDAnisotropicEdgeSpringConstraints::InitFromPatternData(bool bUse3dRestLengths, const TArray<TVec3<FVec2f>>& FaceVertexPatternPositions, const FTriangleMesh& TriangleMesh)
+FXPBDAnisotropicEdgeSpringConstraints::FXPBDAnisotropicEdgeSpringConstraints(
+	const FSolverParticlesRange& Particles,
+	const FTriangleMesh& TriangleMesh,
+	const TArray<TVec3<FVec2f>>& FaceVertexPatternPositions,
+	bool bUse3dRestLengths,
+	const TConstArrayView<FRealSingle>& StiffnessWarpMultipliers,
+	const TConstArrayView<FRealSingle>& StiffnessWeftMultipliers,
+	const TConstArrayView<FRealSingle>& StiffnessBiasMultipliers,
+	const TConstArrayView<FRealSingle>& DampingMultipliers,
+	const TConstArrayView<FRealSingle>& WarpScaleMultipliers,
+	const TConstArrayView<FRealSingle>& WeftScaleMultipliers,
+	const FSolverVec2& InStiffnessWarp,
+	const FSolverVec2& InStiffnessWeft,
+	const FSolverVec2& InStiffnessBias,
+	const FSolverVec2& InDampingRatio,
+	const FSolverVec2& InWarpScale,
+	const FSolverVec2& InWeftScale)
+	: Base(Particles,
+		TriangleMesh.GetElements(),
+		TConstArrayView<FRealSingle>(), // We don't use base stiffness maps
+		InStiffnessWarp,
+		true /*bTrimKinematicConstraints*/,
+		MaxStiffness)
+	, StiffnessWarp(InStiffnessWarp.ClampAxes(0, MaxStiffness),
+		StiffnessWarpMultipliers,
+		TConstArrayView<TVec2<int32>>(Constraints),
+		ParticleOffset,
+		ParticleCount)
+	, StiffnessWeft(InStiffnessWeft.ClampAxes(0, MaxStiffness),
+		StiffnessWeftMultipliers,
+		TConstArrayView<TVec2<int32>>(Constraints),
+		ParticleOffset,
+		ParticleCount)
+	, StiffnessBias(InStiffnessBias.ClampAxes(0, MaxStiffness),
+		StiffnessBiasMultipliers,
+		TConstArrayView<TVec2<int32>>(Constraints),
+		ParticleOffset,
+		ParticleCount)
+	, DampingRatio(
+		InDampingRatio.ClampAxes(MinDampingRatio, MaxDampingRatio),
+		DampingMultipliers,
+		TConstArrayView<TVec2<int32>>(Constraints),
+		ParticleOffset,
+		ParticleCount)
+	, WarpScale(
+		InWarpScale.ClampAxes(MinWarpWeftScale, MaxWarpWeftScale),
+		WarpScaleMultipliers,
+		TConstArrayView<TVec2<int32>>(Constraints),
+		ParticleOffset,
+		ParticleCount)
+	, WeftScale(
+		InWeftScale.ClampAxes(MinWarpWeftScale, MaxWarpWeftScale),
+		WeftScaleMultipliers,
+		TConstArrayView<TVec2<int32>>(Constraints),
+		ParticleOffset,
+		ParticleCount)
+{
+	Init(); // Lambda initialization
+	InitFromPatternData(Particles.Size(), bUse3dRestLengths, FaceVertexPatternPositions, TriangleMesh, /*SavePreResizedSimPosition3D = */false, /*PreResizedSimPosition3D = */TArray<FVector3f>(), /*MeshScale = */ (FSolverReal)1.);
+	InitColor(Particles);
+}
+
+FXPBDAnisotropicEdgeSpringConstraints::FXPBDAnisotropicEdgeSpringConstraints(
+	const FSolverParticles& Particles,
+	int32 InParticleOffset,
+	int32 InParticleCount,
+	const FTriangleMesh& TriangleMesh,
+	const TArray<TVec3<FVec2f>>& FaceVertexPatternPositions,
+	bool bUse3dRestLengths,
+	const TConstArrayView<FRealSingle>& StiffnessWarpMultipliers,
+	const TConstArrayView<FRealSingle>& StiffnessWeftMultipliers,
+	const TConstArrayView<FRealSingle>& StiffnessBiasMultipliers,
+	const TConstArrayView<FRealSingle>& DampingMultipliers,
+	const TConstArrayView<FRealSingle>& WarpScaleMultipliers,
+	const TConstArrayView<FRealSingle>& WeftScaleMultipliers,
+	const FSolverVec2& InStiffnessWarp,
+	const FSolverVec2& InStiffnessWeft,
+	const FSolverVec2& InStiffnessBias,
+	const FSolverVec2& InDampingRatio,
+	const FSolverVec2& InWarpScale,
+	const FSolverVec2& InWeftScale)
+	: Base(Particles,
+		InParticleOffset,
+		InParticleCount,
+		TriangleMesh.GetElements(),
+		TConstArrayView<FRealSingle>(), // We don't use base stiffness maps
+		InStiffnessWarp,
+		true /*bTrimKinematicConstraints*/,
+		MaxStiffness)
+	, StiffnessWarp(InStiffnessWarp.ClampAxes(0, MaxStiffness),
+		StiffnessWarpMultipliers,
+		TConstArrayView<TVec2<int32>>(Constraints),
+		ParticleOffset,
+		ParticleCount)
+	, StiffnessWeft(InStiffnessWeft.ClampAxes(0, MaxStiffness),
+		StiffnessWeftMultipliers,
+		TConstArrayView<TVec2<int32>>(Constraints),
+		ParticleOffset,
+		ParticleCount)
+	, StiffnessBias(InStiffnessBias.ClampAxes(0, MaxStiffness),
+		StiffnessBiasMultipliers,
+		TConstArrayView<TVec2<int32>>(Constraints),
+		ParticleOffset,
+		ParticleCount)
+	, DampingRatio(
+		InDampingRatio.ClampAxes(MinDampingRatio, MaxDampingRatio),
+		DampingMultipliers,
+		TConstArrayView<TVec2<int32>>(Constraints),
+		ParticleOffset,
+		ParticleCount)
+	, WarpScale(
+		InWarpScale.ClampAxes(MinWarpWeftScale, MaxWarpWeftScale),
+		WarpScaleMultipliers,
+		TConstArrayView<TVec2<int32>>(Constraints),
+		ParticleOffset,
+		ParticleCount)
+	, WeftScale(
+		InWeftScale.ClampAxes(MinWarpWeftScale, MaxWarpWeftScale),
+		WeftScaleMultipliers,
+		TConstArrayView<TVec2<int32>>(Constraints),
+		ParticleOffset,
+		ParticleCount)
+{
+	Init(); // Lambda initialization
+	InitFromPatternData(Particles.Size(), bUse3dRestLengths, FaceVertexPatternPositions, TriangleMesh, /*SavePreResizedSimPosition3D = */false, /*PreResizedSimPosition3D = */TArray<FVector3f>(), /*MeshScale = */ (FSolverReal)1.);
+	InitColor(Particles);
+}
+
+void FXPBDAnisotropicEdgeSpringConstraints::InitFromPatternData(int32 ParticlesSize, bool bUse3dRestLengths, const TArray<TVec3<FVec2f>>& FaceVertexPatternPositions, const FTriangleMesh& TriangleMesh, bool SavePreResizedSimPosition3D, const TArray<FVector3f>& PreResizedSimPosition3D, FSolverReal MeshScale)
 {
 	// Calculate multipliers per face. Average the multipliers for faces that share the same edge.
 	struct FMultipliers
@@ -248,6 +382,12 @@ void FXPBDAnisotropicEdgeSpringConstraints::InitFromPatternData(bool bUse3dRestL
 		WarpWeftScaleBaseMultipliers[ConstraintIdx].Normalize();
 		if (!bUse3dRestLengths)
 		{
+			if (SavePreResizedSimPosition3D && (ParticlesSize == PreResizedSimPosition3D.Num()))
+			{
+				const FSolverReal Denom = MeshScale * (PreResizedSimPosition3D[Constraints[ConstraintIdx][0]] - PreResizedSimPosition3D[Constraints[ConstraintIdx][1]]).Size();
+				const FSolverReal Scale = Denom > UE_SMALL_NUMBER ? Dists[ConstraintIdx] / Denom : 1;
+				BaseDists[ConstraintIdx] *= Scale;
+			}
 			BaseDists[ConstraintIdx] /= (FSolverReal)EdgeMultipliers.Num();
 			Dists[ConstraintIdx] = BaseDists[ConstraintIdx];
 		}
@@ -940,6 +1080,9 @@ FXPBDAnisotropicAxialSpringConstraints::FXPBDAnisotropicAxialSpringConstraints(
 	const FTriangleMesh& TriangleMesh,
 	const TArray<TVec3<FVec2f>>& FaceVertexPatternPositions,
 	bool bUse3dRestLengths,
+	bool SavePreResizedSimPosition3D,
+	const TArray<FVector3f>& PreResizedSimPosition3D,
+	FSolverReal MeshScale,
 	const TConstArrayView<FRealSingle>& StiffnessWarpMultipliers,
 	const TConstArrayView<FRealSingle>& StiffnessWeftMultipliers,
 	const TConstArrayView<FRealSingle>& StiffnessBiasMultipliers,
@@ -993,7 +1136,7 @@ FXPBDAnisotropicAxialSpringConstraints::FXPBDAnisotropicAxialSpringConstraints(
 		ParticleCount)
 {
 	Init(); // Lambda initialization
-	InitFromPatternData(bUse3dRestLengths, FaceVertexPatternPositions, TriangleMesh);
+	InitFromPatternData(Particles.Size(), bUse3dRestLengths, FaceVertexPatternPositions, TriangleMesh, SavePreResizedSimPosition3D, PreResizedSimPosition3D, MeshScale);
 	InitColor(Particles);
 }
 
@@ -1004,6 +1147,9 @@ FXPBDAnisotropicAxialSpringConstraints::FXPBDAnisotropicAxialSpringConstraints(
 	const FTriangleMesh& TriangleMesh,
 	const TArray<TVec3<FVec2f>>& FaceVertexPatternPositions,
 	bool bUse3dRestLengths,
+	bool SavePreResizedSimPosition3D,
+	const TArray<FVector3f>& PreResizedSimPosition3D,
+	FSolverReal MeshScale,
 	const TConstArrayView<FRealSingle>& StiffnessWarpMultipliers,
 	const TConstArrayView<FRealSingle>& StiffnessWeftMultipliers,
 	const TConstArrayView<FRealSingle>& StiffnessBiasMultipliers,
@@ -1059,11 +1205,139 @@ FXPBDAnisotropicAxialSpringConstraints::FXPBDAnisotropicAxialSpringConstraints(
 		ParticleCount)
 {
 	Init(); // Lambda initialization
-	InitFromPatternData(bUse3dRestLengths, FaceVertexPatternPositions, TriangleMesh);
+	InitFromPatternData(Particles.Size(), bUse3dRestLengths, FaceVertexPatternPositions, TriangleMesh, SavePreResizedSimPosition3D, PreResizedSimPosition3D, MeshScale);
 	InitColor(Particles);
 }
 
-void FXPBDAnisotropicAxialSpringConstraints::InitFromPatternData(bool bUse3dRestLengths, const TArray<TVec3<FVec2f>>& FaceVertexPatternPositions, const FTriangleMesh& TriangleMesh)
+FXPBDAnisotropicAxialSpringConstraints::FXPBDAnisotropicAxialSpringConstraints(
+	const FSolverParticlesRange& Particles,
+	const FTriangleMesh& TriangleMesh,
+	const TArray<TVec3<FVec2f>>& FaceVertexPatternPositions,
+	bool bUse3dRestLengths,
+	const TConstArrayView<FRealSingle>& StiffnessWarpMultipliers,
+	const TConstArrayView<FRealSingle>& StiffnessWeftMultipliers,
+	const TConstArrayView<FRealSingle>& StiffnessBiasMultipliers,
+	const TConstArrayView<FRealSingle>& DampingMultipliers,
+	const TConstArrayView<FRealSingle>& WarpScaleMultipliers,
+	const TConstArrayView<FRealSingle>& WeftScaleMultipliers,
+	const FSolverVec2& InStiffnessWarp,
+	const FSolverVec2& InStiffnessWeft,
+	const FSolverVec2& InStiffnessBias,
+	const FSolverVec2& InDampingRatio,
+	const FSolverVec2& InWarpScale,
+	const FSolverVec2& InWeftScale)
+	: Base(Particles,
+		TriangleMesh.GetElements(),
+		TConstArrayView<FRealSingle>(), // We don't use base stiffness maps
+		InStiffnessWarp,
+		true /*bTrimKinematicConstraints*/,
+		MaxStiffness)
+	, StiffnessWarp(InStiffnessWarp.ClampAxes(0, MaxStiffness),
+		StiffnessWarpMultipliers,
+		TConstArrayView<TVec3<int32>>(Constraints),
+		ParticleOffset,
+		ParticleCount)
+	, StiffnessWeft(InStiffnessWeft.ClampAxes(0, MaxStiffness),
+		StiffnessWeftMultipliers,
+		TConstArrayView<TVec3<int32>>(Constraints),
+		ParticleOffset,
+		ParticleCount)
+	, StiffnessBias(InStiffnessBias.ClampAxes(0, MaxStiffness),
+		StiffnessBiasMultipliers,
+		TConstArrayView<TVec3<int32>>(Constraints),
+		ParticleOffset,
+		ParticleCount)
+	, DampingRatio(
+		InDampingRatio.ClampAxes(MinDampingRatio, MaxDampingRatio),
+		DampingMultipliers,
+		TConstArrayView<TVec3<int32>>(Constraints),
+		ParticleOffset,
+		ParticleCount)
+	, WarpScale(
+		InWarpScale.ClampAxes(MinWarpWeftScale, MaxWarpWeftScale),
+		WarpScaleMultipliers,
+		TConstArrayView<TVec3<int32>>(Constraints),
+		ParticleOffset,
+		ParticleCount)
+	, WeftScale(
+		InWeftScale.ClampAxes(MinWarpWeftScale, MaxWarpWeftScale),
+		WeftScaleMultipliers,
+		TConstArrayView<TVec3<int32>>(Constraints),
+		ParticleOffset,
+		ParticleCount)
+{
+	Init(); // Lambda initialization
+	InitFromPatternData(Particles.Size(), bUse3dRestLengths, FaceVertexPatternPositions, TriangleMesh, /*SavePreResizedSimPosition3D = */false, /*PreResizedSimPosition3D = */TArray<FVector3f>(), /*MeshScale = */ (FSolverReal)1.);
+	InitColor(Particles);
+}
+
+FXPBDAnisotropicAxialSpringConstraints::FXPBDAnisotropicAxialSpringConstraints(
+	const FSolverParticles& Particles,
+	int32 InParticleOffset,
+	int32 InParticleCount,
+	const FTriangleMesh& TriangleMesh,
+	const TArray<TVec3<FVec2f>>& FaceVertexPatternPositions,
+	bool bUse3dRestLengths,
+	const TConstArrayView<FRealSingle>& StiffnessWarpMultipliers,
+	const TConstArrayView<FRealSingle>& StiffnessWeftMultipliers,
+	const TConstArrayView<FRealSingle>& StiffnessBiasMultipliers,
+	const TConstArrayView<FRealSingle>& DampingMultipliers,
+	const TConstArrayView<FRealSingle>& WarpScaleMultipliers,
+	const TConstArrayView<FRealSingle>& WeftScaleMultipliers,
+	const FSolverVec2& InStiffnessWarp,
+	const FSolverVec2& InStiffnessWeft,
+	const FSolverVec2& InStiffnessBias,
+	const FSolverVec2& InDampingRatio,
+	const FSolverVec2& InWarpScale,
+	const FSolverVec2& InWeftScale)
+	: Base(Particles,
+		InParticleOffset,
+		InParticleCount,
+		TriangleMesh.GetElements(),
+		TConstArrayView<FRealSingle>(), // We don't use base stiffness maps
+		InStiffnessWarp,
+		true /*bTrimKinematicConstraints*/,
+		MaxStiffness)
+	, StiffnessWarp(InStiffnessWarp.ClampAxes(0, MaxStiffness),
+		StiffnessWarpMultipliers,
+		TConstArrayView<TVec3<int32>>(Constraints),
+		ParticleOffset,
+		ParticleCount)
+	, StiffnessWeft(InStiffnessWeft.ClampAxes(0, MaxStiffness),
+		StiffnessWeftMultipliers,
+		TConstArrayView<TVec3<int32>>(Constraints),
+		ParticleOffset,
+		ParticleCount)
+	, StiffnessBias(InStiffnessBias.ClampAxes(0, MaxStiffness),
+		StiffnessBiasMultipliers,
+		TConstArrayView<TVec3<int32>>(Constraints),
+		ParticleOffset,
+		ParticleCount)
+	, DampingRatio(
+		InDampingRatio.ClampAxes(MinDampingRatio, MaxDampingRatio),
+		DampingMultipliers,
+		TConstArrayView<TVec3<int32>>(Constraints),
+		ParticleOffset,
+		ParticleCount)
+	, WarpScale(
+		InWarpScale.ClampAxes(MinWarpWeftScale, MaxWarpWeftScale),
+		WarpScaleMultipliers,
+		TConstArrayView<TVec3<int32>>(Constraints),
+		ParticleOffset,
+		ParticleCount)
+	, WeftScale(
+		InWeftScale.ClampAxes(MinWarpWeftScale, MaxWarpWeftScale),
+		WeftScaleMultipliers,
+		TConstArrayView<TVec3<int32>>(Constraints),
+		ParticleOffset,
+		ParticleCount)
+{
+	Init(); // Lambda initialization
+	InitFromPatternData(Particles.Size(), bUse3dRestLengths, FaceVertexPatternPositions, TriangleMesh, /*SavePreResizedSimPosition3D = */false, /*PreResizedSimPosition3D = */TArray<FVector3f>(), /*MeshScale = */ (FSolverReal)1.);
+	InitColor(Particles);
+}
+
+void FXPBDAnisotropicAxialSpringConstraints::InitFromPatternData(int32 ParticlesSize, bool bUse3dRestLengths, const TArray<TVec3<FVec2f>>& FaceVertexPatternPositions, const FTriangleMesh& TriangleMesh, bool SavePreResizedSimPosition3D, const TArray<FVector3f>& PreResizedSimPosition3D, FSolverReal MeshScale)
 {
 	if (bUse3dRestLengths)
 	{
@@ -1115,8 +1389,7 @@ void FXPBDAnisotropicAxialSpringConstraints::InitFromPatternData(bool bUse3dRest
 		const FVec2f& UV2 = FaceVertexPatternPositions[TriangleIndex][FaceIndex2];
 		const FVec2f& UV3 = FaceVertexPatternPositions[TriangleIndex][FaceIndex3];
 
-		const FVec2f UV = (UV2 - UV3) * Barys[ConstraintIdx] + UV3;
-		const FSolverVec2 UVDiff = UV - UV1;
+		const FSolverVec2 UVDiff = GetSpringDiff(UV1, UV2, UV3, Barys[ConstraintIdx]);
 		const FSolverVec2 UVDiffAbs = UVDiff.GetAbs();
 		FSolverVec3 UVDiffTransformed;
 		if (UVDiffAbs.X > UVDiffAbs.Y)
@@ -1137,7 +1410,17 @@ void FXPBDAnisotropicAxialSpringConstraints::InitFromPatternData(bool bUse3dRest
 		if (!bUse3dRestLengths)
 		{
 			BaseDists[ConstraintIdx] = UVLength;
-			Dists[ConstraintIdx] = UVLength;
+			if (SavePreResizedSimPosition3D && (ParticlesSize == PreResizedSimPosition3D.Num()))
+			{
+				const FSolverReal Denom = MeshScale * GetSpringDiff(
+					PreResizedSimPosition3D[Constraint[0]], 
+					PreResizedSimPosition3D[Constraint[1]], 
+					PreResizedSimPosition3D[Constraint[2]], 
+					Barys[ConstraintIdx]).Size();
+				const FSolverReal Scale = Denom > UE_SMALL_NUMBER ? Dists[ConstraintIdx] / Denom : 1;
+				BaseDists[ConstraintIdx] *= Scale;
+			}
+			Dists[ConstraintIdx] = BaseDists[ConstraintIdx];
 		}
 	}
 }
diff --git a/Engine/Source/Runtime/Experimental/Chaos/Public/Chaos/PBDAxialSpringConstraintsBase.h b/Engine/Source/Runtime/Experimental/Chaos/Public/Chaos/PBDAxialSpringConstraintsBase.h
--- a/Engine/Source/Runtime/Experimental/Chaos/Public/Chaos/PBDAxialSpringConstraintsBase.h
+++ b/Engine/Source/Runtime/Experimental/Chaos/Public/Chaos/PBDAxialSpringConstraintsBase.h
@@ -82,6 +82,12 @@ public:
 	const TArray<FSolverReal>& GetBarys() const { return Barys; }
 
 protected:
+	template<typename VectorType>
+	inline VectorType GetSpringDiff(const VectorType& P1, const VectorType& P2, const VectorType& P3, const FSolverReal Bary) const
+	{
+		return P1 - ((P2 - P3) * Bary + P3);
+	}
+
 	template<typename SolverParticlesOrRange>
 	inline FSolverVec3 GetDelta(const SolverParticlesOrRange& Particles, const int32 ConstraintIndex, const FSolverReal ExpStiffnessValue) const
 	{
@@ -97,8 +103,7 @@ protected:
 		const FSolverVec3& P1 = Particles.P(i1);
 		const FSolverVec3& P2 = Particles.P(i2);
 		const FSolverVec3& P3 = Particles.P(i3);
-		const FSolverVec3 P = (P2 - P3) * Barys[ConstraintIndex] + P3;
-		const FSolverVec3 Difference = P1 - P;
+		const FSolverVec3 Difference = GetSpringDiff(P1, P2, P3, Barys[ConstraintIndex]);
 		const FSolverReal Distance = Difference.Size();
 		if (UNLIKELY(Distance <= (FSolverReal)UE_SMALL_NUMBER))
 		{
diff --git a/Engine/Source/Runtime/Experimental/Chaos/Public/Chaos/XPBDAnisotropicSpringConstraints.h b/Engine/Source/Runtime/Experimental/Chaos/Public/Chaos/XPBDAnisotropicSpringConstraints.h
--- a/Engine/Source/Runtime/Experimental/Chaos/Public/Chaos/XPBDAnisotropicSpringConstraints.h
+++ b/Engine/Source/Runtime/Experimental/Chaos/Public/Chaos/XPBDAnisotropicSpringConstraints.h
@@ -35,6 +35,9 @@ public:
 		const FTriangleMesh& TriangleMesh,
 		const TArray<TVec3<FVec2f>>& FaceVertexPatternPositions,
 		bool bUse3dRestLengths,
+		bool bSavePreResizedSimPosition3D,
+		const TArray<FVector3f>& PreResizedSimPosition3D,
+		FSolverReal MeshScale,
 		const TConstArrayView<FRealSingle>& StiffnessWarpMultipliers,
 		const TConstArrayView<FRealSingle>& StiffnessWeftMultipliers,
 		const TConstArrayView<FRealSingle>& StiffnessBiasMultipliers,
@@ -48,6 +51,49 @@ public:
 		const FSolverVec2& InWarpScale,
 		const FSolverVec2& InWeftScale);
 
+	CHAOS_API FXPBDAnisotropicEdgeSpringConstraints(
+		const FSolverParticles& InParticles,
+		int32 InParticleOffset,
+		int32 InParticleCount,
+		const FTriangleMesh& TriangleMesh,
+		const TArray<TVec3<FVec2f>>& FaceVertexPatternPositions,
+		bool bUse3dRestLengths,
+		bool bSavePreResizedSimPosition3D,
+		const TArray<FVector3f>& PreResizedSimPosition3D,
+		FSolverReal MeshScale,
+		const TConstArrayView<FRealSingle>& StiffnessWarpMultipliers,
+		const TConstArrayView<FRealSingle>& StiffnessWeftMultipliers,
+		const TConstArrayView<FRealSingle>& StiffnessBiasMultipliers,
+		const TConstArrayView<FRealSingle>& DampingMultipliers,
+		const TConstArrayView<FRealSingle>& WarpScaleMultipliers,
+		const TConstArrayView<FRealSingle>& WeftScaleMultipliers,
+		const FSolverVec2& InStiffnessWarp,
+		const FSolverVec2& InStiffnessWeft,
+		const FSolverVec2& InStiffnessBias,
+		const FSolverVec2& InDampingRatio,
+		const FSolverVec2& InWarpScale,
+		const FSolverVec2& InWeftScale);
+
+	UE_DEPRECATED(5.6, "No longer supported, use new constructors instead")
+	CHAOS_API FXPBDAnisotropicEdgeSpringConstraints(
+		const FSolverParticlesRange& Particles,
+		const FTriangleMesh& TriangleMesh,
+		const TArray<TVec3<FVec2f>>& FaceVertexPatternPositions,
+		bool bUse3dRestLengths,
+		const TConstArrayView<FRealSingle>& StiffnessWarpMultipliers,
+		const TConstArrayView<FRealSingle>& StiffnessWeftMultipliers,
+		const TConstArrayView<FRealSingle>& StiffnessBiasMultipliers,
+		const TConstArrayView<FRealSingle>& DampingMultipliers,
+		const TConstArrayView<FRealSingle>& WarpScaleMultipliers,
+		const TConstArrayView<FRealSingle>& WeftScaleMultipliers,
+		const FSolverVec2& InStiffnessWarp,
+		const FSolverVec2& InStiffnessWeft,
+		const FSolverVec2& InStiffnessBias,
+		const FSolverVec2& InDampingRatio,
+		const FSolverVec2& InWarpScale,
+		const FSolverVec2& InWeftScale);
+
+	UE_DEPRECATED(5.6, "No longer supported, use new constructors instead")
 	CHAOS_API FXPBDAnisotropicEdgeSpringConstraints(
 		const FSolverParticles& InParticles,
 		int32 InParticleOffset,
@@ -104,7 +150,7 @@ private:
 	template<typename SolverParticlesOrRange>
 	void InitColor(const SolverParticlesOrRange& InParticles);
 
-	void InitFromPatternData(bool bUse3dRestLengths, const TArray<TVec3<FVec2f>>& FaceVertexPatternPositions, const FTriangleMesh& TriangleMesh);
+	void InitFromPatternData(int32 ParticlesSize, bool bUse3dRestLengths, const TArray<TVec3<FVec2f>>& FaceVertexPatternPositions, const FTriangleMesh& TriangleMesh, bool bSavePreResizedSimPosition3D, const TArray<FVector3f>& PreResizedSimPosition3D, FSolverReal MeshScale);
 
 	CHAOS_API void UpdateDists();
 	
@@ -151,6 +197,9 @@ public:
 		const FTriangleMesh& TriangleMesh,
 		const TArray<TVec3<FVec2f>>& FaceVertexPatternPositions,
 		bool bUse3dRestLengths,
+		bool bSavePreResizedSimPosition3D,
+		const TArray<FVector3f>& PreResizedSimPosition3D,
+		FSolverReal MeshScale,
 		const TConstArrayView<FRealSingle>& StiffnessWarpMultipliers,
 		const TConstArrayView<FRealSingle>& StiffnessWeftMultipliers,
 		const TConstArrayView<FRealSingle>& StiffnessBiasMultipliers,
@@ -164,63 +213,106 @@ public:
 		const FSolverVec2& InWarpScale,
 		const FSolverVec2& InWeftScale);
 
-		CHAOS_API FXPBDAnisotropicAxialSpringConstraints(
-			const FSolverParticles& InParticles,
-			int32 InParticleOffset,
-			int32 InParticleCount,
-			const FTriangleMesh& TriangleMesh,
-			const TArray<TVec3<FVec2f>>& FaceVertexPatternPositions,
-			bool bUse3dRestLengths,
-			const TConstArrayView<FRealSingle>& StiffnessWarpMultipliers,
-			const TConstArrayView<FRealSingle>& StiffnessWeftMultipliers,
-			const TConstArrayView<FRealSingle>& StiffnessBiasMultipliers,
-			const TConstArrayView<FRealSingle>& DampingMultipliers,
-			const TConstArrayView<FRealSingle>& WarpScaleMultipliers,
-			const TConstArrayView<FRealSingle>& WeftScaleMultipliers,
-			const FSolverVec2& InStiffnessWarp,
-			const FSolverVec2& InStiffnessWeft,
-			const FSolverVec2& InStiffnessBias,
-			const FSolverVec2& InDampingRatio,
-			const FSolverVec2& InWarpScale,
-			const FSolverVec2& InWeftScale);
-
-		virtual ~FXPBDAnisotropicAxialSpringConstraints() override {}
-
-		void Init()
-		{
-			Lambdas.Reset();
-			Lambdas.AddZeroed(Constraints.Num());
-			LambdasDamping.Reset();
-			LambdasDamping.AddZeroed(Constraints.Num());
-		}
+	CHAOS_API FXPBDAnisotropicAxialSpringConstraints(
+		const FSolverParticles& InParticles,
+		int32 InParticleOffset,
+		int32 InParticleCount,
+		const FTriangleMesh& TriangleMesh,
+		const TArray<TVec3<FVec2f>>& FaceVertexPatternPositions,
+		bool bUse3dRestLengths,
+		bool bSavePreResizedSimPosition3D,
+		const TArray<FVector3f>& PreResizedSimPosition3D,
+		FSolverReal MeshScale,
+		const TConstArrayView<FRealSingle>& StiffnessWarpMultipliers,
+		const TConstArrayView<FRealSingle>& StiffnessWeftMultipliers,
+		const TConstArrayView<FRealSingle>& StiffnessBiasMultipliers,
+		const TConstArrayView<FRealSingle>& DampingMultipliers,
+		const TConstArrayView<FRealSingle>& WarpScaleMultipliers,
+		const TConstArrayView<FRealSingle>& WeftScaleMultipliers,
+		const FSolverVec2& InStiffnessWarp,
+		const FSolverVec2& InStiffnessWeft,
+		const FSolverVec2& InStiffnessBias,
+		const FSolverVec2& InDampingRatio,
+		const FSolverVec2& InWarpScale,
+		const FSolverVec2& InWeftScale);
+
+	UE_DEPRECATED(5.6, "No longer supported, use new constructors instead")
+	CHAOS_API FXPBDAnisotropicAxialSpringConstraints(
+		const FSolverParticlesRange& Particles,
+		const FTriangleMesh& TriangleMesh,
+		const TArray<TVec3<FVec2f>>& FaceVertexPatternPositions,
+		bool bUse3dRestLengths,
+		const TConstArrayView<FRealSingle>& StiffnessWarpMultipliers,
+		const TConstArrayView<FRealSingle>& StiffnessWeftMultipliers,
+		const TConstArrayView<FRealSingle>& StiffnessBiasMultipliers,
+		const TConstArrayView<FRealSingle>& DampingMultipliers,
+		const TConstArrayView<FRealSingle>& WarpScaleMultipliers,
+		const TConstArrayView<FRealSingle>& WeftScaleMultipliers,
+		const FSolverVec2& InStiffnessWarp,
+		const FSolverVec2& InStiffnessWeft,
+		const FSolverVec2& InStiffnessBias,
+		const FSolverVec2& InDampingRatio,
+		const FSolverVec2& InWarpScale,
+		const FSolverVec2& InWeftScale);
+
+	UE_DEPRECATED(5.6, "No longer supported, use new constructors instead")
+	CHAOS_API FXPBDAnisotropicAxialSpringConstraints(
+		const FSolverParticles& InParticles,
+		int32 InParticleOffset,
+		int32 InParticleCount,
+		const FTriangleMesh& TriangleMesh,
+		const TArray<TVec3<FVec2f>>& FaceVertexPatternPositions,
+		bool bUse3dRestLengths,
+		const TConstArrayView<FRealSingle>& StiffnessWarpMultipliers,
+		const TConstArrayView<FRealSingle>& StiffnessWeftMultipliers,
+		const TConstArrayView<FRealSingle>& StiffnessBiasMultipliers,
+		const TConstArrayView<FRealSingle>& DampingMultipliers,
+		const TConstArrayView<FRealSingle>& WarpScaleMultipliers,
+		const TConstArrayView<FRealSingle>& WeftScaleMultipliers,
+		const FSolverVec2& InStiffnessWarp,
+		const FSolverVec2& InStiffnessWeft,
+		const FSolverVec2& InStiffnessBias,
+		const FSolverVec2& InDampingRatio,
+		const FSolverVec2& InWarpScale,
+		const FSolverVec2& InWeftScale);
+
+	virtual ~FXPBDAnisotropicAxialSpringConstraints() override {}
+
+	void Init()
+	{
+		Lambdas.Reset();
+		Lambdas.AddZeroed(Constraints.Num());
+		LambdasDamping.Reset();
+		LambdasDamping.AddZeroed(Constraints.Num());
+	}
 
-		// Update rest lengths from warp/weft scale
-		void ApplyProperties(const FSolverReal /*Dt*/, const int32 /*NumIterations*/)
+	// Update rest lengths from warp/weft scale
+	void ApplyProperties(const FSolverReal /*Dt*/, const int32 /*NumIterations*/)
+	{
+		bool bWarpScaleChanged = false;
+		WarpScale.ApplyValues(&bWarpScaleChanged);
+		bool bWeftScaleChanged = false;
+		WeftScale.ApplyValues(&bWeftScaleChanged);
+		if (bWarpScaleChanged || bWeftScaleChanged)
 		{
-			bool bWarpScaleChanged = false;
-			WarpScale.ApplyValues(&bWarpScaleChanged);
-			bool bWeftScaleChanged = false;
-			WeftScale.ApplyValues(&bWeftScaleChanged);
-			if (bWarpScaleChanged || bWeftScaleChanged)
-			{
-				// Need to update distances
-				UpdateDists();
-			}
+			// Need to update distances
+			UpdateDists();
 		}
+	}
 
-		template<typename SolverParticlesOrRange>
-		CHAOS_API void Apply(SolverParticlesOrRange& Particles, const FSolverReal Dt) const;
+	template<typename SolverParticlesOrRange>
+	CHAOS_API void Apply(SolverParticlesOrRange& Particles, const FSolverReal Dt) const;
 
-		CHAOS_API void UpdateLinearSystem(const FSolverParticlesRange& Particles, const FSolverReal Dt, FEvolutionLinearSystem& LinearSystem) const;
+	CHAOS_API void UpdateLinearSystem(const FSolverParticlesRange& Particles, const FSolverReal Dt, FEvolutionLinearSystem& LinearSystem) const;
 
-		const TArray<int32>& GetConstraintsPerColorStartIndex() const { return ConstraintsPerColorStartIndex; }
-		const TArray<FSolverVec3>& GetWarpWeftBiasBaseMultipliers() const { return WarpWeftBiasBaseMultipliers; }
+	const TArray<int32>& GetConstraintsPerColorStartIndex() const { return ConstraintsPerColorStartIndex; }
+	const TArray<FSolverVec3>& GetWarpWeftBiasBaseMultipliers() const { return WarpWeftBiasBaseMultipliers; }
 
 private:
 	template<typename SolverParticlesOrRange>
 	void InitColor(const SolverParticlesOrRange& InParticles);
 
-	void InitFromPatternData(bool bUse3dRestLengths, const TArray<TVec3<FVec2f>>& FaceVertexPatternPositions, const FTriangleMesh& TriangleMesh);
+	void InitFromPatternData(int32 ParticlesSize, bool bUse3dRestLengths, const TArray<TVec3<FVec2f>>& FaceVertexPatternPositions, const FTriangleMesh& TriangleMesh, bool bSavePreResizedSimPosition3D, const TArray<FVector3f>& PreResizedSimPosition3D, FSolverReal MeshScale);
 
 	CHAOS_API void UpdateDists();
 
@@ -258,6 +350,7 @@ public:
 	static constexpr FSolverReal DefaultDamping = (FSolverReal)1.;
 	static constexpr bool bDefaultUse3dRestLengths = true;
 	static constexpr FSolverReal DefaultWarpWeftScale = (FSolverReal)1.;
+	static constexpr bool bDefaultSavePreResizedSimPosition3D = false;
 
 	static bool IsEnabled(const FCollectionPropertyConstFacade& PropertyCollection)
 	{
@@ -269,12 +362,17 @@ public:
 		const FTriangleMesh& TriangleMesh,
 		const TArray<TVec3<FVec2f>>& FaceVertexPatternPositions,
 		const TMap<FString, TConstArrayView<FRealSingle>>& WeightMaps,
-		const FCollectionPropertyConstFacade& PropertyCollection)
+		const FCollectionPropertyConstFacade& PropertyCollection,
+		const TArray<FVector3f>& PreResizedSimPosition3D,
+		Softs::FSolverReal MeshScale)
 		: EdgeConstraints(
 			Particles,
 			TriangleMesh,
 			FaceVertexPatternPositions,
 			GetXPBDAnisoSpringUse3dRestLengths(PropertyCollection, bDefaultUse3dRestLengths),
+			GetSavePreResizedSimPosition3D(PropertyCollection, bDefaultSavePreResizedSimPosition3D),
+			PreResizedSimPosition3D,
+			MeshScale,
 			WeightMaps.FindRef(GetXPBDAnisoSpringStiffnessWarpString(PropertyCollection, XPBDAnisoSpringStiffnessWarpName.ToString())),
 			WeightMaps.FindRef(GetXPBDAnisoSpringStiffnessWeftString(PropertyCollection, XPBDAnisoSpringStiffnessWeftName.ToString())),
 			WeightMaps.FindRef(GetXPBDAnisoSpringStiffnessBiasString(PropertyCollection, XPBDAnisoSpringStiffnessBiasName.ToString())),
@@ -292,6 +390,9 @@ public:
 			TriangleMesh,
 			FaceVertexPatternPositions,
 			GetXPBDAnisoSpringUse3dRestLengths(PropertyCollection, bDefaultUse3dRestLengths),
+			GetSavePreResizedSimPosition3D(PropertyCollection, bDefaultSavePreResizedSimPosition3D),
+			PreResizedSimPosition3D,
+			MeshScale,
 			WeightMaps.FindRef(GetXPBDAnisoSpringStiffnessWarpString(PropertyCollection, XPBDAnisoSpringStiffnessWarpName.ToString())),
 			WeightMaps.FindRef(GetXPBDAnisoSpringStiffnessWeftString(PropertyCollection, XPBDAnisoSpringStiffnessWeftName.ToString())),
 			WeightMaps.FindRef(GetXPBDAnisoSpringStiffnessBiasString(PropertyCollection, XPBDAnisoSpringStiffnessBiasName.ToString())),
@@ -311,6 +412,7 @@ public:
 		, XPBDAnisoSpringDampingIndex(PropertyCollection)
 		, XPBDAnisoSpringWarpScaleIndex(PropertyCollection)
 		, XPBDAnisoSpringWeftScaleIndex(PropertyCollection)
+		, SavePreResizedSimPosition3DIndex(PropertyCollection)
 	{}
 
 	FXPBDAnisotropicSpringConstraints(
@@ -320,7 +422,9 @@ public:
 		const FTriangleMesh& TriangleMesh,
 		const TArray<TVec3<FVec2f>>& FaceVertexPatternPositions,
 		const TMap<FString, TConstArrayView<FRealSingle>>& WeightMaps,
-		const FCollectionPropertyConstFacade& PropertyCollection)
+		const FCollectionPropertyConstFacade& PropertyCollection,
+		const TArray<FVector3f>& PreResizedSimPosition3D,
+		Softs::FSolverReal MeshScale)
 		: EdgeConstraints(
 			InParticles,
 			InParticleOffset,
@@ -328,6 +432,9 @@ public:
 			TriangleMesh,
 			FaceVertexPatternPositions,
 			GetXPBDAnisoSpringUse3dRestLengths(PropertyCollection, bDefaultUse3dRestLengths),
+			GetSavePreResizedSimPosition3D(PropertyCollection, bDefaultSavePreResizedSimPosition3D),
+			PreResizedSimPosition3D,
+			MeshScale,
 			WeightMaps.FindRef(GetXPBDAnisoSpringStiffnessWarpString(PropertyCollection, XPBDAnisoSpringStiffnessWarpName.ToString())),
 			WeightMaps.FindRef(GetXPBDAnisoSpringStiffnessWeftString(PropertyCollection, XPBDAnisoSpringStiffnessWeftName.ToString())),
 			WeightMaps.FindRef(GetXPBDAnisoSpringStiffnessBiasString(PropertyCollection, XPBDAnisoSpringStiffnessBiasName.ToString())),
@@ -347,6 +454,9 @@ public:
 			TriangleMesh,
 			FaceVertexPatternPositions,
 			GetXPBDAnisoSpringUse3dRestLengths(PropertyCollection, bDefaultUse3dRestLengths),
+			GetSavePreResizedSimPosition3D(PropertyCollection, bDefaultSavePreResizedSimPosition3D),
+			PreResizedSimPosition3D,
+			MeshScale,
 			WeightMaps.FindRef(GetXPBDAnisoSpringStiffnessWarpString(PropertyCollection, XPBDAnisoSpringStiffnessWarpName.ToString())),
 			WeightMaps.FindRef(GetXPBDAnisoSpringStiffnessWeftString(PropertyCollection, XPBDAnisoSpringStiffnessWeftName.ToString())),
 			WeightMaps.FindRef(GetXPBDAnisoSpringStiffnessBiasString(PropertyCollection, XPBDAnisoSpringStiffnessBiasName.ToString())),
@@ -366,6 +476,7 @@ public:
 		, XPBDAnisoSpringDampingIndex(PropertyCollection)
 		, XPBDAnisoSpringWarpScaleIndex(PropertyCollection)
 		, XPBDAnisoSpringWeftScaleIndex(PropertyCollection)
+		, SavePreResizedSimPosition3DIndex(PropertyCollection)
 	{}
 
 	CHAOS_API void SetProperties(
@@ -399,6 +510,7 @@ private:
 	UE_CHAOS_DECLARE_PROPERTYCOLLECTION_NAME(XPBDAnisoSpringDamping, float);
 	UE_CHAOS_DECLARE_PROPERTYCOLLECTION_NAME(XPBDAnisoSpringWarpScale, float);
 	UE_CHAOS_DECLARE_PROPERTYCOLLECTION_NAME(XPBDAnisoSpringWeftScale, float);
+	UE_CHAOS_DECLARE_PROPERTYCOLLECTION_NAME(SavePreResizedSimPosition3D, bool);
 };
 
 }  // End namespace Chaos::Softs
