Add FOnDemandChunkEntry::GetDiskSize. Fix some more uses of encoded size that should be DiskSize.

diff --git a/Engine/Source/Runtime/Experimental/IoStore/OnDemand/Private/OnDemandContentInstaller.cpp b/Engine/Source/Runtime/Experimental/IoStore/OnDemand/Private/OnDemandContentInstaller.cpp
--- a/Engine/Source/Runtime/Experimental/IoStore/OnDemand/Private/OnDemandContentInstaller.cpp
+++ b/Engine/Source/Runtime/Experimental/IoStore/OnDemand/Private/OnDemandContentInstaller.cpp
@@ -137,7 +137,7 @@ void ResolveChunksToInstall(
 				case EIoChunkType::ShaderCode:
 				{
 					ResolvedChunks.EntryIndices.Emplace(EntryIndex); 
-					ResolvedChunks.TotalSize += Container->ChunkEntries[EntryIndex].EncodedSize;
+					ResolvedChunks.TotalSize += Container->ChunkEntries[EntryIndex].GetDiskSize();
 				}
 				default:
 					break;
@@ -182,7 +182,7 @@ void ResolveChunksToInstall(
 		FOnDemandContainer& Container	= *ResolvedChunks->Container; 
 
 		ResolvedChunks->EntryIndices.Emplace(EntryIndex); 
-		ResolvedChunks->TotalSize += Container.ChunkEntries[EntryIndex].EncodedSize;
+		ResolvedChunks->TotalSize += Container.ChunkEntries[EntryIndex].GetDiskSize();
 
 		// TODO: Installing optional bulkdata should probably be an install argument
 		const EIoChunkType AdditionalPackageChunkTypes[] =
@@ -201,7 +201,7 @@ void ResolveChunksToInstall(
 				check(EntryIndex != INDEX_NONE);
 				FOnDemandContainer& OtherContainer = *ResolvedChunks->Container;
 				ResolvedChunks->EntryIndices.Emplace(EntryIndex);
-				ResolvedChunks->TotalSize += OtherContainer.ChunkEntries[EntryIndex].EncodedSize;
+				ResolvedChunks->TotalSize += OtherContainer.ChunkEntries[EntryIndex].GetDiskSize();
 			}
 		}
 	}
@@ -553,10 +553,10 @@ void FOnDemandContentInstaller::ProcessInstallRequest(FRequest& Request)
 						.ContainerIndex = ContainerIndex,
 						.EntryIndex		= EntryIndex
 					});
-					TotalInstallSize += Entry.EncodedSize;
+					TotalInstallSize += Entry.GetDiskSize();
 				}
 				++InstallRequest.ResolvedChunkCount;
-				TotalContentSize += Entry.EncodedSize;
+				TotalContentSize += Entry.GetDiskSize();
 			}
 
 			// Add references to existing chunk(s)
@@ -641,7 +641,7 @@ void FOnDemandContentInstaller::ExecuteInstallRequest(FRequest& Request, bool bR
 
 					if (IoStore.InstallCache->IsChunkCached(ChunkEntry.Hash))
 					{
-						InstallRequest.Progress.TotalInstallSize -= ChunkEntry.EncodedSize;
+						InstallRequest.Progress.TotalInstallSize -= ChunkEntry.GetDiskSize();
 						return true;
 					}
 
@@ -806,11 +806,11 @@ void FOnDemandContentInstaller::ProcessDownloadedChunk(
 	const FOnDemandChunkEntry& ChunkEntry		= Container->ChunkEntries[HttpRequest.EntryIndex];
 	const FIoChunkId& ChunkId					= Container->ChunkIds[HttpRequest.EntryIndex];
 
-	InstallRequest.Progress.CurrentInstallSize	+= ChunkEntry.EncodedSize;
+	InstallRequest.Progress.CurrentInstallSize	+= ChunkEntry.GetDiskSize();
 
 	UE_LOG(LogIoStoreOnDemand, Verbose, TEXT("Install progress %.2lf/%.2lf MiB, SeqNo=%u, Priority=%d, ContentHandle=(%s), ChunkId='%s', ChunkSize=%.2lf KiB, HttpStatus=%u"),
 		double(InstallRequest.Progress.CurrentInstallSize) / 1024.0 / 1024.0, double(InstallRequest.Progress.TotalInstallSize) / 1024.0 / 1024.0,
-		Request.SeqNo, Request.Priority, *LexToString(InstallRequest.Args.ContentHandle), *LexToString(ChunkId), double(ChunkEntry.EncodedSize) / 1024.0, HttpStatusCode);
+		Request.SeqNo, Request.Priority, *LexToString(InstallRequest.Args.ContentHandle), *LexToString(ChunkId), double(ChunkEntry.GetDiskSize()) / 1024.0, HttpStatusCode);
 
 	if (Request.ErrorCode == EIoErrorCode::Unknown)
 	{
diff --git a/Engine/Source/Runtime/Experimental/IoStore/OnDemand/Private/OnDemandInstallCache.cpp b/Engine/Source/Runtime/Experimental/IoStore/OnDemand/Private/OnDemandInstallCache.cpp
--- a/Engine/Source/Runtime/Experimental/IoStore/OnDemand/Private/OnDemandInstallCache.cpp
+++ b/Engine/Source/Runtime/Experimental/IoStore/OnDemand/Private/OnDemandInstallCache.cpp
@@ -2291,7 +2291,7 @@ FIoStatus FOnDemandInstallCache::Verify()
 		const FIoChunkId& ChunkId					= Container->ChunkIds[EntryIndex];
 		const FOnDemandChunkEntry& ChunkEntry		= Container->ChunkEntries[EntryIndex];
 		FSharedFileHandle FileHandle				= OpenResult.GetValue();
-		const int64 ChunkSize						= Align(int64(ChunkEntry.EncodedSize), FAES::AESBlockSize);
+		const int64 ChunkSize						= ChunkEntry.GetDiskSize();
 		TotalVerifiedBytes							+= ChunkSize;
 
 		if (int64(Chunk.GetSize()) < ChunkSize)
@@ -2508,7 +2508,7 @@ uint64 FOnDemandInstallCache::AddReferencesToBlocks(
 			const bool bToInstall = ChunksToInstall.Contains(Entry.Hash);
 			const bool bIsReferenced = IsReferenced[EntryIndex];
 
-			uint64 ChunkDiskSize = Align(int64(Entry.EncodedSize), FAES::AESBlockSize);
+			const uint64 ChunkDiskSize = Entry.GetDiskSize();
 
 			if (bIsReferenced)
 			{
@@ -2642,7 +2642,7 @@ FIoStatus FOnDemandInstallCache::Defrag(
 	struct FDefragBlockReferencedChunk
 	{
 		uint32 BlockOffset = 0;
-		uint32 EncodedSize = 0;
+		uint32 DiskSize = 0;
 		FIoHash Hash;
 	};
 
@@ -2779,7 +2779,7 @@ FIoStatus FOnDemandInstallCache::Defrag(
 					DefragBlock->ReferencedChunks.Add(FDefragBlockReferencedChunk
 					{
 						.BlockOffset = Loc.BlockOffset,
-						.EncodedSize = Entry.EncodedSize,
+						.DiskSize = Entry.GetDiskSize(),
 						.Hash = Entry.Hash,
 					});
 				}
@@ -2811,8 +2811,7 @@ FIoStatus FOnDemandInstallCache::Defrag(
 			{
 				FileHandle->Seek(ReffedChunk.BlockOffset);
 
-				const int64 ChunkDiskSize = Align(int64(ReffedChunk.EncodedSize), FAES::AESBlockSize);
-				FIoBuffer Buffer(ChunkDiskSize);
+				FIoBuffer Buffer(ReffedChunk.DiskSize);
 				const bool bOk = FileHandle->Read(Buffer.GetData(), Buffer.GetSize());
 				if (!bOk)
 				{
diff --git a/Engine/Source/Runtime/Experimental/IoStore/OnDemand/Private/OnDemandIoStore.cpp b/Engine/Source/Runtime/Experimental/IoStore/OnDemand/Private/OnDemandIoStore.cpp
--- a/Engine/Source/Runtime/Experimental/IoStore/OnDemand/Private/OnDemandIoStore.cpp
+++ b/Engine/Source/Runtime/Experimental/IoStore/OnDemand/Private/OnDemandIoStore.cpp
@@ -1727,7 +1727,7 @@ FOnDemandCacheUsage FOnDemandIoStore::GetCacheUsage(const FOnDemandGetCacheUsage
 					{
 						if (Refs.Indices[EntryIndex])
 						{
-							ReferencedBytesByHandle += ChunkEntry.EncodedSize;
+							ReferencedBytesByHandle += ChunkEntry.GetDiskSize();
 						}
 						EntryIndex++;
 					}
diff --git a/Engine/Source/Runtime/Experimental/IoStore/OnDemand/Private/OnDemandIoStore.h b/Engine/Source/Runtime/Experimental/IoStore/OnDemand/Private/OnDemandIoStore.h
--- a/Engine/Source/Runtime/Experimental/IoStore/OnDemand/Private/OnDemandIoStore.h
+++ b/Engine/Source/Runtime/Experimental/IoStore/OnDemand/Private/OnDemandIoStore.h
@@ -61,6 +61,8 @@ struct FOnDemandChunkEntry
 	uint32	BlockOffset = ~uint32(0);
 	uint32	BlockCount = 0;
 	uint8	CompressionFormatIndex = 0;
+
+	uint32	GetDiskSize() const { return Align(EncodedSize, FAES::AESBlockSize); }
 };
 static_assert(sizeof(FOnDemandChunkEntry) == 40);
 
