Geometry collection : Fixed Nanite based geometry collection rendering that would not updating properly when driven by a static pose chaos cache in cooked games

diff --git a/Engine/Source/Runtime/Experimental/GeometryCollectionEngine/Private/GeometryCollection/GeometryCollectionComponent.cpp b/Engine/Source/Runtime/Experimental/GeometryCollectionEngine/Private/GeometryCollection/GeometryCollectionComponent.cpp
--- a/Engine/Source/Runtime/Experimental/GeometryCollectionEngine/Private/GeometryCollection/GeometryCollectionComponent.cpp
+++ b/Engine/Source/Runtime/Experimental/GeometryCollectionEngine/Private/GeometryCollection/GeometryCollectionComponent.cpp
@@ -1793,12 +1793,11 @@ void UGeometryCollectionComponent::RestTransformsChanged()
 	{
 		OnTransformsDirty();
 
-#if WITH_EDITOR
-		// We need to do this in case we're controlled by Sequencer in editor, which doesn't invoke PostEditChangeProperty
+		// This is needed to make sure geometry Collections driven by a Chaos Cache in Static Pose renders and replay properly
 		UpdateCachedBounds();
-		SendRenderTransform_Concurrent();
-#endif
+		MarkRenderTransformDirty();
 
+		// force update the dynamic data 
 		SendDynamicDataToSceneProxy();
 	}
 	else
