Added install option to include/exclude optional bulk data.

#rnx

diff --git a/Engine/Source/Runtime/Experimental/IoStore/OnDemand/Private/OnDemandContentInstaller.cpp b/Engine/Source/Runtime/Experimental/IoStore/OnDemand/Private/OnDemandContentInstaller.cpp
--- a/Engine/Source/Runtime/Experimental/IoStore/OnDemand/Private/OnDemandContentInstaller.cpp
+++ b/Engine/Source/Runtime/Experimental/IoStore/OnDemand/Private/OnDemandContentInstaller.cpp
@@ -117,6 +117,7 @@ void ResolveChunksToInstall(
 	const TSet<FSharedOnDemandContainer>& Containers,
 	const TSet<FPackageId>& PackageIds,
 	bool bIncludeSoftReferences,
+	bool bIncludeOptionalBulkData,
 	TArray<FResolvedContainerChunks>& OutResolvedContainerChunks,
 	TSet<FPackageId>& OutMissing)
 {
@@ -184,15 +185,24 @@ void ResolveChunksToInstall(
 		ResolvedChunks->EntryIndices.Emplace(EntryIndex); 
 		ResolvedChunks->TotalSize += Container.ChunkEntries[EntryIndex].GetDiskSize();
 
-		// TODO: Installing optional bulkdata should probably be an install argument
-		const EIoChunkType AdditionalPackageChunkTypes[] =
+		static constexpr const EIoChunkType RequiredChunkTypes[] =
+		{
+			EIoChunkType::BulkData,
+			EIoChunkType::MemoryMappedBulkData 
+		};
+
+		static constexpr const EIoChunkType RequiredAndOptionalChunkTypes[] =
 		{
 			EIoChunkType::BulkData,
 			EIoChunkType::OptionalBulkData,
 			EIoChunkType::MemoryMappedBulkData 
 		};
 
-		for (EIoChunkType ChunkType : AdditionalPackageChunkTypes)
+		TConstArrayView<EIoChunkType> AdditionalChunkTypes = bIncludeOptionalBulkData
+			? MakeArrayView<const EIoChunkType>(RequiredAndOptionalChunkTypes, UE_ARRAY_COUNT(RequiredAndOptionalChunkTypes))
+			: MakeArrayView<const EIoChunkType>(RequiredChunkTypes, UE_ARRAY_COUNT(RequiredChunkTypes));
+
+		for (EIoChunkType ChunkType : AdditionalChunkTypes)
 		{
 			// TODO: For Mutable we need to traverse all possible bulk data chunk indices?
 			const FIoChunkId ChunkId = CreateBulkDataIoChunkId(PackageId.Value(), 0, 0, ChunkType);
@@ -487,11 +497,13 @@ void FOnDemandContentInstaller::ProcessInstallRequest(FRequest& Request)
 	}
 
 	TSet<FPackageId> Missing;
-	const bool bIncludeSoftReferences = EnumHasAnyFlags(InstallRequest.Args.Options, EOnDemandInstallOptions::InstallSoftReferences);
+	const bool bIncludeSoftReferences	= EnumHasAnyFlags(InstallRequest.Args.Options, EOnDemandInstallOptions::InstallSoftReferences);
+	const bool bIncludeOptionalBulkData = EnumHasAnyFlags(InstallRequest.Args.Options, EOnDemandInstallOptions::InstallOptionalBulkData);
 	Private::ResolveChunksToInstall(
 		ContainersForInstallation,
 		PackageIdsToInstall,
 		bIncludeSoftReferences,
+		bIncludeOptionalBulkData,
 		InstallRequest.ResolvedChunks,
 		Missing);
 
diff --git a/Engine/Source/Runtime/Experimental/IoStore/OnDemand/Private/OnDemandContentInstaller.h b/Engine/Source/Runtime/Experimental/IoStore/OnDemand/Private/OnDemandContentInstaller.h
--- a/Engine/Source/Runtime/Experimental/IoStore/OnDemand/Private/OnDemandContentInstaller.h
+++ b/Engine/Source/Runtime/Experimental/IoStore/OnDemand/Private/OnDemandContentInstaller.h
@@ -48,6 +48,7 @@ void ResolveChunksToInstall(
 	const TSet<FSharedOnDemandContainer>& Containers,
 	const TSet<FPackageId>& PackageIds,
 	bool bIncludeSoftReferences,
+	bool bIncludeOptionalBulkData,
 	TArray<FResolvedContainerChunks>& OutResolvedContainerChunks,
 	TSet<FPackageId>& OutMissing);
 
diff --git a/Engine/Source/Runtime/Experimental/IoStore/OnDemand/Private/OnDemandIoStore.cpp b/Engine/Source/Runtime/Experimental/IoStore/OnDemand/Private/OnDemandIoStore.cpp
--- a/Engine/Source/Runtime/Experimental/IoStore/OnDemand/Private/OnDemandIoStore.cpp
+++ b/Engine/Source/Runtime/Experimental/IoStore/OnDemand/Private/OnDemandIoStore.cpp
@@ -655,7 +655,8 @@ TIoStatusOr<uint64> FOnDemandIoStore::GetInstallSize(const FOnDemandGetInstallSi
 		return Status;
 	}
 
-	const bool bIncludeSoftReferences = EnumHasAnyFlags(Args.Options, EOnDemandGetInstallSizeOptions::IncludeSoftReferences);
+	const bool bIncludeSoftReferences	= EnumHasAnyFlags(Args.Options, EOnDemandGetInstallSizeOptions::IncludeSoftReferences);
+	const bool bIncludeOptionalBulkData = EnumHasAnyFlags(Args.Options, EOnDemandGetInstallSizeOptions::IncludeOptionalBulkData);
 	TArray<FResolvedContainerChunks> ResolvedChunks;
 	TSet<FPackageId> Missing;
 
@@ -663,6 +664,7 @@ TIoStatusOr<uint64> FOnDemandIoStore::GetInstallSize(const FOnDemandGetInstallSi
 		ContainersForInstallation,
 		PackageIdsToInstall,
 		bIncludeSoftReferences,
+		bIncludeOptionalBulkData,
 		ResolvedChunks,
 		Missing);
 
@@ -693,7 +695,8 @@ FIoStatus FOnDemandIoStore::GetInstallSizesByMountId(const FOnDemandGetInstallSi
 		return Status;
 	}
 
-	const bool bIncludeSoftReferences = EnumHasAnyFlags(Args.Options, EOnDemandGetInstallSizeOptions::IncludeSoftReferences);
+	const bool bIncludeSoftReferences	= EnumHasAnyFlags(Args.Options, EOnDemandGetInstallSizeOptions::IncludeSoftReferences);
+	const bool bIncludeOptionalBulkData = EnumHasAnyFlags(Args.Options, EOnDemandGetInstallSizeOptions::IncludeOptionalBulkData);
 	TArray<FResolvedContainerChunks> ResolvedChunks;
 	TSet<FPackageId> Missing;
 
@@ -701,6 +704,7 @@ FIoStatus FOnDemandIoStore::GetInstallSizesByMountId(const FOnDemandGetInstallSi
 		ContainersForInstallation,
 		PackageIdsToInstall,
 		bIncludeSoftReferences,
+		bIncludeOptionalBulkData,
 		ResolvedChunks,
 		Missing);
 
diff --git a/Engine/Source/Runtime/Experimental/IoStore/OnDemandCore/Public/IO/IoStoreOnDemand.h b/Engine/Source/Runtime/Experimental/IoStore/OnDemandCore/Public/IO/IoStoreOnDemand.h
--- a/Engine/Source/Runtime/Experimental/IoStore/OnDemandCore/Public/IO/IoStoreOnDemand.h
+++ b/Engine/Source/Runtime/Experimental/IoStore/OnDemandCore/Public/IO/IoStoreOnDemand.h
@@ -174,11 +174,13 @@ using FOnDemandMountCompleted = TUniqueFunction<void(FOnDemandMountResult)>;
 enum class EOnDemandInstallOptions
 {
 	/** No additional options. */
-	None				 = 0,
+	None					= 0,
 	/** Trigger callback on game thread. */
-	CallbackOnGameThread = 1 << 0,
-	/** Follow soft references when gathering packages to install */
-	InstallSoftReferences = 1 << 1,
+	CallbackOnGameThread	= 1 << 0,
+	/** Follow soft references when gathering packages to install. */
+	InstallSoftReferences	= 1 << 1,
+	/** Install optional bulk data. */
+	InstallOptionalBulkData = 1 << 2,
 };
 ENUM_CLASS_FLAGS(EOnDemandInstallOptions);
 
@@ -373,9 +375,11 @@ using FOnDemandDefragCompleted = TUniqueFunction<void(FOnDemandDefragResult)>;
 enum class EOnDemandGetInstallSizeOptions
 {
 	/** No additional options. */
-	None = 0,
-	/** Follow soft references when gathering packages to install */
-	IncludeSoftReferences = 1 << 0,
+	None					= 0,
+	/** Follow soft references when gathering packages to install. */
+	IncludeSoftReferences	= 1 << 0,
+	/** Include optional bulk data. */ 
+	IncludeOptionalBulkData	= 1 << 1
 };
 ENUM_CLASS_FLAGS(EOnDemandGetInstallSizeOptions);
 
