Keep track of left offset instead of using value from the base.

#rnx

diff --git a/Engine/Source/Runtime/Experimental/IoStore/HttpClient/Private/Activity.inl b/Engine/Source/Runtime/Experimental/IoStore/HttpClient/Private/Activity.inl
--- a/Engine/Source/Runtime/Experimental/IoStore/HttpClient/Private/Activity.inl
+++ b/Engine/Source/Runtime/Experimental/IoStore/HttpClient/Private/Activity.inl
@@ -18,7 +18,7 @@ public:
 	FAnsiStringView	GetMethod() const;
 	void			EnumerateHeaders(FResponse::FHeaderSink Sink) const;
 
-protected:
+private:
 	FRequestCloak&	operator << (FAnsiStringView Value);
 	uint16			HeaderLeft;
 	int16			HeaderRight = -1;
@@ -86,6 +86,7 @@ public:
 	const char*	TEMP_GetMessageRight() const;
 
 private:
+	int16		MessageLeft = -1;
 	int16		MessageRight = -1;
 };
 
@@ -115,6 +116,11 @@ FOutcome FResponseCloak::TryRecvResponse(FHttpPeer& Peer)
 		return -1;
 	};
 
+	if (MessageLeft < 0)
+	{
+		MessageLeft = int16(GetSize());
+	}
+
 	while (true)
 	{
 		static const uint32 PageSize = 256;
@@ -133,7 +139,7 @@ FOutcome FResponseCloak::TryRecvResponse(FHttpPeer& Peer)
 		// Rewind a little to cover cases where the terminal is fragmented across
 		// recv() calls
 		uint32 DestBias = 0;
-		if (Dest - 3 >= GetData() + HeaderRight)
+		if (Dest - 3 >= GetData() + MessageLeft)
 		{
 			Dest -= (DestBias = 3);
 		}
@@ -141,7 +147,7 @@ FOutcome FResponseCloak::TryRecvResponse(FHttpPeer& Peer)
 		int32 MessageEnd = FindMessageTerminal(Dest, Result + DestBias);
 		if (MessageEnd < 0)
 		{
-			if (GetSize() - HeaderRight > MaxHeaderSize)
+			if (GetSize() - MessageLeft > MaxHeaderSize)
 			{
 				return FOutcome::Error("Headers have grown larger than expected");
 			}
@@ -196,6 +202,8 @@ void FStopwatch::Impl(uint32 Index)
 
 #endif // IAS_HTTP_WITH_PERF
 
+
+
 ////////////////////////////////////////////////////////////////////////////////
 static FLaneEstate* GActivityTraceEstate = LaneEstate_New({
 	.Name = "Iax/Activity",
