Add an "IAX http" line to the IAx display statistics.

#rnx

- This line shows if the http thread has a data rate cap and if so what that cap is as well as the current number of requests in flight and if it has reached the max concurrent request cap or not.
- This new line helps show what sort of load a users IAX http system is under when trying to diagnose slowdown issues.
- Showing the data rate cap is also important as people might not be aware at which points in a processes life span that a cap may be applied, so the slowdown may be intended!

diff --git a/Engine/Source/Runtime/Experimental/IoStore/OnDemand/Private/OnDemandHttpThread.cpp b/Engine/Source/Runtime/Experimental/IoStore/OnDemand/Private/OnDemandHttpThread.cpp
--- a/Engine/Source/Runtime/Experimental/IoStore/OnDemand/Private/OnDemandHttpThread.cpp
+++ b/Engine/Source/Runtime/Experimental/IoStore/OnDemand/Private/OnDemandHttpThread.cpp
@@ -367,6 +367,8 @@ void FOnDemandHttpThread::TickRequests()
 					}
 #endif // UE_ALLOW_INVALID_URL_DEBUGGING
 
+					FOnDemandIoBackendStats::Get()->OnHttpRequestStarted();
+
 					HttpClient->Get(HttpRequest->HostGroup().GetUnderlyingHostGroup(), ChunkUrl, HttpRequest->ChunkRange,
 						[this, HttpRequest, &NumConcurrentRequests]
 						(FMultiEndpointHttpClientResponse&& HttpResponse)
@@ -374,6 +376,7 @@ void FOnDemandHttpThread::TickRequests()
 							TRACE_CPUPROFILER_EVENT_SCOPE(FOnDemandHttpThread::RequestCallback);
 
 							NumConcurrentRequests--;
+							FOnDemandIoBackendStats::Get()->OnHttpRequestCompleted();
 
 							if (HttpResponse.RetryCount > 0)
 							{
diff --git a/Engine/Source/Runtime/Experimental/IoStore/OnDemand/Private/Statistics.cpp b/Engine/Source/Runtime/Experimental/IoStore/OnDemand/Private/Statistics.cpp
--- a/Engine/Source/Runtime/Experimental/IoStore/OnDemand/Private/Statistics.cpp
+++ b/Engine/Source/Runtime/Experimental/IoStore/OnDemand/Private/Statistics.cpp
@@ -61,6 +61,8 @@ static FAutoConsoleVariableRef CVar_DisplayOnScreenStatistics(
 
 #endif // UE_ENABLE_ONSCREEN_STATISTICS
 
+extern int32 GIasHttpRateLimitKiBPerSecond;
+
 ////////////////////////////////////////////////////////////////////////////////
 static int32 BytesToApproxMB(uint64 Bytes) { return int32(Bytes >> 20); }
 static int32 BytesToApproxKB(uint64 Bytes) { return int32(Bytes >> 10); }
@@ -305,6 +307,7 @@ UE_IAX_COUNTER(HttpBandwidthMbps, FCounterInt);
 double GHttpDurationMsAvg[(uint8)EHttpRequestType::NUM_SOURCES] = { 0.0 };
 int32 GHttpDurationMsMax[(uint8)EHttpRequestType::NUM_SOURCES] = { 0 };
 int64 GHttpDurationMsSum[(uint8)EHttpRequestType::NUM_SOURCES] = { 0 };
+int32 GHttpNumConcurrentRequests = 0;
 
 struct FHttpRecentHistoryStatistics
 {
@@ -614,30 +617,12 @@ FOnDemandIoBackendStats::FOnDemandIoBackendStats(FBackendStatus& InStatus)
 				Out.Add(Verbosity, FText::FromStringView(Text.ToView()));
 			});
 
-			{
-				TStringBuilder<256> CachingText;
-				CachingText << TEXT("IAS Cache Stats: ");
-				if (BackendStatus.IsCacheEnabled())
-				{
-					CachingText << TEXT("Cached: ") << FText::AsMemory(GCacheCachedBytes.Get()).ToString();
-					CachingText << TEXT(" | Rejected: ") << FText::AsMemory(GCacheRejectBytes.Get()).ToString();
-					CachingText << TEXT(" | Read: ") << FText::AsMemory(GCacheReadBytes.Get()).ToString();
-					CachingText << TEXT(" (") << GCacheGetCount.Get() << TEXT(")");
-				}
-				else
-				{
-					CachingText << TEXTVIEW("Caching Disabled");
-				}
-
-				Out.Add(FCoreDelegates::EOnScreenMessageSeverity::Info, FText::FromStringView(CachingText.ToView()));
-			}
-
 			{
 				auto HttpStats = [&Out](const TCHAR* Title, EHttpRequestType Type)
 				{
 					TStringBuilder<256> Builder;
 					Builder.Appendf(
-						TEXT("%s Backend Stats: Downloaded: %s (%d) Avg %d ms | Retries: %d | Pending: %d"),
+						TEXT("%s Http: Downloaded: %s (%d) Avg %d ms | Retries: %d | Pending: %d"),
 						Title,
 						*FText::AsMemory(GHttpDownloadedBytes[(uint8)Type].Get()).ToString(),
 						GHttpGetCount[(uint8)Type].Get(),
@@ -649,6 +634,28 @@ FOnDemandIoBackendStats::FOnDemandIoBackendStats(FBackendStatus& InStatus)
 					Out.Add(FCoreDelegates::EOnScreenMessageSeverity::Info, FText::FromStringView(Builder.ToView()));
 				};
 
+				{
+					TStringBuilder<256> Builder;
+					Builder << TEXT("IAX Http:");
+					if (GIasHttpRateLimitKiBPerSecond > 0)
+					{
+						Builder << TEXT(" DataRate Cap: ") << GIasHttpRateLimitKiBPerSecond << TEXT("KiB/s");
+					}
+					else
+					{
+						Builder << TEXT(" DataRate Cap: None");
+					}
+
+					const int32 MaxNumRequests = 8;
+					Builder << TEXT(" ConcurrentRequests: ") << GHttpNumConcurrentRequests << TEXT("/") << MaxNumRequests;
+					if (GHttpNumConcurrentRequests >= MaxNumRequests)
+					{
+						Builder << TEXT( "(saturated)");
+					}
+
+					Out.Add(FCoreDelegates::EOnScreenMessageSeverity::Info, FText::FromStringView(Builder.ToView()));
+				}
+
 				HttpStats(TEXT("IAS"), EHttpRequestType::Streaming);
 				HttpStats(TEXT("IAD"), EHttpRequestType::Installed);
 			}
@@ -673,6 +680,24 @@ FOnDemandIoBackendStats::FOnDemandIoBackendStats(FBackendStatus& InStatus)
 				Out.Add(FCoreDelegates::EOnScreenMessageSeverity::Info, FText::FromStringView(Builder.ToView()));
 			}
 
+			{
+				TStringBuilder<256> CachingText;
+				CachingText << TEXT("IAS Cache Stats: ");
+				if (BackendStatus.IsCacheEnabled())
+				{
+					CachingText << TEXT("Cached: ") << FText::AsMemory(GCacheCachedBytes.Get()).ToString();
+					CachingText << TEXT(" | Rejected: ") << FText::AsMemory(GCacheRejectBytes.Get()).ToString();
+					CachingText << TEXT(" | Read: ") << FText::AsMemory(GCacheReadBytes.Get()).ToString();
+					CachingText << TEXT(" (") << GCacheGetCount.Get() << TEXT(")");
+				}
+				else
+				{
+					CachingText << TEXTVIEW("Caching Disabled");
+				}
+
+				Out.Add(FCoreDelegates::EOnScreenMessageSeverity::Info, FText::FromStringView(CachingText.ToView()));
+			}
+
 			if (GHttpDecodeErrorCount[(uint8)EHttpRequestType::Streaming].Get() > 0 || GCacheDecodeErrorCount.Get() > 0 || GHttpErrorCount[(uint8)EHttpRequestType::Streaming].Get() > 0)
 			{
 				TStringBuilder<256> Builder;
@@ -1111,6 +1136,16 @@ void FOnDemandIoBackendStats::OnHttpCdnCacheReply(EHttpRequestType Type, int32 R
 	}
 }
 
+void FOnDemandIoBackendStats::OnHttpRequestStarted()
+{
+	++GHttpNumConcurrentRequests;
+}
+
+void FOnDemandIoBackendStats::OnHttpRequestCompleted()
+{
+	--GHttpNumConcurrentRequests;
+}
+
 ////////////////////////////////////////////////////////////////////////////////
 void FOnDemandContentInstallerStats::OnRequestEnqueued()
 {
diff --git a/Engine/Source/Runtime/Experimental/IoStore/OnDemand/Private/Statistics.h b/Engine/Source/Runtime/Experimental/IoStore/OnDemand/Private/Statistics.h
--- a/Engine/Source/Runtime/Experimental/IoStore/OnDemand/Private/Statistics.h
+++ b/Engine/Source/Runtime/Experimental/IoStore/OnDemand/Private/Statistics.h
@@ -79,6 +79,8 @@ public:
 	void OnHttpError(EHttpRequestType Type) IAS_STATISTICS_IMPL()
 	void OnHttpDecodeError(EHttpRequestType Type) IAS_STATISTICS_IMPL()
 	void OnHttpCdnCacheReply(EHttpRequestType Type,int32 Reply) IAS_STATISTICS_IMPL()
+	void OnHttpRequestStarted() IAS_STATISTICS_IMPL()
+	void OnHttpRequestCompleted() IAS_STATISTICS_IMPL()
 
 private:
 
