Update a series of USTRUCTs with #if CPP where needed.  UHT will once again report extra bases as errors.

#rnx

diff --git a/Engine/Plugins/Chooser/Source/Chooser/Internal/AnimNode_ChooserPlayer.h b/Engine/Plugins/Chooser/Source/Chooser/Internal/AnimNode_ChooserPlayer.h
--- a/Engine/Plugins/Chooser/Source/Chooser/Internal/AnimNode_ChooserPlayer.h
+++ b/Engine/Plugins/Chooser/Source/Chooser/Internal/AnimNode_ChooserPlayer.h
@@ -22,7 +22,10 @@ enum class EChooserEvaluationFrequency
 };
 
 USTRUCT(BlueprintInternalUseOnly)
-struct FAnimNode_ChooserPlayer : public FAnimNode_BlendStack_Standalone, public IHasContextClass
+struct FAnimNode_ChooserPlayer : public FAnimNode_BlendStack_Standalone
+#if CPP
+	, public IHasContextClass
+#endif
 {
 	GENERATED_BODY()
 
diff --git a/Engine/Plugins/PCG/Source/PCG/Public/Elements/PCGStaticMeshSpawnerContext.h b/Engine/Plugins/PCG/Source/PCG/Public/Elements/PCGStaticMeshSpawnerContext.h
--- a/Engine/Plugins/PCG/Source/PCG/Public/Elements/PCGStaticMeshSpawnerContext.h
+++ b/Engine/Plugins/PCG/Source/PCG/Public/Elements/PCGStaticMeshSpawnerContext.h
@@ -22,7 +22,10 @@ struct FPCGInstancesAndWeights
 };
 
 USTRUCT(BlueprintType)
-struct FPCGStaticMeshSpawnerContext : public FPCGContext, public IPCGAsyncLoadingContext
+struct FPCGStaticMeshSpawnerContext : public FPCGContext
+#if CPP
+	, public IPCGAsyncLoadingContext
+#endif
 {
 	GENERATED_BODY()
 
diff --git a/Engine/Plugins/Runtime/HairStrands/Source/HairStrandsSolver/Public/GroomSolverComponent.h b/Engine/Plugins/Runtime/HairStrands/Source/HairStrandsSolver/Public/GroomSolverComponent.h
--- a/Engine/Plugins/Runtime/HairStrands/Source/HairStrandsSolver/Public/GroomSolverComponent.h
+++ b/Engine/Plugins/Runtime/HairStrands/Source/HairStrandsSolver/Public/GroomSolverComponent.h
@@ -57,7 +57,10 @@ struct FGroomSolverSettings
 
 /** Dataflow groom solver proxy used in dataflow simulation */ 
 USTRUCT(BlueprintType)
-struct FDataflowGroomSolverProxy : public FDataflowPhysicsSolverProxy, public Chaos::FPhysicsSolverEvents
+struct FDataflowGroomSolverProxy : public FDataflowPhysicsSolverProxy
+#if CPP
+	, public Chaos::FPhysicsSolverEvents
+#endif
 {
 	GENERATED_USTRUCT_BODY()
 
diff --git a/Engine/Plugins/Runtime/RigVM/Source/RigVM/Public/RigVMFunctions/RigVMDispatch_CastEnum.h b/Engine/Plugins/Runtime/RigVM/Source/RigVM/Public/RigVMFunctions/RigVMDispatch_CastEnum.h
--- a/Engine/Plugins/Runtime/RigVM/Source/RigVM/Public/RigVMFunctions/RigVMDispatch_CastEnum.h
+++ b/Engine/Plugins/Runtime/RigVM/Source/RigVM/Public/RigVMFunctions/RigVMDispatch_CastEnum.h
@@ -12,7 +12,10 @@ struct FRigVMDispatch_CastEnumBase
 };
 
 USTRUCT(meta=(DisplayName = "Cast", Category = "Enum", Keywords = "As", NodeColor = "1,1,1,1"))
-struct FRigVMDispatch_CastEnumToInt : public FRigVMDispatchFactory, public FRigVMDispatch_CastEnumBase
+struct FRigVMDispatch_CastEnumToInt : public FRigVMDispatchFactory
+#if CPP
+	, public FRigVMDispatch_CastEnumBase
+#endif
 {
 	GENERATED_BODY()
 
@@ -36,7 +39,10 @@ protected:
 };
 
 USTRUCT(meta=(DisplayName = "Cast", Category = "Enum", Keywords = "As", NodeColor = "1,1,1,1"))
-struct FRigVMDispatch_CastIntToEnum : public FRigVMDispatchFactory, public FRigVMDispatch_CastEnumBase
+struct FRigVMDispatch_CastIntToEnum : public FRigVMDispatchFactory
+#if CPP
+	, public FRigVMDispatch_CastEnumBase
+#endif
 {
 	GENERATED_BODY()
 
diff --git a/Engine/Plugins/VirtualProduction/Avalanche/Source/AvalancheTransition/Public/Conditions/AvaTransitionCondition.h b/Engine/Plugins/VirtualProduction/Avalanche/Source/AvalancheTransition/Public/Conditions/AvaTransitionCondition.h
--- a/Engine/Plugins/VirtualProduction/Avalanche/Source/AvalancheTransition/Public/Conditions/AvaTransitionCondition.h
+++ b/Engine/Plugins/VirtualProduction/Avalanche/Source/AvalancheTransition/Public/Conditions/AvaTransitionCondition.h
@@ -8,7 +8,10 @@
 #include "AvaTransitionCondition.generated.h"
 
 USTRUCT(Category="Transition Logic", meta=(Hidden))
-struct FAvaTransitionCondition : public FStateTreeConditionBase, public FAvaTransitionNode
+struct FAvaTransitionCondition : public FStateTreeConditionBase
+#if CPP
+	, public FAvaTransitionNode
+#endif
 {
 	GENERATED_BODY()
 
diff --git a/Engine/Plugins/VirtualProduction/Avalanche/Source/AvalancheTransition/Public/Tasks/AvaTransitionTask.h b/Engine/Plugins/VirtualProduction/Avalanche/Source/AvalancheTransition/Public/Tasks/AvaTransitionTask.h
--- a/Engine/Plugins/VirtualProduction/Avalanche/Source/AvalancheTransition/Public/Tasks/AvaTransitionTask.h
+++ b/Engine/Plugins/VirtualProduction/Avalanche/Source/AvalancheTransition/Public/Tasks/AvaTransitionTask.h
@@ -8,7 +8,10 @@
 #include "AvaTransitionTask.generated.h"
 
 USTRUCT(meta=(Hidden))
-struct FAvaTransitionTask : public FStateTreeTaskBase, public FAvaTransitionNode
+struct FAvaTransitionTask : public FStateTreeTaskBase
+#if CPP
+	, public FAvaTransitionNode
+#endif
 {
 	GENERATED_BODY()
 
diff --git a/Engine/Source/Runtime/Engine/Classes/Animation/AnimNode_DeadBlending.h b/Engine/Source/Runtime/Engine/Classes/Animation/AnimNode_DeadBlending.h
--- a/Engine/Source/Runtime/Engine/Classes/Animation/AnimNode_DeadBlending.h
+++ b/Engine/Source/Runtime/Engine/Classes/Animation/AnimNode_DeadBlending.h
@@ -28,7 +28,10 @@
  * More specifically they specify the "half-life" - or how it takes for the velocities to be decayed by half.
  */
 USTRUCT(BlueprintInternalUseOnly)
-struct FAnimNode_DeadBlending : public FAnimNode_Base, public IBoneReferenceSkeletonProvider
+struct FAnimNode_DeadBlending : public FAnimNode_Base
+#if CPP
+	, public IBoneReferenceSkeletonProvider
+#endif
 {
 	GENERATED_BODY()
 
diff --git a/Engine/Source/Runtime/Engine/Classes/Animation/AnimNode_Inertialization.h b/Engine/Source/Runtime/Engine/Classes/Animation/AnimNode_Inertialization.h
--- a/Engine/Source/Runtime/Engine/Classes/Animation/AnimNode_Inertialization.h
+++ b/Engine/Source/Runtime/Engine/Classes/Animation/AnimNode_Inertialization.h
@@ -297,7 +297,10 @@ private:
 };
 
 USTRUCT(BlueprintInternalUseOnly)
-struct FAnimNode_Inertialization : public FAnimNode_Base, public IBoneReferenceSkeletonProvider
+struct FAnimNode_Inertialization : public FAnimNode_Base
+#if CPP
+	, public IBoneReferenceSkeletonProvider
+#endif
 {
 	GENERATED_BODY()
 
diff --git a/Engine/Source/Runtime/Experimental/Dataflow/Engine/Public/Dataflow/DataflowSubGraphNodes.h b/Engine/Source/Runtime/Experimental/Dataflow/Engine/Public/Dataflow/DataflowSubGraphNodes.h
--- a/Engine/Source/Runtime/Experimental/Dataflow/Engine/Public/Dataflow/DataflowSubGraphNodes.h
+++ b/Engine/Source/Runtime/Experimental/Dataflow/Engine/Public/Dataflow/DataflowSubGraphNodes.h
@@ -33,7 +33,9 @@ namespace UE::Dataflow
 USTRUCT(meta=(Icon="GraphEditor.Function_16x"))
 struct FDataflowSubGraphInputNode
 	: public FDataflowNode
+#if CPP
 	, public FDataflowDynamicConnections::IOwnerInterface
+#endif
 {
 	GENERATED_USTRUCT_BODY()
 
@@ -71,7 +73,9 @@ private:
 USTRUCT(meta = (Icon = "GraphEditor.Function_16x"))
 struct FDataflowSubGraphOutputNode 
 	: public FDataflowNode
+#if CPP
 	, public FDataflowDynamicConnections::IOwnerInterface
+#endif
 {
 	GENERATED_USTRUCT_BODY()
 
@@ -135,8 +139,10 @@ private:
 USTRUCT(meta = (Icon = "GraphEditor.Function_16x"))
 struct FDataflowCallSubGraphNode
 	: public FDataflowNode
+#if CPP
 	, public FDataflowDynamicConnections::IOwnerInterface
 	, public UE::Dataflow::ISubGraphContextCallback
+#endif
 {
 	GENERATED_USTRUCT_BODY()
 
