Collected the second stage of message handling.

#rnx

diff --git a/Engine/Source/Runtime/Experimental/IoStore/HttpClient/Private/Activity.inl b/Engine/Source/Runtime/Experimental/IoStore/HttpClient/Private/Activity.inl
--- a/Engine/Source/Runtime/Experimental/IoStore/HttpClient/Private/Activity.inl
+++ b/Engine/Source/Runtime/Experimental/IoStore/HttpClient/Private/Activity.inl
@@ -82,14 +82,73 @@ class FResponseCloak
 public:
 	using FRequestCloak::FRequestCloak;
 
-	FOutcome	TryRecvResponse(FHttpPeer& Peer);
-	const char*	TEMP_GetMessageRight() const;
+	FOutcome		TryRecvResponse(FHttpPeer& Peer);
+	bool			IsKeepAlive() const;
+	uint32			GetStatusCode() const;
+	FAnsiStringView	GetStatusMessage() const;
+	int64			GetContentLength() const;
+	bool			IsChunked() const;
+	void			ReadHeaders(FResponse::FHeaderSink Sink) const;
+	const char*		TEMP_GetMessageRight() const;
 
 private:
-	int16		MessageLeft = -1;
-	int16		MessageRight = -1;
+	FOutcome		Parse();
+	int64			ContentLength = -1;
+	FMessageOffsets	Offsets = {};
+	int16			MessageLeft = -1;
+	int16			MessageRight = -1;
+	uint16			StatusCode = 0;
+	uint8			bKeepAlive : 1;
+	uint8			bChunked : 1;
+	uint8			_Unused : 6;
 };
 
+////////////////////////////////////////////////////////////////////////////////
+bool FResponseCloak::IsKeepAlive() const
+{
+	check(MessageRight > 0);
+	return bKeepAlive;
+}
+
+////////////////////////////////////////////////////////////////////////////////
+uint32 FResponseCloak::GetStatusCode() const
+{
+	check(StatusCode >= 100);
+	return StatusCode;
+}
+
+////////////////////////////////////////////////////////////////////////////////
+FAnsiStringView	FResponseCloak::GetStatusMessage() const
+{
+	check(Offsets.Message > 0);
+	const char* Cursor = GetData() + MessageLeft;
+	return FAnsiStringView(Cursor + Offsets.Message, Offsets.Headers - Offsets.Message);
+}
+
+////////////////////////////////////////////////////////////////////////////////
+int64 FResponseCloak::GetContentLength() const
+{
+	check(Offsets.Headers > 0);
+	return ContentLength;
+}
+
+////////////////////////////////////////////////////////////////////////////////
+bool FResponseCloak::IsChunked() const
+{
+	check(MessageRight > 0);
+	return bChunked;
+}
+
+////////////////////////////////////////////////////////////////////////////////
+void FResponseCloak::ReadHeaders(FResponse::FHeaderSink Sink) const
+{
+	uint32 Offset = MessageLeft + Offsets.Headers;
+	uint32 Length = MessageRight - Offset - 2; // "-2" trims off '\r\n' that signals end of headers
+	const char* Cursor = GetData() + Offset;
+	FAnsiStringView Headers(Cursor, Length);
+	UE::IoStore::HTTP::EnumerateHeaders(Headers, Sink);
+}
+
 ////////////////////////////////////////////////////////////////////////////////
 const char* FResponseCloak::TEMP_GetMessageRight() const
 {
@@ -119,6 +178,8 @@ FOutcome FResponseCloak::TryRecvResponse(FHttpPeer& Peer)
 	if (MessageLeft < 0)
 	{
 		MessageLeft = int16(GetSize());
+		bKeepAlive = true;
+		bChunked = false;
 	}
 
 	while (true)
@@ -156,10 +217,67 @@ FOutcome FResponseCloak::TryRecvResponse(FHttpPeer& Peer)
 		}
 
 		MessageRight = int16(ptrdiff_t(Dest + MessageEnd - GetData()));
-		return FOutcome::Ok();
+		return Parse();
 	}
 }
 
+////////////////////////////////////////////////////////////////////////////////
+FOutcome FResponseCloak::Parse()
+{
+	const char* Cursor = GetData() + MessageLeft;
+	FAnsiStringView MessageView(Cursor, MessageRight - MessageLeft);
+
+	if (ParseMessage(MessageView, Offsets) < 0)
+	{
+		return FOutcome::Error("Failed to parse message status");
+	}
+
+	FAnsiStringView StatusCodeView(
+		Cursor + Offsets.StatusCode,
+		Offsets.Message - Offsets.StatusCode
+	);
+	StatusCode = uint16(CrudeToInt(StatusCodeView));
+	if (uint32(StatusCode - 100) > 899)
+	{
+		return FOutcome::Error("Invalid status code", StatusCode);
+	}
+
+	static const uint32 Flag_Length		= 1 << 0;
+	static const uint32 Flag_XferEnc	= 1 << 1;
+	static const uint32 Flag_Connection = 1 << 2;
+	static const uint32 Flag_All		= 0x07;
+	uint32 Flags = 0;
+
+	ReadHeaders(
+		[this, &Flags] (FAnsiStringView Name, FAnsiStringView Value) mutable
+		{
+			// todo; may need smarter value handling; ;/, separated options & key-value pairs (ex. in rfc2068)
+
+			if (Name.Equals("Content-Length", ESearchCase::IgnoreCase))
+			{
+				ContentLength = int32(CrudeToInt(Value));
+				Flags |= Flag_Length;
+			}
+
+			else if (Name.Equals("Transfer-Encoding", ESearchCase::IgnoreCase))
+			{
+				bChunked = Value.Equals("chunked", ESearchCase::IgnoreCase);
+				Flags |= Flag_XferEnc;
+			}
+
+			else if (Name.Equals("Connection", ESearchCase::IgnoreCase))
+			{
+				bKeepAlive = !Value.Equals("close");
+				Flags |= Flag_Connection;
+			}
+
+			return Flags != Flag_All;
+		}
+	);
+
+	return FOutcome::Ok();
+}
+
 
 
 #if IAS_HTTP_WITH_PERF
@@ -213,14 +331,6 @@ static FLaneEstate* GActivityTraceEstate = LaneEstate_New({
 });
 
 
-////////////////////////////////////////////////////////////////////////////////
-struct FResponseInternal
-{
-	FMessageOffsets Offsets;
-	int32			ContentLength = 0;
-	uint16			MessageLength;
-	mutable int16	Code;
-};
 
 ////////////////////////////////////////////////////////////////////////////////
 struct alignas(16) FActivity
@@ -260,7 +370,6 @@ struct alignas(16) FActivity
 	};
 	UPTRINT				SinkParam;
 	FTicketSink			Sink;
-	FResponseInternal	Response;
 	FResponseCloak		Buffer;
 };
 
@@ -301,7 +410,7 @@ static uint32 Activity_RemainingKiB(FActivity* Activity)
 	if (Activity->State <= FActivity::EState::RecvStream)  return MAX_uint32;
 	if (Activity->State >  FActivity::EState::RecvContent) return 0;
 
-	uint32 ContentLength = uint32(Activity->Response.ContentLength);
+	uint32 ContentLength = uint32(Activity->Buffer.GetContentLength());
 	check(Activity->StateParam <= ContentLength);
 	return (ContentLength - Activity->StateParam) >> 10;
 }
diff --git a/Engine/Source/Runtime/Experimental/IoStore/HttpClient/Private/Api.inl b/Engine/Source/Runtime/Experimental/IoStore/HttpClient/Private/Api.inl
--- a/Engine/Source/Runtime/Experimental/IoStore/HttpClient/Private/Api.inl
+++ b/Engine/Source/Runtime/Experimental/IoStore/HttpClient/Private/Api.inl
@@ -76,37 +76,21 @@ EStatusCodeClass FResponse::GetStatus() const
 uint32 FResponse::GetStatusCode() const
 {
 	const auto* Activity = (const FActivity*)this;
-	const FResponseInternal& Internal = Activity->Response;
-
-	const char* MessageData = Activity->Buffer.GetData() + Activity->StateParam;
-	if (Internal.Code < 0)
-	{
-		const char* CodePtr = MessageData + Internal.Offsets.StatusCode;
-		Internal.Code = uint16(CrudeToInt(FAnsiStringView(CodePtr, 3)));
-	}
-
-	return Internal.Code;
+	return Activity->Buffer.GetStatusCode();
 }
 
 ////////////////////////////////////////////////////////////////////////////////
 FAnsiStringView FResponse::GetStatusMessage() const
 {
 	const auto* Activity = (const FActivity*)this;
-	const FResponseInternal& Internal = Activity->Response;
-
-	const char* MessageData = Activity->Buffer.GetData() + Activity->StateParam;
-	return FAnsiStringView(
-		MessageData + Internal.Offsets.Message,
-		Internal.Offsets.Headers - Internal.Offsets.Message
-	);
+	return Activity->Buffer.GetStatusMessage();
 }
 
 ////////////////////////////////////////////////////////////////////////////////
 int64 FResponse::GetContentLength() const
 {
 	const auto* Activity = (const FActivity*)this;
-	const FResponseInternal& Internal = Activity->Response;
-	return Internal.ContentLength;
+	return Activity->Buffer.GetContentLength();
 }
 
 ////////////////////////////////////////////////////////////////////////////////
@@ -159,15 +143,7 @@ FAnsiStringView FResponse::GetHeader(FAnsiStringView Name) const
 void FResponse::ReadHeaders(FHeaderSink Sink) const
 {
 	const auto* Activity = (const FActivity*)this;
-	const FResponseInternal& Internal = Activity->Response;
-
-	const char* MessageData = Activity->Buffer.GetData() + Activity->StateParam;
-	FAnsiStringView Result, Headers(
-		MessageData + Internal.Offsets.Headers,
-		Internal.MessageLength - Internal.Offsets.Headers
-	);
-
-	EnumerateHeaders(Headers, Sink);
+	Activity->Buffer.ReadHeaders(Sink);
 }
 
 ////////////////////////////////////////////////////////////////////////////////
@@ -231,7 +207,7 @@ uint32 FTicketStatus::GetContentLength() const
 {
 	check(GetId() <= EId::Content);
 	const auto* Activity = (FActivity*)this;
-	return Activity->Response.ContentLength;
+	return uint32(Activity->Buffer.GetContentLength());
 }
 
 ////////////////////////////////////////////////////////////////////////////////
@@ -288,7 +264,7 @@ FTicketPerf::FSample FTicketPerf::GetSample() const
 	uint32 Bps = ~0u;
 	if (Sample.RecvMs)
 	{
-		Bps = uint32((Activity.Response.ContentLength * 1000ull) / Sample.RecvMs);
+		Bps = uint32((Activity.Buffer.GetContentLength() * 1000ull) / Sample.RecvMs);
 	}
 	Sample.RecvKiBps = Clamp(Bps >> 10);
 
diff --git a/Engine/Source/Runtime/Experimental/IoStore/HttpClient/Private/Loop.inl b/Engine/Source/Runtime/Experimental/IoStore/HttpClient/Private/Loop.inl
--- a/Engine/Source/Runtime/Experimental/IoStore/HttpClient/Private/Loop.inl
+++ b/Engine/Source/Runtime/Experimental/IoStore/HttpClient/Private/Loop.inl
@@ -118,57 +118,10 @@ static FOutcome DoRecvMessage(FActivity* Activity, FHttpPeer& Peer)
 	}
 
 	check(Outcome.IsOk());
-	const char* MessageRight = Buffer.TEMP_GetMessageRight();
-
-	// Fill out the internal response object
-	FResponseInternal& Internal = Activity->Response;
-	const char* MessageData = Buffer.GetData() + Activity->StateParam;
-	Internal.MessageLength = uint16(ptrdiff_t(MessageRight - MessageData));
-
-	FAnsiStringView ResponseView(MessageData, Internal.MessageLength);
-	if (ParseMessage(ResponseView, Internal.Offsets) < 0)
-	{
-		Activity_SetError(Activity, "Failed to parse message status");
-		return FOutcome::Error(Activity->ErrorReason);
-	}
-
-	// Parse headers
-	FAnsiStringView Headers = ResponseView.Mid(
-		Internal.Offsets.Headers,
-		Internal.MessageLength - Internal.Offsets.Headers - 2 // "-2" trims off '\r\n' that signals end of headers
-	);
-
-	int32 Count = 3;
-	bool bChunked = false;
-	bool IsKeepAlive = true;
-	int32 ContentLength = Activity->NoContent ? 0 : -1;
-	EnumerateHeaders(
-		Headers,
-		[&ContentLength, &bChunked, &IsKeepAlive, &Count] (FAnsiStringView Name, FAnsiStringView Value)
-		{
-			// todo; may need smarter value handling; ;/, separated options & key-value pairs (ex. in rfc2068)
-
-			if (Name.Equals("Content-Length", ESearchCase::IgnoreCase))
-			{
-				ContentLength = int32(CrudeToInt(Value));
-				Count--;
-			}
-
-			else if (Name.Equals("Transfer-Encoding", ESearchCase::IgnoreCase))
-			{
-				bChunked = Value.Equals("chunked", ESearchCase::IgnoreCase);
-				Count--;
-			}
 
-			else if (Name.Equals("Connection", ESearchCase::IgnoreCase))
-			{
-				IsKeepAlive = !Value.Equals("close");
-				Count--;
-			}
-
-			return Count > 0;
-		}
-	);
+	bool bChunked = Buffer.IsChunked();
+	bool IsKeepAlive = Buffer.IsKeepAlive();
+	int64 ContentLength = Activity->NoContent ? 0 : Buffer.GetContentLength();
 
 	Activity->IsKeepAlive &= IsKeepAlive;
 
@@ -191,8 +144,6 @@ static FOutcome DoRecvMessage(FActivity* Activity, FHttpPeer& Peer)
 
 	// Call out to the sink to get a content destination
 	FIoBuffer* PriorDest = Activity->Dest; // to retain unioned Host ptr (redirect uses it in sink)
-	Internal.Code = -1;
-	Internal.ContentLength = ContentLength;
 	Activity_CallSink(Activity);
 	Activity->NoContent |= (ContentLength == 0);
 
@@ -210,7 +161,7 @@ static FOutcome DoRecvMessage(FActivity* Activity, FHttpPeer& Peer)
 		if (int32 DestSize = int32(Dest.GetSize()); DestSize == 0)
 		{
 			static const uint32 DefaultChunkSize = 4 << 10;
-			uint32 Size = bChunked ? DefaultChunkSize : ContentLength;
+			uint32 Size = bChunked ? DefaultChunkSize : uint32(ContentLength);
 			Dest = FIoBuffer(Size);
 		}
 		else if (!bChunked && DestSize < ContentLength)
@@ -226,6 +177,7 @@ static FOutcome DoRecvMessage(FActivity* Activity, FHttpPeer& Peer)
 	}
 
 	// Perhaps we have some of the content already?
+	const char* MessageRight = Buffer.TEMP_GetMessageRight();
 	const char* BufferRight = Buffer.GetData() + Buffer.GetSize();
 	uint32 AlreadyReceived = uint32(ptrdiff_t(BufferRight - MessageRight));
 	if (AlreadyReceived > uint32(ContentLength))
@@ -279,8 +231,8 @@ static FOutcome DoRecvContent(FActivity* Activity, FHttpPeer& Peer, int32& MaxRe
 
 	while (true)
 	{
-		const FResponseInternal& Response = Activity->Response;
-		int32 Size = (Response.ContentLength - Activity->StateParam);
+		const auto& Response = Activity->Buffer;
+		int32 Size = (int32(Response.GetContentLength()) - Activity->StateParam);
 		if (Size == 0)
 		{
 			break;
