Move the periodic IAX logging to it's own method.

#rnx

- The FOnDemandIoBackendStats constructor has grown over time and has gotten difficult to read. The periodic logging code was an obvious choice to move to it's own method.

diff --git a/Engine/Source/Runtime/Experimental/IoStore/OnDemand/Private/Statistics.cpp b/Engine/Source/Runtime/Experimental/IoStore/OnDemand/Private/Statistics.cpp
--- a/Engine/Source/Runtime/Experimental/IoStore/OnDemand/Private/Statistics.cpp
+++ b/Engine/Source/Runtime/Experimental/IoStore/OnDemand/Private/Statistics.cpp
@@ -529,51 +529,7 @@ FOnDemandIoBackendStats::FOnDemandIoBackendStats(FBackendStatus& InStatus)
 		HttpCsv(EHttpRequestType::Streaming);
 		HttpCsv(EHttpRequestType::Installed);
 
-		if (GIasStatisticsLogInterval > 0.f)
-		{
-			static double LastLogTime = 0.0;
-			if (double Time = FPlatformTime::Seconds(); Time - LastLogTime > (double)GIasStatisticsLogInterval)
-			{
-				if (BackendStatus.IsCacheEnabled())
-				{
-					UE_LOG(LogIas, Log, TEXT("CacheStats: CachedKiB=%d, WrittenKiB=%d, ReadKiB=%d, RejectedKiB=%d, Get=%d, Error=%d, Put=%d, PutReject=%d, PutExisting=%d"),
-						(int32)CCachedKiB, (int32)CWrittenKiB, (int32)CReadKiB, (int32)CRejectedKiB, CGetCount, CErrorCount, CPutCount, CPutRejectCount, CPutExistingCount);
-				}
-				else
-				{
-					UE_LOG(LogIas, Log, TEXT("CacheStats: Disabled"));
-				}
-
-				auto HttpLog = [](const TCHAR* Title, EHttpRequestType Type)
-				{
-						UE_LOG(LogIas, Log, TEXT("%s - HttpStats: DownloadedKiB=%d, Get=%d, Retry=%d, Cancel=%d, Error=%d, CurPending=%d, CurDurationMsAvg=%d, CurDurationMsMax=%d"),
-							Title, 
-							(int32)((float)GHttpDownloadedBytes[(uint8)Type].Get() * OneOver1024),
-							(int32)GHttpGetCount[(uint8)Type].Get(),
-							(int32)GHttpRetryCount[(uint8)Type].Get(),
-							(int32)GHttpCancelCount[(uint8)Type].Get(),
-							(int32)GHttpErrorCount[(uint8)Type].Get(),
-							(int32)GHttpPendingCount[(uint8)Type].Get(),
-							(int32)GHttpDurationMsAvg[(uint8)Type],
-							GHttpDurationMsMax[(uint8)Type]);
-				};
-
-				HttpLog(TEXT("IAS"), EHttpRequestType::Streaming);
-				HttpLog(TEXT("IAD"), EHttpRequestType::Installed);
-
-#if COUNTERSTRACE_ENABLED
-				{
-					TUniqueLock Lock(InstallerTraceCounters);
-					UE_LOG(LogIoStoreOnDemand, Log, TEXT("IadStats: InstallCount=%llu, Downloaded=%d KiB, AvgInstallDuration=%dms, AvgCacheHitRatio=%d%%"),
-						InstallerTraceCounters.InstallCount.Get(),
-						(int32(InstallerTraceCounters.DownloadedBytes.Get()) / 1024),
-						int32(InstallerTraceCounters.AvgInstallDurationMs.Get()),
-						int32(InstallerTraceCounters.CacheHitRatio.GetMean() * 100.0));
-				}
-#endif // COUNTERSTRACE_ENABLED
-				LastLogTime = Time;
-			}
-		}
+		PrintPeriodicLogging();
 	});
 
 #if UE_ENABLE_ONSCREEN_STATISTICS
@@ -725,6 +681,69 @@ FOnDemandIoBackendStats* FOnDemandIoBackendStats::Get()
 	return GStatistics;
 }
 
+void FOnDemandIoBackendStats::PrintPeriodicLogging()
+{
+	if (GIasStatisticsLogInterval > 0.f)
+	{
+		static constexpr float OneOver1024 = 1.0f / 1024.0f;
+
+		int32 CGetCount = int32(GCacheGetCount.Get());
+		int32 CErrorCount = int32(GCacheErrorCount.Get());
+		int32 CPutCount = int32(GCachePutCount.Get());
+		int32 CPutExistingCount = int32(GCachePutExistingCount.Get());
+		int32 CPutRejectCount = int32(GCachePutRejectCount.Get());
+
+		float CCachedKiB = (float)GCacheCachedBytes.Get() * OneOver1024;
+		float CWrittenKiB = (float)GCacheWrittenBytes.Get() * OneOver1024;
+		float CReadKiB = (float)GCacheReadBytes.Get() * OneOver1024;
+		float CRejectedKiB = (float)GCacheRejectBytes.Get() * OneOver1024;
+
+		static double LastLogTime = 0.0;
+		if (double Time = FPlatformTime::Seconds(); Time - LastLogTime > (double)GIasStatisticsLogInterval)
+		{
+			if (BackendStatus.IsCacheEnabled())
+			{
+				UE_LOG(LogIas, Log, TEXT("CacheStats: CachedKiB=%d, WrittenKiB=%d, ReadKiB=%d, RejectedKiB=%d, Get=%d, Error=%d, Put=%d, PutReject=%d, PutExisting=%d"),
+					(int32)CCachedKiB, (int32)CWrittenKiB, (int32)CReadKiB, (int32)CRejectedKiB, CGetCount, CErrorCount, CPutCount, CPutRejectCount, CPutExistingCount);
+			}
+			else
+			{
+				UE_LOG(LogIas, Log, TEXT("CacheStats: Disabled"));
+			}
+
+			auto HttpLog = [](const TCHAR* Title, EHttpRequestType Type)
+				{
+					UE_LOG(LogIas, Log, TEXT("%s - HttpStats: DownloadedKiB=%d, Get=%d, Retry=%d, Cancel=%d, Error=%d, CurPending=%d, CurDurationMsAvg=%d, CurDurationMsMax=%d"),
+						Title,
+						(int32)((float)GHttpDownloadedBytes[(uint8)Type].Get() * OneOver1024),
+						(int32)GHttpGetCount[(uint8)Type].Get(),
+						(int32)GHttpRetryCount[(uint8)Type].Get(),
+						(int32)GHttpCancelCount[(uint8)Type].Get(),
+						(int32)GHttpErrorCount[(uint8)Type].Get(),
+						(int32)GHttpPendingCount[(uint8)Type].Get(),
+						(int32)GHttpDurationMsAvg[(uint8)Type],
+						GHttpDurationMsMax[(uint8)Type]);
+				};
+
+			HttpLog(TEXT("IAS"), EHttpRequestType::Streaming);
+			HttpLog(TEXT("IAD"), EHttpRequestType::Installed);
+
+#if COUNTERSTRACE_ENABLED
+			{
+				TUniqueLock Lock(InstallerTraceCounters);
+				UE_LOG(LogIoStoreOnDemand, Log, TEXT("IadStats: InstallCount=%llu, Downloaded=%d KiB, AvgInstallDuration=%dms, AvgCacheHitRatio=%d%%"),
+					InstallerTraceCounters.InstallCount.Get(),
+					(int32(InstallerTraceCounters.DownloadedBytes.Get()) / 1024),
+					int32(InstallerTraceCounters.AvgInstallDurationMs.Get()),
+					int32(InstallerTraceCounters.CacheHitRatio.GetMean() * 100.0));
+			}
+#endif // COUNTERSTRACE_ENABLED
+
+			LastLogTime = Time;
+		}
+	}
+}
+
 void FOnDemandIoBackendStats::ReportGeneralAnalytics(TArray<FAnalyticsEventAttribute>& OutAnalyticsArray) const
 {
 	// Note that this analytics section is not optional, if we are reporting analytics then we report this section
diff --git a/Engine/Source/Runtime/Experimental/IoStore/OnDemand/Private/Statistics.h b/Engine/Source/Runtime/Experimental/IoStore/OnDemand/Private/Statistics.h
--- a/Engine/Source/Runtime/Experimental/IoStore/OnDemand/Private/Statistics.h
+++ b/Engine/Source/Runtime/Experimental/IoStore/OnDemand/Private/Statistics.h
@@ -84,6 +84,8 @@ public:
 
 private:
 
+	void PrintPeriodicLogging();
+
 #if IAS_WITH_STATISTICS
 	const FBackendStatus& BackendStatus;
 #endif //IAS_WITH_STATISTICS
