Move updating of IAS csv values to it's own method.

- Follow on clean up work for CL 42878135
- Wrap the new private methods in IAS_WITH_STATISTICS checks.

diff --git a/Engine/Source/Runtime/Experimental/IoStore/OnDemand/Private/Statistics.cpp b/Engine/Source/Runtime/Experimental/IoStore/OnDemand/Private/Statistics.cpp
--- a/Engine/Source/Runtime/Experimental/IoStore/OnDemand/Private/Statistics.cpp
+++ b/Engine/Source/Runtime/Experimental/IoStore/OnDemand/Private/Statistics.cpp
@@ -493,42 +493,7 @@ FOnDemandIoBackendStats::FOnDemandIoBackendStats(FBackendStatus& InStatus)
 	GStatistics = this;
 	GStatisticsEndFrameDelegateHandle = FCoreDelegates::OnEndFrame.AddLambda([this]()
 	{
-		// cache stat totals
-		int32 CGetCount = int32(GCacheGetCount.Get());
-		int32 CErrorCount = int32(GCacheErrorCount.Get());
-		int32 CPutCount = int32(GCachePutCount.Get());
-		int32 CPutExistingCount = int32(GCachePutExistingCount.Get());
-		int32 CPutRejectCount = int32(GCachePutRejectCount.Get());
-		float CCachedKiB = (float)GCacheCachedBytes.Get() * OneOver1024;
-		float CWrittenKiB = (float)GCacheWrittenBytes.Get() * OneOver1024;
-		float CReadKiB = (float)GCacheReadBytes.Get() * OneOver1024;
-		float CRejectedKiB = (float)GCacheRejectBytes.Get() * OneOver1024;
-		CSV_CUSTOM_STAT_DEFINED(CacheGetCount, CGetCount, ECsvCustomStatOp::Set);
-		CSV_CUSTOM_STAT_DEFINED(CacheErrorCount, CErrorCount, ECsvCustomStatOp::Set);
-		CSV_CUSTOM_STAT_DEFINED(CachePutCount, CPutCount, ECsvCustomStatOp::Set);
-		CSV_CUSTOM_STAT_DEFINED(CachePutExistingCount, CPutExistingCount, ECsvCustomStatOp::Set);
-		CSV_CUSTOM_STAT_DEFINED(CachePutRejectCount, CPutRejectCount, ECsvCustomStatOp::Set);
-		CSV_CUSTOM_STAT_DEFINED(CacheCachedMB, CCachedKiB * OneOver1024, ECsvCustomStatOp::Set);
-		CSV_CUSTOM_STAT_DEFINED(CacheWrittenMB, CWrittenKiB * OneOver1024, ECsvCustomStatOp::Set);
-		CSV_CUSTOM_STAT_DEFINED(CacheReadMB, CReadKiB * OneOver1024, ECsvCustomStatOp::Set);
-		CSV_CUSTOM_STAT_DEFINED(CacheRejectedMB, CRejectedKiB * OneOver1024, ECsvCustomStatOp::Set);
-		
-		// http stat totals
-		auto HttpCsv = [](EHttpRequestType Type)
-			{
-				UE_IAX_CSV_CUSTOM_STAT_DEFINED(Type, HttpGetCount, (int32)GHttpGetCount[(uint8)Type].Get(), ECsvCustomStatOp::Set);
-				UE_IAX_CSV_CUSTOM_STAT_DEFINED(Type, HttpCancelCount, (int32)GHttpCancelCount[(uint8)Type].Get(), ECsvCustomStatOp::Set);
-				UE_IAX_CSV_CUSTOM_STAT_DEFINED(Type, HttpErrorCount, (int32)GHttpErrorCount[(uint8)Type].Get(), ECsvCustomStatOp::Set);
-				UE_IAX_CSV_CUSTOM_STAT_DEFINED(Type, HttpPendingCount, (int32)GHttpPendingCount[(uint8)Type].Get(), ECsvCustomStatOp::Set);
-				UE_IAX_CSV_CUSTOM_STAT_DEFINED(Type, HttpDownloadedMB, (float)GHttpDownloadedBytes[(uint8)Type].Get() * OneOver1024 * OneOver1024, ECsvCustomStatOp::Set);
-				UE_IAX_CSV_CUSTOM_STAT_DEFINED(Type, HttpBandwidthMpbs, (int32)GHttpBandwidthMbps[(uint8)Type].Get(), ECsvCustomStatOp::Set);
-				UE_IAX_CSV_CUSTOM_STAT_DEFINED(Type, HttpDurationMsAvg, (int32)GHttpDurationMsAvg[(uint8)Type], ECsvCustomStatOp::Set);
-				UE_IAX_CSV_CUSTOM_STAT_DEFINED(Type, HttpDurationMsMax, GHttpDurationMsMax[(uint8)Type], ECsvCustomStatOp::Set);
-			};
-
-		HttpCsv(EHttpRequestType::Streaming);
-		HttpCsv(EHttpRequestType::Installed);
-
+		UpdateCSVValues();
 		PrintPeriodicLogging();
 	});
 
@@ -681,6 +646,49 @@ FOnDemandIoBackendStats* FOnDemandIoBackendStats::Get()
 	return GStatistics;
 }
 
+#if IAS_WITH_STATISTICS
+
+void FOnDemandIoBackendStats::UpdateCSVValues()
+{
+	static constexpr float OneOver1024 = 1.0f / 1024.0f;
+
+	// cache stat totals
+	int32 CGetCount = int32(GCacheGetCount.Get());
+	int32 CErrorCount = int32(GCacheErrorCount.Get());
+	int32 CPutCount = int32(GCachePutCount.Get());
+	int32 CPutExistingCount = int32(GCachePutExistingCount.Get());
+	int32 CPutRejectCount = int32(GCachePutRejectCount.Get());
+	float CCachedKiB = (float)GCacheCachedBytes.Get() * OneOver1024;
+	float CWrittenKiB = (float)GCacheWrittenBytes.Get() * OneOver1024;
+	float CReadKiB = (float)GCacheReadBytes.Get() * OneOver1024;
+	float CRejectedKiB = (float)GCacheRejectBytes.Get() * OneOver1024;
+	CSV_CUSTOM_STAT_DEFINED(CacheGetCount, CGetCount, ECsvCustomStatOp::Set);
+	CSV_CUSTOM_STAT_DEFINED(CacheErrorCount, CErrorCount, ECsvCustomStatOp::Set);
+	CSV_CUSTOM_STAT_DEFINED(CachePutCount, CPutCount, ECsvCustomStatOp::Set);
+	CSV_CUSTOM_STAT_DEFINED(CachePutExistingCount, CPutExistingCount, ECsvCustomStatOp::Set);
+	CSV_CUSTOM_STAT_DEFINED(CachePutRejectCount, CPutRejectCount, ECsvCustomStatOp::Set);
+	CSV_CUSTOM_STAT_DEFINED(CacheCachedMB, CCachedKiB * OneOver1024, ECsvCustomStatOp::Set);
+	CSV_CUSTOM_STAT_DEFINED(CacheWrittenMB, CWrittenKiB * OneOver1024, ECsvCustomStatOp::Set);
+	CSV_CUSTOM_STAT_DEFINED(CacheReadMB, CReadKiB * OneOver1024, ECsvCustomStatOp::Set);
+	CSV_CUSTOM_STAT_DEFINED(CacheRejectedMB, CRejectedKiB * OneOver1024, ECsvCustomStatOp::Set);
+
+	// http stat totals
+	auto HttpCsv = [](EHttpRequestType Type)
+		{
+			UE_IAX_CSV_CUSTOM_STAT_DEFINED(Type, HttpGetCount, (int32)GHttpGetCount[(uint8)Type].Get(), ECsvCustomStatOp::Set);
+			UE_IAX_CSV_CUSTOM_STAT_DEFINED(Type, HttpCancelCount, (int32)GHttpCancelCount[(uint8)Type].Get(), ECsvCustomStatOp::Set);
+			UE_IAX_CSV_CUSTOM_STAT_DEFINED(Type, HttpErrorCount, (int32)GHttpErrorCount[(uint8)Type].Get(), ECsvCustomStatOp::Set);
+			UE_IAX_CSV_CUSTOM_STAT_DEFINED(Type, HttpPendingCount, (int32)GHttpPendingCount[(uint8)Type].Get(), ECsvCustomStatOp::Set);
+			UE_IAX_CSV_CUSTOM_STAT_DEFINED(Type, HttpDownloadedMB, (float)GHttpDownloadedBytes[(uint8)Type].Get() * OneOver1024 * OneOver1024, ECsvCustomStatOp::Set);
+			UE_IAX_CSV_CUSTOM_STAT_DEFINED(Type, HttpBandwidthMpbs, (int32)GHttpBandwidthMbps[(uint8)Type].Get(), ECsvCustomStatOp::Set);
+			UE_IAX_CSV_CUSTOM_STAT_DEFINED(Type, HttpDurationMsAvg, (int32)GHttpDurationMsAvg[(uint8)Type], ECsvCustomStatOp::Set);
+			UE_IAX_CSV_CUSTOM_STAT_DEFINED(Type, HttpDurationMsMax, GHttpDurationMsMax[(uint8)Type], ECsvCustomStatOp::Set);
+		};
+
+	HttpCsv(EHttpRequestType::Streaming);
+	HttpCsv(EHttpRequestType::Installed);
+}
+
 void FOnDemandIoBackendStats::PrintPeriodicLogging()
 {
 	if (GIasStatisticsLogInterval > 0.f)
@@ -744,6 +752,8 @@ void FOnDemandIoBackendStats::PrintPeriodicLogging()
 	}
 }
 
+#endif // IAS_WITH_STATISTICS
+
 void FOnDemandIoBackendStats::ReportGeneralAnalytics(TArray<FAnalyticsEventAttribute>& OutAnalyticsArray) const
 {
 	// Note that this analytics section is not optional, if we are reporting analytics then we report this section
diff --git a/Engine/Source/Runtime/Experimental/IoStore/OnDemand/Private/Statistics.h b/Engine/Source/Runtime/Experimental/IoStore/OnDemand/Private/Statistics.h
--- a/Engine/Source/Runtime/Experimental/IoStore/OnDemand/Private/Statistics.h
+++ b/Engine/Source/Runtime/Experimental/IoStore/OnDemand/Private/Statistics.h
@@ -84,9 +84,11 @@ public:
 
 private:
 
+#if IAS_WITH_STATISTICS
+
+	void UpdateCSVValues();
 	void PrintPeriodicLogging();
 
-#if IAS_WITH_STATISTICS
 	const FBackendStatus& BackendStatus;
 #endif //IAS_WITH_STATISTICS
 };
