Get cache usage from Stats instead of asking the cache directly.

#rnx

diff --git a/Engine/Source/Runtime/Experimental/IoStore/OnDemand/Private/IasCache.h b/Engine/Source/Runtime/Experimental/IoStore/OnDemand/Private/IasCache.h
--- a/Engine/Source/Runtime/Experimental/IoStore/OnDemand/Private/IasCache.h
+++ b/Engine/Source/Runtime/Experimental/IoStore/OnDemand/Private/IasCache.h
@@ -49,9 +49,6 @@ public:
 	/** Insert a new chunk into the cache. */
 	virtual FIoStatus Put(const FIoHash& Key, FIoBuffer& Data) = 0;
 
-	/** Get cache usage information. */
-	virtual void GetCacheUsage(uint64& OutDiskSize, uint64& OutMaxDiskSize) const = 0;
-
 	/** Evict an item from the cache. For example it might be corrupt */
 	virtual FIoStatus Evict(const FIoHash& Key) = 0;
 };
diff --git a/Engine/Source/Runtime/Experimental/IoStore/OnDemand/Private/JournaledCache.cpp b/Engine/Source/Runtime/Experimental/IoStore/OnDemand/Private/JournaledCache.cpp
--- a/Engine/Source/Runtime/Experimental/IoStore/OnDemand/Private/JournaledCache.cpp
+++ b/Engine/Source/Runtime/Experimental/IoStore/OnDemand/Private/JournaledCache.cpp
@@ -621,7 +621,6 @@ public:
 	void					Drop();
 	uint64					RemainingUntilWrap();
 	uint32					DebugVisit(void* Param, FDebugCacheEntry::Callback* Callback);
-	void					Usage(uint64& OutMappedBytes, uint64& OutMaxSize);
 
 private:
 	struct FMapEntry
@@ -987,15 +986,6 @@ uint32 FDiskCache::DebugVisit(void* Param, FDebugCacheEntry::Callback* Callback)
 	return DataMap.Num();
 }
 
-////////////////////////////////////////////////////////////////////////////////
-void FDiskCache::Usage(uint64& OutMappedBytes, uint64& OutMaxSize)
-{
-	FReadScopeLock _(Lock);
-
-	OutMappedBytes	= MappedBytes;
-	OutMaxSize		= MaxDataSize;
-}
-
 // {{{1 loader .................................................................
 
 ////////////////////////////////////////////////////////////////////////////////
@@ -1300,7 +1290,6 @@ public:
 	uint32			Flush();
 	uint32			WriteMemToDisk(int32 Allowance);
 	uint32			DebugVisit(void* Param, FDebugCacheEntry::Callback* Callback);
-	void			GetDiskUsage(uint64& OutDiskUsage, uint64& OutMaxDiskSize);
 
 private:
 	struct FPartial
@@ -1584,11 +1573,6 @@ uint32 FCache::DebugVisit(void* Param, FDebugCacheEntry::Callback* Callback)
 	return Count;
 }
 
-////////////////////////////////////////////////////////////////////////////////
-void FCache::GetDiskUsage(uint64& OutDiskUsage, uint64& OutMaxDiskSize)
-{
-	DiskCache.Usage(OutDiskUsage, OutMaxDiskSize);
-}
 
 
 // {{{1 governor ...............................................................
@@ -2262,7 +2246,6 @@ public:
 	virtual void				Materialize(const FIoHash& Key, FIoBuffer& Dest, EIoErrorCode& Status, UE::Tasks::FTaskEvent DoneEvent) override;
 	virtual void				Cancel(FIoBuffer& GivenDest) override;
 	virtual FIoStatus			Put(const FIoHash& Key, FIoBuffer& Data) override;
-	virtual void				GetCacheUsage(uint64& OutDiskSize, uint64& OutMaxDiskSize) const override;
 	virtual FIoStatus			Evict(const FIoHash& Key) override;
 
 private:
@@ -2432,12 +2415,6 @@ FIoStatus FJournaledCache::Put(const FIoHash& Key, FIoBuffer& Data)
 	return Ok ? FIoStatus::Ok : FIoStatus(EIoErrorCode::Unknown);
 }
 
-////////////////////////////////////////////////////////////////////////////////
-void FJournaledCache::GetCacheUsage(uint64& OutDiskSize, uint64& OutMaxDiskSize) const
-{
-	Cache->GetDiskUsage(OutDiskSize, OutMaxDiskSize);
-}
-
 ////////////////////////////////////////////////////////////////////////////////
 FIoStatus FJournaledCache::Evict(const FIoHash& Key)
 {
diff --git a/Engine/Source/Runtime/Experimental/IoStore/OnDemand/Private/OnDemandIoDispatcherBackend.cpp b/Engine/Source/Runtime/Experimental/IoStore/OnDemand/Private/OnDemandIoDispatcherBackend.cpp
--- a/Engine/Source/Runtime/Experimental/IoStore/OnDemand/Private/OnDemandIoDispatcherBackend.cpp
+++ b/Engine/Source/Runtime/Experimental/IoStore/OnDemand/Private/OnDemandIoDispatcherBackend.cpp
@@ -1360,7 +1360,7 @@ FOnDemandStreamingCacheUsage FOnDemandIoBackend::GetCacheUsage() const
 	FOnDemandStreamingCacheUsage Usage;
 	if (Cache.IsValid())
 	{
-		Cache->GetCacheUsage(Usage.TotalSize, Usage.MaxSize);
+		Stats.GetIasCacheStats(Usage.TotalSize, Usage.MaxSize);
 	}
 
 	return Usage; 
diff --git a/Engine/Source/Runtime/Experimental/IoStore/OnDemand/Private/Statistics.cpp b/Engine/Source/Runtime/Experimental/IoStore/OnDemand/Private/Statistics.cpp
--- a/Engine/Source/Runtime/Experimental/IoStore/OnDemand/Private/Statistics.cpp
+++ b/Engine/Source/Runtime/Experimental/IoStore/OnDemand/Private/Statistics.cpp
@@ -1011,6 +1011,12 @@ TUniquePtr<IAnalyticsRecording> FOnDemandIoBackendStats::StartAnalyticsRecording
 #endif // UE_ENABLE_ANALYTICS_RECORDING
 }
 
+void FOnDemandIoBackendStats::GetIasCacheStats(uint64& OutUsed, uint64& OutMaxSize) const
+{
+	OutUsed = GCacheCachedBytes.Get();
+	OutMaxSize = GCacheMaxBytes;
+}
+
 void FOnDemandIoBackendStats::OnIoRequestEnqueue()
 {
 	GIoRequestCount.Increment();
diff --git a/Engine/Source/Runtime/Experimental/IoStore/OnDemand/Private/Statistics.h b/Engine/Source/Runtime/Experimental/IoStore/OnDemand/Private/Statistics.h
--- a/Engine/Source/Runtime/Experimental/IoStore/OnDemand/Private/Statistics.h
+++ b/Engine/Source/Runtime/Experimental/IoStore/OnDemand/Private/Statistics.h
@@ -46,6 +46,7 @@ public:
 	void ReportEndPointAnalytics(TArray<FAnalyticsEventAttribute>& OutAnalyticsArray) const IAS_STATISTICS_IMPL()
 
 	TUniquePtr<IAnalyticsRecording> StartAnalyticsRecording() const IAS_STATISTICS_IMPL(TUniquePtr<IAnalyticsRecording>())
+	void GetIasCacheStats(uint64& OutUsed, uint64& OutMaxSize) const IAS_STATISTICS_IMPL(void(OutUsed = OutMaxSize = 0);) 
 
 	void OnIoRequestEnqueue() IAS_STATISTICS_IMPL()
 	void OnIoRequestComplete(uint64 Size, uint64 DurationMs) IAS_STATISTICS_IMPL()
