Move the IAX on screen statistics code to it's own method.

#rnx

diff --git a/Engine/Source/Runtime/Experimental/IoStore/OnDemand/Private/Statistics.cpp b/Engine/Source/Runtime/Experimental/IoStore/OnDemand/Private/Statistics.cpp
--- a/Engine/Source/Runtime/Experimental/IoStore/OnDemand/Private/Statistics.cpp
+++ b/Engine/Source/Runtime/Experimental/IoStore/OnDemand/Private/Statistics.cpp
@@ -499,137 +499,14 @@ FOnDemandIoBackendStats::FOnDemandIoBackendStats(FBackendStatus& InStatus)
 
 #if UE_ENABLE_ONSCREEN_STATISTICS
 	GStatisticsOnScreenDelegateHandle = FCoreDelegates::OnGetOnScreenMessages.AddLambda(
-		[this] (FCoreDelegates::FSeverityMessageMap& Out)
+		[this] (FCoreDelegates::FSeverityMessageMap& OutMessages)
 		{
 			if (!GIasDisplayOnScreenStatistics)
 			{
 				return;
 			}
 
-			{
-				TStringBuilder<256> BackendStatusText;
-				BackendStatusText << TEXT("IAS Backend Status: ");
-				BackendStatus.ToString(BackendStatusText);
-
-				Out.Add(FCoreDelegates::EOnScreenMessageSeverity::Info, FText::FromStringView(BackendStatusText.ToView()));
-			}
-
-			FHostGroupManager::Get().ForEachHostGroup([&Out](const FIASHostGroup& HostGroup)
-			{
-				FCoreDelegates::EOnScreenMessageSeverity Verbosity = FCoreDelegates::EOnScreenMessageSeverity::Info;
-				TStringBuilder<256> Text;
-				Text << TEXT("IAS HostGroup [") << HostGroup.GetName() << TEXT("] ");
-
-				if (HostGroup.IsConnected())
-				{
-					Text << HostGroup.GetPrimaryHostUrl();
-					Text << TEXT(" (") << HostGroup.GetPrimaryHostIndex() << TEXT("/") << HostGroup.GetHostUrls().Num() << TEXT(")");
-				}
-				else if (HostGroup.IsResolved())
-				{
-					Text << TEXT("Resolving...");
-				}
-				else
-				{
-					Text << TEXT("Disconnected");
-					Verbosity = FCoreDelegates::EOnScreenMessageSeverity::Error;
-				}
-
-				Out.Add(Verbosity, FText::FromStringView(Text.ToView()));
-			});
-
-			{
-				auto HttpStats = [&Out](const TCHAR* Title, EHttpRequestType Type)
-				{
-					TStringBuilder<256> Builder;
-					Builder.Appendf(
-						TEXT("%s Http: Downloaded: %s (%d) Avg %d ms | Retries: %d | Pending: %d"),
-						Title,
-						*FText::AsMemory(GHttpDownloadedBytes[(uint8)Type].Get()).ToString(),
-						GHttpGetCount[(uint8)Type].Get(),
-						(int32)GHttpDurationMsAvg[(uint8)Type],
-						GHttpRetryCount[(uint8)Type].Get(),
-						GHttpPendingCount[(uint8)Type].Get()
-					);
-
-					Out.Add(FCoreDelegates::EOnScreenMessageSeverity::Info, FText::FromStringView(Builder.ToView()));
-				};
-
-				{
-					TStringBuilder<256> Builder;
-					Builder << TEXT("IAX Http:");
-					if (GIasHttpRateLimitKiBPerSecond > 0)
-					{
-						Builder << TEXT(" DataRate Cap: ") << GIasHttpRateLimitKiBPerSecond << TEXT("KiB/s");
-					}
-					else
-					{
-						Builder << TEXT(" DataRate Cap: None");
-					}
-
-					const int32 MaxNumRequests = 8;
-					Builder << TEXT(" ConcurrentRequests: ") << GHttpNumConcurrentRequests << TEXT("/") << MaxNumRequests;
-					if (GHttpNumConcurrentRequests >= MaxNumRequests)
-					{
-						Builder << TEXT( "(saturated)");
-					}
-
-					Out.Add(FCoreDelegates::EOnScreenMessageSeverity::Info, FText::FromStringView(Builder.ToView()));
-				}
-
-				HttpStats(TEXT("IAS"), EHttpRequestType::Streaming);
-				HttpStats(TEXT("IAD"), EHttpRequestType::Installed);
-			}
-
-			{
-				TStringBuilder<256> Builder;
-				Builder << TEXT("IAS CDN: Hit/Miss/NoHdr: ");
-
-				Builder << GHttpCdnCacheHit;
-				Builder << TEXT("/") << GHttpCdnCacheMiss;
-				Builder << TEXT("/") << GHttpCdnCacheUnknown;
-
-				if (uint32 Total = GHttpCdnCacheHit + GHttpCdnCacheMiss + GHttpCdnCacheUnknown; Total)
-				{
-					auto AsPercent = [Total] (uint32 Value) { return (Value * 100 + (Total >> 1)) / Total; };
-					Builder << TEXT(" - ") << AsPercent(GHttpCdnCacheHit);
-					Builder << TEXT("%/") << AsPercent(GHttpCdnCacheMiss);
-					Builder << TEXT("%/") << AsPercent(GHttpCdnCacheUnknown);
-					Builder << TEXT("%");
-				}
-
-				Out.Add(FCoreDelegates::EOnScreenMessageSeverity::Info, FText::FromStringView(Builder.ToView()));
-			}
-
-			{
-				TStringBuilder<256> CachingText;
-				CachingText << TEXT("IAS Cache Stats: ");
-				if (BackendStatus.IsCacheEnabled())
-				{
-					CachingText << TEXT("Cached: ") << FText::AsMemory(GCacheCachedBytes.Get()).ToString();
-					CachingText << TEXT(" | Rejected: ") << FText::AsMemory(GCacheRejectBytes.Get()).ToString();
-					CachingText << TEXT(" | Read: ") << FText::AsMemory(GCacheReadBytes.Get()).ToString();
-					CachingText << TEXT(" (") << GCacheGetCount.Get() << TEXT(")");
-				}
-				else
-				{
-					CachingText << TEXTVIEW("Caching Disabled");
-				}
-
-				Out.Add(FCoreDelegates::EOnScreenMessageSeverity::Info, FText::FromStringView(CachingText.ToView()));
-			}
-
-			if (GHttpDecodeErrorCount[(uint8)EHttpRequestType::Streaming].Get() > 0 || GCacheDecodeErrorCount.Get() > 0 || GHttpErrorCount[(uint8)EHttpRequestType::Streaming].Get() > 0)
-			{
-				TStringBuilder<256> Builder;
-				Builder.Appendf(TEXT("IAS Errors: Cache Decode: %d | Http Decode: %d | Http: %d"),
-					GCacheDecodeErrorCount.Get(),
-					GHttpDecodeErrorCount[(uint8)EHttpRequestType::Streaming].Get(),
-					GHttpErrorCount[(uint8)EHttpRequestType::Streaming].Get()
-				);
-
-				Out.Add(FCoreDelegates::EOnScreenMessageSeverity::Error, FText::FromStringView(Builder.ToView()));
-			}
+			PrintOnScreenStatistics(OutMessages);
 		});
 #endif // UE_ENABLE_ONSCREEN_STATISTICS
 }
@@ -752,6 +629,145 @@ void FOnDemandIoBackendStats::PrintPeriodicLogging()
 	}
 }
 
+
+void FOnDemandIoBackendStats::PrintOnScreenStatistics(FCoreDelegates::FSeverityMessageMap& OutMessages)
+{
+#if UE_ENABLE_ONSCREEN_STATISTICS
+
+	// Print the backend status
+	{
+		TStringBuilder<256> BackendStatusText;
+		BackendStatusText << TEXT("IAS Backend Status: ");
+		BackendStatus.ToString(BackendStatusText);
+
+		OutMessages.Add(FCoreDelegates::EOnScreenMessageSeverity::Info, FText::FromStringView(BackendStatusText.ToView()));
+	}
+
+	// Print HostGroup info
+	FHostGroupManager::Get().ForEachHostGroup([&OutMessages](const FIASHostGroup& HostGroup)
+		{
+			FCoreDelegates::EOnScreenMessageSeverity Verbosity = FCoreDelegates::EOnScreenMessageSeverity::Info;
+			TStringBuilder<256> Text;
+			Text << TEXT("IAS HostGroup [") << HostGroup.GetName() << TEXT("] ");
+
+			if (HostGroup.IsConnected())
+			{
+				Text << HostGroup.GetPrimaryHostUrl();
+				Text << TEXT(" (") << HostGroup.GetPrimaryHostIndex() << TEXT("/") << HostGroup.GetHostUrls().Num() << TEXT(")");
+			}
+			else if (HostGroup.IsResolved())
+			{
+				Text << TEXT("Resolving...");
+			}
+			else
+			{
+				Text << TEXT("Disconnected");
+				Verbosity = FCoreDelegates::EOnScreenMessageSeverity::Error;
+			}
+
+			OutMessages.Add(Verbosity, FText::FromStringView(Text.ToView()));
+		});
+
+	// Print Http Stats
+	{
+		auto HttpStats = [&OutMessages](const TCHAR* Title, EHttpRequestType Type)
+			{
+				TStringBuilder<256> Builder;
+				Builder.Appendf(
+					TEXT("%s Http: Downloaded: %s (%d) Avg %d ms | Retries: %d | Pending: %d"),
+					Title,
+					*FText::AsMemory(GHttpDownloadedBytes[(uint8)Type].Get()).ToString(),
+					GHttpGetCount[(uint8)Type].Get(),
+					(int32)GHttpDurationMsAvg[(uint8)Type],
+					GHttpRetryCount[(uint8)Type].Get(),
+					GHttpPendingCount[(uint8)Type].Get()
+				);
+
+				OutMessages.Add(FCoreDelegates::EOnScreenMessageSeverity::Info, FText::FromStringView(Builder.ToView()));
+			};
+
+		{
+			TStringBuilder<256> Builder;
+			Builder << TEXT("IAX Http:");
+			if (GIasHttpRateLimitKiBPerSecond > 0)
+			{
+				Builder << TEXT(" DataRate Cap: ") << GIasHttpRateLimitKiBPerSecond << TEXT("KiB/s");
+			}
+			else
+			{
+				Builder << TEXT(" DataRate Cap: None");
+			}
+
+			const int32 MaxNumRequests = 8;
+			Builder << TEXT(" ConcurrentRequests: ") << GHttpNumConcurrentRequests << TEXT("/") << MaxNumRequests;
+			if (GHttpNumConcurrentRequests >= MaxNumRequests)
+			{
+				Builder << TEXT("(saturated)");
+			}
+
+			OutMessages.Add(FCoreDelegates::EOnScreenMessageSeverity::Info, FText::FromStringView(Builder.ToView()));
+		}
+
+		HttpStats(TEXT("IAS"), EHttpRequestType::Streaming);
+		HttpStats(TEXT("IAD"), EHttpRequestType::Installed);
+	}
+
+	// Print IAX CDN Cache rates
+	{
+		TStringBuilder<256> Builder;
+		Builder << TEXT("IAX CDN: Hit/Miss/NoHdr: ");
+
+		Builder << GHttpCdnCacheHit;
+		Builder << TEXT("/") << GHttpCdnCacheMiss;
+		Builder << TEXT("/") << GHttpCdnCacheUnknown;
+
+		if (uint32 Total = GHttpCdnCacheHit + GHttpCdnCacheMiss + GHttpCdnCacheUnknown; Total)
+		{
+			auto AsPercent = [Total](uint32 Value) { return (Value * 100 + (Total >> 1)) / Total; };
+			Builder << TEXT(" - ") << AsPercent(GHttpCdnCacheHit);
+			Builder << TEXT("%/") << AsPercent(GHttpCdnCacheMiss);
+			Builder << TEXT("%/") << AsPercent(GHttpCdnCacheUnknown);
+			Builder << TEXT("%");
+		}
+
+		OutMessages.Add(FCoreDelegates::EOnScreenMessageSeverity::Info, FText::FromStringView(Builder.ToView()));
+	}
+
+	// Print IAS Cache
+	{
+		TStringBuilder<256> CachingText;
+		CachingText << TEXT("IAS Cache Stats: ");
+		if (BackendStatus.IsCacheEnabled())
+		{
+			CachingText << TEXT("Cached: ") << FText::AsMemory(GCacheCachedBytes.Get()).ToString();
+			CachingText << TEXT(" | Rejected: ") << FText::AsMemory(GCacheRejectBytes.Get()).ToString();
+			CachingText << TEXT(" | Read: ") << FText::AsMemory(GCacheReadBytes.Get()).ToString();
+			CachingText << TEXT(" (") << GCacheGetCount.Get() << TEXT(")");
+		}
+		else
+		{
+			CachingText << TEXTVIEW("Caching Disabled");
+		}
+
+		OutMessages.Add(FCoreDelegates::EOnScreenMessageSeverity::Info, FText::FromStringView(CachingText.ToView()));
+	}
+
+	// Print IAS errors
+	if (GHttpDecodeErrorCount[(uint8)EHttpRequestType::Streaming].Get() > 0 || GCacheDecodeErrorCount.Get() > 0 || GHttpErrorCount[(uint8)EHttpRequestType::Streaming].Get() > 0)
+	{
+		TStringBuilder<256> Builder;
+		Builder.Appendf(TEXT("IAS Errors: Cache Decode: %d | Http Decode: %d | Http: %d"),
+			GCacheDecodeErrorCount.Get(),
+			GHttpDecodeErrorCount[(uint8)EHttpRequestType::Streaming].Get(),
+			GHttpErrorCount[(uint8)EHttpRequestType::Streaming].Get()
+		);
+
+		OutMessages.Add(FCoreDelegates::EOnScreenMessageSeverity::Error, FText::FromStringView(Builder.ToView()));
+	}
+
+#endif // UE_ENABLE_ONSCREEN_STATISTICS
+}
+
 #endif // IAS_WITH_STATISTICS
 
 void FOnDemandIoBackendStats::ReportGeneralAnalytics(TArray<FAnalyticsEventAttribute>& OutAnalyticsArray) const
@@ -768,8 +784,6 @@ void FOnDemandIoBackendStats::ReportGeneralAnalytics(TArray<FAnalyticsEventAttri
 
 void FOnDemandIoBackendStats::ReportEndPointAnalytics(TArray<FAnalyticsEventAttribute>& OutAnalyticsArray) const
 {
-
-
 	if (GIasReportHttpAnalyticsEnabled)
 	{
 		auto ReportHttpStats = [&OutAnalyticsArray](const TCHAR* Prefix, EHttpRequestType Type)
diff --git a/Engine/Source/Runtime/Experimental/IoStore/OnDemand/Private/Statistics.h b/Engine/Source/Runtime/Experimental/IoStore/OnDemand/Private/Statistics.h
--- a/Engine/Source/Runtime/Experimental/IoStore/OnDemand/Private/Statistics.h
+++ b/Engine/Source/Runtime/Experimental/IoStore/OnDemand/Private/Statistics.h
@@ -3,6 +3,7 @@
 #pragma once
 
 #include "HAL/LowLevelMemTracker.h"
+#include "Misc/CoreDelegates.h"
 #include "OnDemandBackendStatus.h"
 #include "ProfilingDebugging/CountersTrace.h"
 #include "ProfilingDebugging/CsvProfiler.h"
@@ -90,6 +91,8 @@ private:
 	void UpdateCSVValues();
 	void PrintPeriodicLogging();
 
+	void PrintOnScreenStatistics(FCoreDelegates::FSeverityMessageMap& OutMessages);
+
 	const FBackendStatus& BackendStatus;
 #endif //IAS_WITH_STATISTICS
 };
