Dataflow : fix core Select node causing an assert when adding a pin while one of the input is connected
- Change set concrete type to also take an type dependency group to allow it to be set even the concrete type is set on teh connection
- Update the Select node accordingly to make sure the type dependency are properly set

diff --git a/Engine/Source/Runtime/Experimental/Dataflow/Core/Private/Dataflow/DataflowConnection.cpp b/Engine/Source/Runtime/Experimental/Dataflow/Core/Private/Dataflow/DataflowConnection.cpp
--- a/Engine/Source/Runtime/Experimental/Dataflow/Core/Private/Dataflow/DataflowConnection.cpp
+++ b/Engine/Source/Runtime/Experimental/Dataflow/Core/Private/Dataflow/DataflowConnection.cpp
@@ -145,6 +145,14 @@ bool FDataflowConnection::SupportsType(FName InType) const
 	return (InType == GetType());
 }
 
+void FDataflowConnection::ForceTypeDependencyGroup(FName InTypeDependencyGroup)
+{
+	if (IsAnyType())
+	{
+		TypeDependencyGroup = InTypeDependencyGroup;
+	}
+}
+
 FDataflowConnection& FDataflowConnection::SetTypeDependencyGroup(FName DependencyGroupName)
 {
 	if (IsAnyType() && !HasConcreteType())
diff --git a/Engine/Source/Runtime/Experimental/Dataflow/Core/Private/Dataflow/DataflowCoreNodes.cpp b/Engine/Source/Runtime/Experimental/Dataflow/Core/Private/Dataflow/DataflowCoreNodes.cpp
--- a/Engine/Source/Runtime/Experimental/Dataflow/Core/Private/Dataflow/DataflowCoreNodes.cpp
+++ b/Engine/Source/Runtime/Experimental/Dataflow/Core/Private/Dataflow/DataflowCoreNodes.cpp
@@ -107,13 +107,16 @@ TArray<UE::Dataflow::FPin> FDataflowSelectNode::AddPins()
 {
 	const int32 Index = Inputs.AddDefaulted();
 	FDataflowInput& Input = RegisterInputArrayConnection(GetConnectionReference(Index));
-	Input.SetTypeDependencyGroup(MainTypeGroup);
 	if (Index > 0)
 	{
 		// Set concrete type the same as Index0.
 		FDataflowInput* const Input0 = FindInput(GetConnectionReference(0));
 		check(Input0);
-		SetInputConcreteType(GetConnectionReference(Index), Input0->GetType());
+		SetConnectionConcreteType(&Input, Input0->GetType(), MainTypeGroup);
+	}
+	else
+	{
+		Input.SetTypeDependencyGroup(MainTypeGroup);
 	}
 	return { { UE::Dataflow::FPin::EDirection::INPUT, Input.GetType(), Input.GetName() } };
 }
@@ -156,8 +159,9 @@ void FDataflowSelectNode::PostSerialize(const FArchive& Ar)
 
 		for (int32 Index = NumInitialInputs; Index < Inputs.Num(); ++Index)
 		{
-			FindOrRegisterInputArrayConnection(GetConnectionReference(Index))
-				.SetTypeDependencyGroup(MainTypeGroup);
+			FDataflowInput& Input = FindOrRegisterInputArrayConnection(GetConnectionReference(Index));
+			// reset the type to allow the type group to be properly set as well 
+			SetConnectionConcreteType(&Input, Input.GetType(), MainTypeGroup);
 		}
 
 		if (Ar.IsTransacting())
diff --git a/Engine/Source/Runtime/Experimental/Dataflow/Core/Private/Dataflow/DataflowNode.cpp b/Engine/Source/Runtime/Experimental/Dataflow/Core/Private/Dataflow/DataflowNode.cpp
--- a/Engine/Source/Runtime/Experimental/Dataflow/Core/Private/Dataflow/DataflowNode.cpp
+++ b/Engine/Source/Runtime/Experimental/Dataflow/Core/Private/Dataflow/DataflowNode.cpp
@@ -1586,7 +1586,7 @@ bool FDataflowNode::SetAllConnectionConcreteType(FName NewType)
 	return bChanged;
 }
 
-bool FDataflowNode::SetConnectionConcreteType(FDataflowConnection* Connection, FName NewType)
+bool FDataflowNode::SetConnectionConcreteType(FDataflowConnection* Connection, FName NewType, FName InTypeDependencyGroup)
 {
 	bool bSuccess = false;
 	if (Connection)
@@ -1620,6 +1620,10 @@ bool FDataflowNode::SetConnectionConcreteType(FDataflowConnection* Connection, F
 				}
 			}
 		}
+		if (InTypeDependencyGroup != NAME_None && Connection->GetTypeDependencyGroup() == NAME_None)
+		{
+			Connection->ForceTypeDependencyGroup(InTypeDependencyGroup);
+		}
 	}
 	return bSuccess;
 }
diff --git a/Engine/Source/Runtime/Experimental/Dataflow/Core/Public/Dataflow/DataflowConnection.h b/Engine/Source/Runtime/Experimental/Dataflow/Core/Public/Dataflow/DataflowConnection.h
--- a/Engine/Source/Runtime/Experimental/Dataflow/Core/Public/Dataflow/DataflowConnection.h
+++ b/Engine/Source/Runtime/Experimental/Dataflow/Core/Public/Dataflow/DataflowConnection.h
@@ -189,6 +189,8 @@ protected:
 	/* Rename the connection */
 	DATAFLOWCORE_API void Rename(FName NewName);
 
+	DATAFLOWCORE_API void ForceTypeDependencyGroup(FName InTypeDependencyGroup);
+
 public:
 	FDataflowConnection() {};
 	UE_DEPRECATED(5.5, "Deprecated constructor : use FConnectionParameters to pass parameters")
diff --git a/Engine/Source/Runtime/Experimental/Dataflow/Core/Public/Dataflow/DataflowNode.h b/Engine/Source/Runtime/Experimental/Dataflow/Core/Public/Dataflow/DataflowNode.h
--- a/Engine/Source/Runtime/Experimental/Dataflow/Core/Public/Dataflow/DataflowNode.h
+++ b/Engine/Source/Runtime/Experimental/Dataflow/Core/Public/Dataflow/DataflowNode.h
@@ -814,7 +814,13 @@ protected:
 	* */
 	DATAFLOWCORE_API bool SetConnectionConcreteType(const UE::Dataflow::FConnectionKey& ConnectionKey, FName NewType);
 
-	DATAFLOWCORE_API bool SetConnectionConcreteType(FDataflowConnection* Connection, FName NewType);
+	/**
+	* Set a connection to be of a specific type or optional set its type dependency group
+	* If the connection is already connected, the method prevents the connection from changing and returns false
+	* The type dependency group can only be set on anytype connections and only if the connecion does not have already a type depencency 
+	* returns true if the connection has successfully changed type 
+	* */
+	DATAFLOWCORE_API bool SetConnectionConcreteType(FDataflowConnection* Connection, FName NewType, FName InTypeDependencyGroup = NAME_None);
 
 	// returns true if any of the types was changed successfully
 	DATAFLOWCORE_API bool SetAllConnectionConcreteType(FName NewType);
