Dataflow : Fix potential crash in when changing properties in a asset using Dataflow

diff --git a/Engine/Source/Runtime/Experimental/Dataflow/Engine/Private/Dataflow/DataflowInstance.cpp b/Engine/Source/Runtime/Experimental/Dataflow/Engine/Private/Dataflow/DataflowInstance.cpp
--- a/Engine/Source/Runtime/Experimental/Dataflow/Engine/Private/Dataflow/DataflowInstance.cpp
+++ b/Engine/Source/Runtime/Experimental/Dataflow/Engine/Private/Dataflow/DataflowInstance.cpp
@@ -505,7 +505,7 @@ FDataflowInstance::~FDataflowInstance()
 #if WITH_EDITOR
 void FDataflowInstance::OnOwnerPostEditChangeProperty(UObject* InObject, struct FPropertyChangedEvent& InPropertyChangedEvent)
 {
-	if (InObject == Owner)
+	if (Owner != nullptr && Owner == InObject)
 	{
 		VariableOverrides.OnOwnerPostEditChangeProperty(InPropertyChangedEvent);
 
diff --git a/Engine/Source/Runtime/Experimental/Dataflow/Engine/Public/Dataflow/DataflowInstance.h b/Engine/Source/Runtime/Experimental/Dataflow/Engine/Public/Dataflow/DataflowInstance.h
--- a/Engine/Source/Runtime/Experimental/Dataflow/Engine/Public/Dataflow/DataflowInstance.h
+++ b/Engine/Source/Runtime/Experimental/Dataflow/Engine/Public/Dataflow/DataflowInstance.h
@@ -271,7 +271,7 @@ private:
 
 	// Owning Object
 	UPROPERTY()
-	TObjectPtr<UObject> Owner;
+	TObjectPtr<UObject> Owner = nullptr;
 };
 
 UINTERFACE(MinimalAPI)
