Turn off partial island sleeping and default to full island sleeping for the entire island if any particle in the island does not allow partial island sleeping.

diff --git a/Engine/Source/Programs/HeadlessChaos/Private/HeadlessChaosTestGraphEvolution.cpp b/Engine/Source/Programs/HeadlessChaos/Private/HeadlessChaosTestGraphEvolution.cpp
--- a/Engine/Source/Programs/HeadlessChaos/Private/HeadlessChaosTestGraphEvolution.cpp
+++ b/Engine/Source/Programs/HeadlessChaos/Private/HeadlessChaosTestGraphEvolution.cpp
@@ -2306,5 +2306,123 @@ namespace ChaosTest
 		}
 	}
 
+	TEST_P(GraphEvolutionTests, TestConstraintGraph_PartialSleepAllowedParticleFlag)
+	{
+		if (GetParam())
+		{
+			FGraphEvolutionTest Test(4, GetParam());
+			Test.MakeChain();
+
+			// Explicitly put two particles asleep
+			Test.Evolution.SetParticleObjectState(Test.ParticleHandles[0], EObjectStateType::Sleeping);
+			Test.Evolution.SetParticleObjectState(Test.ParticleHandles[1], EObjectStateType::Sleeping);
+
+			Test.Advance();
+
+			// Partial sleeping should work as usual
+			EXPECT_TRUE(FConstGenericParticleHandle(Test.ParticleHandles[0])->Sleeping());
+			EXPECT_TRUE(FConstGenericParticleHandle(Test.ParticleHandles[1])->Sleeping());
+			EXPECT_FALSE(FConstGenericParticleHandle(Test.ParticleHandles[2])->Sleeping());
+			EXPECT_FALSE(FConstGenericParticleHandle(Test.ParticleHandles[3])->Sleeping());
+
+			// Disallow partial sleeping for one of the sleeping particles
+			Test.ParticleHandles[0]->SetPartialIslandSleepAllowed(false);
+
+			Test.Advance();
+
+			// All particles wake up
+			EXPECT_FALSE(FConstGenericParticleHandle(Test.ParticleHandles[0])->Sleeping());
+			EXPECT_FALSE(FConstGenericParticleHandle(Test.ParticleHandles[1])->Sleeping());
+			EXPECT_FALSE(FConstGenericParticleHandle(Test.ParticleHandles[2])->Sleeping());
+			EXPECT_FALSE(FConstGenericParticleHandle(Test.ParticleHandles[3])->Sleeping());
+
+			// Explicitly put two particles asleep
+			Test.Evolution.SetParticleObjectState(Test.ParticleHandles[0], EObjectStateType::Sleeping);
+			Test.Evolution.SetParticleObjectState(Test.ParticleHandles[1], EObjectStateType::Sleeping);
+
+			// Allow partial sleeping again of the particles
+			Test.ParticleHandles[0]->SetPartialIslandSleepAllowed(true);
+
+			Test.Advance();
+
+			// Partial sleeping should work as usual
+			EXPECT_TRUE(FConstGenericParticleHandle(Test.ParticleHandles[0])->Sleeping());
+			EXPECT_TRUE(FConstGenericParticleHandle(Test.ParticleHandles[1])->Sleeping());
+			EXPECT_FALSE(FConstGenericParticleHandle(Test.ParticleHandles[2])->Sleeping());
+			EXPECT_FALSE(FConstGenericParticleHandle(Test.ParticleHandles[3])->Sleeping());
+
+			// Disallow partial sleeping for one of the awake particles
+			Test.ParticleHandles[2]->SetPartialIslandSleepAllowed(false);
+
+			Test.Advance();
+
+			// All particles wake up
+			EXPECT_FALSE(FConstGenericParticleHandle(Test.ParticleHandles[0])->Sleeping());
+			EXPECT_FALSE(FConstGenericParticleHandle(Test.ParticleHandles[1])->Sleeping());
+			EXPECT_FALSE(FConstGenericParticleHandle(Test.ParticleHandles[2])->Sleeping());
+			EXPECT_FALSE(FConstGenericParticleHandle(Test.ParticleHandles[3])->Sleeping());
+
+			Test.AdvanceUntilSleeping();
+
+			// All particles are asleep (using island sleeping)
+			EXPECT_TRUE(FConstGenericParticleHandle(Test.ParticleHandles[0])->Sleeping());
+			EXPECT_TRUE(FConstGenericParticleHandle(Test.ParticleHandles[1])->Sleeping());
+			EXPECT_TRUE(FConstGenericParticleHandle(Test.ParticleHandles[2])->Sleeping());
+			EXPECT_TRUE(FConstGenericParticleHandle(Test.ParticleHandles[3])->Sleeping());
+		}
+	}
+
+	TEST_P(GraphEvolutionTests, TestConstraintGraph_PartialSleepAllowedParticleFlag_2)
+	{
+		if (GetParam())
+		{
+			FGraphEvolutionTest Test(4, GetParam());
+			Test.MakeChain();
+
+			// Explicitly put two particles asleep
+			Test.Evolution.SetParticleObjectState(Test.ParticleHandles[0], EObjectStateType::Sleeping);
+			Test.Evolution.SetParticleObjectState(Test.ParticleHandles[1], EObjectStateType::Sleeping);
+
+			Test.Advance();
+
+			// Partial sleeping should work as usual
+			EXPECT_TRUE(FConstGenericParticleHandle(Test.ParticleHandles[0])->Sleeping());
+			EXPECT_TRUE(FConstGenericParticleHandle(Test.ParticleHandles[1])->Sleeping());
+			EXPECT_FALSE(FConstGenericParticleHandle(Test.ParticleHandles[2])->Sleeping());
+			EXPECT_FALSE(FConstGenericParticleHandle(Test.ParticleHandles[3])->Sleeping());
+
+			// Add another particle which does not allow partial sleep
+			Test.ParticleHandles.Add(Test.Evolution.CreateDynamicParticles(1)[0]);
+			Test.Evolution.EnableParticle(Test.ParticleHandles.Last());
+			Test.ConstraintHandles.Add(Test.Constraints.AddConstraint({ Test.ParticleHandles.Last(1), Test.ParticleHandles.Last() }));
+			Test.ParticleHandles.Last()->SetPartialIslandSleepAllowed(false);
+
+			Test.Advance();
+
+			// All particles wake up
+			EXPECT_FALSE(FConstGenericParticleHandle(Test.ParticleHandles[0])->Sleeping());
+			EXPECT_FALSE(FConstGenericParticleHandle(Test.ParticleHandles[1])->Sleeping());
+			EXPECT_FALSE(FConstGenericParticleHandle(Test.ParticleHandles[2])->Sleeping());
+			EXPECT_FALSE(FConstGenericParticleHandle(Test.ParticleHandles[3])->Sleeping());
+			EXPECT_FALSE(FConstGenericParticleHandle(Test.ParticleHandles[4])->Sleeping());
+
+			// Explicitly put two particles asleep
+			Test.Evolution.SetParticleObjectState(Test.ParticleHandles[0], EObjectStateType::Sleeping);
+			Test.Evolution.SetParticleObjectState(Test.ParticleHandles[1], EObjectStateType::Sleeping);
+
+			// Disconnect the particle which does not allow partial sleep
+			Test.ConstraintHandles.Last()->SetEnabled(false);
+
+			Test.Advance();
+
+			// Partial sleeping should work as usual
+			EXPECT_TRUE(FConstGenericParticleHandle(Test.ParticleHandles[0])->Sleeping());
+			EXPECT_TRUE(FConstGenericParticleHandle(Test.ParticleHandles[1])->Sleeping());
+			EXPECT_FALSE(FConstGenericParticleHandle(Test.ParticleHandles[2])->Sleeping());
+			EXPECT_FALSE(FConstGenericParticleHandle(Test.ParticleHandles[2])->Sleeping());
+			EXPECT_FALSE(FConstGenericParticleHandle(Test.ParticleHandles[4])->Sleeping());
+		}
+	}
+
 	INSTANTIATE_TEST_SUITE_P(, GraphEvolutionTests,	testing::Bool());
 }
diff --git a/Engine/Source/Programs/HeadlessChaos/Private/HeadlessChaosTestSleeping.cpp b/Engine/Source/Programs/HeadlessChaos/Private/HeadlessChaosTestSleeping.cpp
--- a/Engine/Source/Programs/HeadlessChaos/Private/HeadlessChaosTestSleeping.cpp
+++ b/Engine/Source/Programs/HeadlessChaos/Private/HeadlessChaosTestSleeping.cpp
@@ -177,6 +177,18 @@ namespace ChaosTest
 			}
 		}
 
+		void MakeQueueOfCubes(const int32 Num, const FReal Size = 200)
+		{
+			MakeFloor();
+			const FReal InitialHeight = Size / 2.0f;
+			for (int32 Id = 0; Id < Num; ++Id)
+			{
+				const FReal Y = Size * Id;
+				const FReal Z = InitialHeight;
+				MakeCube(FVec3(0, Y, Z), Size);
+			}
+		}
+
 		void MakeStackOfSpheres(const int32 Num, const FReal Radius = 100)
 		{
 			MakeFloor();
@@ -466,7 +478,7 @@ namespace ChaosTest
 			}
 		}
 		else
-		{			
+		{
 			// The stack wakes up only above the level of where the impulse is applied.
 			EXPECT_TRUE(Test.ParticleHandles[0]->ObjectState() == EObjectStateType::Sleeping);
 			EXPECT_TRUE(Test.ParticleHandles[1]->ObjectState() == EObjectStateType::Sleeping);
@@ -657,5 +669,60 @@ namespace ChaosTest
 		}
 	}
 
+	TEST_P(SleepingTests, CollideWithQueueOfCubes)
+	{
+		const int32 NumOfStackedObjects = 5;
+		const FReal Dimension = 200;
+		const bool bPartialSleepEnabled = true;
+		const bool bAllowPartialSleep = GetParam();
+
+		// Create a stack of cubes
+		FSleepingTest Test(bPartialSleepEnabled);
+		Test.MakeQueueOfCubes(NumOfStackedObjects, /*Size=*/Dimension);
+
+		// Simulate until all particles are sleeping
+		Test.AdvanceUntilSleeping();
+
+		// Check that everything is sleeping
+		for (FRigidBodyHandle_External* ParticleHandle : Test.ParticleHandles)
+		{
+			EXPECT_TRUE(ParticleHandle->ObjectState() == EObjectStateType::Sleeping);
+		}
+
+		// Create another particle which will collide with the queue
+		FRigidBodyHandle_External* CollidingCube = Test.MakeCube(FVec3(0, -200, 100), 200);
+		CollidingCube->SetV(FVec3(0, 200, 0), true);
+		CollidingCube->SetPartialIslandSleepAllowed(bAllowPartialSleep);
+
+		Test.Advance();
+
+		//// Debug print of all particle positions
+		//for (int32 Id = 0; Id < Test.ParticleHandles.Num(); ++Id)
+		//{
+		//	FRigidBodyHandle_External* ParticleHandle = Test.ParticleHandles[Id];
+		//	printf("Particle Id: %d, State: %d, Position: %f, %f, %f\n", Id,
+		//		ParticleHandle->ObjectState(), ParticleHandle->GetX().X, ParticleHandle->GetX().Y, ParticleHandle->GetX().Z);
+		//}
+
+		if (!bAllowPartialSleep)
+		{
+			// Check that everything is awake
+			for (FRigidBodyHandle_External* ParticleHandle : Test.ParticleHandles)
+			{
+				EXPECT_TRUE(ParticleHandle->ObjectState() == EObjectStateType::Dynamic);
+			}
+		}
+		else
+		{
+			// Only the colliding cube and first particle in the queue wake up
+			EXPECT_TRUE(Test.ParticleHandles[5]->ObjectState() == EObjectStateType::Dynamic);
+			EXPECT_TRUE(Test.ParticleHandles[0]->ObjectState() == EObjectStateType::Dynamic);
+			EXPECT_TRUE(Test.ParticleHandles[1]->ObjectState() == EObjectStateType::Sleeping);
+			EXPECT_TRUE(Test.ParticleHandles[2]->ObjectState() == EObjectStateType::Sleeping);
+			EXPECT_TRUE(Test.ParticleHandles[3]->ObjectState() == EObjectStateType::Sleeping);
+			EXPECT_TRUE(Test.ParticleHandles[4]->ObjectState() == EObjectStateType::Sleeping);
+		}
+	}
+
 	INSTANTIATE_TEST_SUITE_P(, SleepingTests, testing::Bool());
 }
\ No newline at end of file
diff --git a/Engine/Source/Runtime/Experimental/Chaos/Private/Chaos/Island/IslandManager.cpp b/Engine/Source/Runtime/Experimental/Chaos/Private/Chaos/Island/IslandManager.cpp
--- a/Engine/Source/Runtime/Experimental/Chaos/Private/Chaos/Island/IslandManager.cpp
+++ b/Engine/Source/Runtime/Experimental/Chaos/Private/Chaos/Island/IslandManager.cpp
@@ -466,6 +466,7 @@ namespace Chaos::Private
 	void FPBDIsland::Reuse()
 	{
 		Flags.bIsSleepAllowed = true;
+		Flags.bIsPartialSleepAllowed = true;
 		Flags.bIsSleeping = true;
 	}
 
@@ -810,6 +811,7 @@ namespace Chaos::Private
 			{
 				// This is a dynamic particle and we know its island
 				EnqueueIslandCheckSleep(Node->Island, bIsSleepAllowed);
+				UpdateIslandPartialSleepAllowed(Node->Island, Particle);
 			}
 			else
 			{
@@ -817,6 +819,7 @@ namespace Chaos::Private
 				for (const FPBDIslandConstraint* Edge : Node->Edges)
 				{
 					EnqueueIslandCheckSleep(Edge->Island, bIsSleepAllowed);
+					UpdateIslandPartialSleepAllowed(Edge->Island, Particle);
 				}
 			}
 		}
@@ -1290,6 +1293,8 @@ namespace Chaos::Private
 		Node->Flags.bIsSleeping = bIsSleeping;
 		Node->Flags.bNeedsResim = bNeedsResim;
 
+		UpdateIslandPartialSleepAllowed(Node->Island, Node->Particle);
+
 		// Did we change between dynamic and kinematic?
 		if (bIsDynamic != bWasDynamic)
 		{
@@ -1447,6 +1452,7 @@ namespace Chaos::Private
 
 			const bool bIsSleepAllowed = true;
 			EnqueueIslandCheckSleep(Island, bIsSleepAllowed);
+			UpdateIslandPartialSleepAllowed(Island, Node->Particle);
 		}
 	}
 
@@ -1546,12 +1552,14 @@ namespace Chaos::Private
 		if (Node->Flags.bIsDynamic)
 		{
 			EnqueueIslandCheckSleep(Node->Island, bIsSleepAllowed);
+			UpdateIslandPartialSleepAllowed(Node->Island, Node->Particle);
 		}
 		else
 		{
 			for (FPBDIslandConstraint* Edge : Node->Edges)
 			{
 				EnqueueIslandCheckSleep(Edge->Island, bIsSleepAllowed);
+				UpdateIslandPartialSleepAllowed(Edge->Island, Node->Particle);
 			}
 		}
 	}
@@ -1565,9 +1573,19 @@ namespace Chaos::Private
 		}
 	}
 
+	void FPBDIslandManager::UpdateIslandPartialSleepAllowed(FPBDIsland* Island, const FGeometryParticleHandle* Particle)
+	{
+		// All particles in the island need to allow partial sleep for the island to allow partial sleep.
+		const FPBDRigidParticleHandle* Rigid = Particle->CastToRigidParticle();
+		if (Island != nullptr && Rigid != nullptr)
+		{
+			Island->Flags.bIsPartialSleepAllowed = (Island->Flags.bIsPartialSleepAllowed != 0) && Rigid->PartialIslandSleepAllowed();
+		}
+	}
+
 	FPBDIsland* FPBDIslandManager::MergeNodeIslands(FPBDIslandParticle* Node0, FPBDIslandParticle* Node1)
 	{
-		// NOTE: We only regsiter the islands for merging here. The actual merge takes place later in ProcessMerges.
+		// NOTE: We only register the islands for merging here. The actual merge takes place later in ProcessMerges.
 
 		FPBDIsland* Island0 = (Node0 != nullptr) ? Node0->Island : nullptr;
 		FPBDIsland* Island1 = (Node1 != nullptr) ? Node1->Island : nullptr;
@@ -1825,6 +1843,7 @@ namespace Chaos::Private
 			ParentIsland->Flags.bItemsAdded = true;
 			ParentIsland->Flags.bItemsRemoved = !!ParentIsland->Flags.bItemsRemoved || !!ChildIsland->Flags.bItemsRemoved;
 			ParentIsland->Flags.bIsSleepAllowed = !!ParentIsland->Flags.bIsSleepAllowed && !!ChildIsland->Flags.bIsSleepAllowed;
+			ParentIsland->Flags.bIsPartialSleepAllowed = !!ParentIsland->Flags.bIsPartialSleepAllowed && !!ChildIsland->Flags.bIsPartialSleepAllowed;
 			ParentIsland->Flags.bIsSleeping = !!ParentIsland->Flags.bIsSleeping && !!ChildIsland->Flags.bIsSleeping;
 			ParentIsland->Flags.bWasSleeping = !!ParentIsland->Flags.bWasSleeping || !!ChildIsland->Flags.bWasSleeping;
 			ParentIsland->Flags.bCheckSleep = !!ParentIsland->Flags.bCheckSleep || !!ChildIsland->Flags.bCheckSleep;
@@ -1919,6 +1938,10 @@ namespace Chaos::Private
 		// The next island to add nodes/edges to
 		FPBDIsland* CurrentIsland = Island;
 
+		// Reset partial sleep allowed flag
+		// Will be updated when the island is repopulated (call to EnqueueIslandPartialSleepAllowed) 
+		CurrentIsland->Flags.bIsPartialSleepAllowed = true;
+
 		// Visit all nodes and put connected nodes into an island
 		TArray<FPBDIslandParticle*> NodeQueue;
 		NodeQueue.Reserve(Island->Nodes.Num());
@@ -1940,6 +1963,7 @@ namespace Chaos::Private
 
 				// Add the root node to the island and visit all connected nodes
 				AddNodeToIsland(RootNode, CurrentIsland);
+				UpdateIslandPartialSleepAllowed(CurrentIsland, RootNode->Particle);
 				NodeQueue.Add(RootNode);
 				RootNode->VisitEpoch = VisitEpoch;
 
@@ -1962,6 +1986,7 @@ namespace Chaos::Private
 							if ((EdgeOtherNode != nullptr) && EdgeOtherNode->Flags.bIsDynamic && (EdgeOtherNode->VisitEpoch != VisitEpoch))
 							{
 								AddNodeToIsland(EdgeOtherNode, CurrentIsland);
+								UpdateIslandPartialSleepAllowed(CurrentIsland, EdgeOtherNode->Particle);
 								NodeQueue.Add(EdgeOtherNode);
 								EdgeOtherNode->VisitEpoch = VisitEpoch;
 							}
@@ -2152,7 +2177,7 @@ namespace Chaos::Private
 			// @todo(chaos) : (island sleeping only) we should only need to do this when the state changes... ?
 			if (!Island->Flags.bIsSleeping || (Island->Flags.bIsSleeping != Island->Flags.bWasSleeping) || !!Island->Flags.bCheckSleep)
 			{
-				if (CVars::bChaosSolverPartialIslandSleep && !!Island->Flags.bIsSleepAllowed)
+				if (CVars::bChaosSolverPartialIslandSleep && !!Island->Flags.bIsSleepAllowed && !!Island->Flags.bIsPartialSleepAllowed)
 				{
 					if (!Island->Flags.bIsSleeping)
 					{
@@ -2292,7 +2317,7 @@ namespace Chaos::Private
 		{
 			if (!Island->Flags.bIsSleeping && !!Island->Flags.bIsSleepAllowed && !Island->Flags.bIsUsingCache)
 			{
-				if (!CVars::bChaosSolverPartialIslandSleep)
+				if (!CVars::bChaosSolverPartialIslandSleep || !Island->Flags.bIsPartialSleepAllowed)
 				{
 					// Update the sleep state based on particle movement etc
 					ProcessIslandSleep(Island, Dt);
@@ -2313,6 +2338,7 @@ namespace Chaos::Private
 
 			// This gets set to false again next tick if there are moving kinematics in the island
 			Island->Flags.bIsSleepAllowed = true;
+			Island->Flags.bIsPartialSleepAllowed = true;
 		}
 	}
 
diff --git a/Engine/Source/Runtime/Experimental/Chaos/Public/Chaos/Island/IslandManager.h b/Engine/Source/Runtime/Experimental/Chaos/Public/Chaos/Island/IslandManager.h
--- a/Engine/Source/Runtime/Experimental/Chaos/Public/Chaos/Island/IslandManager.h
+++ b/Engine/Source/Runtime/Experimental/Chaos/Public/Chaos/Island/IslandManager.h
@@ -347,6 +347,7 @@ namespace Chaos::Private
 				uint32 bItemsAdded : 1;
 				uint32 bItemsRemoved : 1;
 				uint32 bIsSleepAllowed : 1;
+				uint32 bIsPartialSleepAllowed : 1;
 				uint32 bIsSleeping : 1;
 				uint32 bWasSleeping : 1;
 				uint32 bCheckSleep : 1;
@@ -553,7 +554,7 @@ namespace Chaos::Private
 		CHAOS_API void UpdateGraphNode(FPBDIslandParticle* Node);
 		CHAOS_API void UpdateGraphEdge(FPBDIslandConstraint* Edge);
 
-		// Island managament
+		// Island management
 		CHAOS_API FPBDIsland* CreateIsland();
 		CHAOS_API void DestroyIsland(FPBDIsland* Island);
 		CHAOS_API void AssignNodeIsland(FPBDIslandParticle* Node);
@@ -565,6 +566,7 @@ namespace Chaos::Private
 		CHAOS_API void RemoveEdgeFromIsland(FPBDIslandConstraint* Edge);
 		CHAOS_API void WakeNodeIslands(const FPBDIslandParticle* Node);
 		CHAOS_API void EnqueueIslandCheckSleep(FPBDIsland* Island, const bool bIsSleepAllowed);
+		CHAOS_API void UpdateIslandPartialSleepAllowed(FPBDIsland* Island, const FGeometryParticleHandle* Particle);
 
 		// Island merge set management
 		CHAOS_API FPBDIsland* MergeNodeIslands(FPBDIslandParticle* Node0, FPBDIslandParticle* Node1);
