Fix wrong debug name in WriteSentinel header

#rnx

diff --git a/Engine/Source/Runtime/Experimental/Iris/Core/Private/Iris/ReplicationSystem/ReplicationWriter.cpp b/Engine/Source/Runtime/Experimental/Iris/Core/Private/Iris/ReplicationSystem/ReplicationWriter.cpp
--- a/Engine/Source/Runtime/Experimental/Iris/Core/Private/Iris/ReplicationSystem/ReplicationWriter.cpp
+++ b/Engine/Source/Runtime/Experimental/Iris/Core/Private/Iris/ReplicationSystem/ReplicationWriter.cpp
@@ -2315,7 +2315,7 @@ FReplicationWriter::EWriteObjectStatus FReplicationWriter::WriteObjectAndSubObje
 		constexpr bool bIsDestructionInfo = false;
 		Writer.WriteBool(bIsDestructionInfo);
 
-		WriteSentinel(&Writer, TEXT("DestructionInfo"));
+		WriteSentinel(&Writer, TEXT("RootObjectBatchInfo"));
 
 		// A batch starts with (RefHandleId | BatchSize | bHasBatchObjectData | bHasExports)
 		// We write the header up front, and then we seek back and update relevant info is the object + subobjects is successfully serialized along with necessary exports
