Add BP function to reset cloth rest lengths from current Anim Mesh.

Currently only FPBDEdgeSpringConstraints and FPBDBendingConstraints support this

diff --git a/Engine/Plugins/ChaosCloth/Source/ChaosCloth/Private/ChaosCloth/ChaosClothConstraints.cpp b/Engine/Plugins/ChaosCloth/Source/ChaosCloth/Private/ChaosCloth/ChaosClothConstraints.cpp
--- a/Engine/Plugins/ChaosCloth/Source/ChaosCloth/Private/ChaosCloth/ChaosClothConstraints.cpp
+++ b/Engine/Plugins/ChaosCloth/Source/ChaosCloth/Private/ChaosCloth/ChaosClothConstraints.cpp
@@ -2608,4 +2608,21 @@ void FClothConstraints::Update(
 		SkinnedTriangleCollisionsConstraint->SetProperties(ConfigProperties, WeightMaps);
 	}
 }
+
+void FClothConstraints::ResetRestLengths(
+	const TConstArrayView<Softs::FSolverVec3>& NewRestLengthPositions,
+	const Softs::FCollectionPropertyConstFacade& ConfigProperties,
+	const TMap<FString, TConstArrayView<FRealSingle>>& WeightMaps
+)
+{
+	if (EdgeConstraints)
+	{
+		EdgeConstraints->ResetRestLengths(NewRestLengthPositions);
+	}
+
+	if (BendingElementConstraints)
+	{
+		BendingElementConstraints->ResetRestLengths(NewRestLengthPositions, ConfigProperties, WeightMaps);
+	}
+}
 }  // End namespace Chaos
diff --git a/Engine/Plugins/ChaosCloth/Source/ChaosCloth/Private/ChaosCloth/ChaosClothingSimulationCloth.cpp b/Engine/Plugins/ChaosCloth/Source/ChaosCloth/Private/ChaosCloth/ChaosClothingSimulationCloth.cpp
--- a/Engine/Plugins/ChaosCloth/Source/ChaosCloth/Private/ChaosCloth/ChaosClothingSimulationCloth.cpp
+++ b/Engine/Plugins/ChaosCloth/Source/ChaosCloth/Private/ChaosCloth/ChaosClothingSimulationCloth.cpp
@@ -955,7 +955,7 @@ void FClothingSimulationCloth::Update(FClothingSimulationSolver* Solver)
 				Solver->GetParticleXs(ParticleRangeId),
 				Solver->GetParticleVs(ParticleRangeId));
 
-			if(Solver->IsLegacySolver())
+			if (Solver->IsLegacySolver())
 			{
 				// Update the wind velocity field for the new LOD mesh
 				Solver->SetWindAndPressureGeometry(GroupId, GetTriangleMesh(Solver), ConfigProperties, LODData[LODIndex]->WeightMaps);
@@ -1120,6 +1120,24 @@ void FClothingSimulationCloth::Update(FClothingSimulationSolver* Solver)
 		}
 	}
 
+	if (bResetRestLengthsFromMorphTarget)
+	{
+		check(!Solver->IsLegacySolver()); // Otherwise need to handle particle offsets
+		for (int32 Index = 0; Index < LODData.Num(); ++Index)
+		{
+			const int32 CurrParticleRangeId = GetParticleRangeId(Solver, Index);
+			TArray<Softs::FSolverVec3> RestLengthPositions(Mesh->GetPositions(Index));
+			TArray<Softs::FSolverVec3> RestLengthNormals(Mesh->GetNormals(Index));
+			const int32 MorphTargetIndex = Mesh->FindMorphTargetByName(Index, ResetRestLengthsMorphTargetName);
+			if (MorphTargetIndex != INDEX_NONE)
+			{
+				Mesh->ApplyMorphTarget(Index, MorphTargetIndex, 1.f, RestLengthPositions, RestLengthNormals);
+			}
+			FClothConstraints& ClothConstraints = Solver->GetClothConstraints(CurrParticleRangeId);
+			ClothConstraints.ResetRestLengths(RestLengthPositions, Config->GetProperties(Index), LODData[Index]->WeightMaps);
+		}
+	}
+
 	// Reset trigger flags
 #if CHAOS_DEBUG_DRAW
 	if (bNeedsReset)
@@ -1141,6 +1159,7 @@ void FClothingSimulationCloth::Update(FClothingSimulationSolver* Solver)
 #endif
 	bNeedsTeleport = false;
 	bNeedsReset = false;
+	bResetRestLengthsFromMorphTarget = false;
 }
 
 void FClothingSimulationCloth::PostUpdate(FClothingSimulationSolver* Solver)
diff --git a/Engine/Plugins/ChaosCloth/Source/ChaosCloth/Private/ChaosCloth/ChaosClothingSimulationMesh.cpp b/Engine/Plugins/ChaosCloth/Source/ChaosCloth/Private/ChaosCloth/ChaosClothingSimulationMesh.cpp
--- a/Engine/Plugins/ChaosCloth/Source/ChaosCloth/Private/ChaosCloth/ChaosClothingSimulationMesh.cpp
+++ b/Engine/Plugins/ChaosCloth/Source/ChaosCloth/Private/ChaosCloth/ChaosClothingSimulationMesh.cpp
@@ -180,6 +180,44 @@ bool FClothingSimulationMesh::WrapDeformLOD(
 	return true;
 }
 
+template<typename InOutVectorType>
+void ApplyMorphTargetInternal(
+	float ActiveMorphTargetWeight,
+	const TConstArrayView<FVector3f>& MorphTargetPositionDeltas,
+	const TConstArrayView<FVector3f>& MorphTargetTangentZDeltas,
+	const TConstArrayView<int32>& MorphTargetIndices,
+	const TArrayView<InOutVectorType>& InOutPositions,
+	const TArrayView<InOutVectorType>& InOutNormals)
+{
+	TRACE_CPUPROFILER_EVENT_SCOPE(FClothingSimulationMesh_ApplyMorphTargetInternal);
+
+	for (int32 Index = 0; Index < MorphTargetIndices.Num(); ++Index)
+	{
+		const int32 VertexIndex = MorphTargetIndices[Index];
+		InOutPositions[VertexIndex] += ActiveMorphTargetWeight * MorphTargetPositionDeltas[Index];
+		InOutNormals[VertexIndex] = (InOutNormals[VertexIndex] + ActiveMorphTargetWeight * MorphTargetTangentZDeltas[Index]).GetSafeNormal();
+	}
+}
+
+void FClothingSimulationMesh::ApplyMorphTarget(
+	int32 LODIndex,
+	int32 ActiveMorphTargetIndex,
+	float ActiveMorphTargetWeight,
+	const TArrayView<Softs::FSolverVec3>& InOutPositions,
+	const TArrayView<Softs::FSolverVec3>& InOutNormals) const
+{
+	const TConstArrayView<FVector3f> MorphTargetPositionDeltas = GetMorphTargetPositionDeltas(LODIndex, ActiveMorphTargetIndex);
+	const TConstArrayView<FVector3f> MorphTargetTangentZDeltas = GetMorphTargetTangentZDeltas(LODIndex, ActiveMorphTargetIndex);
+	const TConstArrayView<int32> MorphTargetIndices = GetMorphTargetIndices(LODIndex, ActiveMorphTargetIndex);
+
+	if (!FMath::IsNearlyZero(ActiveMorphTargetWeight) && !MorphTargetPositionDeltas.IsEmpty() &&
+		MorphTargetPositionDeltas.Num() == MorphTargetTangentZDeltas.Num() &&
+		MorphTargetPositionDeltas.Num() == MorphTargetIndices.Num())
+	{
+		ApplyMorphTargetInternal(ActiveMorphTargetWeight, MorphTargetPositionDeltas, MorphTargetTangentZDeltas, MorphTargetIndices, InOutPositions, InOutNormals);
+	}
+}
+
 // Inline function used to force the unrolling of the skinning loop, LWC: note skinning is all done in float to match the asset data type
 FORCEINLINE static void AddInfluence(FVector3f& OutPosition, FVector3f& OutNormal, const FVector3f& RefParticle, const FVector3f& RefNormal, const FMatrix44f& BoneMatrix, const float Weight)
 {
@@ -226,14 +264,7 @@ void FClothingSimulationMesh::SkinPhysicsMesh(int32 LODIndex, int32 ActiveMorphT
 		// TODO optimize this
 		WritablePositions = Positions;
 		WritableNormals = Normals;
-
-		for (int32 Index = 0; Index < MorphTargetIndices.Num(); ++Index)
-		{
-			const int32 VertexIndex = MorphTargetIndices[Index];
-			WritablePositions[VertexIndex] += ActiveMorphTargetWeight * MorphTargetPositionDeltas[Index];
-			WritableNormals[VertexIndex] = (WritableNormals[VertexIndex] + ActiveMorphTargetWeight * MorphTargetTangentZDeltas[Index]).GetSafeNormal();
-		}
-
+		ApplyMorphTargetInternal(ActiveMorphTargetWeight, MorphTargetPositionDeltas, MorphTargetTangentZDeltas, MorphTargetIndices, TArrayView<FVector3f>(WritablePositions), TArrayView<FVector3f>(WritableNormals));
 		Positions = WritablePositions;
 		Normals = WritableNormals;
 	}
diff --git a/Engine/Plugins/ChaosCloth/Source/ChaosCloth/Public/ChaosCloth/ChaosClothConstraints.h b/Engine/Plugins/ChaosCloth/Source/ChaosCloth/Public/ChaosCloth/ChaosClothConstraints.h
--- a/Engine/Plugins/ChaosCloth/Source/ChaosCloth/Public/ChaosCloth/ChaosClothConstraints.h
+++ b/Engine/Plugins/ChaosCloth/Source/ChaosCloth/Public/ChaosCloth/ChaosClothConstraints.h
@@ -94,6 +94,11 @@ namespace Chaos
 		// NOTE: this only does something if using the PBDSolver. Force-based solver constraints are activated automatically
 		// when activating a particle range.
 		void Enable(bool bEnable);
+
+		void ResetRestLengths(
+			const TConstArrayView<Softs::FSolverVec3>& NewRestLengthPositions,
+			const Softs::FCollectionPropertyConstFacade& ConfigProperties,
+			const TMap<FString, TConstArrayView<FRealSingle>>& WeightMaps);
 		// ---- End of Cloth interface ----
 
 		// ---- Debug functions ----
diff --git a/Engine/Plugins/ChaosCloth/Source/ChaosCloth/Public/ChaosCloth/ChaosClothingSimulationCloth.h b/Engine/Plugins/ChaosCloth/Source/ChaosCloth/Public/ChaosCloth/ChaosClothingSimulationCloth.h
--- a/Engine/Plugins/ChaosCloth/Source/ChaosCloth/Public/ChaosCloth/ChaosClothingSimulationCloth.h
+++ b/Engine/Plugins/ChaosCloth/Source/ChaosCloth/Public/ChaosCloth/ChaosClothingSimulationCloth.h
@@ -52,6 +52,11 @@ namespace Chaos
 
 		void Reset() { bNeedsReset = true; }
 		void Teleport() { bNeedsTeleport = true; }
+		void ResetRestLengthsWithMorphTarget(const FString& MorphTargetName) 
+		{
+			bResetRestLengthsFromMorphTarget = true; 
+			ResetRestLengthsMorphTargetName = MorphTargetName; 
+		}
 		// ---- End of the animatable property setters ----
 
 		// ---- Node property getters/setters
@@ -156,6 +161,10 @@ namespace Chaos
 		bool bNeedsReset = false;
 		bool bNeedsTeleport = false;
 
+		/** Reset restlengths from morph target. */
+		bool bResetRestLengthsFromMorphTarget = false;
+		FString ResetRestLengthsMorphTargetName;
+
 #if CHAOS_DEBUG_DRAW
 		FRealSingle TimeSinceLastTeleport = 0.f;
 		FRealSingle TimeSinceLastReset = 0.f;
diff --git a/Engine/Plugins/ChaosCloth/Source/ChaosCloth/Public/ChaosCloth/ChaosClothingSimulationMesh.h b/Engine/Plugins/ChaosCloth/Source/ChaosCloth/Public/ChaosCloth/ChaosClothingSimulationMesh.h
--- a/Engine/Plugins/ChaosCloth/Source/ChaosCloth/Public/ChaosCloth/ChaosClothingSimulationMesh.h
+++ b/Engine/Plugins/ChaosCloth/Source/ChaosCloth/Public/ChaosCloth/ChaosClothingSimulationMesh.h
@@ -201,6 +201,14 @@ namespace Chaos
 			int32 Offset,
 			int32 ActiveMorphTargetIndex = INDEX_NONE,
 			float ActiveMorphTargetWeight = 0.f);
+
+		/** Apply a Morph Target to a given set of positions */
+		UE_API void ApplyMorphTarget(
+			int32 LODIndex,
+			int32 ActiveMorphTargetIndex,
+			float ActiveMorphTargetWeight,
+			const TArrayView<Softs::FSolverVec3>& InOutPositions,
+			const TArrayView<Softs::FSolverVec3>& InOutNormals) const;
 		// ---- End of the Cloth interface ----
 
 	private:
diff --git a/Engine/Plugins/ChaosClothAsset/Source/ChaosClothAssetEngine/Private/ChaosClothAsset/ClothComponent.cpp b/Engine/Plugins/ChaosClothAsset/Source/ChaosClothAssetEngine/Private/ChaosClothAsset/ClothComponent.cpp
--- a/Engine/Plugins/ChaosClothAsset/Source/ChaosClothAssetEngine/Private/ChaosClothAsset/ClothComponent.cpp
+++ b/Engine/Plugins/ChaosClothAsset/Source/ChaosClothAssetEngine/Private/ChaosClothAsset/ClothComponent.cpp
@@ -337,6 +337,7 @@ void UChaosClothComponent::TickComponent(float DeltaTime, enum ELevelTick TickTy
 		bTickOnceInEditor = false;
 	}
 #endif
+	bResetRestLengthsFromMorphTarget = false;
 }
 
 bool UChaosClothComponent::RequiresPreEndOfFrameSync() const
diff --git a/Engine/Plugins/ChaosClothAsset/Source/ChaosClothAssetEngine/Private/ChaosClothAsset/ClothSimulationContext.cpp b/Engine/Plugins/ChaosClothAsset/Source/ChaosClothAssetEngine/Private/ChaosClothAsset/ClothSimulationContext.cpp
--- a/Engine/Plugins/ChaosClothAsset/Source/ChaosClothAssetEngine/Private/ChaosClothAsset/ClothSimulationContext.cpp
+++ b/Engine/Plugins/ChaosClothAsset/Source/ChaosClothAssetEngine/Private/ChaosClothAsset/ClothSimulationContext.cpp
@@ -13,6 +13,7 @@ namespace UE::Chaos::ClothAsset
 {
 	void FClothSimulationContext::Fill(const UChaosClothComponent& ClothComponent, float InDeltaTime, float MaxDeltaTime, bool bIsInitialization, FClothingSimulationCacheData* InCacheData)
 	{
+		TRACE_CPUPROFILER_EVENT_SCOPE(FClothSimulationContext_Fill);
 		// Set the time
 		DeltaTime = FMath::Min(InDeltaTime, MaxDeltaTime);
 
@@ -29,6 +30,12 @@ namespace UE::Chaos::ClothAsset
 		bTeleport = (InDeltaTime > MaxDeltaTime * MaxDeltaTimeTeleportMultiplier) ? true : 
 			(ComponentTeleportMode >= EClothingTeleportMode::Teleport);
 		bReset = ComponentTeleportMode == EClothingTeleportMode::TeleportAndReset;
+		
+		bResetRestLengthsWithMorphTarget = ClothComponent.NeedsResetRestLengths();
+		if (bResetRestLengthsWithMorphTarget)
+		{
+			ResetRestLengthsMorphTargetName = ClothComponent.GetRestLengthsMorphTargetName();
+		}
 
 		VelocityScale = (!bTeleport && !bReset && InDeltaTime > 0.f) ?
 			FMath::Min(InDeltaTime, MaxDeltaTime) / InDeltaTime :
diff --git a/Engine/Plugins/ChaosClothAsset/Source/ChaosClothAssetEngine/Private/ChaosClothAsset/ClothSimulationContext.h b/Engine/Plugins/ChaosClothAsset/Source/ChaosClothAssetEngine/Private/ChaosClothAsset/ClothSimulationContext.h
--- a/Engine/Plugins/ChaosClothAsset/Source/ChaosClothAssetEngine/Private/ChaosClothAsset/ClothSimulationContext.h
+++ b/Engine/Plugins/ChaosClothAsset/Source/ChaosClothAssetEngine/Private/ChaosClothAsset/ClothSimulationContext.h
@@ -42,6 +42,10 @@ namespace UE::Chaos::ClothAsset
 		/** Reset. */
 		bool bReset = false;
 
+		/** Reset RestLengths (with morph target) */
+		bool bResetRestLengthsWithMorphTarget = false;
+		FString ResetRestLengthsMorphTargetName;
+
 		/** Velocity scale to compensate for time differences when the MaxPhysicsDeltaTime kicks in. */
 		float VelocityScale;
 
diff --git a/Engine/Plugins/ChaosClothAsset/Source/ChaosClothAssetEngine/Private/ChaosClothAsset/ClothSimulationProxy.cpp b/Engine/Plugins/ChaosClothAsset/Source/ChaosClothAssetEngine/Private/ChaosClothAsset/ClothSimulationProxy.cpp
--- a/Engine/Plugins/ChaosClothAsset/Source/ChaosClothAssetEngine/Private/ChaosClothAsset/ClothSimulationProxy.cpp
+++ b/Engine/Plugins/ChaosClothAsset/Source/ChaosClothAssetEngine/Private/ChaosClothAsset/ClothSimulationProxy.cpp
@@ -282,7 +282,7 @@ PRAGMA_ENABLE_DEPRECATION_WARNINGS
 
 	void FClothSimulationProxy::PreProcess_GameThread(float DeltaTime, bool bForceWaitForInitialization)
 	{
-		SCOPE_CYCLE_COUNTER(STAT_ClothSimulationProxy_TickGame);
+		TRACE_CPUPROFILER_EVENT_SCOPE(FClothSimulationProxy_PreProcess_GameThread);
 
 		// Set bIsSimulating to its default value, this will be changed in PreSimulate_GameThread if the simulation ever runs
 		bIsSimulating = false;
@@ -325,8 +325,6 @@ PRAGMA_ENABLE_DEPRECATION_WARNINGS
 
 	bool FClothSimulationProxy::PreSimulate_GameThread(float DeltaTime)
 	{
-		SCOPE_CYCLE_COUNTER(STAT_ClothSimulationProxy_TickGame);
-
 		if (!bIsPreProcessed)
 		{
 			return false;
@@ -485,6 +483,10 @@ PRAGMA_ENABLE_DEPRECATION_WARNINGS
 			{
 				Cloth->Teleport();
 			}
+			if (ClothSimulationContext->bResetRestLengthsWithMorphTarget)
+			{
+				Cloth->ResetRestLengthsWithMorphTarget(ClothSimulationContext->ResetRestLengthsMorphTargetName);
+			}
 		}
 
 		// Step the simulation
@@ -817,6 +819,7 @@ PRAGMA_ENABLE_DEPRECATION_WARNINGS
 
 	void FClothSimulationProxy::InitializeConfigs()
 	{
+		TRACE_CPUPROFILER_EVENT_SCOPE(FClothSimulationProxy_InitializeConfigs);
 		// Replace physics thread's configs with the game thread's configs
 		for (int32 ConfigIndex = 0; ConfigIndex < Configs.Num(); ++ConfigIndex)
 		{
diff --git a/Engine/Plugins/ChaosClothAsset/Source/ChaosClothAssetEngine/Public/ChaosClothAsset/ClothComponent.h b/Engine/Plugins/ChaosClothAsset/Source/ChaosClothAssetEngine/Public/ChaosClothAsset/ClothComponent.h
--- a/Engine/Plugins/ChaosClothAsset/Source/ChaosClothAssetEngine/Public/ChaosClothAsset/ClothComponent.h
+++ b/Engine/Plugins/ChaosClothAsset/Source/ChaosClothAssetEngine/Public/ChaosClothAsset/ClothComponent.h
@@ -97,12 +97,19 @@ public:
 	UFUNCTION(BlueprintCallable, Category = "ClothComponent", Meta = (Keywords = "Chaos Cloth Teleport Reset"))
 	void ForceNextUpdateTeleportAndReset() { bTeleportOnce = bResetOnce = true; }
 
+	/** Reset the cloth rest lengths with a given morph target (applied to base pose). If MorphTargetName is empty or does not exist, the rest lengths will be reset to default. */
+	UFUNCTION(BlueprintCallable, Category = "ClothComponent", Meta = (Keywords = "Chaos Cloth Reset Rest Length Morph Target"))
+	void ResetRestLengthsWithMorphTarget(const FString& MorphTargetName) { bResetRestLengthsFromMorphTarget = true; ResetRestLengthsMorphTargetName = MorphTargetName; }
+
 	/** Return whether teleport is currently requested. Use GetClothTeleportMode to get teleport mode including any auto-teleport based on teleport thresholds. */
 	bool NeedsTeleport() const { return bTeleport || bTeleportOnce; }
 
 	/** Return whether reseting the pose is currently requested. Use GetClothTeleportMode to get teleport mode including any auto-teleport based on teleport thresholds*/
 	bool NeedsReset() const { return bReset || bResetOnce; }
 
+	bool NeedsResetRestLengths() const { return bResetRestLengthsFromMorphTarget; }
+	const FString& GetRestLengthsMorphTargetName() const { return ResetRestLengthsMorphTargetName; }
+
 	/** Get currently calculated teleport mode. */
 	EClothingTeleportMode GetClothTeleportMode() const { return ClothTeleportMode; }
 
@@ -441,6 +448,10 @@ private:
 
 	bool bHasValidRenderDataForVisibility = false;
 
+	/** Reset restlengths from morph target. */
+	bool bResetRestLengthsFromMorphTarget = false;
+	FString ResetRestLengthsMorphTargetName;
+
 	/**
 	 * Simulation properties per cloth model. 
 	 * Transient because it doesn't need to be serialized but contains pointers to UObject that needs not be garbage collected.
diff --git a/Engine/Source/Runtime/Experimental/Chaos/Private/Chaos/PBDBendingConstraintsBase.cpp b/Engine/Source/Runtime/Experimental/Chaos/Private/Chaos/PBDBendingConstraintsBase.cpp
--- a/Engine/Source/Runtime/Experimental/Chaos/Private/Chaos/PBDBendingConstraintsBase.cpp
+++ b/Engine/Source/Runtime/Experimental/Chaos/Private/Chaos/PBDBendingConstraintsBase.cpp
@@ -19,26 +19,24 @@ namespace Chaos::Softs
 {
 	namespace Private
 	{
-		template<typename SolverParticlesOrRange>
 		void Calculate3DRestAngles(
-			const SolverParticlesOrRange& InParticles,
+			const TConstArrayView<FSolverVec3>& InPositions,
 			const TArray<TVec4<int32>>& Constraints,
 			TArray<FSolverReal>& RestAngles)
 		{
 			RestAngles.Reset(Constraints.Num());
 			for (const TVec4<int32>& Constraint : Constraints)
 			{
-				const FSolverVec3& P1 = InParticles.X(Constraint[0]);
-				const FSolverVec3& P2 = InParticles.X(Constraint[1]);
-				const FSolverVec3& P3 = InParticles.X(Constraint[2]);
-				const FSolverVec3& P4 = InParticles.X(Constraint[3]);
+				const FSolverVec3& P1 = InPositions[Constraint[0]];
+				const FSolverVec3& P2 = InPositions[Constraint[1]];
+				const FSolverVec3& P3 = InPositions[Constraint[2]];
+				const FSolverVec3& P4 = InPositions[Constraint[3]];
 				RestAngles.Add(FMath::Clamp(FPBDBendingConstraintsBase::CalcAngle(P1, P2, P3, P4), -UE_PI, UE_PI));
 			}
 		}
 
-		template<typename SolverParticlesOrRange>
 		void CalculateFlatnessRestAngles(
-			const SolverParticlesOrRange& InParticles,
+			const TConstArrayView<FSolverVec3>& InPositions,
 			int32 InParticleOffset,
 			int32 InParticleCount,
 			const TConstArrayView<FRealSingle>& RestAngleMap,
@@ -69,10 +67,10 @@ namespace Chaos::Softs
 			RestAngles.Reset(Constraints.Num());
 			for (const TVec4<int32>& Constraint : Constraints)
 			{
-				const FSolverVec3& P1 = InParticles.X(Constraint[0]);
-				const FSolverVec3& P2 = InParticles.X(Constraint[1]);
-				const FSolverVec3& P3 = InParticles.X(Constraint[2]);
-				const FSolverVec3& P4 = InParticles.X(Constraint[3]);
+				const FSolverVec3& P1 = InPositions[Constraint[0]];
+				const FSolverVec3& P2 = InPositions[Constraint[1]];
+				const FSolverVec3& P3 = InPositions[Constraint[2]];
+				const FSolverVec3& P4 = InPositions[Constraint[3]];
 				const FSolverReal Flatness = FMath::Clamp(FlatnessValue(Constraint), (FSolverReal)0.f, (FSolverReal)1.f);
 				RestAngles.Add(FMath::Clamp(((FSolverReal)1.f - Flatness) * FPBDBendingConstraintsBase::CalcAngle(P1, P2, P3, P4), -UE_PI, UE_PI));
 			}
@@ -111,9 +109,8 @@ namespace Chaos::Softs
 		}
 	}
 
-	template<typename SolverParticlesOrRange>
 	void FPBDBendingConstraintsBase::CalculateRestAngles(
-		const SolverParticlesOrRange& InParticles,
+		const TConstArrayView<FSolverVec3>& InPositions,
 		int32 InParticleOffset,
 		int32 InParticleCount,
 		const TConstArrayView<FRealSingle>& RestAngleMap,
@@ -124,32 +121,16 @@ namespace Chaos::Softs
 		{
 		default:
 		case ERestAngleConstructionType::Use3DRestAngles:
-			Private::Calculate3DRestAngles(InParticles, Constraints, RestAngles);
+			Private::Calculate3DRestAngles(InPositions, Constraints, RestAngles);
 			break;
 		case ERestAngleConstructionType::FlatnessRatio:
-			Private::CalculateFlatnessRestAngles(InParticles, InParticleOffset, InParticleCount, RestAngleMap, RestAngleValue, Constraints, RestAngles);
+			Private::CalculateFlatnessRestAngles(InPositions, InParticleOffset, InParticleCount, RestAngleMap, RestAngleValue, Constraints, RestAngles);
 			break;
 		case ERestAngleConstructionType::ExplicitRestAngles:
 			Private::CalculateExplicitRestAngles(InParticleOffset, InParticleCount, RestAngleMap, RestAngleValue, Constraints, RestAngles);
 		}
 	}
 
-	template CHAOS_API void FPBDBendingConstraintsBase::CalculateRestAngles(
-		const FSolverParticles& InParticles,
-		int32 InParticleOffset,
-		int32 InParticleCount,
-		const TConstArrayView<FRealSingle>& RestAngleMap,
-		const FSolverVec2& RestAngleValue,
-		ERestAngleConstructionType RestAngleConstructionType);
-	template CHAOS_API void FPBDBendingConstraintsBase::CalculateRestAngles(
-		const FSolverParticlesRange& InParticles,
-		int32 InParticleOffset,
-		int32 InParticleCount,
-		const TConstArrayView<FRealSingle>& RestAngleMap,
-		const FSolverVec2& RestAngleValue,
-		ERestAngleConstructionType RestAngleConstructionType);
-
-
 	template<typename SolverParticlesOrRange>
 	void FPBDBendingConstraintsBase::Init(const SolverParticlesOrRange& InParticles)
 	{
diff --git a/Engine/Source/Runtime/Experimental/Chaos/Private/Chaos/PBDSpringConstraints.cpp b/Engine/Source/Runtime/Experimental/Chaos/Private/Chaos/PBDSpringConstraints.cpp
--- a/Engine/Source/Runtime/Experimental/Chaos/Private/Chaos/PBDSpringConstraints.cpp
+++ b/Engine/Source/Runtime/Experimental/Chaos/Private/Chaos/PBDSpringConstraints.cpp
@@ -371,6 +371,19 @@ void FPBDEdgeSpringConstraints::ApplyProperties(const FSolverReal Dt, const int3
 	}
 }
 
+void FPBDEdgeSpringConstraints::ResetRestLengths(const TConstArrayView<FSolverVec3>& Positions)
+{
+	TRACE_CPUPROFILER_EVENT_SCOPE(FPBDSpringConstraints_ResetRestLengths);
+	
+	TArray<FSolverReal>& DistsToCompute = bWarpWeftScaleEnabled ? BaseDists : Dists;
+	for (int32 Index = 0; Index < Constraints.Num(); ++Index)
+	{
+		DistsToCompute[Index] = (Positions[Constraints[Index][0]] - Positions[Constraints[Index][1]]).Size();
+	}
+
+	UpdateDists();
+}
+
 void FPBDEdgeSpringConstraints::UpdateDists()
 {
 	if (bWarpWeftScaleEnabled)
diff --git a/Engine/Source/Runtime/Experimental/Chaos/Public/Chaos/PBDBendingConstraints.h b/Engine/Source/Runtime/Experimental/Chaos/Public/Chaos/PBDBendingConstraints.h
--- a/Engine/Source/Runtime/Experimental/Chaos/Public/Chaos/PBDBendingConstraints.h
+++ b/Engine/Source/Runtime/Experimental/Chaos/Public/Chaos/PBDBendingConstraints.h
@@ -157,6 +157,16 @@ public:
 
 	const TArray<int32>& GetConstraintsPerColorStartIndex() const { return ConstraintsPerColorStartIndex; }
 
+	void ResetRestLengths(const TConstArrayView<FSolverVec3>& Positions,
+		const FCollectionPropertyConstFacade& PropertyCollection,
+		const TMap<FString, TConstArrayView<FRealSingle>>& WeightMaps)
+	{
+		CalculateRestAngles(Positions, ParticleOffset, ParticleCount,
+			GetRestAngleMapFromCollection(WeightMaps, PropertyCollection),
+			GetRestAngleValueFromCollection(PropertyCollection),
+			(ERestAngleConstructionType)GetRestAngleType(PropertyCollection, (int32)ERestAngleConstructionType::Use3DRestAngles));
+	}
+
 private:
 	template<typename SolverParticlesOrRange>
 	CHAOS_API void InitColor(const SolverParticlesOrRange& InParticles);
diff --git a/Engine/Source/Runtime/Experimental/Chaos/Public/Chaos/PBDBendingConstraintsBase.h b/Engine/Source/Runtime/Experimental/Chaos/Public/Chaos/PBDBendingConstraintsBase.h
--- a/Engine/Source/Runtime/Experimental/Chaos/Public/Chaos/PBDBendingConstraintsBase.h
+++ b/Engine/Source/Runtime/Experimental/Chaos/Public/Chaos/PBDBendingConstraintsBase.h
@@ -66,7 +66,7 @@ PRAGMA_DISABLE_DEPRECATION_WARNINGS
 		, BucklingRatio((FSolverReal)BucklingRatioWeighted)
 PRAGMA_ENABLE_DEPRECATION_WARNINGS
 	{
-		CalculateRestAngles(InParticles, ParticleOffset, ParticleCount, RestAngleMap, RestAngleValue, RestAngleConstructionType);
+		CalculateRestAngles(InParticles.XArray(), ParticleOffset, ParticleCount, RestAngleMap, RestAngleValue, RestAngleConstructionType);
 	}
 
 	UE_DEPRECATED(5.5, "Use a constructor with BucklingRatioMultipliers.")
@@ -145,7 +145,7 @@ PRAGMA_ENABLE_DEPRECATION_WARNINGS
 		, BucklingRatio((FSolverReal)BucklingRatioWeighted)
 		PRAGMA_ENABLE_DEPRECATION_WARNINGS
 	{
-		CalculateRestAngles(InParticles, ParticleOffset, ParticleCount, RestAngleMap, RestAngleValue, RestAngleConstructionType);
+		CalculateRestAngles(InParticles.XArray(), ParticleOffset, ParticleCount, RestAngleMap, RestAngleValue, RestAngleConstructionType);
 	}
 
 	UE_DEPRECATED(5.5, "Use a constructor with BucklingRatioMultipliers.")
@@ -344,6 +344,13 @@ protected:
 		return Grads;
 	}
 
+	CHAOS_API void CalculateRestAngles(const TConstArrayView<FSolverVec3>& InPositions,
+		int32 InParticleOffset,
+		int32 InParticleCount,
+		const TConstArrayView<FRealSingle>& RestAngleMap,
+		const FSolverVec2& RestAngleValue,
+		ERestAngleConstructionType RestAngleConstructionType);
+
 private:
 	template<typename SolverParticlesOrRange>
 	static TArray<TVec4<int32>> TrimKinematicConstraints(const TArray<TVec4<int32>>& InConstraints, const SolverParticlesOrRange& InParticles)
@@ -371,14 +378,6 @@ private:
 		}
 		return ExtractedEdges;
 	}
-
-	template<typename SolverParticlesOrRange>
-	CHAOS_API void CalculateRestAngles(const SolverParticlesOrRange& InParticles,
-		int32 InParticleOffset,
-		int32 InParticleCount,
-		const TConstArrayView<FRealSingle>& RestAngleMap,
-		const FSolverVec2& RestAngleValue,
-		ERestAngleConstructionType RestAngleConstructionType);
 	
 protected:
 	TArray<TVec4<int32>> Constraints;
diff --git a/Engine/Source/Runtime/Experimental/Chaos/Public/Chaos/PBDSpringConstraints.h b/Engine/Source/Runtime/Experimental/Chaos/Public/Chaos/PBDSpringConstraints.h
--- a/Engine/Source/Runtime/Experimental/Chaos/Public/Chaos/PBDSpringConstraints.h
+++ b/Engine/Source/Runtime/Experimental/Chaos/Public/Chaos/PBDSpringConstraints.h
@@ -161,6 +161,8 @@ public:
 
 	CHAOS_API void ApplyProperties(const FSolverReal Dt, const int32 NumIterations);
 
+	CHAOS_API void ResetRestLengths(const TConstArrayView<FSolverVec3>& Positions);
+
 private:
 
 	void InitFromPatternData(const TArray<TVec3<FVec2f>>& FaceVertexPatternPositions, const FTriangleMesh& TriangleMesh);
