Fix simulation crash in when muscle activation is set up

diff --git a/Engine/Source/Runtime/Experimental/Chaos/Private/Chaos/Deformable/ChaosDeformableSolver.cpp b/Engine/Source/Runtime/Experimental/Chaos/Private/Chaos/Deformable/ChaosDeformableSolver.cpp
--- a/Engine/Source/Runtime/Experimental/Chaos/Private/Chaos/Deformable/ChaosDeformableSolver.cpp
+++ b/Engine/Source/Runtime/Experimental/Chaos/Private/Chaos/Deformable/ChaosDeformableSolver.cpp
@@ -2017,11 +2017,14 @@ namespace Chaos::Softs
 		GeometryCollection::Facades::FMuscleActivationFacade MuscleActivationFacade(Rest);
 		if (MuscleActivationFacade.IsValid())
 		{
-			int32 VertexOffset = Proxy.GetSolverParticleRange().Start;
-			// Only supports Gauss Seidel for now because we are using AllElements
-			int32 ElementOffset = AllElements->Num(); //InitializeMuscleActivation has to go before adding tetrahedrons
-			MuscleIndexOffset.Add(Proxy.GetOwner(), MuscleActivationConstraints->NumMuscles()); 
-			MuscleActivationConstraints->AddMuscles(Evolution->Particles(), MuscleActivationFacade, VertexOffset, ElementOffset);
+			if (AllElements)
+			{
+				int32 VertexOffset = Proxy.GetSolverParticleRange().Start;
+				// Only supports Gauss Seidel for now because we are using AllElements
+				int32 ElementOffset = AllElements->Num(); //InitializeMuscleActivation has to go before adding tetrahedrons
+				MuscleIndexOffset.Add(Proxy.GetOwner(), MuscleActivationConstraints->NumMuscles());
+				MuscleActivationConstraints->AddMuscles(Evolution->Particles(), MuscleActivationFacade, VertexOffset, ElementOffset);
+			}
 		}
 	}
 
