Chaos - fix inertia conditioning for medium to large objects

Inertia conditioning was only working for small objects typically used on characters but was early-outing on an inverse inertia tolerance that was far too large.

diff --git a/Engine/Source/Runtime/Experimental/Chaos/Private/Chaos/MassConditioning.cpp b/Engine/Source/Runtime/Experimental/Chaos/Private/Chaos/MassConditioning.cpp
--- a/Engine/Source/Runtime/Experimental/Chaos/Private/Chaos/MassConditioning.cpp
+++ b/Engine/Source/Runtime/Experimental/Chaos/Private/Chaos/MassConditioning.cpp
@@ -6,18 +6,30 @@
 #include "Chaos/PBDJointConstraints.h"
 #include "Chaos/PBDJointConstraintUtilities.h"
 
+namespace Chaos::CVars
+{
+	float ChaosInertiaConditioningInvMassTolerance = 1.e-20f;
+	float ChaosInertiaConditioningInvInertiaTolerance = 1.e-30f;
+	float ChaosInertiaConditioningExtentTolerance = 1.e-8f;
+	FAutoConsoleVariableRef  CVarChaosInertiaConditioningInvMassTolerance(TEXT("p.Chaos.InertiaConditioning.InvMassTolerance"), ChaosInertiaConditioningInvMassTolerance, TEXT(""));
+	FAutoConsoleVariableRef  CVarChaosInertiaConditioningInvInertiaTolerance(TEXT("p.Chaos.InertiaConditioning.InvInertiaTolerance"), ChaosInertiaConditioningInvInertiaTolerance, TEXT(""));
+	FAutoConsoleVariableRef  CVarChaosInertiaConditioningExtentTolerance(TEXT("p.Chaos.InertiaConditioning.ExtentTolerance"), ChaosInertiaConditioningExtentTolerance, TEXT(""));
+}
+
 namespace Chaos
 {
 	FVec3f CalculateInertiaConditioning(const FRealSingle InvM, const FVec3f& InvI, const FVec3f& ConstraintExtents, const FRealSingle MaxDistance, const FRealSingle MaxRotationRatio, const FRealSingle MaxInvInertiaComponentRatio)
 	{
-		check(InvM > 0);
-		check(!InvI.IsZero());
-
 		FVec3f InvInertiaScale = FVec3f(1);
 
-		// Sanity check for zero dimensions along at least one axis
+		// Sanity check for zero mass or zero dimensions along at least one axis
 		// We should never get this unless there are asset errors
-		if (ConstraintExtents.GetMin() > UE_SMALL_NUMBER)
+		bool bIsValid = 
+			(InvM > CVars::ChaosInertiaConditioningInvMassTolerance) 
+			&& (InvI.GetMin() > CVars::ChaosInertiaConditioningInvInertiaTolerance)
+			&& (ConstraintExtents.GetMin() > CVars::ChaosInertiaConditioningExtentTolerance);
+
+		if (bIsValid)
 		{
 			// The I,J elements of this matrix can be interpreted the ratio of the rotation correction versus the translation 
 			// correction for a constraint arm along the Ith axis and a correction along the Jth axis. It has zero diagonal 
@@ -113,12 +125,11 @@ namespace Chaos
 
 	FVec3f CalculateParticleInertiaConditioning(const FPBDRigidParticleHandle* Rigid, const FRealSingle MaxDistance, const FRealSingle MaxRotationRatio, const FRealSingle MaxInvInertiaComponentRatio)
 	{
-		if (FMath::IsNearlyZero(Rigid->InvM()) || Rigid->InvI().IsNearlyZero())
+		if (Rigid->IsDynamic())
 		{
-			return FVec3f(1);
+			const FVec3 ConstraintExtents = CalculateCoMConstraintExtents(Rigid);
+			return CalculateInertiaConditioning(FRealSingle(Rigid->InvM()), Rigid->InvI(), ConstraintExtents, MaxDistance, MaxRotationRatio, MaxInvInertiaComponentRatio);
 		}
-
-		const FVec3 ConstraintExtents = CalculateCoMConstraintExtents(Rigid);
-		return CalculateInertiaConditioning(FRealSingle(Rigid->InvM()), Rigid->InvI(), ConstraintExtents, MaxDistance, MaxRotationRatio, MaxInvInertiaComponentRatio);
+		return FVec3f(1);
 	}
 }
\ No newline at end of file
