Dataflow : fix crash when executing a redo command on a deleted node
- when calling undo command the Dataflow node could not be found yet ( because still missing from the graph ) and would fail to serialize the node
- then when redo was called , the node woudl be there and the serialization would be called on the node while the data did not saved anything causing out of bound read from the archive

diff --git a/Engine/Source/Runtime/Experimental/Dataflow/Engine/Private/Dataflow/DataflowEdNode.cpp b/Engine/Source/Runtime/Experimental/Dataflow/Engine/Private/Dataflow/DataflowEdNode.cpp
--- a/Engine/Source/Runtime/Experimental/Dataflow/Engine/Private/Dataflow/DataflowEdNode.cpp
+++ b/Engine/Source/Runtime/Experimental/Dataflow/Engine/Private/Dataflow/DataflowEdNode.cpp
@@ -613,15 +613,18 @@ void UDataflowEdNode::Serialize(FArchive& Ar)
 	{
 		const UDataflow* const DataflowObject = UDataflow::GetDataflowAssetFromEdGraph(GetGraph());
 		bool bCanSerializeNode = !DataflowObject || DataflowObject->IsPerNodeTransactionSerializationEnabled();
-		if (bCanSerializeNode)
+
+		TSharedPtr<FDataflowNode> DataflowNode = GetDataflowNode();
+
+		// the DataflowNode may not be always valid so we need to serialize that part first so that Svaing and Loading behave exactly the same way 
+		bool bNodeSerializable = bCanSerializeNode && DataflowNode.IsValid();
+		Ar << bNodeSerializable;
+
+		if (bNodeSerializable)
 		{
-			if (TSharedPtr<FDataflowNode> DataflowNode = GetDataflowNode())
-			{
-				DataflowNode->SerializeInternal(Ar);
-			}
+			DataflowNode->SerializeInternal(Ar);
 		}
 	}
-
 	// some double level template type have unwanted spaces 
 	// now that the dataflow connection are no longer having spaces in their type name we need to fix that up on the pins
 	RemoveSpacesInAllPinTypes();
