Changed log verbosity for on-demand install request from Log -> Verbose.

#rnx

diff --git a/Engine/Source/Runtime/Experimental/IoStore/OnDemand/Private/OnDemandContentInstaller.cpp b/Engine/Source/Runtime/Experimental/IoStore/OnDemand/Private/OnDemandContentInstaller.cpp
--- a/Engine/Source/Runtime/Experimental/IoStore/OnDemand/Private/OnDemandContentInstaller.cpp
+++ b/Engine/Source/Runtime/Experimental/IoStore/OnDemand/Private/OnDemandContentInstaller.cpp
@@ -444,7 +444,7 @@ void FOnDemandContentInstaller::ProcessInstallRequest(FRequest& Request)
 	FRequest::FInstall& InstallRequest	= Request.AsInstall();
 	Request.Priority					= InstallRequest.Args.Priority;
 
-	UE_LOG(LogIoStoreOnDemand, Log, TEXT("Processing install request, SeqNo=%u, Priority=%d, ContentHandle=(%s)"),
+	UE_LOG(LogIoStoreOnDemand, Verbose, TEXT("Processing install request, SeqNo=%u, Priority=%d, ContentHandle=(%s)"),
 		Request.SeqNo, Request.Priority, *LexToString(InstallRequest.Args.ContentHandle));
 
 	if (InstallRequest.Args.ContentHandle.IsValid() == false)
@@ -633,7 +633,7 @@ void FOnDemandContentInstaller::ExecuteInstallRequest(FRequest& Request, bool bR
 	FRequest::FInstall& InstallRequest = Request.AsInstall();
 	check(InstallRequest.HttpRequestHandles.IsEmpty() == false);
 
-	UE_LOG(LogIoStoreOnDemand, Log, TEXT("Executing install request, SeqNo=%u, Priority=%d, ContentHandle=(%s)"),
+	UE_LOG(LogIoStoreOnDemand, Verbose, TEXT("Executing install request, SeqNo=%u, Priority=%d, ContentHandle=(%s)"),
 		Request.SeqNo, Request.Priority, *LexToString(InstallRequest.Args.ContentHandle));
 
 	if (Request.ErrorCode.load(std::memory_order_relaxed) == EIoErrorCode::Cancelled)
@@ -974,7 +974,7 @@ void FOnDemandContentInstaller::CompleteInstallRequest(FRequest& Request)
 		DurationCycles,
 		Request.ErrorCode);
 
-	UE_LOG(LogIoStoreOnDemand, Log, TEXT("Install request completed, Result='%s', SeqNo=%u, Priority=%d, ContentHandle=(%s), ContentSize=%.2lf MiB, InstallSize=%.2lf MiB, CacheHitRatio=%d%%, Duration=%dms"),
+	UE_LOG(LogIoStoreOnDemand, Verbose, TEXT("Install request completed, Result='%s', SeqNo=%u, Priority=%d, ContentHandle=(%s), ContentSize=%.2lf MiB, InstallSize=%.2lf MiB, CacheHitRatio=%d%%, Duration=%dms"),
 		GetIoErrorText(InstallResult.Status.GetErrorCode()), Request.SeqNo, Request.Priority, *LexToString(InstallRequest.Args.ContentHandle), double(InstallResult.Progress.TotalContentSize) / 1024.0 / 1024.0,
 		double(InstallResult.Progress.TotalInstallSize) / 1024.0 / 1024.0, int32(CacheHitRatio * 100), int32(InstallResult.DurationInSeconds * 1000));
 
