Move some validation to where it is parsed.

#rnx

diff --git a/Engine/Source/Runtime/Experimental/IoStore/HttpClient/Private/Activity.inl b/Engine/Source/Runtime/Experimental/IoStore/HttpClient/Private/Activity.inl
--- a/Engine/Source/Runtime/Experimental/IoStore/HttpClient/Private/Activity.inl
+++ b/Engine/Source/Runtime/Experimental/IoStore/HttpClient/Private/Activity.inl
@@ -88,6 +88,7 @@ public:
 	FAnsiStringView	GetStatusMessage() const;
 	int64			GetContentLength() const;
 	bool			IsChunked() const;
+	bool			HasBody() const;
 	void			ReadHeaders(FResponse::FHeaderSink Sink) const;
 	const char*		TEMP_GetMessageRight() const;
 
@@ -100,6 +101,7 @@ private:
 	uint16			StatusCode = 0;
 	uint8			bKeepAlive : 1;
 	uint8			bChunked : 1;
+	uint8			bHasBody : 1;
 	uint8			_Unused : 6;
 };
 
@@ -139,6 +141,12 @@ bool FResponseCloak::IsChunked() const
 	return bChunked;
 }
 
+////////////////////////////////////////////////////////////////////////////////
+bool FResponseCloak::HasBody() const
+{
+	return bHasBody;
+}
+
 ////////////////////////////////////////////////////////////////////////////////
 void FResponseCloak::ReadHeaders(FResponse::FHeaderSink Sink) const
 {
@@ -275,6 +283,35 @@ FOutcome FResponseCloak::Parse()
 		}
 	);
 
+	bHasBody = (StatusCode != 304);
+	bHasBody &= !GetMethod().Equals("HEAD", ESearchCase::IgnoreCase);
+	bHasBody &= bChunked | (ContentLength > 0);
+
+	if (Flags & Flag_Length)
+	{
+		if ((StatusCode <= 100) | (StatusCode == 204) | (Flags & Flag_XferEnc))
+		{
+			return FOutcome::Error("Unexpected Content-Length field", StatusCode);
+		}
+
+		if (ContentLength < 0)
+		{
+			return FOutcome::Error("Invalid Content-Length field", uint32(ContentLength));
+		}
+	}
+	else if (Flags & Flag_XferEnc)
+	{
+		if (!bChunked)
+		{
+			return FOutcome::Error("Unsupported Transfer-Encoding");
+		}
+		ContentLength = -1;
+	}
+	else if (bHasBody)
+	{
+		return FOutcome::Error("Missing Content-Length field");
+	}
+
 	return FOutcome::Ok();
 }
 
diff --git a/Engine/Source/Runtime/Experimental/IoStore/HttpClient/Private/Loop.inl b/Engine/Source/Runtime/Experimental/IoStore/HttpClient/Private/Loop.inl
--- a/Engine/Source/Runtime/Experimental/IoStore/HttpClient/Private/Loop.inl
+++ b/Engine/Source/Runtime/Experimental/IoStore/HttpClient/Private/Loop.inl
@@ -121,7 +121,7 @@ static FOutcome DoRecvMessage(FActivity* Activity, FHttpPeer& Peer)
 
 	bool bChunked = Buffer.IsChunked();
 	bool IsKeepAlive = Buffer.IsKeepAlive();
-	int64 ContentLength = Activity->NoContent ? 0 : Buffer.GetContentLength();
+	int64 ContentLength = Buffer.GetContentLength();
 
 	Activity->IsKeepAlive &= IsKeepAlive;
 
@@ -133,23 +133,15 @@ static FOutcome DoRecvMessage(FActivity* Activity, FHttpPeer& Peer)
 			Activity_SetError(Activity, "Chunked transfer encoding disabled (ERRNOCHUNK)");
 			return FOutcome::Error(Activity->ErrorReason);
 		}
-
-		ContentLength = -1;
-	}
-	else if (ContentLength < 0)
-	{
-		Activity_SetError(Activity, "Missing/invalid Content-Length header");
-		return FOutcome::Error(Activity->ErrorReason);
 	}
 
 	// Call out to the sink to get a content destination
 	FIoBuffer* PriorDest = Activity->Dest; // to retain unioned Host ptr (redirect uses it in sink)
 	Activity_CallSink(Activity);
-	Activity->NoContent |= (ContentLength == 0);
 
 	// Check the user gave us a destination for content
 	FIoBuffer& Dest = *(Activity->Dest);
-	if (Activity->NoContent == 0)
+	if (Buffer.HasBody())
 	{
 		if (&Dest == PriorDest)
 		{
@@ -187,7 +179,7 @@ static FOutcome DoRecvMessage(FActivity* Activity, FHttpPeer& Peer)
 	}
 
 	// HEAD methods
-	if (Activity->NoContent == 1)
+	if (!Buffer.HasBody())
 	{
 		if (AlreadyReceived)
 		{
@@ -1166,8 +1158,6 @@ FRequest FEventLoop::FImpl::Request(
 		Path = "/";
 	}
 
-	Activity->NoContent = (Method == "HEAD");
-
 	auto& Buffer = Activity->Buffer;
 	Buffer.Begin(Method, Path);
 	Buffer.AddHeader("Host", Activity->Host->GetHostName());
@@ -1502,6 +1492,7 @@ FRequest FEventLoop::Request(
 	Activity->bFollow30x = (Params->bAutoRedirect == true);
 	Activity->bAllowChunked = (Params->bAllowChunked == true);
 	Activity->LengthScore = 0;
+	Activity->NoContent = 0;
 
 	Activity_SetScore(Activity, Params->ContentSizeEst);
 
