Addressed sign mismatch

#rnx

diff --git a/Engine/Source/Runtime/Experimental/IoStore/HttpClient/Private/Misc.inl b/Engine/Source/Runtime/Experimental/IoStore/HttpClient/Private/Misc.inl
--- a/Engine/Source/Runtime/Experimental/IoStore/HttpClient/Private/Misc.inl
+++ b/Engine/Source/Runtime/Experimental/IoStore/HttpClient/Private/Misc.inl
@@ -72,7 +72,7 @@ static IAS_CVAR(int32, IdleMs,				50'000,	"Time in milliseconds to close idle co
 class FOutcome
 {
 public:
-	static FOutcome Ok(uint32 Result=0);
+	static FOutcome Ok(int32 Result=0);
 	static FOutcome Waiting();
 	static FOutcome Error(const char* Message, int32 Code=-1);
 	static FOutcome None()				{ return Error(""); }
@@ -95,7 +95,7 @@ private:
 			PTRINT	Code	: 16;
 		};
 		struct {
-			uint32	Result;
+			int32	Result;
 			uint32	Tag;
 		};
 	};
@@ -103,7 +103,7 @@ private:
 static_assert(sizeof(FOutcome) == sizeof(void*));
 
 ////////////////////////////////////////////////////////////////////////////////
-FOutcome FOutcome::Ok(uint32 Result)
+FOutcome FOutcome::Ok(int32 Result)
 {
 	FOutcome Outcome;
 	Outcome.Tag = OkTag;
diff --git a/Engine/Source/Runtime/Experimental/IoStore/HttpClient/Private/Test.inl b/Engine/Source/Runtime/Experimental/IoStore/HttpClient/Private/Test.inl
--- a/Engine/Source/Runtime/Experimental/IoStore/HttpClient/Private/Test.inl
+++ b/Engine/Source/Runtime/Experimental/IoStore/HttpClient/Private/Test.inl
@@ -391,7 +391,7 @@ static void ChunkedTest(const ANSICHAR* TestHost)
 
 	struct FTestState
 	{
-		uint32	Size = 0;
+		int32	Size = 0;
 		uint32	Hash = 0x493;
 		uint32	ExpectedHash = 0;
 		int32	ExpectedSize = -1;
