Iris - Better logging when handling end of bitstream issues. Reset read journal when deserializing huge objects.

diff --git a/Engine/Source/Runtime/Experimental/Iris/Core/Private/Iris/ReplicationSystem/ReplicationReader.cpp b/Engine/Source/Runtime/Experimental/Iris/Core/Private/Iris/ReplicationSystem/ReplicationReader.cpp
--- a/Engine/Source/Runtime/Experimental/Iris/Core/Private/Iris/ReplicationSystem/ReplicationReader.cpp
+++ b/Engine/Source/Runtime/Experimental/Iris/Core/Private/Iris/ReplicationSystem/ReplicationReader.cpp
@@ -2492,6 +2492,7 @@ void FReplicationReader::ReadObjects(FNetSerializationContext& Context, uint32 O
 
 	FNetBitStreamReader& Reader = *Context.GetBitStreamReader();
 	
+	const uint32 OriginalObjectBatchCountToRead = ObjectBatchCountToRead;
 	while (ObjectBatchCountToRead && !Context.HasErrorOrOverflow() && (!bGracefullyHandleReachingEndOfBitstream || (Reader.GetBitsLeft() > 0)))
 	{
 		ReadObjectBatch(Context, ReadObjectFlags);
@@ -2505,7 +2506,7 @@ void FReplicationReader::ReadObjects(FNetSerializationContext& Context, uint32 O
 		if (bGracefullyHandleReachingEndOfBitstream)
 		{
 			// ObjectBatchCountToRead should be zero at this point otherwise there's a problem on the writing side.
-			UE_CLOG(ObjectBatchCountToRead > 0, LogIris, Error, TEXT("Reached end of bitstream prior to reading all expected objects. %u left to read"), ObjectBatchCountToRead);
+			UE_CLOG(ObjectBatchCountToRead > 0, LogIris, Error, TEXT("Reached end of bitstream prior to reading all expected objects. %u/%u left to read. PacketId: %u %s"), ObjectBatchCountToRead, OriginalObjectBatchCountToRead, Context.GetPacketId(), *Context.PrintReadJournal());
 			ensure(ObjectBatchCountToRead == 0);
 		}
 	}
@@ -2533,6 +2534,8 @@ void FReplicationReader::ProcessHugeObjectAttachment(FNetSerializationContext& C
 	HugeObjectReader.InitBits(NetObjectBlob.GetRawData().GetData(), NetObjectBlob.GetRawDataBitCount());
 	FNetSerializationContext HugeObjectSerializationContext = Context.MakeSubContext(&HugeObjectReader);
 	HugeObjectSerializationContext.SetTraceCollector(HugeObjectTraceCollector);
+	UE_RESET_READ_JOURNAL(HugeObjectSerializationContext);
+	UE_ADD_READ_JOURNAL_ENTRY(HugeObjectSerializationContext, TEXT("HugeObject"));
 
 	UE_NET_TRACE_NAMED_SCOPE(HugeObjectTraceScope, HugeObjectState, HugeObjectReader, HugeObjectTraceCollector, ENetTraceVerbosity::Trace);
 
