Chaos: Fixing mixed usage of PopOutputData_External and PopFutureOutputData_External. The first always unqueues the future output data into a member variable, even if it doesn't return it, when that output data is passed the results time. If a user then calls PopFutureOutputData_External to get one future output data (for interpolation), it would not get that output data in the future, but either nothing or the one after that, which is unexpected.

diff --git a/Engine/Source/Runtime/Experimental/Chaos/Private/Chaos/Utilities.cpp b/Engine/Source/Runtime/Experimental/Chaos/Private/Chaos/Utilities.cpp
--- a/Engine/Source/Runtime/Experimental/Chaos/Private/Chaos/Utilities.cpp
+++ b/Engine/Source/Runtime/Experimental/Chaos/Private/Chaos/Utilities.cpp
@@ -15,7 +15,7 @@ namespace Chaos
 
 	namespace Utilities
 	{
-		FReal GetSolverPhysicsResultsTime(FPhysicsSolverBase* Solver)
+		FReal GetSolverPhysicsResultsTime(const FPhysicsSolverBase* Solver)
 		{
 			return Solver->GetPhysicsResultsTime_External();
 		}
diff --git a/Engine/Source/Runtime/Experimental/Chaos/Public/Chaos/SimCallbackObject.h b/Engine/Source/Runtime/Experimental/Chaos/Public/Chaos/SimCallbackObject.h
--- a/Engine/Source/Runtime/Experimental/Chaos/Public/Chaos/SimCallbackObject.h
+++ b/Engine/Source/Runtime/Experimental/Chaos/Public/Chaos/SimCallbackObject.h
@@ -26,7 +26,7 @@ class FSingleParticlePhysicsProxy;
 
 namespace Utilities
 {
-	CHAOS_API FReal GetSolverPhysicsResultsTime(FPhysicsSolverBase*);
+	CHAOS_API FReal GetSolverPhysicsResultsTime(const FPhysicsSolverBase*);
 }
 
 enum class ESimCallbackOptions : uint16
@@ -538,8 +538,8 @@ public:
 	}
 
 	/**
-	* Gets the output data produced in order up to and including SimTime. Typical usage is:
-	* while(auto Output = PopOutputData_External(ExternalTime)) { //process output }
+	* Gets the output data produced in order up to and including ResultsTime. Typical usage is:
+	* while(auto Output = PopOutputData_External()) { //process output }
 	*/
 	TSimCallbackOutputHandle<TOutputType> PopOutputData_External()
 	{
@@ -572,7 +572,11 @@ public:
 	*/
 	TSimCallbackOutputHandle<TOutputType> PopFutureOutputData_External()
 	{
-		OutputQueue.Dequeue(CurrentOutput_External);
+		if (!CurrentOutput_External)
+		{
+			OutputQueue.Dequeue(CurrentOutput_External);
+		}
+		
 		TOutputType* Output = CurrentOutput_External;
 		CurrentOutput_External = nullptr;
 		return TSimCallbackOutputHandle<TOutputType>(Output, this);
@@ -582,7 +586,7 @@ public:
 	* Check if the output queue is empty of data
 	* Can be used while iterating the queue through PopOutputData_External and PopFutureOutputData_External to check if the current data is the last data.
 	*/
-	bool IsOutputQueueEmpty_External()
+	bool IsOutputQueueEmpty_External() const
 	{
 		return OutputQueue.IsEmpty();
 	}
