Unused local variable names

#rnx

diff --git a/Engine/Source/Runtime/Experimental/IoStore/HttpClient/Private/Test.inl b/Engine/Source/Runtime/Experimental/IoStore/HttpClient/Private/Test.inl
--- a/Engine/Source/Runtime/Experimental/IoStore/HttpClient/Private/Test.inl
+++ b/Engine/Source/Runtime/Experimental/IoStore/HttpClient/Private/Test.inl
@@ -617,9 +617,7 @@ static void HttpTest(const ANSICHAR* TestHost, FCertRootsRef VerifyCert)
 	{
 		FRequest Request = Loop.Request("GET", BuildUrl("/data/67"), ReqParams);
 		Request.Accept(EMimeType::Json);
-
-		FTicket Ticket = Loop.Send(MoveTemp(Request), HashSink);
-
+		Loop.Send(MoveTemp(Request), HashSink);
 		WaitForLoopIdle();
 	}
 
@@ -627,18 +625,17 @@ static void HttpTest(const ANSICHAR* TestHost, FCertRootsRef VerifyCert)
 	{
 		FRequest Request = Loop.Get(BuildUrl("/data"), ReqParams).Accept(EMimeType::Json);
 
-		FTicket Tickets[] = {
-			Loop.Send(Loop.Get(BuildUrl("/data"), ReqParams).Accept(EMimeType::Json), HashSink),
-			Loop.Send(MoveTemp(Request), HashSink),
-			Loop.Send(Loop.Get("http://httpbin.org/get"), NoErrorSink),
-		};
+		Loop.Send(Loop.Get(BuildUrl("/data"), ReqParams).Accept(EMimeType::Json), HashSink),
+		Loop.Send(MoveTemp(Request), HashSink),
+		Loop.Send(Loop.Get("http://httpbin.org/get"), NoErrorSink),
+
 		WaitForLoopIdle();
 	}
 
 	// convenience
 	{
 		FRequest Request = Loop.Get(BuildUrl("/data"), ReqParams).Accept(EMimeType::Json);
-		FTicket Ticket = Loop.Send(MoveTemp(Request), HashSink);
+		Loop.Send(MoveTemp(Request), HashSink);
 		WaitForLoopIdle();
 	}
 
@@ -750,7 +747,6 @@ static void HttpTest(const ANSICHAR* TestHost, FCertRootsRef VerifyCert)
 		char AsciiData[257];
 		for (char& c : AsciiData)
 		{
-			int32 i = int32(ptrdiff_t(&c - AsciiData));
 			c = 0x41 + (MixTh() % 26);
 			c += (MixTh() & 2) ? 0x20 : 0;
 		}
