Chaos : fix geometry collection leaf particles not properly being waken up after having been in a sleep state once
- when created particles are properly being added to the active particle arrays in the SOA on the physics thread
- however when waken up they would not be added back to it

diff --git a/Engine/Source/Runtime/Experimental/Chaos/Public/Chaos/PBDRigidsSOAs.h b/Engine/Source/Runtime/Experimental/Chaos/Public/Chaos/PBDRigidsSOAs.h
--- a/Engine/Source/Runtime/Experimental/Chaos/Public/Chaos/PBDRigidsSOAs.h
+++ b/Engine/Source/Runtime/Experimental/Chaos/Public/Chaos/PBDRigidsSOAs.h
@@ -797,16 +797,16 @@ public:
 			check(ClusteredParticle->GetParticleType() != Chaos::EParticleType::GeometryCollection);
 			RemoveClusteredParticle(ClusteredParticle);
 			InsertClusteredParticle(ClusteredParticle);
+		}
 
-			// Cluster particle must go in the active array because the Cluster arrays are not in the active view
-			if (ClusteredParticle->ObjectState() != EObjectStateType::Dynamic || ClusteredParticle->Disabled())
-			{
-				RemoveFromActiveArray(ClusteredParticle->CastToRigidParticle(), /*bStillDirty=*/true);
-			}
-			else
-			{
-				AddToActiveArray(ClusteredParticle->CastToRigidParticle());
-			}
+		// Cluster particle must go in the active array because the Cluster arrays are not in the active view
+		if (ClusteredParticle->ObjectState() != EObjectStateType::Dynamic || ClusteredParticle->Disabled())
+		{
+			RemoveFromActiveArray(ClusteredParticle->CastToRigidParticle(), /*bStillDirty=*/true);
+		}
+		else
+		{
+			AddToActiveArray(ClusteredParticle->CastToRigidParticle());
 		}
 
 		CheckParticleViewMask(ClusteredParticle);
