Further consolidation of request-related logic

#rnx

diff --git a/Engine/Source/Runtime/Experimental/IoStore/HttpClient/Private/Activity.inl b/Engine/Source/Runtime/Experimental/IoStore/HttpClient/Private/Activity.inl
--- a/Engine/Source/Runtime/Experimental/IoStore/HttpClient/Private/Activity.inl
+++ b/Engine/Source/Runtime/Experimental/IoStore/HttpClient/Private/Activity.inl
@@ -17,10 +17,12 @@ public:
 	void			End();
 	FAnsiStringView	GetMethod() const;
 	void			EnumerateHeaders(FResponse::FHeaderSink Sink) const;
+	FOutcome		TrySendRequest(FHttpPeer& Peer);
 
 private:
 	FRequestCloak&	operator << (FAnsiStringView Value);
 	uint16			HeaderLeft;
+	uint16			AlreadySent = 0;
 	int16			HeaderRight = -1;
 	int8			MethodLength = -1;
 };
@@ -73,6 +75,33 @@ FRequestCloak& FRequestCloak::operator << (FAnsiStringView Value)
 	return *this;
 }
 
+////////////////////////////////////////////////////////////////////////////////
+FOutcome FRequestCloak::TrySendRequest(FHttpPeer& Peer)
+{
+	const char* SendData = GetData();
+	int32 SendSize = GetSize();
+
+	SendData += AlreadySent;
+	SendSize -= AlreadySent;
+	check(SendSize > 0);
+
+	FOutcome Outcome = Peer.Send(SendData, SendSize);
+	if (!Outcome.IsOk())
+	{
+		return Outcome;
+	}
+
+	int32 Result = Outcome.GetResult();
+	AlreadySent += uint16(Result);
+	if (AlreadySent == GetSize())
+	{
+		return FOutcome::Ok();
+	}
+
+	check(AlreadySent < GetSize());
+	return FOutcome::Waiting();
+}
+
 
 
 ////////////////////////////////////////////////////////////////////////////////
diff --git a/Engine/Source/Runtime/Experimental/IoStore/HttpClient/Private/Loop.inl b/Engine/Source/Runtime/Experimental/IoStore/HttpClient/Private/Loop.inl
--- a/Engine/Source/Runtime/Experimental/IoStore/HttpClient/Private/Loop.inl
+++ b/Engine/Source/Runtime/Experimental/IoStore/HttpClient/Private/Loop.inl
@@ -10,24 +10,15 @@ namespace UE::IoStore::HTTP
 ////////////////////////////////////////////////////////////////////////////////
 static FOutcome DoSend(FActivity* Activity, FHttpPeer& Peer)
 {
-
 #if IAS_HTTP_WITH_PERF
 	Activity->Stopwatch.SendStart();
 #endif
 
 	TRACE_CPUPROFILER_EVENT_SCOPE(IasHttp::DoSend);
 
-	FBuffer& Buffer = Activity->Buffer;
-	const char* SendData = Buffer.GetData();
-	int32 SendSize = Buffer.GetSize();
-
-	uint32 AlreadySent = Activity->StateParam;
-	SendData += AlreadySent;
-	SendSize -= AlreadySent;
-	check(SendSize > 0);
-
-	FOutcome Outcome = Peer.Send(SendData, SendSize);
+	auto& Buffer = Activity->Buffer;
 
+	FOutcome Outcome = Buffer.TrySendRequest(Peer);
 	if (Outcome.IsError())
 	{
 		Activity_SetError(Activity, Outcome);
@@ -41,19 +32,12 @@ static FOutcome DoSend(FActivity* Activity, FHttpPeer& Peer)
 
 	check(Outcome.IsOk());
 
-	int32 Result = Outcome.GetResult();
-	Activity->StateParam += Result;
-	if (Activity->StateParam < Buffer.GetSize())
-	{
-		return DoSend(Activity, Peer);
-	}
-
 #if IAS_HTTP_WITH_PERF
 	Activity->Stopwatch.SendEnd();
 #endif
 
-	Activity_ChangeState(Activity, FActivity::EState::RecvMessage, Buffer.GetSize());
-	return FOutcome::Ok(Result);
+	Activity_ChangeState(Activity, FActivity::EState::RecvMessage);
+	return FOutcome::Ok();
 }
 
 ////////////////////////////////////////////////////////////////////////////////
