Add logging for common operations on indexed cache storage

diff --git a/Engine/Source/Runtime/Experimental/IndexedCacheStorage/Private/IndexedCacheStorageManager.cpp b/Engine/Source/Runtime/Experimental/IndexedCacheStorage/Private/IndexedCacheStorageManager.cpp
--- a/Engine/Source/Runtime/Experimental/IndexedCacheStorage/Private/IndexedCacheStorageManager.cpp
+++ b/Engine/Source/Runtime/Experimental/IndexedCacheStorage/Private/IndexedCacheStorageManager.cpp
@@ -154,6 +154,9 @@ bool FIndexedCacheStorageManager::CreateCacheStorage(uint64 RequestNumberOfBytes
 
 	ensure(IndexedCacheEntriesMetaData[CacheIndex].MountRefCount == 0);
 	IndexedCacheEntriesMetaData[CacheIndex].bCacheExists = IIndexedCacheStorage::Get().CreateCacheStorage(RequestNumberOfBytes, CacheIndex);
+
+	UE_LOG(LogIndexedCacheStorageManager, Display, TEXT("CreateCacheStorage %ull bytes at index %d: %d"), RequestNumberOfBytes, CacheIndex, IndexedCacheEntriesMetaData[CacheIndex].bCacheExists);
+
 	return IndexedCacheEntriesMetaData[CacheIndex].bCacheExists;
 }
 
@@ -172,6 +175,7 @@ void FIndexedCacheStorageManager::DestroyCacheStorage(int32 CacheIndex)
 	}
 
 	ensure(IndexedCacheEntriesMetaData[CacheIndex].MountRefCount == 0);
+	UE_LOG(LogIndexedCacheStorageManager, Display, TEXT("DestroyCacheStorage  index %d"), CacheIndex);
 	IIndexedCacheStorage::Get().DestroyCacheStorage(CacheIndex);
 	IndexedCacheEntriesMetaData[CacheIndex].bCacheExists = false;
 }
@@ -213,8 +217,10 @@ FString FIndexedCacheStorageManager::MountCacheStorage(int32 CacheIndex)
 		bool bMounted = IIndexedCacheStorage::Get().MountCacheStorage(MountedPath, CacheIndex, IndexedCacheEntriesMetaData[CacheIndex].MountName);
 		if (!bMounted)
 		{
+			UE_LOG(LogIndexedCacheStorageManager, Error, TEXT("MountCacheStorage failed to mount cache %d"), CacheIndex);
 			return ActualMountName;
 		}
+		UE_LOG(LogIndexedCacheStorageManager, Display, TEXT("MountCacheStorage mounted cache %d at %s"), CacheIndex, *MountedPath);
 	}
 	else
 	{
@@ -241,6 +247,7 @@ void FIndexedCacheStorageManager::UnmountCacheStorage(int32 CacheIndex)
 	--IndexedCacheEntriesMetaData[CacheIndex].MountRefCount;
 	if (IndexedCacheEntriesMetaData[CacheIndex].MountRefCount == 0)
 	{
+		UE_LOG(LogIndexedCacheStorageManager, Display, TEXT("UnmountCacheStorage cache %d at %s"), CacheIndex, *IndexedCacheEntriesMetaData[CacheIndex].MountName);
 		IIndexedCacheStorage::Get().UnmountCacheStorage(*IndexedCacheEntriesMetaData[CacheIndex].MountName);
 	}
 }
