Moved some more validation out of the higher-level loop logic.

#rnx

diff --git a/Engine/Source/Runtime/Experimental/IoStore/HttpClient/Private/Activity.inl b/Engine/Source/Runtime/Experimental/IoStore/HttpClient/Private/Activity.inl
--- a/Engine/Source/Runtime/Experimental/IoStore/HttpClient/Private/Activity.inl
+++ b/Engine/Source/Runtime/Experimental/IoStore/HttpClient/Private/Activity.inl
@@ -327,6 +327,11 @@ FOutcome FResponseCloak::Parse()
 		{
 			return FOutcome::Error("Invalid Content-Length field", uint32(ContentLength));
 		}
+
+		if ((GetSize() - MessageRight) > ContentLength)
+		{
+			return FOutcome::Error("More data received that expected");
+		}
 	}
 	else if (Flags & Flag_XferEnc)
 	{
@@ -340,6 +345,10 @@ FOutcome FResponseCloak::Parse()
 	{
 		return FOutcome::Error("Missing Content-Length field");
 	}
+	else if (GetSize() > uint32(MessageRight))
+	{
+		return FOutcome::Error("Received content when none was expected");
+	}
 
 	return FOutcome::Ok();
 }
diff --git a/Engine/Source/Runtime/Experimental/IoStore/HttpClient/Private/Loop.inl b/Engine/Source/Runtime/Experimental/IoStore/HttpClient/Private/Loop.inl
--- a/Engine/Source/Runtime/Experimental/IoStore/HttpClient/Private/Loop.inl
+++ b/Engine/Source/Runtime/Experimental/IoStore/HttpClient/Private/Loop.inl
@@ -156,20 +156,10 @@ static FOutcome DoRecvMessage(FActivity* Activity, FHttpPeer& Peer)
 	const char* MessageRight = Buffer.TEMP_GetMessageRight();
 	const char* BufferRight = Buffer.GetData() + Buffer.GetSize();
 	uint32 AlreadyReceived = uint32(ptrdiff_t(BufferRight - MessageRight));
-	if (AlreadyReceived > uint32(ContentLength))
-	{
-		Activity_SetError(Activity, "More data received that expected");
-		return FOutcome::Error(Activity->ErrorReason);
-	}
 
 	// HEAD methods
 	if (!Buffer.HasBody())
 	{
-		if (AlreadyReceived)
-		{
-			Activity_SetError(Activity, "Received content when none was expected");
-			return FOutcome::Error(Activity->ErrorReason);
-		}
 		Activity_ChangeState(Activity, FActivity::EState::RecvDone);
 		return FOutcome::Ok();
 	}
