Renamed local variables

#rnx

diff --git a/Engine/Source/Runtime/Experimental/IoStore/HttpClient/Private/Loop.inl b/Engine/Source/Runtime/Experimental/IoStore/HttpClient/Private/Loop.inl
--- a/Engine/Source/Runtime/Experimental/IoStore/HttpClient/Private/Loop.inl
+++ b/Engine/Source/Runtime/Experimental/IoStore/HttpClient/Private/Loop.inl
@@ -87,9 +87,9 @@ static FOutcome DoRecvMessage(FActivity* Activity, FHttpPeer& Peer)
 	Activity->Stopwatch.RecvStart();
 #endif
 
-	auto& Buffer = Activity->Buffer;
+	auto& Response = Activity->Buffer;
 
-	FOutcome Outcome = Buffer.TryRecvResponse(Peer);
+	FOutcome Outcome = Response.TryRecvResponse(Peer);
 	if (Outcome.IsError())
 	{
 		Activity_SetError(Activity, Outcome);
@@ -103,9 +103,9 @@ static FOutcome DoRecvMessage(FActivity* Activity, FHttpPeer& Peer)
 
 	check(Outcome.IsOk());
 
-	bool bChunked = Buffer.IsChunked();
-	bool IsKeepAlive = Buffer.IsKeepAlive();
-	int64 ContentLength = Buffer.GetContentLength();
+	bool bChunked = Response.IsChunked();
+	bool IsKeepAlive = Response.IsKeepAlive();
+	int64 ContentLength = Response.GetContentLength();
 
 	Activity->IsKeepAlive &= IsKeepAlive;
 
@@ -125,7 +125,7 @@ static FOutcome DoRecvMessage(FActivity* Activity, FHttpPeer& Peer)
 
 	// Check the user gave us a destination for content
 	FIoBuffer& Dest = *(Activity->Dest);
-	if (Buffer.HasBody())
+	if (Response.HasBody())
 	{
 		if (&Dest == PriorDest)
 		{
@@ -153,12 +153,12 @@ static FOutcome DoRecvMessage(FActivity* Activity, FHttpPeer& Peer)
 	}
 
 	// Perhaps we have some of the content already?
-	const char* MessageRight = Buffer.TEMP_GetMessageRight();
-	const char* BufferRight = Buffer.GetData() + Buffer.GetSize();
+	const char* MessageRight = Response.TEMP_GetMessageRight();
+	const char* BufferRight = Response.GetData() + Response.GetSize();
 	uint32 AlreadyReceived = uint32(ptrdiff_t(BufferRight - MessageRight));
 
 	// HEAD methods
-	if (!Buffer.HasBody())
+	if (!Response.HasBody())
 	{
 		Activity_ChangeState(Activity, FActivity::EState::RecvDone);
 		return FOutcome::Ok();
