Fix IAX hostgroups immediately if one of the first few requests fail after a reconnection.

#rnx

- When a hostgroup disconnects/reconnects we attempt to clear the request history so that we track failures from scratch. However we were using the wrong TBitArray method for this and were not actually resetting anything.
- TBitArray::SetNum only initializes new elements added but we wanted to reset all elements to the given value and should've been using TBitArray::Init.

diff --git a/Engine/Source/Runtime/Experimental/IoStore/OnDemand/Private/IasHostGroup.cpp b/Engine/Source/Runtime/Experimental/IoStore/OnDemand/Private/IasHostGroup.cpp
--- a/Engine/Source/Runtime/Experimental/IoStore/OnDemand/Private/IasHostGroup.cpp
+++ b/Engine/Source/Runtime/Experimental/IoStore/OnDemand/Private/IasHostGroup.cpp
@@ -38,7 +38,7 @@ struct FBitWindow
 	void Reset(uint32 Count)
 	{
 		Count = FMath::RoundUpToPowerOfTwo(Count);
-		Bits.SetNum(int32(Count), false);
+		Bits.Init(false, Count);
 		Counter = 0;
 		Mask = Count - 1;
 	}
