Fix potential buffer over-read when receiving a chunked transfer response.

#rnx

- CrudeToInt takes a string view and we were only passing in a pointer (Cursor) so the constructor had to try to calculate the string length itself by looking for a null terminator but there is no guarantee that the data in Activity->Dest will be null terminated. Eventually we get unluky and Dest does not contain any 0 values and the string view will read past the end of the buffer and eventually raise an exception.
- We now explicitly create the string view and pass in the known length of the buffer, this is both safer and faster.

diff --git a/Engine/Source/Runtime/Experimental/IoStore/HttpClient/Private/Loop.inl b/Engine/Source/Runtime/Experimental/IoStore/HttpClient/Private/Loop.inl
--- a/Engine/Source/Runtime/Experimental/IoStore/HttpClient/Private/Loop.inl
+++ b/Engine/Source/Runtime/Experimental/IoStore/HttpClient/Private/Loop.inl
@@ -346,7 +346,8 @@ static FOutcome DoRecvStream(FActivity* Activity, FHttpPeer& Peer, int32& MaxRec
 			}
 			++HeaderLength;
 
-			ChunkSize = int32(CrudeToInt<16>(Cursor));
+			ChunkSize = int32(CrudeToInt<16>(FAnsiStringView(Cursor, (int32)View.GetSize())));
+
 			if (ChunkSize < 0)
 			{
 				Activity_SetError(Activity, "Unparsable chunk size");
