Iris - Reset NetBitStreamReader between batches to get correct read journal entry and clear overflow status.

diff --git a/Engine/Source/Runtime/Experimental/Iris/Core/Private/Iris/ReplicationSystem/ReplicationReader.cpp b/Engine/Source/Runtime/Experimental/Iris/Core/Private/Iris/ReplicationSystem/ReplicationReader.cpp
--- a/Engine/Source/Runtime/Experimental/Iris/Core/Private/Iris/ReplicationSystem/ReplicationReader.cpp
+++ b/Engine/Source/Runtime/Experimental/Iris/Core/Private/Iris/ReplicationSystem/ReplicationReader.cpp
@@ -2670,6 +2670,7 @@ void FReplicationReader::ProcessQueuedBatches()
 			UE_LOG(LogIris, Verbose, TEXT("ProcessQueuedBatches processing %d queued batches for Handle %s "), PendingBatchData.QueuedDataChunks.Num(), *PendingBatchData.Handle.ToString());
 
 			// Reset for batch
+			Reader = FNetBitStreamReader();
 			Context.ResetErrorContext();
 			Context.SetErrorHandleContext(PendingBatchData.Handle);
 			UE_RESET_READ_JOURNAL(Context);
