Iris - Async loading failure debugging cvars.

net.Iris.AsyncLoading.FailPackageName Name - fails async loading all packages containing the package name.
net.Iris.AsyncLoading.FailNextLoad 1 - fails the next async load regardless of the name and resets the cvar.
net.Iris.AsyncLoading.FailAllLoads 1 - fails all loads until the cvar is set to 0.

diff --git a/Engine/Source/Runtime/Experimental/Iris/Core/Private/Iris/ReplicationSystem/ObjectReferenceCache.cpp b/Engine/Source/Runtime/Experimental/Iris/Core/Private/Iris/ReplicationSystem/ObjectReferenceCache.cpp
--- a/Engine/Source/Runtime/Experimental/Iris/Core/Private/Iris/ReplicationSystem/ObjectReferenceCache.cpp
+++ b/Engine/Source/Runtime/Experimental/Iris/Core/Private/Iris/ReplicationSystem/ObjectReferenceCache.cpp
@@ -30,6 +30,10 @@
 #	define UE_NET_VALIDATE_REFERENCECACHE 0
 #endif
 
+#ifndef UE_NET_OBJECTREFERENCECACHE_ASYNCLOAD_DEBUG
+#	define UE_NET_OBJECTREFERENCECACHE_ASYNCLOAD_DEBUG !(UE_BUILD_SHIPPING)
+#endif
+
 // Don't compile verbose logs in Shipping builds
 #if (UE_BUILD_SHIPPING)
 #	define UE_NET_REFERENCECACHE_LOG_COMPILE_VERBOSITY Log
@@ -52,6 +56,17 @@ namespace UE::Net::Private
 static bool bIrisAllowAsyncLoading = true;
 FAutoConsoleVariableRef CVarIrisAllowAsyncLoading(TEXT("net.iris.AllowAsyncLoading"), bIrisAllowAsyncLoading, TEXT("Flag to allow or disallow async loading when using iris replication. Note: net.allowAsyncLoading must also be enabled."), ECVF_Default);
 
+#if UE_NET_OBJECTREFERENCECACHE_ASYNCLOAD_DEBUG
+static FString IrisAsyncLoadingFailPackageName;
+FAutoConsoleVariableRef CVarIrisAsyncLoadingFailPackageName(TEXT("net.Iris.AsyncLoading.FailPackageName"), IrisAsyncLoadingFailPackageName, TEXT("Always fail loading the specified package name. For debugging purposes."), ECVF_Default);
+
+static bool bIrisAsyncLoadingFailNextLoad;
+FAutoConsoleVariableRef CVarIrisAsyncLoadingFailNextLoad(TEXT("net.Iris.AsyncLoading.FailNextLoad"), bIrisAsyncLoadingFailNextLoad, TEXT("Simulate failing the next async loading request. For debugging purposes."), ECVF_Default);
+
+static bool bIrisAsyncLoadingFailAllLoads;
+FAutoConsoleVariableRef CVarIrisAsyncLoadingFailAllLoads(TEXT("net.Iris.AsyncLoading.FailAllLoads"), bIrisAsyncLoadingFailAllLoads, TEXT("Simulate failing all async loading requests. For debugging purposes."), ECVF_Default);
+#endif
+
 static bool bIrisEnableRemoveReferenceFallbackPath = true;
 FAutoConsoleVariableRef CVarIrisEnableRemoveReferenceFallbackPath(TEXT("net.iris.ObjectReferenceCache.EnableRemoveReferenceFallbackPath"), bIrisEnableRemoveReferenceFallbackPath, TEXT("Whether to enable the slow fallback path or not when ObjectReferenceCache::RemoveReference is passed a null pointer object."), ECVF_Default);
 
@@ -1982,12 +1997,23 @@ void FObjectReferenceCache::StartAsyncLoadingPackage(FCachedNetObjectReference&
 
 	//DelinquentAsyncLoads.MaxConcurrentAsyncLoads = FMath::Max<uint32>(DelinquentAsyncLoads.MaxConcurrentAsyncLoads, PendingAsyncLoadRequests.Num());
 
-	LoadPackageAsync(PackagePath.ToString(), FLoadPackageAsyncDelegate::CreateWeakLambda(ReplicationSystem, 
-		[this](const FName& PackageName, UPackage* Package, EAsyncLoadingResult::Type Result)
-		{
-			AsyncPackageCallback(PackageName, Package, Result);
-		}
-	));
+#if UE_NET_OBJECTREFERENCECACHE_ASYNCLOAD_DEBUG
+	if (bIrisAsyncLoadingFailAllLoads || bIrisAsyncLoadingFailNextLoad || (IrisAsyncLoadingFailPackageName.Len() > 0 && PackagePath.ToString().Find(IrisAsyncLoadingFailPackageName) != INDEX_NONE))
+	{
+		bIrisAsyncLoadingFailNextLoad = false;
+		// Not thread safe. We don't want to load the package so we call our fail callback immediately.
+		AsyncPackageForcedFailCallback(PackagePath, nullptr, EAsyncLoadingResult::Failed);
+	}
+	else
+#endif
+	{
+		LoadPackageAsync(PackagePath.ToString(), FLoadPackageAsyncDelegate::CreateWeakLambda(ReplicationSystem, 
+			[this](const FName& PackageName, UPackage* Package, EAsyncLoadingResult::Type Result)
+			{
+				AsyncPackageCallback(PackageName, Package, Result);
+			}
+		));
+	}
 }
 
 void FObjectReferenceCache::AsyncPackageCallback(const FName& PackageName, UPackage* Package, EAsyncLoadingResult::Type Result)
@@ -2071,6 +2097,34 @@ void FObjectReferenceCache::AsyncPackageCallback(const FName& PackageName, UPack
 	}
 }
 
+void FObjectReferenceCache::AsyncPackageForcedFailCallback(const FName& PackageName, UPackage* Package, EAsyncLoadingResult::Type Result)
+{
+#if UE_NET_OBJECTREFERENCECACHE_ASYNCLOAD_DEBUG
+	UE_LOG_REFERENCECACHE(Warning, TEXT("AsyncPackageCallback: Force fail async package load. Path: %s"), ToCStr(PackageName.ToString()));
+
+	if (FPendingAsyncLoadRequest const* const PendingLoadRequest = PendingAsyncLoadRequests.Find(PackageName))
+	{
+		for (FNetRefHandle RefHandleToProcess : PendingLoadRequest->NetRefHandles)
+		{
+			if (FCachedNetObjectReference* CacheObject = ReferenceHandleToCachedReference.Find(RefHandleToProcess))
+			{
+				CacheObject->bIsPending = false;
+				CacheObject->bIsBroken = true;
+
+				if (UObject* Object = CacheObject->Object.Get())
+				{
+					// Something is loaded, we can now resolve queued data
+					UpdateTrackedQueuedBatchObjectReference(RefHandleToProcess, Object);
+
+				}
+			}
+		}
+
+		PendingAsyncLoadRequests.Remove(PackageName);
+	}
+#endif
+}
+
 void FObjectReferenceCache::SetAsyncLoadMode(const EAsyncLoadMode NewMode)
 {
 	AsyncLoadMode = NewMode;
diff --git a/Engine/Source/Runtime/Experimental/Iris/Core/Private/Iris/ReplicationSystem/ObjectReferenceCache.h b/Engine/Source/Runtime/Experimental/Iris/Core/Private/Iris/ReplicationSystem/ObjectReferenceCache.h
--- a/Engine/Source/Runtime/Experimental/Iris/Core/Private/Iris/ReplicationSystem/ObjectReferenceCache.h
+++ b/Engine/Source/Runtime/Experimental/Iris/Core/Private/Iris/ReplicationSystem/ObjectReferenceCache.h
@@ -145,7 +145,7 @@ public:
 
 	void SetAsyncLoadMode(const EAsyncLoadMode NewMode);
 	bool ShouldAsyncLoad() const;
-	void AsyncPackageCallback(const FName& PackageName, UPackage* Package, EAsyncLoadingResult::Type Result);
+
 
 	// While async loading of pending must be mapped references we need to maintain references to already resolved objects as there will be no instance referencing them
 	void AddReferencedObjects(FReferenceCollector& ReferenceCollector);
@@ -159,6 +159,8 @@ public:
 	const TCHAR* GetObjectRelativePath(FNetRefHandle NetRefHandle) const;
 
 private:
+	void AsyncPackageCallback(const FName& PackageName, UPackage* Package, EAsyncLoadingResult::Type Result);
+	void AsyncPackageForcedFailCallback(const FName& PackageName, UPackage* Package, EAsyncLoadingResult::Type Result);
 
 	struct FCachedNetObjectReference
 	{
