Chaos - Fix race condition in FPBDRigidsSolver::UnregisterObject when bodies are destroyed asynchronously.

diff --git a/Engine/Source/Runtime/Experimental/Chaos/Private/PBDRigidsSolver.cpp b/Engine/Source/Runtime/Experimental/Chaos/Private/PBDRigidsSolver.cpp
--- a/Engine/Source/Runtime/Experimental/Chaos/Private/PBDRigidsSolver.cpp
+++ b/Engine/Source/Runtime/Experimental/Chaos/Private/PBDRigidsSolver.cpp
@@ -856,8 +856,12 @@ namespace Chaos
 		// Remove the proxy from the GT proxy map
 		FUniqueIdx UniqueIdx = Proxy->GetGameThreadAPI().UniqueIdx();
 
-		FIgnoreCollisionManager& CollisionManager = GetEvolution()->GetBroadPhase().GetIgnoreCollisionManager();
 		{
+			// The collision manager is also modified from FPBDRigidsSolver::PushPhysicsState which uses the 
+			// MarshallingManager lock so we need to lock this here as well.
+			// @todo(chaos): should we lock this for the whole function?
+			Chaos::Private::FChaosAsyncInitBodyWriteScopeLock MarshallingManagerLock(MarshallingManager.GetMarshallingManagerLock());
+			FIgnoreCollisionManager& CollisionManager = GetEvolution()->GetBroadPhase().GetIgnoreCollisionManager();
 			int32 ExternalTimestamp = GetMarshallingManager().GetExternalTimestamp_External();
 			FIgnoreCollisionManager::FDeactivationSet& PendingMap = CollisionManager.GetPendingDeactivationsForGameThread(ExternalTimestamp);
 			PendingMap.Add(UniqueIdx);
