Let waiting outcome communicate a result value too.

#rnx

diff --git a/Engine/Source/Runtime/Experimental/IoStore/HttpClient/Private/Misc.inl b/Engine/Source/Runtime/Experimental/IoStore/HttpClient/Private/Misc.inl
--- a/Engine/Source/Runtime/Experimental/IoStore/HttpClient/Private/Misc.inl
+++ b/Engine/Source/Runtime/Experimental/IoStore/HttpClient/Private/Misc.inl
@@ -73,7 +73,7 @@ class FOutcome
 {
 public:
 	static FOutcome Ok(int32 Result=0);
-	static FOutcome Waiting();
+	static FOutcome Waiting(int32 Result=0);
 	static FOutcome Error(const char* Message, int32 Code=-1);
 	static FOutcome None()				{ return Error(""); }
 	bool			IsError() const		{ return Message < 0x8000'0000'0000; }
@@ -81,7 +81,7 @@ public:
 	bool			IsOk() const		{ return Tag == OkTag; }
 	FAnsiStringView GetMessage() const	{ check(IsError()); return (const char*)(Message); }
 	int32			GetErrorCode() const{ check(IsError()); return int32(Code); }
-	uint32			GetResult() const	{ check(IsOk()); return Result; }
+	uint32			GetResult() const	{ check(IsOk()|IsWaiting()); return Result; }
 
 private:
 					FOutcome() = default;
@@ -112,10 +112,11 @@ FOutcome FOutcome::Ok(int32 Result)
 }
 
 ////////////////////////////////////////////////////////////////////////////////
-FOutcome FOutcome::Waiting()
+FOutcome FOutcome::Waiting(int32 Result)
 {
 	FOutcome Outcome;
 	Outcome.Tag = WaitTag;
+	Outcome.Result = Result;
 	return Outcome;
 }
 
