Tweak support test .py to run out of the tree.

#rnx

diff --git a/Engine/Source/Runtime/Experimental/IoStore/HttpClient/Private/TestServer.py b/Engine/Source/Runtime/Experimental/IoStore/HttpClient/Private/TestServer.py
--- a/Engine/Source/Runtime/Experimental/IoStore/HttpClient/Private/TestServer.py
+++ b/Engine/Source/Runtime/Experimental/IoStore/HttpClient/Private/TestServer.py
@@ -424,15 +424,19 @@ def main():
     for item in Path(__file__).parents:
         if (item / "GenerateProjectFiles.bat").is_file():
             os.chdir(item)
+            with_tls = True
             break
     else:
-        assert False
-
-    print("\n## generating certificates")
-    openssl_bin = os.getenv("OPENSSL_BIN", "openssl")
-    if (x := Path(".") / "Engine/Binaries/DotNET/IOS/openssl.exe").is_file():
-        openssl_bin = str(x)
-    root_cert, server_cert, server_key = gen_test_certs(openssl_bin)
+        with_tls = False
+
+    if with_tls:
+        print("\n## generating certificates")
+        openssl_bin = os.getenv("OPENSSL_BIN", "openssl")
+        if (x := Path(".") / "Engine/Binaries/DotNET/IOS/openssl.exe").is_file():
+            openssl_bin = str(x)
+        root_cert, server_cert, server_key = gen_test_certs(openssl_bin)
+    else:
+        root_cert = None
 
     httpd_port = int(os.getenv("IasTestServerPort", 0))
     httpd_port = httpd_port or (int(random.random() * 0x8000) + 0x4000)
@@ -442,8 +446,9 @@ def main():
 
     print("\n## starting servers")
     start_svc(proxy_loop, httpd_port)
-    start_svc(tls_proxy_loop, httpd_port, root_cert, server_key, server_cert)
     start_svc(plain_httpd_loop, httpd_port, root_cert)
+    if with_tls:
+        start_svc(tls_proxy_loop, httpd_port, root_cert, server_key, server_cert)
 
     time.sleep(0.25)
     print("\n## ready")
