Match tests to what is compile in

#rnx

diff --git a/Engine/Source/Runtime/Experimental/IoStore/HttpClient/Private/Test.inl b/Engine/Source/Runtime/Experimental/IoStore/HttpClient/Private/Test.inl
--- a/Engine/Source/Runtime/Experimental/IoStore/HttpClient/Private/Test.inl
+++ b/Engine/Source/Runtime/Experimental/IoStore/HttpClient/Private/Test.inl
@@ -204,6 +204,8 @@ static void ThrottleTest(FAnsiStringView TestUrl)
 	}
 }
 
+#if IAS_HTTP_HAS_OPENSSL
+
 ////////////////////////////////////////////////////////////////////////////////
 static void TlsLoadRootCerts()
 {
@@ -261,6 +263,8 @@ static void TlsTest()
 	}
 }
 
+#endif // IAS_HTTP_HAS_OPENSSL
+
 ////////////////////////////////////////////////////////////////////////////////
 static void RedirectTest(const ANSICHAR* TestHost, FCertRootsRef VerifyCert)
 {
@@ -313,6 +317,13 @@ static void RedirectTest(const ANSICHAR* TestHost, FCertRootsRef VerifyCert)
 
 	for (auto ReTest : { ReAbs, ReAbsTls, ReRel, ReRelTls })
 	{
+#if !IAS_HTTP_HAS_OPENSSL
+		if (ReTest & 1)
+		{
+			continue;
+		}
+#endif
+
 		RequestParams.VerifyCert = (ReTest & 1) ? VerifyCert : 0;
 		RecvCount = 0;
 		for (uint32 Code : TestCodes)
@@ -333,6 +344,13 @@ static void RedirectTest(const ANSICHAR* TestHost, FCertRootsRef VerifyCert)
 
 	for (auto ReTest : { ReAbs, ReAbsTls, ReRel, ReRelTls })
 	{
+#if !IAS_HTTP_HAS_OPENSSL
+		if (ReTest & 1)
+		{
+			continue;
+		}
+#endif
+
 		FConnectionPool::FParams Params;
 		Params.SetHostFromUrl(BuildUrl(ReTest, 0));
 		Params.VerifyCert = (ReTest & 1) ? VerifyCert : 0;
@@ -921,6 +939,7 @@ IOSTOREHTTPCLIENT_API void IasHttpTest(const ANSICHAR* TestHost="localhost", uin
 
 	MiscTest();
 
+#if IAS_HTTP_HAS_OPENSSL
 	FCertRoots TestServerCaChain;
 	{
 		TAnsiStringBuilder<64> CaUrl;
@@ -944,14 +963,20 @@ IOSTOREHTTPCLIENT_API void IasHttpTest(const ANSICHAR* TestHost="localhost", uin
 		TestServerCaChain = FCertRoots(CertBuffer.GetView());
 	}
 	FCertRootsRef TestServerCertRef = FCertRoots::Explicit(TestServerCaChain);
+#else
+	FCertRootsRef TestServerCertRef = FCertRoots::NoTls();
+#endif // IAS_HTTP_HAS_OPENSSL
 
 	SeedHttp(TestHost, Seed);
 	HttpTest(TestHost, FCertRoots::NoTls());
-	HttpTest(TestHost, TestServerCertRef);
 	ChunkedTest(TestHost);
 	RedirectTest(TestHost, TestServerCertRef);
+
+#if IAS_HTTP_HAS_OPENSSL
+	HttpTest(TestHost, TestServerCertRef);
 	TlsLoadRootCerts();
 	TlsTest();
+#endif
 }
 
 #endif // !SHIP|TEST
