Iris
- Added Conditional variants fo IRIS_PROFILER macros

diff --git a/Engine/Source/Runtime/Experimental/Iris/Core/Public/Iris/Core/IrisProfiler.h b/Engine/Source/Runtime/Experimental/Iris/Core/Public/Iris/Core/IrisProfiler.h
--- a/Engine/Source/Runtime/Experimental/Iris/Core/Public/Iris/Core/IrisProfiler.h
+++ b/Engine/Source/Runtime/Experimental/Iris/Core/Public/Iris/Core/IrisProfiler.h
@@ -35,29 +35,39 @@
 #		include "c:/Program Files/Superluminal/Performance/API/include/Superluminal/PerformanceAPI.h"
 #		include "HAL/PreprocessorHelpers.h"
 #		pragma comment (lib, "c:/Program Files/Superluminal/Performance/API/lib/x64/PerformanceAPI_MD.lib")
-#		define IRIS_PROFILER_SCOPE(x) PERFORMANCEAPI_INSTRUMENT(PREPROCESSOR_TO_STRING(x))
-#		define IRIS_PROFILER_SCOPE_TEXT(x) PERFORMANCEAPI_INSTRUMENT_DATA(PREPROCESSOR_JOIN(IrisProfilerScope, __LINE__), x)
+#		define IRIS_PROFILER_SCOPE(X) PERFORMANCEAPI_INSTRUMENT(PREPROCESSOR_TO_STRING(X))
+#		define IRIS_PROFILER_SCOPE_CONDITIONAL(X,Cond) PERFORMANCEAPI_INSTRUMENT(PREPROCESSOR_TO_STRING(X))
+#		define IRIS_PROFILER_SCOPE_TEXT(X) PERFORMANCEAPI_INSTRUMENT_DATA(PREPROCESSOR_JOIN(IrisProfilerScope, __LINE__), X)
+#		define IRIS_PROFILER_SCOPE_TEXT_CONDITIONAL(X, Cond) PERFORMANCEAPI_INSTRUMENT_DATA(PREPROCESSOR_JOIN(IrisProfilerScope, __LINE__), X)
 #	else
 #		include "ProfilingDebugging/CpuProfilerTrace.h"
-#		define IRIS_PROFILER_SCOPE(x) TRACE_CPUPROFILER_EVENT_SCOPE(x)
+#		define IRIS_PROFILER_SCOPE(X) TRACE_CPUPROFILER_EVENT_SCOPE(X)
+#		define IRIS_PROFILER_SCOPE_CONDITIONAL(X,Cond) TRACE_CPUPROFILER_EVENT_SCOPE_CONDITIONAL(X, Cond)
 #		define IRIS_PROFILER_SCOPE_TEXT(X) TRACE_CPUPROFILER_EVENT_SCOPE_TEXT(X)
+#		define IRIS_PROFILER_SCOPE_TEXT_CONDITIONAL(X, Cond) TRACE_CPUPROFILER_EVENT_SCOPE_TEXT_CONDITIONAL(X, Cond)
 #	endif
 #else
 #	define PERFORMANCEAPI_ENABLED 0
-#	define IRIS_PROFILER_SCOPE(x)
+#	define IRIS_PROFILER_SCOPE(X)
+#	define IRIS_PROFILER_SCOPE_CONDITIONAL(X, Cond)
 #	define IRIS_PROFILER_SCOPE_TEXT(X)
+#	define IRIS_PROFILER_SCOPE_TEXT_CONDITIONAL(X, Cond)
 #endif
 
 #if UE_IRIS_PROFILER_ENABLE_PROTOCOL_NAMES
-#	define IRIS_PROFILER_PROTOCOL_NAME(x) IRIS_PROFILER_SCOPE_TEXT(x)
+#	define IRIS_PROFILER_PROTOCOL_NAME(X) IRIS_PROFILER_SCOPE_TEXT(X)
+#	define IRIS_PROFILER_PROTOCOL_NAME_CONDITIONAL(X, Cond) IRIS_PROFILER_SCOPE_TEXT_CONDITIONAL(X, Cond)
 #else
-#	define IRIS_PROFILER_PROTOCOL_NAME(x)
+#	define IRIS_PROFILER_PROTOCOL_NAME(X)
+#	define IRIS_PROFILER_PROTOCOL_NAME_CONDITIONAL(X, Cond)
 #endif
 
 #if UE_IRIS_PROFILER_ENABLE_VERBOSE
-#	define IRIS_PROFILER_SCOPE_VERBOSE(x) IRIS_PROFILER_SCOPE(x);
+#	define IRIS_PROFILER_SCOPE_VERBOSE(X) IRIS_PROFILER_SCOPE(X);
+#	define IRIS_PROFILER_SCOPE_VERBOSE_CONDITIONAL(X) IRIS_PROFILER_SCOPE_CONDITIONAL(X, Cond);
 #else
-#	define IRIS_PROFILER_SCOPE_VERBOSE(x)
+#	define IRIS_PROFILER_SCOPE_VERBOSE(X)
+#	define IRIS_PROFILER_SCOPE_VERBOSE_CONDITIONAL(X)
 #endif
 
 #ifndef IRIS_CLIENT_PROFILER_ENABLE
