Appease PVS

#rnx

diff --git a/Engine/Source/Runtime/Experimental/IoStore/HttpClient/Private/Misc.inl b/Engine/Source/Runtime/Experimental/IoStore/HttpClient/Private/Misc.inl
--- a/Engine/Source/Runtime/Experimental/IoStore/HttpClient/Private/Misc.inl
+++ b/Engine/Source/Runtime/Experimental/IoStore/HttpClient/Private/Misc.inl
@@ -81,7 +81,7 @@ public:
 	bool			IsOk() const		{ return Tag == OkTag; }
 	FAnsiStringView GetMessage() const	{ check(IsError()); return (const char*)(Message); }
 	int32			GetErrorCode() const{ check(IsError()); return int32(Code); }
-	uint32			GetResult() const	{ check(IsOk()|IsWaiting()); return Result; }
+	uint32			GetResult() const	{ check(!IsError()); return Result; }
 
 private:
 					FOutcome() = default;
