Further strategic encapsulation of logiv embedded in activities.

#rnx

diff --git a/Engine/Source/Runtime/Experimental/IoStore/HttpClient/Private/Activity.inl b/Engine/Source/Runtime/Experimental/IoStore/HttpClient/Private/Activity.inl
--- a/Engine/Source/Runtime/Experimental/IoStore/HttpClient/Private/Activity.inl
+++ b/Engine/Source/Runtime/Experimental/IoStore/HttpClient/Private/Activity.inl
@@ -355,6 +355,101 @@ FOutcome FResponseCloak::Parse()
 
 
 
+////////////////////////////////////////////////////////////////////////////////
+class FBodyCloak
+	: public FResponseCloak
+{
+public:
+	using FResponseCloak::FResponseCloak;
+
+	FOutcome	TryRecv(FMutableMemoryView Dest, FHttpPeer& Peer);
+
+private:
+	enum class EState : uint8 { Init, Recv, Done };
+
+	FOutcome	Init(FMutableMemoryView Dest, FHttpPeer& Peer);
+	FOutcome	Recv(FMutableMemoryView Dest, FHttpPeer& Peer);
+	FMemoryView	Overspill;
+	int64		Remaining = 0;
+	EState		State = EState::Init;
+};
+
+////////////////////////////////////////////////////////////////////////////////
+FOutcome FBodyCloak::TryRecv(FMutableMemoryView Dest, FHttpPeer& Peer)
+{
+	FOutcome Outcome = FOutcome::None();
+
+	switch (State)
+	{
+	case EState::Init:		Outcome = Init(Dest, Peer); break;
+	case EState::Recv:		Outcome = Recv(Dest, Peer); break;
+	case EState::Done:		Outcome = FOutcome::Error("Transaction is complete"); break;
+	}
+
+	if (Outcome.IsOk())
+	{
+		State = EState::Done;
+	}
+
+	return Outcome;
+}
+
+////////////////////////////////////////////////////////////////////////////////
+FOutcome FBodyCloak::Init(FMutableMemoryView Dest, FHttpPeer& Peer)
+{
+	const char* DataEnd = GetData() + GetSize();
+	const char* MessageEnd = TEMP_GetMessageRight();
+	uint32 AlreadyReceived = uint32(ptrdiff_t(DataEnd - MessageEnd));
+	if (AlreadyReceived)
+	{
+		Overspill = { (const uint8*)MessageEnd, AlreadyReceived };
+	}
+	
+	Remaining = GetContentLength();
+	check(Remaining > 0);
+	State = EState::Recv;
+	return Recv(Dest, Peer);
+}
+
+////////////////////////////////////////////////////////////////////////////////
+FOutcome FBodyCloak::Recv(FMutableMemoryView Dest, FHttpPeer& Peer)
+{
+	int64 OverspillSize = Overspill.GetSize();
+	if (OverspillSize != 0)
+	{
+		int64 Size = FMath::Min<int64>(Dest.GetSize(), OverspillSize);
+		::memcpy(Dest.GetData(), Overspill.GetData(), Size);
+		Dest = Dest.Mid(Size);
+		Overspill = Overspill.Mid(Size);
+		Remaining -= Size;
+		if (Remaining == 0 || Dest.IsEmpty())
+		{
+			int32 Result = int32(Size);
+			return Remaining ? FOutcome::Waiting(Result) : FOutcome::Ok(Result);
+		}
+	}
+
+	auto* Cursor = (char*)(Dest.GetData());
+	uint32 Size = uint32(FMath::Min<int64>(Dest.GetSize(), Remaining));
+	FOutcome Outcome = Peer.Recv(Cursor, Size);
+	if (Outcome.IsError())
+	{
+		return Outcome;
+	}
+
+	int32 Result = Outcome.GetResult();
+	Remaining -= Result;
+	if (Remaining < 0)
+	{
+		return FOutcome::Error("Unexpectedly received too much", int32(Remaining));
+	}
+
+	Result += int32(OverspillSize);
+	return Remaining ? FOutcome::Waiting(Result) : FOutcome::Ok(Result);
+}
+
+
+
 #if IAS_HTTP_WITH_PERF
 
 ////////////////////////////////////////////////////////////////////////////////
@@ -445,7 +540,7 @@ struct alignas(16) FActivity
 	};
 	UPTRINT				SinkParam;
 	FTicketSink			Sink;
-	FResponseCloak		Buffer;
+	FBodyCloak			Buffer;
 };
 
 ////////////////////////////////////////////////////////////////////////////////
diff --git a/Engine/Source/Runtime/Experimental/IoStore/HttpClient/Private/Loop.inl b/Engine/Source/Runtime/Experimental/IoStore/HttpClient/Private/Loop.inl
--- a/Engine/Source/Runtime/Experimental/IoStore/HttpClient/Private/Loop.inl
+++ b/Engine/Source/Runtime/Experimental/IoStore/HttpClient/Private/Loop.inl
@@ -152,11 +152,6 @@ static FOutcome DoRecvMessage(FActivity* Activity, FHttpPeer& Peer)
 		}
 	}
 
-	// Perhaps we have some of the content already?
-	const char* MessageRight = Response.TEMP_GetMessageRight();
-	const char* BufferRight = Response.GetData() + Response.GetSize();
-	uint32 AlreadyReceived = uint32(ptrdiff_t(BufferRight - MessageRight));
-
 	// HEAD methods
 	if (!Response.HasBody())
 	{
@@ -167,8 +162,18 @@ static FOutcome DoRecvMessage(FActivity* Activity, FHttpPeer& Peer)
 	// We're all set to go and get content
 	check(Activity->Dest != nullptr);
 
-	auto NextState = bChunked ? FActivity::EState::RecvStream : FActivity::EState::RecvContent;
-	Activity_ChangeState(Activity, NextState, AlreadyReceived);
+	if (!bChunked)
+	{
+		Activity_ChangeState(Activity, FActivity::EState::RecvContent);
+		return FOutcome::Ok();
+	}
+
+	// Perhaps we have some of the content already?
+	const char* MessageRight = Response.TEMP_GetMessageRight();
+	const char* BufferRight = Response.GetData() + Response.GetSize();
+	uint32 AlreadyReceived = uint32(ptrdiff_t(BufferRight - MessageRight));
+
+	Activity_ChangeState(Activity, FActivity::EState::RecvStream, AlreadyReceived);
 
 	// Copy any of the content we may have already received.
 	if (AlreadyReceived == 0)
@@ -195,22 +200,26 @@ static FOutcome DoRecvContent(FActivity* Activity, FHttpPeer& Peer, int32& MaxRe
 {
 	TRACE_CPUPROFILER_EVENT_SCOPE(IasHttp::DoRecvContent);
 
-	while (true)
+	FMutableMemoryView Dest = Activity->Dest->GetMutableView();
+	Dest = Dest.Right(Dest.GetSize() - Activity->StateParam).Left(MaxRecvSize);
+	MaxRecvSize -= int32(Dest.GetSize());
+
+	auto& Body = Activity->Buffer;
+	FOutcome Outcome = Body.TryRecv(Dest, Peer);
+	if (Outcome.IsError())
 	{
-		const auto& Response = Activity->Buffer;
-		int32 Size = (int32(Response.GetContentLength()) - Activity->StateParam);
-		if (Size == 0)
-		{
-			break;
-		}
+		Activity_SetError(Activity, Outcome);
+		return Outcome;
+	}
 
-		FOutcome Outcome = DoRecvPeer(Activity, Peer, MaxRecvSize, Size);
-		if (!Outcome.IsOk())
-		{
-			return Outcome;
-		}
+	if (Outcome.IsWaiting())
+	{
+		Activity->StateParam += Outcome.GetResult();
+		return Outcome;
 	}
 
+	check(Activity->StateParam + Outcome.GetResult() == Body.GetContentLength());
+
 #if IAS_HTTP_WITH_PERF
 	Activity->Stopwatch.RecvEnd();
 #endif
