Explicit overload that is deleted to prevent accidental misuse

#rnx

diff --git a/Engine/Source/Runtime/Experimental/IoStore/HttpClient/Private/Misc.inl b/Engine/Source/Runtime/Experimental/IoStore/HttpClient/Private/Misc.inl
--- a/Engine/Source/Runtime/Experimental/IoStore/HttpClient/Private/Misc.inl
+++ b/Engine/Source/Runtime/Experimental/IoStore/HttpClient/Private/Misc.inl
@@ -224,6 +224,8 @@ static int64 CrudeToInt(FAnsiStringView View)
 	return Value;
 };
 
+template <uint32 Base=10> static int64 CrudeToInt(const char*) = delete;
+
 ////////////////////////////////////////////////////////////////////////////////
 struct FMessageOffsets
 {
diff --git a/Engine/Source/Runtime/Experimental/IoStore/HttpClient/Private/Test.inl b/Engine/Source/Runtime/Experimental/IoStore/HttpClient/Private/Test.inl
--- a/Engine/Source/Runtime/Experimental/IoStore/HttpClient/Private/Test.inl
+++ b/Engine/Source/Runtime/Experimental/IoStore/HttpClient/Private/Test.inl
@@ -68,24 +68,26 @@ static void MiscTest()
 	EnumerateHeaders("foo:bar" CRLF " ",	IsBar); check(AllIsWell);
 	EnumerateHeaders("foo:bar" CRLF "n:ej",	IsBar); check(AllIsWell);
 
-	check(CrudeToInt("") < 0);
-	check(CrudeToInt("X") < 0);
-	check(CrudeToInt("/") < 0);
-	check(CrudeToInt(":") < 0);
-	check(CrudeToInt("-1") < -1);
-	check(CrudeToInt("0") == 0);
-	check(CrudeToInt("9") == 9);
-	check(CrudeToInt("493") == 493);
-
-	check(CrudeToInt<16>("56") == 0x56);
-	check(CrudeToInt<16>("1") == 0x01);
-	check(CrudeToInt<16>("9") == 0x09);
-	check(CrudeToInt<16>("a") == 0x0a); check(CrudeToInt<16>("A") == 0x0a);
-	check(CrudeToInt<16>("f") == 0x0f); check(CrudeToInt<16>("F") == 0x0f);
-	check(CrudeToInt<16>("g") < 0);
-	check(CrudeToInt<16>("49e") == 0x49e);
-	check(CrudeToInt<16>("aBcD") == 0xabcd);
-	check(CrudeToInt<16>("eEeE") == 0xeeee);
+	check(CrudeToInt(FAnsiStringView("")) < 0);
+	check(CrudeToInt(FAnsiStringView("X")) < 0);
+	check(CrudeToInt(FAnsiStringView("/")) < 0);
+	check(CrudeToInt(FAnsiStringView(":")) < 0);
+	check(CrudeToInt(FAnsiStringView("-1")) < -1);
+	check(CrudeToInt(FAnsiStringView("0")) == 0);
+	check(CrudeToInt(FAnsiStringView("9")) == 9);
+	check(CrudeToInt(FAnsiStringView("493")) == 493);
+
+	check(CrudeToInt<16>(FAnsiStringView("56")) == 0x56);
+	check(CrudeToInt<16>(FAnsiStringView("1")) == 0x01);
+	check(CrudeToInt<16>(FAnsiStringView("9")) == 0x09);
+	check(CrudeToInt<16>(FAnsiStringView("a")) == 0x0a);
+	check(CrudeToInt<16>(FAnsiStringView("A")) == 0x0a);
+	check(CrudeToInt<16>(FAnsiStringView("f")) == 0x0f);
+	check(CrudeToInt<16>(FAnsiStringView("F")) == 0x0f);
+	check(CrudeToInt<16>(FAnsiStringView("g")) < 0);
+	check(CrudeToInt<16>(FAnsiStringView("49e")) == 0x49e);
+	check(CrudeToInt<16>(FAnsiStringView("aBcD")) == 0xabcd);
+	check(CrudeToInt<16>(FAnsiStringView("eEeE")) == 0xeeee);
 
 	FUrlOffsets UrlOut;
 
