Add more logging to Indexed Cache Storage

diff --git a/Engine/Source/Runtime/Experimental/IndexedCacheStorage/Private/IndexedCacheStorageManager.cpp b/Engine/Source/Runtime/Experimental/IndexedCacheStorage/Private/IndexedCacheStorageManager.cpp
--- a/Engine/Source/Runtime/Experimental/IndexedCacheStorage/Private/IndexedCacheStorageManager.cpp
+++ b/Engine/Source/Runtime/Experimental/IndexedCacheStorage/Private/IndexedCacheStorageManager.cpp
@@ -144,11 +144,13 @@ bool FIndexedCacheStorageManager::CreateCacheStorage(uint64 RequestNumberOfBytes
 	FRWScopeLock _(IndexedCacheEntriesMetaDataLock, SLT_Write);
 	if (CacheIndex <= 0 && !ensure(IndexedCacheEntriesMetaData.IsValidIndex(CacheIndex)))
 	{
+		UE_LOG(LogIndexedCacheStorageManager, Error, TEXT("CreateCacheStorage invalid index %d"), CacheIndex);
 		return false;
 	}
 
 	if (IndexedCacheEntriesMetaData[CacheIndex].MountName.IsEmpty())
 	{
+		UE_LOG(LogIndexedCacheStorageManager, Error, TEXT("CreateCacheStorage invalid empty mount name"));
 		return false;
 	}
 
@@ -206,6 +208,7 @@ FString FIndexedCacheStorageManager::MountCacheStorage(int32 CacheIndex)
 	FString ActualMountName;
 	if (!IndexedCacheEntriesMetaData.IsValidIndex(CacheIndex) || !IndexedCacheEntriesMetaData[CacheIndex].bCacheExists)
 	{
+		UE_LOG(LogIndexedCacheStorageManager, Display, TEXT("MountCacheStorage skip mount for cache %d, exists=%d"), CacheIndex, IndexedCacheEntriesMetaData[CacheIndex].bCacheExists);
 		return ActualMountName;
 	}
 
@@ -225,6 +228,7 @@ FString FIndexedCacheStorageManager::MountCacheStorage(int32 CacheIndex)
 	else
 	{
 		IIndexedCacheStorage::Get().GetCacheStorageMountPath(MountedPath, IndexedCacheEntriesMetaData[CacheIndex].MountName);
+		UE_LOG(LogIndexedCacheStorageManager, Display, TEXT("MountCacheStorage cache %d already mounted %d times at %s"), CacheIndex, IndexedCacheEntriesMetaData[CacheIndex].MountRefCount, *MountedPath);
 	}
 
 	check(!MountedPath.IsEmpty());
@@ -241,6 +245,7 @@ void FIndexedCacheStorageManager::UnmountCacheStorage(int32 CacheIndex)
 	FString ActualMountName;
 	if (!IndexedCacheEntriesMetaData.IsValidIndex(CacheIndex) || IndexedCacheEntriesMetaData[CacheIndex].MountRefCount == 0)
 	{
+		UE_LOG(LogIndexedCacheStorageManager, Display, TEXT("UnmountCacheStorage skip unmount for cache %d mounted %d times"), CacheIndex, IndexedCacheEntriesMetaData[CacheIndex].MountRefCount);
 		return;
 	}
 
@@ -250,6 +255,10 @@ void FIndexedCacheStorageManager::UnmountCacheStorage(int32 CacheIndex)
 		UE_LOG(LogIndexedCacheStorageManager, Display, TEXT("UnmountCacheStorage cache %d at %s"), CacheIndex, *IndexedCacheEntriesMetaData[CacheIndex].MountName);
 		IIndexedCacheStorage::Get().UnmountCacheStorage(*IndexedCacheEntriesMetaData[CacheIndex].MountName);
 	}
+	else
+	{
+		UE_LOG(LogIndexedCacheStorageManager, Display, TEXT("UnmountCacheStorage cache %d still mounted %d times"), CacheIndex, IndexedCacheEntriesMetaData[CacheIndex].MountRefCount);
+	}
 }
 
 
