Iris
* Added net.iris.AsyncLoadingPriority cvar to be able to set the loading priority of packages referenced by replicated properties

#rnx

diff --git a/Engine/Source/Runtime/Experimental/Iris/Core/Private/Iris/ReplicationSystem/ObjectReferenceCache.cpp b/Engine/Source/Runtime/Experimental/Iris/Core/Private/Iris/ReplicationSystem/ObjectReferenceCache.cpp
--- a/Engine/Source/Runtime/Experimental/Iris/Core/Private/Iris/ReplicationSystem/ObjectReferenceCache.cpp
+++ b/Engine/Source/Runtime/Experimental/Iris/Core/Private/Iris/ReplicationSystem/ObjectReferenceCache.cpp
@@ -70,6 +70,10 @@ FAutoConsoleVariableRef CVarIrisAsyncLoadingFailAllLoads(TEXT("net.Iris.AsyncLoa
 static bool bIrisEnableRemoveReferenceFallbackPath = true;
 FAutoConsoleVariableRef CVarIrisEnableRemoveReferenceFallbackPath(TEXT("net.iris.ObjectReferenceCache.EnableRemoveReferenceFallbackPath"), bIrisEnableRemoveReferenceFallbackPath, TEXT("Whether to enable the slow fallback path or not when ObjectReferenceCache::RemoveReference is passed a null pointer object."), ECVF_Default);
 
+static int32 IrisAsyncLoadingPriority = 0;
+FAutoConsoleVariableRef CVarIrisAsyncLoadingPriority(TEXT("net.iris.AsyncLoadingPriority"), IrisAsyncLoadingPriority, TEXT("The default priority to use when async loading packages referenced by replicated properties."), ECVF_Default);
+
+
 FObjectReferenceCache::FPendingAsyncLoadRequest::FPendingAsyncLoadRequest(FNetRefHandle InNetRefHandle, double InRequestStartTime)
 : RequestStartTime(InRequestStartTime)
 {
@@ -2007,12 +2011,15 @@ void FObjectReferenceCache::StartAsyncLoadingPackage(FCachedNetObjectReference&
 	else
 #endif
 	{
-		LoadPackageAsync(PackagePath.ToString(), FLoadPackageAsyncDelegate::CreateWeakLambda(ReplicationSystem, 
-			[this](const FName& PackageName, UPackage* Package, EAsyncLoadingResult::Type Result)
-			{
-				AsyncPackageCallback(PackageName, Package, Result);
-			}
-		));
+		LoadPackageAsync(
+			PackagePath.ToString(), 
+			FLoadPackageAsyncDelegate::CreateWeakLambda(ReplicationSystem, 
+				[this](const FName& PackageName, UPackage* Package, EAsyncLoadingResult::Type Result)
+				{
+					AsyncPackageCallback(PackageName, Package, Result);
+				}),
+			IrisAsyncLoadingPriority
+		);
 	}
 }
 
