Failed IAX requests are now logged at warning verbosity and fixed a number of cases where a 4XX response would be reported as an internal error.

#rnx

### Logging improvements
- The original goal of this changelist was to improve the logging of response issues in live as internally we are quite used to running with the IAX log categories set to VeryVerbose which gives a great breakdown of issues but in live we only log errors for the CDN latency checks so if a user if having problems with their requests but not the latency check we don't have any info to work with.
- FMultiEndpointHttpClient::Log will now use the Warning verbosity if the request had problems, either a response outside of the 2XX range or an internal error.
- We now append the reason for the error to the end of the log line. This is especially useful in the case of an internal error as the status code being 0 doesn't really tell us much.
- We now record the response message (FRequest::GetStatusMessage) so that we can use it for log messages for non internal errors.
- Add a new cvar 'iax.HttpWarnOnFailure' that can be set to false to restore the old behavior.

### Misc
- While testing 4XX messages I found that a number of CDNs return responses for these with chunked transfer enabled. Since IAX was not enabling chunked transfer this resulted in an internal error which we'd end up reporting rather than the actual problem. So 'ias.HttpAllowChunkedXfer' is now enabled by default.
- As it turned out FMultiEndpointHttpClient did not support chunked transfer so that needed to be fixed up first. Currently we do this in a pretty inefficient way by just appending each chunk to FRequest::Body which as a FIoBuffer is not really designed for this so we need to reallocate and copy data on each chunk append. We probably could do this better, either by pre-allocating the length of the data (we do accept this as an optional input optimization) but that is way bigger than needed for errors.
-- Since we don't use responses if a request failed I added a new FMultiEndpointHttpClientConfig option 'bResponseBodyOnError' (false by default) which indicates if the caller wants the request response for failed requests (response code outside of the 2XX range) or not. So gathering the chunks might be inefficient but we just opt of that entirely. If we ever start using chunked transfers properly then we might need to revisit this (ideally we'd just use FCompositeBuffer but that would require a more extensive rewrite then I am willing to do at the moment)
- We now only parse the CDN hit/miss statistics if the response was in the 2XX range as most other responses will just return unknown anyway. We also only apply the CDN hit/miss stat to our statistics if the request succeeded.

diff --git a/Engine/Source/Runtime/Experimental/IoStore/OnDemand/Private/IasHostGroup.cpp b/Engine/Source/Runtime/Experimental/IoStore/OnDemand/Private/IasHostGroup.cpp
--- a/Engine/Source/Runtime/Experimental/IoStore/OnDemand/Private/IasHostGroup.cpp
+++ b/Engine/Source/Runtime/Experimental/IoStore/OnDemand/Private/IasHostGroup.cpp
@@ -256,7 +256,7 @@ bool FIASHostGroup::OnFailedResponse()
 	const float Average = Impl->HttpErrorHistory.AvgSetBits();
 	const bool bAboveHighWaterMark = Average > GIasHttpErrorHighWater;
 
-	UE_LOG(LogIas, Log, TEXT("[%s] %.2f%% the last %d HTTP requests failed"), *Impl->Name.ToString(),  Average * 100.0f, GIasHttpErrorSampleCount);
+	UE_LOG(LogIas, Verbose, TEXT("[%s] %.2f%% the last %d HTTP requests failed"), *Impl->Name.ToString(),  Average * 100.0f, GIasHttpErrorSampleCount);
 
 	if (bAboveHighWaterMark && IsConnected())
 	{
diff --git a/Engine/Source/Runtime/Experimental/IoStore/OnDemand/Private/OnDemandHttpClient.cpp b/Engine/Source/Runtime/Experimental/IoStore/OnDemand/Private/OnDemandHttpClient.cpp
--- a/Engine/Source/Runtime/Experimental/IoStore/OnDemand/Private/OnDemandHttpClient.cpp
+++ b/Engine/Source/Runtime/Experimental/IoStore/OnDemand/Private/OnDemandHttpClient.cpp
@@ -20,6 +20,25 @@
 namespace UE::IoStore
 {
 
+/** Not intended as a long term cvar, only including so that we can revert the new behavior easily if needed */
+bool GIaxWarnOnFailure = true;
+static FAutoConsoleVariableRef CVar_IaxWarnOnFailure(
+	TEXT("iax.HttpWarnOnFailure"),
+	GIaxWarnOnFailure,
+	TEXT("When enabled FMultiEndpointHttpClient will log failed requests as warnings instead of the config verbosity")
+);
+
+/** Utility function to join two FIoBuffers together and return the combined result as a new FIoBuffer */
+static FIoBuffer JoinIoBuffers(const FIoBuffer& LHS, const FIoBuffer& RHS)
+{
+	const uint64 TotalSize = LHS.GetSize() + RHS.GetSize();
+	FIoBuffer Output(TotalSize);
+
+	Output.GetMutableView().CopyFrom(LHS.GetView()).CopyFrom(RHS.GetView());
+
+	return Output;
+}
+
 static int8 TrackCdnCacheStats(const HTTP::FResponse& Response)
 {
 #if !UE_BUILD_TEST
@@ -274,50 +293,76 @@ void FMultiEndpointHttpClient::IssueRequest(FRequest&& Request)
 
 	auto HttpSink = [this, Request = MoveTemp(Request)](const FTicketStatus& TicketStatus) mutable
 	{
-		const uint64 DurationMs = uint64(FPlatformTime::ToMilliseconds64(FPlatformTime::Cycles64() - Request.StartTime));
-
 		switch (TicketStatus.GetId())
 		{
 			case FTicketStatus::EId::Response:
 			{
 				FResponse& HttpResponse = TicketStatus.GetResponse();
-				Request.StatusCode		= HttpResponse.GetStatusCode();
-				HttpResponse.SetDestination(&Request.Body);
-				Request.CDNCacheStatus = TrackCdnCacheStats(HttpResponse);
+
+				Request.StatusCode = HttpResponse.GetStatusCode();
+				Request.bIsChunkedTransfer = HttpResponse.GetContentLength() == -1;
+
+				if (IsHttpStatusOk(Request.StatusCode))
+				{
+					Request.CDNCacheStatus = TrackCdnCacheStats(HttpResponse);
+				}
+				else
+				{
+					Request.ResponseMessage = FString(HttpResponse.GetStatusMessage().TrimEnd());
+				}
+				
+				HttpResponse.SetDestination(Request.bIsChunkedTransfer == false ? &Request.Body : &Request.Chunk);
+				
 				break;
 			}
 			case FTicketStatus::EId::Content:
+			{
+				// ::DataSize of zero means that all the chunks have been transfered and FRequest::Body should be complete
+				if (Request.bIsChunkedTransfer && Request.Chunk.DataSize() != 0)
+				{
+					if (IsHttpStatusOk(Request.StatusCode) || Config.bResponseBodyOnError)
+					{
+						// Could consider using FRequest::Range to presize FRequest::Body and copy the chunks
+						// into it rather than resizing each time.
+						Request.Body = JoinIoBuffers(Request.Body, Request.Chunk);
+					}
+					return;
+				}
+			}
 			case FTicketStatus::EId::Error:
 			case FTicketStatus::EId::Cancelled:
 			{
-				Log(Request);
+				Log(Request, TicketStatus);
 
-				const bool bError		= TicketStatus.GetId() == FTicketStatus::EId::Error;
-				const bool bServerError = Request.StatusCode > 499 && Request.StatusCode < 600;
+				// We only want to retry the request if there was an internal or server problem. There is no point retrying a 404 error for example/
+				const bool bError = IsHttpServerError(Request.StatusCode) || TicketStatus.GetId() == FTicketStatus::EId::Error;
 
-				if ((bError || bServerError) && Request.RetryCount < GetRetryLimitForRequest(Request))
+				if (bError && Request.RetryCount < GetRetryLimitForRequest(Request))
 				{
 					RetryRequest(MoveTemp(Request));
 				}
 				else
 				{
-					
 					FMultiEndpointHttpClientResponse Response
 					{
-						.Body					= MoveTemp(Request.Body),
-						.DurationMilliseconds	= DurationMs,
+						.DurationMilliseconds	= uint64(FPlatformTime::ToMilliseconds64(FPlatformTime::Cycles64() - Request.StartTime)),
 						.StatusCode				= Request.StatusCode,
 						.RetryCount				= Request.RetryCount,
 						.HostIndex				= Request.Host,
 						.CDNCacheStatus			= Request.CDNCacheStatus
 					};
 
+					if (IsHttpStatusOk(Request.StatusCode) || Config.bResponseBodyOnError)
+					{
+						Response.Body = MoveTemp(Request.Body);
+					}
+
 					if (TicketStatus.GetId() == FTicketStatus::EId::Content)
 					{
 						Response.Sample = TicketStatus.GetPerf().GetSample();
 					}
-
-					if (bError)
+					
+					if (TicketStatus.GetId() == FTicketStatus::EId::Error)
 					{
 						Response.Reason = TicketStatus.GetError().Reason;
 					}
@@ -325,6 +370,10 @@ void FMultiEndpointHttpClient::IssueRequest(FRequest&& Request)
 					{
 						Response.Reason = TEXT("Cancelled");
 					}
+					else if (!Request.ResponseMessage.IsEmpty())
+					{
+						Response.Reason = MoveTemp(Request.ResponseMessage);
+					}
 
 					FOnHttpResponse OnResponse = MoveTemp(Request.OnResponse);
 					OnResponse(MoveTemp(Response));
@@ -437,10 +486,15 @@ uint32 FMultiEndpointHttpClient::GetRetryLimitForRequest(const FRequest& Request
 	return Config.MaxRetryCount == -1 ? Request.Connection.HostGroup.Hosts().Num() : Config.MaxRetryCount;
 }
 
-void FMultiEndpointHttpClient::Log(const FRequest& Request) const
+void FMultiEndpointHttpClient::Log(const FRequest& Request, const HTTP::FTicketStatus& TicketStatus) const
 {
 #if !NO_LOGGING
-	if (Config.LogCategory == nullptr || Config.LogCategory->IsSuppressed(Config.LogVerbosity))
+
+	using namespace UE::IoStore::HTTP;
+	 
+	const ELogVerbosity::Type Verbosity = !GIaxWarnOnFailure || IsHttpStatusOk(Request.StatusCode) ? Config.LogVerbosity : ELogVerbosity::Type::Warning;
+
+	if (Config.LogCategory == nullptr || Config.LogCategory->IsSuppressed(Verbosity))
 	{
 		return;
 	}
@@ -448,15 +502,30 @@ void FMultiEndpointHttpClient::Log(const FRequest& Request) const
 	const uint64 DurationMs = (uint64)FPlatformTime::ToMilliseconds64(FPlatformTime::Cycles64() - Request.StartTime);
 	const uint64 Size = Request.Body.GetSize() >> 10;
 
-	FMsg::Logf(__FILE__, __LINE__, Config.LogCategory->GetCategoryName(), Config.LogVerbosity,
-		TEXT("http-%3u: %5" UINT64_FMT "ms %5" UINT64_FMT "KiB [%4s] %s (Attempt %u/%u)"),
+	TStringBuilder<256> Reason;
+	if (TicketStatus.GetId() == FTicketStatus::EId::Error)
+	{
+		Reason << TEXT(" ") << TicketStatus.GetError().Reason;
+	}
+	else if (TicketStatus.GetId() == FTicketStatus::EId::Cancelled)
+	{
+		Reason << TEXT(" Cancelled");
+	}
+	else if (!Request.ResponseMessage.IsEmpty())
+	{
+		Reason << TEXT(" ") << Request.ResponseMessage;
+	}
+
+	FMsg::Logf(__FILE__, __LINE__, Config.LogCategory->GetCategoryName(), Verbosity,
+		TEXT("http-%3u: %5" UINT64_FMT "ms %5" UINT64_FMT "KiB [%4s] %s (Attempt %u/%u)%s"),
 		Request.StatusCode,
 		DurationMs,
 		Size,
 		CDNCacheStatusToString(Request.CDNCacheStatus),
 		WriteToString<512>(Request.Connection.HostGroup.Host(Request.Host), Request.RelativeUrl).ToString(),
 		Request.RetryCount,
-		GetRetryLimitForRequest(Request)
+		GetRetryLimitForRequest(Request),
+		*Reason
 	);
 
 #endif //!NO_LOGGING
diff --git a/Engine/Source/Runtime/Experimental/IoStore/OnDemand/Private/OnDemandHttpClient.h b/Engine/Source/Runtime/Experimental/IoStore/OnDemand/Private/OnDemandHttpClient.h
--- a/Engine/Source/Runtime/Experimental/IoStore/OnDemand/Private/OnDemandHttpClient.h
+++ b/Engine/Source/Runtime/Experimental/IoStore/OnDemand/Private/OnDemandHttpClient.h
@@ -25,10 +25,20 @@
 namespace UE::IoStore
 {
 
-FIoStatus LoadDefaultHttpCertificates(bool& bWasLoaded); 
+namespace HTTP { class FTicketStatus; }
+
+FIoStatus LoadDefaultHttpCertificates(bool& bWasLoaded);
+
+/** Returns if a Http status code is okay (2XX) or not */
 inline bool IsHttpStatusOk(uint32 StatusCode)
 {
-	return StatusCode > 199 && StatusCode < 300;
+	return StatusCode >= 200 && StatusCode < 300;
+}
+
+/** Returns if a Http status code indicates a server error (5XX) or not */
+inline bool IsHttpServerError(uint32 StatusCode)
+{
+	return StatusCode >= 500 && StatusCode < 600;
 }
 
 enum class EHttpRedirects
@@ -50,6 +60,8 @@ struct FMultiEndpointHttpClientConfig
 	EHttpRedirects				Redirects = EHttpRedirects::Follow;
 	bool						bEnableThreadSafetyChecks = false;
 	bool						bAllowChunkedTransfer = true;
+	/** When true the client will return the response body for requests outside of the 2XX range. When false only requests within the 2XX range will return the response body */
+	bool						bResponseBodyOnError = false;
 	/** Logging will be disabled if this is set to a nullptr, it is up to the calling system to assign a LogCategory */
 	const UE_LOG_CATEGORY_TYPE*	LogCategory = nullptr;
 	ELogVerbosity::Type			LogVerbosity = ELogVerbosity::Log;
@@ -118,11 +130,20 @@ private:
 		FAnsiString			RelativeUrl;
 		FIoOffsetAndLength	Range;
 		FConnection&		Connection;
+		/** Stored the response body */
 		FIoBuffer			Body;
+		/**
+		 * Stores each chunk being sent if the response uses chunked transfer, the chunks will be combined
+		 * together so that the final result will be in Body once all chunks are processed.
+		 */
+		FIoBuffer			Chunk;
+		/** Stores the response message if the StatusCode indicates a problem */
+		FString				ResponseMessage;
 		uint64				StartTime = 0;
 		uint32				RetryCount = 0;
 		uint32				StatusCode = 0;
 		int32				Host = INDEX_NONE;
+		bool				bIsChunkedTransfer = false;
 		int8				CDNCacheStatus = -1;
 	};
 
@@ -136,7 +157,7 @@ private:
 
 	uint32								GetRetryLimitForRequest(const FRequest& Request) const;
 
-	void								Log(const FRequest& Request) const;
+	void								Log(const FRequest& Request, const HTTP::FTicketStatus& TicketStatus) const;
 
 	FMultiEndpointHttpClientConfig		Config;
 	TArray<TUniquePtr<FConnection>>		Connections;
diff --git a/Engine/Source/Runtime/Experimental/IoStore/OnDemand/Private/OnDemandHttpThread.cpp b/Engine/Source/Runtime/Experimental/IoStore/OnDemand/Private/OnDemandHttpThread.cpp
--- a/Engine/Source/Runtime/Experimental/IoStore/OnDemand/Private/OnDemandHttpThread.cpp
+++ b/Engine/Source/Runtime/Experimental/IoStore/OnDemand/Private/OnDemandHttpThread.cpp
@@ -95,7 +95,7 @@ static FAutoConsoleVariableRef CVar_IasHttpFailTimeOutMs(
 	TEXT("Fail infinite network waits that take longer than this (in ms, a value of zero will use the default)")
 );
 
-bool GIasHttpAllowChunkedXfer = false;
+bool GIasHttpAllowChunkedXfer = true;
 static FAutoConsoleVariableRef CVar_IasHttpAllowChunkedXfer(
 	TEXT("ias.HttpAllowChunkedXfer"),
 	GIasHttpAllowChunkedXfer,
@@ -383,12 +383,11 @@ void FOnDemandHttpThread::TickRequests()
 								FOnDemandIoBackendStats::Get()->OnHttpRetry(HttpRequest->Type);
 							}
 
-							FOnDemandIoBackendStats::Get()->OnHttpCdnCacheReply(HttpRequest->Type, HttpResponse.CDNCacheStatus);
-
 							if (HttpResponse.IsOk())
 							{
 								HttpRequest->HostGroup().OnSuccessfulResponse();
 
+								FOnDemandIoBackendStats::Get()->OnHttpCdnCacheReply(HttpRequest->Type, HttpResponse.CDNCacheStatus);
 								FOnDemandIoBackendStats::Get()->OnHttpGet(HttpRequest->Type, HttpResponse.Body.DataSize(), HttpResponse.DurationMilliseconds);
 
 								HttpRequest->OnRequestCompleted(HttpResponse.StatusCode, HttpResponse.Reason, MoveTemp(HttpResponse.Body));
