Added command line option -Iax.DefaultHostGroup=<host1, host2,..> that will override the default host group for both IAD and IAS.

#rnx

diff --git a/Engine/Source/Runtime/Experimental/IoStore/OnDemand/Private/OnDemandConfig.cpp b/Engine/Source/Runtime/Experimental/IoStore/OnDemand/Private/OnDemandConfig.cpp
--- a/Engine/Source/Runtime/Experimental/IoStore/OnDemand/Private/OnDemandConfig.cpp
+++ b/Engine/Source/Runtime/Experimental/IoStore/OnDemand/Private/OnDemandConfig.cpp
@@ -306,6 +306,20 @@ TIoStatusOr<FOnDemandEndpointConfig> TryParseEndpointConfig(const TCHAR* Command
 		}
 	}
 
+#if !UE_BUILD_SHIPPING
+	FString HostsOverride;
+	if (FParse::Value(CommandLine, TEXT("Iax.DefaultHostGroup="), HostsOverride))
+	{
+		OutConfig.DistributionUrl.Empty();
+		OutConfig.ServiceUrls.Empty();
+		HostsOverride.ParseIntoArray(OutConfig.ServiceUrls, TEXT(","), true);
+		for (FString& Host : OutConfig.ServiceUrls)
+		{
+			Host.TrimStartAndEndInline();
+		}
+	}
+#endif
+
 	return OutConfig;
 }
 
diff --git a/Engine/Source/Runtime/Experimental/IoStore/OnDemand/Private/OnDemandIoStore.cpp b/Engine/Source/Runtime/Experimental/IoStore/OnDemand/Private/OnDemandIoStore.cpp
--- a/Engine/Source/Runtime/Experimental/IoStore/OnDemand/Private/OnDemandIoStore.cpp
+++ b/Engine/Source/Runtime/Experimental/IoStore/OnDemand/Private/OnDemandIoStore.cpp
@@ -1124,14 +1124,33 @@ FIoStatus FOnDemandIoStore::TickMountRequest(FMountRequest& MountRequest)
 				if (NewHostGroup.IsOk())
 				{
 					Container->HostGroup = NewHostGroup.ConsumeValueOrDie();
-					if (FIoStatus Status = Container->HostGroup.Resolve(Args.HostGroup.Hosts()); !Status.IsOk())
+					TConstArrayView<FAnsiString> Hosts = Args.HostGroup.Hosts();
+#if !UE_BUILD_SHIPPING
+					TArray<FAnsiString> HostsOverride;
+					if (Container->HostGroupName == FOnDemandHostGroup::DefaultName)
+					{
+						FString HostValue;
+						if (FParse::Value(FCommandLine::Get(), TEXT("Iax.DefaultHostGroup="), HostValue))
+						{
+							TArray<FString> HostArray;
+							HostValue.ParseIntoArray(HostArray, TEXT(","), true);
+							for (FString& Host : HostArray)
+							{
+								Host.TrimStartAndEndInline();
+								HostsOverride.Add(FAnsiString(StringCast<ANSICHAR>(*Host)));
+							}
+							Hosts = HostsOverride; 
+						}
+					}
+#endif
+					if (FIoStatus Status = Container->HostGroup.Resolve(Hosts); !Status.IsOk())
 					{
 						return Status;
 					}
 
 					{
 						const FOnDemandHostGroup& HostGroup = Container->HostGroup.GetUnderlyingHostGroup();
-						TConstArrayView<FAnsiString> Hosts	= HostGroup.Hosts();
+						Hosts								= HostGroup.Hosts();
 						TStringBuilder<128> Sb;
 
 						UE_LOG(LogIoStoreOnDemand, Log, TEXT("Created new host group '%s'"), *Container->HostGroupName.ToString());
