Optimize on-demand package store.
- reduced time for updating which packages are installed, total time when running IAD stress-test 455ms -> 45ms
- reduced package store entry from 16B -> 8B

#rnx

diff --git a/Engine/Source/Runtime/Experimental/IoStore/OnDemand/Private/OnDemandIoStore.cpp b/Engine/Source/Runtime/Experimental/IoStore/OnDemand/Private/OnDemandIoStore.cpp
--- a/Engine/Source/Runtime/Experimental/IoStore/OnDemand/Private/OnDemandIoStore.cpp
+++ b/Engine/Source/Runtime/Experimental/IoStore/OnDemand/Private/OnDemandIoStore.cpp
@@ -1691,6 +1691,12 @@ void FOnDemandIoStore::GetReferencedContent(TArray<FSharedOnDemandContainer>& Ou
 	}
 }
 
+TBitArray<> FOnDemandIoStore::GetReferencedContent(const FSharedOnDemandContainer& Container)
+{
+	UE::TUniqueLock Lock(ContainerMutex);
+	return Container->GetReferencedChunkEntries();
+}
+
 void FOnDemandIoStore::GetReferencedContentByHandle(TMap<FOnDemandWeakContentHandle, TArray<FOnDemandContainerChunkEntryReferences>>& OutReferencesByHandle) const
 {
 	UE::TUniqueLock Lock(ContainerMutex);
diff --git a/Engine/Source/Runtime/Experimental/IoStore/OnDemand/Private/OnDemandIoStore.h b/Engine/Source/Runtime/Experimental/IoStore/OnDemand/Private/OnDemandIoStore.h
--- a/Engine/Source/Runtime/Experimental/IoStore/OnDemand/Private/OnDemandIoStore.h
+++ b/Engine/Source/Runtime/Experimental/IoStore/OnDemand/Private/OnDemandIoStore.h
@@ -333,6 +333,7 @@ public:
 	TArray<FIoChunkId>					DebugFindStreamingChunkIds(int32 NumToFind);
 #endif //!UE_BUILD_SHIPPING
 	void								GetReferencedContent(TArray<FSharedOnDemandContainer>& OutContainers, TArray<TBitArray<>>& OutChunkEntryIndices, bool bPackageStore = false);
+	TBitArray<>							GetReferencedContent(const FSharedOnDemandContainer& Container);
 	void								GetReferencedContentByHandle(TMap<FOnDemandWeakContentHandle, TArray<FOnDemandContainerChunkEntryReferences>>& OutReferencesByHandle) const;
 
 	// IOnDemandIoStore interface
diff --git a/Engine/Source/Runtime/Experimental/IoStore/OnDemand/Private/OnDemandPackageStoreBackend.cpp b/Engine/Source/Runtime/Experimental/IoStore/OnDemand/Private/OnDemandPackageStoreBackend.cpp
--- a/Engine/Source/Runtime/Experimental/IoStore/OnDemand/Private/OnDemandPackageStoreBackend.cpp
+++ b/Engine/Source/Runtime/Experimental/IoStore/OnDemand/Private/OnDemandPackageStoreBackend.cpp
@@ -21,22 +21,11 @@ class FOnDemandPackageStoreBackend final
 {
 	using FSoftPackageReferenceMap = TMap<FPackageId, const FFilePackageStoreEntrySoftReferences*>;
 
-	struct FContainer
-	{
-		FContainer(FString&& ContainerName, FSharedContainerHeader ContainerHeader)
-			: Name(MoveTemp(ContainerName))
-			, Header(MoveTemp(ContainerHeader))
-		{ }
-
-		FString						Name;
-		FSharedContainerHeader		Header;
-		FSoftPackageReferenceMap	SoftRefs;
-	};
-
 	struct FOnDemandPackageStoreEntry
 	{
-		const FFilePackageStoreEntry* FilePackageStoreEntry = nullptr;
-		bool bIsInstalled = false; // uninstalled packages are not 'missing' and not 'pending'
+		uint32 PackageEntryIndex : 31;
+		uint32 IsReferenced : 1;
+		int32 ChunkEntryIndex;
 	};
 
 	using FSharedBackendContext = TSharedPtr<const FPackageStoreBackendContext>;
@@ -45,6 +34,15 @@ class FOnDemandPackageStoreBackend final
 	using FLocalizedMap			= TMap<FPackageId, FName>;
 	using FRedirectMap			= TMap<FPackageId, FRedirect>;
 
+	struct FContainer
+	{
+		FSharedOnDemandContainer	Container;
+		FEntryMap					EntryMap;
+		FLocalizedMap				LocalizedMap;
+		FRedirectMap				RedirectMap;
+		FSoftPackageReferenceMap	SoftRefs;
+	};
+
 public:
 						FOnDemandPackageStoreBackend(TWeakPtr<FOnDemandIoStore> OnDemandIoStore);
 						virtual ~FOnDemandPackageStoreBackend ();
@@ -69,13 +67,10 @@ public:
 
 private:
 	void				UpdateLookupTables(const TArray<FSharedOnDemandContainer>& AllContainers, const TArray<TBitArray<>>& ReferencedChunkEntryIndices);
-	void				UpdateReferencedPackages(const TArray<FSharedOnDemandContainer>& AllContainers, const TArray<TBitArray<>>& ReferencedChunkEntryIndices);
+	void				UpdateReferencedPackages(const TArray<TBitArray<>>& ReferencedChunkEntryIndices);
 
 	TWeakPtr<FOnDemandIoStore>			OnDemandIoStore;
 	TArray<FContainer>					Containers;
-	FEntryMap							EntryMap;
-	FLocalizedMap						LocalizedMap;
-	FRedirectMap						RedirectMap;
 	UE::FMutex							Mutex;
 	std::atomic<EOnDemandPackageStoreUpdateMode> NeedsUpdateMode{ EOnDemandPackageStoreUpdateMode::Full };
 };
@@ -109,6 +104,8 @@ void FOnDemandPackageStoreBackend::OnMounted(TSharedRef<const FPackageStoreBacke
 
 void FOnDemandPackageStoreBackend::BeginRead()
 {
+	TRACE_CPUPROFILER_EVENT_SCOPE(FOnDemandPackageStoreBackend::BeginRead);
+
 	EOnDemandPackageStoreUpdateMode LocalNeedsUpdate = NeedsUpdateMode.exchange(EOnDemandPackageStoreUpdateMode::None);
 
 	if (LocalNeedsUpdate == EOnDemandPackageStoreUpdateMode::None)
@@ -117,24 +114,32 @@ void FOnDemandPackageStoreBackend::BeginRead()
 		return;
 	}
 
-	TArray<FSharedOnDemandContainer> AllContainers;
 	TArray<TBitArray<>> ReferencedChunkEntryIndices;
-
-	if (TSharedPtr<FOnDemandIoStore> PinOnDemandIoStore = OnDemandIoStore.Pin())
-	{
-		const bool bPackageStore = true;
-		PinOnDemandIoStore->GetReferencedContent(AllContainers, ReferencedChunkEntryIndices, bPackageStore);
-	}
-
-	Mutex.Lock();
+	TSharedPtr<FOnDemandIoStore> PinOnDemandIoStore = OnDemandIoStore.Pin();
 
 	if (LocalNeedsUpdate == EOnDemandPackageStoreUpdateMode::Full)
 	{
+		TArray<FSharedOnDemandContainer> AllContainers;
+		if (PinOnDemandIoStore.IsValid())
+		{
+			const bool bPackageStore = true;
+			PinOnDemandIoStore->GetReferencedContent(AllContainers, ReferencedChunkEntryIndices, bPackageStore);
+		}
+		Mutex.Lock();
 		UpdateLookupTables(AllContainers, ReferencedChunkEntryIndices);
 	}
-	else if (LocalNeedsUpdate == EOnDemandPackageStoreUpdateMode::ReferencedPackages)
+	else
 	{
-		UpdateReferencedPackages(AllContainers, ReferencedChunkEntryIndices);
+		ensure(LocalNeedsUpdate == EOnDemandPackageStoreUpdateMode::ReferencedPackages);
+		if (PinOnDemandIoStore.IsValid())
+		{
+			for (const FContainer& Container : Containers)
+			{
+				ReferencedChunkEntryIndices.Add(PinOnDemandIoStore->GetReferencedContent(Container.Container));
+			}
+		}
+		Mutex.Lock();
+		UpdateReferencedPackages(ReferencedChunkEntryIndices);
 	}
 }
 
@@ -148,20 +153,26 @@ EPackageStoreEntryStatus FOnDemandPackageStoreBackend::GetPackageStoreEntry(
 	FName PackageName,
 	FPackageStoreEntry& OutPackageStoreEntry)
 {
-	if (const FOnDemandPackageStoreEntry* Entry = EntryMap.Find(PackageId))
+	for (const FContainer& Container : Containers)
 	{
-		const FFilePackageStoreEntry* FilePackageStoreEntry = Entry->FilePackageStoreEntry;
-		OutPackageStoreEntry.ImportedPackageIds =
-			MakeArrayView(FilePackageStoreEntry->ImportedPackages.Data(), FilePackageStoreEntry->ImportedPackages.Num());
-		OutPackageStoreEntry.ShaderMapHashes =
-			MakeArrayView(FilePackageStoreEntry->ShaderMapHashes.Data(), FilePackageStoreEntry->ShaderMapHashes.Num());
-
-		if (Entry->bIsInstalled)
+		if (const FOnDemandPackageStoreEntry* PackageEntry = Container.EntryMap.Find(PackageId))
 		{
-			return EPackageStoreEntryStatus::Ok;
-		}
+			const FIoContainerHeader& ContainerHeader			= *Container.Container->Header;
+			const FFilePackageStoreEntry* PackageStoreEntries	= reinterpret_cast<const FFilePackageStoreEntry*>(ContainerHeader.StoreEntries.GetData());
+			const FFilePackageStoreEntry& PackageStoreEntry		= *(PackageStoreEntries + PackageEntry->PackageEntryIndex);
 
-		return EPackageStoreEntryStatus::NotInstalled;
+			OutPackageStoreEntry.ImportedPackageIds =
+				MakeArrayView(PackageStoreEntry.ImportedPackages.Data(), PackageStoreEntry.ImportedPackages.Num());
+			OutPackageStoreEntry.ShaderMapHashes =
+				MakeArrayView(PackageStoreEntry.ShaderMapHashes.Data(), PackageStoreEntry.ShaderMapHashes.Num());
+
+			if (PackageEntry->IsReferenced)
+			{
+				return EPackageStoreEntryStatus::Ok;
+			}
+
+			return EPackageStoreEntryStatus::NotInstalled;
+		}
 	}
 
 	return EPackageStoreEntryStatus::Missing;
@@ -172,24 +183,27 @@ bool FOnDemandPackageStoreBackend::GetPackageRedirectInfo(
 	FName& OutSourcePackageName,
 	FPackageId& OutRedirectedToPackageId)
 {
-	if (const FRedirect* Redirect = RedirectMap.Find(PackageId))
-	{
-		OutSourcePackageName		= Redirect->Key;
-		OutRedirectedToPackageId	= Redirect->Value;
-		return true;
-	}
-	
-	if (const FName* SourcePkgName = LocalizedMap.Find(PackageId))
+	for (const FContainer& Container : Containers)
 	{
-		const FName LocalizedPkgName = FPackageLocalizationManager::Get().FindLocalizedPackageName(*SourcePkgName);
-		if (LocalizedPkgName.IsNone() == false)
+		if (const FRedirect* Redirect = Container.RedirectMap.Find(PackageId))
 		{
-			const FPackageId LocalizedPkgId = FPackageId::FromName(LocalizedPkgName);
-			if (EntryMap.Find(LocalizedPkgId))
+			OutSourcePackageName		= Redirect->Key;
+			OutRedirectedToPackageId	= Redirect->Value;
+			return true;
+		}
+
+		if (const FName* SourcePkgName = Container.LocalizedMap.Find(PackageId))
+		{
+			const FName LocalizedPkgName = FPackageLocalizationManager::Get().FindLocalizedPackageName(*SourcePkgName);
+			if (LocalizedPkgName.IsNone() == false)
 			{
-				OutSourcePackageName		= *SourcePkgName;
-				OutRedirectedToPackageId	= LocalizedPkgId;
-				return true;
+				const FPackageId LocalizedPkgId = FPackageId::FromName(LocalizedPkgName);
+				if (Container.EntryMap.Find(LocalizedPkgId))
+				{
+					OutSourcePackageName		= *SourcePkgName;
+					OutRedirectedToPackageId	= LocalizedPkgId;
+					return true;
+				}
 			}
 		}
 	}
@@ -203,7 +217,7 @@ TConstArrayView<uint32> FOnDemandPackageStoreBackend::GetSoftReferences(FPackage
 	{
 		if (const FFilePackageStoreEntrySoftReferences* SoftRefs = Container.SoftRefs.FindRef(PackageId))
 		{
-			OutPackageIds = Container.Header->SoftPackageReferences.PackageIds;
+			OutPackageIds = Container.Container->Header->SoftPackageReferences.PackageIds;
 			return TConstArrayView<uint32>(SoftRefs->Indices.Data(), SoftRefs->Indices.Num());
 		}
 	}
@@ -214,22 +228,17 @@ void FOnDemandPackageStoreBackend::UpdateLookupTables(const TArray<FSharedOnDema
 {
 	check(Mutex.IsLocked());
 
-	int32 PackageCount = 0;
-	for (const FSharedOnDemandContainer& C : AllContainers)
-	{
-		PackageCount += C->Header->PackageIds.Num();
-	}
-
-	LocalizedMap.Empty();
-	RedirectMap.Empty();
 	Containers.Empty(AllContainers.Num());
-	EntryMap.Empty(PackageCount);
 
 	for (int32 ContainerIdx = 0; const FSharedOnDemandContainer& SharedContainer : AllContainers)
 	{
-		FContainer& Container = Containers.Emplace_GetRef(CopyTemp(SharedContainer->UniqueName()), SharedContainer->Header);
+		FContainer& Container			= Containers.AddDefaulted_GetRef();
+		Container.Container				= SharedContainer;
+		FEntryMap& EntryMap				= Container.EntryMap;
+		const FIoContainerHeader& Hdr	= *SharedContainer->Header;
+	
+		EntryMap.Reserve(Hdr.PackageIds.Num());
 
-		const FIoContainerHeader& Hdr = *Container.Header;
 		TConstArrayView<FFilePackageStoreEntry> Entries(
 			reinterpret_cast<const FFilePackageStoreEntry*>(Hdr.StoreEntries.GetData()),
 			Hdr.PackageIds.Num());
@@ -251,7 +260,13 @@ void FOnDemandPackageStoreBackend::UpdateLookupTables(const TArray<FSharedOnDema
 			check(PkgChunkIdx != INDEX_NONE);
 			const bool IsReferenced = ReferencedChunkEntryIndices[ContainerIdx][PkgChunkIdx];
 
-			EntryMap.Add(PkgId, FOnDemandPackageStoreEntry{ .FilePackageStoreEntry = &Entry, .bIsInstalled = IsReferenced });
+			EntryMap.Add(PkgId, FOnDemandPackageStoreEntry
+			{
+				.PackageEntryIndex	= uint32(EntryIdx),
+				.IsReferenced		= uint32(IsReferenced),
+				.ChunkEntryIndex	= PkgChunkIdx
+			});
+
 			if (AllSoftReferences.IsEmpty() == false)
 			{
 				const FFilePackageStoreEntrySoftReferences& SoftRefs = AllSoftReferences[EntryIdx];
@@ -262,8 +277,10 @@ void FOnDemandPackageStoreBackend::UpdateLookupTables(const TArray<FSharedOnDema
 			}
 			++EntryIdx;
 		}
+		EntryMap.Shrink();
 		Container.SoftRefs.Shrink();
 
+		FLocalizedMap& LocalizedMap = Container.LocalizedMap;
 		for (const FIoContainerHeaderLocalizedPackage& Localized : Hdr.LocalizedPackages)
 		{
 			FName& SourcePackageName = LocalizedMap.FindOrAdd(Localized.SourcePackageId);
@@ -273,7 +290,9 @@ void FOnDemandPackageStoreBackend::UpdateLookupTables(const TArray<FSharedOnDema
 				SourcePackageName = NameEntry.ToName(Localized.SourcePackageName.GetNumber());
 			}
 		}
+		LocalizedMap.Shrink();
 
+		FRedirectMap& RedirectMap = Container.RedirectMap;
 		for (const FIoContainerHeaderPackageRedirect& Redirect : Hdr.PackageRedirects)
 		{
 			FRedirect& RedirectEntry = RedirectMap.FindOrAdd(Redirect.SourcePackageId);
@@ -285,39 +304,37 @@ void FOnDemandPackageStoreBackend::UpdateLookupTables(const TArray<FSharedOnDema
 				RedirectEntry.Value = Redirect.TargetPackageId;
 			}
 		}
+		RedirectMap.Shrink();
 
 		++ContainerIdx;
 	}
-
-	EntryMap.Shrink();
-	LocalizedMap.Shrink();
-	RedirectMap.Shrink();
 }
 
-void FOnDemandPackageStoreBackend::UpdateReferencedPackages(const TArray<FSharedOnDemandContainer>& AllContainers, const TArray<TBitArray<>>& ReferencedChunkEntryIndices)
+void FOnDemandPackageStoreBackend::UpdateReferencedPackages(const TArray<TBitArray<>>& ReferencedChunkEntryIndices)
 {
 	check(Mutex.IsLocked());
 
-	for (int32 ContainerIdx = 0; const FSharedOnDemandContainer& SharedContainer : AllContainers)
+	if (ReferencedChunkEntryIndices.IsEmpty())
 	{
-		const FIoContainerHeader& Hdr = *SharedContainer->Header;
-
-		for (const FPackageId PkgId : Hdr.PackageIds)
+		for (FContainer& Container : Containers)
 		{
-			// Entry may not be present if this function is racing with a container unmount
-			// It will be removed on the next full update after the unmount completes
-			if (FOnDemandPackageStoreEntry* OnDemandEntry = EntryMap.Find(PkgId))
+			for (TPair<FPackageId, FOnDemandPackageStoreEntry>& Kv : Container.EntryMap)
 			{
-				const FIoChunkId PkgChunkId = CreatePackageDataChunkId(PkgId);
-				const int32 PkgChunkIdx = SharedContainer->FindChunkEntryIndex(PkgChunkId);
-				check(PkgChunkIdx != INDEX_NONE);
-				const bool IsReferenced = ReferencedChunkEntryIndices[ContainerIdx][PkgChunkIdx];
-
-				OnDemandEntry->bIsInstalled = IsReferenced;
+				Kv.Value.IsReferenced = 0; 
 			}
 		}
-
-		++ContainerIdx;
+	}
+	else
+	{
+		for (int32 ContainerIdx = 0; FContainer& Container : Containers)
+		{
+			const TBitArray<>& Refs = ReferencedChunkEntryIndices[ContainerIdx];
+			for (TPair<FPackageId, FOnDemandPackageStoreEntry>& Kv : Container.EntryMap)
+			{
+				Kv.Value.IsReferenced = Refs[Kv.Value.ChunkEntryIndex];
+			}
+			++ContainerIdx;
+		}
 	}
 }
 
