Register the -IgnoreContainerHeader parameter in IAS tool

#rnx

diff --git a/Engine/Source/Runtime/Experimental/IoStore/OnDemand/Private/Tool/ChunkPlugin.cpp b/Engine/Source/Runtime/Experimental/IoStore/OnDemand/Private/Tool/ChunkPlugin.cpp
--- a/Engine/Source/Runtime/Experimental/IoStore/OnDemand/Private/Tool/ChunkPlugin.cpp
+++ b/Engine/Source/Runtime/Experimental/IoStore/OnDemand/Private/Tool/ChunkPlugin.cpp
@@ -790,6 +790,7 @@ static FCommand ChunkPluginCommand(
 		TArgument<bool>(TEXT("-KeepContainerFiles"),		TEXT("Should we keep the container files after processing them.")),
 		TArgument<bool>(TEXT("-StreamOnDemand"),			TEXT("Set the content to be streamed on-demand")),
 		TArgument<bool>(TEXT("-InstallOnDemand"),			TEXT("Set the content to be installed on-demand")),
+		TArgument<bool>(TEXT("-IgnoreContainerHeader"),		TEXT("Ignore the container header?")),
 		TArgument<FStringView>(TEXT("-BucketPrefix"),		TEXT("Path to prefix to bucket objects")),
 		TArgument<int32>(TEXT("-MaxConcurrentUploads"),		TEXT("Number of simultaneous uploads")),
 		S3Arguments,
