Dataflow vector math library update
- Added multiply and divide node
- Updated property to display values without having to expand themn in the node details
- Added keywords for easier discovery of the nodes

diff --git a/Engine/Source/Runtime/Experimental/Dataflow/Core/Private/Dataflow/DataflowVectorNodes.cpp b/Engine/Source/Runtime/Experimental/Dataflow/Core/Private/Dataflow/DataflowVectorNodes.cpp
--- a/Engine/Source/Runtime/Experimental/Dataflow/Core/Private/Dataflow/DataflowVectorNodes.cpp
+++ b/Engine/Source/Runtime/Experimental/Dataflow/Core/Private/Dataflow/DataflowVectorNodes.cpp
@@ -15,6 +15,8 @@ namespace UE::Dataflow
 		DATAFLOW_NODE_REGISTER_CREATION_FACTORY(FDataflowVectorBreakNode);
 		DATAFLOW_NODE_REGISTER_CREATION_FACTORY(FDataflowVectorAddNode);
 		DATAFLOW_NODE_REGISTER_CREATION_FACTORY(FDataflowVectorSubtractNode);
+		DATAFLOW_NODE_REGISTER_CREATION_FACTORY(FDataflowVectorMultiplyNode);
+		DATAFLOW_NODE_REGISTER_CREATION_FACTORY(FDataflowVectorDivideNode);
 		DATAFLOW_NODE_REGISTER_CREATION_FACTORY(FDataflowVectorDotProductNode);
 		DATAFLOW_NODE_REGISTER_CREATION_FACTORY(FDataflowVectorLengthNode);
 		DATAFLOW_NODE_REGISTER_CREATION_FACTORY(FDataflowVectorSquaredLengthNode);
@@ -184,6 +186,59 @@ void FDataflowVectorSubtractNode::Evaluate(UE::Dataflow::FContext& Context, cons
 
 //-----------------------------------------------------------
 
+FDataflowVectorMultiplyNode::FDataflowVectorMultiplyNode(const UE::Dataflow::FNodeParameters& InParam, FGuid InGuid)
+	: FDataflowNode(InParam, InGuid)
+{
+	RegisterInputConnection(&A);
+	RegisterInputConnection(&B);
+	RegisterOutputConnection(&V);
+	// set default to 4D vector
+	SetOutputConcreteType<FVector4>(&V);
+
+}
+
+void FDataflowVectorMultiplyNode::Evaluate(UE::Dataflow::FContext& Context, const FDataflowOutput* Out) const
+{
+	if (Out->IsA(&V))
+	{
+		const FVector4 InA = GetValue(Context, &A);
+		const FVector4 InB = GetValue(Context, &B);
+		SetValue(Context, (InA * InB), &V);
+	}
+}
+
+//-----------------------------------------------------------
+
+FDataflowVectorDivideNode::FDataflowVectorDivideNode(const UE::Dataflow::FNodeParameters& InParam, FGuid InGuid)
+	: FDataflowNode(InParam, InGuid)
+{
+	RegisterInputConnection(&A);
+	RegisterInputConnection(&B);
+	RegisterInputConnection(&Fallback);
+	RegisterOutputConnection(&V);
+	// set default to 4D vector
+	SetOutputConcreteType<FVector4>(&V);
+
+}
+
+void FDataflowVectorDivideNode::Evaluate(UE::Dataflow::FContext& Context, const FDataflowOutput* Out) const
+{
+	if (Out->IsA(&V))
+	{
+		const FVector4 InA = GetValue(Context, &A);
+		const FVector4 InB = GetValue(Context, &B);
+		const FVector4 InFallBack = GetValue(Context, &Fallback);
+		FVector4 OutV(0);
+		OutV.X = (InB.X == 0) ? InFallBack.X : (InA.X / InB.X);
+		OutV.Y = (InB.Y == 0) ? InFallBack.Y : (InA.Y / InB.Y);
+		OutV.Z = (InB.Z == 0) ? InFallBack.Z : (InA.Z / InB.Z);
+		OutV.W = (InB.W == 0) ? InFallBack.W : (InA.W / InB.W);
+		SetValue(Context, OutV, &V);
+	}
+}
+
+//-----------------------------------------------------------
+
 FDataflowVectorDotProductNode::FDataflowVectorDotProductNode(const UE::Dataflow::FNodeParameters& InParam, FGuid InGuid)
 	: FDataflowNode(InParam, InGuid)
 {
diff --git a/Engine/Source/Runtime/Experimental/Dataflow/Core/Public/Dataflow/DataflowVectorNodes.h b/Engine/Source/Runtime/Experimental/Dataflow/Core/Public/Dataflow/DataflowVectorNodes.h
--- a/Engine/Source/Runtime/Experimental/Dataflow/Core/Public/Dataflow/DataflowVectorNodes.h
+++ b/Engine/Source/Runtime/Experimental/Dataflow/Core/Public/Dataflow/DataflowVectorNodes.h
@@ -13,18 +13,18 @@ USTRUCT()
 struct FDataflowVectorMakeVec2Node : public FDataflowNode
 {
 	GENERATED_USTRUCT_BODY()
-	DATAFLOW_NODE_DEFINE_INTERNAL(FDataflowVectorMakeVec2Node, "MakeVector2", DATAFLOW_MATH_VECTOR_NODES_CATEGORY, "")
+	DATAFLOW_NODE_DEFINE_INTERNAL(FDataflowVectorMakeVec2Node, "MakeVector2", DATAFLOW_MATH_VECTOR_NODES_CATEGORY, "Create Add 2D XY")
 
 public:
 	FDataflowVectorMakeVec2Node(const UE::Dataflow::FNodeParameters& InParam, FGuid InGuid = FGuid::NewGuid());
 
 private:
 	/** X component */
-	UPROPERTY(EditAnywhere, Category = Operands, meta = (DataflowInput));
+	UPROPERTY(EditAnywhere, Category = Operands, meta = (DataflowInput, ShowOnlyInnerProperties));
 	FDataflowNumericTypes X;
 
 	/** Y component */
-	UPROPERTY(EditAnywhere, Category = Operands, meta = (DataflowInput));
+	UPROPERTY(EditAnywhere, Category = Operands, meta = (DataflowInput, ShowOnlyInnerProperties));
 	FDataflowNumericTypes Y;
 
 	/** 2D Vector {X, Y} */
@@ -39,22 +39,22 @@ USTRUCT()
 struct FDataflowVectorMakeVec3Node : public FDataflowNode
 {
 	GENERATED_USTRUCT_BODY()
-	DATAFLOW_NODE_DEFINE_INTERNAL(FDataflowVectorMakeVec3Node, "MakeVector3", DATAFLOW_MATH_VECTOR_NODES_CATEGORY, "")
+	DATAFLOW_NODE_DEFINE_INTERNAL(FDataflowVectorMakeVec3Node, "MakeVector3", DATAFLOW_MATH_VECTOR_NODES_CATEGORY, "Create Add 3D XYZ")
 
 public:
 	FDataflowVectorMakeVec3Node(const UE::Dataflow::FNodeParameters& InParam, FGuid InGuid = FGuid::NewGuid());
 
 private:
 	/** X component */
-	UPROPERTY(EditAnywhere, Category = Operands, meta = (DataflowInput));
+	UPROPERTY(EditAnywhere, Category = Operands, meta = (DataflowInput, ShowOnlyInnerProperties));
 	FDataflowNumericTypes X;
 
 	/** Y component */
-	UPROPERTY(EditAnywhere, Category = Operands, meta = (DataflowInput));
+	UPROPERTY(EditAnywhere, Category = Operands, meta = (DataflowInput, ShowOnlyInnerProperties));
 	FDataflowNumericTypes Y;
 
 	/** Z component */
-	UPROPERTY(EditAnywhere, Category = Operands, meta = (DataflowInput));
+	UPROPERTY(EditAnywhere, Category = Operands, meta = (DataflowInput, ShowOnlyInnerProperties));
 	FDataflowNumericTypes Z;
 
 	/** 3D Vector {X, Y, Z} */
@@ -70,22 +70,22 @@ USTRUCT()
 struct FDataflowVectorMakeVec4Node : public FDataflowNode
 {
 	GENERATED_USTRUCT_BODY()
-	DATAFLOW_NODE_DEFINE_INTERNAL(FDataflowVectorMakeVec4Node, "MakeVector4", DATAFLOW_MATH_VECTOR_NODES_CATEGORY, "")
+	DATAFLOW_NODE_DEFINE_INTERNAL(FDataflowVectorMakeVec4Node, "MakeVector4", DATAFLOW_MATH_VECTOR_NODES_CATEGORY, "Create Add 4D XYZW")
 
 public:
 	FDataflowVectorMakeVec4Node(const UE::Dataflow::FNodeParameters& InParam, FGuid InGuid = FGuid::NewGuid());
 
 private:
 	/** X component */
-	UPROPERTY(EditAnywhere, Category = Operands, meta = (DataflowInput));
+	UPROPERTY(EditAnywhere, Category = Operands, meta = (DataflowInput, ShowOnlyInnerProperties));
 	FDataflowNumericTypes X;
 
 	/** Y component */
-	UPROPERTY(EditAnywhere, Category = Operands, meta = (DataflowInput));
+	UPROPERTY(EditAnywhere, Category = Operands, meta = (DataflowInput, ShowOnlyInnerProperties));
 	FDataflowNumericTypes Y;
 
 	/** Z component */
-	UPROPERTY(EditAnywhere, Category = Operands, meta = (DataflowInput));
+	UPROPERTY(EditAnywhere, Category = Operands, meta = (DataflowInput, ShowOnlyInnerProperties));
 	FDataflowNumericTypes Z;
 
 	/** W component */
@@ -108,14 +108,14 @@ USTRUCT()
 struct FDataflowVectorBreakNode: public FDataflowNode
 {
 	GENERATED_USTRUCT_BODY()
-	DATAFLOW_NODE_DEFINE_INTERNAL(FDataflowVectorBreakNode, "BreakVector", DATAFLOW_MATH_VECTOR_NODES_CATEGORY, "")
+	DATAFLOW_NODE_DEFINE_INTERNAL(FDataflowVectorBreakNode, "BreakVector", DATAFLOW_MATH_VECTOR_NODES_CATEGORY, "Expand Split X Y Z W Components")
 
 public:
 	FDataflowVectorBreakNode(const UE::Dataflow::FNodeParameters& InParam, FGuid InGuid = FGuid::NewGuid());
 
 private:
 	/** Vector to break into components */
-	UPROPERTY(EditAnywhere, Category = Operands, meta = (DataflowInput));
+	UPROPERTY(EditAnywhere, Category = Operands, meta = (DataflowInput, ShowOnlyInnerProperties));
 	FDataflowVectorTypes V;
 
 	/** X component */
@@ -142,18 +142,18 @@ USTRUCT()
 struct FDataflowVectorAddNode : public FDataflowNode
 {
 	GENERATED_USTRUCT_BODY()
-	DATAFLOW_NODE_DEFINE_INTERNAL(FDataflowVectorAddNode, "AddVector", DATAFLOW_MATH_VECTOR_NODES_CATEGORY, "")
+	DATAFLOW_NODE_DEFINE_INTERNAL(FDataflowVectorAddNode, "AddVector", DATAFLOW_MATH_VECTOR_NODES_CATEGORY, "+ Plus Addition")
 
 public:
 	FDataflowVectorAddNode(const UE::Dataflow::FNodeParameters& InParam, FGuid InGuid = FGuid::NewGuid());
 
 private:
 	/** A Vector operand */
-	UPROPERTY(EditAnywhere, Category = Operands, meta = (DataflowInput));
+	UPROPERTY(EditAnywhere, Category = Operands, meta = (DataflowInput, ShowOnlyInnerProperties));
 	FDataflowVectorTypes A;
 
 	/** B Vector operand */
-	UPROPERTY(EditAnywhere, Category = Operands, meta = (DataflowInput));
+	UPROPERTY(EditAnywhere, Category = Operands, meta = (DataflowInput, ShowOnlyInnerProperties));
 	FDataflowVectorTypes B;
 
 	/** Add result V=A+B */
@@ -168,18 +168,18 @@ USTRUCT()
 struct FDataflowVectorSubtractNode : public FDataflowNode
 {
 	GENERATED_USTRUCT_BODY()
-	DATAFLOW_NODE_DEFINE_INTERNAL(FDataflowVectorSubtractNode, "SubtractVector", DATAFLOW_MATH_VECTOR_NODES_CATEGORY, "")
+	DATAFLOW_NODE_DEFINE_INTERNAL(FDataflowVectorSubtractNode, "SubtractVector", DATAFLOW_MATH_VECTOR_NODES_CATEGORY, "- Minus Subtraction")
 
 public:
 	FDataflowVectorSubtractNode(const UE::Dataflow::FNodeParameters& InParam, FGuid InGuid = FGuid::NewGuid());
 
 private:
 	/** A Vector operand */
-	UPROPERTY(EditAnywhere, Category = Operands, meta = (DataflowInput));
+	UPROPERTY(EditAnywhere, Category = Operands, meta = (DataflowInput, ShowOnlyInnerProperties));
 	FDataflowVectorTypes A;
 
 	/** B Vector operand */
-	UPROPERTY(EditAnywhere, Category = Operands, meta = (DataflowInput));
+	UPROPERTY(EditAnywhere, Category = Operands, meta = (DataflowInput, ShowOnlyInnerProperties));
 	FDataflowVectorTypes B;
 
 	/** Add result V=A-B */
@@ -189,23 +189,82 @@ private:
 	virtual void Evaluate(UE::Dataflow::FContext& Context, const FDataflowOutput* Out) const override;
 };
 
+/** Multiply two vectors component wise: V = (A * B) */
+USTRUCT()
+struct FDataflowVectorMultiplyNode : public FDataflowNode
+{
+	GENERATED_USTRUCT_BODY()
+	DATAFLOW_NODE_DEFINE_INTERNAL(FDataflowVectorMultiplyNode, "MultiplyVector", DATAFLOW_MATH_VECTOR_NODES_CATEGORY, "* Multiplication Times")
+
+public:
+	FDataflowVectorMultiplyNode(const UE::Dataflow::FNodeParameters& InParam, FGuid InGuid = FGuid::NewGuid());
+
+private:
+	/** A Vector operand */
+	UPROPERTY(EditAnywhere, Category = Operands, meta = (DataflowInput, ShowOnlyInnerProperties));
+	FDataflowVectorTypes A;
+
+	/** B Vector operand */
+	UPROPERTY(EditAnywhere, Category = Operands, meta = (DataflowInput, ShowOnlyInnerProperties));
+	FDataflowVectorTypes B;
+
+	/** Add result V=A*B */
+	UPROPERTY(meta = (DataflowOutput))
+	FDataflowVectorTypes V;
+
+	virtual void Evaluate(UE::Dataflow::FContext& Context, const FDataflowOutput* Out) const override;
+};
+
+/** 
+* Multiply two vectors component wise: V = (A / B) 
+* When a component of B is zero, use the falback value as a return value for the specific component 
+*/
+USTRUCT()
+struct FDataflowVectorDivideNode : public FDataflowNode
+{
+	GENERATED_USTRUCT_BODY()
+	DATAFLOW_NODE_DEFINE_INTERNAL(FDataflowVectorDivideNode, "DivideVector", DATAFLOW_MATH_VECTOR_NODES_CATEGORY, "/ Division Over")
+
+public:
+	FDataflowVectorDivideNode(const UE::Dataflow::FNodeParameters& InParam, FGuid InGuid = FGuid::NewGuid());
+
+private:
+	/** A Vector operand */
+	UPROPERTY(EditAnywhere, Category = Operands, meta = (DataflowInput, ShowOnlyInnerProperties));
+	FDataflowVectorTypes A;
+
+	/** B Vector operand */
+	UPROPERTY(EditAnywhere, Category = Operands, meta = (DataflowInput, ShowOnlyInnerProperties));
+	FDataflowVectorTypes B;
+
+	/** Fallback Vector used when components of B are zero */
+	UPROPERTY(EditAnywhere, Category = Operands, meta = (DataflowInput, ShowOnlyInnerProperties));
+	FDataflowVectorTypes Fallback;
+
+	/** Add result V=A*B */
+	UPROPERTY(meta = (DataflowOutput))
+	FDataflowVectorTypes V;
+
+	virtual void Evaluate(UE::Dataflow::FContext& Context, const FDataflowOutput* Out) const override;
+};
+
 /** Compute the dot product between two vectors : DotProduct = A.B */
 USTRUCT()
 struct FDataflowVectorDotProductNode : public FDataflowNode
 {
 	GENERATED_USTRUCT_BODY()
-	DATAFLOW_NODE_DEFINE_INTERNAL(FDataflowVectorDotProductNode, "VectorDotProduct", DATAFLOW_MATH_VECTOR_NODES_CATEGORY, "")
+	DATAFLOW_NODE_DEFINE_INTERNAL(FDataflowVectorDotProductNode, "VectorDotProduct", DATAFLOW_MATH_VECTOR_NODES_CATEGORY, "Project")
 
 public:
 	FDataflowVectorDotProductNode(const UE::Dataflow::FNodeParameters& InParam, FGuid InGuid = FGuid::NewGuid());
 
 private:
 	/** A Vector operand */
-	UPROPERTY(EditAnywhere, Category = Operands, meta = (DataflowInput));
+	UPROPERTY(EditAnywhere, Category = Operands, meta = (DataflowInput, ShowOnlyInnerProperties));
 	FDataflowVectorTypes A;
 
 	/** B Vector operand */
-	UPROPERTY(EditAnywhere, Category = Operands, meta = (DataflowInput));
+	UPROPERTY(EditAnywhere, Category = Operands, meta = (DataflowInput, ShowOnlyInnerProperties));
 	FDataflowVectorTypes B;
 
 	/** Resulting dot product : DotProduct=A.B  */
@@ -220,14 +279,14 @@ USTRUCT()
 struct FDataflowVectorLengthNode : public FDataflowNode
 {
 	GENERATED_USTRUCT_BODY()
-	DATAFLOW_NODE_DEFINE_INTERNAL(FDataflowVectorLengthNode, "VectorLength", DATAFLOW_MATH_VECTOR_NODES_CATEGORY, "")
+	DATAFLOW_NODE_DEFINE_INTERNAL(FDataflowVectorLengthNode, "VectorLength", DATAFLOW_MATH_VECTOR_NODES_CATEGORY, "Size Magnitude")
 
 public:
 	FDataflowVectorLengthNode(const UE::Dataflow::FNodeParameters& InParam, FGuid InGuid = FGuid::NewGuid());
 
 private:
 	/** Vector to get length from */
-	UPROPERTY(EditAnywhere, Category = Operands, meta = (DataflowInput));
+	UPROPERTY(EditAnywhere, Category = Operands, meta = (DataflowInput, ShowOnlyInnerProperties));
 	FDataflowVectorTypes V;
 
 	/** Length of the input vector : Length=|V| */
@@ -242,14 +301,14 @@ USTRUCT()
 struct FDataflowVectorSquaredLengthNode : public FDataflowNode
 {
 	GENERATED_USTRUCT_BODY()
-	DATAFLOW_NODE_DEFINE_INTERNAL(FDataflowVectorSquaredLengthNode, "VectorSquaredLength", DATAFLOW_MATH_VECTOR_NODES_CATEGORY, "")
+	DATAFLOW_NODE_DEFINE_INTERNAL(FDataflowVectorSquaredLengthNode, "VectorSquaredLength", DATAFLOW_MATH_VECTOR_NODES_CATEGORY, "Size Magnitude")
 
 public:
 	FDataflowVectorSquaredLengthNode(const UE::Dataflow::FNodeParameters& InParam, FGuid InGuid = FGuid::NewGuid());
 
 private:
 	/** Vector to get squared length from */
-	UPROPERTY(EditAnywhere, Category = Operands, meta = (DataflowInput));
+	UPROPERTY(EditAnywhere, Category = Operands, meta = (DataflowInput, ShowOnlyInnerProperties));
 	FDataflowVectorTypes V;
 
 	/** Length of the input vector : SquaredLength = |V||V| */
@@ -271,11 +330,11 @@ public:
 
 private:
 	/** A Vector operand */
-	UPROPERTY(EditAnywhere, Category = Operands, meta = (DataflowInput));
+	UPROPERTY(EditAnywhere, Category = Operands, meta = (DataflowInput, ShowOnlyInnerProperties));
 	FDataflowVectorTypes A;
 
 	/** B Vector operand */
-	UPROPERTY(EditAnywhere, Category = Operands, meta = (DataflowInput));
+	UPROPERTY(EditAnywhere, Category = Operands, meta = (DataflowInput, ShowOnlyInnerProperties));
 	FDataflowVectorTypes B;
 
 	/** Distance between A and B : Distance=|B-A| */
@@ -300,11 +359,11 @@ public:
 
 private:
 	/** A Vector operand */
-	UPROPERTY(EditAnywhere, Category = Operands, meta = (DataflowInput));
+	UPROPERTY(EditAnywhere, Category = Operands, meta = (DataflowInput, ShowOnlyInnerProperties));
 	FDataflowVectorTypes A;
 
 	/** B Vector operand */
-	UPROPERTY(EditAnywhere, Category = Operands, meta = (DataflowInput));
+	UPROPERTY(EditAnywhere, Category = Operands, meta = (DataflowInput, ShowOnlyInnerProperties));
 	FDataflowVectorTypes B;
 
 	/** Resulting cross product of A and B : CrossProduct=B^A */
@@ -319,18 +378,18 @@ USTRUCT()
 struct FDataflowVectorScaleNode : public FDataflowNode
 {
 	GENERATED_USTRUCT_BODY()
-	DATAFLOW_NODE_DEFINE_INTERNAL(FDataflowVectorScaleNode, "ScaleVector", DATAFLOW_MATH_VECTOR_NODES_CATEGORY, "")
+	DATAFLOW_NODE_DEFINE_INTERNAL(FDataflowVectorScaleNode, "ScaleVector", DATAFLOW_MATH_VECTOR_NODES_CATEGORY, "Scalar Float Double")
 
 public:
 	FDataflowVectorScaleNode(const UE::Dataflow::FNodeParameters& InParam, FGuid InGuid = FGuid::NewGuid());
 
 private:
 	/** Vector to scale */
-	UPROPERTY(EditAnywhere, Category = Operands, meta = (DataflowInput));
+	UPROPERTY(EditAnywhere, Category = Operands, meta = (DataflowInput, ShowOnlyInnerProperties));
 	FDataflowVectorTypes V;
 
 	/** Scale factor */
-	UPROPERTY(EditAnywhere, Category = Operands, meta = (DataflowInput));
+	UPROPERTY(EditAnywhere, Category = Operands, meta = (DataflowInput, ShowOnlyInnerProperties));
 	FDataflowNumericTypes Scale;
 
 	/** Scaled vector : Scaled=(V * Scale) */
@@ -352,7 +411,7 @@ public:
 
 private:
 	/** Vector to normalize */
-	UPROPERTY(EditAnywhere, Category = Operands, meta = (DataflowInput));
+	UPROPERTY(EditAnywhere, Category = Operands, meta = (DataflowInput, ShowOnlyInnerProperties));
 	FDataflowVectorTypes V;
 
 	/** Normalized vector : Normalized=V/|V| */
