Dataflow math library update
- Add keywords to each node
- Make anytype inputs to be auto expanded

diff --git a/Engine/Source/Runtime/Experimental/Dataflow/Core/Public/Dataflow/DataflowMathNodes.h b/Engine/Source/Runtime/Experimental/Dataflow/Core/Public/Dataflow/DataflowMathNodes.h
--- a/Engine/Source/Runtime/Experimental/Dataflow/Core/Public/Dataflow/DataflowMathNodes.h
+++ b/Engine/Source/Runtime/Experimental/Dataflow/Core/Public/Dataflow/DataflowMathNodes.h
@@ -24,7 +24,7 @@ struct FDataflowMathOneInputOperatorNode : public FDataflowNode
 {
 	GENERATED_USTRUCT_BODY()
 
-	UPROPERTY(EditAnywhere, Category = Operands, meta = (DataflowInput));
+	UPROPERTY(EditAnywhere, Category = Operands, meta = (DataflowInput, ShowOnlyInnerProperties));
 	FDataflowNumericTypes A;
 
 	UPROPERTY(meta = (DataflowOutput))
@@ -46,10 +46,10 @@ struct FDataflowMathTwoInputsOperatorNode : public FDataflowNode
 {
 	GENERATED_USTRUCT_BODY()
 
-	UPROPERTY(EditAnywhere, Category = Operands, meta = (DataflowInput));
+	UPROPERTY(EditAnywhere, Category = Operands, meta = (DataflowInput, ShowOnlyInnerProperties));
 	FDataflowNumericTypes A;
 
-	UPROPERTY(EditAnywhere, Category = Operands, meta = (DataflowInput));
+	UPROPERTY(EditAnywhere, Category = Operands, meta = (DataflowInput, ShowOnlyInnerProperties));
 	FDataflowNumericTypes B;
 
 	UPROPERTY(meta = (DataflowOutput))
@@ -70,7 +70,7 @@ USTRUCT()
 struct FDataflowMathAddNode : public FDataflowMathTwoInputsOperatorNode
 {
 	GENERATED_USTRUCT_BODY()
-	DATAFLOW_NODE_DEFINE_INTERNAL(FDataflowMathAddNode, "Add", DATAFLOW_MATH_NODES_CATEGORY, "")
+	DATAFLOW_NODE_DEFINE_INTERNAL(FDataflowMathAddNode, "Add", DATAFLOW_MATH_NODES_CATEGORY, "+ Addition Plus")
 
 public:
 	FDataflowMathAddNode(const UE::Dataflow::FNodeParameters& InParam, FGuid InGuid = FGuid::NewGuid());
@@ -82,7 +82,7 @@ USTRUCT()
 struct FDataflowMathSubtractNode : public FDataflowMathTwoInputsOperatorNode
 {
 	GENERATED_USTRUCT_BODY()
-	DATAFLOW_NODE_DEFINE_INTERNAL(FDataflowMathSubtractNode, "Subtract", DATAFLOW_MATH_NODES_CATEGORY, "")
+	DATAFLOW_NODE_DEFINE_INTERNAL(FDataflowMathSubtractNode, "Subtract", DATAFLOW_MATH_NODES_CATEGORY, "- Subtraction Minus")
 
 public:
 	FDataflowMathSubtractNode(const UE::Dataflow::FNodeParameters& InParam, FGuid InGuid = FGuid::NewGuid());
@@ -94,7 +94,7 @@ USTRUCT()
 struct FDataflowMathMultiplyNode : public FDataflowMathTwoInputsOperatorNode
 {
 	GENERATED_USTRUCT_BODY()
-	DATAFLOW_NODE_DEFINE_INTERNAL(FDataflowMathMultiplyNode, "Multiply", DATAFLOW_MATH_NODES_CATEGORY, "")
+	DATAFLOW_NODE_DEFINE_INTERNAL(FDataflowMathMultiplyNode, "Multiply", DATAFLOW_MATH_NODES_CATEGORY, "* Multiplication Time")
 
 public:
 	FDataflowMathMultiplyNode(const UE::Dataflow::FNodeParameters& InParam, FGuid InGuid = FGuid::NewGuid());
@@ -109,9 +109,9 @@ USTRUCT()
 struct FDataflowMathDivideNode : public FDataflowMathTwoInputsOperatorNode
 {
 	GENERATED_USTRUCT_BODY()
-	DATAFLOW_NODE_DEFINE_INTERNAL(FDataflowMathDivideNode, "Divide", DATAFLOW_MATH_NODES_CATEGORY, "")
+	DATAFLOW_NODE_DEFINE_INTERNAL(FDataflowMathDivideNode, "Divide", DATAFLOW_MATH_NODES_CATEGORY, "/ Division")
 
-	UPROPERTY(EditAnywhere, Category = SafeDivide, meta = (DataflowInput));
+	UPROPERTY(EditAnywhere, Category = SafeDivide, meta = (DataflowInput, ShowOnlyInnerProperties));
 	FDataflowNumericTypes Fallback;
 
 public:
@@ -128,7 +128,7 @@ USTRUCT(meta = (Deprecated = "5.6"))
 struct FDataflowMathMinimumNode : public FDataflowMathTwoInputsOperatorNode
 {
 	GENERATED_USTRUCT_BODY()
-	DATAFLOW_NODE_DEFINE_INTERNAL(FDataflowMathMinimumNode, "Minimum", DATAFLOW_MATH_NODES_CATEGORY, "")
+	DATAFLOW_NODE_DEFINE_INTERNAL(FDataflowMathMinimumNode, "Minimum", DATAFLOW_MATH_NODES_CATEGORY, "Lowest Smallest")
 
 public:
 	FDataflowMathMinimumNode(const UE::Dataflow::FNodeParameters& InParam, FGuid InGuid = FGuid::NewGuid());
@@ -142,7 +142,7 @@ USTRUCT()
 struct FDataflowMathMinimumNode_v2 : public FDataflowNode
 {
 	GENERATED_USTRUCT_BODY()
-	DATAFLOW_NODE_DEFINE_INTERNAL(FDataflowMathMinimumNode_v2, "Minimum", DATAFLOW_MATH_NODES_CATEGORY, "")
+	DATAFLOW_NODE_DEFINE_INTERNAL(FDataflowMathMinimumNode_v2, "Minimum", DATAFLOW_MATH_NODES_CATEGORY, "Lowest Smallest")
 
 private:
 	UPROPERTY()
@@ -179,7 +179,7 @@ USTRUCT(meta = (Deprecated = "5.6"))
 struct FDataflowMathMaximumNode : public FDataflowMathTwoInputsOperatorNode
 {
 	GENERATED_USTRUCT_BODY()
-	DATAFLOW_NODE_DEFINE_INTERNAL(FDataflowMathMaximumNode, "Maximum", DATAFLOW_MATH_NODES_CATEGORY, "")
+	DATAFLOW_NODE_DEFINE_INTERNAL(FDataflowMathMaximumNode, "Maximum", DATAFLOW_MATH_NODES_CATEGORY, "Highest Larger")
 
 public:
 	FDataflowMathMaximumNode(const UE::Dataflow::FNodeParameters& InParam, FGuid InGuid = FGuid::NewGuid());
@@ -191,7 +191,7 @@ USTRUCT()
 struct FDataflowMathMaximumNode_v2 : public FDataflowNode
 {
 	GENERATED_USTRUCT_BODY()
-	DATAFLOW_NODE_DEFINE_INTERNAL(FDataflowMathMaximumNode_v2, "Maximum", DATAFLOW_MATH_NODES_CATEGORY, "")
+	DATAFLOW_NODE_DEFINE_INTERNAL(FDataflowMathMaximumNode_v2, "Maximum", DATAFLOW_MATH_NODES_CATEGORY, "Highest Larger")
 
 private:
 	UPROPERTY()
@@ -227,9 +227,9 @@ USTRUCT()
 struct FDataflowMathReciprocalNode : public FDataflowMathOneInputOperatorNode
 {
 	GENERATED_USTRUCT_BODY()
-	DATAFLOW_NODE_DEFINE_INTERNAL(FDataflowMathReciprocalNode, "Reciprocal", DATAFLOW_MATH_NODES_CATEGORY, "")
+	DATAFLOW_NODE_DEFINE_INTERNAL(FDataflowMathReciprocalNode, "Reciprocal", DATAFLOW_MATH_NODES_CATEGORY, "OneOver OneDividedBy")
 
-	UPROPERTY(EditAnywhere, Category = SafeDivide, meta = (DataflowInput));
+	UPROPERTY(EditAnywhere, Category = SafeDivide, meta = (DataflowInput, ShowOnlyInnerProperties));
 	FDataflowNumericTypes Fallback;
 
 public:
@@ -281,9 +281,9 @@ USTRUCT()
 struct FDataflowMathInverseSquareRootNode : public FDataflowMathOneInputOperatorNode
 {
 	GENERATED_USTRUCT_BODY()
-	DATAFLOW_NODE_DEFINE_INTERNAL(FDataflowMathInverseSquareRootNode, "InverseSquareRoot", DATAFLOW_MATH_NODES_CATEGORY, "")
+	DATAFLOW_NODE_DEFINE_INTERNAL(FDataflowMathInverseSquareRootNode, "InverseSquareRoot", DATAFLOW_MATH_NODES_CATEGORY, "OneOverSquareRoot")
 
-	UPROPERTY(EditAnywhere, Category = SafeDivide, meta = (DataflowInput));
+	UPROPERTY(EditAnywhere, Category = SafeDivide, meta = (DataflowInput, ShowOnlyInnerProperties));
 	FDataflowNumericTypes Fallback;
 
 public:
@@ -296,7 +296,7 @@ USTRUCT()
 struct FDataflowMathNegateNode : public FDataflowMathOneInputOperatorNode
 {
 	GENERATED_USTRUCT_BODY()
-	DATAFLOW_NODE_DEFINE_INTERNAL(FDataflowMathNegateNode, "Negate", DATAFLOW_MATH_NODES_CATEGORY, "")
+	DATAFLOW_NODE_DEFINE_INTERNAL(FDataflowMathNegateNode, "Negate", DATAFLOW_MATH_NODES_CATEGORY, "Minus Negative")
 
 public:
 	FDataflowMathNegateNode(const UE::Dataflow::FNodeParameters& InParam, FGuid InGuid = FGuid::NewGuid());
@@ -308,7 +308,7 @@ USTRUCT()
 struct FDataflowMathAbsNode : public FDataflowMathOneInputOperatorNode
 {
 	GENERATED_USTRUCT_BODY()
-	DATAFLOW_NODE_DEFINE_INTERNAL(FDataflowMathAbsNode, "Abs", DATAFLOW_MATH_NODES_CATEGORY, "")
+	DATAFLOW_NODE_DEFINE_INTERNAL(FDataflowMathAbsNode, "Abs", DATAFLOW_MATH_NODES_CATEGORY, "Absolute Positive")
 
 public:
 	FDataflowMathAbsNode(const UE::Dataflow::FNodeParameters& InParam, FGuid InGuid = FGuid::NewGuid());
@@ -320,7 +320,7 @@ USTRUCT()
 struct FDataflowMathFloorNode : public FDataflowMathOneInputOperatorNode
 {
 	GENERATED_USTRUCT_BODY()
-	DATAFLOW_NODE_DEFINE_INTERNAL(FDataflowMathFloorNode, "Floor", DATAFLOW_MATH_NODES_CATEGORY, "")
+	DATAFLOW_NODE_DEFINE_INTERNAL(FDataflowMathFloorNode, "Floor", DATAFLOW_MATH_NODES_CATEGORY, "Round Lowest Integer")
 
 public:
 	FDataflowMathFloorNode(const UE::Dataflow::FNodeParameters& InParam, FGuid InGuid = FGuid::NewGuid());
@@ -332,7 +332,7 @@ USTRUCT()
 struct FDataflowMathCeilNode : public FDataflowMathOneInputOperatorNode
 {
 	GENERATED_USTRUCT_BODY()
-	DATAFLOW_NODE_DEFINE_INTERNAL(FDataflowMathCeilNode, "Ceil", DATAFLOW_MATH_NODES_CATEGORY, "")
+	DATAFLOW_NODE_DEFINE_INTERNAL(FDataflowMathCeilNode, "Ceil", DATAFLOW_MATH_NODES_CATEGORY, "Round Higher Integer")
 
 public:
 	FDataflowMathCeilNode(const UE::Dataflow::FNodeParameters& InParam, FGuid InGuid = FGuid::NewGuid());
@@ -344,7 +344,7 @@ USTRUCT()
 struct FDataflowMathRoundNode : public FDataflowMathOneInputOperatorNode
 {
 	GENERATED_USTRUCT_BODY()
-	DATAFLOW_NODE_DEFINE_INTERNAL(FDataflowMathRoundNode, "Round", DATAFLOW_MATH_NODES_CATEGORY, "")
+	DATAFLOW_NODE_DEFINE_INTERNAL(FDataflowMathRoundNode, "Round", DATAFLOW_MATH_NODES_CATEGORY, "Round Mid 0.5")
 
 public:
 	FDataflowMathRoundNode(const UE::Dataflow::FNodeParameters& InParam, FGuid InGuid = FGuid::NewGuid());
@@ -356,7 +356,7 @@ USTRUCT()
 struct FDataflowMathTruncNode : public FDataflowMathOneInputOperatorNode
 {
 	GENERATED_USTRUCT_BODY()
-	DATAFLOW_NODE_DEFINE_INTERNAL(FDataflowMathTruncNode, "Trunc", DATAFLOW_MATH_NODES_CATEGORY, "")
+	DATAFLOW_NODE_DEFINE_INTERNAL(FDataflowMathTruncNode, "Trunc", DATAFLOW_MATH_NODES_CATEGORY, "Integer Truncate")
 
 public:
 	FDataflowMathTruncNode(const UE::Dataflow::FNodeParameters& InParam, FGuid InGuid = FGuid::NewGuid());
@@ -368,7 +368,7 @@ USTRUCT()
 struct FDataflowMathFracNode : public FDataflowMathOneInputOperatorNode
 {
 	GENERATED_USTRUCT_BODY()
-	DATAFLOW_NODE_DEFINE_INTERNAL(FDataflowMathFracNode, "Frac", DATAFLOW_MATH_NODES_CATEGORY, "")
+	DATAFLOW_NODE_DEFINE_INTERNAL(FDataflowMathFracNode, "Frac", DATAFLOW_MATH_NODES_CATEGORY, "Fractional Decimal Point")
 
 public:
 	FDataflowMathFracNode(const UE::Dataflow::FNodeParameters& InParam, FGuid InGuid = FGuid::NewGuid());
@@ -380,7 +380,7 @@ USTRUCT()
 struct FDataflowMathPowNode : public FDataflowMathTwoInputsOperatorNode
 {
 	GENERATED_USTRUCT_BODY()
-	DATAFLOW_NODE_DEFINE_INTERNAL(FDataflowMathPowNode, "Pow", DATAFLOW_MATH_NODES_CATEGORY, "")
+	DATAFLOW_NODE_DEFINE_INTERNAL(FDataflowMathPowNode, "Pow", DATAFLOW_MATH_NODES_CATEGORY, "Power Exponent")
 
 public:
 	FDataflowMathPowNode(const UE::Dataflow::FNodeParameters& InParam, FGuid InGuid = FGuid::NewGuid());
@@ -395,9 +395,9 @@ USTRUCT()
 struct FDataflowMathLogXNode : public FDataflowMathOneInputOperatorNode
 {
 	GENERATED_USTRUCT_BODY()
-	DATAFLOW_NODE_DEFINE_INTERNAL(FDataflowMathLogXNode, "LogX", DATAFLOW_MATH_NODES_CATEGORY, "")
+	DATAFLOW_NODE_DEFINE_INTERNAL(FDataflowMathLogXNode, "LogX", DATAFLOW_MATH_NODES_CATEGORY, "Logarithm")
 
-	UPROPERTY(EditAnywhere, Category = SafeDivide, meta = (DataflowInput));
+	UPROPERTY(EditAnywhere, Category = SafeDivide, meta = (DataflowInput, ShowOnlyInnerProperties));
 	FDataflowNumericTypes Base;
 
 public:
@@ -410,7 +410,7 @@ USTRUCT()
 struct FDataflowMathLogNode : public FDataflowMathOneInputOperatorNode
 {
 	GENERATED_USTRUCT_BODY()
-	DATAFLOW_NODE_DEFINE_INTERNAL(FDataflowMathLogNode, "Log", DATAFLOW_MATH_NODES_CATEGORY, "")
+	DATAFLOW_NODE_DEFINE_INTERNAL(FDataflowMathLogNode, "Log", DATAFLOW_MATH_NODES_CATEGORY, "Logarithm")
 
 public:
 	FDataflowMathLogNode(const UE::Dataflow::FNodeParameters& InParam, FGuid InGuid = FGuid::NewGuid());
@@ -422,7 +422,7 @@ USTRUCT()
 struct FDataflowMathExpNode : public FDataflowMathOneInputOperatorNode
 {
 	GENERATED_USTRUCT_BODY()
-	DATAFLOW_NODE_DEFINE_INTERNAL(FDataflowMathExpNode, "Exp", DATAFLOW_MATH_NODES_CATEGORY, "")
+	DATAFLOW_NODE_DEFINE_INTERNAL(FDataflowMathExpNode, "Exp", DATAFLOW_MATH_NODES_CATEGORY, "Exponential")
 
 public:
 	FDataflowMathExpNode(const UE::Dataflow::FNodeParameters& InParam, FGuid InGuid = FGuid::NewGuid());
@@ -434,7 +434,7 @@ USTRUCT()
 struct FDataflowMathSignNode : public FDataflowMathOneInputOperatorNode
 {
 	GENERATED_USTRUCT_BODY()
-	DATAFLOW_NODE_DEFINE_INTERNAL(FDataflowMathSignNode, "Sign", DATAFLOW_MATH_NODES_CATEGORY, "")
+	DATAFLOW_NODE_DEFINE_INTERNAL(FDataflowMathSignNode, "Sign", DATAFLOW_MATH_NODES_CATEGORY, "Minus Plus Positive Negative")
 
 public:
 	FDataflowMathSignNode(const UE::Dataflow::FNodeParameters& InParam, FGuid InGuid = FGuid::NewGuid());
@@ -479,7 +479,7 @@ USTRUCT()
 struct FDataflowMathConstantNode : public FDataflowNode
 {
 	GENERATED_USTRUCT_BODY()
-	DATAFLOW_NODE_DEFINE_INTERNAL(FDataflowMathConstantNode, "Constants", DATAFLOW_MATH_NODES_CATEGORY, "")
+	DATAFLOW_NODE_DEFINE_INTERNAL(FDataflowMathConstantNode, "Constants", DATAFLOW_MATH_NODES_CATEGORY, "Pi Half Two Four Inverse SquareRoot Sqrt Cube Square e Gamma Golden Ratio")
 
 	/** Math constant to output */
 	UPROPERTY(EditAnywhere, Category = "Constants");
@@ -501,12 +501,12 @@ USTRUCT()
 struct FDataflowMathClampNode : public FDataflowMathOneInputOperatorNode
 {
 	GENERATED_USTRUCT_BODY()
-	DATAFLOW_NODE_DEFINE_INTERNAL(FDataflowMathClampNode, "Clamp", DATAFLOW_MATH_NODES_CATEGORY, "")
+	DATAFLOW_NODE_DEFINE_INTERNAL(FDataflowMathClampNode, "Clamp", DATAFLOW_MATH_NODES_CATEGORY, "Saturate Limits")
 
-	UPROPERTY(EditAnywhere, Category = "Clamp", meta = (DataflowInput))
+	UPROPERTY(EditAnywhere, Category = "Clamp", meta = (DataflowInput, ShowOnlyInnerProperties))
 	FDataflowNumericTypes Min;
 
-	UPROPERTY(EditAnywhere, Category = "Clamp", meta = (DataflowInput))
+	UPROPERTY(EditAnywhere, Category = "Clamp", meta = (DataflowInput, ShowOnlyInnerProperties))
 	FDataflowNumericTypes Max;
 
 public:
@@ -527,7 +527,7 @@ USTRUCT()
 struct FDataflowMathSinNode : public FDataflowMathOneInputOperatorNode
 {
 	GENERATED_USTRUCT_BODY()
-	DATAFLOW_NODE_DEFINE_INTERNAL(FDataflowMathSinNode, "Sin", DATAFLOW_MATH_TRIG_NODES_CATEGORY, "")
+	DATAFLOW_NODE_DEFINE_INTERNAL(FDataflowMathSinNode, "Sin", DATAFLOW_MATH_TRIG_NODES_CATEGORY, "Sine Trigonometry Circle Angle Degrees Radians")
 
 public:
 	FDataflowMathSinNode(const UE::Dataflow::FNodeParameters& InParam, FGuid InGuid = FGuid::NewGuid());
@@ -539,7 +539,7 @@ USTRUCT()
 struct FDataflowMathCosNode : public FDataflowMathOneInputOperatorNode
 {
 	GENERATED_USTRUCT_BODY()
-	DATAFLOW_NODE_DEFINE_INTERNAL(FDataflowMathCosNode, "Cos", DATAFLOW_MATH_TRIG_NODES_CATEGORY, "")
+	DATAFLOW_NODE_DEFINE_INTERNAL(FDataflowMathCosNode, "Cos", DATAFLOW_MATH_TRIG_NODES_CATEGORY, "Cosine Trigonometry Circle Angle Degrees Radians")
 
 public:
 	FDataflowMathCosNode(const UE::Dataflow::FNodeParameters& InParam, FGuid InGuid = FGuid::NewGuid());
@@ -551,7 +551,7 @@ USTRUCT()
 struct FDataflowMathTanNode : public FDataflowMathOneInputOperatorNode
 {
 	GENERATED_USTRUCT_BODY()
-	DATAFLOW_NODE_DEFINE_INTERNAL(FDataflowMathTanNode, "Tan", DATAFLOW_MATH_TRIG_NODES_CATEGORY, "")
+	DATAFLOW_NODE_DEFINE_INTERNAL(FDataflowMathTanNode, "Tan", DATAFLOW_MATH_TRIG_NODES_CATEGORY, "Tangent Trigonometry Circle Angle Degrees Radians")
 
 public:
 	FDataflowMathTanNode(const UE::Dataflow::FNodeParameters& InParam, FGuid InGuid = FGuid::NewGuid());
@@ -563,7 +563,7 @@ USTRUCT()
 struct FDataflowMathArcSinNode : public FDataflowMathOneInputOperatorNode
 {
 	GENERATED_USTRUCT_BODY()
-	DATAFLOW_NODE_DEFINE_INTERNAL(FDataflowMathArcSinNode, "ArcSin", DATAFLOW_MATH_TRIG_NODES_CATEGORY, "")
+	DATAFLOW_NODE_DEFINE_INTERNAL(FDataflowMathArcSinNode, "ArcSin", DATAFLOW_MATH_TRIG_NODES_CATEGORY, "Sine Trigonometry Circle Angle Degrees Radians")
 
 public:
 	FDataflowMathArcSinNode(const UE::Dataflow::FNodeParameters& InParam, FGuid InGuid = FGuid::NewGuid());
@@ -575,7 +575,7 @@ USTRUCT()
 struct FDataflowMathArcCosNode : public FDataflowMathOneInputOperatorNode
 {
 	GENERATED_USTRUCT_BODY()
-	DATAFLOW_NODE_DEFINE_INTERNAL(FDataflowMathArcCosNode, "ArcCos", DATAFLOW_MATH_TRIG_NODES_CATEGORY, "")
+	DATAFLOW_NODE_DEFINE_INTERNAL(FDataflowMathArcCosNode, "ArcCos", DATAFLOW_MATH_TRIG_NODES_CATEGORY, "Cosine Trigonometry Circle Angle Degrees Radians")
 
 public:
 	FDataflowMathArcCosNode(const UE::Dataflow::FNodeParameters& InParam, FGuid InGuid = FGuid::NewGuid());
@@ -587,7 +587,7 @@ USTRUCT()
 struct FDataflowMathArcTanNode : public FDataflowMathOneInputOperatorNode
 {
 	GENERATED_USTRUCT_BODY()
-	DATAFLOW_NODE_DEFINE_INTERNAL(FDataflowMathArcTanNode, "ArcTan", DATAFLOW_MATH_TRIG_NODES_CATEGORY, "")
+	DATAFLOW_NODE_DEFINE_INTERNAL(FDataflowMathArcTanNode, "ArcTan", DATAFLOW_MATH_TRIG_NODES_CATEGORY, "Tangent Trigonometry Circle Angle Degrees Radians")
 
 public:
 	FDataflowMathArcTanNode(const UE::Dataflow::FNodeParameters& InParam, FGuid InGuid = FGuid::NewGuid());
@@ -599,7 +599,7 @@ USTRUCT()
 struct FDataflowMathArcTan2Node : public FDataflowMathTwoInputsOperatorNode
 {
 	GENERATED_USTRUCT_BODY()
-	DATAFLOW_NODE_DEFINE_INTERNAL(FDataflowMathArcTan2Node, "ArcTan2", DATAFLOW_MATH_TRIG_NODES_CATEGORY, "")
+	DATAFLOW_NODE_DEFINE_INTERNAL(FDataflowMathArcTan2Node, "ArcTan2", DATAFLOW_MATH_TRIG_NODES_CATEGORY, "Tangent Trigonometry Circle Angle Degrees Radians")
 
 public:
 	FDataflowMathArcTan2Node(const UE::Dataflow::FNodeParameters& InParam, FGuid InGuid = FGuid::NewGuid());
@@ -612,7 +612,7 @@ USTRUCT()
 struct FDataflowMathDegToRadNode : public FDataflowMathOneInputOperatorNode
 {
 	GENERATED_USTRUCT_BODY()
-	DATAFLOW_NODE_DEFINE_INTERNAL(FDataflowMathDegToRadNode, "DegToRad", DATAFLOW_MATH_TRIG_NODES_CATEGORY, "")
+	DATAFLOW_NODE_DEFINE_INTERNAL(FDataflowMathDegToRadNode, "DegToRad", DATAFLOW_MATH_TRIG_NODES_CATEGORY, "Trigonometry Circle Angle Degrees Radians Convert Unit")
 
 public:
 	FDataflowMathDegToRadNode(const UE::Dataflow::FNodeParameters& InParam, FGuid InGuid = FGuid::NewGuid());
@@ -624,7 +624,7 @@ USTRUCT()
 struct FDataflowMathRadToDegNode : public FDataflowMathOneInputOperatorNode
 {
 	GENERATED_USTRUCT_BODY()
-	DATAFLOW_NODE_DEFINE_INTERNAL(FDataflowMathRadToDegNode, "RadToDeg", DATAFLOW_MATH_TRIG_NODES_CATEGORY, "")
+	DATAFLOW_NODE_DEFINE_INTERNAL(FDataflowMathRadToDegNode, "RadToDeg", DATAFLOW_MATH_TRIG_NODES_CATEGORY, "Trigonometry Circle Angle Degrees Radians Convert Unit")
 
 public:
 	FDataflowMathRadToDegNode(const UE::Dataflow::FNodeParameters& InParam, FGuid InGuid = FGuid::NewGuid());
