Fixed chains bouncing off of Spherical Collisions
On XSX and in single precision, when code is optimized the closest segment is not generated the correct result.
Used a special dot product that avoid avoid FMA.
So this issue is due to the FMA intrinsic used in single precision on the dot product which bring huge errors.

diff --git a/Engine/Source/Runtime/Experimental/Chaos/Private/Chaos/CollisionOneShotManifoldsMiscShapes.cpp b/Engine/Source/Runtime/Experimental/Chaos/Private/Chaos/CollisionOneShotManifoldsMiscShapes.cpp
--- a/Engine/Source/Runtime/Experimental/Chaos/Private/Chaos/CollisionOneShotManifoldsMiscShapes.cpp
+++ b/Engine/Source/Runtime/Experimental/Chaos/Private/Chaos/CollisionOneShotManifoldsMiscShapes.cpp
@@ -117,7 +117,7 @@ namespace Chaos
 			// @todo(chaos) this would be much simpler if the sphere's were always at the origin and capsules were at the origin and axis aligned
 			const FRigidTransform3f SphereToCapsuleTransform = FRigidTransform3f(SphereTransform.GetRelativeTransformNoScale(CapsuleTransform));
 			const FVec3f SpherePos = SphereToCapsuleTransform.TransformPositionNoScale(Sphere.GetCenterf());
-			const FRealSingle NearPosT = Utilities::ClosestTimeOnLineSegment(SpherePos, Capsule.GetX1f(), Capsule.GetX2f());
+			const FRealSingle NearPosT = Utilities::ClosestTimeOnLineSegment<FRealSingle>(SpherePos, Capsule.GetX1f(), Capsule.GetX2f());
 
 			// Add the closest contact point to the manifold
 			const FVec3f NearPos = FMath::Lerp(Capsule.GetX1f(), Capsule.GetX2f(), NearPosT);
@@ -131,7 +131,7 @@ namespace Chaos
 					const FVec3 NearPosDir = NearPosDelta / NearPosDistance;
 					const FRealSingle NearPhi = NearPosDistance - Sphere.GetRadiusf() - Capsule.GetRadiusf();
 
-					FContactPoint NearContactPoint;
+					FContactPointf NearContactPoint;
 					NearContactPoint.ShapeContactPoints[0] = SphereToCapsuleTransform.InverseTransformPositionNoScale(SpherePos - Sphere.GetRadiusf() * NearPosDir);
 					NearContactPoint.ShapeContactPoints[1] = NearPos + Capsule.GetRadiusf() * NearPosDir;
 					NearContactPoint.ShapeContactNormal = NearPosDir;
diff --git a/Engine/Source/Runtime/Experimental/Chaos/Public/Chaos/Utilities.h b/Engine/Source/Runtime/Experimental/Chaos/Public/Chaos/Utilities.h
--- a/Engine/Source/Runtime/Experimental/Chaos/Public/Chaos/Utilities.h
+++ b/Engine/Source/Runtime/Experimental/Chaos/Public/Chaos/Utilities.h
@@ -763,21 +763,37 @@ namespace Chaos
 			C2 = P2 + T * D2;
 		}
 
+		namespace
+		{
+			// Use dot product with a barrier to avoid fma to generate inaccuracy results
+			template<typename T>
+			T DotWithoutFMA(const TVec3<T>& A, const TVec3<T>& B)
+			{
+				T X = A.X * B.X;
+				T Y = A.Y * B.Y;
+				T Z = A.Z * B.Z;
+#if defined(_MSC_VER) && !defined(__clang__)
+				_ReadWriteBarrier();
+#endif
+				return X + Y + Z;
+			}
+		}
+
 		template<typename T>
 		inline T ClosestTimeOnLineSegment(const TVec3<T>& Point, const TVec3<T>& StartPoint, const TVec3<T>& EndPoint)
 		{
 			const TVec3<T> Segment = EndPoint - StartPoint;
 			const TVec3<T> VectToPoint = Point - StartPoint;
 
-			// See if closest point is before StartPoint
-			const T Dot1 = TVec3<T>::DotProduct(VectToPoint, Segment);
-			if (Dot1 <= 0)
+			// See if closet point is before StartPoint
+			const T Dot1 = DotWithoutFMA<T>(VectToPoint, Segment);
+			if (Dot1 <= T(0))
 			{
 				return T(0);
 			}
 
 			// See if closest point is beyond EndPoint
-			const T Dot2 = TVec3<T>::DotProduct(Segment, Segment);
+			const T Dot2 = DotWithoutFMA<T>(Segment, Segment);
 			if (Dot2 <= Dot1)
 			{
 				return T(1);
