Fixed iostore error on request command so it doesn't use unknown errors.

diff --git a/Engine/Source/Runtime/Experimental/IoStore/OnDemand/Private/OnDemandContentInstaller.cpp b/Engine/Source/Runtime/Experimental/IoStore/OnDemand/Private/OnDemandContentInstaller.cpp
--- a/Engine/Source/Runtime/Experimental/IoStore/OnDemand/Private/OnDemandContentInstaller.cpp
+++ b/Engine/Source/Runtime/Experimental/IoStore/OnDemand/Private/OnDemandContentInstaller.cpp
@@ -483,7 +483,7 @@ void FOnDemandContentInstaller::ProcessInstallRequest(FRequest& Request)
 		if (FCString::Strstr(*ContentHandle->DebugName, *UE::IoStore::CVars::IoStoreErrorOnRequest) != nullptr || 
 			FCString::Strstr(*InstallRequest.Args.DebugName, *UE::IoStore::CVars::IoStoreErrorOnRequest) != nullptr)
 		{
-			Request.ErrorCode	= (EIoErrorCode)((FMath::Rand() / (RAND_MAX / ((int)EIoErrorCode::Last-1))) + 1);
+			Request.ErrorCode	= (EIoErrorCode)((FMath::Rand() / (RAND_MAX / ((int)EIoErrorCode::Last-2))) + 2); // ignore unknown and ok states
 			Request.ErrorReason = FString::Printf(TEXT("Debug error requested on debug name %s"), *UE::IoStore::CVars::IoStoreErrorOnRequest);
 
 			return CompleteInstallRequest(Request);
