add a general dataflow selection -> index array conversion node (so that e.g. one can iterate the selected indices)

diff --git a/Engine/Source/Runtime/Experimental/Dataflow/Core/Private/Dataflow/DataflowConvertNodes.cpp b/Engine/Source/Runtime/Experimental/Dataflow/Core/Private/Dataflow/DataflowConvertNodes.cpp
--- a/Engine/Source/Runtime/Experimental/Dataflow/Core/Private/Dataflow/DataflowConvertNodes.cpp
+++ b/Engine/Source/Runtime/Experimental/Dataflow/Core/Private/Dataflow/DataflowConvertNodes.cpp
@@ -22,6 +22,7 @@ namespace UE::Dataflow
 		// Needs to be fixed
 //		DATAFLOW_NODE_REGISTER_CREATION_FACTORY(FConvertUObjectConvertibleTypesDataflowNode);
 		DATAFLOW_NODE_REGISTER_CREATION_FACTORY(FConvertSelectionTypesDataflowNode);
+		DATAFLOW_NODE_REGISTER_CREATION_FACTORY(FConvertSelectionTypesToIndicesArrayDataflowNode);
 		DATAFLOW_NODE_REGISTER_CREATION_FACTORY(FConvertVectorArrayTypesDataflowNode);
 		DATAFLOW_NODE_REGISTER_CREATION_FACTORY(FConvertNumericArrayTypesDataflowNode);
 		DATAFLOW_NODE_REGISTER_CREATION_FACTORY(FConvertStringArrayTypesDataflowNode);
@@ -312,6 +313,26 @@ void FConvertSelectionTypesDataflowNode::Evaluate(UE::Dataflow::FContext& Contex
 
 //-----------------------------------------------------------------------------------------------
 
+FConvertSelectionTypesToIndicesArrayDataflowNode::FConvertSelectionTypesToIndicesArrayDataflowNode(const UE::Dataflow::FNodeParameters& InParam, FGuid InGuid)
+	: FDataflowNode(InParam, InGuid)
+{
+	RegisterInputConnection(&In);
+	RegisterOutputConnection(&Out);
+}
+
+void FConvertSelectionTypesToIndicesArrayDataflowNode::Evaluate(UE::Dataflow::FContext& Context, const FDataflowOutput* Output) const
+{
+	if (Output->IsA(&Out))
+	{
+		const FDataflowSelection InValue = GetValue(Context, &In);
+		TArray<int32> NewSelection;
+		InValue.AsArray(NewSelection);
+		SetValue(Context, MoveTemp(NewSelection), &Out);
+	}
+}
+
+//-----------------------------------------------------------------------------------------------
+
 FConvertVectorArrayTypesDataflowNode::FConvertVectorArrayTypesDataflowNode(const UE::Dataflow::FNodeParameters& InParam, FGuid InGuid)
 	: FDataflowNode(InParam, InGuid)
 {
diff --git a/Engine/Source/Runtime/Experimental/Dataflow/Core/Public/Dataflow/DataflowConvertNodes.h b/Engine/Source/Runtime/Experimental/Dataflow/Core/Public/Dataflow/DataflowConvertNodes.h
--- a/Engine/Source/Runtime/Experimental/Dataflow/Core/Public/Dataflow/DataflowConvertNodes.h
+++ b/Engine/Source/Runtime/Experimental/Dataflow/Core/Public/Dataflow/DataflowConvertNodes.h
@@ -243,6 +243,33 @@ public:
 	FConvertSelectionTypesDataflowNode(const UE::Dataflow::FNodeParameters& InParam, FGuid InGuid = FGuid::NewGuid());
 };
 
+/**
+* Convert Selection to Index Array
+*/
+USTRUCT()
+struct FConvertSelectionTypesToIndicesArrayDataflowNode : public FDataflowNode
+{
+	GENERATED_USTRUCT_BODY()
+	DATAFLOW_NODE_DEFINE_INTERNAL(FConvertSelectionTypesToIndicesArrayDataflowNode, "ConvertSelectionToIndicesArray", DATAFLOW_CONVERT_NODES_CATEGORY, "")
+
+private:
+
+	/** Input value */
+	UPROPERTY(meta = (DataflowInput, DataflowIntrinsic))
+	FDataflowSelectionTypes In;
+
+	/** Output value */
+	UPROPERTY(meta = (DataflowOutput))
+	TArray<int32> Out;
+
+	//~ Begin FDataflowNode interface
+	virtual void Evaluate(UE::Dataflow::FContext& Context, const FDataflowOutput* Out) const override;
+	//~ End FDataflowNode interface
+
+public:
+	FConvertSelectionTypesToIndicesArrayDataflowNode(const UE::Dataflow::FNodeParameters& InParam, FGuid InGuid = FGuid::NewGuid());
+};
+
 /**
 * Convert Vector array types
 */
