[ChaosVD] Fix for an issue where in complex recordings with a large amount of geometry, data at the start of the recording might be dropped.

This is an problem because right at the begining of the recording we will have the FArchive headers we use for backwards compatibility and most of the level geometry.

diff --git a/Engine/Source/Runtime/Experimental/ChaosVisualDebugger/Private/ChaosVDRuntimeModule.cpp b/Engine/Source/Runtime/Experimental/ChaosVisualDebugger/Private/ChaosVDRuntimeModule.cpp
--- a/Engine/Source/Runtime/Experimental/ChaosVisualDebugger/Private/ChaosVDRuntimeModule.cpp
+++ b/Engine/Source/Runtime/Experimental/ChaosVisualDebugger/Private/ChaosVDRuntimeModule.cpp
@@ -178,6 +178,34 @@ bool FChaosVDRuntimeModule::RecordingTimerTick(float DeltaTime)
 	return true;
 }
 
+void FChaosVDRuntimeModule::HandleTraceConnectionEstablished()
+{
+	{
+		UE::TReadScopeLock ReadLock(DelegatesRWLock);
+		RecordingStartedDelegate.Broadcast();
+	}
+
+	FTraceAuxiliary::OnConnection.RemoveAll(this);
+	
+	constexpr int32 MinAllowedTimeInSecondsBetweenCaptures = 1;
+	int32 ConfiguredTimeBetweenCaptures = CVarChaosVDGTimeBetweenFullCaptures->GetInt();
+
+	ensureAlwaysMsgf(ConfiguredTimeBetweenCaptures > MinAllowedTimeInSecondsBetweenCaptures,
+	                 TEXT("The minimum allowed time interval between full captures is [%d] seconds, but [%d] seconds were configured. Clamping to [%d] seconds"),
+	                 MinAllowedTimeInSecondsBetweenCaptures, ConfiguredTimeBetweenCaptures, MinAllowedTimeInSecondsBetweenCaptures);
+
+	FullCaptureRequesterHandle = FTSTicker::GetCoreTicker().AddTicker(FTickerDelegate::CreateRaw(this, &FChaosVDRuntimeModule::RequestFullCapture),
+	                                                                  FMath::Clamp(ConfiguredTimeBetweenCaptures, MinAllowedTimeInSecondsBetweenCaptures, TNumericLimits<int32>::Max()));
+
+	RecordingTimerHandle = FTSTicker::GetCoreTicker().AddTicker(FTickerDelegate::CreateRaw(this, &FChaosVDRuntimeModule::RecordingTimerTick));
+
+	{
+		UE::TReadScopeLock ReadLock(DelegatesRWLock);
+		PostRecordingStartedDelegate.Broadcast();
+	}
+}
+
+
 void FChaosVDRuntimeModule::StartRecording(TConstArrayView<FString> Args)
 {
 	if (bIsRecording)
@@ -240,7 +268,7 @@ void FChaosVDRuntimeModule::StartRecording(TConstArrayView<FString> Args)
 		FTraceAuxiliary::EConnectionType::Network,
 		*Target,
 		nullptr, &TracingOptions);
-		
+
 		CurrentRecordingMode = EChaosVDRecordingMode::Live;
 	}
 	else
@@ -253,27 +281,7 @@ void FChaosVDRuntimeModule::StartRecording(TConstArrayView<FString> Args)
 
 	if (ensure(bIsRecording))
 	{
-		{
-			UE::TReadScopeLock ReadLock(DelegatesRWLock);
-			RecordingStartedDelegate.Broadcast();
-		}
-		
-		constexpr int32 MinAllowedTimeInSecondsBetweenCaptures = 1;
-		int32 ConfiguredTimeBetweenCaptures = CVarChaosVDGTimeBetweenFullCaptures->GetInt();
-
-		ensureAlwaysMsgf(ConfiguredTimeBetweenCaptures > MinAllowedTimeInSecondsBetweenCaptures,
-			TEXT("The minimum allowed time interval between full captures is [%d] seconds, but [%d] seconds were configured. Clamping to [%d] seconds"),
-			MinAllowedTimeInSecondsBetweenCaptures, ConfiguredTimeBetweenCaptures, MinAllowedTimeInSecondsBetweenCaptures);
-
-		FullCaptureRequesterHandle = FTSTicker::GetCoreTicker().AddTicker(FTickerDelegate::CreateRaw(this, &FChaosVDRuntimeModule::RequestFullCapture),
-			FMath::Clamp(ConfiguredTimeBetweenCaptures, MinAllowedTimeInSecondsBetweenCaptures, TNumericLimits<int32>::Max()));
-
-		RecordingTimerHandle = FTSTicker::GetCoreTicker().AddTicker(FTickerDelegate::CreateRaw(this, &FChaosVDRuntimeModule::RecordingTimerTick));
-
-		{
-			UE::TReadScopeLock ReadLock(DelegatesRWLock);
-			PostRecordingStartedDelegate.Broadcast();
-		}
+		FTraceAuxiliary::OnConnection.AddRaw(this, &FChaosVDRuntimeModule::HandleTraceConnectionEstablished);
 	}
 	else
 	{
@@ -294,12 +302,14 @@ void FChaosVDRuntimeModule::StartRecording(TConstArrayView<FString> Args)
 
 void FChaosVDRuntimeModule::StopRecording()
 {
+	FTraceAuxiliary::OnConnection.RemoveAll(this);
+
 	if (!bIsRecording)
 	{
 		UE_LOG(LogChaosVDRuntime, Warning, TEXT("[%s] Attempted to stop recorded when there is no CVD recording active."), ANSI_TO_TCHAR(__FUNCTION__));
 		return;
 	}
-	
+
 	FTraceAuxiliary::OnTraceStopped.RemoveAll(this);
 
 #if UE_TRACE_ENABLED
diff --git a/Engine/Source/Runtime/Experimental/ChaosVisualDebugger/Public/ChaosVDRuntimeModule.h b/Engine/Source/Runtime/Experimental/ChaosVisualDebugger/Public/ChaosVDRuntimeModule.h
--- a/Engine/Source/Runtime/Experimental/ChaosVisualDebugger/Public/ChaosVDRuntimeModule.h
+++ b/Engine/Source/Runtime/Experimental/ChaosVisualDebugger/Public/ChaosVDRuntimeModule.h
@@ -139,6 +139,8 @@ private:
 
 	/** Queues a full Capture of the simulation on the next frame */
 	UE_API bool RecordingTimerTick(float DeltaTime);
+	
+	void HandleTraceConnectionEstablished();
 
 	/** Used to handle stop requests to the active trace session that were not done by us
 	 * That is a possible scenario because Trace is shared by other In-Editor tools
