Tentative fix for crashin DataflowInstance receiving Uobject property change
- move the unregistration of the delegate out of the owner check to make sure it never get missed and survive post destruction of the DataflowInstance structure

diff --git a/Engine/Source/Runtime/Experimental/Dataflow/Engine/Private/Dataflow/DataflowInstance.cpp b/Engine/Source/Runtime/Experimental/Dataflow/Engine/Private/Dataflow/DataflowInstance.cpp
--- a/Engine/Source/Runtime/Experimental/Dataflow/Engine/Private/Dataflow/DataflowInstance.cpp
+++ b/Engine/Source/Runtime/Experimental/Dataflow/Engine/Private/Dataflow/DataflowInstance.cpp
@@ -497,8 +497,9 @@ FDataflowInstance::~FDataflowInstance()
 	if (Owner)
 	{
 		FDataflowAssetDelegates::OnVariablesChanged.RemoveAll(Owner.Get());
-		FCoreUObjectDelegates::OnObjectPropertyChanged.Remove(OnOwnerPropertyChangedHandle);
 	}
+	// make sure we remove this outside of the Owner check in case of the Owner changed value since construction
+	FCoreUObjectDelegates::OnObjectPropertyChanged.Remove(OnOwnerPropertyChangedHandle);
 #endif
 }
 
