Remove 'IAS' from the hostgroup debug display lines as they are not specifically for IAS anymore.

#rnx

diff --git a/Engine/Source/Runtime/Experimental/IoStore/OnDemand/Private/Statistics.cpp b/Engine/Source/Runtime/Experimental/IoStore/OnDemand/Private/Statistics.cpp
--- a/Engine/Source/Runtime/Experimental/IoStore/OnDemand/Private/Statistics.cpp
+++ b/Engine/Source/Runtime/Experimental/IoStore/OnDemand/Private/Statistics.cpp
@@ -648,7 +648,7 @@ void FOnDemandIoBackendStats::PrintOnScreenStatistics(FCoreDelegates::FSeverityM
 		{
 			FCoreDelegates::EOnScreenMessageSeverity Verbosity = FCoreDelegates::EOnScreenMessageSeverity::Info;
 			TStringBuilder<256> Text;
-			Text << TEXT("IAS HostGroup [") << HostGroup.GetName() << TEXT("] ");
+			Text << TEXT("HostGroup [") << HostGroup.GetName() << TEXT("] ");
 
 			if (HostGroup.IsConnected())
 			{
