[NetPhysics] Scheduled PT Commands with resimulation handling.
- Enqueue a lambda (command) to run on a specified physics frame
- Cache the command in rewind history via PushData and apply the command again during resimulation on the correct frame
- Commands not scheduled for specific frames will not be ran again during resimulations

diff --git a/Engine/Source/Runtime/Engine/Classes/GameFramework/PlayerController.h b/Engine/Source/Runtime/Engine/Classes/GameFramework/PlayerController.h
--- a/Engine/Source/Runtime/Engine/Classes/GameFramework/PlayerController.h
+++ b/Engine/Source/Runtime/Engine/Classes/GameFramework/PlayerController.h
@@ -2309,10 +2309,13 @@ public:
 	*/
 	ENGINE_API void ExecuteAsyncPhysicsCommand(const FAsyncPhysicsTimestamp& AsyncPhysicsTimestamp, UObject* OwningObject, const TFunction<void()>& Command, const bool bEnableResim = true);
 
-	/** Generates a timestamp for the upcoming physics step (plus any pending time). Useful for synchronizing client and server events on a specific physics step 
-	* On the server LocalFrame and ServerFrame are the same, on the client LocalFrame is the current physics frame and ServerFrame is the server-side frame that the LocalFrame corresponds.
+	/** Generates a timestamp for the upcoming physics step. Useful for synchronizing client and server events on a specific physics step 
+	* On the server LocalFrame and ServerFrame are the same, on the client LocalFrame is the upcoming physics frame and ServerFrame is the server-side frame that the LocalFrame corresponds to.
 	*/
-	ENGINE_API FAsyncPhysicsTimestamp GetPhysicsTimestamp(float DeltaSeconds = 0.0f);
+	ENGINE_API FAsyncPhysicsTimestamp GetPhysicsTimestamp();
+
+	UE_DEPRECATED(5.7, "Deprecated, use the GetPhysicsTimestamp() without the DeltaSeconds parameter, it will return the upcoming frame as intended.")
+	ENGINE_API FAsyncPhysicsTimestamp GetPhysicsTimestamp(float DeltaSeconds) { return GetPhysicsTimestamp(); }
 
 	/** Get the physics frame number offset between Local and Server, recommended is to use GetPhysicsTimestamp() which takes this offset into account on the client for the ServerFrame value */
 	int32 GetNetworkPhysicsTickOffset() const { return NetworkPhysicsTickOffset; }
diff --git a/Engine/Source/Runtime/Engine/Private/PlayerController.cpp b/Engine/Source/Runtime/Engine/Private/PlayerController.cpp
--- a/Engine/Source/Runtime/Engine/Private/PlayerController.cpp
+++ b/Engine/Source/Runtime/Engine/Private/PlayerController.cpp
@@ -6489,32 +6489,20 @@ void APlayerController::ExecuteAsyncPhysicsCommand(const FAsyncPhysicsTimestamp&
 	}
 }
 
-FAsyncPhysicsTimestamp APlayerController::GetPhysicsTimestamp(float DeltaSeconds)
+FAsyncPhysicsTimestamp APlayerController::GetPhysicsTimestamp()
 {
 	using namespace Chaos;
 
 	FAsyncPhysicsTimestamp Timestamp;
 
-	if(UWorld* World = GetWorld())
+	if (UWorld* World = GetWorld())
 	{
 		if (FPhysScene* PhysScene = World->GetPhysicsScene())
 		{
-			if (FPBDRigidsSolver* Solver = static_cast<FPBDRigidsSolver*>(PhysScene->GetSolver()))
+			if (FPhysicsSolverBase* Solver = static_cast<FPhysicsSolverBase*>(PhysScene->GetSolver()))
 			{
-				const FReal DeltaTime = Solver->GetAsyncDeltaTime();
-				const int32 PendingSteps = (DeltaTime > 0.0) ? DeltaSeconds / DeltaTime : 0;
-
-				int32 LocalPhysicsStep = Solver->GetCurrentFrame();
-				
-				LocalPhysicsStep += PendingSteps;	//Add any pending steps user wants to wait on
-				Timestamp.ServerFrame = LocalPhysicsStep;
-				Timestamp.LocalFrame = LocalPhysicsStep;
-
-				if (IsLocalController())
-				{
-					//If local controller we update server frame based on our estimate
-					Timestamp.ServerFrame = LocalPhysicsStep + NetworkPhysicsTickOffset;
-				}
+				Timestamp.LocalFrame = Solver->GetMarshallingManager().GetInternalStep_External();
+				Timestamp.ServerFrame = Timestamp.LocalFrame + NetworkPhysicsTickOffset; // On the server NetworkPhysicsTickOffset is always 0 so it's safe to always add on both client and server to convert LocalFrame to ServerFrame
 			}
 		}
 	}
diff --git a/Engine/Source/Runtime/Experimental/Chaos/Private/Chaos/ChaosMarshallingManager.cpp b/Engine/Source/Runtime/Experimental/Chaos/Private/Chaos/ChaosMarshallingManager.cpp
--- a/Engine/Source/Runtime/Experimental/Chaos/Private/Chaos/ChaosMarshallingManager.cpp
+++ b/Engine/Source/Runtime/Experimental/Chaos/Private/Chaos/ChaosMarshallingManager.cpp
@@ -11,10 +11,13 @@ int32 SimDelay = 0;
 FAutoConsoleVariableRef CVarSimDelay(TEXT("p.simDelay"),SimDelay,TEXT(""));
 
 bool bCachePushDataDirtyProxies = true;
-FAutoConsoleVariableRef CVarCachePushDataDirtyProxies(TEXT("p.Resim.CachePushDataDirtyProxies"), bCachePushDataDirtyProxies, TEXT("Default = false. Set true to enable resim caching dirty proxies in the push data from game thread to physics thread. This will make physics proxy changes from GT play out during a resimulation."));
+FAutoConsoleVariableRef CVarCachePushDataDirtyProxies(TEXT("p.Resim.CachePushDataDirtyProxies"), bCachePushDataDirtyProxies, TEXT("Default = false. Set true to enable resim caching dirty proxies in the push data from game thread to physics thread. This will make physics proxy changes from GT play out during resimulation."));
 
 bool bCachePushDataAsyncInputs = true;
-FAutoConsoleVariableRef CVarCachePushDataAsyncInputs(TEXT("p.Resim.CachePushDataAsyncInputs"), bCachePushDataAsyncInputs, TEXT("Default = false. Set true to enable resim caching of async inputs in the push data from game thread to physics thread. This will make async inputs available again during a resimulation."));
+FAutoConsoleVariableRef CVarCachePushDataAsyncInputs(TEXT("p.Resim.CachePushDataAsyncInputs"), bCachePushDataAsyncInputs, TEXT("Default = false. Set true to enable resim caching of async inputs in the push data from game thread to physics thread. This will make async inputs available again during resimulation."));
+
+bool bCachePushDataSimCommands = true;
+FAutoConsoleVariableRef CVarCachePushDataSimCommands(TEXT("p.Resim.CachePushDataSimCommands"), bCachePushDataSimCommands, TEXT("Default = false. Set true to enable resim caching of sim commands in the push data from game thread to physics thread. This will make sim commands available again during resimulation."));
 
 FChaosMarshallingManager::FChaosMarshallingManager()
 : ExternalTime_External(0)
@@ -160,7 +163,7 @@ void FPushPhysicsData::Reset()
 		delete CallbackToRemove;
 	}
 
-	SimCommands.Reset();
+	ClearSimCommands();
 	DirtyProxiesDataBuffer.Reset();
 	SimCallbackObjectsToAdd.Reset();
 	SimCallbackObjectsToRemove.Reset();
@@ -168,20 +171,25 @@ void FPushPhysicsData::Reset()
 
 void FPushPhysicsData::ResetForHistory()
 {
-	// Clear dirty proxies data if we should not them cache in history
+	// Clear dirty proxies data if we should not cache them in history
 	if (bCachePushDataDirtyProxies == false)
 	{
 		ResetDirtyProxiesBuffer();
 		DirtyProxiesDataBuffer.Reset();
 	}
 
-	// Clear async inputs if we should not them cache in history
+	// Clear async inputs if we should not cache them in history
 	if (bCachePushDataAsyncInputs == false)
 	{
 		ClearAsyncInputs();
 	}
 
-	SimCommands.Reset();
+	// Clear simcommands if we should not cache them in history
+	if (bCachePushDataSimCommands == false)
+	{
+		ClearSimCommands();
+	}
+
 	SimCallbackObjectsToAdd.Reset();
 }
 
@@ -210,6 +218,18 @@ void FPushPhysicsData::ClearAsyncInputs()
 	SimCallbackInputs.Reset();
 }
 
+void FPushPhysicsData::ClearSimCommands()
+{
+	for (FSimCallbackCommandObject* SimCallbackObject : SimCommands)
+	{
+		if (SimCallbackObject)
+		{
+			delete SimCallbackObject;
+		}
+	}
+	SimCommands.Reset();
+}
+
 void FPushPhysicsData::ClearAsyncInputForCallback(ISimCallbackObject* Callback)
 {
 	for (TArray<FSimCallbackInputAndObject>::TIterator Itr = SimCallbackInputs.CreateIterator(); Itr; ++Itr)
diff --git a/Engine/Source/Runtime/Experimental/Chaos/Private/PBDRigidsSolver.cpp b/Engine/Source/Runtime/Experimental/Chaos/Private/PBDRigidsSolver.cpp
--- a/Engine/Source/Runtime/Experimental/Chaos/Private/PBDRigidsSolver.cpp
+++ b/Engine/Source/Runtime/Experimental/Chaos/Private/PBDRigidsSolver.cpp
@@ -1972,13 +1972,43 @@ namespace Chaos
 		//so for example if the user deletes a floor surface, we want all sub-steps to see that in the same way
 		//also note, the commands run after data is marshalled over. This is important because data marshalling ensures any GT property changes are seen by command
 		//for example a particle may not be created until marshalling occurs, and then a command could explicitly modify something like a collision setting
-		for (FSimCallbackCommandObject* SimCallbackObject : PushData.SimCommands)
+
+		// If we are not resimulating, insert scheduled commands for this step into the current PushData to be executed
+		if (!IsResimming())
+		{
+			TArray<FSimCallbackCommandObject*>& ScheduledSimCommands = MarshallingManager.GetSimCommandsScheduled_Internal();
+			for (int32 Idx = ScheduledSimCommands.Num() - 1; Idx >= 0; Idx--)
+			{
+				FSimCallbackCommandObject* SimCallbackObject = ScheduledSimCommands[Idx];
+				if (SimCallbackObject->ShouldExecute(CurrentFrame, /*bIsResim*/false))
+				{
+					// If this scheduled command should get processed this step, inject it into the start of PushData:SimCommands to preserve call order and to cache the SimCommand in rewind history via PushData
+					PushData.SimCommands.Insert(SimCallbackObject, 0);
+					ScheduledSimCommands.RemoveAt(Idx, EAllowShrinking::No);
+				}
+			}
+		}
+
+		for (FSimCallbackCommandObject*& SimCallbackObject : PushData.SimCommands)
 		{
-			SimCallbackObject->SetSimAndDeltaTime_Internal(GetSolverTime(), MLastDt);
-			SimCallbackObject->PreSimulate_Internal();
-			delete SimCallbackObject;
+			if (!SimCallbackObject)
+			{
+				// A scheduled command delayed for a later frame will leave a nullptr in the array
+				continue;
+			}
+
+			if (SimCallbackObject->ShouldExecute(CurrentFrame, IsResimming()))
+			{
+				SimCallbackObject->SetSimAndDeltaTime_Internal(GetSolverTime(), MLastDt);
+				SimCallbackObject->PreSimulate_Internal();
+			}
+			else if (SimCallbackObject->ShouldDelay(CurrentFrame, IsResimming()))
+			{
+				// Cache command to run on a future frame and null the pointer in this PushData
+				MarshallingManager.EnqueueCommandScheduled_Internal(SimCallbackObject);
+				SimCallbackObject = nullptr; // Stop tracking the command in this PushData
+			}
 		}
-		PushData.SimCommands.Reset();
 		
 		for (ISimCallbackObject* SimCallbackObject : PushData.SimCallbackObjectsToAdd)
 		{
diff --git a/Engine/Source/Runtime/Experimental/Chaos/Public/Chaos/ChaosMarshallingManager.h b/Engine/Source/Runtime/Experimental/Chaos/Public/Chaos/ChaosMarshallingManager.h
--- a/Engine/Source/Runtime/Experimental/Chaos/Public/Chaos/ChaosMarshallingManager.h
+++ b/Engine/Source/Runtime/Experimental/Chaos/Public/Chaos/ChaosMarshallingManager.h
@@ -303,6 +303,7 @@ struct FPushPhysicsData
 	void ResetForHistory();
 	void ResetDirtyProxiesBuffer();
 	void ClearAsyncInputs();
+	void ClearSimCommands();
 	void ClearAsyncInputForCallback(ISimCallbackObject* Callback);
 	void CopySubstepData(const FPushPhysicsData& FirstStepData);
 };
@@ -414,6 +415,16 @@ public:
 		ProducerData->SimCommands.Add(SimCommand);
 	}
 
+	void EnqueueCommandScheduled_Internal(FSimCallbackCommandObject* SimCommand)
+	{
+		SimCommandsScheduled_Internal.Add(SimCommand);
+	}
+
+	TArray<FSimCallbackCommandObject*>& GetSimCommandsScheduled_Internal()
+	{
+		return SimCommandsScheduled_Internal;
+	}
+
 	void UnregisterSimCallbackObject_External(ISimCallbackObject* SimCallbackObject)
 	{
 		UE_CHAOS_ASYNC_INITBODY_WRITESCOPELOCK(MarshallingManagerLock);
@@ -490,6 +501,7 @@ private:
 	TSpscQueue<FPushPhysicsData*> PushDataPool;	//pool to grab more push data from to avoid expensive reallocs
 	TArray<TUniquePtr<FPushPhysicsData>> BackingBuffer;	//all push data is cleaned up by this
 	TArray<FPushPhysicsData*> HistoryQueue_Internal;	//all push data still needed for potential rewinds. Latest data comes first
+	TArray<FSimCallbackCommandObject*> SimCommandsScheduled_Internal; // SimCommands scheduled for a future frame, will be inserted into PushData on the correct frame
 
 	//pull
 	FPullPhysicsData* CurPullData;	//the current pull data sim is writing to
diff --git a/Engine/Source/Runtime/Experimental/Chaos/Public/Chaos/Framework/PhysicsSolverBase.h b/Engine/Source/Runtime/Experimental/Chaos/Public/Chaos/Framework/PhysicsSolverBase.h
--- a/Engine/Source/Runtime/Experimental/Chaos/Public/Chaos/Framework/PhysicsSolverBase.h
+++ b/Engine/Source/Runtime/Experimental/Chaos/Public/Chaos/Framework/PhysicsSolverBase.h
@@ -435,6 +435,15 @@ namespace Chaos
 			RegisterSimOneShotCallback(MoveTemp(Func));
 		}
 
+		/** Enqueue a lambda function (command) to run on the physics thread on a specific physics step
+		* The command will be reapplied during physics resimulation on the same step */
+		template <typename Lambda>
+		void EnqueueCommandScheduled_External(const int32 ExecuteOnStep, Lambda&& Func)
+		{
+			FSimCallbackCommandObject* CommandObject = new FSimCallbackCommandObject(MoveTemp(Func), ExecuteOnStep);
+			MarshallingManager.RegisterSimCommand_External(CommandObject);
+		}
+
 		CHAOS_API void SetRewindCallback(TUniquePtr<IRewindCallback>&& RewindCallback);
 
 		FRewindData* GetRewindData()
diff --git a/Engine/Source/Runtime/Experimental/Chaos/Public/Chaos/SimCallbackObject.h b/Engine/Source/Runtime/Experimental/Chaos/Public/Chaos/SimCallbackObject.h
--- a/Engine/Source/Runtime/Experimental/Chaos/Public/Chaos/SimCallbackObject.h
+++ b/Engine/Source/Runtime/Experimental/Chaos/Public/Chaos/SimCallbackObject.h
@@ -426,18 +426,20 @@ private:
 class FSimCallbackCommandObject : public ISimCallbackObject
 {
 public:
-	FSimCallbackCommandObject(TUniqueFunction<void()>&& InFunc)
+	FSimCallbackCommandObject(TUniqueFunction<void()>&& InFunc, int32 Step = INDEX_NONE)
 		: Func(MoveTemp(InFunc))
 		, Func2(nullptr)
 		, bFuncHasTimeParameters(false)
+		, ExecuteOnStep(Step)
 	{
 		ensureMsgf(Func, TEXT("Created a sim callback object with an unbound function. This command will not be executed."));
 	}
 
-	FSimCallbackCommandObject(TUniqueFunction<void(FReal DeltaTime, FReal SimTime)>&& InFunc)
+	FSimCallbackCommandObject(TUniqueFunction<void(FReal DeltaTime, FReal SimTime)>&& InFunc, int32 Step = INDEX_NONE)
 		: Func(nullptr)
 		, Func2(MoveTemp(InFunc))
 		, bFuncHasTimeParameters(true)
+		, ExecuteOnStep(Step)
 	{
 		ensureMsgf(Func2, TEXT("Created a sim callback object with an unbound function. This command will not be executed."));
 	}
@@ -454,6 +456,21 @@ public:
 		return StaticName;
 	}
 
+	const bool ShouldExecute(const int32 Step, const bool bIsResim) const
+	{
+		return bIsResim ? ExecuteOnStep == Step : ExecuteOnStep <= Step;
+	}
+
+	const bool ShouldDelay(const int32 Step, const bool bIsResim) const
+	{
+		return !bIsResim && ExecuteOnStep > Step;
+	}
+
+	const int32 GetExecuteOnStep() const
+	{
+		return ExecuteOnStep;
+	}
+
 private:
 
 	virtual FSimCallbackInput* AllocateInputData_External()
@@ -493,6 +510,7 @@ private:
 	TUniqueFunction<void(FReal deltaTime, FReal SimTime)> Func2;
 
 	bool bFuncHasTimeParameters;
+	int32 ExecuteOnStep;
 
 	friend struct FSimCallbackInput;
 };
