Bound the check of the found journal range against data so that it doesn't take forever in failure cases.

#rnx

diff --git a/Engine/Source/Runtime/Experimental/IoStore/OnDemand/Private/JournaledCache.cpp b/Engine/Source/Runtime/Experimental/IoStore/OnDemand/Private/JournaledCache.cpp
--- a/Engine/Source/Runtime/Experimental/IoStore/OnDemand/Private/JournaledCache.cpp
+++ b/Engine/Source/Runtime/Experimental/IoStore/OnDemand/Private/JournaledCache.cpp
@@ -1192,22 +1192,38 @@ static int32 LoadCache(FDiskCache& DiskCache)
 		return false;
 	};
 
-	for (; BasisIndex < Paragraphs.Num(); ++BasisIndex)
+	// Find the first good paragraph.
 	{
-		const FParagraph& Stock = Paragraphs[BasisIndex];
+		enum { SearchDepth = 14 };
+		int32 Last = BasisIndex + (1 << SearchDepth);
+		Last = FMath::Min(Last, Paragraphs.Num());
 
-		uint64 DataBase = Stock.Phrase->DataCursor;
-		if (DataBase + Stock.DataSize > DiskCache.MaxDataSize)
+		for (int32 End = Last; BasisIndex < End;)
 		{
-			DataBase = 0;
+			int32 Mid = BasisIndex + ((End - BasisIndex) >> 1);
+
+			const FParagraph& Stock = Paragraphs[Mid];
+
+			uint64 DataBase = Stock.Phrase->DataCursor;
+			if (DataBase + Stock.DataSize > DiskCache.MaxDataSize)
+			{
+				DataBase = 0;
+			}
+
+			uint32 Hash;
+			const uint32 Seed = uint32(Stock.Phrase->Marker);
+			if (ReadHash(DataBase, Hash, Stock.DataSize, Seed) && (Hash == Stock.Phrase->Hash))
+			{
+				End = Mid;
+				continue;
+			}
+
+			BasisIndex = Mid + 1;
 		}
 
-		uint32 Hash;
-		const uint32 Seed = uint32(Stock.Phrase->Marker);
-		// There could be a phrase with only one very short data entry. Make sure we don't hash too much.
-		if (ReadHash(DataBase, Hash, Stock.DataSize, Seed) && (Hash == Stock.Phrase->Hash))
+		if (BasisIndex == Last)
 		{
-			break;
+			BasisIndex = Paragraphs.Num();
 		}
 	}
 
