Iris
- Improved support for utilizing per property dirty tracking
- Added cvar net.Iris.UsePerPropertyDirtyTracking to toggle between using per property tracking
- Added cvar net.Iris.Poll.EnableVerboseProfiling that can be enabled to get verbose tracking over polled properties. Very useful to track down properties that constantly are marked dirty
- Added test to verify behavior of pushbased initialonly/initialorowner properties

diff --git a/Engine/Plugins/Runtime/GameplayAbilities/Source/GameplayAbilities/Private/Serialization/MinimalGameplayCueReplicationProxyReplicationFragment.cpp b/Engine/Plugins/Runtime/GameplayAbilities/Source/GameplayAbilities/Private/Serialization/MinimalGameplayCueReplicationProxyReplicationFragment.cpp
--- a/Engine/Plugins/Runtime/GameplayAbilities/Source/GameplayAbilities/Private/Serialization/MinimalGameplayCueReplicationProxyReplicationFragment.cpp
+++ b/Engine/Plugins/Runtime/GameplayAbilities/Source/GameplayAbilities/Private/Serialization/MinimalGameplayCueReplicationProxyReplicationFragment.cpp
@@ -54,6 +54,10 @@ FMinimalGameplayCueReplicationProxyReplicationFragment::FMinimalGameplayCueRepli
 	if (EnumHasAnyFlags(InDescriptor->Traits, EReplicationStateTraits::HasPushBasedDirtiness))
 	{
 		Traits |= EReplicationFragmentTraits::HasPushBasedDirtiness;
+		if (EnumHasAnyFlags(InDescriptor->Traits, EReplicationStateTraits::HasFullPushBasedDirtiness))
+		{
+			Traits |= EReplicationFragmentTraits::HasFullPushBasedDirtiness;
+		}
 	}
 #endif
 	
@@ -97,7 +101,11 @@ void FMinimalGameplayCueReplicationProxyReplicationFragment::ApplyReplicatedStat
 
 bool FMinimalGameplayCueReplicationProxyReplicationFragment::PollReplicatedState(EReplicationFragmentPollFlags PollOption)
 {
-	if (EnumHasAnyFlags(PollOption, EReplicationFragmentPollFlags::PollAllState))
+	// We can early out if we are pushbased and not dirty for polling
+	const bool bPoll = EnumHasAnyFlags(PollOption, EReplicationFragmentPollFlags::PollAllState) ||
+	(EnumHasAnyFlags(PollOption, EReplicationFragmentPollFlags::PollDirtyState) && (!EnumHasAnyFlags(EReplicationFragmentTraits::HasPushBasedDirtiness, Traits) || SrcReplicationState->IsDirtyForPolling()));
+
+	if (bPoll)
 	{
 		const uint8* ExternalStateBuffer = reinterpret_cast<uint8*>(Owner) + ReplicationStateDescriptor->MemberProperties[0]->GetOffset_ForGC();
 		const FMinimalGameplayCueReplicationProxy* ExternalSourceState = reinterpret_cast<const FMinimalGameplayCueReplicationProxy*>(ExternalStateBuffer);
@@ -115,6 +123,11 @@ bool FMinimalGameplayCueReplicationProxyReplicationFragment::PollReplicatedState
 	return SrcReplicationState->IsDirty(0);
 }
 
+void FMinimalGameplayCueReplicationProxyReplicationFragment::CollectOwner(FReplicationStateOwnerCollector* Owners) const
+{
+	Owners->AddOwner(Owner);
+}
+
 void FMinimalGameplayCueReplicationProxyReplicationFragment::CallRepNotifies(FReplicationStateApplyContext& Context)
 {
 	CallRepNotify(Context);
diff --git a/Engine/Plugins/Runtime/GameplayAbilities/Source/GameplayAbilities/Private/Serialization/MinimalReplicationTagCountMapReplicationFragment.cpp b/Engine/Plugins/Runtime/GameplayAbilities/Source/GameplayAbilities/Private/Serialization/MinimalReplicationTagCountMapReplicationFragment.cpp
--- a/Engine/Plugins/Runtime/GameplayAbilities/Source/GameplayAbilities/Private/Serialization/MinimalReplicationTagCountMapReplicationFragment.cpp
+++ b/Engine/Plugins/Runtime/GameplayAbilities/Source/GameplayAbilities/Private/Serialization/MinimalReplicationTagCountMapReplicationFragment.cpp
@@ -48,6 +48,10 @@ FMinimalReplicationTagCountMapReplicationFragment::FMinimalReplicationTagCountMa
 	if (EnumHasAnyFlags(InDescriptor->Traits, EReplicationStateTraits::HasPushBasedDirtiness))
 	{
 		Traits |= EReplicationFragmentTraits::HasPushBasedDirtiness;
+		if (EnumHasAnyFlags(InDescriptor->Traits, EReplicationStateTraits::HasFullPushBasedDirtiness))
+		{
+			Traits |= EReplicationFragmentTraits::HasFullPushBasedDirtiness;
+		}
 	}
 #endif
 		
@@ -93,7 +97,11 @@ void FMinimalReplicationTagCountMapReplicationFragment::ApplyReplicatedState(FRe
 
 bool FMinimalReplicationTagCountMapReplicationFragment::PollReplicatedState(EReplicationFragmentPollFlags PollOption)
 {
-	if (EnumHasAnyFlags(PollOption, EReplicationFragmentPollFlags::PollAllState))
+	// We can early out if we are pushbased and not dirty for polling
+	const bool bPoll = EnumHasAnyFlags(PollOption, EReplicationFragmentPollFlags::PollAllState) ||
+	(EnumHasAnyFlags(PollOption, EReplicationFragmentPollFlags::PollDirtyState) && (!EnumHasAnyFlags(EReplicationFragmentTraits::HasPushBasedDirtiness, Traits) || SrcReplicationState->IsDirtyForPolling()));
+
+	if (bPoll)
 	{
 		const uint8* ExternalStateBuffer = reinterpret_cast<uint8*>(Owner) + ReplicationStateDescriptor->MemberProperties[0]->GetOffset_ForGC();
 		const FMinimalReplicationTagCountMap* ExternalSourceState = reinterpret_cast<const FMinimalReplicationTagCountMap*>(ExternalStateBuffer);
@@ -111,6 +119,11 @@ bool FMinimalReplicationTagCountMapReplicationFragment::PollReplicatedState(ERep
 	return SrcReplicationState->IsDirty(0);
 }
 
+void FMinimalReplicationTagCountMapReplicationFragment::CollectOwner(FReplicationStateOwnerCollector* Owners) const
+{
+	Owners->AddOwner(Owner);
+}
+
 void FMinimalReplicationTagCountMapReplicationFragment::CallRepNotifies(FReplicationStateApplyContext& Context)
 {
 	CallRepNotify(Context);
diff --git a/Engine/Plugins/Runtime/GameplayAbilities/Source/GameplayAbilities/Private/Serialization/MinimalReplicationTagCountMapReplicationFragment.h b/Engine/Plugins/Runtime/GameplayAbilities/Source/GameplayAbilities/Private/Serialization/MinimalReplicationTagCountMapReplicationFragment.h
--- a/Engine/Plugins/Runtime/GameplayAbilities/Source/GameplayAbilities/Private/Serialization/MinimalReplicationTagCountMapReplicationFragment.h
+++ b/Engine/Plugins/Runtime/GameplayAbilities/Source/GameplayAbilities/Private/Serialization/MinimalReplicationTagCountMapReplicationFragment.h
@@ -30,6 +30,7 @@ protected:
 	virtual void ApplyReplicatedState(FReplicationStateApplyContext& Context) const override;
 	virtual bool PollReplicatedState(EReplicationFragmentPollFlags PollOption) override;
 	virtual void CallRepNotifies(FReplicationStateApplyContext& Context) override;
+	virtual void CollectOwner(FReplicationStateOwnerCollector* Owners) const override;
 
 private:
 	void MimicMinimalReplicationTagCountMapReceiveLogic(FReplicationStateApplyContext& Context) const;
diff --git a/Engine/Plugins/Runtime/GameplayAbilities/Source/GameplayAbilities/Public/Serialization/MinimalGameplayCueReplicationProxyReplicationFragment.h b/Engine/Plugins/Runtime/GameplayAbilities/Source/GameplayAbilities/Public/Serialization/MinimalGameplayCueReplicationProxyReplicationFragment.h
--- a/Engine/Plugins/Runtime/GameplayAbilities/Source/GameplayAbilities/Public/Serialization/MinimalGameplayCueReplicationProxyReplicationFragment.h
+++ b/Engine/Plugins/Runtime/GameplayAbilities/Source/GameplayAbilities/Public/Serialization/MinimalGameplayCueReplicationProxyReplicationFragment.h
@@ -30,6 +30,7 @@ protected:
 	virtual void ApplyReplicatedState(FReplicationStateApplyContext& Context) const override;
 	virtual bool PollReplicatedState(EReplicationFragmentPollFlags PollOption) override;
 	virtual void CallRepNotifies(FReplicationStateApplyContext& Context) override;
+	virtual void CollectOwner(FReplicationStateOwnerCollector* Owners) const override;
 
 private:
 	void MimicMinimalGameplayCueReplicationProxyReceiveLogic(FReplicationStateApplyContext& Context) const;
diff --git a/Engine/Plugins/Runtime/ReplicationSystemTestPlugin/Source/Private/Tests/ReplicationState/TestReplicationStateDescriptorBuilder.cpp b/Engine/Plugins/Runtime/ReplicationSystemTestPlugin/Source/Private/Tests/ReplicationState/TestReplicationStateDescriptorBuilder.cpp
--- a/Engine/Plugins/Runtime/ReplicationSystemTestPlugin/Source/Private/Tests/ReplicationState/TestReplicationStateDescriptorBuilder.cpp
+++ b/Engine/Plugins/Runtime/ReplicationSystemTestPlugin/Source/Private/Tests/ReplicationState/TestReplicationStateDescriptorBuilder.cpp
@@ -235,15 +235,17 @@ static FReplicationStateMemberTraitsDescriptor TestClassMemberTraitsDescriptors[
 	{ EReplicationStateMemberTraits::None },
 };
 
+// 4 byte header + changemask + pollmask
+constexpr SIZE_T TestClassMemberDescriptors_MemberOffset = 12;
 static FReplicationStateMemberDescriptor TestClassMemberDescriptors[] = 
 {
-	{8, 0},
-	{10, 2},
-	{12, 4},
-	{16, 8},
-	{24, 16},
-	{25, 17},
-	{26, 18}
+	{TestClassMemberDescriptors_MemberOffset, 0},
+	{TestClassMemberDescriptors_MemberOffset + 2, 2},
+	{TestClassMemberDescriptors_MemberOffset + 4, 4},
+	{TestClassMemberDescriptors_MemberOffset + 12, 8},
+	{TestClassMemberDescriptors_MemberOffset + 20, 16},
+	{TestClassMemberDescriptors_MemberOffset + 21, 17},
+	{TestClassMemberDescriptors_MemberOffset + 22, 18}
 };
 
 static FReplicationStateMemberChangeMaskDescriptor TestClassMemberChangeMaskDescriptors[] =
@@ -274,7 +276,7 @@ static const FReplicationStateDescriptor TestClassReplictionStateDescriptor
 	static_cast<const UScriptStruct*>(nullptr), // BaseStruct
 	static_cast<const FNetDebugName*>(nullptr), // DebugName
 	static_cast<const FReplicationStateMemberDebugDescriptor*>(nullptr), // MemberDebugDescriptors
-	32,			//uint32 ExternalSize;
+	40,			//uint32 ExternalSize;
 	24,			//uint32 InternalSize;
 	8,			//uint16 ExternalAlignment;
 	8,			//uint16 InternalAlignment;
@@ -310,7 +312,7 @@ static const FReplicationStateDescriptor TestClassReplictionStateDescriptorNotRe
 	static_cast<const UScriptStruct*>(nullptr), // BaseStruct
 	static_cast<const FNetDebugName*>(nullptr), // DebugName
 	static_cast<const FReplicationStateMemberDebugDescriptor*>(nullptr), // MemberDebugDescriptors
-	32,			//uint32 ExternalSize;
+	40,			//uint32 ExternalSize;
 	24,			//uint32 InternalSize;
 	8,			//uint16 ExternalAlignment;
 	8,			//uint16 InternalAlignment;
@@ -346,7 +348,7 @@ static const FReplicationStateDescriptor TestClassWithNonReplicatedDataReplicati
 	static_cast<const UScriptStruct*>(nullptr), // BaseStruct
 	static_cast<const FNetDebugName*>(nullptr), // DebugName
 	static_cast<const FReplicationStateMemberDebugDescriptor*>(nullptr), // MemberDebugDescriptors
-	32,			//uint32 ExternalSize;
+	40,			//uint32 ExternalSize;
 	24,			//uint32 InternalSize;
 	8,			//uint16 ExternalAlignment;
 	8,			//uint16 InternalAlignment;
@@ -397,18 +399,19 @@ UE_NET_TEST_FIXTURE(FTestReplicationStateDescriptorBuilderFixture, UTestReplicat
 namespace TestReplicationStateDescriptor_FakeDescriptorData
 {
 
-// Fake desc for FTestReplicationStateDescriptor_TestStruct
+// 4 byte header + changemask + pollmask
+constexpr SIZE_T TestClassWithInheritanceMemberDescriptors_MemberOffset = 12;
 static FReplicationStateMemberDescriptor TestClassWithInheritanceMemberDescriptors[] = 
 {
-	{8, 0},
-	{10, 2},
-	{12, 4},
-	{16, 8},
-	{24, 16},
-	{25, 17},
-	{26, 18},
-	{27, 19},
-	{28, 20}
+	{TestClassWithInheritanceMemberDescriptors_MemberOffset, 0},
+	{TestClassWithInheritanceMemberDescriptors_MemberOffset + 2, 2},
+	{TestClassWithInheritanceMemberDescriptors_MemberOffset + 4, 4},
+	{TestClassWithInheritanceMemberDescriptors_MemberOffset + 12, 8},
+	{TestClassWithInheritanceMemberDescriptors_MemberOffset + 20, 16},
+	{TestClassWithInheritanceMemberDescriptors_MemberOffset + 21, 17},
+	{TestClassWithInheritanceMemberDescriptors_MemberOffset + 22, 18},
+	{TestClassWithInheritanceMemberDescriptors_MemberOffset + 23, 19},
+	{TestClassWithInheritanceMemberDescriptors_MemberOffset + 24, 20}
 };
 
 // Fake member serializer, config is complicated to setup
@@ -467,7 +470,7 @@ static FReplicationStateDescriptor TestClassWithInheritanceReplictionStateDescri
 	static_cast<const UScriptStruct*>(nullptr), // BaseStruct
 	static_cast<const FNetDebugName*>(nullptr), // DebugName
 	static_cast<const FReplicationStateMemberDebugDescriptor*>(nullptr), // MemberDebugDescriptors
-	32, // ExternalSize (rounded up since it also contains statemask)
+	40, // ExternalSize (rounded up since it also contains changemask)
 	24,	// InternalSize
 	8,	// ExternalAlignment
 	8,	// InternalAlignment
@@ -499,10 +502,12 @@ namespace TestReplicationStateDescriptor_FakeDescriptorData
 {
 
 // Fake desc for FTestReplicationStateDescriptor_TestStruct
+// 4 byte header + changemask + pollmask
+constexpr SIZE_T TestClassWithInheritanceNoSuperMemberDescriptors_MemberOffset = 12;
 static FReplicationStateMemberDescriptor TestClassWithInheritanceNoSuperMemberDescriptors[] = 
 {
-	{8, 0},
-	{10, 2}
+	{TestClassWithInheritanceNoSuperMemberDescriptors_MemberOffset, 0},
+	{TestClassWithInheritanceNoSuperMemberDescriptors_MemberOffset + 2, 2}
 };
 
 // Fake member serializer, config is complicated to setup
@@ -540,7 +545,7 @@ static FReplicationStateDescriptor TestClassWithInheritanceNoSuperReplictionStat
 	static_cast<const UScriptStruct*>(nullptr), // BaseStruct
 	static_cast<const FNetDebugName*>(nullptr), // DebugName
 	static_cast<const FReplicationStateMemberDebugDescriptor*>(nullptr), // MemberDebugDescriptors
-	12, // ExternalSize (rounded up since it also contains statemask)
+	16, // ExternalSize (rounded up since it also contains changemask)
 	4,	// InternalSize
 	4,	// ExternalAlignment
 	2,	// InternalAlignment
@@ -571,9 +576,11 @@ namespace TestReplicationStateDescriptor_FakeDescriptorData
 {
 
 // Fake desc for FTestReplicationStateDescriptor_TestStruct
+// 4 byte header + changemask + pollmask
+constexpr SIZE_T TestClassWithReplicatedStructMemberDescriptors_MemberOffset = 16;
 static FReplicationStateMemberDescriptor TestClassWithReplicatedStructMemberDescriptors[] = 
 {
-	{8, 0},
+	{TestClassWithReplicatedStructMemberDescriptors_MemberOffset, 0},
 };
 
 // Fake member serializer, config is complicated to setup
@@ -608,7 +615,7 @@ static FReplicationStateDescriptor TestClassWithReplicatedStructReplictionStateD
 	static_cast<const UScriptStruct*>(nullptr), // BaseStruct
 	static_cast<const FNetDebugName*>(nullptr), // DebugName
 	static_cast<const FReplicationStateMemberDebugDescriptor*>(nullptr), // MemberDebugDescriptors
-	32, // ExternalSize (rounded up since it also contains statemask)
+	40, // ExternalSize (Header + ChangeMasks)
 	24,	// InternalSize
 	8,	// ExternalAlignment
 	8,	// InternalAlignment
@@ -1016,10 +1023,10 @@ UE_NET_TEST_FIXTURE(FTestReplicationStateDescriptorBuilderFixture, UTestReplicat
 	InitDescriptorsFromClass(UTestReplicationStateDescriptor_TestClassWithCondition_InitialOrOwner::StaticClass());
 	UE_NET_ASSERT_GE(Descriptors.Num(), 2);
 
-	const FReplicationStateDescriptor* InitOnlyDesc = GetDescriptorWithTraits(EReplicationStateTraits::InitOnly);
-	UE_NET_ASSERT_NE(InitOnlyDesc, nullptr);
-	UE_NET_ASSERT_TRUE(EnumHasAllFlags(InitOnlyDesc->Traits, EReplicationStateTraits::InitOnly));
-	UE_NET_ASSERT_FALSE(EnumHasAnyFlags(InitOnlyDesc->Traits, EReplicationStateTraits::HasLifetimeConditionals));
+	//const FReplicationStateDescriptor* InitOnlyDesc = GetDescriptorWithTraits(EReplicationStateTraits::InitOnly);
+	//UE_NET_ASSERT_NE(InitOnlyDesc, nullptr);
+	//UE_NET_ASSERT_TRUE(EnumHasAllFlags(InitOnlyDesc->Traits, EReplicationStateTraits::InitOnly));
+	//UE_NET_ASSERT_FALSE(EnumHasAnyFlags(InitOnlyDesc->Traits, EReplicationStateTraits::HasLifetimeConditionals));
 
 	const FReplicationStateDescriptor* LifetimeConditionalsDesc = GetDescriptorWithTraits(EReplicationStateTraits::HasLifetimeConditionals);
 	UE_NET_ASSERT_NE(LifetimeConditionalsDesc, nullptr);
@@ -1222,7 +1229,7 @@ UE_NET_TEST_FIXTURE(FTestReplicationStateDescriptorBuilderFixture, TestClassWith
 	}
 }
 
-UE_NET_TEST_FIXTURE(FTestReplicationStateDescriptorBuilderFixture, TestClassWithPushBasedPropertiesAndNeverReplicatedPropertiesHasPushBasedDirtinessTrait)
+UE_NET_TEST_FIXTURE(FTestReplicationStateDescriptorBuilderFixture, TestClassWithPushBasedPropertiesAndNeverReplicatedPropertiesHasFullPushBasedDirtinessTrait)
 {
 	InitDescriptorsFromClass(UTestReplicationStateDescriptor_TestClassWithPushBasedPropertiesAndNeverReplicatedProperties::StaticClass());
 	UE_NET_ASSERT_GE(Descriptors.Num(), 1);
@@ -1230,10 +1237,11 @@ UE_NET_TEST_FIXTURE(FTestReplicationStateDescriptorBuilderFixture, TestClassWith
 	for (const FReplicationStateDescriptor* TheDescriptor : Descriptors)
 	{
 		UE_NET_ASSERT_TRUE(EnumHasAnyFlags(TheDescriptor->Traits, EReplicationStateTraits::HasPushBasedDirtiness));
+		UE_NET_ASSERT_TRUE(EnumHasAnyFlags(TheDescriptor->Traits, EReplicationStateTraits::HasFullPushBasedDirtiness));
 	}
 }
 
-UE_NET_TEST_FIXTURE(FTestReplicationStateDescriptorBuilderFixture, TestClassWithPushBasedPropertiesAndNeverReplicatedPropertiesAndFunctionsHasPushBasedDirtinessTrait)
+UE_NET_TEST_FIXTURE(FTestReplicationStateDescriptorBuilderFixture, TestClassWithPushBasedPropertiesAndNeverReplicatedPropertiesAndFunctionsHasFullPushBasedDirtinessTrait)
 {
 	InitDescriptorsFromClass(UTestReplicationStateDescriptor_TestClassWithPushBasedPropertiesAndNeverReplicatedPropertiesAndFunctions::StaticClass());
 	UE_NET_ASSERT_GE(Descriptors.Num(), 2);
@@ -1241,6 +1249,7 @@ UE_NET_TEST_FIXTURE(FTestReplicationStateDescriptorBuilderFixture, TestClassWith
 	for (const FReplicationStateDescriptor* TheDescriptor : Descriptors)
 	{
 		UE_NET_ASSERT_TRUE(EnumHasAnyFlags(TheDescriptor->Traits, EReplicationStateTraits::HasPushBasedDirtiness));
+		UE_NET_ASSERT_TRUE(EnumHasAnyFlags(TheDescriptor->Traits, EReplicationStateTraits::HasFullPushBasedDirtiness));
 	}
 }
 
@@ -1251,7 +1260,8 @@ UE_NET_TEST_FIXTURE(FTestReplicationStateDescriptorBuilderFixture, UTestReplicat
 
 	for (const FReplicationStateDescriptor* TheDescriptor : Descriptors)
 	{
-		UE_NET_ASSERT_FALSE(EnumHasAnyFlags(TheDescriptor->Traits, EReplicationStateTraits::HasPushBasedDirtiness));
+		UE_NET_ASSERT_TRUE(EnumHasAnyFlags(TheDescriptor->Traits, EReplicationStateTraits::HasPushBasedDirtiness));
+		UE_NET_ASSERT_FALSE(EnumHasAnyFlags(TheDescriptor->Traits, EReplicationStateTraits::HasFullPushBasedDirtiness));
 	}
 }
 
diff --git a/Engine/Plugins/Runtime/ReplicationSystemTestPlugin/Source/Private/Tests/ReplicationSystem/Conditionals/TestConditionals.cpp b/Engine/Plugins/Runtime/ReplicationSystemTestPlugin/Source/Private/Tests/ReplicationSystem/Conditionals/TestConditionals.cpp
--- a/Engine/Plugins/Runtime/ReplicationSystemTestPlugin/Source/Private/Tests/ReplicationSystem/Conditionals/TestConditionals.cpp
+++ b/Engine/Plugins/Runtime/ReplicationSystemTestPlugin/Source/Private/Tests/ReplicationSystem/Conditionals/TestConditionals.cpp
@@ -1076,4 +1076,127 @@ UE_NET_TEST_FIXTURE(FTestConditionalsFixture, ReplayOnlyIsNotReplicatedToRegular
 	UE_NET_ASSERT_EQ(ClientObject->ConnectionFilteredComponents[0]->ReplayOnlyInt, ExpectedClientValue);
 }
 
+UE_NET_TEST_FIXTURE(FTestConditionalsFixture, PushBased_COND_InitialOrOwner)
+{
+	// Add clients
+	FReplicationSystemTestClient* SomeClients[] = { CreateClient(), CreateClient() };
+	constexpr SIZE_T NonOwnerClientIndex = 0;
+	constexpr SIZE_T OwnerClientIndex = 1;
+
+	UTestReplicatedIrisPushModelObject* ServerObject = Server->CreateObject<UTestReplicatedIrisPushModelObject>();
+
+	// Set value
+	ServerObject->SetInitialOrOwnerInt(13);
+
+	// Set owning connection
+	Server->ReplicationSystem->SetOwningNetConnection(ServerObject->NetRefHandle, SomeClients[OwnerClientIndex]->ConnectionIdOnServer);
+
+	// Send and deliver packets
+	for (auto Client : SomeClients)
+	{
+		Server->NetUpdate();
+		Server->SendAndDeliverTo(Client, DeliverPacket);
+		Server->PostSendUpdate();
+	}
+	
+	// Check value on owning client
+	{
+		const UTestReplicatedIrisPushModelObject* ClientObject = Cast<UTestReplicatedIrisPushModelObject>(SomeClients[OwnerClientIndex]->GetReplicationBridge()->GetReplicatedObject(ServerObject->NetRefHandle));
+		UE_NET_ASSERT_NE(ClientObject, nullptr);
+
+		// Check that the value is the same as on the server
+		UE_NET_ASSERT_EQ(ClientObject->GetInitialOrOwnerInt(), ServerObject->GetInitialOrOwnerInt());
+	}
+
+	// Check value on non-owning client
+	{
+		const UTestReplicatedIrisPushModelObject* ClientObject = Cast<UTestReplicatedIrisPushModelObject>(SomeClients[NonOwnerClientIndex]->GetReplicationBridge()->GetReplicatedObject(ServerObject->NetRefHandle));
+		UE_NET_ASSERT_NE(ClientObject, nullptr);
+
+		// Check that the value is the same as on the server
+		UE_NET_ASSERT_EQ(ClientObject->GetInitialOrOwnerInt(), ServerObject->GetInitialOrOwnerInt());
+	}
+
+	// Modify value
+	const int32 PrevValue = ServerObject->GetInitialOrOwnerInt();
+	ServerObject->SetInitialOrOwnerInt(PrevValue + 13);
+
+	// Send and deliver packets
+	for (auto Client : SomeClients)
+	{
+		Server->NetUpdate();
+		Server->SendAndDeliverTo(Client, DeliverPacket);
+		Server->PostSendUpdate();
+	}
+
+	// Check value on owning client. It should've been updated.
+	{
+		const UTestReplicatedIrisPushModelObject* ClientObject = Cast<UTestReplicatedIrisPushModelObject>(SomeClients[OwnerClientIndex]->GetReplicationBridge()->GetReplicatedObject(ServerObject->NetRefHandle));
+		UE_NET_ASSERT_NE(ClientObject, nullptr);
+
+		// Check that the value the same as on the server
+		UE_NET_ASSERT_EQ(ClientObject->GetInitialOrOwnerInt(), ServerObject->GetInitialOrOwnerInt());
+	}
+
+	// Check value on non-owning client. It should remain the same as the previous value.
+	{
+		const UTestReplicatedIrisPushModelObject* ClientObject = Cast<UTestReplicatedIrisPushModelObject>(SomeClients[NonOwnerClientIndex]->GetReplicationBridge()->GetReplicatedObject(ServerObject->NetRefHandle));
+		UE_NET_ASSERT_NE(ClientObject, nullptr);
+
+		// Check that the value is unmodified.
+		UE_NET_ASSERT_EQ(ClientObject->GetInitialOrOwnerInt(), PrevValue);
+	}
+}
+
+UE_NET_TEST_FIXTURE(FTestConditionalsFixture, PushBased_COND_InitialOnly)
+{
+	// Create ClientA
+	FReplicationSystemTestClient* ClientA =  CreateClient();
+
+	UTestReplicatedIrisPushModelObject* ServerObject = Server->CreateObject<UTestReplicatedIrisPushModelObject>();
+
+	// Set value
+	ServerObject->SetInitialOnlyInt(13);
+
+	// Send and deliver to clientA
+	Server->UpdateAndSend({ ClientA });
+	
+	// Check value on ClientA
+	{
+		const UTestReplicatedIrisPushModelObject* ClientObject = Cast<UTestReplicatedIrisPushModelObject>(ClientA->GetReplicationBridge()->GetReplicatedObject(ServerObject->NetRefHandle));
+		UE_NET_ASSERT_NE(ClientObject, nullptr);
+
+		// Check that the value is the same as on the server
+		UE_NET_ASSERT_EQ(ClientObject->GetInitialOnlyInt(), ServerObject->GetInitialOnlyInt());
+	}
+
+	// Modify value
+	const int32 PrevValue = ServerObject->GetInitialOnlyInt();
+	ServerObject->SetInitialOnlyInt(PrevValue + 13);
+
+	// Create ClientB
+	FReplicationSystemTestClient* ClientB =  CreateClient();
+
+	// Send and deliver to clientA
+	Server->UpdateAndSend({ ClientA, ClientB});
+
+	// Check value on ClientA, should still be the PrevValue
+	{
+		const UTestReplicatedIrisPushModelObject* ClientObject = Cast<UTestReplicatedIrisPushModelObject>(ClientA->GetReplicationBridge()->GetReplicatedObject(ServerObject->NetRefHandle));
+		UE_NET_ASSERT_NE(ClientObject, nullptr);
+
+		// Check that the value is the same as on the server
+		UE_NET_ASSERT_EQ(ClientObject->GetInitialOnlyInt(), PrevValue);
+	}
+
+	// Check value on ClientB, it should be the updated value
+	{
+		const UTestReplicatedIrisPushModelObject* ClientObject = Cast<UTestReplicatedIrisPushModelObject>(ClientB->GetReplicationBridge()->GetReplicatedObject(ServerObject->NetRefHandle));
+		UE_NET_ASSERT_NE(ClientObject, nullptr);
+
+		// Check that the value is the same as on the server
+		UE_NET_ASSERT_EQ(ClientObject->GetInitialOnlyInt(), ServerObject->GetInitialOnlyInt());
+	}
+}
+
 }
diff --git a/Engine/Plugins/Runtime/ReplicationSystemTestPlugin/Source/Private/Tests/ReplicationSystem/ReplicatedTestObject.cpp b/Engine/Plugins/Runtime/ReplicationSystemTestPlugin/Source/Private/Tests/ReplicationSystem/ReplicatedTestObject.cpp
--- a/Engine/Plugins/Runtime/ReplicationSystemTestPlugin/Source/Private/Tests/ReplicationSystem/ReplicatedTestObject.cpp
+++ b/Engine/Plugins/Runtime/ReplicationSystemTestPlugin/Source/Private/Tests/ReplicationSystem/ReplicatedTestObject.cpp
@@ -489,6 +489,14 @@ void UTestReplicatedIrisPushModelObject::GetLifetimeReplicatedProps(TArray< clas
 	FDoRepLifetimeParams Params{ .bIsPushBased = true };
 	DOREPLIFETIME_WITH_PARAMS(UTestReplicatedIrisPushModelObject, IntA, Params);
 	DOREPLIFETIME_WITH_PARAMS(UTestReplicatedIrisPushModelObject, IntB, Params);
+
+	// Push based initial only
+	Params.Condition = ELifetimeCondition::COND_InitialOnly;
+	DOREPLIFETIME_WITH_PARAMS(UTestReplicatedIrisPushModelObject, InitialOnlyInt, Params);
+
+	// Push based initial or owner
+	Params.Condition = ELifetimeCondition::COND_InitialOrOwner;
+	DOREPLIFETIME_WITH_PARAMS(UTestReplicatedIrisPushModelObject, InitialOrOwnerInt, Params);
 }
 
 void UTestReplicatedIrisPushModelObject::RegisterReplicationFragments(UE::Net::FFragmentRegistrationContext& Context, UE::Net::EFragmentRegistrationFlags RegistrationFlags)
@@ -500,6 +508,27 @@ void UTestReplicatedIrisPushModelObject::RegisterReplicationFragments(UE::Net::F
 	}
 }
 
+void UTestReplicatedIrisPushModelObject::SetInitialOnlyInt(int32 InValue)
+{
+	InitialOnlyInt = InValue;
+	MARK_PROPERTY_DIRTY_FROM_NAME(UTestReplicatedIrisPushModelObject, InitialOnlyInt, this);
+}
+
+int32 UTestReplicatedIrisPushModelObject::GetInitialOnlyInt() const
+{
+	return InitialOnlyInt;
+}
+
+void UTestReplicatedIrisPushModelObject::SetInitialOrOwnerInt(int32 InValue)
+{
+	InitialOrOwnerInt = InValue;
+	MARK_PROPERTY_DIRTY_FROM_NAME(UTestReplicatedIrisPushModelObject, InitialOrOwnerInt, this);
+}
+
+int32 UTestReplicatedIrisPushModelObject::GetInitialOrOwnerInt() const
+{
+	return InitialOrOwnerInt;
+}
 
 void UTestReplicatedIrisPushModelObject::SetIntA(int32 InValue)
 {
diff --git a/Engine/Plugins/Runtime/ReplicationSystemTestPlugin/Source/Private/Tests/ReplicationSystem/ReplicatedTestObject.h b/Engine/Plugins/Runtime/ReplicationSystemTestPlugin/Source/Private/Tests/ReplicationSystem/ReplicatedTestObject.h
--- a/Engine/Plugins/Runtime/ReplicationSystemTestPlugin/Source/Private/Tests/ReplicationSystem/ReplicatedTestObject.h
+++ b/Engine/Plugins/Runtime/ReplicationSystemTestPlugin/Source/Private/Tests/ReplicationSystem/ReplicatedTestObject.h
@@ -571,7 +571,19 @@ public:
 	void SetIntB(int32 InValue);
 	int32 GetIntB() const;
 
+	void SetInitialOnlyInt(int32 InValue);
+	int32 GetInitialOnlyInt() const;
+
+	void SetInitialOrOwnerInt(int32 InValue);
+	int32 GetInitialOrOwnerInt() const;
+
 private:
+	UPROPERTY(Transient, Replicated)
+	int32 InitialOnlyInt;
+
+	UPROPERTY(Transient, Replicated)
+	int32 InitialOrOwnerInt;
+	
 	UPROPERTY(Transient, Replicated)
 	int32 IntA;
 
diff --git a/Engine/Plugins/Runtime/ReplicationSystemTestPlugin/Source/Private/Tests/ReplicationSystem/TestReplicationBridge.cpp b/Engine/Plugins/Runtime/ReplicationSystemTestPlugin/Source/Private/Tests/ReplicationSystem/TestReplicationBridge.cpp
--- a/Engine/Plugins/Runtime/ReplicationSystemTestPlugin/Source/Private/Tests/ReplicationSystem/TestReplicationBridge.cpp
+++ b/Engine/Plugins/Runtime/ReplicationSystemTestPlugin/Source/Private/Tests/ReplicationSystem/TestReplicationBridge.cpp
@@ -35,7 +35,7 @@ UE_NET_TEST_FIXTURE(FReplicationBridgeTestFixture, CreateObjectWithSingleFragmen
 	UE_NET_ASSERT_EQ(1u, Protocol->ReplicationStateCount);
 	UE_NET_ASSERT_EQ(4u, Protocol->InternalTotalAlignment);
 	UE_NET_ASSERT_EQ(20u, Protocol->InternalTotalSize);
-	UE_NET_ASSERT_EQ(32u, Protocol->MaxExternalStateSize);
+	UE_NET_ASSERT_EQ(36u, Protocol->MaxExternalStateSize);
 	UE_NET_ASSERT_EQ(5u, Protocol->ChangeMaskBitCount);
 	UE_NET_ASSERT_EQ(EReplicationProtocolTraits::SupportsDeltaCompression, Protocol->ProtocolTraits);
 
diff --git a/Engine/Source/Runtime/Experimental/Iris/Core/Private/Iris/IrisConfigInternal.h b/Engine/Source/Runtime/Experimental/Iris/Core/Private/Iris/IrisConfigInternal.h
--- a/Engine/Source/Runtime/Experimental/Iris/Core/Private/Iris/IrisConfigInternal.h
+++ b/Engine/Source/Runtime/Experimental/Iris/Core/Private/Iris/IrisConfigInternal.h
@@ -47,3 +47,6 @@
 #endif 
 #endif
 
+#ifndef UE_NET_WITH_VERBOSE_POLL_PROFILING
+#   define UE_NET_WITH_VERBOSE_POLL_PROFILING (!UE_BUILD_SHIPPING)
+#endif
diff --git a/Engine/Source/Runtime/Experimental/Iris/Core/Private/Iris/ReplicationState/PropertyReplicationState.cpp b/Engine/Source/Runtime/Experimental/Iris/Core/Private/Iris/ReplicationState/PropertyReplicationState.cpp
--- a/Engine/Source/Runtime/Experimental/Iris/Core/Private/Iris/ReplicationState/PropertyReplicationState.cpp
+++ b/Engine/Source/Runtime/Experimental/Iris/Core/Private/Iris/ReplicationState/PropertyReplicationState.cpp
@@ -16,6 +16,7 @@
 #include "UObject/PropertyPortFlags.h"
 #include "Containers/StringFwd.h"
 #include "ProfilingDebugging/CsvProfiler.h"
+#include "Iris/IrisConfigInternal.h"
 
 DEFINE_LOG_CATEGORY_STATIC(LogIrisRepNotify, Warning, All);
 
@@ -84,6 +85,11 @@ bool FPropertyReplicationState::IsInitState() const
 	return IsValid() && ReplicationStateDescriptor->IsInitState();
 }
 
+bool FPropertyReplicationState::IsDirtyForPolling() const
+{
+	return IsDirty() || Private::GetMemberPollMask(StateBuffer, ReplicationStateDescriptor).IsAnyBitSet();
+}
+
 bool FPropertyReplicationState::IsDirty() const
 {
 	if (IsValid())
@@ -220,6 +226,11 @@ void FPropertyReplicationState::MarkDirty(uint32 Index)
 	}
 }
 
+bool FPropertyReplicationState::IsDirtyForPolling(uint32 Index) const
+{
+	return IsDirty(Index) || Private::GetMemberPollMask(StateBuffer, ReplicationStateDescriptor).GetBit(Index);
+}
+
 bool FPropertyReplicationState::IsDirty(uint32 Index) const
 {
 	const FReplicationStateDescriptor* Descriptor = ReplicationStateDescriptor;
@@ -270,7 +281,9 @@ bool FPropertyReplicationState::PollPropertyReplicationState(const void* RESTRIC
 		const FReplicationStateDescriptor* Descriptor = ReplicationStateDescriptor;
 		const uint8* SrcBuffer = reinterpret_cast<const uint8*>(SrcStateData);
 
+#if UE_NET_WITH_VERBOSE_POLL_PROFILING
 		IRIS_PROFILER_PROTOCOL_NAME(ReplicationStateDescriptor->DebugName->Name);
+#endif
 
 		const FReplicationStateMemberDescriptor* MemberDescriptors = Descriptor->MemberDescriptors;
 		const FProperty** MemberProperties = Descriptor->MemberProperties;
@@ -291,6 +304,44 @@ bool FPropertyReplicationState::PollPropertyReplicationState(const void* RESTRIC
 	return IsDirty();
 }
 
+bool FPropertyReplicationState::PollPushBasedPropertyReplicationState(const void* RESTRICT SrcStateData, const FPollParameters& PollParameters)
+{
+	if (IsValid())
+	{
+		const FReplicationStateDescriptor* Descriptor = ReplicationStateDescriptor;
+		const uint8* SrcBuffer = reinterpret_cast<const uint8*>(SrcStateData);
+
+#if UE_NET_WITH_VERBOSE_POLL_PROFILING
+		IRIS_PROFILER_SCOPE_TEXT_CONDITIONAL(Descriptor->DebugName->Name, PollParameters.bEnableVerboseProfiling);
+#endif
+
+		const FReplicationStateMemberDescriptor* MemberDescriptors = Descriptor->MemberDescriptors;
+		const FProperty** MemberProperties = Descriptor->MemberProperties;
+		const FReplicationStateMemberPropertyDescriptor* MemberPropertyDescriptors = Descriptor->MemberPropertyDescriptors;
+		const uint32 MemberCount = Descriptor->MemberCount;
+		const FReplicationStateMemberTraitsDescriptor* MemberTraits = Descriptor->MemberTraitsDescriptors;
+		for (uint32 MemberIt = 0; MemberIt < MemberCount; ++MemberIt)
+		{
+			const bool bForcePoll = (PollParameters.bForceRefreshObjectReferences && EnumHasAnyFlags(MemberTraits[MemberIt].Traits, EReplicationStateMemberTraits::HasObjectReference)) || !EnumHasAnyFlags(MemberTraits[MemberIt].Traits, EReplicationStateMemberTraits::HasPushBasedDirtiness);
+			if (bForcePoll || IsDirtyForPolling(MemberIt))
+			{
+#if UE_NET_WITH_VERBOSE_POLL_PROFILING
+				IRIS_PROFILER_SCOPE_TEXT_CONDITIONAL(EnumHasAnyFlags(MemberTraits[MemberIt].Traits, EReplicationStateMemberTraits::HasPushBasedDirtiness) ? TEXT("PushBased") : TEXT("AlwaysPolled"), PollParameters.bEnableVerboseProfiling);
+				IRIS_PROFILER_SCOPE_TEXT_CONDITIONAL(Descriptor->MemberDebugDescriptors[MemberIt].DebugName->Name, PollParameters.bEnableVerboseProfiling);
+#endif
+				const FReplicationStateMemberDescriptor& MemberDescriptor = MemberDescriptors[MemberIt];
+				const FReplicationStateMemberPropertyDescriptor& MemberPropertyDescriptor = MemberPropertyDescriptors[MemberIt];
+				const FProperty* Property = MemberProperties[MemberIt];
+
+				// $TODO: make special version to avoid unnecessary overhead related to looking up same things for every single member.
+				PollPropertyValue(MemberIt, SrcBuffer + Property->GetOffset_ForGC() + Property->GetElementSize()*MemberPropertyDescriptor.ArrayIndex);
+			}
+		}		
+	}
+
+	return IsDirty();
+}
+
 bool FPropertyReplicationState::StoreCurrentPropertyReplicationStateForRepNotifies(const void* RESTRICT SrcStateData, const FPropertyReplicationState* NewStateToBeApplied)
 {
 	if (IsValid())
@@ -341,7 +392,9 @@ void FPropertyReplicationState::PushPropertyReplicationState(const UObject* Owne
 		const FReplicationStateDescriptor* Descriptor = ReplicationStateDescriptor;
 		uint8* DstBuffer = static_cast<uint8*>(DstData);
 
+#if UE_NET_WITH_VERBOSE_POLL_PROFILING
 		IRIS_PROFILER_PROTOCOL_NAME(ReplicationStateDescriptor->DebugName->Name);
+#endif
 
 		const FReplicationStateMemberDescriptor* MemberDescriptors = Descriptor->MemberDescriptors;
 		const FProperty** MemberProperties = Descriptor->MemberProperties;
diff --git a/Engine/Source/Runtime/Experimental/Iris/Core/Private/Iris/ReplicationState/ReplicationStateDescriptorBuilder.cpp b/Engine/Source/Runtime/Experimental/Iris/Core/Private/Iris/ReplicationState/ReplicationStateDescriptorBuilder.cpp
--- a/Engine/Source/Runtime/Experimental/Iris/Core/Private/Iris/ReplicationState/ReplicationStateDescriptorBuilder.cpp
+++ b/Engine/Source/Runtime/Experimental/Iris/Core/Private/Iris/ReplicationState/ReplicationStateDescriptorBuilder.cpp
@@ -1146,9 +1146,12 @@ void FPropertyReplicationStateDescriptorBuilder::BuildMemberChangeMaskDescriptor
 	Descriptor->ChangeMaskBitCount = static_cast<uint16>(CurrentBitOffset);
 	Descriptor->ChangeMasksExternalOffset = static_cast<uint32>(Align(Context.External.Size, ChangeMaskMemberAlignment));
 
+	// We also reserve room for the MemberPollMask
+	const SIZE_T MemberPollMaskMemberSize = FNetBitArrayView::CalculateRequiredWordCount(Members.Num()) * sizeof(FNetBitArrayView::StorageWordType);
+
 	// Update context
 	Context.External.Alignment = FMath::Max(Context.External.Alignment, ChangeMaskMemberAlignment);
-	Context.External.Size = Descriptor->ChangeMasksExternalOffset + ChangeMaskMemberSize + ConditionalChangeMaskMemberSize;
+	Context.External.Size = Descriptor->ChangeMasksExternalOffset + ChangeMaskMemberSize + ConditionalChangeMaskMemberSize + MemberPollMaskMemberSize;
 }
 
 void FPropertyReplicationStateDescriptorBuilder::BuildMemberSerializerDescriptors(FReplicationStateMemberSerializerDescriptor* MemberSerializerDescriptors, FReplicationStateDescriptor* Descriptor) const
@@ -1179,6 +1182,7 @@ void FPropertyReplicationStateDescriptorBuilder::BuildMemberTraitsDescriptors(FR
 		CurrentMemberTraitsDescriptor->Traits |= (EnumHasAnyFlags(Info.Traits, EMemberPropertyTraits::HasObjectReference) ? EReplicationStateMemberTraits::HasObjectReference : EReplicationStateMemberTraits::None);
 		CurrentMemberTraitsDescriptor->Traits |= (EnumHasAnyFlags(Info.Traits, EMemberPropertyTraits::RepNotifyAlways) ? EReplicationStateMemberTraits::HasRepNotifyAlways : EReplicationStateMemberTraits::None);
 		CurrentMemberTraitsDescriptor->Traits |= (EnumHasAnyFlags(Info.Traits, EMemberPropertyTraits::UseSerializerIsEqual) ? EReplicationStateMemberTraits::UseSerializerIsEqual : EReplicationStateMemberTraits::None);
+		CurrentMemberTraitsDescriptor->Traits |= (EnumHasAnyFlags(Info.Traits, EMemberPropertyTraits::HasPushBasedDirtiness) ? EReplicationStateMemberTraits::HasPushBasedDirtiness : EReplicationStateMemberTraits::None);
 		++CurrentMemberTraitsDescriptor;
 	}
 
@@ -1586,11 +1590,16 @@ EReplicationStateTraits FPropertyReplicationStateDescriptorBuilder::BuildReplica
 	{
 		Traits |= EReplicationStateTraits::UseSerializerIsEqual;
 	}
-
-	// Only set the PushBasedDirtiness trait if all properties are push based. An alternative could be to split into multiple states.
-	if (EnumHasAnyFlags(SharedPropertyTraits, EMemberPropertyTraits::HasPushBasedDirtiness))
+	
+	if (EnumHasAnyFlags(PropertyTraits, EMemberPropertyTraits::HasPushBasedDirtiness))
 	{
 		Traits |= EReplicationStateTraits::HasPushBasedDirtiness;
+
+		// Only set the HasFullPushBasedDirtiness trait if all properties are push based. An alternative could be to split into multiple states.
+		if (EnumHasAnyFlags(SharedPropertyTraits, EMemberPropertyTraits::HasPushBasedDirtiness))
+		{
+			Traits |= EReplicationStateTraits::HasFullPushBasedDirtiness;
+		}
 	}
 
 	// Struct traits
@@ -1651,6 +1660,7 @@ EReplicationStateTraits FPropertyReplicationStateDescriptorBuilder::BuildReplica
 	if (Context.MemberFunctionCache.Num() && Members.Num() == 0)
 	{
 		Traits |= EReplicationStateTraits::HasPushBasedDirtiness;
+		Traits |= EReplicationStateTraits::HasFullPushBasedDirtiness;
 	}
 
 	Traits |= EReplicationStateTraits::NeedsRefCount;
@@ -2030,9 +2040,7 @@ EMemberPropertyTraits FPropertyReplicationStateDescriptorBuilder::GetConnectionF
 
 EMemberPropertyTraits FPropertyReplicationStateDescriptorBuilder::GetInitOnlyTrait(ELifetimeCondition Condition)
 { 
-	bool bResult = 
-		Condition == COND_InitialOnly ||
-		Condition == COND_InitialOrOwner;
+	bool bResult = Condition == COND_InitialOnly;
 
 	return bResult ? EMemberPropertyTraits::InitOnly : EMemberPropertyTraits::None;
 }
diff --git a/Engine/Source/Runtime/Experimental/Iris/Core/Private/Iris/ReplicationSystem/DirtyNetObjectTracker.cpp b/Engine/Source/Runtime/Experimental/Iris/Core/Private/Iris/ReplicationSystem/DirtyNetObjectTracker.cpp
--- a/Engine/Source/Runtime/Experimental/Iris/Core/Private/Iris/ReplicationSystem/DirtyNetObjectTracker.cpp
+++ b/Engine/Source/Runtime/Experimental/Iris/Core/Private/Iris/ReplicationSystem/DirtyNetObjectTracker.cpp
@@ -3,7 +3,9 @@
 #include "DirtyNetObjectTracker.h"
 #include "Iris/IrisConstants.h"
 #include "Iris/Core/IrisProfiler.h"
+#include "Iris/Core/IrisLog.h"
 #include "Iris/ReplicationSystem/ReplicationSystem.h"
+#include "Iris/ReplicationState/ReplicationStateUtil.h"
 #include "Iris/ReplicationSystem/ReplicationSystemInternal.h"
 
 #include "ProfilingDebugging/CsvProfiler.h"
@@ -63,8 +65,133 @@ void FDirtyNetObjectTracker::OnMaxInternalNetRefIndexIncreased(FInternalNetRefIn
 	NetObjectIdCount = NewMaxInternalIndex;
 }
 
+bool FDirtyNetObjectTracker::MarkPushbasedPropertiesDirty(FInternalNetRefIndex ObjectIndex, uint16 OwnerIndex, const FNetBitArrayView& DirtyProperties)
+{
+	bool bShouldMarkObjectDirty = false;
+
+	// For now, we do not support more than a single push based owner per protocol. JIRA: UE-278338
+	if (!ensure(OwnerIndex == 0U))
+	{
+		return false;
+	}
+
+	const FNetRefHandleManager::FReplicatedObjectData& ReplicatedObjectData = NetRefHandleManager->GetReplicatedObjectDataNoCheck(ObjectIndex);
+	const FReplicationProtocol* Protocol = ReplicatedObjectData.Protocol;
+	const FReplicationInstanceProtocol* InstanceProtocol = ReplicatedObjectData.InstanceProtocol;
+
+	// Nothing to poll to or from.
+	if (!InstanceProtocol || !NetRefHandleManager->GetReplicatedObjectStateBufferNoCheck(ObjectIndex))
+	{
+		return false;
+	}
+	
+	if (Protocol->PushModelOwnerCount == 0U)
+	{
+		return false;
+	}
+
+	if (Protocol->PushModelOwnerCount > 1U)
+	{
+		UE_LOG(LogIris, Error, TEXT("Failed to mark dirty due to too many pushmodel state owners ObjectIndex %d"), ObjectIndex);					
+		return false;
+	}
+
+	const FNetRefHandleManager* LocalNetRefHandleManager = NetRefHandleManager;
+	DirtyProperties.ForAllSetBits([ObjectIndex, OwnerIndex, LocalNetRefHandleManager,  &Protocol, &InstanceProtocol, &bShouldMarkObjectDirty](uint32 RepIndex)
+		{
+			// Dynamic part
+			const FReplicationProtocol::FRepIndexToFragmentIndexTable& FragmentIndexTable = Protocol->PushModelOwnerRepIndexToFragmentIndexTable[OwnerIndex];
+			if (RepIndex >= FragmentIndexTable.NumEntries)
+			{
+				ensureMsgf(false, TEXT("Trying to mark invalid property dirty (Could be a disabled property). Invalid memberindex %s : RepIndex [%d][%d]"), *LocalNetRefHandleManager->PrintObjectFromIndex(ObjectIndex), OwnerIndex, RepIndex);
+				return;
+			}
+			const FReplicationProtocol::FRepIndexToFragmentIndex FragmentIndex = FragmentIndexTable.RepIndexToFragmentIndex[RepIndex];
+
+			if (FragmentIndex.FragmentIndex == FReplicationProtocol::FRepIndexToFragmentIndex::InvalidEntry)
+			{
+				return;
+			}
+
+			const FReplicationInstanceProtocol::FFragmentData& Fragment = InstanceProtocol->FragmentData[FragmentIndex.FragmentIndex];
+			const FReplicationStateDescriptor* StateDescriptor = Protocol->ReplicationStateDescriptors[FragmentIndex.FragmentIndex];
+
+			// Does this state contain this property?
+			const FReplicationStateMemberRepIndexToMemberIndexDescriptor& RepIndexToMemberIndexDescriptor = StateDescriptor->MemberRepIndexToMemberIndexDescriptors[RepIndex];
+			if (RepIndexToMemberIndexDescriptor.MemberIndex == FReplicationStateMemberRepIndexToMemberIndexDescriptor::InvalidEntry)
+			{
+				ensureMsgf(false, TEXT("Trying to mark not existing property dirty. Invalid memberindex %s : RepIndex [%d][%d]"), *LocalNetRefHandleManager->PrintObjectFromIndex(ObjectIndex), OwnerIndex, RepIndex);
+				return;
+			}
+
+			// Skip custom conditionals
+			{
+				const EReplicationStateTraits Traits = StateDescriptor->Traits;
+				if (EnumHasAnyFlags(Traits, EReplicationStateTraits::HasLifetimeConditionals))
+				{
+					const FReplicationStateMemberChangeMaskDescriptor& ChangeMaskInfo = StateDescriptor->MemberChangeMaskDescriptors[RepIndexToMemberIndexDescriptor.MemberIndex];
+
+					if (ChangeMaskInfo.BitCount > 0)
+					{
+						FNetBitArrayView MemberConditionalChangeMask = Private::GetMemberConditionalChangeMask(Fragment.ExternalSrcBuffer, StateDescriptor);
+						if (!MemberConditionalChangeMask.GetBit(ChangeMaskInfo.BitOffset))
+						{
+							return;
+						}
+					}
+				}
+			}
+
+			// Mark dirty
+			FNetBitArrayView DirtyMembersForPolling = GetMemberPollMask(Fragment.ExternalSrcBuffer, StateDescriptor);
+			DirtyMembersForPolling.SetBit(RepIndexToMemberIndexDescriptor.MemberIndex);
+			bShouldMarkObjectDirty |= true;
+		}
+	);
+
+	return bShouldMarkObjectDirty;
+}
+
 void FDirtyNetObjectTracker::ApplyGlobalDirtyObjectList()
 {
+	if (FGlobalDirtyNetObjectTracker::IsUsingPerPropertyDirtyTracking())
+	{
+		if (UReplicationSystem* ReplicationSystem = GetReplicationSystem(ReplicationSystemId))
+		{
+			FReplicationConditionals& Conditionals = ReplicationSystem->GetReplicationSystemInternal()->GetConditionals();
+
+			const FGlobalDirtyNetObjectTracker::FDirtyHandleAndPropertyMap& GlobalDirtyNetObjects = FGlobalDirtyNetObjectTracker::GetDirtyNetObjectsAndProperties(GlobalDirtyTrackerPollHandle);
+			for (const FGlobalDirtyNetObjectTracker::FDirtyHandleAndPropertyMap::ElementType& Element : GlobalDirtyNetObjects)
+			{
+				const FNetHandle NetHandle = Element.Key;
+
+				const FInternalNetRefIndex NetObjectIndex = NetRefHandleManager->GetInternalIndexFromNetHandle(NetHandle);
+				if (NetObjectIndex != FNetRefHandleManager::InvalidInternalIndex)
+				{
+					// Update changemasks based on rep-indices marked dirty
+					const FGlobalDirtyNetObjectTracker::FDirtyPropertyStorage& DirtyPropertyData = Element.Value;
+					if (DirtyPropertyData.Num())
+					{
+						const uint32 BitCount = DirtyPropertyData.Num() * sizeof(FGlobalDirtyNetObjectTracker::FDirtyPropertyStorage::ElementType) * 8U;						
+						constexpr uint32 OwnerIndex = 0U; // Always 0 as we currently onyly support a single push based fragment owner
+
+						FNetBitArrayView DirtyRepIndices = MakeNetBitArrayView(DirtyPropertyData.GetData(), BitCount, FNetBitArrayView::NoResetNoValidate);
+						if (const bool bShouldMarkObjectDirty = MarkPushbasedPropertiesDirty(NetObjectIndex, OwnerIndex, DirtyRepIndices))
+						{
+							// MarkObjectDirty
+							DirtyNetObjects.SetBit(NetObjectIndex);							
+						}
+					}
+					else
+					{
+						// If it was an explicit dirty call for the object, mark object dirty
+						DirtyNetObjects.SetBit(NetObjectIndex);
+					}
+				}
+			}
+		}
+	}
+	else
 	{
 		const TSet<FNetHandle>& GlobalDirtyNetObjects = FGlobalDirtyNetObjectTracker::GetDirtyNetObjects(GlobalDirtyTrackerPollHandle);
 		for (FNetHandle NetHandle : GlobalDirtyNetObjects)
@@ -175,7 +302,7 @@ void FDirtyNetObjectTracker::ForceNetUpdate(FInternalNetRefIndex NetObjectIndex)
 
 	ForceNetUpdateObjects.SetBit(NetObjectIndex);
 
-	// Flag the object dirty so we update his filters too
+	// Flag the object dirty so we update filters etc too
 	MarkNetObjectDirty(NetObjectIndex);
 
 	UE_LOG(LogIrisDirtyTracker, Verbose, TEXT("FDirtyNetObjectTracker::ForceNetUpdateObjects[%u]: %s"), ReplicationSystemId, *NetRefHandleManager->PrintObjectFromIndex(NetObjectIndex));
diff --git a/Engine/Source/Runtime/Experimental/Iris/Core/Private/Iris/ReplicationSystem/DirtyNetObjectTracker.h b/Engine/Source/Runtime/Experimental/Iris/Core/Private/Iris/ReplicationSystem/DirtyNetObjectTracker.h
--- a/Engine/Source/Runtime/Experimental/Iris/Core/Private/Iris/ReplicationSystem/DirtyNetObjectTracker.h
+++ b/Engine/Source/Runtime/Experimental/Iris/Core/Private/Iris/ReplicationSystem/DirtyNetObjectTracker.h
@@ -112,6 +112,9 @@ private:
 	/** Can only be accessed via FDirtyObjectsAccessor */
 	FNetBitArrayView GetDirtyNetObjectsThisFrame();
 
+	/** Propagate properties marked dirty for given object and OwnerIndex */
+	bool MarkPushbasedPropertiesDirty(FInternalNetRefIndex ObjectIndex, uint16 OwnerIndex, const FNetBitArrayView& DirtyProperties);
+
 private:
 
 	// Dirty objects that persist across frames.
diff --git a/Engine/Source/Runtime/Experimental/Iris/Core/Private/Iris/ReplicationSystem/FastArrayReplicationFragment.cpp b/Engine/Source/Runtime/Experimental/Iris/Core/Private/Iris/ReplicationSystem/FastArrayReplicationFragment.cpp
--- a/Engine/Source/Runtime/Experimental/Iris/Core/Private/Iris/ReplicationSystem/FastArrayReplicationFragment.cpp
+++ b/Engine/Source/Runtime/Experimental/Iris/Core/Private/Iris/ReplicationSystem/FastArrayReplicationFragment.cpp
@@ -90,6 +90,10 @@ FFastArrayReplicationFragmentBase::FFastArrayReplicationFragmentBase(EReplicatio
 		FFastArraySerializer* FastArraySerializer = reinterpret_cast<FFastArraySerializer*>(reinterpret_cast<uint8*>(InOwner) + InDescriptor->MemberProperties[0]->GetOffset_ForGC());
 		FastArraySerializer->CachePushModelState(InOwner, InDescriptor->MemberProperties[0]->RepIndex);
 		Traits |= EReplicationFragmentTraits::HasPushBasedDirtiness;
+		if (EnumHasAnyFlags(InDescriptor->Traits, EReplicationStateTraits::HasFullPushBasedDirtiness))
+		{
+			Traits |= EReplicationFragmentTraits::HasFullPushBasedDirtiness;
+		}
 	}
 #endif
 }
diff --git a/Engine/Source/Runtime/Experimental/Iris/Core/Private/Iris/ReplicationSystem/LegacyPushModel.cpp b/Engine/Source/Runtime/Experimental/Iris/Core/Private/Iris/ReplicationSystem/LegacyPushModel.cpp
--- a/Engine/Source/Runtime/Experimental/Iris/Core/Private/Iris/ReplicationSystem/LegacyPushModel.cpp
+++ b/Engine/Source/Runtime/Experimental/Iris/Core/Private/Iris/ReplicationSystem/LegacyPushModel.cpp
@@ -61,13 +61,13 @@ void FNetHandleLegacyPushModelHelper::ClearNetPushID(UObject* Object)
 void FNetHandleLegacyPushModelHelper::MarkPropertyOwnerDirty(const UObject* Object, UEPushModelPrivate::FNetIrisPushObjectId PushId, const int32 RepIndex)
 {
 	const FNetPushObjectHandle Handle(PushId);
-	MarkNetObjectStateDirty(Handle.GetNetHandle());
+	MarkNetObjectStateDirty(Handle.GetNetHandle(), RepIndex, RepIndex);
 }
 
 void FNetHandleLegacyPushModelHelper::MarkPropertiesOwnerDirty(const UObject* Object, UEPushModelPrivate::FNetIrisPushObjectId PushId, const int32 StartRepIndex, const int32 EndRepIndex)
 {
 	const FNetPushObjectHandle Handle(PushId);
-	MarkNetObjectStateDirty(Handle.GetNetHandle());
+	MarkNetObjectStateDirty(Handle.GetNetHandle(), StartRepIndex, EndRepIndex);
 }
 
 void FNetHandleLegacyPushModelHelper::OptionallyMarkPropertyOwnerDirty(const UObject* Object, UEPushModelPrivate::FNetIrisPushObjectId PushId, const int32 RepIndex)
@@ -75,7 +75,7 @@ void FNetHandleLegacyPushModelHelper::OptionallyMarkPropertyOwnerDirty(const UOb
 	if (IsIrisPushModelEnabled(true))
 	{
 		const FNetPushObjectHandle Handle(PushId);
-		MarkNetObjectStateDirty(Handle.GetNetHandle());
+		MarkNetObjectStateDirty(Handle.GetNetHandle(), RepIndex, RepIndex);
 	}
 }
 
@@ -84,7 +84,7 @@ void FNetHandleLegacyPushModelHelper::OptionallyMarkPropertiesOwnerDirty(const U
 	if (IsIrisPushModelEnabled(true))
 	{
 		const FNetPushObjectHandle Handle(PushId);
-		MarkNetObjectStateDirty(Handle.GetNetHandle());
+		MarkNetObjectStateDirty(Handle.GetNetHandle(), StartRepIndex, EndRepIndex);
 	}
 }	
 
diff --git a/Engine/Source/Runtime/Experimental/Iris/Core/Private/Iris/ReplicationSystem/NetRefHandleManager.cpp b/Engine/Source/Runtime/Experimental/Iris/Core/Private/Iris/ReplicationSystem/NetRefHandleManager.cpp
--- a/Engine/Source/Runtime/Experimental/Iris/Core/Private/Iris/ReplicationSystem/NetRefHandleManager.cpp
+++ b/Engine/Source/Runtime/Experimental/Iris/Core/Private/Iris/ReplicationSystem/NetRefHandleManager.cpp
@@ -74,6 +74,7 @@ void FNetRefHandleManager::Init(const FInitParams& InitParams)
 		InitNetBitArray(&WantToBeDormantInternalIndices);
 		InitNetBitArray(&ObjectsWithPreUpdate);
 		InitNetBitArray(&DormantObjectsPendingFlushNet);
+		InitNetBitArray(&ObjectsWithFullPushBasedDirtiness);	
 	}
 
 	// Mark the invalid index as used
@@ -235,6 +236,7 @@ FInternalNetRefIndex FNetRefHandleManager::InternalCreateNetObject(const FNetRef
 	Data.Protocol = ReplicationProtocol;
 	Data.InstanceProtocol = nullptr;
 	ObjectsWithPreUpdate.ClearBit(InternalIndex);
+	ObjectsWithFullPushBasedDirtiness.ClearBit(InternalIndex);
 	ReplicatedObjectStateBuffers[InternalIndex] = nullptr;
 	Data.ReceiveStateBuffer = nullptr;
 	Data.bShouldPropagateChangedStates = 1U;
@@ -279,6 +281,7 @@ void FNetRefHandleManager::AttachInstanceProtocol(FInternalNetRefIndex InternalI
 		ReplicatedInstances[InternalIndex] = Instance;
 
 		ObjectsWithPreUpdate.SetBitValue(InternalIndex, EnumHasAnyFlags(InstanceProtocol->InstanceTraits, EReplicationInstanceProtocolTraits::NeedsPreSendUpdate));
+		ObjectsWithFullPushBasedDirtiness.SetBitValue(InternalIndex, EnumHasAnyFlags(InstanceProtocol->InstanceTraits, EReplicationInstanceProtocolTraits::HasFullPushBasedDirtiness));
 	}
 }
 
@@ -293,6 +296,7 @@ const FReplicationInstanceProtocol* FNetRefHandleManager::DetachInstanceProtocol
 		ReplicatedInstances[InternalIndex] = nullptr;
 
 		ObjectsWithPreUpdate.ClearBit(InternalIndex);
+		ObjectsWithFullPushBasedDirtiness.ClearBit(InternalIndex);
 		
 		return InstanceProtocol;
 	}
@@ -494,6 +498,7 @@ void FNetRefHandleManager::ClearStateForFreedInternalIndex(FInternalNetRefIndex
 	ObjectsWithDependentObjectsInternalIndices.ClearBit(FreedInternalIndex);
 	WantToBeDormantInternalIndices.ClearBit(FreedInternalIndex);
 	DormantObjectsPendingFlushNet.ClearBit(FreedInternalIndex);
+	ObjectsWithFullPushBasedDirtiness.ClearBit(FreedInternalIndex);
 }
 
 FNetRefHandle FNetRefHandleManager::CreateHandleForDestructionInfo(FNetRefHandle Handle, const FReplicationProtocol* DestroyedObjectProtocol)
diff --git a/Engine/Source/Runtime/Experimental/Iris/Core/Private/Iris/ReplicationSystem/NetRefHandleManager.h b/Engine/Source/Runtime/Experimental/Iris/Core/Private/Iris/ReplicationSystem/NetRefHandleManager.h
--- a/Engine/Source/Runtime/Experimental/Iris/Core/Private/Iris/ReplicationSystem/NetRefHandleManager.h
+++ b/Engine/Source/Runtime/Experimental/Iris/Core/Private/Iris/ReplicationSystem/NetRefHandleManager.h
@@ -327,6 +327,9 @@ public:
 	/** Get Objects that are flagged for PreUpdate (aka PreReplication) */
 	FNetBitArrayView GetObjectsWithPreUpdate() const { return MakeNetBitArrayView(ObjectsWithPreUpdate); }
 
+	/** Get Objects that fully supports pushbased dirtiness */
+	FNetBitArrayView GetObjectsWithFullPushBasedDirtiness() const { return MakeNetBitArrayView(ObjectsWithFullPushBasedDirtiness); }
+
 public:
 
 	/** Return a string to identify the object linked to an index in logs */
@@ -482,6 +485,9 @@ private:
 	// Bitset marking internal indices that is flagged as requiring PreUpdate
 	FNetBitArray ObjectsWithPreUpdate;
 
+	// Bitset marking internal indices that is flagged as supporting FullPushBasedDirtiness
+	FNetBitArray ObjectsWithFullPushBasedDirtiness;
+
 	// Both are marked in DestroyedStartupObjectsInternalIndicies
 	// When one is destroyed both mapping will be cleared
 	// Map OriginalIndex -> DestructionInfoIndex
diff --git a/Engine/Source/Runtime/Experimental/Iris/Core/Private/Iris/ReplicationSystem/ObjectReplicationBridge.cpp b/Engine/Source/Runtime/Experimental/Iris/Core/Private/Iris/ReplicationSystem/ObjectReplicationBridge.cpp
--- a/Engine/Source/Runtime/Experimental/Iris/Core/Private/Iris/ReplicationSystem/ObjectReplicationBridge.cpp
+++ b/Engine/Source/Runtime/Experimental/Iris/Core/Private/Iris/ReplicationSystem/ObjectReplicationBridge.cpp
@@ -1730,7 +1730,7 @@ void UObjectReplicationBridge::SetNetPushIdOnInstance(UE::Net::FReplicationInsta
 	using namespace UE::Net;
 
 	// Set push ID only if any state supports it. If no state supports it then we might crash if setting the ID.
-	if (EnumHasAnyFlags(InstanceProtocol->InstanceTraits, EReplicationInstanceProtocolTraits::HasPartialPushBasedDirtiness | EReplicationInstanceProtocolTraits::HasFullPushBasedDirtiness))
+	if (EnumHasAnyFlags(InstanceProtocol->InstanceTraits, EReplicationInstanceProtocolTraits::HasPushBasedDirtiness))
 	{
 		const Private::FNetPushObjectHandle PushHandle(NetHandle);
 		TArrayView<const FReplicationFragment* const> Fragments(InstanceProtocol->Fragments, InstanceProtocol->FragmentCount);
diff --git a/Engine/Source/Runtime/Experimental/Iris/Core/Private/Iris/ReplicationSystem/Polling/ObjectPoller.cpp b/Engine/Source/Runtime/Experimental/Iris/Core/Private/Iris/ReplicationSystem/Polling/ObjectPoller.cpp
--- a/Engine/Source/Runtime/Experimental/Iris/Core/Private/Iris/ReplicationSystem/Polling/ObjectPoller.cpp
+++ b/Engine/Source/Runtime/Experimental/Iris/Core/Private/Iris/ReplicationSystem/Polling/ObjectPoller.cpp
@@ -17,12 +17,28 @@
 #include "Iris/Core/IrisProfiler.h"
 #include "Iris/Stats/NetStatsContext.h"
 
+#include "Net/Core/DirtyNetObjectTracker/GlobalDirtyNetObjectTracker.h"
 #include "Net/Core/Trace/NetTrace.h"
 #include "Net/Core/Trace/NetDebugName.h"
 #include "HAL/IConsoleManager.h"
 
 namespace UE::Net::Private
 {
+namespace CVars
+{
+	static bool bFilterOutNonDirtyPushBasedObjects = true;
+	static FAutoConsoleVariableRef CVarFilterOutNonDirtyPushBasedObjects(
+		TEXT("net.Iris.Poll.FilterOutNonDirtyPushBasedObjects"),
+		bFilterOutNonDirtyPushBasedObjects,
+		TEXT("When true fully push based objects that are not considered to require polling are masked off before the poll loop.")
+	);
+
+	static bool bEnableVerbosePollProfiling = false;
+	static FAutoConsoleVariableRef CVarEnableVerbosePollProfiling(
+		TEXT("net.Iris.Poll.EnableVerboseProfiling"),
+		bEnableVerbosePollProfiling,
+		TEXT("Enable conditional profiler scopes for polling. Useful to see exactly what properties gets polled in cpu profiler..")
+	);
 
 	static int32 GNumPollingTasks = 8;
 	static FAutoConsoleVariableRef CvarNumPollingTasks(
@@ -37,6 +53,7 @@ namespace UE::Net::Private
 		GNumWordsPerChunk,
 		TEXT("Number of words in a row to interleave before queuing to the next task")
 	);
+}
 
 	struct FReplicationPollTaskData
 	{
@@ -91,6 +108,9 @@ namespace UE::Net::Private
 		{
 			check(ReplicationSystemInternal->GetIsInParallelPhase());
 
+			IRIS_PROFILER_SCOPE_CONDITIONAL(PollAndCopyPushBased, PollTaskData.bUsePushModel);
+			IRIS_PROFILER_SCOPE_CONDITIONAL(ForcePollAndCopy, !PollTaskData.bUsePushModel);
+
 			FNetStatsContext* ChildNetStatsContext = PollTaskData.NetTypeStats->AcquireChildNetStatsContext();
 
 			//Set up the lambda to use for polling, selecting a different one for PushModel or ForcePoll
@@ -154,24 +174,57 @@ FObjectPoller::FObjectPoller(const FInitParams& InitParams)
 	GarbageCollectionAffectedObjects = MakeNetBitArrayView(ObjectReplicationBridge->GarbageCollectionAffectedObjects);
 
 	// DirtyObjectsThisFrame is acquired only during polling 
+	bUsePerPropertyDirtyTracking = FGlobalDirtyNetObjectTracker::IsUsingPerPropertyDirtyTracking();
 }
 
-void FObjectPoller::PollAndCopyObjects(const FNetBitArrayView& ObjectsConsideredForPolling)
+void FObjectPoller::PollAndCopyObjects(const FNetBitArrayView& InObjectsConsideredForPolling)
 {
 	FDirtyObjectsAccessor DirtyObjectsAccessor(ReplicationSystemInternal->GetDirtyNetObjectTracker());
 	DirtyObjectsThisFrame = DirtyObjectsAccessor.GetDirtyNetObjects();
 
 	NetStatsContext = ReplicationSystemInternal->GetNetTypeStats().GetNetStatsContext();
 
-	IRIS_PROFILER_SCOPE_VERBOSE_CONDITIONAL(PollAndCopyPushBased, IsIrisPushModelEnabled());
-	IRIS_PROFILER_SCOPE_VERBOSE_CONDITIONAL(ForcePollAndCopy, !IsIrisPushModelEnabled());
+	const bool bIsIrisPushModelEnabled = IsIrisPushModelEnabled();
 	
+	FNetBitArray TempObjectsToPoll;
+	
+	// Prepare bitarray with objects to poll, filtering out all fully push based that we can skip.
+	const bool bUseFilteredObjectsConsideredForPolling = CVars::bFilterOutNonDirtyPushBasedObjects && bIsIrisPushModelEnabled;
+	if (bUseFilteredObjectsConsideredForPolling)
+	{
+		IRIS_PROFILER_SCOPE(FilterObjectsConsideredForPolling);
+
+		using StorageWordType = UE::Net::FNetBitArrayBase::StorageWordType;
+
+		TempObjectsToPoll.Init(InObjectsConsideredForPolling.GetNumBits());
+
+		const FNetBitArrayView ObjectsWithFullPushBasedDirtiness = LocalNetRefHandleManager.GetObjectsWithFullPushBasedDirtiness();
+
+		const StorageWordType* ObjectsConsideredForPollingData = InObjectsConsideredForPolling.GetData();
+		const StorageWordType* ObjectsWithFullPushBasedDirtinessData = ObjectsWithFullPushBasedDirtiness.GetData();
+		const StorageWordType* DirtyObjectsThisFrameData = DirtyObjectsThisFrame.GetData();
+		const StorageWordType* AccumulatedDirtyObjectsData = AccumulatedDirtyObjects.GetData();
+		const StorageWordType* GarbageCollectionAffectedObjectsData = GarbageCollectionAffectedObjects.GetData();
+		StorageWordType* TempObjectsToPollData = TempObjectsToPoll.GetData();
+
+		// Build a bitarray for all objects that we need to poll, masking off all objects we can skip (fully push based, not dirty, and not affected by GC).
+		// Assumes that "new"-objects are properly marked as dirty.
+		for (uint32 WordIt = 0U; WordIt < InObjectsConsideredForPolling.GetNumWords(); ++WordIt)
+		{
+			const FNetBitArrayBase::StorageWordType SkippableObjects = ObjectsWithFullPushBasedDirtinessData[WordIt] & ~(DirtyObjectsThisFrameData[WordIt] | AccumulatedDirtyObjectsData[WordIt] | GarbageCollectionAffectedObjectsData[WordIt]);
+			TempObjectsToPollData[WordIt] = ObjectsConsideredForPollingData[WordIt] & ~SkippableObjects;
+		}
+	}
+
+	// Pick filtered or non-filtered polling list.
+	const FNetBitArrayView ObjectsConsideredForPolling = bUseFilteredObjectsConsideredForPolling ? UE::Net::MakeNetBitArrayView(TempObjectsToPoll, FNetBitArrayBase::NoResetNoValidate) : InObjectsConsideredForPolling;
+
 	// We need to split the work somewhat evenly across multiple tasks, so we define a Chunk as
 	// a number of contiguous words (each 32-bit word represents 32 objects), so 16 words per chunk 
 	// is 64-bytes and should fill out a whole cache line on standard platforms.
 	// The Chunk size must be small enough to result in at least one chunk per task, otherwise fall back to synchronous
-	const uint32 NumPollingTasks = GNumPollingTasks;
-	const uint32 BitsPerChunk = GNumWordsPerChunk * FNetBitArray::WordBitCount;
+	const uint32 NumPollingTasks = CVars::GNumPollingTasks;
+	const uint32 BitsPerChunk = CVars::GNumWordsPerChunk * FNetBitArray::WordBitCount;
 	const bool bEachTaskHasAChunkOfWork = (BitsPerChunk * NumPollingTasks) < ObjectsConsideredForPolling.GetNumBits();
 	
 	if (ReplicationSystemInternal->AreParallelTasksAllowed() && bEachTaskHasAChunkOfWork)
@@ -184,10 +237,10 @@ void FObjectPoller::PollAndCopyObjects(const FNetBitArrayView& ObjectsConsidered
 
 		for (uint32 CurrentTaskIndex = 0; CurrentTaskIndex < NumPollingTasks; CurrentTaskIndex++)
 		{
-			IRIS_PROFILER_SCOPE(UObjectPoller_PrepareAndDispatchTask);
+			IRIS_PROFILER_SCOPE(PrepareAndDispatchTask);
 			
 			FReplicationPollTaskData& Data = TaskData[CurrentTaskIndex];
-			Data.bUsePushModel = IsIrisPushModelEnabled();//If pushmodel is disabled, call ForcePollObject instead
+			Data.bUsePushModel = bIsIrisPushModelEnabled;//If pushmodel is disabled, call ForcePollObject instead
 			Data.NetTypeStats = &ReplicationSystemInternal->GetNetTypeStats();
 			Data.CurrentTaskIndex = CurrentTaskIndex;
 			Data.BitsPerChunk = BitsPerChunk;
@@ -215,8 +268,10 @@ void FObjectPoller::PollAndCopyObjects(const FNetBitArrayView& ObjectsConsidered
 	else
 	{
 		//Synchronous, non-task based polling of all objects considered for polling
-		if(IsIrisPushModelEnabled())
+		if (IsIrisPushModelEnabled())
 		{
+			IRIS_PROFILER_SCOPE(PollAndCopyPushBased);
+
 			ObjectsConsideredForPolling.ForAllSetBits([this](FInternalNetRefIndex Objectindex)
 			{
 				PushModelPollObject(Objectindex, NetStatsContext, PollStats);
@@ -224,6 +279,8 @@ void FObjectPoller::PollAndCopyObjects(const FNetBitArrayView& ObjectsConsidered
 		}
 		else
 		{
+			IRIS_PROFILER_SCOPE(ForcePollAndCopy);
+
 			ObjectsConsideredForPolling.ForAllSetBits([this](FInternalNetRefIndex Objectindex)
 			{
 				ForcePollObject(Objectindex, NetStatsContext, PollStats);
@@ -241,6 +298,8 @@ void FObjectPoller::PollAndCopySingleObject(FInternalNetRefIndex ObjectIndex)
 
 	//We don't run this PollAndCopySingleObject during the parallel Poll phase, so we can use the parent context NetStatsContext and PollStats
 	check(!ReplicationSystemInternal->GetIsInParallelPhase());
+
+	IRIS_PROFILER_SCOPE_VERBOSE(ForcePollAndCopy);
 	ForcePollObject(ObjectIndex, NetStatsContext, PollStats);
 
 	// Clear ref to locked dirty bit array
@@ -256,15 +315,13 @@ void FObjectPoller::ForcePollObject(FInternalNetRefIndex ObjectIndex, FNetStatsC
 		return;
 	}
 
-	IRIS_PROFILER_PROTOCOL_NAME(ObjectData.Protocol->DebugName->Name);
-
 	// We always poll all states here.
 	ObjectData.bWantsFullPoll = 0U;
 
 	// Poll properties if the instance protocol requires it
 	if (EnumHasAnyFlags(ObjectData.InstanceProtocol->InstanceTraits, EReplicationInstanceProtocolTraits::NeedsPoll))
 	{
-		IRIS_PROFILER_SCOPE_VERBOSE(PollAndCopy);
+		IRIS_PROFILER_PROTOCOL_NAME_CONDITIONAL(ObjectData.Protocol->DebugName->Name, CVars::bEnableVerbosePollProfiling);
 		UE_NET_IRIS_STATS_TIMER(Timer, InNetStatsContext);
 		UE_NET_TRACE_POLL_OBJECT_SCOPE(ObjectData.RefHandle, Timer);
 
@@ -273,6 +330,7 @@ void FObjectPoller::ForcePollObject(FInternalNetRefIndex ObjectIndex, FNetStatsC
 
 		// If this object has been around for a garbage collect and it has object references we must make sure that we update all cached object references
 		EReplicationFragmentPollFlags PollOptions = EReplicationFragmentPollFlags::PollAllState;
+		PollOptions |= CVars::bEnableVerbosePollProfiling ? EReplicationFragmentPollFlags::EnableVerboseProfiling : EReplicationFragmentPollFlags::None;
 		PollOptions |= bIsGCAffectedObject ? EReplicationFragmentPollFlags::ForceRefreshCachedObjectReferencesAfterGC : EReplicationFragmentPollFlags::None;
 
 		const bool bWasAlreadyDirty = DirtyObjectsThisFrame.IsBitSet(ObjectIndex);
@@ -301,15 +359,12 @@ void FObjectPoller::ForcePollObject(FInternalNetRefIndex ObjectIndex, FNetStatsC
 void FObjectPoller::PushModelPollObject(FInternalNetRefIndex ObjectIndex, FNetStatsContext* InNetStatsContext, FPreUpdateAndPollStats& InPollStats)
 {
 	FNetRefHandleManager::FReplicatedObjectData& ObjectData = LocalNetRefHandleManager.GetReplicatedObjectDataNoCheck(ObjectIndex);
-	if (UNLIKELY(ObjectData.InstanceProtocol == nullptr))
+	const FReplicationInstanceProtocol* InstanceProtocol = ObjectData.InstanceProtocol;
+	if (UNLIKELY(InstanceProtocol == nullptr))
 	{
 		return;
 	}
 
-	IRIS_PROFILER_PROTOCOL_NAME(ObjectData.Protocol->DebugName->Name);
-
-	const FReplicationInstanceProtocol* InstanceProtocol = ObjectData.InstanceProtocol;
-
 	const EReplicationInstanceProtocolTraits InstanceTraits = InstanceProtocol->InstanceTraits;
 	const bool bNeedsPoll = EnumHasAnyFlags(InstanceTraits, EReplicationInstanceProtocolTraits::NeedsPoll);
 
@@ -330,10 +385,6 @@ void FObjectPoller::PushModelPollObject(FInternalNetRefIndex ObjectIndex, FNetSt
 		return;
 	}
 
-	IRIS_PROFILER_SCOPE_VERBOSE(PollPushBased);
-	UE_NET_IRIS_STATS_TIMER(Timer, InNetStatsContext);
-	UE_NET_TRACE_POLL_OBJECT_SCOPE(ObjectData.RefHandle, Timer);
-
 	// Does the object need to poll all states once.
 	const bool bWantsFullPoll = ObjectData.bWantsFullPoll;
 	ObjectData.bWantsFullPoll = 0U;
@@ -344,55 +395,79 @@ void FObjectPoller::PushModelPollObject(FInternalNetRefIndex ObjectIndex, FNetSt
 	{
 		if (bIsDirtyObject || bWantsFullPoll)
 		{
+			IRIS_PROFILER_PROTOCOL_NAME_CONDITIONAL(ObjectData.Protocol->DebugName->Name, CVars::bEnableVerbosePollProfiling);
+			UE_NET_IRIS_STATS_TIMER(Timer, NetStatsContext);
+			UE_NET_TRACE_POLL_OBJECT_SCOPE(ObjectData.RefHandle, Timer);
+
 			// We need to do a poll if object is marked as dirty
-			EReplicationFragmentPollFlags PollOptions = EReplicationFragmentPollFlags::PollAllState;
+			EReplicationFragmentPollFlags PollOptions = bUsePerPropertyDirtyTracking && !bWantsFullPoll ? EReplicationFragmentPollFlags::PollDirtyState : EReplicationFragmentPollFlags::PollAllState;
+			PollOptions |= CVars::bEnableVerbosePollProfiling ? EReplicationFragmentPollFlags::EnableVerboseProfiling : EReplicationFragmentPollFlags::None;
 			PollOptions |= bIsGCAffectedObject ? EReplicationFragmentPollFlags::ForceRefreshCachedObjectReferencesAfterGC : EReplicationFragmentPollFlags::None;
 			bPollFoundDirty = FReplicationInstanceOperations::PollAndCopyPropertyData(InstanceProtocol, EReplicationFragmentTraits::None, PollOptions);
 			++InPollStats.PolledObjectCount;
+
+			UE_NET_IRIS_STATS_ADD_TIME_AND_COUNT_FOR_OBJECT(Timer, Poll, ObjectIndex);
 		}
 		else if (bIsGCAffectedObject)
 		{
+			IRIS_PROFILER_PROTOCOL_NAME_CONDITIONAL(ObjectData.Protocol->DebugName->Name, CVars::bEnableVerbosePollProfiling);
+			UE_NET_IRIS_STATS_TIMER(Timer, NetStatsContext);
+			UE_NET_TRACE_POLL_OBJECT_SCOPE(ObjectData.RefHandle, Timer);
+
 			// If this object might have been affected by GC, only refresh cached references
-			const EReplicationFragmentTraits RequiredTraits = EReplicationFragmentTraits::HasPushBasedDirtiness;
-			bPollFoundDirty = FReplicationInstanceOperations::PollAndCopyObjectReferences(InstanceProtocol, RequiredTraits);
+			bPollFoundDirty = FReplicationInstanceOperations::PollAndCopyObjectReferences(InstanceProtocol, EReplicationFragmentTraits::None);
 			++InPollStats.PolledReferencesObjectCount;
+
+			UE_NET_IRIS_STATS_ADD_TIME_AND_COUNT_FOR_OBJECT(Timer, Poll, ObjectIndex);
+		}
+		else
+		{
+			++PollStats.SkippedObjectCount;
 		}
 	}
 	else
 	{
-		// If the object has pushed based fragments, and not is marked dirty and object is affected by GC we need to make sure that we refresh cached references for all push based fragments
-		const bool bIsPushBasedObject = EnumHasAnyFlags(InstanceTraits, EReplicationInstanceProtocolTraits::HasPartialPushBasedDirtiness | EReplicationInstanceProtocolTraits::HasFullPushBasedDirtiness);
+		IRIS_PROFILER_PROTOCOL_NAME_CONDITIONAL(ObjectData.Protocol->DebugName->Name, CVars::bEnableVerbosePollProfiling);
+		UE_NET_IRIS_STATS_TIMER(Timer, NetStatsContext);
+		UE_NET_TRACE_POLL_OBJECT_SCOPE(ObjectData.RefHandle, Timer);
+
+		// If the object has fragments with pushed based properties, and is not marked dirty and object is affected by GC we need to make sure that we refresh cached references for all fragments with push based properties
+		const bool bIsFullPushBasedObject = EnumHasAnyFlags(InstanceTraits, EReplicationInstanceProtocolTraits::HasFullPushBasedDirtiness);
 		const bool bHasObjectReferences = EnumHasAnyFlags(InstanceTraits, EReplicationInstanceProtocolTraits::HasObjectReference);
-		const bool bNeedsRefreshOfCachedObjectReferences = ((!(bWantsFullPoll | bIsDirtyObject)) & bIsGCAffectedObject & bIsPushBasedObject & bHasObjectReferences);
+		const bool bNeedsRefreshOfCachedObjectReferences = ((!(bWantsFullPoll | bIsDirtyObject)) & bIsGCAffectedObject & bIsFullPushBasedObject & bHasObjectReferences);
 		if (bNeedsRefreshOfCachedObjectReferences)
 		{
-			// Only states which has push based dirtiness need to be updated as the other states will be polled in full anyway.
-			const EReplicationFragmentTraits RequiredTraits = EReplicationFragmentTraits::HasPushBasedDirtiness;
+			// Only states which has full push based dirtiness need to be updated as the other states will be at least partially polled anyway.
+			const EReplicationFragmentTraits RequiredTraits = EReplicationFragmentTraits::HasFullPushBasedDirtiness;
 			bPollFoundDirty = FReplicationInstanceOperations::PollAndCopyObjectReferences(InstanceProtocol, RequiredTraits);
 			++InPollStats.PolledReferencesObjectCount;
 		}
 
+		// We currently cannot trust changemask for multi-owner instances so we need to poll.
+		const bool bIsMultiOwnerInstance = EnumHasAnyFlags(InstanceTraits, EReplicationInstanceProtocolTraits::IsMultiObjectInstance);
+
 		// If this object has been around for a garbage collect and it has object references we must make sure that we update all cached object references 
-		EReplicationFragmentPollFlags PollOptions = EReplicationFragmentPollFlags::PollAllState;
+		EReplicationFragmentPollFlags PollOptions = bUsePerPropertyDirtyTracking && !bWantsFullPoll && !bIsMultiOwnerInstance ? EReplicationFragmentPollFlags::PollDirtyState : EReplicationFragmentPollFlags::PollAllState;
+		PollOptions |= CVars::bEnableVerbosePollProfiling ? EReplicationFragmentPollFlags::EnableVerboseProfiling : EReplicationFragmentPollFlags::None;
 		PollOptions |= bIsGCAffectedObject ? EReplicationFragmentPollFlags::ForceRefreshCachedObjectReferencesAfterGC : EReplicationFragmentPollFlags::None;
 
-		// If the object is not new or dirty at this point we only need to poll non-push based fragments as we know that pushed based states have not been modified
-		const EReplicationFragmentTraits ExcludeTraits = (bIsDirtyObject || bWantsFullPoll) ? EReplicationFragmentTraits::None : EReplicationFragmentTraits::HasPushBasedDirtiness;
+		// If the object is not new or dirty at this point we only need to poll non-fully push based fragments as we know that fully pushed based states have not been modified or have already had their 
+		const EReplicationFragmentTraits ExcludeTraits = (bIsDirtyObject || bWantsFullPoll || bIsMultiOwnerInstance) ? EReplicationFragmentTraits::None : EReplicationFragmentTraits::HasFullPushBasedDirtiness;
 		bPollFoundDirty |= FReplicationInstanceOperations::PollAndCopyPropertyData(InstanceProtocol, ExcludeTraits, PollOptions);
 		++InPollStats.PolledObjectCount;
-	}
 
-	if (bIsDirtyObject || bPollFoundDirty)
-	{
-		UE_NET_IRIS_STATS_ADD_TIME_AND_COUNT_FOR_OBJECT(Timer, Poll, ObjectIndex);
+		if (bPollFoundDirty)
+		{
+			UE_NET_IRIS_STATS_ADD_TIME_AND_COUNT_FOR_OBJECT(Timer, Poll, ObjectIndex);
 
-		DirtyObjectsToQuantize.SetBit(ObjectIndex);
-		DirtyObjectsThisFrame.SetBit(ObjectIndex);
-	}
-	else
-	{
-		UE_NET_IRIS_STATS_ADD_TIME_AND_COUNT_FOR_OBJECT_AS_WASTE(Timer, Poll, ObjectIndex);
-		UE_NET_TRACE_POLL_OBJECT_IS_WASTE();
+			DirtyObjectsToQuantize.SetBit(ObjectIndex);
+			DirtyObjectsThisFrame.SetBit(ObjectIndex);
+		}
+		else
+		{
+			UE_NET_IRIS_STATS_ADD_TIME_AND_COUNT_FOR_OBJECT_AS_WASTE(Timer, Poll, ObjectIndex);
+			UE_NET_TRACE_POLL_OBJECT_IS_WASTE();
+		}
 	}
 }
 
diff --git a/Engine/Source/Runtime/Experimental/Iris/Core/Private/Iris/ReplicationSystem/Polling/ObjectPoller.h b/Engine/Source/Runtime/Experimental/Iris/Core/Private/Iris/ReplicationSystem/Polling/ObjectPoller.h
--- a/Engine/Source/Runtime/Experimental/Iris/Core/Private/Iris/ReplicationSystem/Polling/ObjectPoller.h
+++ b/Engine/Source/Runtime/Experimental/Iris/Core/Private/Iris/ReplicationSystem/Polling/ObjectPoller.h
@@ -37,6 +37,7 @@ public:
 	{
 		uint32 PreUpdatedObjectCount = 0;
 		uint32 PolledObjectCount = 0;
+		uint32 SkippedObjectCount = 0;
 		uint32 PolledReferencesObjectCount = 0;
 
 		void Accumulate(const FPreUpdateAndPollStats& StatsToAdd)
@@ -89,6 +90,8 @@ private:
 	FNetBitArrayView GarbageCollectionAffectedObjects;
 
 	FPreUpdateAndPollStats PollStats;
+
+	bool bUsePerPropertyDirtyTracking = false;
 };
 
 } // end namespace UE::Net::Private
diff --git a/Engine/Source/Runtime/Experimental/Iris/Core/Private/Iris/ReplicationSystem/PropertyReplicationFragment.cpp b/Engine/Source/Runtime/Experimental/Iris/Core/Private/Iris/ReplicationSystem/PropertyReplicationFragment.cpp
--- a/Engine/Source/Runtime/Experimental/Iris/Core/Private/Iris/ReplicationSystem/PropertyReplicationFragment.cpp
+++ b/Engine/Source/Runtime/Experimental/Iris/Core/Private/Iris/ReplicationSystem/PropertyReplicationFragment.cpp
@@ -7,6 +7,7 @@
 #include "Iris/Core/IrisProfiler.h"
 #include "Iris/Core/IrisLog.h"
 #include "HAL/IConsoleManager.h"
+#include "Iris/IrisConfigInternal.h"
 
 namespace UE::Net
 {
@@ -64,6 +65,10 @@ FPropertyReplicationFragment::FPropertyReplicationFragment(EReplicationFragmentT
 	if (EnumHasAnyFlags(InDescriptor->Traits, EReplicationStateTraits::HasPushBasedDirtiness))
 	{
 		Traits |= EReplicationFragmentTraits::HasPushBasedDirtiness;
+		if (EnumHasAnyFlags(InDescriptor->Traits, EReplicationStateTraits::HasFullPushBasedDirtiness))
+		{
+			Traits |= EReplicationFragmentTraits::HasFullPushBasedDirtiness;
+		}
 	}
 
 	// Propagate object reference.
@@ -148,20 +153,41 @@ void FPropertyReplicationFragment::ApplyReplicatedState(FReplicationStateApplyCo
 
 bool FPropertyReplicationFragment::PollReplicatedState(EReplicationFragmentPollFlags PollOption)
 {
-	// Since PollPropertyReplicationState always copies the new value we do not need to do anything special
-	// with object references if the EReplicationFragmentPollOption::All flag is set
+	bool bIsDirty = false;
+	const bool bEnableVerboseProfiling = EnumHasAnyFlags(PollOption, EReplicationFragmentPollFlags::EnableVerboseProfiling);
+
 	if (EnumHasAnyFlags(PollOption, EReplicationFragmentPollFlags::PollAllState))
 	{
+		IRIS_PROFILER_SCOPE_CONDITIONAL(PropertyReplicationFragment_PollAllState, bEnableVerboseProfiling);
+
+		// Since PollPropertyReplicationState always copies the new value we do not need to do anything special
+		// with object references if the EReplicationFragmentPollOption::All flag is set
 		const void* ExternalSourceState = static_cast<void*>(Owner);
 		return SrcReplicationState->PollPropertyReplicationState(ExternalSourceState);
 	}
+	else if (EnumHasAnyFlags(PollOption, EReplicationFragmentPollFlags::PollDirtyState))
+	{
+		const bool bForceRefreshObjectReferences = EnumHasAnyFlags(PollOption, EReplicationFragmentPollFlags::ForceRefreshCachedObjectReferencesAfterGC) && ReplicationStateDescriptor->HasObjectReference();
+
+		// We can skip, fully push based states if they are not dirty and we are not refreshing object references
+		if (!EnumHasAnyFlags(ReplicationStateDescriptor->Traits, EReplicationStateTraits::HasFullPushBasedDirtiness) || bForceRefreshObjectReferences || SrcReplicationState->IsDirtyForPolling())
+		{
+			IRIS_PROFILER_SCOPE_CONDITIONAL(PropertyReplicationFragment_PollDirtyState, bEnableVerboseProfiling);
+
+			const void* ExternalSourceState = static_cast<void*>(Owner);
+			const FPropertyReplicationState::FPollParameters Params = {.bEnableVerboseProfiling = bEnableVerboseProfiling, .bForceRefreshObjectReferences = bEnableVerboseProfiling };
+			bIsDirty = SrcReplicationState->PollPushBasedPropertyReplicationState(ExternalSourceState, Params);
+		}
+	}
 	else if (EnumHasAnyFlags(PollOption, EReplicationFragmentPollFlags::ForceRefreshCachedObjectReferencesAfterGC))
 	{
+		IRIS_PROFILER_SCOPE_CONDITIONAL(PropertyReplicationFragment_PollObjectReferences, bEnableVerboseProfiling);
+
 		const void* ExternalSourceState = static_cast<void*>(Owner);
 		return SrcReplicationState->PollObjectReferences(ExternalSourceState);
 	}
 
-	return false;
+	return bIsDirty;
 }
 
 void FPropertyReplicationFragment::ReplicatedStateToString(FStringBuilderBase& StringBuilder, FReplicationStateApplyContext& Context, EReplicationStateToStringFlags Flags) const
diff --git a/Engine/Source/Runtime/Experimental/Iris/Core/Private/Iris/ReplicationSystem/ReplicationBridge.cpp b/Engine/Source/Runtime/Experimental/Iris/Core/Private/Iris/ReplicationSystem/ReplicationBridge.cpp
--- a/Engine/Source/Runtime/Experimental/Iris/Core/Private/Iris/ReplicationSystem/ReplicationBridge.cpp
+++ b/Engine/Source/Runtime/Experimental/Iris/Core/Private/Iris/ReplicationSystem/ReplicationBridge.cpp
@@ -1196,7 +1196,7 @@ void UReplicationBridge::ClearNetPushIds(UE::Net::Private::FInternalNetRefIndex
 	const FNetRefHandleManager::FReplicatedObjectData& ObjectData = NetRefHandleManager->GetReplicatedObjectDataNoCheck(InternalReplicationIndex);
 	if (const FReplicationInstanceProtocol* InstanceProtocol = ObjectData.InstanceProtocol)
 	{
-		if (EnumHasAnyFlags(InstanceProtocol->InstanceTraits, EReplicationInstanceProtocolTraits::HasPartialPushBasedDirtiness | EReplicationInstanceProtocolTraits::HasFullPushBasedDirtiness))
+		if (EnumHasAnyFlags(InstanceProtocol->InstanceTraits, EReplicationInstanceProtocolTraits::HasPushBasedDirtiness))
 		{
 			TArrayView<const FReplicationFragment* const> Fragments(InstanceProtocol->Fragments, InstanceProtocol->FragmentCount);
 			ClearNetPushIdOnFragments(Fragments);
diff --git a/Engine/Source/Runtime/Experimental/Iris/Core/Private/Iris/ReplicationSystem/ReplicationOperations.cpp b/Engine/Source/Runtime/Experimental/Iris/Core/Private/Iris/ReplicationSystem/ReplicationOperations.cpp
--- a/Engine/Source/Runtime/Experimental/Iris/Core/Private/Iris/ReplicationSystem/ReplicationOperations.cpp
+++ b/Engine/Source/Runtime/Experimental/Iris/Core/Private/Iris/ReplicationSystem/ReplicationOperations.cpp
@@ -671,17 +671,16 @@ void FReplicationStateOperations::ApplyStruct(FNetSerializationContext& Context,
 bool FReplicationInstanceOperations::PollAndCopyPropertyData(const FReplicationInstanceProtocol* InstanceProtocol, EReplicationFragmentTraits ExcludeTraits, EReplicationFragmentPollFlags PollOptions)
 {
 	bool bIsStateDirty = false;
-	FReplicationFragment* const * Fragments = InstanceProtocol->Fragments;
 
 	// @TODO Cache traits to avoid touching fragments
 	const uint32 FragmentCount = InstanceProtocol->FragmentCount;
 	for (uint32 StateIt = 0; StateIt < FragmentCount; ++StateIt)
 	{
 		// Only poll fragments with NeedsPoll and none of the ExcludedTraits set.
-		const EReplicationFragmentTraits MaskedFragmentTraits = Fragments[StateIt]->GetTraits() & (EReplicationFragmentTraits::NeedsPoll | ExcludeTraits);
+		const EReplicationFragmentTraits MaskedFragmentTraits = InstanceProtocol->FragmentData[StateIt].Traits & (EReplicationFragmentTraits::NeedsPoll | ExcludeTraits);
 		if (MaskedFragmentTraits == EReplicationFragmentTraits::NeedsPoll)
 		{
-			bIsStateDirty |= Fragments[StateIt]->PollReplicatedState(PollOptions);
+			bIsStateDirty |= InstanceProtocol->Fragments[StateIt]->PollReplicatedState(PollOptions);
 		}
 	}
 
@@ -693,19 +692,18 @@ bool FReplicationInstanceOperations::PollAndCopyPropertyData(const FReplicationI
 	return PollAndCopyPropertyData(InstanceProtocol, EReplicationFragmentTraits::None, PollOptions);
 }
 
-bool FReplicationInstanceOperations::PollAndCopyObjectReferences(const FReplicationInstanceProtocol* InstanceProtocol, EReplicationFragmentTraits RequiredTraits)
+bool FReplicationInstanceOperations::PollAndCopyObjectReferences(const FReplicationInstanceProtocol* InstanceProtocol, EReplicationFragmentTraits RequiredTraits, EReplicationFragmentPollFlags PollOptions)
 {
 	bool bIsStateDirty = false;
 	const EReplicationFragmentTraits ExpectedTraits = EReplicationFragmentTraits::HasObjectReference | RequiredTraits;
 
-	FReplicationFragment* const * Fragments = InstanceProtocol->Fragments;
 	const uint32 FragmentCount = InstanceProtocol->FragmentCount;
 	for (uint32 StateIt = 0; StateIt < FragmentCount; ++StateIt)
 	{
-		const EReplicationFragmentTraits MaskedFragmentTraits = Fragments[StateIt]->GetTraits() & ExpectedTraits;
+		const EReplicationFragmentTraits MaskedFragmentTraits = InstanceProtocol->FragmentData[StateIt].Traits & ExpectedTraits;
 		if (MaskedFragmentTraits == ExpectedTraits)
 		{			
-			bIsStateDirty |= Fragments[StateIt]->PollReplicatedState(EReplicationFragmentPollFlags::ForceRefreshCachedObjectReferencesAfterGC);
+			bIsStateDirty |= InstanceProtocol->Fragments[StateIt]->PollReplicatedState(PollOptions | EReplicationFragmentPollFlags::ForceRefreshCachedObjectReferencesAfterGC);
 		}
 	}
 
@@ -815,6 +813,7 @@ void FReplicationInstanceOperations::ResetDirtiness(const FReplicationInstancePr
 	{
 		const FReplicationStateDescriptor* CurrentDescriptor = ReplicationStateDescriptors[StateIt];
 		Private::ResetMemberChangeMasks(FragmentData[StateIt].ExternalSrcBuffer, CurrentDescriptor);
+		Private::ResetMemberPollMask(FragmentData[StateIt].ExternalSrcBuffer, CurrentDescriptor);
 	}
 }
 
diff --git a/Engine/Source/Runtime/Experimental/Iris/Core/Private/Iris/ReplicationSystem/ReplicationProtocol.h b/Engine/Source/Runtime/Experimental/Iris/Core/Private/Iris/ReplicationSystem/ReplicationProtocol.h
--- a/Engine/Source/Runtime/Experimental/Iris/Core/Private/Iris/ReplicationSystem/ReplicationProtocol.h
+++ b/Engine/Source/Runtime/Experimental/Iris/Core/Private/Iris/ReplicationSystem/ReplicationProtocol.h
@@ -3,7 +3,9 @@
 #pragma once
 
 #include "CoreTypes.h"
+#include "Containers/StringFwd.h"
 #include "Misc/EnumClassFlags.h"
+#include "Iris/ReplicationSystem/ReplicationFragment.h"
 
 namespace UE::Net
 {
@@ -23,10 +25,12 @@ enum class EReplicationInstanceProtocolTraits : uint16
 	NeedsLegacyCallbacks		= NeedsPoll << 1,
 	IsBound						= NeedsLegacyCallbacks << 1,
 	NeedsPreSendUpdate			= IsBound << 1,
-	HasPartialPushBasedDirtiness = NeedsPreSendUpdate << 1,
-	HasFullPushBasedDirtiness	 = HasPartialPushBasedDirtiness << 1,
+	HasPushBasedDirtiness		= NeedsPreSendUpdate << 1,
+	HasFullPushBasedDirtiness	= HasPushBasedDirtiness << 1,
+	// Whether this instance contains fragments from multiple owners
+	IsMultiObjectInstance	 = HasFullPushBasedDirtiness << 1,
 	// Whether there's a state that has object references
-	HasObjectReference = HasFullPushBasedDirtiness << 1,
+	HasObjectReference = IsMultiObjectInstance << 1,
 };
 ENUM_CLASS_FLAGS(EReplicationInstanceProtocolTraits);
 
@@ -40,6 +44,7 @@ struct FReplicationInstanceProtocol
 	struct FFragmentData
 	{
 		uint8* ExternalSrcBuffer;
+		EReplicationFragmentTraits Traits;
 	};
 	
 	FFragmentData* FragmentData;
@@ -63,6 +68,10 @@ enum class EReplicationProtocolTraits : uint16
 	HasObjectReference = HasConnectionSpecificSerialization << 1U,
 	// Whether delta compression is supported or not, essentially whether it makes sense to create baselines or not.
 	SupportsDeltaCompression = HasObjectReference << 1U,
+	// If there are some states in the protocol that partially uses pushbased dirtiness
+	HasPushBasedDirtiness = SupportsDeltaCompression << 1U,
+	// If all states uses full pushbased dirtiness
+	HasFullPushBasedDirtiness	 = HasPushBasedDirtiness << 1U,
 };
 ENUM_CLASS_FLAGS(EReplicationProtocolTraits);
 
@@ -72,6 +81,21 @@ ENUM_CLASS_FLAGS(EReplicationProtocolTraits);
  */
 struct FReplicationProtocol
 {
+	struct FRepIndexToFragmentIndex
+	{
+		enum : uint16
+		{
+			InvalidEntry = 65535U,
+		};
+		
+		uint16 FragmentIndex = InvalidEntry;
+	};
+	struct FRepIndexToFragmentIndexTable
+	{
+		const FRepIndexToFragmentIndex* RepIndexToFragmentIndex;
+		uint32 NumEntries;
+	};
+
 	uint32 GetConditionalChangeMaskOffset() const { return InternalChangeMasksOffset; }
 
 	// RefCounting used to track usage of the replication protocol
@@ -105,6 +129,11 @@ struct FReplicationProtocol
 	// Selected traits from the ReplicationStateDescriptors that might be handy to cache in the protocol.
 	EReplicationProtocolTraits ProtocolTraits;
 
+	// Keep a table of repindex -> fragment for each owner with pushbased data.
+	const FRepIndexToFragmentIndexTable* PushModelOwnerRepIndexToFragmentIndexTable;
+	// Owner count
+	uint16 PushModelOwnerCount;
+
 	mutable int32 RefCount;
 };
 
diff --git a/Engine/Source/Runtime/Experimental/Iris/Core/Private/Iris/ReplicationSystem/ReplicationProtocolManager.cpp b/Engine/Source/Runtime/Experimental/Iris/Core/Private/Iris/ReplicationSystem/ReplicationProtocolManager.cpp
--- a/Engine/Source/Runtime/Experimental/Iris/Core/Private/Iris/ReplicationSystem/ReplicationProtocolManager.cpp
+++ b/Engine/Source/Runtime/Experimental/Iris/Core/Private/Iris/ReplicationSystem/ReplicationProtocolManager.cpp
@@ -14,6 +14,7 @@
 #include "Containers/StringFwd.h"
 #include "Logging/LogScopedVerbosityOverride.h"
 #include "Iris/Stats/NetStats.h"
+#include "Iris/ReplicationState/ReplicationStateDescriptor.h"
 
 #ifndef UE_NET_ENABLE_PROTOCOLMANAGER_LOG
 // Default is to enable protocol logs in non-shipping
@@ -38,6 +39,21 @@ static FAutoConsoleVariableRef CVarIrisLogReplicationProtocols(
 	TEXT("If true, log all created replication protocols."
 	));
 
+static void GetInstanceTraits(const FReplicationFragments& Fragments, EReplicationFragmentTraits& OutAccumulatedTraits, EReplicationFragmentTraits& OutSharedTraits, const EReplicationFragmentTraits InObjectTraits)
+{
+	EReplicationFragmentTraits AccumulatedTraits = InObjectTraits;
+	EReplicationFragmentTraits SharedTraits = Fragments.Num() > 0 ? ~EReplicationFragmentTraits::None : EReplicationFragmentTraits::None;
+	for (const FReplicationFragmentInfo& Info : Fragments)
+	{
+		const EReplicationFragmentTraits FragmentTraits = Info.Fragment->GetTraits();
+		AccumulatedTraits |= FragmentTraits;
+		SharedTraits &= FragmentTraits;
+	}
+
+	OutAccumulatedTraits = AccumulatedTraits;
+	OutSharedTraits = SharedTraits;
+}
+
 FReplicationInstanceProtocol* FReplicationProtocolManager::CreateInstanceProtocol(const FReplicationFragments& Fragments, UE::Net::EReplicationFragmentTraits ObjectTraits)
 {
 	const uint32 FragmentCount = Fragments.Num();
@@ -69,20 +85,33 @@ FReplicationInstanceProtocol* FReplicationProtocolManager::CreateInstanceProtoco
 	InstanceProtocol->Fragments = reinterpret_cast<FReplicationFragment* const *>(Buffer + LayoutData.FragmentsSizeAndOffset.Offset);
 	InstanceProtocol->FragmentCount = static_cast<uint16>(FragmentCount);
 
+	// Setup owner collector
+	constexpr uint32 MaxFragmentOwnerCount = 2U;
+	UObject* FragmentOwnersForPushBasedFragments[MaxFragmentOwnerCount] = {};
+	FReplicationStateOwnerCollector UniquePushBasedOwners(FragmentOwnersForPushBasedFragments, MaxFragmentOwnerCount);
+
 	// Fill in fragment data and fragment pointer
-	uint32 FragmentIt = 0u;
-	EReplicationFragmentTraits AccumulatedTraits = ObjectTraits;
-	EReplicationFragmentTraits SharedTraits = FragmentCount > 0 ? ~EReplicationFragmentTraits::None : EReplicationFragmentTraits::None;
+	uint32 FragmentIt = 0U;
+
 	for (const FReplicationFragmentInfo& Info : Fragments)
 	{
 		InstanceProtocol->FragmentData[FragmentIt].ExternalSrcBuffer = reinterpret_cast<uint8*>(Info.SrcReplicationStateBuffer);
+		InstanceProtocol->FragmentData[FragmentIt].Traits = Info.Fragment->GetTraits();
 		const_cast<FReplicationFragment**>(InstanceProtocol->Fragments)[FragmentIt] = Info.Fragment;
-		const EReplicationFragmentTraits FragmentTraits = Info.Fragment->GetTraits();
-		AccumulatedTraits |= FragmentTraits;
-		SharedTraits &= FragmentTraits;
+
+		// We collect unique owners with properties that are pushbased
+		if (Info.Descriptor->MemberCount > 0 && EnumHasAnyFlags(Info.Fragment->GetTraits(), EReplicationFragmentTraits::HasPushBasedDirtiness))
+		{
+			Info.Fragment->CollectOwner(&UniquePushBasedOwners);		
+		}
 		++FragmentIt;
 	}
 
+	// Get accumulated and shared traits
+	EReplicationFragmentTraits AccumulatedTraits;
+	EReplicationFragmentTraits SharedTraits;
+	GetInstanceTraits(Fragments, AccumulatedTraits, SharedTraits, ObjectTraits);
+
 	// Init instance traits
 	EReplicationInstanceProtocolTraits InstanceTraits = EReplicationInstanceProtocolTraits::None;
 	if (EnumHasAnyFlags(AccumulatedTraits, EReplicationFragmentTraits::NeedsPoll))
@@ -101,10 +130,19 @@ FReplicationInstanceProtocol* FReplicationProtocolManager::CreateInstanceProtoco
 	{
 		InstanceTraits |= EReplicationInstanceProtocolTraits::HasObjectReference;
 	}
-	// Push based dirtiness?
-	if (EnumHasAnyFlags(AccumulatedTraits, EReplicationFragmentTraits::HasPushBasedDirtiness))
+	// Currently we do not support push based dirtiness for protocols with multiple pushbased fragments with different owners.
+	if (UniquePushBasedOwners.GetOwnerCount() > 1U)
 	{
-		InstanceTraits |= (EnumHasAnyFlags(SharedTraits, EReplicationFragmentTraits::HasPushBasedDirtiness) ? EReplicationInstanceProtocolTraits::HasFullPushBasedDirtiness : EReplicationInstanceProtocolTraits::HasPartialPushBasedDirtiness);
+		InstanceTraits |= EReplicationInstanceProtocolTraits::IsMultiObjectInstance;
+		UE_LOG(LogIris, Warning, TEXT("Currently we do not support pushbased dirtiness for multiowner instances"));
+	}
+	else if (EnumHasAnyFlags(AccumulatedTraits, EReplicationFragmentTraits::HasPushBasedDirtiness))
+	{
+		InstanceTraits |= EReplicationInstanceProtocolTraits::HasPushBasedDirtiness;
+		if (EnumHasAnyFlags(SharedTraits, EReplicationFragmentTraits::HasFullPushBasedDirtiness))
+		{
+			InstanceTraits |= EReplicationInstanceProtocolTraits::HasFullPushBasedDirtiness;
+		}
 	}
 
 	InstanceProtocol->InstanceTraits = InstanceTraits;
@@ -245,11 +283,56 @@ const FReplicationProtocol* FReplicationProtocolManager::CreateReplicationProtoc
 		return nullptr;
 	}
 
+	// Build table to map from RepIndex to fragment
+	TArray<TArray<FReplicationProtocol::FRepIndexToFragmentIndex, TInlineAllocator<128>>> OwnerRepIndicesToFragment;
+	{
+		// $TODO: Remove when we change FReplicationFragment::GetOwner()
+		auto GetFragmentOwner = [](const FReplicationFragment* Fragment)
+		{
+			constexpr uint32 MaxFragmentOwnerCount = 2U;
+			UObject* FragmentOwnersForPushBasedFragments[MaxFragmentOwnerCount];
+			FReplicationStateOwnerCollector OwnerCollector(FragmentOwnersForPushBasedFragments, MaxFragmentOwnerCount);
+			Fragment->CollectOwner(&OwnerCollector);
+
+			ensure(OwnerCollector.GetOwnerCount() > 0U);
+			return OwnerCollector.GetOwnerCount() > 0U ? OwnerCollector.GetOwners()[0U] : nullptr;
+		};
+		
+		// Fill in fragment data and fragment pointer
+		TArray<const UObject*, TInlineAllocator<16>> Owners;
+
+		for (const FReplicationFragmentInfo& Info : Fragments)
+		{
+			const uint16 FragmentIndex = static_cast<uint16>(&Info - Fragments.GetData());
+			if (Info.Descriptor->MemberCount > 0U && EnumHasAnyFlags(Info.Fragment->GetTraits(), EReplicationFragmentTraits::HasPushBasedDirtiness))
+			{
+				const UObject* Owner = GetFragmentOwner(Info.Fragment);
+				const int32 OwnerId = Owners.AddUnique(Owner);
+
+				// Resize if necessary
+				OwnerRepIndicesToFragment.SetNum(Owners.Num());
+				OwnerRepIndicesToFragment[OwnerId].SetNum(FPlatformMath::Max(Info.Descriptor->RepIndexCount, (uint16)OwnerRepIndicesToFragment[OwnerId].Num()));
+
+				for (const FReplicationStateMemberRepIndexToMemberIndexDescriptor& MemberIndex : MakeArrayView(Info.Descriptor->MemberRepIndexToMemberIndexDescriptors, Info.Descriptor->RepIndexCount))
+				{
+					const int32 RepIndex = static_cast<int32>(&MemberIndex - Info.Descriptor->MemberRepIndexToMemberIndexDescriptors);
+					if (MemberIndex.MemberIndex != FReplicationStateMemberRepIndexToMemberIndexDescriptor::InvalidEntry)
+					{
+						//UE_LOG(LogIris, VeryVerbose, TEXT("Mapping Property: %s:%s repindex %d memberindex %d to fragment %d"), *GetNameSafe(Owner), Info.Descriptor->MemberDebugDescriptors[MemberIndex.MemberIndex].DebugName->Name, RepIndex, MemberIndex.MemberIndex, FragmentIndex)					
+						OwnerRepIndicesToFragment[OwnerId][RepIndex].FragmentIndex = (uint16)FragmentIndex;
+					}
+				}
+			}
+		}
+	}
+
 	// We want to keep this as a single allocation so we first build a layout and allocate enough space for the InstanceProtocol and its data
 	struct FReplicationProtocolLayoutData
 	{
 		FMemoryLayoutUtil::FOffsetAndSize ReplicationProtocolSizeAndOffset;
 		FMemoryLayoutUtil::FOffsetAndSize ReplicationStateDescriptorsSizeAndOffset;
+		FMemoryLayoutUtil::FOffsetAndSize PushModelOwnerTableSizeAndOffset;
+		FMemoryLayoutUtil::FOffsetAndSize PushModelOwnerRepIndexToFragmentIndexDataSizeAndOffset;
 	};
 	FReplicationProtocolLayoutData LayoutData;
 
@@ -257,6 +340,20 @@ const FReplicationProtocol* FReplicationProtocolManager::CreateReplicationProtoc
 	FMemoryLayoutUtil::AddToLayout<FReplicationProtocol>(Layout, LayoutData.ReplicationProtocolSizeAndOffset, 1);
 	FMemoryLayoutUtil::AddToLayout<const FReplicationStateDescriptor*>(Layout, LayoutData.ReplicationStateDescriptorsSizeAndOffset, FragmentCount);
 
+	// Fill in RepIndex -> Fragmeent lookup tables
+	const uint16 OwnerCount = (uint16)OwnerRepIndicesToFragment.Num();
+
+	uint32 TotalRepIndexToFragmentCount = 0U;
+	for (const TArray<FReplicationProtocol::FRepIndexToFragmentIndex, TInlineAllocator<128>>& TableEntry : OwnerRepIndicesToFragment)
+	{
+		TotalRepIndexToFragmentCount += TableEntry.Num();
+	}
+
+	// Add lookup table to layout
+	FMemoryLayoutUtil::AddToLayout<const FReplicationProtocol::FRepIndexToFragmentIndexTable>(Layout, LayoutData.PushModelOwnerTableSizeAndOffset, OwnerCount);
+	// Actual data for all table entries.
+	FMemoryLayoutUtil::AddToLayout<const FReplicationProtocol::FRepIndexToFragmentIndex>(Layout, LayoutData.PushModelOwnerRepIndexToFragmentIndexDataSizeAndOffset, TotalRepIndexToFragmentCount);
+
 	// Allocate memory for the protocol, the replication protocol must be refcounted by all NetObjects
 	// We could also choose to explicitly control the lifetime 
 	uint8* Buffer = static_cast<uint8*>(FMemory::MallocZeroed(Layout.CurrentOffset, static_cast<uint32>(Layout.MaxAlignment)));
@@ -266,6 +363,30 @@ const FReplicationProtocol* FReplicationProtocolManager::CreateReplicationProtoc
 	Protocol->ReplicationStateDescriptors = FragmentCount ? reinterpret_cast<const FReplicationStateDescriptor**>(Buffer + LayoutData.ReplicationStateDescriptorsSizeAndOffset.Offset) : nullptr;
 	Protocol->ReplicationStateCount = FragmentCount;
 
+	// Fill in the data for the the RepIndex -> Fragment lookup tables.
+	{
+		// Setup table pointer and count
+		FReplicationProtocol::FRepIndexToFragmentIndexTable* RepIndexToFragmentIndexTables = OwnerCount ? reinterpret_cast<FReplicationProtocol::FRepIndexToFragmentIndexTable*>(Buffer + LayoutData.PushModelOwnerTableSizeAndOffset.Offset) : nullptr;
+		Protocol->PushModelOwnerRepIndexToFragmentIndexTable = RepIndexToFragmentIndexTables;
+		Protocol->PushModelOwnerCount = OwnerCount;
+
+		// Setup table data.
+		FReplicationProtocol::FRepIndexToFragmentIndex* PushModelRepIndexToFragmentIndexData = OwnerCount ? reinterpret_cast<FReplicationProtocol::FRepIndexToFragmentIndex*>(Buffer + LayoutData.PushModelOwnerRepIndexToFragmentIndexDataSizeAndOffset.Offset) : nullptr;
+		for (uint32 OwnerIt = 0U; OwnerIt < OwnerCount; ++OwnerIt)
+		{
+			const uint32 RepIndexToFragmentCount = OwnerRepIndicesToFragment[OwnerIt].Num();
+
+			// Copy data
+			FPlatformMemory::Memcpy(PushModelRepIndexToFragmentIndexData, OwnerRepIndicesToFragment[OwnerIt].GetData(), RepIndexToFragmentCount * sizeof(FReplicationProtocol::FRepIndexToFragmentIndex)); 
+
+			// Setup table entry
+			RepIndexToFragmentIndexTables[OwnerIt].RepIndexToFragmentIndex = PushModelRepIndexToFragmentIndexData;
+			RepIndexToFragmentIndexTables[OwnerIt].NumEntries = OwnerRepIndicesToFragment[OwnerIt].Num();
+
+			PushModelRepIndexToFragmentIndexData += RepIndexToFragmentCount;
+		}
+	}
+
 	// Cached data
 	uint32 MaxExternalSize = 0;
 	uint32 MaxExternalAlign = 0;
@@ -275,6 +396,11 @@ const FReplicationProtocol* FReplicationProtocolManager::CreateReplicationProtoc
 	EReplicationProtocolTraits ProtocolTraits = EReplicationProtocolTraits::None;
 	EReplicationStateTraits CombinedStateTraits = EReplicationStateTraits::None;
 
+	// Get accumulated and shared traits for fragments
+	EReplicationFragmentTraits AccumulatedTraits;
+	EReplicationFragmentTraits SharedTraits;
+	GetInstanceTraits(Fragments, AccumulatedTraits, SharedTraits, EReplicationFragmentTraits::None);
+
 	// Fill in data for the protocol
 	uint32 FirstLifetimeConditionalsStateIndex = ~0U;
 	uint32 LifetimeConditionalsStateCount = 0;
@@ -355,6 +481,16 @@ const FReplicationProtocol* FReplicationProtocolManager::CreateReplicationProtoc
 		InternalSize += FNetBitArrayView::CalculateRequiredWordCount(ChangeMaskBitCount)*sizeof(FNetBitArrayView::StorageWordType);
 	}
 
+	// Setup Pushbased traits
+	if (EnumHasAnyFlags(AccumulatedTraits, EReplicationFragmentTraits::HasPushBasedDirtiness))
+	{
+		ProtocolTraits |= EReplicationProtocolTraits::HasPushBasedDirtiness;
+		if (EnumHasAnyFlags(SharedTraits, EReplicationFragmentTraits::HasPushBasedDirtiness))
+		{
+			ProtocolTraits |= EReplicationProtocolTraits::HasFullPushBasedDirtiness;	
+		}
+	}
+
 	Protocol->InternalTotalAlignment = FPlatformMath::Max(InternalAlign, 1U);
 	Protocol->InternalTotalSize = InternalSize;
 	Protocol->MaxExternalStateSize = MaxExternalSize;
diff --git a/Engine/Source/Runtime/Experimental/Iris/Core/Public/Iris/ReplicationState/PropertyReplicationState.h b/Engine/Source/Runtime/Experimental/Iris/Core/Public/Iris/ReplicationState/PropertyReplicationState.h
--- a/Engine/Source/Runtime/Experimental/Iris/Core/Public/Iris/ReplicationState/PropertyReplicationState.h
+++ b/Engine/Source/Runtime/Experimental/Iris/Core/Public/Iris/ReplicationState/PropertyReplicationState.h
@@ -87,6 +87,23 @@ public:
 	*/
 	IRISCORE_API bool PollPropertyReplicationState(const void* RESTRICT SrcData);
 
+	struct FPollParameters
+	{
+		bool bEnableVerboseProfiling = false;
+		bool bForceRefreshObjectReferences = false;
+	};
+
+	/**
+	 * Partial poll src data from properties in SrcData based on the state of the inlined per member pollmask
+	 * non-pushbased properties are always polled, 
+	 * Compare and update representation in DstStateBuffer and update ChangeMask. This is currently used to track and update dirtiness for the properties.
+	 * @SrcData is a UClass/UStruct containing properties.
+	 * @return True if any state is dirty, false if not.
+	 * @todo This requires some thought as the property offset we get in all are expressed from the base, but we do generate descriptors for all parts of the inheritance chain.
+	 * Do we carry around the base pointer or should we recalculate the offsets when we build the descriptor
+	*/
+	IRISCORE_API bool PollPushBasedPropertyReplicationState(const void* RESTRICT SrcData, const FPollParameters& PollParameters);
+
 	/** Certain rep notifies requires the current state to be stored before we overwrite it, this method will copy property values from src data, , where SrcData is a UClass/UStruct containing properties
 	 * if the property is marked as dirty in NewStateToBeApplied
 	 */
@@ -120,7 +137,6 @@ public:
 		bool bOnlyCallIfDiffersFromLocal = false;
 	};
 
-
 	/** Invoke repnotifies for all dirty members */
 	IRISCORE_API void CallRepNotifies(void* RESTRICT DstData, const FCallRepNotifiesParameters& Params) const;
 	inline void CallRepNotifies(void* RESTRICT DstData, const FPropertyReplicationState* PreviousState, bool bIsInit) const;
@@ -157,16 +173,24 @@ public:
 	 */
 	IRISCORE_API bool IsCustomConditionEnabled(uint32 Index) const;
 
+	/**
+	 * Returns true if state is marked dirty for polling
+	 */
+	IRISCORE_API bool IsDirtyForPolling() const;
+
 private:
-	bool IsInitState() const;
+	// Returns true if member should be polled, returns true if memberpollmask is dirty or if changemask is dirty.
+	bool IsDirtyForPolling(uint32 Index) const;
+
+	// Returns true if changemask or header of the state is marked dirty
 	bool IsDirty() const;
+	bool IsInitState() const;
 	void ConstructStateInternal();
 	void DestructStateInternal();
 	void InjectState(const FReplicationStateDescriptor* Descriptor, uint8* InStateBuffer);
 	void PollPropertyValue(uint32 Index, const void* SrcValue);
 	void PushPropertyValue(uint32 Index, void* DstValue) const;
 
-
 	TRefCountPtr<const FReplicationStateDescriptor> ReplicationStateDescriptor;
 	uint8* StateBuffer;
 	uint32 bOwnState : 1;
diff --git a/Engine/Source/Runtime/Experimental/Iris/Core/Public/Iris/ReplicationState/ReplicationStateDescriptor.h b/Engine/Source/Runtime/Experimental/Iris/Core/Public/Iris/ReplicationState/ReplicationStateDescriptor.h
--- a/Engine/Source/Runtime/Experimental/Iris/Core/Public/Iris/ReplicationState/ReplicationStateDescriptor.h
+++ b/Engine/Source/Runtime/Experimental/Iris/Core/Public/Iris/ReplicationState/ReplicationStateDescriptor.h
@@ -48,6 +48,7 @@ enum class EReplicationStateMemberTraits : uint16
 	HasConnectionSpecificSerialization = HasObjectReference << 1U,
 	HasRepNotifyAlways = HasConnectionSpecificSerialization << 1U,
 	UseSerializerIsEqual = HasRepNotifyAlways << 1U,
+	HasPushBasedDirtiness = UseSerializerIsEqual << 1U,
 };
 ENUM_CLASS_FLAGS(EReplicationStateMemberTraits);
 
@@ -176,8 +177,9 @@ enum class EReplicationStateTraits : uint32
 	IsNativeFastArrayReplicationState	= IsFastArrayReplicationState << 1U,
 	HasConnectionSpecificSerialization	= IsNativeFastArrayReplicationState << 1U,
 	HasPushBasedDirtiness				= HasConnectionSpecificSerialization << 1U,
+	HasFullPushBasedDirtiness			= HasPushBasedDirtiness << 1U,
 	// Whether delta compression is supported or not
-	SupportsDeltaCompression			= HasPushBasedDirtiness << 1U,
+	SupportsDeltaCompression			= HasFullPushBasedDirtiness << 1U,
 	UseSerializerIsEqual				= SupportsDeltaCompression << 1U,
 	// Whether this is a descriptor for a struct derived from something with a NetSerializer, but doesn't have a custom NetSerializer itself.
 	IsDerivedFromStructWithCustomSerializer	= UseSerializerIsEqual << 1U,
@@ -204,7 +206,7 @@ struct FReplicationStateDescriptor
 
 	uint32 GetChangeMaskOffset() const { return ChangeMasksExternalOffset; }
 	uint32 GetConditionalChangeMaskOffset() const { return ChangeMasksExternalOffset + 4U*((ChangeMaskBitCount + 31)/32); }
-
+	uint32 GetMemberPollMaskOffset() const { return GetConditionalChangeMaskOffset() + (EnumHasAnyFlags(Traits, EReplicationStateTraits::HasLifetimeConditionals) ? (4U*((ChangeMaskBitCount + 31)/32)) : 0U); }
 
 	const FReplicationStateMemberDescriptor* MemberDescriptors;
 	const FReplicationStateMemberChangeMaskDescriptor* MemberChangeMaskDescriptors;
@@ -252,8 +254,9 @@ struct FReplicationStateDescriptor
 
 	// How many bits do we need for our tracking of dirty changes
 	uint16 ChangeMaskBitCount;
-	// This is the offset to where we store data for changemasks in the external state.
-	// Retrieve offset with helper methods, GetChangeMaskOffset(), GetConditionalChangeMaskOffset().
+
+	// This is the offset to where we store data for ChangeMask, ConditionalChangeMask and MemberPollMask in the external state.
+	// Retrieve offset with helper methods, GetChangeMaskOffset(), GetConditionalChangeMaskOffset(), GetMemberPollMaskOffset()
 	uint32 ChangeMasksExternalOffset;
 
 	// We need to assign a unique key that is stable between server and client (name hash of class + state type for now)
diff --git a/Engine/Source/Runtime/Experimental/Iris/Core/Public/Iris/ReplicationState/ReplicationStateUtil.h b/Engine/Source/Runtime/Experimental/Iris/Core/Public/Iris/ReplicationState/ReplicationStateUtil.h
--- a/Engine/Source/Runtime/Experimental/Iris/Core/Public/Iris/ReplicationState/ReplicationStateUtil.h
+++ b/Engine/Source/Runtime/Experimental/Iris/Core/Public/Iris/ReplicationState/ReplicationStateUtil.h
@@ -53,6 +53,22 @@ inline FNetBitArrayView GetMemberConditionalChangeMask(uint8* StateBuffer, const
 	return FNetBitArrayView((FNetBitArrayView::StorageWordType*)(StateBuffer + Descriptor->GetConditionalChangeMaskOffset()), Descriptor->ChangeMaskBitCount, FNetBitArrayView::NoResetNoValidate);
 }
 
+/**
+ * Get MemberPollMask. Valid if state has members
+ */
+inline FNetBitArrayView GetMemberPollMask(uint8* StateBuffer, const FReplicationStateDescriptor* Descriptor)
+{
+	return FNetBitArrayView((FNetBitArrayView::StorageWordType*)(StateBuffer + Descriptor->GetMemberPollMaskOffset()), Descriptor->MemberCount, FNetBitArrayView::NoResetNoValidate);
+}
+
+/**
+ * Reset MemberPollMask
+ */
+inline void ResetMemberPollMask(uint8* StateBuffer, const FReplicationStateDescriptor* Descriptor)
+{
+	FNetBitArrayView((FNetBitArrayView::StorageWordType*)(StateBuffer + Descriptor->GetMemberPollMaskOffset()), Descriptor->MemberCount, FNetBitArrayView::ResetOnInit);
+}
+
 /**
  * Mark specific member dirty, if this is the first bit marked as dirty in the local MemberChangeMask and the state is bound, mark owning object as dirty as well.
  * Note that all bits described by the ChangeMaskInfo will be dirtied by the call. 
diff --git a/Engine/Source/Runtime/Experimental/Iris/Core/Public/Iris/ReplicationSystem/FastArrayReplicationFragment.h b/Engine/Source/Runtime/Experimental/Iris/Core/Public/Iris/ReplicationSystem/FastArrayReplicationFragment.h
--- a/Engine/Source/Runtime/Experimental/Iris/Core/Public/Iris/ReplicationSystem/FastArrayReplicationFragment.h
+++ b/Engine/Source/Runtime/Experimental/Iris/Core/Public/Iris/ReplicationSystem/FastArrayReplicationFragment.h
@@ -50,6 +50,9 @@ protected:
 	// Returns true if the FastArray is dirty
 	bool IsDirty() const;
 
+	// Returns true if the FastArray should be polled
+	bool IsDirtyForPolling() const;
+
 	// Mark FastArray dirty
 	void MarkDirty();
 
@@ -169,11 +172,20 @@ bool TFastArrayReplicationFragment<FastArrayItemType, FastArrayType>::PollReplic
 		constexpr bool bForceFullCompare = true;
 		return PollAllState(bForceFullCompare);
 	}
-	else
+
+	// We can early out if we are pushbased and not dirty for polling
+	const bool bPoll = EnumHasAnyFlags(PollOption, EReplicationFragmentPollFlags::PollAllState) ||
+		(EnumHasAnyFlags(PollOption, EReplicationFragmentPollFlags::PollDirtyState) && (!EnumHasAnyFlags(EReplicationFragmentTraits::HasPushBasedDirtiness, Traits) || IsDirtyForPolling()));
+
+	if (bPoll)
 	{
 		constexpr bool bForceFullCompare = false;
 		return PollAllState(bForceFullCompare);
 	}
+	else
+	{
+		return IsDirty();
+	}
 }
 
 template <typename FastArrayItemType, typename FastArrayType>
@@ -290,6 +302,14 @@ bool TFastArrayReplicationFragment<FastArrayItemType, FastArrayType>::IsDirty()
 	return Private::FReplicationStateHeaderAccessor::GetIsStateDirty(ReplicationStateHeader);
 }
 
+template <typename FastArrayItemType, typename FastArrayType>
+bool TFastArrayReplicationFragment<FastArrayItemType, FastArrayType>::IsDirtyForPolling() const
+{
+	FReplicationStateHeader& ReplicationStateHeader = Private::GetReplicationStateHeader(ReplicationState->GetStateBuffer(), ReplicationStateDescriptor);
+	bool bIsStateDirty = Private::FReplicationStateHeaderAccessor::GetIsStateDirty(ReplicationStateHeader);
+	return bIsStateDirty || Private::GetMemberPollMask(ReplicationState->GetStateBuffer(), ReplicationStateDescriptor).IsAnyBitSet();
+}
+
 template <typename FastArrayItemType, typename FastArrayType>
 void TFastArrayReplicationFragment<FastArrayItemType, FastArrayType>::MarkDirty()
 {
@@ -376,9 +396,12 @@ template <typename FastArrayItemType, typename FastArrayType, typename PollingPo
 bool TNativeFastArrayReplicationFragment<FastArrayItemType, FastArrayType, PollingPolicyType>::PollReplicatedState(EReplicationFragmentPollFlags PollOption)
 {
 	// We ignore object references polling. Since the source state will have references cleaned up they will be valid once any affected item is dirtied.
-	if (PollOption == EReplicationFragmentPollFlags::PollAllState)
+	if (EnumHasAnyFlags(PollOption, EReplicationFragmentPollFlags::PollAllState | EReplicationFragmentPollFlags::PollDirtyState))
 	{
-		return PollAllState();
+		if (!EnumHasAnyFlags(EReplicationFragmentTraits::HasPushBasedDirtiness, Traits) || IsDirty())
+		{
+			return PollAllState();
+		}
 	}
 
 	return IsDirty();
diff --git a/Engine/Source/Runtime/Experimental/Iris/Core/Public/Iris/ReplicationSystem/ReplicationFragment.h b/Engine/Source/Runtime/Experimental/Iris/Core/Public/Iris/ReplicationSystem/ReplicationFragment.h
--- a/Engine/Source/Runtime/Experimental/Iris/Core/Public/Iris/ReplicationSystem/ReplicationFragment.h
+++ b/Engine/Source/Runtime/Experimental/Iris/Core/Public/Iris/ReplicationSystem/ReplicationFragment.h
@@ -132,8 +132,11 @@ enum class EReplicationFragmentTraits : uint32
 	// Fragment supports push based dirtiness
 	HasPushBasedDirtiness				= NeedsPreSendUpdate << 1,
 
+	// Fragment supports full push based dirtiness
+	HasFullPushBasedDirtiness			= HasPushBasedDirtiness << 1,
+
 	// Fragment is a PropertyReplication, or is using PropertyReplicationState
-	HasPropertyReplicationState			= HasPushBasedDirtiness << 1,
+	HasPropertyReplicationState			= HasFullPushBasedDirtiness << 1,
 
 	// Fragment has object references
 	HasObjectReference					= HasPropertyReplicationState << 1,
@@ -146,10 +149,18 @@ ENUM_CLASS_FLAGS(EReplicationFragmentTraits);
 enum class EReplicationFragmentPollFlags : uint32
 {
 	None = 0,
-	/** If set, we need to refresh all cached object references, typically set after a GC to detect stale references in cached data. */
+
+	/** If set, we need to refresh all cached object references, typically set after a GC to detect stale references in cached data, regardless of if they are dirty or not. */
 	ForceRefreshCachedObjectReferencesAfterGC = 1,
-	/** Normal poll. */
-	PollAllState = ForceRefreshCachedObjectReferencesAfterGC << 1,
+
+	/** Pull members marked for polling or marked as dirty in the changemask, non-push based members will alwasys be considered for polling. Note: if ForceRefreshCachedObjectReferencesAfterGC is set as well, object references must be polled as well. */
+	PollDirtyState = ForceRefreshCachedObjectReferencesAfterGC << 1,
+
+	/** Normal full poll. */
+	PollAllState = PollDirtyState << 1,
+
+	/** Enable Verbose profiling  */
+	EnableVerboseProfiling = PollAllState << 1,
 };
 ENUM_CLASS_FLAGS(EReplicationFragmentPollFlags);
 
diff --git a/Engine/Source/Runtime/Experimental/Iris/Core/Public/Iris/ReplicationSystem/ReplicationOperations.h b/Engine/Source/Runtime/Experimental/Iris/Core/Public/Iris/ReplicationSystem/ReplicationOperations.h
--- a/Engine/Source/Runtime/Experimental/Iris/Core/Public/Iris/ReplicationSystem/ReplicationOperations.h
+++ b/Engine/Source/Runtime/Experimental/Iris/Core/Public/Iris/ReplicationSystem/ReplicationOperations.h
@@ -87,7 +87,7 @@ struct FReplicationInstanceOperations
 	static IRISCORE_API bool PollAndCopyPropertyData(const FReplicationInstanceProtocol* InstanceProtocol, EReplicationFragmentPollFlags PollOptions = EReplicationFragmentPollFlags::PollAllState);
 
 	/** Update object references in fragments that has object references and additional required traits. Returns true if a polled state is dirty. */
-	static IRISCORE_API bool PollAndCopyObjectReferences(const FReplicationInstanceProtocol* InstanceProtocol, EReplicationFragmentTraits RequiredTraits);
+	static IRISCORE_API bool PollAndCopyObjectReferences(const FReplicationInstanceProtocol* InstanceProtocol, EReplicationFragmentTraits RequiredTraits, EReplicationFragmentPollFlags PollOptions = EReplicationFragmentPollFlags::None);
 
 	/**
 	 * Quantize the state for a replicated object with a given InstanceProtocol using the ReplicationProtocol. 
diff --git a/Engine/Source/Runtime/Net/Core/Private/Net/Core/DirtyNetObjectTracker/GlobalDirtyNetObjectTracker.cpp b/Engine/Source/Runtime/Net/Core/Private/Net/Core/DirtyNetObjectTracker/GlobalDirtyNetObjectTracker.cpp
--- a/Engine/Source/Runtime/Net/Core/Private/Net/Core/DirtyNetObjectTracker/GlobalDirtyNetObjectTracker.cpp
+++ b/Engine/Source/Runtime/Net/Core/Private/Net/Core/DirtyNetObjectTracker/GlobalDirtyNetObjectTracker.cpp
@@ -2,17 +2,33 @@
 
 #include "Net/Core/DirtyNetObjectTracker/GlobalDirtyNetObjectTracker.h"
 #include "Containers/Set.h"
+#include "Containers/Map.h"
+#include "Containers/Array.h"
 #include "Misc/CoreDelegates.h"
+#include "HAL/IConsoleManager.h"
 
 namespace UE::Net
 {
 
+namespace Private::CVars
+{
+
+bool bUsePerPropertyDirtyTracking = true;
+static FAutoConsoleVariableRef CVarUsePerPropertyDirtyTracking(
+	TEXT("net.Iris.UsePerPropertyDirtyTracking"),
+	bUsePerPropertyDirtyTracking,
+	TEXT("When true we use per property dirty tracking.")
+);
+
+}
+
 class FGlobalDirtyNetObjectTracker::FPimpl
 {
 private:
 	friend FGlobalDirtyNetObjectTracker;
 
 	static TSet<FNetHandle> EmptyDirtyObjects;
+	static FDirtyHandleAndPropertyMap EmptyDirtyObjectsAndProperties;
 
 private:
 
@@ -38,6 +54,8 @@ private:
 
 	TSet<FNetHandle> DirtyObjects;
 
+	FDirtyHandleAndPropertyMap DirtyObjectsAndProperties;
+
 	FNetBitArray AssignedHandleIndices;
 	FNetBitArray Pollers;
 
@@ -50,16 +68,34 @@ private:
 };
 
 TSet<FNetHandle> FGlobalDirtyNetObjectTracker::FPimpl::EmptyDirtyObjects;
+FGlobalDirtyNetObjectTracker::FDirtyHandleAndPropertyMap FGlobalDirtyNetObjectTracker::FPimpl::EmptyDirtyObjectsAndProperties;
 
 FGlobalDirtyNetObjectTracker::FPimpl* FGlobalDirtyNetObjectTracker::Instance = nullptr;
 
-void FGlobalDirtyNetObjectTracker::MarkNetObjectStateDirty(FNetHandle NetHandle)
+void FGlobalDirtyNetObjectTracker::MarkNetObjectStateDirty(FNetHandle NetHandle, const int32 StartRepIndex, const int32 EndRepIndex)
 {
 	if (Instance && Instance->PollerCount > 0)
 	{
 		if (ensureMsgf(!Instance->bLockDirtyList, TEXT("MarkNetObjectStateDirty was called while the dirty list was set to read-only.")))
 		{
-			Instance->DirtyObjects.Add(NetHandle);
+			if (Private::CVars::bUsePerPropertyDirtyTracking)
+			{
+				FDirtyPropertyStorage& Entry = Instance->DirtyObjectsAndProperties.FindOrAdd(NetHandle);
+				if (EndRepIndex != INDEX_NONE)
+				{
+					const int32 WordCount = (int32)FNetBitArrayView::CalculateRequiredWordCount(EndRepIndex + 1);
+					if (WordCount > Entry.Num())
+					{
+						Entry.SetNumZeroed(WordCount);
+					}
+					FNetBitArrayView BitMask(Entry.GetData(), Entry.Num() * 32U, FNetBitArrayView::NoResetNoValidate);
+					BitMask.SetBits(StartRepIndex, EndRepIndex - StartRepIndex + 1U);
+				}
+			}
+			else
+			{
+				Instance->DirtyObjects.Add(NetHandle);
+			}
 
 			// If there are multiple pollers, any poller that has already gathered this frame will need to gather again
 			// to know about this newly dirtied object, since it may have happened in another replication system's PreUpdate.
@@ -125,6 +161,7 @@ void FGlobalDirtyNetObjectTracker::DestroyPoller(uint32 HandleIndex)
 		if (Instance->PollerCount <= 0)
 		{
 			Instance->DirtyObjects.Reset();
+			Instance->DirtyObjectsAndProperties.Reset();
 			Instance->bLockDirtyList = false;
 		}
 	}
@@ -144,6 +181,20 @@ const TSet<FNetHandle>& FGlobalDirtyNetObjectTracker::GetDirtyNetObjects(const F
 	return FPimpl::EmptyDirtyObjects;
 }
 
+const FGlobalDirtyNetObjectTracker::FDirtyHandleAndPropertyMap& FGlobalDirtyNetObjectTracker::GetDirtyNetObjectsAndProperties(const FPollHandle& Handle)
+{
+	if (Instance && Handle.IsValid())
+	{
+		check(Instance->PollerStatuses[Handle.Index].bIsActive);
+
+		Instance->Pollers.SetBit(Handle.Index);
+		Instance->PollerStatuses[Handle.Index].bNeedsGather = false;
+		return Instance->DirtyObjectsAndProperties;
+	}
+
+	return FPimpl::EmptyDirtyObjectsAndProperties;
+}
+
 void FGlobalDirtyNetObjectTracker::LockDirtyListUntilReset(const FPollHandle& Handle)
 {
 	if (Instance && Handle.IsValid())
@@ -180,6 +231,7 @@ void FGlobalDirtyNetObjectTracker::ResetDirtyNetObjects(const FPollHandle& Handl
 			}
 
 			Instance->DirtyObjects.Reset();
+			Instance->DirtyObjectsAndProperties.Reset();
 		}
 
 		Instance->bLockDirtyList = false;
@@ -200,7 +252,7 @@ bool FGlobalDirtyNetObjectTracker::ResetDirtyNetObjectsIfSinglePoller(const FPol
 			check(Instance->Pollers.IsBitSet(Handle.Index));
 
 			Instance->DirtyObjects.Reset();
-
+			Instance->DirtyObjectsAndProperties.Reset();
 			Instance->Pollers.ClearBit(Handle.Index);
 
 			return true;
@@ -210,6 +262,11 @@ bool FGlobalDirtyNetObjectTracker::ResetDirtyNetObjectsIfSinglePoller(const FPol
 	return false;
 }
 
+bool FGlobalDirtyNetObjectTracker::IsUsingPerPropertyDirtyTracking()
+{
+	return UE::Net::Private::CVars::bUsePerPropertyDirtyTracking;
+}
+
 void FGlobalDirtyNetObjectTracker::Init()
 {
 	checkf(Instance == nullptr, TEXT("%s"), TEXT("Only one FGlobalDirtyNetObjectTracker instance may exist."));
diff --git a/Engine/Source/Runtime/Net/Core/Public/Net/Core/DirtyNetObjectTracker/GlobalDirtyNetObjectTracker.h b/Engine/Source/Runtime/Net/Core/Public/Net/Core/DirtyNetObjectTracker/GlobalDirtyNetObjectTracker.h
--- a/Engine/Source/Runtime/Net/Core/Public/Net/Core/DirtyNetObjectTracker/GlobalDirtyNetObjectTracker.h
+++ b/Engine/Source/Runtime/Net/Core/Public/Net/Core/DirtyNetObjectTracker/GlobalDirtyNetObjectTracker.h
@@ -11,7 +11,7 @@ class FNetCoreModule;
 namespace UE::Net
 {
 
-void MarkNetObjectStateDirty(FNetHandle NetHandle);
+void MarkNetObjectStateDirty(FNetHandle NetHandle, int32 StartRepIndex, int32 EndRepIndex);
 
 class FGlobalDirtyNetObjectTracker
 {
@@ -46,9 +46,21 @@ public:
 		uint32 Index = InvalidIndex;
 	};
 
+	using FDirtyPropertyStorage = TArray<FNetBitArrayBase::StorageWordType, TInlineAllocator<4>>;
+	using FDirtyHandleAndPropertyMap = TMap<FNetHandle, FDirtyPropertyStorage>;
+
 public:
+	/** Returns true if we are tracking dirtiness per property. */
+	NETCORE_API static bool IsUsingPerPropertyDirtyTracking();
+
 	/** Marks an object as dirty. Assumes Init() has been called. */
-	NETCORE_API static void MarkNetObjectStateDirty(FNetHandle);
+	NETCORE_API static void MarkNetObjectStateDirty(FNetHandle, const int32, const int32);
+
+	/** Marks an object as dirty. Assumes Init() has been called. */
+	inline static void MarkNetObjectStateDirty(FNetHandle Handle) 
+	{
+		MarkNetObjectStateDirty(Handle, INDEX_NONE, INDEX_NONE);
+	}
 
 	/** Delegate called on pollers that haven't gathered dirty state before it resets */
 	DECLARE_DELEGATE(FPreResetDelegate);
@@ -63,6 +75,13 @@ public:
 	 */
 	NETCORE_API static const TSet<FNetHandle>& GetDirtyNetObjects(const FPollHandle& Handle);
 
+	/**
+	 * Returns the objects with dirty state including per property dirtiness. The reference can only be assumed to be valid until the next call to ResetDirtyNetObjects().
+	 * A matching call to ResetDirtyNetObjects() is required but should not be called until all pollers have had the chance to call GetDirtyNetObjects().
+	 * Assumes Init() has been called.
+	 */
+	NETCORE_API static const FDirtyHandleAndPropertyMap& GetDirtyNetObjectsAndProperties(const FPollHandle& Handle);
+
 	/**
 	 * Detect any updates to the dirty list until all pollers have gathered and reset the list.
 	 */
@@ -100,7 +119,12 @@ private:
 
 inline void MarkNetObjectStateDirty(FNetHandle NetHandle)
 {
-	FGlobalDirtyNetObjectTracker::MarkNetObjectStateDirty(NetHandle);
+	FGlobalDirtyNetObjectTracker::MarkNetObjectStateDirty(NetHandle, INDEX_NONE, INDEX_NONE);
+}
+
+inline void MarkNetObjectStateDirty(FNetHandle NetHandle, int32 StartRepIndex, int32 EndRepIndex)
+{
+	FGlobalDirtyNetObjectTracker::MarkNetObjectStateDirty(NetHandle, StartRepIndex, EndRepIndex);
 }
 
 inline FGlobalDirtyNetObjectTracker::FPollHandle::FPollHandle(uint32 InIndex)
