Corrected caching proxy header check for cache hit/miss.

#rnx

diff --git a/Engine/Source/Runtime/Experimental/IoStore/OnDemand/Private/OnDemandHttpClient.cpp b/Engine/Source/Runtime/Experimental/IoStore/OnDemand/Private/OnDemandHttpClient.cpp
--- a/Engine/Source/Runtime/Experimental/IoStore/OnDemand/Private/OnDemandHttpClient.cpp
+++ b/Engine/Source/Runtime/Experimental/IoStore/OnDemand/Private/OnDemandHttpClient.cpp
@@ -41,20 +41,22 @@ static FIoBuffer JoinIoBuffers(const FIoBuffer& LHS, const FIoBuffer& RHS)
 
 static int8 TrackCdnCacheStats(const HTTP::FResponse& Response)
 {
+	int8 Result = -1;
+
 #if !UE_BUILD_TEST
-	for (FAnsiStringView Header : { "X-Cache", "CF-Cache-Status" })
+	// All header fields are considered, with the assumption that later fields are
+	// more accurate than prior ones (e.g. a caching proxy sits infront if CDNs)
+	Response.ReadHeaders([&Result] (FAnsiStringView Key, FAnsiStringView Value)
 	{
-		FAnsiStringView Value = Response.GetHeader(Header);
-		if (Value.IsEmpty())
+		if (Key == "CF-Cache-Status" || (Key.StartsWith("X-") && Key.EndsWith("-Cache")))
 		{
-			continue;
+			Result = Value.Find("HIT", 0, ESearchCase::IgnoreCase) >= 0 ? 1 : 0;
 		}
-
-		return Value.Find("HIT", 0, ESearchCase::IgnoreCase) >= 0 ? 1 : 0;
-	}
+		return true;
+	});
 #endif
 
-	return -1;
+	return Result;
 }
 
 static const TCHAR* CDNCacheStatusToString(int8 Status)
