Surface IAS cache boot time to on-screen display.

#rnx

diff --git a/Engine/Source/Runtime/Experimental/IoStore/OnDemand/Private/JournaledCache.cpp b/Engine/Source/Runtime/Experimental/IoStore/OnDemand/Private/JournaledCache.cpp
--- a/Engine/Source/Runtime/Experimental/IoStore/OnDemand/Private/JournaledCache.cpp
+++ b/Engine/Source/Runtime/Experimental/IoStore/OnDemand/Private/JournaledCache.cpp
@@ -1020,6 +1020,14 @@ static int32 LoadCache(FDiskCache& DiskCache)
 	}
 
 	UE_LOG(LogIas, VeryVerbose, TEXT("JournaledCache: %u byte journal found"), DataSize);
+	
+	uint64 BootTimestamp = FPlatformTime::Cycles64();
+	ON_SCOPE_EXIT
+	{
+		BootTimestamp = FPlatformTime::Cycles64() - BootTimestamp;
+		uint64 TimeMs = uint64(FPlatformTime::ToMilliseconds64(BootTimestamp));
+		FOnDemandIoBackendStats::Get()->OnCacheBootMs(TimeMs);
+	};
 
 	struct FParagraph
 	{
diff --git a/Engine/Source/Runtime/Experimental/IoStore/OnDemand/Private/Statistics.cpp b/Engine/Source/Runtime/Experimental/IoStore/OnDemand/Private/Statistics.cpp
--- a/Engine/Source/Runtime/Experimental/IoStore/OnDemand/Private/Statistics.cpp
+++ b/Engine/Source/Runtime/Experimental/IoStore/OnDemand/Private/Statistics.cpp
@@ -277,6 +277,7 @@ FCounterAtomicInt		GIoRequestReadBytes(TEXT("Ias/IoRequestReadBytes"), TraceCoun
 FCounterInt				GIoRequestCancelCount(TEXT("Ias/IoRequestCancelCount"), TraceCounterDisplayHint_None);
 FCounterAtomicInt		GIoRequestErrorCount(TEXT("Ias/IoRequestErrorCount"), TraceCounterDisplayHint_None);
 // cache stats
+static uint32			GCacheBootMs = 0;
 FCounterAtomicInt		GCacheErrorCount(TEXT("Ias/CacheErrorCount"), TraceCounterDisplayHint_None);
 FCounterAtomicInt		GCacheDecodeErrorCount(TEXT("Ias/CacheDecodeErrorCount"), TraceCounterDisplayHint_None);
 FCounterAtomicInt		GCacheGetCount(TEXT("Ias/CacheGetCount"), TraceCounterDisplayHint_None);
@@ -743,6 +744,7 @@ void FOnDemandIoBackendStats::PrintOnScreenStatistics(FCoreDelegates::FSeverityM
 			CachingText << TEXT(" | Rejected: ") << FText::AsMemory(GCacheRejectBytes.Get()).ToString();
 			CachingText << TEXT(" | Read: ") << FText::AsMemory(GCacheReadBytes.Get()).ToString();
 			CachingText << TEXT(" (") << GCacheGetCount.Get() << TEXT(")");
+			CachingText << TEXT(" (Boot: ") << GCacheBootMs << TEXT("ms)");
 		}
 		else
 		{
@@ -1058,6 +1060,11 @@ void FOnDemandIoBackendStats::OnIoRequestError()
 	CSV_CUSTOM_STAT_DEFINED(FrameIoRequestErrorCount, int32(GIoRequestErrorCount.Get()), ECsvCustomStatOp::Set);
 }
 
+void FOnDemandIoBackendStats::OnCacheBootMs(uint64 TimeMs)
+{
+	GCacheBootMs = uint32(TimeMs);
+}
+
 void FOnDemandIoBackendStats::OnCacheError()
 {
 	GCacheErrorCount.Increment();
diff --git a/Engine/Source/Runtime/Experimental/IoStore/OnDemand/Private/Statistics.h b/Engine/Source/Runtime/Experimental/IoStore/OnDemand/Private/Statistics.h
--- a/Engine/Source/Runtime/Experimental/IoStore/OnDemand/Private/Statistics.h
+++ b/Engine/Source/Runtime/Experimental/IoStore/OnDemand/Private/Statistics.h
@@ -54,6 +54,7 @@ public:
 	void OnIoRequestCancel() IAS_STATISTICS_IMPL()
 	void OnIoRequestError() IAS_STATISTICS_IMPL()
 
+	void OnCacheBootMs(uint64 TimeMs) IAS_STATISTICS_IMPL()
 	void OnCacheError() IAS_STATISTICS_IMPL()
 	void OnCacheDecodeError() IAS_STATISTICS_IMPL()
 	void OnCacheGet(uint64 DataSize) IAS_STATISTICS_IMPL()
