Dataflow - Fixed a missing AssetRegistry module dependency that was causing an error while buillding Headless Chaos Flesh.

diff --git a/Engine/Source/Runtime/Experimental/Dataflow/Core/DataflowCore.Build.cs b/Engine/Source/Runtime/Experimental/Dataflow/Core/DataflowCore.Build.cs
--- a/Engine/Source/Runtime/Experimental/Dataflow/Core/DataflowCore.Build.cs
+++ b/Engine/Source/Runtime/Experimental/Dataflow/Core/DataflowCore.Build.cs
@@ -9,6 +9,7 @@ namespace UnrealBuildTool.Rules
 			PublicDependencyModuleNames.AddRange(
 				new string[]
 				{
+					"AssetRegistry",
 					"Core",
 					"CoreUObject",
 					"Chaos",
