Raise the default number of Iris Polling Tasks to 32 for slightly better task scheduling. Also allow us to disable Multithreading Polling separately from other future multithreading features by setting NumPollingTasks to 0.

#rnx

diff --git a/Engine/Source/Runtime/Experimental/Iris/Core/Private/Iris/ReplicationSystem/Polling/ObjectPoller.cpp b/Engine/Source/Runtime/Experimental/Iris/Core/Private/Iris/ReplicationSystem/Polling/ObjectPoller.cpp
--- a/Engine/Source/Runtime/Experimental/Iris/Core/Private/Iris/ReplicationSystem/Polling/ObjectPoller.cpp
+++ b/Engine/Source/Runtime/Experimental/Iris/Core/Private/Iris/ReplicationSystem/Polling/ObjectPoller.cpp
@@ -40,11 +40,11 @@ namespace CVars
 		TEXT("Enable conditional profiler scopes for polling. Useful to see exactly what properties gets polled in cpu profiler..")
 	);
 
-	static int32 GNumPollingTasks = 8;
+	static int32 GNumPollingTasks = 32;
 	static FAutoConsoleVariableRef CvarNumPollingTasks(
 		TEXT("net.Iris.NumPollingTasks"),
 		GNumPollingTasks,
-		TEXT("Number of Tasks that will be created to poll dirty objects in a single frame")
+		TEXT("Number of Tasks that will be created to poll dirty objects in a single frame. Setting to 0 will run Polling Synchronously instead.")
 	);
 
 	static int32 GNumWordsPerChunk = 16;
@@ -227,7 +227,7 @@ void FObjectPoller::PollAndCopyObjects(const FNetBitArrayView& InObjectsConsider
 	const uint32 BitsPerChunk = CVars::GNumWordsPerChunk * FNetBitArray::WordBitCount;
 	const bool bEachTaskHasAChunkOfWork = (BitsPerChunk * NumPollingTasks) < ObjectsConsideredForPolling.GetNumBits();
 	
-	if (ReplicationSystemInternal->AreParallelTasksAllowed() && bEachTaskHasAChunkOfWork)
+	if (ReplicationSystemInternal->AreParallelTasksAllowed() && bEachTaskHasAChunkOfWork && NumPollingTasks > 0)
 	{
 		ReplicationSystemInternal->StartParallelPhase();
 
