Iris
- Changed NetRefHandleManager::AddDependentObject Assert to ensure to avoid taking down server.

diff --git a/Engine/Source/Runtime/Experimental/Iris/Core/Private/Iris/ReplicationSystem/NetRefHandleManager.cpp b/Engine/Source/Runtime/Experimental/Iris/Core/Private/Iris/ReplicationSystem/NetRefHandleManager.cpp
--- a/Engine/Source/Runtime/Experimental/Iris/Core/Private/Iris/ReplicationSystem/NetRefHandleManager.cpp
+++ b/Engine/Source/Runtime/Experimental/Iris/Core/Private/Iris/ReplicationSystem/NetRefHandleManager.cpp
@@ -800,7 +800,12 @@ FInternalNetRefIndex FNetRefHandleManager::GetRootObjectInternalIndexOfSubObject
 
 bool FNetRefHandleManager::AddDependentObject(FNetRefHandle ParentRefHandle, FNetRefHandle DependentObjectRefHandle, EDependentObjectSchedulingHint SchedulingHint, EDependentFilterTrait FilterTrait)
 {
-	check(ParentRefHandle != DependentObjectRefHandle);
+	if (ParentRefHandle == DependentObjectRefHandle)
+	{
+		UE_LOG(LogIris, Error, TEXT("FNetRefHandleManager::AddDependentObject: ParentObject %s cannot be dependent on itself. Please fix calling code."), *PrintObjectFromNetRefHandle(ParentRefHandle));
+		ensureMsgf(ParentRefHandle != DependentObjectRefHandle, TEXT("ParentObject %s cannot be dependent on itself."), *PrintObjectFromNetRefHandle(ParentRefHandle));
+		return false;
+	}
 
 	// validate objects
 	FInternalNetRefIndex ParentInternalIndex = GetInternalIndex(ParentRefHandle);
