Iris
* Added ability to configure unique async loading priority per class
* Configured via ObjectReplicationBridgeConfig::AsyncLoadingClassConfigs
* Possible settings are Default, High, VeryHigh
* The enum settings are mapped to an actual loading priority value via cvar "net.iris.AsyncLoading.[DefaultPriority|HighPriority|VeryHighPriority]"

#rnx

diff --git a/Engine/Source/Runtime/Experimental/Iris/Core/Private/Iris/ReplicationSystem/ChunkedDataStream/ChunkedDataReader.cpp b/Engine/Source/Runtime/Experimental/Iris/Core/Private/Iris/ReplicationSystem/ChunkedDataStream/ChunkedDataReader.cpp
--- a/Engine/Source/Runtime/Experimental/Iris/Core/Private/Iris/ReplicationSystem/ChunkedDataStream/ChunkedDataReader.cpp
+++ b/Engine/Source/Runtime/Experimental/Iris/Core/Private/Iris/ReplicationSystem/ChunkedDataStream/ChunkedDataReader.cpp
@@ -147,7 +147,7 @@ ON_SCOPE_EXIT
 		const uint32 ReturnPos = ExportsReader.GetPosBits();
 		ExportsReader.Seek(ReturnPos + ExportsOffset);
 
-		if (!ObjectReferenceCache->ReadExports(SubContext, &Entry.References->MustBeMappedReferences))
+		if (!ObjectReferenceCache->ReadExports(FNetRefHandle::GetInvalid(), SubContext, &Entry.References->MustBeMappedReferences, Entry.References->IrisAsyncLoadingPriority))
 		{
 			return false;
 		}
@@ -266,7 +266,7 @@ void FChunkedDataReader::AssemblePayloadsPendingAssembly(UE::Net::FNetSerializat
 	}
 }
 
-bool FChunkedDataReader::TryResolveUnresolvedMustBeMappedReferences(TArray<FNetRefHandle>& MustBeMappedReferences)
+bool FChunkedDataReader::TryResolveUnresolvedMustBeMappedReferences(TArray<FNetRefHandle>& MustBeMappedReferences, EIrisAsyncLoadingPriority IrisAsyncLoadingPriority)
 {
 	// Resolve
 	TArray<FNetRefHandle> Unresolved;
@@ -275,6 +275,12 @@ bool FChunkedDataReader::TryResolveUnresolvedMustBeMappedReferences(TArray<FNetR
 	Unresolved.Reserve(MustBeMappedReferences.Num());
 	QueuedObjectsToTrack.SetNum(MustBeMappedReferences.Num());
 
+	ResolveContext.AsyncLoadingPriority = ConvertAsyncLoadingPriority(IrisAsyncLoadingPriority);
+	ON_SCOPE_EXIT
+	{
+		ResolveContext.AsyncLoadingPriority = INDEX_NONE;
+	};
+
 	// Try to resolve references
 	for (FNetRefHandle Handle : MustBeMappedReferences)
 	{
@@ -330,7 +336,7 @@ UChunkedDataStream::EDispatchResult FChunkedDataReader::DispatchReceivedPayload(
 	const bool bProcessReferences = Entry.References.IsValid();
 	if (bProcessReferences)
 	{					
-		if (ObjectReferenceCache->ShouldAsyncLoad() && !TryResolveUnresolvedMustBeMappedReferences(Entry.References->MustBeMappedReferences))
+		if (ObjectReferenceCache->ShouldAsyncLoad() && !TryResolveUnresolvedMustBeMappedReferences(Entry.References->MustBeMappedReferences, Entry.References->IrisAsyncLoadingPriority))
 		{
 			// Wait for async loading to complete if we have any mustbemapped entries
 			UE_LOG_CHUNKEDDATASTREAM_CONN(Verbose, TEXT("Waiting for %d MustBeMapped references to be resolvable"), Entry.References->MustBeMappedReferences.Num());
diff --git a/Engine/Source/Runtime/Experimental/Iris/Core/Private/Iris/ReplicationSystem/ChunkedDataStream/ChunkedDataReader.h b/Engine/Source/Runtime/Experimental/Iris/Core/Private/Iris/ReplicationSystem/ChunkedDataStream/ChunkedDataReader.h
--- a/Engine/Source/Runtime/Experimental/Iris/Core/Private/Iris/ReplicationSystem/ChunkedDataStream/ChunkedDataReader.h
+++ b/Engine/Source/Runtime/Experimental/Iris/Core/Private/Iris/ReplicationSystem/ChunkedDataStream/ChunkedDataReader.h
@@ -5,12 +5,17 @@
 #include "ChunkedDataStreamCommon.h"
 #include "Containers/Array.h"
 #include "Containers/RingBuffer.h"
+
 #include "Iris/DataStream/DataStream.h"
+
 #include "Iris/ReplicationSystem/ChunkedDataStream/ChunkedDataStream.h"
 #include "Iris/ReplicationSystem/ObjectReferenceCacheFwd.h"
+#include "Iris/ReplicationSystem/ObjectReferenceTypes.h"
+
 #include "Iris/Serialization/IrisObjectReferencePackageMap.h"
 #include "Iris/Serialization/IrisPackageMapExportUtil.h"
 #include "Iris/Serialization/NetExportContext.h"
+
 #include "Net/Core/NetToken/NetTokenExportContext.h"
 
 namespace UE::Net::Private
@@ -29,6 +34,7 @@ public:
 	{
 		FIrisPackageMapExportsQuantizedType QuantizedExports;
 		TArray<FNetRefHandle> MustBeMappedReferences;
+		EIrisAsyncLoadingPriority IrisAsyncLoadingPriority = EIrisAsyncLoadingPriority::Default;
 
 		~FReferencesForImport()
 		{
@@ -73,7 +79,7 @@ public:
 
 	bool ProcessExportPayload(FNetSerializationContext& Context, FRecvQueueEntry& Entry);
 	void AssemblePayloadsPendingAssembly(UE::Net::FNetSerializationContext& Context);
-	bool TryResolveUnresolvedMustBeMappedReferences(TArray<FNetRefHandle>& MustBeMappedReferences);
+	bool TryResolveUnresolvedMustBeMappedReferences(TArray<FNetRefHandle>& MustBeMappedReferences, EIrisAsyncLoadingPriority IrisAsyncLoadingPriority);
 	UChunkedDataStream::EDispatchResult DispatchReceivedPayload(TFunctionRef<void(TConstArrayView64<uint8>)> DispatchPayloadFunction);
 	UChunkedDataStream::EDispatchResult DispatchReceivedPayloads(TFunctionRef<void(TConstArrayView64<uint8>)> DispatchPayloadFunction);
 
diff --git a/Engine/Source/Runtime/Experimental/Iris/Core/Private/Iris/ReplicationSystem/NetRefHandleManager.cpp b/Engine/Source/Runtime/Experimental/Iris/Core/Private/Iris/ReplicationSystem/NetRefHandleManager.cpp
--- a/Engine/Source/Runtime/Experimental/Iris/Core/Private/Iris/ReplicationSystem/NetRefHandleManager.cpp
+++ b/Engine/Source/Runtime/Experimental/Iris/Core/Private/Iris/ReplicationSystem/NetRefHandleManager.cpp
@@ -1,21 +1,30 @@
 // Copyright Epic Games, Inc. All Rights Reserved.
 
 #include "NetRefHandleManager.h"
-#include "ReplicationOperationsInternal.h"
-#include "ReplicationProtocol.h"
-#include "ReplicationProtocolManager.h"
+
 #include "Iris/Core/IrisLog.h"
 #include "Iris/Core/IrisProfiler.h"
-#include "Iris/Serialization/NetSerializationContext.h"
+
+#include "Iris/ReplicationSystem/NetRefHandleManagerTypes.h"
+
 #include "Iris/Serialization/InternalNetSerializationContext.h"
+#include "Iris/Serialization/NetSerializationContext.h"
+
 #include "Net/Core/Trace/NetDebugName.h"
 #include "Net/Core/Trace/NetTrace.h"
+
+#include "ProfilingDebugging/CsvProfiler.h"
+
+#include "ReplicationOperationsInternal.h"
+#include "ReplicationProtocol.h"
+#include "ReplicationProtocolManager.h"
+
+#include "Stats/StatsMisc.h"
+
+#include "UObject/CoreNetTypes.h"
 #include "UObject/CoreNetTypes.h"
 #include "UObject/Object.h"
 #include "UObject/UObjectGlobals.h"
-#include "UObject/CoreNetTypes.h"
-#include "ProfilingDebugging/CsvProfiler.h"
-#include "Stats/StatsMisc.h"
 
 CSV_DEFINE_CATEGORY(IrisCommon, true);
 
@@ -187,7 +196,7 @@ FInternalNetRefIndex FNetRefHandleManager::GetNextFreeInternalIndex() const
 	return NextFreeIndex != FNetBitArray::InvalidIndex ? NextFreeIndex : InvalidInternalIndex;
 }
 
-FInternalNetRefIndex FNetRefHandleManager::InternalCreateNetObject(const FNetRefHandle NetRefHandle, const FNetHandle GlobalHandle, const FReplicationProtocol* ReplicationProtocol)
+FInternalNetRefIndex FNetRefHandleManager::InternalCreateNetObject(const FNetRefHandle NetRefHandle, const FNetHandle GlobalHandle, const FCreateNetObjectParams& Params)
 {
 	if (ActiveObjectCount >= MaxActiveObjectCount)
 	{
@@ -213,7 +222,7 @@ FInternalNetRefIndex FNetRefHandleManager::InternalCreateNetObject(const FNetRef
 		// If we could not grow anymore, kill the process now. The system cannot replicate objects anymore and the game behavior is undefined.
 		if (InternalIndex == InvalidInternalIndex)
 		{
-			UE_LOG(LogIris, Fatal, TEXT("NetRefHandleManager: Hit the maximum limit of active replicated objects: %u. Aborting since we cannot replicate %s"), MaxActiveObjectCount, ReplicationProtocol->DebugName->Name);
+			UE_LOG(LogIris, Fatal, TEXT("NetRefHandleManager: Hit the maximum limit of active replicated objects: %u. Aborting since we cannot replicate %s"), MaxActiveObjectCount, Params.ReplicationProtocol->DebugName->Name);
 			return InvalidInternalIndex;
 		}
 	}
@@ -233,7 +242,9 @@ FInternalNetRefIndex FNetRefHandleManager::InternalCreateNetObject(const FNetRef
 
 	Data.RefHandle = NetRefHandle;
 	Data.NetHandle = GlobalHandle;
-	Data.Protocol = ReplicationProtocol;
+	Data.Protocol = Params.ReplicationProtocol;
+	Data.NetFactoryId = Params.NetFactoryId;
+	Data.IrisAsyncLoadingPriority = Params.IrisAsyncLoadingPriority;
 	Data.InstanceProtocol = nullptr;
 	ObjectsWithPreUpdate.ClearBit(InternalIndex);
 	ObjectsWithFullPushBasedDirtiness.ClearBit(InternalIndex);
@@ -331,14 +342,16 @@ FNetRefHandle FNetRefHandleManager::AllocateNetRefHandle(bool bIsStatic)
 	return NewHandle;
 }
 
-FNetRefHandle FNetRefHandleManager::CreateNetObject(FNetRefHandle WantedHandle, FNetHandle GlobalHandle, const FReplicationProtocol* ReplicationProtocol)
+FNetRefHandle FNetRefHandleManager::CreateNetObject(FNetRefHandle WantedHandle, FNetHandle GlobalHandle, const FCreateNetObjectParams& Params)
 {
 	FNetRefHandle NetRefHandle = WantedHandle;
 
-	const FInternalNetRefIndex InternalIndex = InternalCreateNetObject(NetRefHandle, GlobalHandle, ReplicationProtocol);
+	const FInternalNetRefIndex InternalIndex = InternalCreateNetObject(NetRefHandle, GlobalHandle, Params);
 	if (InternalIndex != InvalidInternalIndex)
 	{
 		FReplicatedObjectData& Data = ReplicatedObjectData[InternalIndex];
+
+		const UE::Net::FReplicationProtocol* ReplicationProtocol = Params.ReplicationProtocol;
 		
 		// Allocate storage for outgoing data. We need a valid pointer even if the size is zero.
 		uint8* StateBuffer = (uint8*)FMemory::MallocZeroed(FPlatformMath::Max(ReplicationProtocol->InternalTotalSize, 1U), ReplicationProtocol->InternalTotalAlignment);
@@ -363,37 +376,35 @@ FNetRefHandle FNetRefHandleManager::CreateNetObject(FNetRefHandle WantedHandle,
 }
 
 // Create NetRefHandle not owned by us
-FNetRefHandle FNetRefHandleManager::CreateNetObjectFromRemote(FNetRefHandle WantedHandle, const FReplicationProtocol* ReplicationProtocol,  UE::Net::FNetObjectFactoryId FactoryId)
+FNetRefHandle FNetRefHandleManager::CreateNetObjectFromRemote(FNetRefHandle WantedHandle, const FCreateNetObjectParams& Params)
 {
 	if (!ensureMsgf(WantedHandle.IsValid() && !WantedHandle.IsCompleteHandle(), TEXT("FNetRefHandleManager::CreateNetObjectFromRemote Expected WantedHandle %s to be valid and incomplete"), *WantedHandle.ToString()))
 	{
 		return FNetRefHandle();
 	}
 
-	check(FactoryId != UE::Net::InvalidNetObjectFactoryId);
+	check(Params.NetFactoryId != UE::Net::InvalidNetObjectFactoryId);
 
 	FNetRefHandle NetRefHandle = MakeNetRefHandle(WantedHandle.GetId(), ReplicationSystemId);
 
-	const FInternalNetRefIndex InternalIndex = InternalCreateNetObject(NetRefHandle, FNetHandle(), ReplicationProtocol);
+	const FInternalNetRefIndex InternalIndex = InternalCreateNetObject(NetRefHandle, FNetHandle(), Params);
 	if (InternalIndex != InvalidInternalIndex)
 	{
 		FReplicatedObjectData& Data = ReplicatedObjectData[InternalIndex];
 
-		Data.NetFactoryId = FactoryId;
-
 		// Allocate storage for incoming data
-		Data.ReceiveStateBuffer = (uint8*)FMemory::Malloc(ReplicationProtocol->InternalTotalSize, ReplicationProtocol->InternalTotalAlignment);
+		Data.ReceiveStateBuffer = (uint8*)FMemory::Malloc(Params.ReplicationProtocol->InternalTotalSize, Params.ReplicationProtocol->InternalTotalAlignment);
 
 		// Since we currently do not have any default values, just initialize to zero
 		// N.B. If this Memzero is optimized away for some reason it still needs to be done for protocols with dynamic state.
-		FMemory::Memzero(Data.ReceiveStateBuffer, ReplicationProtocol->InternalTotalSize);
+		FMemory::Memzero(Data.ReceiveStateBuffer, Params.ReplicationProtocol->InternalTotalSize);
 
 		// Note: We could initialize this from default but at the moment it is part of the contract for all serializers to write the value when we serialize and we will only apply dirty states
 
 		// Don't bother initializing the conditional changemask if present since it's currently unused on the receiving end.
 
 		// Bump protocol refcount
-		ReplicationProtocol->AddRef();
+		Params.ReplicationProtocol->AddRef();
 
 		return NetRefHandle;
 	}
@@ -501,12 +512,12 @@ void FNetRefHandleManager::ClearStateForFreedInternalIndex(FInternalNetRefIndex
 	ObjectsWithFullPushBasedDirtiness.ClearBit(FreedInternalIndex);
 }
 
-FNetRefHandle FNetRefHandleManager::CreateHandleForDestructionInfo(FNetRefHandle Handle, const FReplicationProtocol* DestroyedObjectProtocol)
+FNetRefHandle FNetRefHandleManager::CreateHandleForDestructionInfo(FNetRefHandle Handle, const FCreateNetObjectParams& Params)
 {
 	// Create destruction info handle carrying destruction info
 	constexpr bool bIsStaticHandle = false;
 	FNetRefHandle AllocatedHandle = AllocateNetRefHandle(bIsStaticHandle);
-	FNetRefHandle DestructionInfoHandle = CreateNetObject(AllocatedHandle, FNetHandle(), DestroyedObjectProtocol);
+	FNetRefHandle DestructionInfoHandle = CreateNetObject(AllocatedHandle, FNetHandle(), Params);
 
 	if (DestructionInfoHandle.IsValid())
 	{
diff --git a/Engine/Source/Runtime/Experimental/Iris/Core/Private/Iris/ReplicationSystem/NetRefHandleManager.h b/Engine/Source/Runtime/Experimental/Iris/Core/Private/Iris/ReplicationSystem/NetRefHandleManager.h
--- a/Engine/Source/Runtime/Experimental/Iris/Core/Private/Iris/ReplicationSystem/NetRefHandleManager.h
+++ b/Engine/Source/Runtime/Experimental/Iris/Core/Private/Iris/ReplicationSystem/NetRefHandleManager.h
@@ -2,20 +2,28 @@
 
 #pragma once
 
-#include "CoreTypes.h"
 #include "Containers/Array.h"
-#include "Containers/Map.h"
 #include "Containers/ArrayView.h"
-#include "Net/Core/NetBitArray.h"
-#include "Net/Core/NetHandle/NetHandle.h"
-#include "Iris/IrisConstants.h"
+#include "Containers/Map.h"
+
+#include "CoreTypes.h"
+
+#include "Delegates/Delegate.h"
+
 #include "Iris/Core/NetChunkedArray.h"
+
+#include "Iris/IrisConstants.h"
+
 #include "Iris/ReplicationState/ReplicationStateDescriptor.h"
-#include "Iris/ReplicationSystem/NetRefHandle.h"
 #include "Iris/ReplicationSystem/NetDependencyData.h"
 #include "Iris/ReplicationSystem/NetObjectFactoryRegistry.h"
+#include "Iris/ReplicationSystem/NetRefHandle.h"
+#include "Iris/ReplicationSystem/ObjectReferenceTypes.h"
+
+#include "Net/Core/NetBitArray.h"
+#include "Net/Core/NetHandle/NetHandle.h"
+
 #include "UObject/ObjectPtr.h"
-#include "Delegates/Delegate.h"
 
 class FReferenceCollector;
 namespace UE::Net
@@ -26,6 +34,8 @@ namespace UE::Net
 	namespace Private
 	{
 		class FReplicationProtocolManager;
+
+		struct FCreateNetObjectParams;
 	}
 }
 
@@ -89,6 +99,9 @@ public:
 		/** The factory responsible for instantiating this object */
 		FNetObjectFactoryId NetFactoryId = InvalidNetObjectFactoryId;
 
+		/** The async loading priority for references collected from this object */
+		EIrisAsyncLoadingPriority IrisAsyncLoadingPriority = EIrisAsyncLoadingPriority::Default;
+
 		union
 		{
 			uint16 Flags : 16U;
@@ -153,10 +166,10 @@ public:
 	FNetRefHandle AllocateNetRefHandle(bool bIsStatic);
 
 	// Create local Net Object
-	FNetRefHandle CreateNetObject(FNetRefHandle WantedHandle, FNetHandle GlobalHandle, const FReplicationProtocol* ReplicationProtocol);
+	FNetRefHandle CreateNetObject(FNetRefHandle WantedHandle, FNetHandle GlobalHandle, const FCreateNetObjectParams& Params);
 
 	// Create NetObject on request from remote
-	FNetRefHandle CreateNetObjectFromRemote(FNetRefHandle WantedHandle, const FReplicationProtocol* ReplicationProtocol, FNetObjectFactoryId FactoryId);
+	FNetRefHandle CreateNetObjectFromRemote(FNetRefHandle WantedHandle, const FCreateNetObjectParams& Params);
 
 	// Attach Instance protocol to handle
 	// Instance can be null, we only track the Instance for legacy support
@@ -171,7 +184,7 @@ public:
 	// Creates a DestructionInfo Handle which is used to replicate persistently destroyed static objects when late joining or streaming in levels with static objects that already has been destroyed
 	// Returns a new handle used to replicate the destruction info, and adds a cross-reference to the original handle being destroyed if that still exists 
 	// Handle is a reference to an already replicated Handle that is about to be destroyed
-	FNetRefHandle CreateHandleForDestructionInfo(FNetRefHandle Handle, const FReplicationProtocol* DestroydObjectProtocol);
+	FNetRefHandle CreateHandleForDestructionInfo(FNetRefHandle Handle, const FCreateNetObjectParams& Params);
 	
 	void DestroyNetObject(FNetRefHandle Handle);
 
@@ -190,6 +203,8 @@ public:
 	inline uint8* GetReplicatedObjectStateBufferNoCheck(FInternalNetRefIndex InternalObjectIndex) { return ReplicatedObjectStateBuffers[InternalObjectIndex]; }
 	inline const TNetChunkedArray<uint8*>& GetReplicatedObjectStateBuffers() const { return ReplicatedObjectStateBuffers; }
 
+	inline EIrisAsyncLoadingPriority GetIrisAsyncLoadingPriority(FInternalNetRefIndex InternalIndex) const;
+
 	/** Verify a handle against internal handle. A handle is valid if it matches internal storage. */
 	inline bool IsValidNetRefHandle(FNetRefHandle Handle) const;
 
@@ -374,7 +389,8 @@ private:
 
 	void InitNetBitArray(FNetBitArray* NetBitArray);
 
-	FInternalNetRefIndex InternalCreateNetObject(const FNetRefHandle NetRefHandle, const FNetHandle GlobalHandle, const FReplicationProtocol* ReplicationProtocol);
+	FInternalNetRefIndex InternalCreateNetObject(const FNetRefHandle NetRefHandle, const FNetHandle GlobalHandle, const FCreateNetObjectParams& Params);
+
 	void InternalDestroyNetObject(FInternalNetRefIndex InternalIndex);
 
 	/** Clear the status of all internal bit arrays when an internal index is put back in the free list */
@@ -563,6 +579,19 @@ FNetRefHandle FNetRefHandleManager::GetCompleteNetRefHandle(FNetRefHandle Incomp
 	}
 }
 
+EIrisAsyncLoadingPriority FNetRefHandleManager::GetIrisAsyncLoadingPriority(FInternalNetRefIndex InternalIndex) const
+{
+	const FNetRefHandleManager::FReplicatedObjectData& ObjectData = GetReplicatedObjectDataNoCheck(InternalIndex);
+	
+	// If it's a subobject, take the priority of the RootObject
+	if (ObjectData.SubObjectRootIndex != FNetRefHandleManager::InvalidInternalIndex)
+	{
+		return GetReplicatedObjectDataNoCheck(ObjectData.SubObjectRootIndex).IrisAsyncLoadingPriority;
+	}
+
+	return ObjectData.IrisAsyncLoadingPriority;
+}
+
 bool FNetRefHandleManager::IsValidNetRefHandle(FNetRefHandle Handle) const
 {
 	return RefHandleToInternalIndex.Contains(Handle);
diff --git a/Engine/Source/Runtime/Experimental/Iris/Core/Private/Iris/ReplicationSystem/NetRefHandleManagerTypes.h b/Engine/Source/Runtime/Experimental/Iris/Core/Private/Iris/ReplicationSystem/NetRefHandleManagerTypes.h
new file mode 100644
--- /dev/null
+++ b/Engine/Source/Runtime/Experimental/Iris/Core/Private/Iris/ReplicationSystem/NetRefHandleManagerTypes.h
@@ -0,0 +1,24 @@
+// Copyright Epic Games, Inc. All Rights Reserved.
+
+#pragma once
+
+#include "Iris/ReplicationSystem/NetObjectFactoryRegistry.h"
+#include "Iris/ReplicationSystem/ObjectReferenceTypes.h"
+
+namespace UE::Net
+{
+	struct FReplicationProtocol;
+}
+
+namespace UE::Net::Private
+{
+
+	/** Holds important parameters needed to create a NetObject */
+	struct FCreateNetObjectParams
+	{
+		FNetObjectFactoryId NetFactoryId = InvalidNetObjectFactoryId;
+		EIrisAsyncLoadingPriority IrisAsyncLoadingPriority = EIrisAsyncLoadingPriority::Default;
+		const UE::Net::FReplicationProtocol* ReplicationProtocol = nullptr;
+	};
+
+}
\ No newline at end of file
diff --git a/Engine/Source/Runtime/Experimental/Iris/Core/Private/Iris/ReplicationSystem/ObjectReferenceCache.cpp b/Engine/Source/Runtime/Experimental/Iris/Core/Private/Iris/ReplicationSystem/ObjectReferenceCache.cpp
--- a/Engine/Source/Runtime/Experimental/Iris/Core/Private/Iris/ReplicationSystem/ObjectReferenceCache.cpp
+++ b/Engine/Source/Runtime/Experimental/Iris/Core/Private/Iris/ReplicationSystem/ObjectReferenceCache.cpp
@@ -1,14 +1,21 @@
 // Copyright Epic Games, Inc. All Rights Reserved.
 
 #include "ObjectReferenceCache.h"
+
+#include "HAL/IConsoleManager.h"
+
 #include "Iris/Core/IrisLog.h"
 #include "Iris/Core/IrisMemoryTracker.h"
+
 #include "Iris/IrisConstants.h"
+
+#include "Iris/ReplicationSystem/ObjectReferenceTypes.h"
 #include "Iris/ReplicationSystem/ObjectReplicationBridge.h"
+#include "Iris/ReplicationSystem/PendingBatchData.h"
 #include "Iris/ReplicationSystem/ReplicationSystem.h"
 #include "Iris/ReplicationSystem/ReplicationSystemInternal.h"
 #include "Iris/ReplicationSystem/StringTokenStore.h"
-#include "Iris/ReplicationSystem/PendingBatchData.h"
+
 #include "Iris/Serialization/InternalNetSerializationContext.h"
 #include "Iris/Serialization/NetBitStreamReader.h"
 #include "Iris/Serialization/NetBitStreamUtil.h"
@@ -17,14 +24,17 @@
 #include "Iris/Serialization/NetSerializationContext.h"
 #include "Iris/Serialization/ObjectNetSerializer.h"
 #include "Iris/Serialization/QuantizedObjectReference.h"
+
 #include "Iris/Stats/NetStatsContext.h"
+
 #include "Misc/CoreMiscDefines.h"
 #include "Misc/PackageName.h"
 #include "Misc/StringBuilder.h"
+
 #include "Net/Core/Trace/NetDebugName.h"
 #include "Net/Core/Trace/NetTrace.h"
+
 #include "UObject/Package.h"
-#include "HAL/IConsoleManager.h"
 
 #ifndef UE_NET_VALIDATE_REFERENCECACHE
 #	define UE_NET_VALIDATE_REFERENCECACHE 0
@@ -70,9 +80,6 @@ FAutoConsoleVariableRef CVarIrisAsyncLoadingFailAllLoads(TEXT("net.Iris.AsyncLoa
 static bool bIrisEnableRemoveReferenceFallbackPath = true;
 FAutoConsoleVariableRef CVarIrisEnableRemoveReferenceFallbackPath(TEXT("net.iris.ObjectReferenceCache.EnableRemoveReferenceFallbackPath"), bIrisEnableRemoveReferenceFallbackPath, TEXT("Whether to enable the slow fallback path or not when ObjectReferenceCache::RemoveReference is passed a null pointer object."), ECVF_Default);
 
-static int32 IrisAsyncLoadingPriority = 0;
-FAutoConsoleVariableRef CVarIrisAsyncLoadingPriority(TEXT("net.iris.AsyncLoadingPriority"), IrisAsyncLoadingPriority, TEXT("The default priority to use when async loading packages referenced by replicated properties."), ECVF_Default);
-
 
 FObjectReferenceCache::FPendingAsyncLoadRequest::FPendingAsyncLoadRequest(FNetRefHandle InNetRefHandle, double InRequestStartTime)
 : RequestStartTime(InRequestStartTime)
@@ -914,7 +921,7 @@ UObject* FObjectReferenceCache::ResolveObjectReferenceHandleInternal(FNetRefHand
 				{
 					CacheObjectPtr->bIsPending = true;
 
-					StartAsyncLoadingPackage(*CacheObjectPtr, ObjectPathName, RefHandle, false);
+					StartAsyncLoadingPackage(*CacheObjectPtr, ObjectPathName, ResolveContext.AsyncLoadingPriority, RefHandle, false);
 					UE_LOG_REFERENCECACHE(VeryVerbose, TEXT("GetObjectFromRefHandle: Async loading package. Path: %s, %s"), *ObjectPath, *RefHandle.ToString());
 				}
 				else
@@ -976,7 +983,7 @@ UObject* FObjectReferenceCache::ResolveObjectReferenceHandleInternal(FNetRefHand
 			if (ShouldAsyncLoad() && !ResolveContext.bForceSyncLoad && Package->HasAnyInternalFlags(EInternalObjectFlags_AsyncLoading))
 			{
 				// Something else is already async loading this package, calling load again will add our callback to the existing load request
-				StartAsyncLoadingPackage(*CacheObjectPtr, ObjectPathName, RefHandle, true);
+				StartAsyncLoadingPackage(*CacheObjectPtr, ObjectPathName, ResolveContext.AsyncLoadingPriority, RefHandle, true);
 				UE_LOG(LogIrisReferences, Verbose, TEXT("GetObjectFromRefHandle: Listening to existing async load. Path: %s, NetRefHandle: %s"), *ObjectPath, *RefHandle.ToString());
 			}
 			else if (/*!GbNetCheckNoLoadPackages ||*/ !CacheObjectPtr->bNoLoad)
@@ -1699,7 +1706,7 @@ FObjectReferenceCache::EWriteExportsResult FObjectReferenceCache::WritePendingEx
 					NetTokenStore->WriteNetToken(Context, NetToken);
 					NetTokenStore->WriteTokenData(Context, NetToken);
 
-					ExportContext->AddExported(NetToken);			
+					ExportContext->AddExported(NetToken);
 				}
 			}
 		}
@@ -1715,6 +1722,8 @@ FObjectReferenceCache::EWriteExportsResult FObjectReferenceCache::WritePendingEx
 				Writer.WriteBool(true);
 				WriteFullReference(Context, Reference);
 
+				//UE_LOG(LogIrisReferences, Verbose, TEXT("WritePendingExports for %s : Export[%d] %s"), *NetRefHandleManager->PrintObjectFromIndex(ObjectIndex), CountExports, *NetRefHandleManager->PrintObjectFromNetRefHandle(Handle));
+
 #if UE_NET_IRIS_CSV_STATS
 				++CountExports;
 #endif
@@ -1747,8 +1756,9 @@ FObjectReferenceCache::EWriteExportsResult FObjectReferenceCache::WritePendingEx
 	return EWriteExportsResult::NoExports;
 }
 
-bool FObjectReferenceCache::ReadExports(FNetSerializationContext& Context, TArray<FNetRefHandle>* MustBeMappedExports)
+bool FObjectReferenceCache::ReadExports(const FNetRefHandle& NetObjectHandle, FNetSerializationContext& Context, TArray<FNetRefHandle>* MustBeMappedExports, EIrisAsyncLoadingPriority& OutIrisAsyncLoadingPriority)
 {
+	check(MustBeMappedExports);
 	FNetBitStreamReader& Reader = *Context.GetBitStreamReader();
 
 	// Read NetToken exports
@@ -1789,13 +1799,13 @@ bool FObjectReferenceCache::ReadExports(FNetSerializationContext& Context, TArra
 
 	if (!Context.HasErrorOrOverflow())
 	{
-		ReadMustBeMappedExports(Context, MustBeMappedExports);
+		ReadMustBeMappedExports(NetObjectHandle, Context, MustBeMappedExports, OutIrisAsyncLoadingPriority);
 	}
 
 	return !Context.HasErrorOrOverflow();
 }
 
-bool FObjectReferenceCache::WriteMustBeMappedExports(FNetSerializationContext& Context, FInternalNetRefIndex ObjectIndex,TArrayView<const FNetObjectReference> ExportsView) const
+bool FObjectReferenceCache::WriteMustBeMappedExports(FNetSerializationContext& Context, FInternalNetRefIndex ObjectIndex, TArrayView<const FNetObjectReference> ExportsView) const
 {
 	FNetBitStreamWriter& Writer = *Context.GetBitStreamWriter();
 	FNetExportContext* ExportContext = Context.GetExportContext();
@@ -1808,6 +1818,11 @@ bool FObjectReferenceCache::WriteMustBeMappedExports(FNetSerializationContext& C
 		uint32 CountExports = 0;
 #endif
 
+		const EIrisAsyncLoadingPriority IrisAsyncLoadingPriority = NetRefHandleManager->GetIrisAsyncLoadingPriority(ObjectIndex);
+
+		// Only write the loading priority before the first reference
+		bool bHasWrittenPriority = false;
+
 		for (const FNetObjectReference& Reference : ExportsView)
 		{
 			const FNetRefHandle Handle = Reference.GetRefHandle();
@@ -1816,14 +1831,29 @@ bool FObjectReferenceCache::WriteMustBeMappedExports(FNetSerializationContext& C
 			{
 				UE_NET_TRACE_OBJECT_SCOPE(Handle, *Context.GetBitStreamWriter(), Context.GetTraceCollector(), ENetTraceVerbosity::Verbose);
 				Writer.WriteBool(true);
+
+				if (!bHasWrittenPriority)
+				{
+					bHasWrittenPriority = true;
+					Writer.WriteBits((uint32)IrisAsyncLoadingPriority, GetIrisAsyncLoadingPriorityBits());
+				}
+
 				WriteNetRefHandle(Context, Handle);
 
+#if UE_NET_IRIS_CSV_STATS
+				UE_LOG(LogIrisReferences, VeryVerbose, TEXT("WriteMustBeMappedExports for %s. MustBeMapped[%d] %s (%s)"), *NetRefHandleManager->PrintObjectFromIndex(ObjectIndex), CountExports, *GetNameSafe(CachedObject->Object.Get()), *Handle.ToString());
+#endif
+
 #if UE_NET_IRIS_CSV_STATS
 				++CountExports;
 #endif
 			}
 		}
 
+#if UE_NET_IRIS_CSV_STATS
+		UE_CLOG(CountExports>0, LogIrisReferences, VeryVerbose, TEXT("WriteMustBeMappedExports for %s. Wrote %d exports with priority: %s"), *NetRefHandleManager->PrintObjectFromIndex(ObjectIndex), CountExports, LexToString(IrisAsyncLoadingPriority));
+#endif
+
 		// $TODO: This should be a separate state, as MustBeMapped is not really an export, just a list of important exports
 		UE_NET_IRIS_STATS_ADD_COUNT_FOR_OBJECT(Context.GetNetStatsContext(), WriteExports, ObjectIndex, CountExports);
 	}
@@ -1834,13 +1864,20 @@ bool FObjectReferenceCache::WriteMustBeMappedExports(FNetSerializationContext& C
 	return !Writer.IsOverflown();
 }
 
-void FObjectReferenceCache::ReadMustBeMappedExports(FNetSerializationContext& Context, TArray<FNetRefHandle>* MustBeMappedExports)
+void FObjectReferenceCache::ReadMustBeMappedExports(const FNetRefHandle& NetObjectHandle, FNetSerializationContext& Context, TArray<FNetRefHandle>* MustBeMappedExports, EIrisAsyncLoadingPriority& OutIrisAsyncLoadingPriority)
 {
+	check(MustBeMappedExports);
+
 	FNetBitStreamReader& Reader = *Context.GetBitStreamReader();
 
 	bool bHasExportsToRead = Reader.ReadBool(); 
 	UE_NET_TRACE_SCOPE(MustBeMappedExports, Reader, Context.GetTraceCollector(), ENetTraceVerbosity::Verbose);
 
+	if (bHasExportsToRead)
+	{	
+		OutIrisAsyncLoadingPriority = static_cast<EIrisAsyncLoadingPriority>(Reader.ReadBits(GetIrisAsyncLoadingPriorityBits()));
+	}
+
 	while (bHasExportsToRead && !Context.HasErrorOrOverflow())
 	{
 		UE_NET_TRACE_NAMED_OBJECT_SCOPE(ReferenceScope, FNetRefHandle::GetInvalid(), *Context.GetBitStreamReader(), Context.GetTraceCollector(), ENetTraceVerbosity::Verbose);
@@ -1849,11 +1886,12 @@ void FObjectReferenceCache::ReadMustBeMappedExports(FNetSerializationContext& Co
 
 		bHasExportsToRead = Reader.ReadBool();
 
-		if (MustBeMappedExports)
-		{
-			MustBeMappedExports->Add(MustBeMappedHandle);
-		}
+		UE_LOG(LogIrisReferences, VeryVerbose, TEXT("ReadMustBeMappedExports for %s. MustBeMapped[%d] %s"), *NetObjectHandle.ToString(), MustBeMappedExports->Num(), *MustBeMappedHandle.ToString());
+
+		MustBeMappedExports->Add(MustBeMappedHandle);
 	}
+
+	UE_CLOG(MustBeMappedExports->Num() > 0, LogIrisReferences, VeryVerbose, TEXT("ReadMustBeMappedExports for %s. Read %d references with priority: %s"), *NetObjectHandle.ToString(), MustBeMappedExports->Num(), LexToString(OutIrisAsyncLoadingPriority));
 }
 
 FString FObjectReferenceCache::DescribeObjectReference(const FNetObjectReference Ref, const FNetObjectResolveContext& ResolveContext)
@@ -1976,10 +2014,19 @@ void FObjectReferenceCache::ValidateAsyncLoadingPackage(FCachedNetObjectReferenc
 //#endif
 }
 
-void FObjectReferenceCache::StartAsyncLoadingPackage(FCachedNetObjectReference& CacheObject, FName PackagePath, const FNetRefHandle RefHandle, const bool bWasAlreadyAsyncLoading)
+void FObjectReferenceCache::StartAsyncLoadingPackage(FCachedNetObjectReference& CacheObject, FName PackagePath, TAsyncLoadPriority AsyncLoadingPriority, const FNetRefHandle RefHandle, const bool bWasAlreadyAsyncLoading)
 {
 	LLM_SCOPE_BYTAG(Iris);
 
+	const bool bIsPrioritySet = AsyncLoadingPriority != INDEX_NONE;
+	if (!bIsPrioritySet)
+	{
+		ensureMsgf(false, TEXT("AsyncLoadingPriorty wasn't set in the ResolveContext. Using default value as fallback"));
+		AsyncLoadingPriority = ConvertAsyncLoadingPriority(EIrisAsyncLoadingPriority::Default);
+	}
+
+	UE_LOG(LogIrisReferences, Verbose, TEXT("StartAsyncLoadingPackage for %s: Package: %s (Async loading priority: %d)"), *NetRefHandleManager->PrintObjectFromNetRefHandle(RefHandle), *PackagePath.ToString(), AsyncLoadingPriority);
+
 	// Need timer?
 	FPendingAsyncLoadRequest LoadRequest(RefHandle, ReplicationSystem->GetElapsedTime());
 	
@@ -2018,7 +2065,7 @@ void FObjectReferenceCache::StartAsyncLoadingPackage(FCachedNetObjectReference&
 				{
 					AsyncPackageCallback(PackageName, Package, Result);
 				}),
-			IrisAsyncLoadingPriority
+			AsyncLoadingPriority
 		);
 	}
 }
diff --git a/Engine/Source/Runtime/Experimental/Iris/Core/Private/Iris/ReplicationSystem/ObjectReferenceCache.h b/Engine/Source/Runtime/Experimental/Iris/Core/Private/Iris/ReplicationSystem/ObjectReferenceCache.h
--- a/Engine/Source/Runtime/Experimental/Iris/Core/Private/Iris/ReplicationSystem/ObjectReferenceCache.h
+++ b/Engine/Source/Runtime/Experimental/Iris/Core/Private/Iris/ReplicationSystem/ObjectReferenceCache.h
@@ -11,6 +11,8 @@
 #include "ObjectReferenceCacheFwd.h"
 #include "UObject/ObjectPtr.h"
 
+enum class EIrisAsyncLoadingPriority : uint8;
+
 namespace UE::Net
 {
 	class FStringTokenStore;
@@ -131,7 +133,7 @@ public:
 	// it is up to the caller to roll back written data and pending exports
 	EWriteExportsResult WritePendingExports(FNetSerializationContext& Context, FInternalNetRefIndex ObjectIndex);
 
-	bool ReadExports(FNetSerializationContext& Context, TArray<FNetRefHandle>* MustBeMappedExports);
+	bool ReadExports(const FNetRefHandle& NetObjectHandle, FNetSerializationContext& Context, TArray<FNetRefHandle>* MustBeMappedExports, EIrisAsyncLoadingPriority& OutIrisAsyncLoadingPriority);
 
 	static FNetObjectReference MakeNetObjectReference(FNetRefHandle Handle);
 
@@ -228,9 +230,9 @@ private:
 	// Must be mapped exports are written for each batch that serializes object references, if async loading is enabled the client
 	// will defer application of data contained in the batch until the must be mapped exports are resolvable.
 	bool WriteMustBeMappedExports(FNetSerializationContext& Context, FInternalNetRefIndex ObjectIndex, TArrayView<const FNetObjectReference> ExportsView) const;
-	void ReadMustBeMappedExports(FNetSerializationContext& Context, TArray<FNetRefHandle>* MustBeMappedExports);
+	void ReadMustBeMappedExports(const FNetRefHandle& NetObjectHandle, FNetSerializationContext& Context, TArray<FNetRefHandle>* MustBeMappedExports, EIrisAsyncLoadingPriority& OutIrisAsyncLoadingPriority);
 
-	void StartAsyncLoadingPackage(FCachedNetObjectReference& Object, FName PackagePath, const FNetRefHandle RefHandle, const bool bWasAlreadyAsyncLoading);
+	void StartAsyncLoadingPackage(FCachedNetObjectReference& Object, FName PackagePath, TAsyncLoadPriority AsyncLoadingPriority, const FNetRefHandle RefHandle, const bool bWasAlreadyAsyncLoading);
 	void ValidateAsyncLoadingPackage(FCachedNetObjectReference& Object, FName PackagePath, const FNetRefHandle RefHandle);
 
 private:
diff --git a/Engine/Source/Runtime/Experimental/Iris/Core/Private/Iris/ReplicationSystem/ObjectReferenceCacheFwd.h b/Engine/Source/Runtime/Experimental/Iris/Core/Private/Iris/ReplicationSystem/ObjectReferenceCacheFwd.h
--- a/Engine/Source/Runtime/Experimental/Iris/Core/Private/Iris/ReplicationSystem/ObjectReferenceCacheFwd.h
+++ b/Engine/Source/Runtime/Experimental/Iris/Core/Private/Iris/ReplicationSystem/ObjectReferenceCacheFwd.h
@@ -7,21 +7,30 @@
 class UReplicationSystem;
 class UObjectReplicationBridge;
 
+// From UObjectGlobals.h
+typedef int32 TAsyncLoadPriority;
+
+// Forward declarations
 namespace UE::Net
 {
 	class FNetTokenStoreState;
 	class FStringTokenStore;
+}
+
+namespace UE::Net::Private
+{
+	class FNetExportContext;
+	class FNetPendingBatches;
+}
 
+// Definitions
+namespace UE::Net
+{
 	struct FNetObjectResolveContext
 	{
 		FNetTokenStoreState* RemoteNetTokenStoreState = nullptr;
 		uint32 ConnectionId = InvalidConnectionId;
+		TAsyncLoadPriority AsyncLoadingPriority = INDEX_NONE; // default to -1 to trap codepaths that are missing the assignment
 		bool bForceSyncLoad = false;
 	};
-
-	namespace Private
-	{
-		class FNetExportContext;
-		class FNetPendingBatches;
-	}
 }
diff --git a/Engine/Source/Runtime/Experimental/Iris/Core/Private/Iris/ReplicationSystem/ObjectReferenceTypes.cpp b/Engine/Source/Runtime/Experimental/Iris/Core/Private/Iris/ReplicationSystem/ObjectReferenceTypes.cpp
new file mode 100644
--- /dev/null
+++ b/Engine/Source/Runtime/Experimental/Iris/Core/Private/Iris/ReplicationSystem/ObjectReferenceTypes.cpp
@@ -0,0 +1,51 @@
+// Copyright Epic Games, Inc. All Rights Reserved.
+
+#include "Iris/ReplicationSystem/ObjectReferenceTypes.h"
+
+#include "HAL/IConsoleManager.h"
+
+namespace UE::Net::Private
+{
+	static int32 IrisDefaultAsyncLoadingPriority = 0;
+	static FAutoConsoleVariableRef CVarIrisDefaultAsyncLoadingPriority(
+		TEXT("net.iris.AsyncLoading.DefaultPriority"),
+		IrisDefaultAsyncLoadingPriority,
+		TEXT("The default priority to use when async loading packages referenced by replicated properties.")
+		);
+
+	static int32 IrisHighAsyncLoadingPriority = 50;
+	static FAutoConsoleVariableRef CVarIrisHighAsyncLoadingPriority(
+		TEXT("net.iris.AsyncLoading.HighPriority"),
+		IrisHighAsyncLoadingPriority,
+		TEXT("The high priority value to use when async loading packages referenced by replicated properties.")
+		);
+
+	static int32 IrisVeryHighAsyncLoadingPriority = 100;
+	static FAutoConsoleVariableRef CVarIrisVeryHighAsyncLoadingPriority(
+		TEXT("net.iris.AsyncLoading.VeryHighPriority"),
+		IrisVeryHighAsyncLoadingPriority,
+		TEXT("The very high priority value to use when async loading packages referenced by replicated properties.")
+		);
+
+}
+
+TAsyncLoadPriority ConvertAsyncLoadingPriority(EIrisAsyncLoadingPriority IrisPriority)
+{
+	switch(IrisPriority)
+	{
+		case EIrisAsyncLoadingPriority::High:
+		{
+			return UE::Net::Private::IrisHighAsyncLoadingPriority;
+		} break;
+
+		case EIrisAsyncLoadingPriority::VeryHigh:
+		{
+			return UE::Net::Private::IrisVeryHighAsyncLoadingPriority;
+		} break;
+
+		default:
+		{
+			return UE::Net::Private::IrisDefaultAsyncLoadingPriority;
+		} break;
+	}
+}
\ No newline at end of file
diff --git a/Engine/Source/Runtime/Experimental/Iris/Core/Private/Iris/ReplicationSystem/ObjectReplicationBridge.cpp b/Engine/Source/Runtime/Experimental/Iris/Core/Private/Iris/ReplicationSystem/ObjectReplicationBridge.cpp
--- a/Engine/Source/Runtime/Experimental/Iris/Core/Private/Iris/ReplicationSystem/ObjectReplicationBridge.cpp
+++ b/Engine/Source/Runtime/Experimental/Iris/Core/Private/Iris/ReplicationSystem/ObjectReplicationBridge.cpp
@@ -21,26 +21,29 @@
 #include "Net/Core/Trace/NetTrace.h"
 #include "Net/Core/Trace/NetDebugName.h"
 
+#include "Iris/ReplicationState/ReplicationStateDescriptorBuilder.h"
+#include "Iris/ReplicationState/ReplicationStateUtil.h"
+
 #include "Iris/ReplicationSystem/LegacyPushModel.h"
+#include "Iris/ReplicationSystem/NetObjectFactory.h"
+#include "Iris/ReplicationSystem/NetObjectFactoryRegistry.h"
+#include "Iris/ReplicationSystem/NetRefHandleManagerTypes.h"
 #include "Iris/ReplicationSystem/ObjectPollFrequencyLimiter.h"
 #include "Iris/ReplicationSystem/ObjectReplicationBridgeConfig.h"
+#include "Iris/ReplicationSystem/Polling/ObjectPoller.h"
 #include "Iris/ReplicationSystem/Prioritization/NetObjectPrioritizer.h"
 #include "Iris/ReplicationSystem/ReplicationConnections.h"
 #include "Iris/ReplicationSystem/ReplicationFragment.h"
 #include "Iris/ReplicationSystem/ReplicationFragmentInternal.h"
 #include "Iris/ReplicationSystem/ReplicationFragmentUtil.h"
-#include "Iris/ReplicationState/ReplicationStateDescriptorBuilder.h"
-#include "Iris/ReplicationState/ReplicationStateUtil.h"
-#include "Iris/ReplicationSystem/ReplicationSystem.h"
-#include "Iris/ReplicationSystem/ReplicationSystemTypes.h"
-#include "Iris/ReplicationSystem/ReplicationSystemInternal.h"
 #include "Iris/ReplicationSystem/ReplicationOperations.h"
 #include "Iris/ReplicationSystem/ReplicationOperationsInternal.h"
+#include "Iris/ReplicationSystem/ReplicationSystem.h"
+#include "Iris/ReplicationSystem/ReplicationSystemInternal.h"
+#include "Iris/ReplicationSystem/ReplicationSystemTypes.h"
 #include "Iris/ReplicationSystem/ReplicationWriter.h"
-#include "Iris/ReplicationSystem/NetObjectFactory.h"
-#include "Iris/ReplicationSystem/NetObjectFactoryRegistry.h"
 #include "Iris/ReplicationSystem/RepTag.h"
-#include "Iris/ReplicationSystem/Polling/ObjectPoller.h"
+
 #include "Iris/Serialization/NetSerializationContext.h"
 #include "Iris/Serialization/InternalNetSerializationContext.h"
 #include "Iris/Serialization/NetBitStreamUtil.h"
@@ -91,7 +94,7 @@ static FAutoConsoleVariableRef CVarEnableForceNetUpdate(
 	TEXT("net.Iris.EnableForceNetUpdate"),
 	bEnableForceNetUpdate,
 	TEXT("When true the system only allows ForceNetUpdate to skip the poll frequency of objects. When false any MarkDirty object will be immediately polled.")
-);
+	);
 
 #if UE_BUILD_SHIPPING
 static int32 GInvalidLocationWarnings = 0; // don't check OOB locations in Shipping
@@ -418,11 +421,18 @@ UE::Net::FNetRefHandle UObjectReplicationBridge::StartReplicatingNetObject(UObje
 		return FNetRefHandle::GetInvalid();
 	}
 
-	IRIS_PROFILER_PROTOCOL_NAME(ReplicationProtocol->DebugName->Name);		
+	IRIS_PROFILER_PROTOCOL_NAME(ReplicationProtocol->DebugName->Name);
+
+	const FCreateNetObjectParams CreateNetObjectParams
+	{
+		.NetFactoryId = NetFactoryId,
+		.IrisAsyncLoadingPriority = GetClassIrisAsyncLoadingPriority(Instance->GetClass()),
+		.ReplicationProtocol = ReplicationProtocol,
+	};
 
 	// Create NetHandle and bind instance
 	FNetHandle NetHandle = FNetHandleManager::GetOrCreateNetHandle(Instance);
-	FNetRefHandle RefHandle = InternalCreateNetObject(AllocatedRefHandle, NetHandle, ReplicationProtocol);
+	FNetRefHandle RefHandle = InternalCreateNetObject(AllocatedRefHandle, NetHandle, CreateNetObjectParams);
 
 	if (!RefHandle.IsValid())
 	{
@@ -1055,21 +1065,26 @@ FReplicationBridgeCreateNetRefHandleResult UObjectReplicationBridge::CreateNetRe
 		FIrisDelegates::GetCriticalErrorDetectedDelegate().Broadcast(ReplicationSystem);
 
 		OnProtocolMismatchDetected(WantedNetHandle);
+		return FReplicationBridgeCreateNetRefHandleResult();
 	}
-	else
+	
+	// Create NetHandle
+	const FCreateNetObjectParams CreateNetObjectParams
 	{
-		// Create NetHandle
-		FNetRefHandle Handle = InternalCreateNetObjectFromRemote(WantedNetHandle, ReplicationProtocol, FactoryId);
-		CreateResult.NetRefHandle = Handle;
-		if (Handle.IsValid())
-		{
-			RegisterRemoteInstance(Handle, InstancePtr, ReplicationProtocol, InstanceProtocol.Get(), Context.ConnectionId);
+		.NetFactoryId = FactoryId,
+		.ReplicationProtocol = ReplicationProtocol,
+	};
 
-			// Release instance protocol from the uniquePtr as it is now successfully bound to the handle
-			(void)InstanceProtocol.Release();
+	FNetRefHandle Handle = InternalCreateNetObjectFromRemote(WantedNetHandle, CreateNetObjectParams);
+	CreateResult.NetRefHandle = Handle;
+	if (Handle.IsValid())
+	{
+		RegisterRemoteInstance(Handle, InstancePtr, ReplicationProtocol, InstanceProtocol.Get(), Context.ConnectionId);
 
-			Factory->PostInstantiation(UNetObjectFactory::FPostInstantiationContext{ .Instance = InstancePtr, .Header = Header.Get(), .ConnectionId = Context.ConnectionId });
-		}
+		// Release instance protocol from the uniquePtr as it is now successfully bound to the handle
+		(void)InstanceProtocol.Release();
+
+		Factory->PostInstantiation(UNetObjectFactory::FPostInstantiationContext{ .Instance = InstancePtr, .Header = Header.Get(), .ConnectionId = Context.ConnectionId });
 	}
 
 	return CreateResult;
@@ -1900,6 +1915,23 @@ bool UObjectReplicationBridge::FindOrCachePollFrequency(const UClass* Class, flo
 	return false;
 }
 
+EIrisAsyncLoadingPriority UObjectReplicationBridge::GetClassIrisAsyncLoadingPriority(const UClass* Class)
+{
+	const FName ClassName = GetConfigClassPathName(Class);
+	
+	//$IRIS TODO: Add a passive cache of classes we looked up here
+	for (const UClass* ClassToFind = Class; ClassToFind != nullptr; ClassToFind = ClassToFind->GetSuperClass())
+	{
+		const FName ClassToFindName = GetConfigClassPathName(ClassToFind);
+		if (const EIrisAsyncLoadingPriority* AsyncLoadingPriority = ClassesIrisAsyncLoadingPriority.Find(ClassToFindName))
+		{
+			return *AsyncLoadingPriority;
+		}
+	}
+
+	return EIrisAsyncLoadingPriority::Default;
+}
+
 FName UObjectReplicationBridge::GetConfigClassPathName(const UClass* Class)
 {
 	if (FName* CachedPathName = ConfigClassPathNameCache.Find(Class))
@@ -2215,6 +2247,7 @@ void UObjectReplicationBridge::LoadConfig()
 	ClassesWithPrioritizer.Empty();
 	ClassesWithDeltaCompression.Empty();
 	ClassesFlaggedCritical.Empty();
+	ClassesIrisAsyncLoadingPriority.Empty();
 
 	// Reset PathNameCache
 	ConfigClassPathNameCache.Empty();
@@ -2307,6 +2340,7 @@ void UObjectReplicationBridge::LoadConfig()
 			ClassesFlaggedCritical.Add(CriticalClassConfig.ClassName, CriticalClassConfig.bDisconnectOnProtocolMismatch);
 		}
 	}
+
 	// Load TypeStats settings
 	{
 		for (const FObjectReplicationBridgeTypeStatsConfig& TypeStatsConfig : BridgeConfig->GetTypeStatsConfigs())
@@ -2320,6 +2354,14 @@ void UObjectReplicationBridge::LoadConfig()
 			ClassesWithTypeStats.Add(TypeStatsConfig.ClassName, TypeStatsConfig.TypeStatsName);
 		}
 	}
+
+	// Async loading priority
+	{
+		for (const FObjectReplicationBridgeAsyncLoadingClassConfig& AsyncLoadingConfig : BridgeConfig->GetAsyncLoadingClassConfigs())
+		{
+			ClassesIrisAsyncLoadingPriority.Add(AsyncLoadingConfig.ClassName, AsyncLoadingConfig.IrisAsyncLoadingPriority);
+		}
+	}
 }
 
 void UObjectReplicationBridge::FindClassesInPollPeriodOverrides()
diff --git a/Engine/Source/Runtime/Experimental/Iris/Core/Private/Iris/ReplicationSystem/ObjectReplicationBridgeConfig.cpp b/Engine/Source/Runtime/Experimental/Iris/Core/Private/Iris/ReplicationSystem/ObjectReplicationBridgeConfig.cpp
--- a/Engine/Source/Runtime/Experimental/Iris/Core/Private/Iris/ReplicationSystem/ObjectReplicationBridgeConfig.cpp
+++ b/Engine/Source/Runtime/Experimental/Iris/Core/Private/Iris/ReplicationSystem/ObjectReplicationBridgeConfig.cpp
@@ -42,3 +42,8 @@ TConstArrayView<FObjectReplicationBridgeTypeStatsConfig> UObjectReplicationBridg
 {
 	return MakeArrayView(TypeStatsConfigs);
 }
+
+TConstArrayView<FObjectReplicationBridgeAsyncLoadingClassConfig> UObjectReplicationBridgeConfig::GetAsyncLoadingClassConfigs() const
+{
+	return MakeArrayView(AsyncLoadingClassConfigs);
+}
diff --git a/Engine/Source/Runtime/Experimental/Iris/Core/Private/Iris/ReplicationSystem/PendingBatchData.h b/Engine/Source/Runtime/Experimental/Iris/Core/Private/Iris/ReplicationSystem/PendingBatchData.h
--- a/Engine/Source/Runtime/Experimental/Iris/Core/Private/Iris/ReplicationSystem/PendingBatchData.h
+++ b/Engine/Source/Runtime/Experimental/Iris/Core/Private/Iris/ReplicationSystem/PendingBatchData.h
@@ -3,7 +3,9 @@
 #pragma once
 
 #include "Containers/Array.h"
+
 #include "Iris/ReplicationSystem/NetRefHandle.h"
+#include "Iris/ReplicationSystem/ObjectReferenceTypes.h"
 
 namespace UE::Net::Private
 {
@@ -56,6 +58,9 @@ struct FPendingBatchData
 	
 	// Incremented every time we try to process queued batches, reset each time we output warning.
 	int32 PendingBatchTryProcessCount = 0;
+
+	// The async loading priority of the object stored in the queued batch
+	EIrisAsyncLoadingPriority IrisAsyncLoadingPriority = EIrisAsyncLoadingPriority::Default;
 };
 
 struct FPendingBatches
diff --git a/Engine/Source/Runtime/Experimental/Iris/Core/Private/Iris/ReplicationSystem/ReplicationBridge.cpp b/Engine/Source/Runtime/Experimental/Iris/Core/Private/Iris/ReplicationSystem/ReplicationBridge.cpp
--- a/Engine/Source/Runtime/Experimental/Iris/Core/Private/Iris/ReplicationSystem/ReplicationBridge.cpp
+++ b/Engine/Source/Runtime/Experimental/Iris/Core/Private/Iris/ReplicationSystem/ReplicationBridge.cpp
@@ -1,32 +1,42 @@
 // Copyright Epic Games, Inc. All Rights Reserved.
 
 #include "Iris/ReplicationSystem/ReplicationBridge.h"
+
 #include "Containers/ArrayView.h"
+
 #include "HAL/IConsoleManager.h"
-#include "Iris/IrisConfigInternal.h"
+
 #include "Iris/Core/IrisCsv.h"
 #include "Iris/Core/IrisLog.h"
 #include "Iris/Core/IrisProfiler.h"
-#include "Net/Core/NetBitArray.h"
-#include "Net/Core/NetHandle/NetHandleManager.h"
-#include "Net/Core/Trace/NetTrace.h"
-#include "Net/Core/Trace/NetDebugName.h"
+#include "Iris/IrisConfigInternal.h"
+
 #include "Iris/ReplicationState/ReplicationStateUtil.h"
+
+#include "Iris/ReplicationSystem/ChangeMaskCache.h"
 #include "Iris/ReplicationSystem/LegacyPushModel.h"
-#include "Iris/ReplicationSystem/Prioritization/NetObjectPrioritizer.h"
+#include "Iris/ReplicationSystem/NetRefHandleManagerTypes.h"
 #include "Iris/ReplicationSystem/ReplicationFragment.h"
 #include "Iris/ReplicationSystem/ReplicationFragmentInternal.h"
 #include "Iris/ReplicationSystem/ReplicationSystem.h"
 #include "Iris/ReplicationSystem/ReplicationSystemInternal.h"
+#include "Iris/ReplicationSystem/ReplicationWriter.h"
 #include "Iris/ReplicationSystem/RepTag.h"
 #include "Iris/ReplicationSystem/WorldLocations.h"
-#include "ReplicationOperationsInternal.h"
-#include "Iris/Serialization/ObjectNetSerializer.h"
-#include "Iris/Serialization/NetBitStreamWriter.h"
-#include "Iris/ReplicationSystem/ChangeMaskCache.h"
-#include "Iris/ReplicationSystem/ReplicationWriter.h"
+
+#include "Iris/ReplicationSystem/Prioritization/NetObjectPrioritizer.h"
+
 #include "Iris/Serialization/InternalNetSerializationContext.h"
+#include "Iris/Serialization/NetBitStreamWriter.h"
+#include "Iris/Serialization/ObjectNetSerializer.h"
+
+#include "Net/Core/NetBitArray.h"
+#include "Net/Core/NetHandle/NetHandleManager.h"
+#include "Net/Core/Trace/NetDebugName.h"
+#include "Net/Core/Trace/NetTrace.h"
+
 #include "ReplicationFragmentInternal.h"
+#include "ReplicationOperationsInternal.h"
 
 #define UE_LOG_REPLICATIONBRIDGE(Category, Format, ...)  UE_LOG(LogIris, Category, TEXT("RepBridge(%u)::") Format, ReplicationSystem->GetId(), ##__VA_ARGS__)
 
@@ -333,6 +343,10 @@ void UReplicationBridge::ReadAndExecuteDestructionInfoFromRemote(FReplicationBri
 
 		CallDetachInstanceFromRemote(ReferenceToDestroy.GetRefHandle(), DestroyReason, DestroyFlags, NetFactoryId);
 	}
+	else
+	{
+		//$TODO: Need to validate that it's safe if the resolve returns false
+	}
 }
 
 void UReplicationBridge::CallSubObjectCreatedFromReplication(UE::Net::Private::FInternalNetRefIndex RootObjectIndex, FNetRefHandle SubObjectCreated)
@@ -418,32 +432,27 @@ void UReplicationBridge::Deinitialize()
 	Groups = nullptr;
 }
 
-UE::Net::FNetRefHandle UReplicationBridge::InternalCreateNetObject(FNetRefHandle AllocatedHandle, FNetHandle GlobalHandle, const UE::Net::FReplicationProtocol* ReplicationProtocol)
+UE::Net::FNetRefHandle UReplicationBridge::InternalCreateNetObject(FNetRefHandle AllocatedHandle, FNetHandle GlobalHandle, const UE::Net::Private::FCreateNetObjectParams& Params)
 {
 	check(AllocatedHandle.IsValid() && AllocatedHandle.IsCompleteHandle());
 
-	FNetRefHandle Handle = NetRefHandleManager->CreateNetObject(AllocatedHandle, GlobalHandle, ReplicationProtocol);
+	FNetRefHandle Handle = NetRefHandleManager->CreateNetObject(AllocatedHandle, GlobalHandle, Params);
 
 	if (Handle.IsValid())
 	{
-		UE_NET_TRACE_NETHANDLE_CREATED(Handle, ReplicationProtocol->DebugName, (uint64)ReplicationProtocol->ProtocolIdentifier, 0/*Local*/);
+		UE_NET_TRACE_NETHANDLE_CREATED(Handle, Params.ReplicationProtocol->DebugName, (uint64)Params.ReplicationProtocol->ProtocolIdentifier, 0/*Local*/);
 	}
 
 	return Handle;
 }
 
-UE::Net::FNetRefHandle UReplicationBridge::InternalCreateNetObject(FNetRefHandle AllocatedHandle, const UE::Net::FReplicationProtocol* ReplicationProtocol)
+UE::Net::FNetRefHandle UReplicationBridge::InternalCreateNetObjectFromRemote(FNetRefHandle WantedNetHandle, const UE::Net::Private::FCreateNetObjectParams& Params)
 {
-	return InternalCreateNetObject(AllocatedHandle, FNetHandle(), ReplicationProtocol);
-}
-
-UE::Net::FNetRefHandle UReplicationBridge::InternalCreateNetObjectFromRemote(FNetRefHandle WantedNetHandle, const UE::Net::FReplicationProtocol* ReplicationProtocol, UE::Net::FNetObjectFactoryId FactoryId)
-{
-	FNetRefHandle Handle = NetRefHandleManager->CreateNetObjectFromRemote(WantedNetHandle, ReplicationProtocol, FactoryId);
+	FNetRefHandle Handle = NetRefHandleManager->CreateNetObjectFromRemote(WantedNetHandle, Params);
 
 	if (Handle.IsValid())
 	{
-		UE_NET_TRACE_NETHANDLE_CREATED(Handle, ReplicationProtocol->DebugName, (uint64)ReplicationProtocol->ProtocolIdentifier, 1/*Remote*/);
+		UE_NET_TRACE_NETHANDLE_CREATED(Handle, Params.ReplicationProtocol->DebugName, (uint64)Params.ReplicationProtocol->ProtocolIdentifier, 1/*Remote*/);
 	}
 
 	return Handle;
@@ -1029,11 +1038,15 @@ UE::Net::FNetRefHandle UReplicationBridge::StoreDestructionInfo(FNetRefHandle Ha
 
 	FReplicationProtocolManager* ProtocolManager = GetReplicationProtocolManager();
 
-	const FNetObjectFactoryId NetFactoryId = Parameters.NetFactoryId;
+	const FCreateNetObjectParams CreateNetObjectParams
+	{
+		.NetFactoryId = Parameters.NetFactoryId,
+		.ReplicationProtocol = DestructionInfoProtocol,
+	};
 
 	// Create handle for DestructionInfo to be able to scope destruction infos
 	const FNetObjectGroupHandle LevelGroupHandle = GetLevelGroup(Parameters.Level);
-	FNetRefHandle DestructionInfoHandle = NetRefHandleManager->CreateHandleForDestructionInfo(Handle, DestructionInfoProtocol);
+	FNetRefHandle DestructionInfoHandle = NetRefHandleManager->CreateHandleForDestructionInfo(Handle, CreateNetObjectParams);
 
 	if (LevelGroupHandle.IsValid())
 	{
@@ -1047,7 +1060,7 @@ UE::Net::FNetRefHandle UReplicationBridge::StoreDestructionInfo(FNetRefHandle Ha
 	{
 		.StaticRef = FObjectReferenceCache::MakeNetObjectReference(Handle),
 		.LevelGroupHandle = LevelGroupHandle,
-		.NetFactoryId = NetFactoryId,
+		.NetFactoryId = Parameters.NetFactoryId,
 		.InternalReplicationIndex = InternalReplicationIndex
 	};
 
diff --git a/Engine/Source/Runtime/Experimental/Iris/Core/Private/Iris/ReplicationSystem/ReplicationReader.cpp b/Engine/Source/Runtime/Experimental/Iris/Core/Private/Iris/ReplicationSystem/ReplicationReader.cpp
--- a/Engine/Source/Runtime/Experimental/Iris/Core/Private/Iris/ReplicationSystem/ReplicationReader.cpp
+++ b/Engine/Source/Runtime/Experimental/Iris/Core/Private/Iris/ReplicationSystem/ReplicationReader.cpp
@@ -1,34 +1,46 @@
 // Copyright Epic Games, Inc. All Rights Reserved.
 
 #include "Iris/ReplicationSystem/ReplicationReader.h"
-#include "Iris/IrisConfigInternal.h"
+
+#include "Algo/RemoveIf.h"
+
+#include "HAL/IConsoleManager.h"
+
+#include "Iris/Core/IrisDebugging.h"
 #include "Iris/Core/IrisLog.h"
 #include "Iris/Core/IrisProfiler.h"
-#include "Iris/Core/IrisDebugging.h"
-#include "Net/Core/Trace/NetTrace.h"
-#include "Net/Core/Trace/NetDebugName.h"
+
 #include "Iris/DataStream/DataStream.h"
+
+#include "Iris/IrisConfigInternal.h"
+
 #include "Iris/ReplicationSystem/DeltaCompression/DeltaCompressionBaselineManager.h"
+#include "Iris/ReplicationSystem/DequantizeAndApplyHelper.h"
+#include "Iris/ReplicationSystem/NetBlob/NetObjectBlobHandler.h"
+#include "Iris/ReplicationSystem/NetBlob/PartialNetObjectAttachmentHandler.h"
 #include "Iris/ReplicationSystem/NetRefHandleManager.h"
-#include "Iris/Serialization/NetReferenceCollector.h"
+#include "Iris/ReplicationSystem/ObjectReferenceTypes.h"
 #include "Iris/ReplicationSystem/ReplicationOperations.h"
 #include "Iris/ReplicationSystem/ReplicationOperationsInternal.h"
-#include "Iris/ReplicationSystem/DequantizeAndApplyHelper.h"
 #include "Iris/ReplicationSystem/ReplicationSystem.h"
 #include "Iris/ReplicationSystem/ReplicationSystemInternal.h"
-#include "Iris/ReplicationSystem/NetBlob/NetObjectBlobHandler.h"
-#include "Iris/ReplicationSystem/NetBlob/PartialNetObjectAttachmentHandler.h"
+
 #include "Iris/Serialization/InternalNetSerializationContext.h"
 #include "Iris/Serialization/NetBitStreamReader.h"
 #include "Iris/Serialization/NetBitStreamUtil.h"
+#include "Iris/Serialization/NetReferenceCollector.h"
 #include "Iris/Serialization/NetSerializer.h"
-#include "HAL/IConsoleManager.h"
-#include "UObject/Class.h"
-#include "Algo/RemoveIf.h"
-#include "Misc/ScopeExit.h"
+
 #include "Misc/CoreMisc.h"
+#include "Misc/ScopeExit.h"
+
+#include "Net/Core/Trace/NetDebugName.h"
+#include "Net/Core/Trace/NetTrace.h"
+
 #include "ProfilingDebugging/CsvProfiler.h"
 
+#include "UObject/Class.h"
+
 #if UE_NET_ENABLE_REPLICATIONREADER_LOG
 #	define UE_LOG_REPLICATIONREADER(Format, ...)  UE_LOG(LogIris, Log, Format, ##__VA_ARGS__)
 #	define UE_LOG_REPLICATIONREADER_CONN(Format, ...)  UE_LOG(LogIris, Log, TEXT("Conn: %u ") Format, Parameters.ConnectionId, ##__VA_ARGS__)
@@ -731,7 +743,7 @@ void FReplicationReader::DeserializeObjectStateDelta(FNetSerializationContext& C
 	}
 }
 
-FPendingBatchData* FReplicationReader::UpdateUnresolvedMustBeMappedReferences(FNetRefHandle InHandle, TArray<FNetRefHandle>& MustBeMappedReferences)
+FPendingBatchData* FReplicationReader::UpdateUnresolvedMustBeMappedReferences(FNetRefHandle InHandle, TArray<FNetRefHandle>& MustBeMappedReferences, EIrisAsyncLoadingPriority IrisAsyncLoadingPriority)
 {
 	FPendingBatchData* PendingBatch = PendingBatches.Find(InHandle);
 	// If we already have a pending batch we append any new must be mapped references to it.
@@ -749,7 +761,15 @@ FPendingBatchData* FReplicationReader::UpdateUnresolvedMustBeMappedReferences(FN
 
 	Unresolved.Reserve(MustBeMappedReferences.Num());
 	QueuedObjectsToTrack.SetNum(MustBeMappedReferences.Num());
-	
+
+	// Assign the priority of the current object
+	ResolveContext.AsyncLoadingPriority = ConvertAsyncLoadingPriority(IrisAsyncLoadingPriority);
+	ON_SCOPE_EXIT
+	{
+		// Revert to an invalid priority to detect paths that are not setting the priority
+		ResolveContext.AsyncLoadingPriority = INDEX_NONE;
+	};
+
 	for (FNetRefHandle Handle : MustBeMappedReferences)
 	{
 		UObject* ResolvedObject = nullptr;
@@ -1008,6 +1028,7 @@ uint32 FReplicationReader::ReadObjectBatch(FNetSerializationContext& Context, ui
 	uint32 BatchEndPos = BatchEndOrStartOfExportsPos;
 
 	TempMustBeMappedReferences.Reset();
+	EIrisAsyncLoadingPriority IrisAsyncLoadingPriority = EIrisAsyncLoadingPriority::Default;
 	if (bHasExports)
 	{
 		const uint32 ReturnPos = Reader.GetPosBits();
@@ -1016,7 +1037,7 @@ uint32 FReplicationReader::ReadObjectBatch(FNetSerializationContext& Context, ui
 		Reader.Seek(BatchEndPos);
 
 		// Read exports and any must be mapped references
-		ObjectReferenceCache->ReadExports(Context, &TempMustBeMappedReferences);
+		ObjectReferenceCache->ReadExports(IncompleteHandle, Context, &TempMustBeMappedReferences, IrisAsyncLoadingPriority);
 		if (Context.HasErrorOrOverflow())
 		{
 			UE_LOG(LogIris, Error, TEXT("FReplicationReader::ReadObject Failed to read exports for handle: %s.\n%s"), *IncompleteHandle.ToString(), *Context.PrintReadJournal());
@@ -1044,10 +1065,12 @@ uint32 FReplicationReader::ReadObjectBatch(FNetSerializationContext& Context, ui
 	FPendingBatchData* PendingBatch = PendingBatches.Find(IncompleteHandle);
 
 	// This object has pending must be mapped references that must be resolved before we can process the data.
-	FPendingBatchData* PendingBatchData = ObjectReferenceCache->ShouldAsyncLoad() ? UpdateUnresolvedMustBeMappedReferences(IncompleteHandle, TempMustBeMappedReferences) : nullptr;
+	FPendingBatchData* PendingBatchData = ObjectReferenceCache->ShouldAsyncLoad() ? UpdateUnresolvedMustBeMappedReferences(IncompleteHandle, TempMustBeMappedReferences, IrisAsyncLoadingPriority) : nullptr;
 	if (PendingBatchData)
 	{
-		UE_LOG(LogIris, Verbose, TEXT("FReplicationReader::ReadObjectBatch Handle %s will be defered as it has unresolved must be mapped references"), *IncompleteHandle.ToString());
+		UE_LOG(LogIris, Verbose, TEXT("FReplicationReader::ReadObjectBatch Handle %s will be defered as it has %d unresolved must be mapped references"), *IncompleteHandle.ToString(), PendingBatchData->PendingMustBeMappedReferences.Num());
+
+		PendingBatchData->IrisAsyncLoadingPriority = IrisAsyncLoadingPriority;
 
 		UE_NET_TRACE_OBJECT_SCOPE(IncompleteHandle, Reader, Context.GetTraceCollector(), ENetTraceVerbosity::Trace);
 		UE_NET_TRACE_SCOPE(QueuedBatch, Reader, Context.GetTraceCollector(), ENetTraceVerbosity::Trace);
@@ -2662,7 +2685,7 @@ void FReplicationReader::ProcessQueuedBatches()
 
 		// Try to resolve remaining must be mapped references
 		TempMustBeMappedReferences.Reset();
-		UpdateUnresolvedMustBeMappedReferences(PendingBatchData.Handle, TempMustBeMappedReferences);
+		UpdateUnresolvedMustBeMappedReferences(PendingBatchData.Handle, TempMustBeMappedReferences, PendingBatchData.IrisAsyncLoadingPriority);
 
 		// If we have no more pending must be references we can apply the received state
 		if (PendingBatchData.PendingMustBeMappedReferences.IsEmpty())		
diff --git a/Engine/Source/Runtime/Experimental/Iris/Core/Private/Iris/ReplicationSystem/ReplicationReader.h b/Engine/Source/Runtime/Experimental/Iris/Core/Private/Iris/ReplicationSystem/ReplicationReader.h
--- a/Engine/Source/Runtime/Experimental/Iris/Core/Private/Iris/ReplicationSystem/ReplicationReader.h
+++ b/Engine/Source/Runtime/Experimental/Iris/Core/Private/Iris/ReplicationSystem/ReplicationReader.h
@@ -15,6 +15,9 @@
 
 // Forward declaration
 class UReplicationBridge;
+
+enum class EIrisAsyncLoadingPriority : uint8;
+
 namespace UE::Net
 {
 	class FNetSerializationContext;
@@ -216,7 +219,7 @@ private:
 	void DeserializeObjectStateDelta(FNetSerializationContext& Context, uint32 InternalIndex, FDispatchObjectInfo& Info, FReplicatedObjectInfo& ObjectInfo, const FNetRefHandleManager::FReplicatedObjectData& ObjectData, uint32& OutNewBaselineIndex);
 
 	// If async loading is enabled this function will verify if we can resolve all PendingMustBeMappedReferences
-	FPendingBatchData* UpdateUnresolvedMustBeMappedReferences(FNetRefHandle Handle, TArray<FNetRefHandle>& MustBeMappedReferences);
+	FPendingBatchData* UpdateUnresolvedMustBeMappedReferences(FNetRefHandle Handle, TArray<FNetRefHandle>& MustBeMappedReferences, EIrisAsyncLoadingPriority AsyncLoadingPriority);
 
 	// If we are queuing data for a batch we must also defer calls to EndReplication
 	// This method writes this method in the form of a QueuedChunk
diff --git a/Engine/Source/Runtime/Experimental/Iris/Core/Public/Iris/ReplicationSystem/ObjectReferenceTypes.h b/Engine/Source/Runtime/Experimental/Iris/Core/Public/Iris/ReplicationSystem/ObjectReferenceTypes.h
new file mode 100644
--- /dev/null
+++ b/Engine/Source/Runtime/Experimental/Iris/Core/Public/Iris/ReplicationSystem/ObjectReferenceTypes.h
@@ -0,0 +1,57 @@
+// Copyright Epic Games, Inc. All Rights Reserved.
+
+#pragma once
+
+#include "Iris/Core/BitTwiddling.h"
+
+#include "ObjectReferenceTypes.generated.h"
+
+// From UObjectGlobals.h
+typedef int32 TAsyncLoadPriority;
+
+UENUM()
+enum class EIrisAsyncLoadingPriority : uint8
+{
+	/** The default loading priority used for all types. Generally corresponds to 0 */
+	Default,
+	/** A loading priority setting for important classes. */
+	High,
+	/** A loading priority setting for critical classes. */
+	VeryHigh,
+	/** Maximum possible value to the calculate minimum bits needed to serialize the enum */
+	Max = VeryHigh,
+};
+
+static uint32 GetIrisAsyncLoadingPriorityBits()
+{
+	return UE::Net::GetBitsNeeded(EIrisAsyncLoadingPriority::Max);
+}
+
+TAsyncLoadPriority ConvertAsyncLoadingPriority(EIrisAsyncLoadingPriority IrisPriority);
+
+static const TCHAR* LexToString(EIrisAsyncLoadingPriority Priority)
+{
+	switch(Priority)
+	{
+		case EIrisAsyncLoadingPriority::Default:
+		{
+			return TEXT("Default");
+		} break;
+
+		case EIrisAsyncLoadingPriority::High:
+		{
+			return TEXT("High");
+		} break;
+
+		case EIrisAsyncLoadingPriority::VeryHigh:
+		{
+			return TEXT("VeryHigh");
+		} break;
+
+		default:
+		{
+			ensure(false);
+			return TEXT("Missing");
+		} break;
+	}
+}
\ No newline at end of file
diff --git a/Engine/Source/Runtime/Experimental/Iris/Core/Public/Iris/ReplicationSystem/ObjectReplicationBridge.h b/Engine/Source/Runtime/Experimental/Iris/Core/Public/Iris/ReplicationSystem/ObjectReplicationBridge.h
--- a/Engine/Source/Runtime/Experimental/Iris/Core/Public/Iris/ReplicationSystem/ObjectReplicationBridge.h
+++ b/Engine/Source/Runtime/Experimental/Iris/Core/Public/Iris/ReplicationSystem/ObjectReplicationBridge.h
@@ -359,6 +359,9 @@ protected:
 	 */
 	IRISCORE_API bool GetClassPollFrequency(const UClass* Class, float& OutPollPeriod) const;
 
+	/** Find the configured async loading priority for a given class.  Will return EIrisAsyncLoadingPriority::Default if no class config exist */
+	IRISCORE_API EIrisAsyncLoadingPriority GetClassIrisAsyncLoadingPriority(const UClass* Class);
+
 	/** Returns true if the class is considered critical and we force a disconnection if a protocol mismatch prevents instances of this class from replicating. */
 	IRISCORE_API bool IsClassCritical(const UClass* Class);
 
@@ -515,6 +518,9 @@ private:
 	// Classes that may force a disconnection when a protocol mismatch is detected.
 	TMap<FName, bool> ClassesFlaggedCritical;
 
+	// Classes that have an async loading priority
+	TMap<FName, EIrisAsyncLoadingPriority> ClassesIrisAsyncLoadingPriority;
+
 	// Type stats
 	TMap<FName, FName> ClassesWithTypeStats;
 
diff --git a/Engine/Source/Runtime/Experimental/Iris/Core/Public/Iris/ReplicationSystem/ObjectReplicationBridgeConfig.h b/Engine/Source/Runtime/Experimental/Iris/Core/Public/Iris/ReplicationSystem/ObjectReplicationBridgeConfig.h
--- a/Engine/Source/Runtime/Experimental/Iris/Core/Public/Iris/ReplicationSystem/ObjectReplicationBridgeConfig.h
+++ b/Engine/Source/Runtime/Experimental/Iris/Core/Public/Iris/ReplicationSystem/ObjectReplicationBridgeConfig.h
@@ -4,7 +4,9 @@
 
 #include "Containers/ContainersFwd.h"
 #include "UObject/ObjectMacros.h"
-//#include "UObject/NameTypes.h"
+
+#include "Iris/ReplicationSystem/ObjectReferenceTypes.h"
+
 #include "ObjectReplicationBridgeConfig.generated.h"
 
 USTRUCT()
@@ -102,6 +104,20 @@ struct FObjectReplicatedBridgeCriticalClassConfig
 	bool bDisconnectOnProtocolMismatch = true;
 };
 
+USTRUCT()
+struct FObjectReplicationBridgeAsyncLoadingClassConfig
+{
+	GENERATED_BODY()
+
+	/** Path names of classes that want to configure their async loading priority. */
+	UPROPERTY()
+	FName ClassName;
+
+	/** The async loading priority to use for assets reference by instances of a class */
+	UPROPERTY()
+	EIrisAsyncLoadingPriority IrisAsyncLoadingPriority = EIrisAsyncLoadingPriority::Default;
+};
+
 USTRUCT()
 struct FObjectReplicationBridgeTypeStatsConfig
 {
@@ -135,6 +151,7 @@ public:
 	IRISCORE_API TConstArrayView<FObjectReplicationBridgeDeltaCompressionConfig> GetDeltaCompressionConfigs() const;
 	IRISCORE_API TConstArrayView<FObjectReplicatedBridgeCriticalClassConfig> GetCriticalClassConfigs() const;
 	IRISCORE_API TConstArrayView<FObjectReplicationBridgeTypeStatsConfig> GetTypeStatsConfigs() const;
+	IRISCORE_API TConstArrayView<FObjectReplicationBridgeAsyncLoadingClassConfig> GetAsyncLoadingClassConfigs() const;
 
 	FName GetDefaultSpatialFilterName() const;
 	FName GetRequiredNetDriverChannelClassName() const;
@@ -194,12 +211,17 @@ private:
 	UPROPERTY(Config)
 	FName DefaultSpatialFilterName;
 
+	/** The names of classes that want to configure the priority when async loading packages of replicated references. */
+	UPROPERTY(Config)
+	TArray<FObjectReplicationBridgeAsyncLoadingClassConfig> AsyncLoadingClassConfigs;
+
 	/**
 	 * The name of the channel class required for object replication to work.
 	 */
 	UPROPERTY(Config)
 	FName RequiredNetDriverChannelClassName;
 
+	/** Classes that must be instantiated otherwise we need to disconnect from the session entirely.  */
 	UPROPERTY(Config)
 	TArray<FName> CriticalActorClasses;
 };
diff --git a/Engine/Source/Runtime/Experimental/Iris/Core/Public/Iris/ReplicationSystem/ReplicationBridge.h b/Engine/Source/Runtime/Experimental/Iris/Core/Public/Iris/ReplicationSystem/ReplicationBridge.h
--- a/Engine/Source/Runtime/Experimental/Iris/Core/Public/Iris/ReplicationSystem/ReplicationBridge.h
+++ b/Engine/Source/Runtime/Experimental/Iris/Core/Public/Iris/ReplicationSystem/ReplicationBridge.h
@@ -2,16 +2,24 @@
 
 #pragma once
 
-#include "HAL/Platform.h"
 #include "Containers/ContainersFwd.h"
 #include "Containers/Map.h"
+
+#include "HAL/Platform.h"
+
 #include "Iris/Core/NetObjectReference.h"
-#include "Iris/ReplicationSystem/NetRefHandle.h"
+
 #include "Iris/ReplicationSystem/NetObjectFactory.h"
+#include "Iris/ReplicationSystem/NetObjectFactoryRegistry.h"
 #include "Iris/ReplicationSystem/NetObjectGroupHandle.h"
+#include "Iris/ReplicationSystem/NetRefHandle.h"
+#include "Iris/ReplicationSystem/ObjectReferenceTypes.h"
 #include "Iris/ReplicationSystem/ReplicationBridgeTypes.h"
+
 #include "Misc/EnumClassFlags.h"
+
 #include "Net/Core/NetHandle/NetHandle.h"
+
 #include "UObject/ObjectKey.h"
 #include "UObject/ObjectMacros.h"
 
@@ -24,7 +32,6 @@ class UNetDriver;
 namespace UE::Net
 {
 	enum class ENetRefHandleError : uint32;
-	typedef uint8 FNetObjectFactoryId;
 
 	struct FNetDependencyInfo;
 	struct FReplicationInstanceProtocol;
@@ -49,7 +56,9 @@ namespace UE::Net
 		class FReplicationSystemImpl;
 		class FReplicationSystemInternal;
 		class FReplicationWriter;
+		
 		struct FChangeMaskCache;
+		struct FCreateNetObjectParams;
 	}
 
 	typedef TArray<FNetDependencyInfo, TInlineAllocator<32> > FNetDependencyInfoArray;
@@ -200,16 +209,14 @@ protected:
 	virtual bool CanCreateDestructionInfo() const { return true; }
 
 protected:
-	// Forward calls to internal operations that we allow replication bridges to access
 
-	/** Create a local NetRefHandle / NetObject using the ReplicationProtocol. */
-	IRISCORE_API FNetRefHandle InternalCreateNetObject(FNetRefHandle AllocatedHandle, FNetHandle GlobalHandle, const UE::Net::FReplicationProtocol* ReplicationProtocol);
+	// Forward calls to internal operations that we allow replication bridges to access
 
 	/** Create a local NetRefHandle / NetObject using the ReplicationProtocol. */
-	IRISCORE_API FNetRefHandle InternalCreateNetObject(FNetRefHandle AllocatedHandle, const UE::Net::FReplicationProtocol* ReplicationProtocol);
+	IRISCORE_API FNetRefHandle InternalCreateNetObject(FNetRefHandle AllocatedHandle, FNetHandle GlobalHandle, const UE::Net::Private::FCreateNetObjectParams& Params);
 
 	/** Create a NetRefHandle / NetObject on request from the authoritative end. */
-	IRISCORE_API FNetRefHandle InternalCreateNetObjectFromRemote(FNetRefHandle WantedNetHandle, const UE::Net::FReplicationProtocol* ReplicationProtocol,  UE::Net::FNetObjectFactoryId FactoryId);
+	IRISCORE_API FNetRefHandle InternalCreateNetObjectFromRemote(FNetRefHandle WantedNetHandle, const UE::Net::Private::FCreateNetObjectParams& Params);
 
 	/** Attach instance to NetRefHandle. */
 	IRISCORE_API void InternalAttachInstanceToNetRefHandle(FNetRefHandle RefHandle, bool bBindInstanceProtocol, UE::Net::FReplicationInstanceProtocol* InstanceProtocol, UObject* Instance, FNetHandle NetHandle);
