Chaos - Prevent debug draw for resim frames

#rnx

diff --git a/Engine/Source/Runtime/Experimental/Chaos/Private/PBDRigidsSolver.cpp b/Engine/Source/Runtime/Experimental/Chaos/Private/PBDRigidsSolver.cpp
--- a/Engine/Source/Runtime/Experimental/Chaos/Private/PBDRigidsSolver.cpp
+++ b/Engine/Source/Runtime/Experimental/Chaos/Private/PBDRigidsSolver.cpp
@@ -2683,6 +2683,12 @@ TRACE_COUNTER_SET(ChaosTraceCounter_##Name, Value)
 #if CHAOS_DEBUG_DRAW
 		QUICK_SCOPE_CYCLE_COUNTER(SolverDebugDraw);
 
+		// Don't debug draw resim frames
+		if (GetEvolution()->IsResimming())
+		{
+			return;
+		}
+
 		if (CDDScene.IsValid())
 		{
 			// @todo(chaos): this logic should be in UChaosDebugDrawSubsystem but the CVars are here. Move them and this...
