[NetPhysics] Fix KinematicTarget resim and refactor RewindData to correctly use (and only use) the PrePushData phase.
- Cache data at the start of the physics frame, inside the PrePushData phase. Previously the data was cached after game thread alterations were made on the physics thread, inside the PostPushData.
- Cache KinematicTarget just before it's used during solve but cache it inside PrePushData so that we can rewind and have the correct KinematicTarget for the frame applied while it also can be overwritten by both GT and PT data if needed.
- When a particle first goes dirty or is created, cache that particles state from before it went dirty by using pre-solve data or calculating the data back one step.
- Fix headless tests

diff --git a/Engine/Source/Programs/HeadlessChaos/Private/HeadlessChaosTestRewind.cpp b/Engine/Source/Programs/HeadlessChaos/Private/HeadlessChaosTestRewind.cpp
--- a/Engine/Source/Programs/HeadlessChaos/Private/HeadlessChaosTestRewind.cpp
+++ b/Engine/Source/Programs/HeadlessChaos/Private/HeadlessChaosTestRewind.cpp
@@ -31,16 +31,9 @@ namespace ChaosTest {
 		Solver->UpdateGameThreadStructures();
 	}
 
-	auto* CreateSolverHelper(int32 StepMode, int32 RewindHistorySize, int32 Optimization, FReal& OutSimDt)
+	template<typename T_SOLVER>
+	void EnableRewindCaptureHelper(T_SOLVER& Solver, const int32 RewindHistorySize, const int32& Optimization)
 	{
-		constexpr FReal FixedDt = 1;
-		constexpr FReal DtSizes[] = { FixedDt, FixedDt, FixedDt * 0.25, FixedDt * 4 };	//test fixed dt, sub-stepping, step collapsing
-
-		// Make a solver
-		FChaosSolversModule* Module = FChaosSolversModule::GetModule();
-		auto* Solver = Module->CreateSolver(nullptr,/*AsyncDt=*/-1);
-		InitSolverSettings(Solver);
-
 		/** int32 Optimization starts at 0 and iterate up to 2
 		* 0 = No optimizations
 		* 1 = CollisionResimCache
@@ -51,6 +44,20 @@ namespace ChaosTest {
 		{
 			RewindData->SetRewindDataOptimization(Optimization == 2); // Use optimized RewindData in as optimization 2
 		}
+	}
+
+	auto* CreateSolverHelper(int32 StepMode, int32 RewindHistorySize, int32 Optimization, FReal& OutSimDt)
+	{
+		constexpr FReal FixedDt = 1;
+		constexpr FReal DtSizes[] = { FixedDt, FixedDt, FixedDt * 0.25, FixedDt * 4 };	//test fixed dt, sub-stepping, step collapsing
+
+		// Make a solver
+		FChaosSolversModule* Module = FChaosSolversModule::GetModule();
+		auto* Solver = Module->CreateSolver(nullptr,/*AsyncDt=*/-1);
+		InitSolverSettings(Solver);
+
+		EnableRewindCaptureHelper(Solver, RewindHistorySize, Optimization);
+
 		Solver->SetThreadingMode_External(EThreadingModeTemp::SingleThread);
 		OutSimDt = DtSizes[StepMode];
 		if (StepMode > 0)
@@ -158,7 +165,9 @@ namespace ChaosTest {
 				const FReal TimeEnd = (SimStep + 1) * SimDt;
 				const FReal LastInputTime = SimDt <= 1 ? TimeStart : TimeEnd - 1;	//latest gt time associated with this interval
 
-				const auto ParticleState = RewindData->GetPastStateAtFrame(*Proxy->GetHandle_LowLevel(), SimStep);
+				/** Grab the cached data from the beginning of the next frame. When the particle is set by GT data we cache the pre-GT state in the RewindData's PrePushData phase, 
+				* the dirty GT data is separately cached as marshaled PushData and the result of the GT being applied will be cached in the next frames PrePushData. */
+				const auto ParticleState = RewindData->GetPastStateAtFrame(*Proxy->GetHandle_LowLevel(), SimStep + 1, Chaos::FFrameAndPhase::PrePushData);
 				EXPECT_EQ(ParticleState.GetX()[2], 100 - FMath::FloorToInt(LastInputTime));	//We teleported on GT so no interpolation
 
 				if (LastInputTime < 3)
@@ -475,19 +484,24 @@ namespace ChaosTest {
 
 						for (int32 Step = 0; Step < PhysicsStep; ++Step)
 						{
-							const FReal OldV = RewindData->GetPastStateAtFrame(*Proxy->GetHandle_LowLevel(), Step).GetV()[2];
+							const FReal OldV = RewindData->GetPastStateAtFrame(*Proxy->GetHandle_LowLevel(), Step, Chaos::FFrameAndPhase::PrePushData).GetV()[2];
 							const FReal Time = Step * SimDt;
-							if (Time == 2)	//velocity was reset by gt
+
+							if (Time == 0)
 							{
-								EXPECT_NEAR(OldV, 0, 1e-2);
+								// First frame is not cached since the proxy is marked as Uninitialized at that point
 							}
-							else if (Time < 2)	//simple acceleration
+							else if (Time <= 2)	//simple acceleration
 							{
 								EXPECT_NEAR(OldV, -Time, 1e-2);
 							}
-						else if (Time <= 4)	//we set velocity at time = 4, so rewind data will not see it until frame 5
+							else if (Time <= 3)	//simple acceleration after GT reset the velocity on step 2
 							{
-							if (SimDt > 2)	//if SimDt is this large, the reset of zero at time 2 is swallowed by first step so we don't really know about it
+								EXPECT_NEAR(OldV, -(Time - 2), 1e-2);
+							}
+							else if (Time <= 4)	//we set velocity at time = 4, so rewind data will not see it until frame 5
+							{
+								if (SimDt > 2)	//if SimDt is this large, the reset of zero at time 2 is swallowed by first step so we don't really know about it
 								{
 									EXPECT_NEAR(OldV, -Time, 1e-2);
 								}
@@ -537,8 +551,8 @@ namespace ChaosTest {
 				FRewindCallbackTestHelper* Helper = RegisterCallbackHelper(Solver);
 				Helper->TriggerRewindFunc = [&Floor, &bHasResimmed](int32 PhysicsStep) -> int32
 				{
-					if(!bHasResimmed && Floor)
-				{
+					if (!bHasResimmed && Floor)
+					{
 						bHasResimmed = true;
 						return 0;
 					}
@@ -557,7 +571,6 @@ namespace ChaosTest {
 							}
 						}
 
-
 						if (!bHasResimmed || Time < 4.5)
 						{
 							//simply movement without hitting floor because it's spawned too late
@@ -612,7 +625,6 @@ namespace ChaosTest {
 				FRewindCallbackTestHelper* Helper = RegisterCallbackHelper(Solver);
 				Helper->TriggerRewindFunc = [&SpawnedProxy, &SpawnedProxyNoCorrection, &bHasResimmed](int32 PhysicsStep) -> int32
 				{
-					
 					if(!bHasResimmed && SpawnedProxy)
 					{
 						bHasResimmed = true;
@@ -632,7 +644,6 @@ namespace ChaosTest {
 						{
 							RewindData->SpawnProxyIfNeeded(*SpawnedProxy);
 							SpawnedProxy->GetPhysicsThreadAPI()->SetX(FVec3(500, 0, 100.0));
-							
 						}
 					}
 
@@ -664,10 +675,8 @@ namespace ChaosTest {
 				{
 					TickSolverHelper(Solver);
 
-					//spawn floor way late to ensure no collision on first run
 					if (Step == 12)
 					{
-
 						// Make particles. E.g:
 						//	we just found out from the server these were spawned and these are the latest positions replicated from the server.
 						//	but actually these positions are the server positions from frame 10. Lets see if we can correct this.
@@ -694,8 +703,6 @@ namespace ChaosTest {
 							Solver->RegisterObject(SpawnedProxyNoCorrection);
 							ChaosTest::SetParticleSimDataToCollide({ SpawnedProxyNoCorrection->GetParticle_LowLevel() });
 						}
-						
-
 					}
 				}
 			});
@@ -1473,24 +1480,24 @@ namespace ChaosTest {
 
 						for (int32 Step = 0; Step < PhysicsStep; ++Step)
 						{
-							if (Optimization != 2) // Optimization 2 does not cache data in PrePushData so no need to check that phase
 							{
 								//Before push no force or torque at all
-								FGeometryParticleState PrePushState = RewindData->GetPastStateAtFrame(*Proxy->GetHandle_LowLevel(), Step, FFrameAndPhase::EParticleHistoryPhase::PrePushData);
+								FGeometryParticleState PrePushState = RewindData->GetPastStateAtFrame(*Proxy->GetHandle_LowLevel(), Step, Chaos::FFrameAndPhase::PrePushData);
 								EXPECT_EQ(PrePushState.Acceleration()[2], 0);
 								EXPECT_EQ(PrePushState.AngularAcceleration()[2], 0);
 							}
 
+							if (Optimization != 2) // Optimization 2 does not cache data in PostPushState so no need to check that phase
 							{
 								//After push (but before callback) only see GT force, and no torque
-								FGeometryParticleState PostPushState = RewindData->GetPastStateAtFrame(*Proxy->GetHandle_LowLevel(), Step, FFrameAndPhase::EParticleHistoryPhase::PostPushData);
+								FGeometryParticleState PostPushState = RewindData->GetPastStateAtFrame(*Proxy->GetHandle_LowLevel(), Step, Chaos::FFrameAndPhase::PostPushData);
 								EXPECT_EQ(PostPushState.Acceleration()[2], 1);	//GT always sets force of 1
 								EXPECT_EQ(PostPushState.AngularAcceleration()[2], 0);	//GT never sets torque
 							}
 
 							if (Optimization != 2) // Optimization 2 does not cache data in PostCallbacks so no need to check that phase
 							{
-							FGeometryParticleState State = RewindData->GetPastStateAtFrame(*Proxy->GetHandle_LowLevel(), Step, FFrameAndPhase::EParticleHistoryPhase::PostCallbacks);
+							FGeometryParticleState State = RewindData->GetPastStateAtFrame(*Proxy->GetHandle_LowLevel(), Step, Chaos::FFrameAndPhase::PostCallbacks);
 								if (Step == 3)
 								{
 									EXPECT_EQ(State.Acceleration()[2], 11);	//1 from GT + 10 from callback
@@ -1521,13 +1528,15 @@ namespace ChaosTest {
 	{
 		TRewindHelper::TestDynamicSphere([](auto* Solver, FReal SimDt, int32 Optimization, auto Proxy, auto Sphere)
 		{
-				auto& Particle = Proxy->GetGameThreadAPI();
+			auto& Particle = Proxy->GetGameThreadAPI();
+			Particle.SetGravityEnabled(false);
+
 			const int32 LastGameStep = 20;
 
 			for (int Step = 0; Step <= LastGameStep; ++Step)
 			{
 				//sim-writable property that changes every step
-					Particle.AddForce(FVec3(0, 0, Step + 1));
+				Particle.AddForce(FVec3(0, 0, Step + 1));
 				TickSolverHelper(Solver);
 			}
 
@@ -1535,10 +1544,15 @@ namespace ChaosTest {
 
 			//check state at every step except latest
 			const int32 LastSimStep = LastGameStep / SimDt;
+			FReal ExpectedX = 0;
+			FReal ExpectedVelocity = 0;
+			FReal PreviousForce = 0;
+
 			for (int Step = 0; Step < LastSimStep - 1; ++Step)
 			{
 				const auto ParticleState = RewindData->GetPastStateAtFrame(*Proxy->GetHandle_LowLevel(), Step);
 				FReal ExpectedForce = Step + 1;
+
 				if (SimDt < 1)
 				{
 					//each sub-step gets a constant force applied
@@ -1550,7 +1564,23 @@ namespace ChaosTest {
 					//where step is game step: so really it's step * 4
 					ExpectedForce = Step * 4 + 2.5;
 				}
-				EXPECT_EQ(ParticleState.Acceleration()[2], ExpectedForce);
+
+				ExpectedVelocity += PreviousForce * SimDt;
+				ExpectedX += ExpectedVelocity * SimDt;
+
+				EXPECT_EQ(ParticleState.GetV()[2], ExpectedVelocity);
+				EXPECT_EQ(ParticleState.GetX()[2], ExpectedX);
+
+				/** With Optimization 2 we only cache PrePushData where Acceleration is always 0. Calling GetPastStateAtFrame for PostPushData will return next frames PrePushData since it's the next available data */
+				if (Optimization == 2)
+				{
+					EXPECT_EQ(ParticleState.Acceleration()[2], 0);
+				}
+				else
+				{
+					EXPECT_EQ(ParticleState.Acceleration()[2], ExpectedForce);
+				}
+				PreviousForce = ExpectedForce;
 			}
 		});
 	}
@@ -1559,7 +1589,8 @@ namespace ChaosTest {
 	{
 		TRewindHelper::TestDynamicSphere([](auto* Solver, FReal SimDt, int32 Optimization, auto Proxy, auto Sphere)
 		{
-				auto& Particle = Proxy->GetGameThreadAPI();
+			auto& Particle = Proxy->GetGameThreadAPI();
+			Particle.SetGravityEnabled(false);
 			const int32 LastGameStep = 20;
 
 			for (int Step = 0; Step <= LastGameStep; ++Step)
@@ -1567,39 +1598,54 @@ namespace ChaosTest {
 				//sim-writable property that changes infrequently and not at beginning
 				if (Step == 3)
 				{
-						Particle.AddForce(FVec3(0, 0, Step));
+					Particle.AddForce(FVec3(0, 0, Step));
 				}
 
 				if (Step == 5)
 				{
-						Particle.AddForce(FVec3(0, 0, Step));
+					Particle.AddForce(FVec3(0, 0, Step));
 				}
 
 				TickSolverHelper(Solver);
 			}
 
 			const FRewindData* RewindData = Solver->GetRewindData();
+			FReal ExpectedX = 0;
+			FReal ExpectedVelocity = 0;
+			FReal PreviousForce = 0;
 
 			//check state at every step except latest
 			const int32 LastSimStep = LastGameStep / SimDt;
 			for (int Step = 0; Step < LastSimStep - 1; ++Step)
 			{
-				const auto ParticleState = RewindData->GetPastStateAtFrame(*Proxy->GetHandle_LowLevel(), Step);
+				const float SimTime = Step * SimDt;
+
+				const auto ParticleState = RewindData->GetPastStateAtFrame(*Proxy->GetHandle_LowLevel(), Step, (Optimization == 2 ? Chaos::FFrameAndPhase::PrePushData : Chaos::FFrameAndPhase::PostPushData));
+
+				FReal ExpectedForce = 0;
 
+				ExpectedVelocity += PreviousForce * SimDt;
+				ExpectedX += ExpectedVelocity * SimDt;
+
+				EXPECT_EQ(ParticleState.GetV()[2], ExpectedVelocity);
+				EXPECT_EQ(ParticleState.GetX()[2], ExpectedX);
+
+				/** With Optimization 2 we only cache PrePushData where Acceleration is always 0. */
 				if (SimDt <= 1)
 				{
-					const float SimTime = Step * SimDt;
 					if (SimTime >= 3 && SimTime < 4)
 					{
-						EXPECT_EQ(ParticleState.Acceleration()[2], 3);
+						ExpectedForce = 3;
+						EXPECT_EQ(ParticleState.Acceleration()[2], (Optimization == 2) ? 0.0f : ExpectedForce);
 					}
 					else if (SimTime >= 5 && SimTime < 6)
 					{
-						EXPECT_EQ(ParticleState.Acceleration()[2], 5);
+						ExpectedForce = 5;
+						EXPECT_EQ(ParticleState.Acceleration()[2], (Optimization == 2) ? 0.0f : ExpectedForce);
 					}
 					else
 					{
-						EXPECT_EQ(ParticleState.Acceleration()[2], 0);
+						EXPECT_EQ(ParticleState.Acceleration()[2], (Optimization == 2) ? 0.0f : ExpectedForce);
 					}
 				}
 				else
@@ -1607,18 +1653,20 @@ namespace ChaosTest {
 					//we get an average
 					if (Step == 0)
 					{
-						EXPECT_EQ(ParticleState.Acceleration()[2], 3 / 4.f);
+						ExpectedForce = 3 / 4.f;
+						EXPECT_EQ(ParticleState.Acceleration()[2], (Optimization == 2) ? 0.0f : ExpectedForce);
 					}
 					else if (Step == 1)
 					{
-						EXPECT_EQ(ParticleState.Acceleration()[2], 5 / 4.f);
+						ExpectedForce = 5 / 4.f;
+						EXPECT_EQ(ParticleState.Acceleration()[2], (Optimization == 2) ? 0.0f : ExpectedForce);
 					}
 					else
 					{
-						EXPECT_EQ(ParticleState.Acceleration()[2], 0);
+						EXPECT_EQ(ParticleState.Acceleration()[2], (Optimization == 2) ? 0.0f : ExpectedForce);
 					}
 				}
-
+				PreviousForce = ExpectedForce;
 			}
 		});
 	}
@@ -1660,7 +1708,9 @@ namespace ChaosTest {
 			const int32 LastSimStep = LastGameStep / SimDt;
 			for (int Step = 0; Step < LastSimStep - 1; ++Step)
 			{
-				const auto ParticleState = RewindData->GetPastStateAtFrame(*Proxy->GetHandle_LowLevel(), Step);
+				/** Grab the cached data from the beginning of the next frame. When the particle is set by GT data we cache the pre-GT state in the RewindData's PrePushData phase,
+				* the dirty GT data is separately cached as marshaled PushData and the result of the GT being applied will be cached in the next frames PrePushData. */
+				const auto ParticleState = RewindData->GetPastStateAtFrame(*Proxy->GetHandle_LowLevel(), Step + 1, Chaos::FFrameAndPhase::PrePushData);
 				if (SimDt <= 1)
 				{
 					const float SimTime = Step * SimDt;
@@ -1693,19 +1743,19 @@ namespace ChaosTest {
 	{
 		TRewindHelper::TestDynamicSphere([](auto* Solver, FReal SimDt, int32 Optimization, auto Proxy, auto Sphere)
 		{
-				auto& Particle = Proxy->GetGameThreadAPI();
-				Solver->GetEvolution()->GetGravityForces().SetAcceleration(FVec3(0, 0, -1), 0);
-				Particle.SetGravityEnabled(true);
-				Particle.SetX(FVec3(0, 0, 100));
+			auto& Particle = Proxy->GetGameThreadAPI();
+			Solver->GetEvolution()->GetGravityForces().SetAcceleration(FVec3(0, 0, -1), 0);
+			Particle.SetGravityEnabled(true);
+			Particle.SetX(FVec3(0, 0, 100));
 
 			const int32 LastGameStep = 20;
-				for (int Step = 0; Step <= LastGameStep; ++Step)
+			for (int Step = 0; Step <= LastGameStep; ++Step)
 			{
 				//teleport from GT
-					if (Step == 5)
+				if (Step == 5)
 				{
-						Particle.SetX(FVec3(0, 0, 10));
-						Particle.SetV(FVec3(0, 0, 0));
+					Particle.SetX(FVec3(0, 0, 10));
+					Particle.SetV(FVec3(0, 0, 0));
 				}
 
 				TickSolverHelper(Solver);
@@ -1720,19 +1770,37 @@ namespace ChaosTest {
 
 			for (int Step = 0; Step < LastSimStep - 1; ++Step)
 			{
-				const auto ParticleState = RewindData->GetPastStateAtFrame(*Proxy->GetHandle_LowLevel(), Step);
+				const auto ParticleState = RewindData->GetPastStateAtFrame(*Proxy->GetHandle_LowLevel(), Step, (Optimization == 2 ? Chaos::FFrameAndPhase::PrePushData : Chaos::FFrameAndPhase::PostPushData));
 
 				const FReal SimStart = SimDt * Step;
 				const FReal SimEnd = SimDt * (Step + 1);
-					if (SimStart <= 5 && SimEnd > 5)
+
+				if (Optimization == 2)
 				{
-					ExpectedVZ = 0;
-					ExpectedXZ = 10;
+					/* Teleport was made on GT and is cached in PushData when Optimization 2 is enabled so we will have the first RewindData cached on the next physics frame after the teleport
+					* Track expected V and X from the next frame by adding SimDt to the teleport frame and then step the expected V and X once*/
+					if (SimStart <= (5 + SimDt) && SimEnd > (5 + SimDt))
+					{
+						ExpectedVZ = 0;
+						ExpectedXZ = 10;
+
+						ExpectedVZ -= SimDt;
+						ExpectedXZ += ExpectedVZ * SimDt;
+					}
+				}
+				else
+				{
+					if (SimStart <= 5 && SimEnd > 5)
+					{
+						ExpectedVZ = 0;
+						ExpectedXZ = 10;
+					}
 				}
 
 				EXPECT_NEAR(ParticleState.GetX()[2], ExpectedXZ, 1e-4);
 				EXPECT_NEAR(ParticleState.GetV()[2], ExpectedVZ, 1e-4);
 
+				// Step expected V and X once
 				ExpectedVZ -= SimDt;
 				ExpectedXZ += ExpectedVZ * SimDt;
 			}
@@ -2353,7 +2421,7 @@ namespace ChaosTest {
 
 	GTEST_TEST(AllTraits, RewindTest_NumDirty)
 	{
-		for (int Optimization = 0; Optimization < 2; ++Optimization)
+		for (int Optimization = 0; Optimization < 3; ++Optimization)
 		{
 			auto Sphere = Chaos::FImplicitObjectPtr(new Chaos::FSphere(FVec3(0), 10));
 
@@ -2364,7 +2432,7 @@ namespace ChaosTest {
 			InitSolverSettings(Solver);
 
 			//note: this 5 is just a suggestion, there could be more frames saved than that
-			Solver->EnableRewindCapture(5, !!Optimization);
+			EnableRewindCaptureHelper(Solver, 5, Optimization);
 
 			// Make particles
 			auto Proxy = FSingleParticlePhysicsProxy::Create(Chaos::FPBDRigidParticle::CreateParticle());
@@ -2418,7 +2486,7 @@ namespace ChaosTest {
 
 	GTEST_TEST(AllTraits, RewindTest_Resim)
 	{
-		for (int Optimization = 0; Optimization < 2; ++Optimization)
+		for (int Optimization = 0; Optimization < 3; ++Optimization)
 		{
 			auto Sphere = Chaos::FImplicitObjectPtr(new Chaos::FSphere(FVec3(0), 10));
 
@@ -2430,7 +2498,7 @@ namespace ChaosTest {
 
 			InitSolverSettings(Solver);
 
-			Solver->EnableRewindCapture(5, !!Optimization);
+			EnableRewindCaptureHelper(Solver, 5, Optimization);
 
 			// Make particles
 			auto Proxy = FSingleParticlePhysicsProxy::Create(Chaos::FPBDRigidParticle::CreateParticle());
@@ -2572,7 +2640,7 @@ namespace ChaosTest {
 
 	GTEST_TEST(AllTraits, RewindTest_ResimDesyncAfterMissingTeleport)
 	{
-		for (int Optimization = 0; Optimization < 2; ++Optimization)
+		for (int Optimization = 0; Optimization < 3; ++Optimization)
 		{
 			auto Sphere = Chaos::FImplicitObjectPtr(new Chaos::FSphere(FVec3(0), 10));
 
@@ -2582,7 +2650,7 @@ namespace ChaosTest {
 			auto* Solver = Module->CreateSolver(nullptr, /*AsyncDt=*/-1);
 			InitSolverSettings(Solver);
 
-			Solver->EnableRewindCapture(7, !!Optimization);
+			EnableRewindCaptureHelper(Solver, 7, Optimization);
 
 			// Make particles
 			auto Proxy = FSingleParticlePhysicsProxy::Create(Chaos::FPBDRigidParticle::CreateParticle());
@@ -2669,7 +2737,7 @@ namespace ChaosTest {
 
 	GTEST_TEST(AllTraits, RewindTest_ResimDesyncAfterChangingMass)
 	{
-		for (int Optimization = 0; Optimization < 2; ++Optimization)
+		for (int Optimization = 0; Optimization < 3; ++Optimization)
 		{
 			auto Sphere = Chaos::FImplicitObjectPtr(new Chaos::FSphere(FVec3(0), 10));
 
@@ -2679,7 +2747,7 @@ namespace ChaosTest {
 			auto* Solver = Module->CreateSolver(nullptr, /*AsyncDt=*/-1);
 			InitSolverSettings(Solver);
 
-			Solver->EnableRewindCapture(7, !!Optimization);
+			EnableRewindCaptureHelper(Solver, 7, Optimization);
 
 			FReal CurMass = 1.0;
 			int32 LastStep = 11;
@@ -2760,7 +2828,7 @@ namespace ChaosTest {
 	{
 #if REWIND_DESYNC
 
-		for (int Optimization = 0; Optimization < 2; ++Optimization)
+		for (int Optimization = 0; Optimization < 3; ++Optimization)
 		{
 			//We want to detect when sim results change
 			//Detecting output of position and velocity is expensive and hard to track
@@ -2774,7 +2842,7 @@ namespace ChaosTest {
 			auto* Solver = Module->CreateSolver(nullptr, /*AsyncDt=*/-1);
 			InitSolverSettings(Solver);
 
-			Solver->EnableRewindCapture(7, !!Optimization);
+			EnableRewindCaptureHelper(Solver, 7, Optimization);
 
 			// Make particles
 
@@ -2862,7 +2930,7 @@ namespace ChaosTest {
 
 	GTEST_TEST(AllTraits, DISABLED_RewindTest_ResimDesyncFromChangeForce)
 	{
-		for (int Optimization = 0; Optimization < 2; ++Optimization)
+		for (int Optimization = 0; Optimization < 3; ++Optimization)
 		{
 			auto Sphere = Chaos::FImplicitObjectPtr(new Chaos::FSphere(FVec3(0), 10));
 
@@ -2872,7 +2940,7 @@ namespace ChaosTest {
 			auto* Solver = Module->CreateSolver(nullptr, /*AsyncDt=*/-1);
 			InitSolverSettings(Solver);
 
-			Solver->EnableRewindCapture(7, !!Optimization);
+			EnableRewindCaptureHelper(Solver, 7, Optimization);
 
 			// Make particles
 			auto Proxy = FSingleParticlePhysicsProxy::Create(Chaos::FPBDRigidParticle::CreateParticle());
@@ -2955,7 +3023,7 @@ namespace ChaosTest {
 	GTEST_TEST(AllTraits, RewindTest_ResimAsFollower)
 	{
 #if REWIND_DESYNC
-		for (int Optimization = 0; Optimization < 2; ++Optimization)
+		for (int Optimization = 0; Optimization < 3; ++Optimization)
 		{
 			auto Sphere = Chaos::FImplicitObjectPtr(new Chaos::FSphere(FVec3(0), 10));
 			auto Box = Chaos::FImplicitObjectPtr(new TBox<FReal, 3>(FVec3(-100, -100, -100), FVec3(100, 100, 0)));
@@ -2966,7 +3034,7 @@ namespace ChaosTest {
 			auto* Solver = Module->CreateSolver(nullptr, /*AsyncDt=*/-1);
 			InitSolverSettings(Solver);
 
-			Solver->EnableRewindCapture(7, !!Optimization);
+			EnableRewindCaptureHelper(Solver, 7, Optimization);
 
 			// Make particles
 			auto DynamicProxy = FSingleParticlePhysicsProxy::Create(Chaos::FPBDRigidParticle::CreateParticle());
@@ -3046,7 +3114,7 @@ namespace ChaosTest {
 
 	GTEST_TEST(AllTraits, DISABLED_RewindTest_FullResimFallSeeCollisionCorrection)
 	{
-		for (int Optimization = 0; Optimization < 2; ++Optimization)
+		for (int Optimization = 0; Optimization < 3; ++Optimization)
 		{
 			auto Sphere = Chaos::FImplicitObjectPtr(new Chaos::FSphere(FVec3(0), 10));
 			auto Box = Chaos::FImplicitObjectPtr(new TBox<FReal, 3>(FVec3(-100, -100, -100), FVec3(100, 100, 0)));
@@ -3057,7 +3125,7 @@ namespace ChaosTest {
 			auto* Solver = Module->CreateSolver(nullptr, /*AsyncDt=*/-1);
 			InitSolverSettings(Solver);
 
-			Solver->EnableRewindCapture(100, !!Optimization);
+			EnableRewindCaptureHelper(Solver, 100, Optimization);
 
 			// Make particles
 			auto DynamicProxy = FSingleParticlePhysicsProxy::Create(Chaos::FPBDRigidParticle::CreateParticle());
@@ -3140,7 +3208,7 @@ namespace ChaosTest {
 
 	GTEST_TEST(AllTraits, DISABLED_RewindTest_ResimAsFollowerFallIgnoreCollision)
 	{
-		for (int Optimization = 0; Optimization < 2; ++Optimization)
+		for (int Optimization = 0; Optimization < 3; ++Optimization)
 		{
 			auto Sphere = Chaos::FImplicitObjectPtr(new Chaos::FSphere(FVec3(0), 10));
 			auto Box = Chaos::FImplicitObjectPtr(new TBox<FReal, 3>(FVec3(-100, -100, -100), FVec3(100, 100, 0)));
@@ -3151,7 +3219,7 @@ namespace ChaosTest {
 			auto* Solver = Module->CreateSolver(nullptr, /*AsyncDt=*/-1);
 			InitSolverSettings(Solver);
 
-			Solver->EnableRewindCapture(100, !!Optimization);
+			EnableRewindCaptureHelper(Solver, 100, Optimization);
 
 			// Make particles
 			auto DynamicProxy = FSingleParticlePhysicsProxy::Create(Chaos::FPBDRigidParticle::CreateParticle());
@@ -3231,7 +3299,7 @@ namespace ChaosTest {
 	GTEST_TEST(AllTraits, RewindTest_ResimAsFollowerWithForces)
 	{
 #if REWIND_DESYNC
-		for (int Optimization = 0; Optimization < 2; ++Optimization)
+		for (int Optimization = 0; Optimization < 3; ++Optimization)
 		{
 			auto Box = Chaos::FImplicitObjectPtr(new TBox<FReal, 3>(FVec3(-10, -10, -10), FVec3(10, 10, 10)));
 
@@ -3241,7 +3309,7 @@ namespace ChaosTest {
 			auto* Solver = Module->CreateSolver(nullptr, /*AsyncDt=*/-1);
 			InitSolverSettings(Solver);
 
-			Solver->EnableRewindCapture(7, !!Optimization);
+			EnableRewindCaptureHelper(Solver, 7, Optimization);
 
 			// Make particles
 			auto FullSimProxy = FSingleParticlePhysicsProxy::Create(Chaos::FPBDRigidParticle::CreateParticle());
@@ -3311,7 +3379,7 @@ namespace ChaosTest {
 	GTEST_TEST(AllTraits, RewindTest_ResimAsFollowerWokenUp)
 	{
 #if REWIND_DESYNC
-		for (int Optimization = 0; Optimization < 2; ++Optimization)
+		for (int Optimization = 0; Optimization < 3; ++Optimization)
 		{
 			auto Box = Chaos::FImplicitObjectPtr(new TBox<FReal, 3>(FVec3(-10, -10, -10), FVec3(10, 10, 10)));
 
@@ -3321,7 +3389,7 @@ namespace ChaosTest {
 			auto* Solver = Module->CreateSolver(nullptr, /*AsyncDt=*/-1);
 			InitSolverSettings(Solver);
 
-			Solver->EnableRewindCapture(7, !!Optimization);
+			EnableRewindCaptureHelper(Solver, 7, Optimization);
 
 			// Make particles
 			auto ImpulsedObjProxy = FSingleParticlePhysicsProxy::Create(Chaos::FPBDRigidParticle::CreateParticle());
@@ -3401,7 +3469,7 @@ namespace ChaosTest {
 	GTEST_TEST(AllTraits, RewindTest_ResimAsFollowerWokenUpNoHistory)
 	{
 #if REWIND_DESYNC
-		for (int Optimization = 0; Optimization < 2; ++Optimization)
+		for (int Optimization = 0; Optimization < 3; ++Optimization)
 		{
 			auto Box = Chaos::FImplicitObjectPtr(new TBox<FReal, 3>(FVec3(-10, -10, -10), FVec3(10, 10, 10)));
 
@@ -3411,7 +3479,7 @@ namespace ChaosTest {
 			auto* Solver = Module->CreateSolver(nullptr, /*AsyncDt=*/-1);
 			InitSolverSettings(Solver);
 
-			Solver->EnableRewindCapture(7, !!Optimization);
+			EnableRewindCaptureHelper(Solver, 7, Optimization);
 
 			// Make particles
 			auto ImpulsedObjProxy = FSingleParticlePhysicsProxy::Create(Chaos::FPBDRigidParticle::CreateParticle());
@@ -3491,7 +3559,7 @@ namespace ChaosTest {
 
 	GTEST_TEST(AllTraits, DISABLED_RewindTest_DesyncSimOutOfCollision)
 	{
-		for (int Optimization = 0; Optimization < 2; ++Optimization)
+		for (int Optimization = 0; Optimization < 3; ++Optimization)
 		{
 			auto Sphere = Chaos::FImplicitObjectPtr(new Chaos::FSphere(FVec3(0), 10));
 			auto Box = Chaos::FImplicitObjectPtr(new TBox<FReal, 3>(FVec3(-100, -100, -100), FVec3(100, 100, 0)));
@@ -3502,7 +3570,7 @@ namespace ChaosTest {
 			auto* Solver = Module->CreateSolver(nullptr, /*AsyncDt=*/-1);
 			InitSolverSettings(Solver);
 
-			Solver->EnableRewindCapture(100, !!Optimization);
+			EnableRewindCaptureHelper(Solver, 100, Optimization);
 
 			// Make particles
 			auto DynamicProxy = FSingleParticlePhysicsProxy::Create(Chaos::FPBDRigidParticle::CreateParticle());
@@ -3833,12 +3901,12 @@ namespace ChaosTest {
 	{
 		//TODO: disabled because at the moment GC particles are always marked as dirty - this messes up with transient dirty during rewind
 		//Should probably rethink why GC has its own dirty view
-		for (int Optimization = 0; Optimization < 2; ++Optimization)
+		for (int Optimization = 0; Optimization < 3; ++Optimization)
 		{
 			FGeometryCollectionWrapper* Collection = TNewSimulationObject<GeometryType::GeometryCollectionWithSingleRigid>::Init()->As<FGeometryCollectionWrapper>();
 
 			FFramework UnitTest;
-			UnitTest.Solver->EnableRewindCapture(100, !!Optimization);
+			EnableRewindCaptureHelper(UnitTest.Solver, 100, Optimization);
 			UnitTest.AddSimulationObject(Collection);
 			UnitTest.Initialize();
 
diff --git a/Engine/Source/Runtime/Engine/Private/PhysicsEngine/PhysicsReplication.cpp b/Engine/Source/Runtime/Engine/Private/PhysicsEngine/PhysicsReplication.cpp
--- a/Engine/Source/Runtime/Engine/Private/PhysicsEngine/PhysicsReplication.cpp
+++ b/Engine/Source/Runtime/Engine/Private/PhysicsEngine/PhysicsReplication.cpp
@@ -1051,7 +1051,7 @@ void FPhysicsReplicationAsync::UpdateRewindDataTarget(const FPhysicsRepAsyncInpu
 	{
 		// Cache all target states inside RewindData
 		const int32 LocalFrame = Input.ServerFrame - *Input.FrameOffset;
-		RewindData->SetTargetStateAtFrame(*Handle, LocalFrame, Chaos::FFrameAndPhase::EParticleHistoryPhase::PostPushData,
+		RewindData->SetTargetStateAtFrame(*Handle, LocalFrame, Chaos::FFrameAndPhase::EParticleHistoryPhase::PrePushData,
 			Input.TargetState.Position, Input.TargetState.Quaternion,
 			Input.TargetState.LinVel, FMath::DegreesToRadians(Input.TargetState.AngVel), (Input.TargetState.Flags & ERigidBodyFlags::Sleeping));
 	}
@@ -2187,7 +2187,7 @@ bool FPhysicsReplicationAsync::ResimulationReplication(Chaos::FPBDRigidParticleH
 	const bool bShouldSleep = (Target.TargetState.Flags & ERigidBodyFlags::Sleeping) != 0;
 	bool bClearTarget = true;
 
-	static constexpr Chaos::FFrameAndPhase::EParticleHistoryPhase RewindPhase = Chaos::FFrameAndPhase::EParticleHistoryPhase::PostPushData;
+	static constexpr Chaos::FFrameAndPhase::EParticleHistoryPhase RewindPhase = Chaos::FFrameAndPhase::EParticleHistoryPhase::PrePushData;
 	
 	// Get state from locally cached history for frame corresponding to received data
 	const Chaos::FGeometryParticleState PastState = RewindData->GetPastStateAtFrame(*Handle, LocalFrame, RewindPhase);
diff --git a/Engine/Source/Runtime/Experimental/Chaos/Private/Chaos/PBDRigidsEvolutionGBF.cpp b/Engine/Source/Runtime/Experimental/Chaos/Private/Chaos/PBDRigidsEvolutionGBF.cpp
--- a/Engine/Source/Runtime/Experimental/Chaos/Private/Chaos/PBDRigidsEvolutionGBF.cpp
+++ b/Engine/Source/Runtime/Experimental/Chaos/Private/Chaos/PBDRigidsEvolutionGBF.cpp
@@ -537,6 +537,11 @@ void FPBDRigidsEvolutionGBF::AdvanceOneTimeStepImpl(const FReal Dt, const FSubSt
 	}
 
 	{
+		if (CaptureRewindKinematicTarget)
+		{
+			CaptureRewindKinematicTarget(Particles.GetActiveMovingKinematicParticlesView());
+		}
+
 		SCOPE_CYCLE_COUNTER(STAT_Evolution_KinematicTargets);
 		CSV_SCOPED_TIMING_STAT(PhysicsVerbose, StepSolver_KinematicTargets);
 		ApplyKinematicTargets(Dt, SubStepInfo.PseudoFraction);
@@ -746,7 +751,7 @@ void FPBDRigidsEvolutionGBF::AdvanceOneTimeStepImpl(const FReal Dt, const FSubSt
 
 	Clustering.AdvanceClustering(Dt, GetCollisionConstraints());
 
-	if(CaptureRewindData)
+	if (CaptureRewindData)
 	{
 		CaptureRewindData(Particles.GetDirtyParticlesView());
 	}
diff --git a/Engine/Source/Runtime/Experimental/Chaos/Private/PBDRigidsSolver.cpp b/Engine/Source/Runtime/Experimental/Chaos/Private/PBDRigidsSolver.cpp
--- a/Engine/Source/Runtime/Experimental/Chaos/Private/PBDRigidsSolver.cpp
+++ b/Engine/Source/Runtime/Experimental/Chaos/Private/PBDRigidsSolver.cpp
@@ -1219,16 +1219,11 @@ namespace Chaos
 		//todo: do we need this?
 		//MarshallingManager.Reset();
 
-		if (bUseCollisionResimCache)
+		if (GetRewindData())
 		{
-			EnableRewindCapture(true);
+			EnableRewindCapture(GetRewindData()->Capacity(), GetUseCollisionResimCache());
 		}
 
-		MEvolution->SetCaptureRewindDataFunction([this](const TParticleView<TPBDRigidParticles<FReal,3>>& ActiveParticles)
-		{
-			FinalizeRewindData(ActiveParticles);
-		});
-
 		FEventDefaults::RegisterSystemEvents(*GetEventManager());
 	}
 
@@ -1981,13 +1976,13 @@ namespace Chaos
 			}
 		}
 
-		ProcessSinglePushedData_Internal(PushData);
-
 		if (FRewindData* RewindData = GetRewindData())
 		{
-			RewindData->CacheCurrentDirtyData(FFrameAndPhase::PostPushData);
+			RewindData->CacheCurrentDirtyData(FFrameAndPhase::PrePushData);
 		}
 
+		ProcessSinglePushedData_Internal(PushData);
+
 		//run any commands passed in. These don't generate outputs and are a one off so just do them here
 		//note: commands run before sim callbacks. This is important for sub-stepping since we want each sub-step to have a consistent view
 		//so for example if the user deletes a floor surface, we want all sub-steps to see that in the same way
@@ -2877,24 +2872,6 @@ TRACE_COUNTER_SET(ChaosTraceCounter_##Name, Value)
 		QueryMaterialMasks_External = SimMaterialMasks;
 	}
 
-
-	void FPBDRigidsSolver::FinalizeRewindData(const TParticleView<TPBDRigidParticles<FReal,3>>& DirtyParticles)
-	{
-		using namespace Chaos;
-
-		//Simulated objects must have their properties captured for rewind
-		if(MRewindData && DirtyParticles.Num())
-		{
-			QUICK_SCOPE_CYCLE_COUNTER(RecordRewindData);
-			
-			int32 DataIdx = 0;
-			for(TPBDRigidParticleHandleImp<FReal,3,false>& DirtyParticle : DirtyParticles)
-			{
-				MRewindData->PushPTDirtyData(*DirtyParticle.Handle(), DataIdx++);
-			}
-		}
-	}
-
 	void FPBDRigidsSolver::UpdateExternalAccelerationStructure_External(ISpatialAccelerationCollection<FAccelerationStructureHandle,FReal,3>*& ExternalStructure)
 	{
 		GetEvolution()->UpdateExternalAccelerationStructure_External(ExternalStructure,*PendingSpatialOperations_External);
diff --git a/Engine/Source/Runtime/Experimental/Chaos/Private/RewindData.cpp b/Engine/Source/Runtime/Experimental/Chaos/Private/RewindData.cpp
--- a/Engine/Source/Runtime/Experimental/Chaos/Private/RewindData.cpp
+++ b/Engine/Source/Runtime/Experimental/Chaos/Private/RewindData.cpp
@@ -142,7 +142,7 @@ bool FJointStateBase::IsInSync(const FPBDJointConstraintHandle& Handle, const FF
 
 
 bool bCVarRewindDataOptimization = true;
-FAutoConsoleVariableRef CVarRewindDataOptimization(TEXT("p.Resim.RewindDataOptimization"), bCVarRewindDataOptimization, TEXT("Default value for RewinData optimization, note that this can be overridden at runtime by API calls. Effect: Only alter the minimum required properties during a resim for particles not marked for FullResim and only cache data during the PostPushData phase and lower memory allocation for the history cache to 1/3 of non-optimized flow."));
+FAutoConsoleVariableRef CVarRewindDataOptimization(TEXT("p.Resim.RewindDataOptimization"), bCVarRewindDataOptimization, TEXT("Default value for RewinData optimization, note that this can be overridden at runtime by API calls. Effect: Only alter the minimum required properties during a resim for particles not marked for FullResim and only cache data during the PrePushData phase and lower memory allocation for the history cache to 1/3 of non-optimized flow."));
 
 FRewindData::FRewindData(FPBDRigidsSolver* InSolver, int32 NumFrames, bool InRewindDataOptimization, int32 InCurrentFrame)
 	: Managers(NumFrames + 1)	//give 1 extra for saving at head
@@ -155,6 +155,7 @@ FRewindData::FRewindData(FPBDRigidsSolver* InSolver, int32 NumFrames, bool InRew
 	, bRewindDataOptimization(InRewindDataOptimization)
 	, LatestTargetFrame(0)
 {
+	RegisterEvolutionCallbacks();
 }
 
 FRewindData::FRewindData(FPBDRigidsSolver* InSolver, int32 NumFrames, int32 InCurrentFrame)
@@ -168,8 +169,31 @@ FRewindData::FRewindData(FPBDRigidsSolver* InSolver, int32 NumFrames, int32 InCu
 	, bRewindDataOptimization(bCVarRewindDataOptimization)
 	, LatestTargetFrame(0)
 {
+	RegisterEvolutionCallbacks();
 }
 
+void FRewindData::RegisterEvolutionCallbacks()
+{
+	if (!ensureMsgf(Solver, TEXT("FRewindData::RegisterEvolutionCallbacks No Solver")))
+	{
+		return;
+	}
+
+	if (!ensureMsgf(Solver->GetEvolution(), TEXT("FRewindData::RegisterEvolutionCallbacks No Evolution")))
+	{
+		return;
+	}
+
+	Solver->GetEvolution()->SetCaptureRewindKinematicTargetFunction([this](const TParticleView<TPBDRigidParticles<FReal, 3>>& ActiveKinematicParticles)
+		{
+			ProcessDirtyKinematicTargets(ActiveKinematicParticles);
+		});
+
+	Solver->GetEvolution()->SetCaptureRewindDataFunction([this](const TParticleView<TPBDRigidParticles<FReal, 3>>& ActiveParticles)
+		{
+			ProcessDirtyPTParticles(ActiveParticles);
+		});
+}
 
 void FRewindData::ApplyInputs(const int32 ApplyFrame, const bool bResetSolver)
 {
@@ -207,7 +231,7 @@ void FRewindData::ApplyTargets(const int32 Frame, const bool bResetSimulation)
 	EnsureIsInPhysicsThreadContext();
 
 	//If property changed between Frame and CurFrame, record the latest value and rewind to old
-	FFrameAndPhase RewindFrameAndPhase{ Frame, FFrameAndPhase::PostPushData };
+	FFrameAndPhase RewindFrameAndPhase{ Frame, FFrameAndPhase::PrePushData };
 
 	auto RewindHelper = [RewindFrameAndPhase, this](auto Obj, bool bResimAsFollower, auto& Property, const auto& RewindFunc)
 	{
@@ -290,7 +314,7 @@ void FRewindData::ApplyTargets(const int32 Frame, const bool bResetSimulation)
 const int32 FRewindData::CompareTargetsToLastFrame()
 {
 	int32 RewindFrame = INDEX_NONE;
-	const FFrameAndPhase FrameAndPhase{ CurrentFrame() - 1, FFrameAndPhase::PostPushData };
+	const FFrameAndPhase FrameAndPhase{ CurrentFrame() - 1, FFrameAndPhase::PrePushData };
 
 	if (LatestTargetFrame < FrameAndPhase.Frame)
 	{
@@ -414,7 +438,7 @@ bool FRewindData::RewindToFrame(int32 Frame)
 	}
 
 	//If property changed between Frame and CurFrame, record the latest value and rewind to old
-	FFrameAndPhase RewindFrameAndPhase{ Frame, FFrameAndPhase::PostPushData };
+	FFrameAndPhase RewindFrameAndPhase{ Frame, FFrameAndPhase::PrePushData };
 	FFrameAndPhase CurFrameAndPhase{ CurFrame, FFrameAndPhase::PrePushData };
 
 	BlockResimFrame = bResimAllowRewindToResimulatedFrames ? Frame : CurFrame;
@@ -454,43 +478,54 @@ bool FRewindData::RewindToFrame(int32 Frame)
 			
 	for(FDirtyParticleInfo& DirtyParticleInfo : DirtyParticles)
 	{
-		FGeometryParticleHandle* PTParticle = DirtyParticleInfo.GetObjectPtr();
+		FGeometryParticleHandle* GeometryParticle = DirtyParticleInfo.GetObjectPtr();
+		FKinematicGeometryParticleHandle* KinematicParticle = GeometryParticle->CastToKinematicParticle();
+		FPBDRigidParticleHandle* RigidParticle = GeometryParticle->CastToRigidParticle();
 
 		//rewind is about to start, all particles should be in sync at this point
-		ensure(PTParticle->SyncState() == ESyncState::InSync);
+		ensure(GeometryParticle->SyncState() == ESyncState::InSync);
 		
 		FGeometryParticleStateBase& History = DirtyParticleInfo.GetHistory(); //non-const in case we need to record what's at head for a rewind (CurFrame has already been increased to the next frame)
 
 		PRAGMA_DISABLE_DEPRECATION_WARNINGS
-		History.CachePreCorrectionState(*PTParticle); // Deprecated UE 5.6
+		History.CachePreCorrectionState(*GeometryParticle); // Deprecated UE 5.6
 		PRAGMA_ENABLE_DEPRECATION_WARNINGS
 
-		CachePreResimState(*PTParticle);
+		CachePreResimState(*GeometryParticle);
 
 		// Todo: This should be set by bubble resimulation so that Dynamic and Sleeping particles outside of relevancy of resim doesn't actually resimulate as dynamic particles.
-		DirtyParticleInfo.bResimAsFollower = PTParticle->ObjectState() != EObjectStateType::Dynamic && PTParticle->ObjectState() != EObjectStateType::Sleeping;
+		DirtyParticleInfo.bResimAsFollower = GeometryParticle->ObjectState() != EObjectStateType::Dynamic && GeometryParticle->ObjectState() != EObjectStateType::Sleeping;
 		const bool bResimAsFollower = DirtyParticleInfo.bResimAsFollower;
 
-		bool bAnyChange = RewindHelper(PTParticle, bResimAsFollower, History.ParticlePositionRotation, [](auto Particle, const auto& Data) {Particle->SetXR(Data); });
-		bAnyChange |= RewindHelper(PTParticle->CastToKinematicParticle(), bResimAsFollower, History.Velocities, [](auto Particle, const auto& Data) { Particle->SetV(Data.V()); Particle->SetW(Data.W()); });
-		bAnyChange |= RewindHelper(PTParticle, bResimAsFollower, History.NonFrequentData, [this](auto Particle, const auto& Data)
+		if (KinematicParticle)
+		{
+			// Always clear the kinematic target
+			static const FKinematicTarget DefaultKinematicTarget = FKinematicTarget();
+			KinematicParticle->SetKinematicTarget(DefaultKinematicTarget);
+		}
+
+		bool bAnyChange = RewindHelper(GeometryParticle, bResimAsFollower, History.ParticlePositionRotation, [](auto Particle, const auto& Data) {Particle->SetXR(Data); });
+		bAnyChange |= RewindHelper(KinematicParticle, bResimAsFollower, History.Velocities, [](auto Particle, const auto& Data) { Particle->SetV(Data.V()); Particle->SetW(Data.W()); });
+		bAnyChange |= RewindHelper(GeometryParticle, bResimAsFollower, History.NonFrequentData, [this](auto Particle, const auto& Data)
+		{
+			Solver->GetEvolution()->InvalidateParticle(Particle); // Clear collision/constraints before updating NonFrequentData
+			Particle->SetNonFrequentData(Data); 
+		});
+		bAnyChange |= RewindHelper(KinematicParticle, bResimAsFollower, History.KinematicTarget, [this](auto Particle, const auto& Data)
+		{
+			if (Data.GetMode() != EKinematicTargetMode::None)
 			{
-				Solver->GetEvolution()->InvalidateParticle(Particle); // Clear collision/constraints before updating NonFrequentData
-				Particle->SetNonFrequentData(Data); 
-			});
-		bAnyChange |= RewindHelper(PTParticle->CastToKinematicParticle(), bResimAsFollower, History.KinematicTarget, [](auto Particle, const auto& Data) { Particle->SetKinematicTarget(Data); });
-		bAnyChange |= RewindHelper(PTParticle->CastToRigidParticle(), bResimAsFollower, History.Dynamics, [](auto Particle, const auto& Data) {Particle->SetDynamics(Data); });
-		bAnyChange |= RewindHelper(PTParticle->CastToRigidParticle(), bResimAsFollower, History.DynamicsMisc, [this](auto Particle, const auto& Data) {Solver->SetParticleDynamicMisc(Particle, Data); });
-		bAnyChange |= RewindHelper(PTParticle->CastToRigidParticle(), bResimAsFollower, History.MassProps, [](auto Particle, const auto& Data) {Particle->SetMassProps(Data); });
+				Particle->SetKinematicTarget(Data);
+				Solver->GetEvolution()->GetParticles().MarkMovingKinematic(Particle->Handle());
+			}
+		});
+		bAnyChange |= RewindHelper(RigidParticle, bResimAsFollower, History.Dynamics, [](auto Particle, const auto& Data) {Particle->SetDynamics(Data); });
+		bAnyChange |= RewindHelper(RigidParticle, bResimAsFollower, History.DynamicsMisc, [this](auto Particle, const auto& Data) {Solver->SetParticleDynamicMisc(Particle, Data); });
+		bAnyChange |= RewindHelper(RigidParticle, bResimAsFollower, History.MassProps, [](auto Particle, const auto& Data) {Particle->SetMassProps(Data); });
 
 		// Todo: This should be set by bubble resimulation so that Dynamic and Sleeping particles outside of relevancy of resim doesn't actually resimulate as dynamic particles.
 		// Set this after rewinding, since ResimType gets overwritten if NonFrequentData is cached
-		PTParticle->SetResimType(bResimAsFollower ? EResimType::ResimAsFollower : EResimType::FullResim);
-
-		if (bAnyChange)
-		{
-			CVD_TRACE_PARTICLE(PTParticle);
-		}
+		GeometryParticle->SetResimType(bResimAsFollower ? EResimType::ResimAsFollower : EResimType::FullResim);
 
 		if (!bResimAsFollower)
 		{
@@ -506,14 +541,19 @@ bool FRewindData::RewindToFrame(int32 Frame)
 				//for now just mark anything that changed as enabled during resim. TODO: use bubble
 				DirtyParticleInfo.GetObjectPtr()->SetEnabledDuringResim(true);
 			}
+		}
 
-			if (DirtyParticleInfo.InitializedOnStep > Frame)
-			{
-				//hasn't initialized yet, so disable
-				//must do this after rewind because SetDynamicsMisc will re-enable
-				//(the disable is a temp way to ignore objects not spawned yet, they weren't really disabled which is why it gets re-enabled)
-				Solver->GetEvolution()->DisableParticle(DirtyParticleInfo.GetObjectPtr());
-			}
+		if (DirtyParticleInfo.InitializedOnStep > Frame)
+		{
+			//hasn't initialized yet, so disable
+			//must do this after rewind because SetDynamicsMisc will re-enable
+			//(the disable is a temp way to ignore objects not spawned yet, they weren't really disabled which is why it gets re-enabled)
+			Solver->GetEvolution()->DisableParticle(DirtyParticleInfo.GetObjectPtr());
+		}
+
+		if (bAnyChange)
+		{
+			CVD_TRACE_PARTICLE(GeometryParticle);
 		}
 	}
 
@@ -532,7 +572,7 @@ bool FRewindData::RewindToFrame(int32 Frame)
 
 void FRewindData::StepNonResimParticles(const int32 Frame)
 {
-	const FFrameAndPhase FrameAndPhase{ Frame, FFrameAndPhase::PostPushData };
+	const FFrameAndPhase FrameAndPhase{ Frame, FFrameAndPhase::PrePushData };
 	auto RewindHelper = [FrameAndPhase, this](auto Obj, auto& Property, const auto& RewindFunc) -> bool
 	{
 		if (auto Val = Property.Read(FrameAndPhase, PropertiesPool))
@@ -544,8 +584,11 @@ void FRewindData::StepNonResimParticles(const int32 Frame)
 
 	for (FDirtyParticleInfo& DirtyParticleInfo : DirtyParticles)
 	{
-		FGeometryParticleHandle* PTParticle = DirtyParticleInfo.GetObjectPtr();
-		if (PTParticle->ResimType() != EResimType::ResimAsFollower)
+		FGeometryParticleHandle* GeometryParticle = DirtyParticleInfo.GetObjectPtr();
+		FKinematicGeometryParticleHandle* KinematicParticle = GeometryParticle->CastToKinematicParticle();
+		FPBDRigidParticleHandle* RigidParticle = GeometryParticle->CastToRigidParticle();
+
+		if (GeometryParticle->ResimType() != EResimType::ResimAsFollower)
 		{
 			continue;
 		}
@@ -553,41 +596,51 @@ void FRewindData::StepNonResimParticles(const int32 Frame)
 		bool bHasChanged = false;
 		const FGeometryParticleStateBase& History = DirtyParticleInfo.GetHistory();
 
+		if (KinematicParticle)
+		{
+			// Always clear the kinematic target
+			static const FKinematicTarget DefaultKinematicTarget = FKinematicTarget();
+			KinematicParticle->SetKinematicTarget(DefaultKinematicTarget);
+		}
+
 		//  Set Postion and Rotation
-		bHasChanged = RewindHelper(PTParticle, History.ParticlePositionRotation, [](auto Particle, const auto& Data) -> bool
+		bHasChanged = RewindHelper(GeometryParticle, History.ParticlePositionRotation, [](auto Particle, const auto& Data) -> bool
+		{
+			if (Particle->GetX() != Data.GetX() || Particle->GetR() != Data.GetR())
 			{
-				if (Particle->GetX() != Data.GetX() || Particle->GetR() != Data.GetR())
-				{
-					Particle->SetXR(Data);
-					return true; 
-				}
-				return false;
-			});
+				Particle->SetXR(Data);
+				return true; 
+			}
+			return false;
+		});
 
 		// Set Velocity and Angular Velocity
-		bHasChanged |= RewindHelper(PTParticle->CastToKinematicParticle(), History.Velocities, [](auto Particle, const auto& Data) -> bool
+		bHasChanged |= RewindHelper(KinematicParticle, History.Velocities, [](auto Particle, const auto& Data) -> bool
+		{
+			if (Particle->GetV() != Data.GetV() || Particle->GetW() != Data.GetW())
 			{
-				if (Particle->GetV() != Data.GetV() || Particle->GetW() != Data.GetW())
-				{
-					Particle->SetV(Data.GetV()); Particle->SetW(Data.GetW());
-					return true;
-				}
-				return false;
-			});
+				Particle->SetV(Data.GetV()); Particle->SetW(Data.GetW());
+				return true;
+			}
+			return false;
+		});
 
-		// If XRVW has not changed for the non-resim particle, continue to the next particle
-		if (!bHasChanged)
+		// Set kinematic target
+		bHasChanged |= RewindHelper(KinematicParticle, History.KinematicTarget, [this](auto Particle, const FKinematicTarget& Data) -> bool
 		{
-			continue;
-		}
+			if (Data.GetMode() != EKinematicTargetMode::None)
+			{
+				Particle->SetKinematicTarget(Data);
+				Solver->GetEvolution()->GetParticles().MarkMovingKinematic(Particle->Handle()); // Update SOA views with this particle as being moving kinematic
+				return true;
+			}
+			return false;
+		});
 
 		if (bRewindDataOptimization)
 		{
-			// Set kinematic target
-			RewindHelper(PTParticle->CastToKinematicParticle(), History.KinematicTarget, [](auto Particle, const auto& Data) -> bool { Particle->SetKinematicTarget(Data); return true; });
-
 			// Set disabled true/false and object state
-			bool bHasUpdatedSOAs = RewindHelper(PTParticle->CastToRigidParticle(), History.DynamicsMisc, [this](auto Particle, const auto& Data) -> bool
+			bool bHasUpdatedSOAs = RewindHelper(RigidParticle, History.DynamicsMisc, [this](auto Particle, const auto& Data) -> bool
 				{
 					if (Particle == nullptr)
 					{
@@ -618,11 +671,11 @@ void FRewindData::StepNonResimParticles(const int32 Frame)
 			// If not already done, update SOA views else particles might not get updated
 			if (!bHasUpdatedSOAs)
 			{
-				if (FPBDRigidParticleHandle* Rigid = PTParticle->CastToRigidParticle())
+				if (RigidParticle)
 				{
-					Solver->GetEvolution()->GetParticles().SetDynamicParticleSOA(Rigid->Handle());
+					Solver->GetEvolution()->GetParticles().SetDynamicParticleSOA(RigidParticle->Handle());
 				}
-				else if (FPBDRigidClusteredParticleHandle* Clustered = PTParticle->CastToClustered())
+				else if (FPBDRigidClusteredParticleHandle* Clustered = GeometryParticle->CastToClustered())
 				{
 					Solver->GetEvolution()->GetParticles().SetClusteredParticleSOA(Clustered->Handle());
 				}	
@@ -630,25 +683,30 @@ void FRewindData::StepNonResimParticles(const int32 Frame)
 		}
 		else
 		{
-			RewindHelper(PTParticle, History.NonFrequentData, [this](auto Particle, const auto& Data) -> bool
-				{
-					Solver->GetEvolution()->InvalidateParticle(Particle); // Clear collision/constraints before updating NonFrequentData
-					Particle->SetNonFrequentData(Data);
-					return true;
-				});
-			RewindHelper(PTParticle->CastToKinematicParticle(), History.KinematicTarget, [](auto Particle, const auto& Data) -> bool { Particle->SetKinematicTarget(Data); return true; });
-			RewindHelper(PTParticle->CastToRigidParticle(), History.Dynamics, [](auto Particle, const auto& Data) -> bool { Particle->SetDynamics(Data); return true; });
-			RewindHelper(PTParticle->CastToRigidParticle(), History.DynamicsMisc, [this](auto Particle, const auto& Data) -> bool { Solver->SetParticleDynamicMisc(Particle, Data); return true; });
-			RewindHelper(PTParticle->CastToRigidParticle(), History.MassProps, [](auto Particle, const auto& Data) -> bool { Particle->SetMassProps(Data); return true; });
+			// If XRVW has not changed, or we don't have an active KinematicTarget for the non-resim particle, continue to the next particle
+			if (!bHasChanged)
+			{
+				continue;
+			}
+
+			RewindHelper(GeometryParticle, History.NonFrequentData, [this](auto Particle, const auto& Data) -> bool
+			{
+				Solver->GetEvolution()->InvalidateParticle(Particle); // Clear collision/constraints before updating NonFrequentData
+				Particle->SetNonFrequentData(Data);
+				return true;
+			});
+			RewindHelper(RigidParticle, History.Dynamics, [](auto Particle, const auto& Data) -> bool { Particle->SetDynamics(Data); return true; });
+			RewindHelper(RigidParticle, History.DynamicsMisc, [this](auto Particle, const auto& Data) -> bool { Solver->SetParticleDynamicMisc(Particle, Data); return true; });
+			RewindHelper(RigidParticle, History.MassProps, [](auto Particle, const auto& Data) -> bool { Particle->SetMassProps(Data); return true; });
 		}
 
 		// If the particle is dynamic we must fix the collision anchors so that friction doesn't undo the movement
-		if (PTParticle->ObjectState() == EObjectStateType::Dynamic)
+		if (GeometryParticle->ObjectState() == EObjectStateType::Dynamic)
 		{
-			PTParticle->ParticleCollisions().VisitCollisions(
-				[this, PTParticle](FPBDCollisionConstraint& Collision)
+			GeometryParticle->ParticleCollisions().VisitCollisions(
+				[this, GeometryParticle](FPBDCollisionConstraint& Collision)
 				{
-					Collision.UpdateParticleTransform(PTParticle);
+					Collision.UpdateParticleTransform(GeometryParticle);
 					return ECollisionVisitorResult::Continue;
 				});
 		}
@@ -835,7 +893,8 @@ void FRewindData::AdvanceFrameImp(IResimCacheBase* ResimCache)
 				if (auto Rigid = Handle->CastToRigidParticle())
 				{
 					//sim callback is finished so record the dynamics before solve starts
-					FGeometryParticleStateBase& Latest = Info.AddFrame(CurFrame);
+					Info.MarkDirty(CurFrame);
+					FGeometryParticleStateBase& Latest = Info.GetHistory();
 					Latest.Dynamics.WriteAccessMonotonic(FrameAndPhase, PropertiesPool).CopyFrom(*Rigid);
 				}
 			}
@@ -894,7 +953,8 @@ void FRewindData::PushGTDirtyData(const FDirtyPropertiesManager& SrcManager,cons
 		}
 
 		auto& Info = FindOrAddDirtyObj(*Object, Proxy->IsInitialized() ? INDEX_NONE : CurFrame);
-		auto& Latest = Info.AddFrame(CurFrame);
+		Info.MarkDirty(CurFrame);
+		auto& Latest = Info.GetHistory();
 
 		//At this point all phases should be clean
 		ensure(Latest.IsClean(FFrameAndPhase{ CurFrame, FFrameAndPhase::PrePushData }));
@@ -941,9 +1001,17 @@ void FRewindData::PushGTDirtyData(const FDirtyPropertiesManager& SrcManager,cons
 
 		if (bRewindDataOptimization)
 		{
-			/** When using optimization, add and mark particle as dirty and then rely on FRewindData::CacheCurrentDirtyData() to cache PostPushData */
-			FDirtyParticleInfo& Info = FindOrAddDirtyObj(*PTParticle, Proxy->IsInitialized() ? INDEX_NONE : CurFrame);
-			Info.LastDirtyFrame = CurFrame;
+			/** When using optimization, if this is the first time the particle is dirtied, cache the particles current state and register as dirty */
+			if (FDirtyParticleInfo* DirtyParticleInfo = FindDirtyObj(*PTParticle))
+			{
+				DirtyParticleInfo->MarkDirty(CurFrame);
+			}
+			else // Not yet tracked as dirty
+			{
+				// Call FindOrAddDirtyObj here to set the InitializedOnFrame, if the proxy is not yet initialized
+				FindOrAddDirtyObj(*PTParticle, ParticleProxy->IsInitialized() ? INDEX_NONE : CurFrame);
+				CacheDirtyParticleData(PTParticle, FFrameAndPhase::PrePushData, /*bDirty*/true);				
+			}
 			break;
 		}
 
@@ -971,7 +1039,8 @@ void FRewindData::PushGTDirtyData(const FDirtyPropertiesManager& SrcManager,cons
 			if (auto NewData = Dirty.PropertyData.FindDynamics(SrcManager, SrcDataIdx))
 			{
 				FDirtyParticleInfo& Info = FindOrAddDirtyObj(*PTParticle, ParticleProxy->IsInitialized() ? INDEX_NONE : CurFrame);
-				FGeometryParticleStateBase& Latest = Info.AddFrame(CurFrame);
+				Info.MarkDirty(CurFrame);
+				FGeometryParticleStateBase& Latest = Info.GetHistory();
 				const FFrameAndPhase PostPushData{ CurFrame, FFrameAndPhase::PostPushData };
 				Latest.Dynamics.WriteAccessMonotonic(PostPushData, PropertiesPool) = *NewData;
 				Info.DirtyDynamics = CurFrame;	//Need to save the dirty dynamics into the next phase as well (it's possible a callback will stomp the dynamics value, so that's why it's pending)
@@ -996,9 +1065,17 @@ void FRewindData::PushGTDirtyData(const FDirtyPropertiesManager& SrcManager,cons
 
 		if (bRewindDataOptimization)
 		{
-			/** When using optimization, add and mark joint as dirty and then rely on FRewindData::CacheCurrentDirtyData() to cache PostPushData */
-			FDirtyJointInfo& Info = FindOrAddDirtyObj(*Joint, Proxy->IsInitialized() ? INDEX_NONE : CurFrame);
-			Info.LastDirtyFrame = CurFrame;
+			/** When using optimization, if this is the first time the joint is dirtied, cache the joints current state and register as dirty */
+			if (FDirtyJointInfo* DirtyJointInfo = FindDirtyObj(*Joint))
+			{
+				DirtyJointInfo->MarkDirty(CurFrame);
+			}
+			else // Not yet tracked as dirty
+			{
+				// Call FindOrAddDirtyObj here to set the InitializedOnFrame, if the proxy is not yet initialized
+				FindOrAddDirtyObj(*Joint, JointProxy->IsInitialized() ? INDEX_NONE : CurFrame);
+				CacheDirtyJointData(Joint, FFrameAndPhase::PrePushData, /*bDirty*/true);
+			}
 			break;
 		}
 
@@ -1182,13 +1259,27 @@ void FRewindData::BufferPhysicsResults(TMap<const IPhysicsProxyBase*, struct FDi
 
 void FRewindData::MarkDirtyFromPT(FGeometryParticleHandle& Handle)
 {
-	FDirtyParticleInfo& Info = FindOrAddDirtyObj(Handle);
+	if (bRewindDataOptimization)
+	{
+		/** When using optimization, if this is the first time the joint is dirtied, cache the joints current state and register as dirty */
+		if (FDirtyParticleInfo* DirtyParticleInfo = FindDirtyObj(Handle))
+		{
+			DirtyParticleInfo->MarkDirty(CurFrame);
+		}
+		else // Not yet tracked as dirty
+		{
+			CacheDirtyParticleData(&Handle, FFrameAndPhase::PrePushData, /*bDirty*/true);
+		}
+		return;
+	}
 
+	FDirtyParticleInfo& Info = FindOrAddDirtyObj(Handle);
 	Info.DirtyDynamics = CurFrame;
 
 	const bool bRecordingHistory = !IsResimAndInSync(Handle);
 
-	FGeometryParticleStateBase& Latest = Info.AddFrame(CurFrame);
+	Info.MarkDirty(CurFrame);
+	FGeometryParticleStateBase& Latest = Info.GetHistory();
 
 	//TODO: use property system
 	//For now we just dirty all PT properties that we typically use
@@ -1230,11 +1321,25 @@ void FRewindData::MarkDirtyFromPT(FGeometryParticleHandle& Handle)
 
 void FRewindData::MarkDirtyJointFromPT(FPBDJointConstraintHandle& Handle)
 {
-	FDirtyJointInfo& Info = FindOrAddDirtyObj(Handle);
+	if (bRewindDataOptimization)
+	{
+		/** When using optimization, if this is the first time the joint is dirtied, cache the joints current state and register as dirty */
+		if (FDirtyJointInfo* DirtyJointInfo = FindDirtyObj(Handle))
+		{
+			DirtyJointInfo->MarkDirty(CurFrame);
+		}
+		else // Not yet tracked as dirty
+		{
+			CacheDirtyJointData(&Handle, FFrameAndPhase::PrePushData, /*bDirty*/true);
+		}
+		return;
+	}
 
 	const bool bRecordingHistory = !IsResimAndInSync(Handle);
 
-	FJointStateBase& Latest = Info.AddFrame(CurFrame);
+	FDirtyJointInfo& Info = FindOrAddDirtyObj(Handle);
+	Info.MarkDirty(CurFrame);
+	FJointStateBase& Latest = Info.GetHistory();
 
 	//TODO: use property system
 
@@ -1386,7 +1491,7 @@ int32 FRewindData::FindValidResimFrame(const int32 RequestedFrame)
 
 	for (int32 CheckFrame = RequestedFrame; CheckFrame > EarliestFrame; CheckFrame--)
 	{
-		const FFrameAndPhase FrameAndPhase{ CheckFrame, FFrameAndPhase::PostPushData };
+		const FFrameAndPhase FrameAndPhase{ CheckFrame, FFrameAndPhase::PrePushData };
 		bHasTargetHistory = true;
 		bHasCustomDataHistory = true;
 		
@@ -1493,29 +1598,55 @@ void FRewindData::PushStateAtFrame(FGeometryParticleHandle& Handle, int32 Frame,
 	const bool bRecordingHistory = !IsResimAndInSync(Handle);
 
 	FDirtyParticleInfo& Info = FindOrAddDirtyObj(Handle);
-	FGeometryParticleStateBase& Latest = Info.GetHistory();
+	FGeometryParticleStateBase& History = Info.GetHistory();
 	const FFrameAndPhase FrameAndPhase{ Frame, Phase };
 	LatestTargetFrame = bRecordingHistory ? FMath::Max(LatestTargetFrame, Frame) : LatestTargetFrame;
 
-	if (bRecordingHistory || Latest.TargetPositions.IsClean(FrameAndPhase))
+	if (bRecordingHistory || History.TargetPositions.IsClean(FrameAndPhase))
 	{
-		FParticlePositionRotation& PositionRotation = Latest.TargetPositions.Insert(FrameAndPhase, PropertiesPool);
+		FParticlePositionRotation& PositionRotation = History.TargetPositions.Insert(FrameAndPhase, PropertiesPool);
 		PositionRotation.SetX(Position);
 		PositionRotation.SetR(Quaternion);
 	}
 
-	if (bRecordingHistory || Latest.TargetVelocities.IsClean(FrameAndPhase))
+	if (bRecordingHistory || History.TargetVelocities.IsClean(FrameAndPhase))
 	{
-		FParticleVelocities& PreVelocities = Latest.TargetVelocities.Insert(FrameAndPhase, PropertiesPool);
-		PreVelocities.SetV(LinVelocity);
-		PreVelocities.SetW(AngVelocity);
+		FParticleVelocities& Velocities = History.TargetVelocities.Insert(FrameAndPhase, PropertiesPool);
+		Velocities.SetV(LinVelocity);
+		Velocities.SetW(AngVelocity);
 	}
 
-	if (bRecordingHistory || Latest.TargetStates.IsClean(FrameAndPhase))
+	if (bRecordingHistory || History.TargetStates.IsClean(FrameAndPhase))
 	{
-		FParticleDynamicMisc& PreDynamicsMisc = Latest.TargetStates.Insert(FrameAndPhase, PropertiesPool);
-		PreDynamicsMisc.SetObjectState(bShouldSleep ? EObjectStateType::Sleeping : EObjectStateType::Dynamic);
-		PreDynamicsMisc.SetDisabled(false);
+		FParticleDynamicMisc& DynamicsMisc = History.TargetStates.Insert(FrameAndPhase, PropertiesPool);
+		DynamicsMisc.SetObjectState(bShouldSleep ? EObjectStateType::Sleeping : EObjectStateType::Dynamic);
+		DynamicsMisc.SetDisabled(false);
+	}
+}
+
+void FRewindData::ProcessDirtyKinematicTargets(const TParticleView<TPBDRigidParticles<FReal, 3>>& ActiveKinematicParticles)
+{
+	// Cache the kinematic target of moving kinematic objects
+	for (FTransientPBDRigidParticleHandle& ActiveKinematicParticle : ActiveKinematicParticles)
+	{
+		CacheKinematicTarget(*ActiveKinematicParticle.Handle());
+	}
+}
+
+void FRewindData::ProcessDirtyPTParticles(const TParticleView<TPBDRigidParticles<FReal, 3>>&  DirtyPTParticles)
+{
+	PRAGMA_DISABLE_DEPRECATION_WARNINGS
+	if (Solver)
+	{
+		Solver->FinalizeRewindData(DirtyPTParticles);
+	}
+	PRAGMA_ENABLE_DEPRECATION_WARNINGS
+
+	// Cache dirty particles from simulation
+	int32 DataIdx = 0;
+	for (TPBDRigidParticleHandleImp<FReal, 3, false>& DirtyPTParticle : DirtyPTParticles)
+	{
+		PushPTDirtyData(*DirtyPTParticle.Handle(), DataIdx++);
 	}
 }
 
@@ -1523,21 +1654,27 @@ void FRewindData::PushPTDirtyData(TPBDRigidParticleHandle<FReal, 3>& Handle, con
 {
 	if (bRewindDataOptimization)
 	{
-		/** When using optimization, mark particle as dirty and then rely on FRewindData::CacheCurrentDirtyData() to cache PostPushData 
-		* If this is a new entry in the dirty particles collection, cache it here since CacheCurrentDirtyData for PostPushData has already happened this frame */
+		/** When using optimization, mark particle as dirty and then rely on FRewindData::CacheCurrentDirtyData() to cache PrePushData 
+		* If this is a new entry in the dirty particles collection, cache now since CacheCurrentDirtyData for PrePushData has already happened this frame */
 
-		if (FDirtyParticleInfo* Info = FindDirtyObj(Handle))
+		FDirtyParticleInfo* Info = FindDirtyObj(Handle);
+		if (!Info || Info->InitializedOnStep == CurFrame)
+		{
+			CacheDirtyParticleData(&Handle, FFrameAndPhase::PostCallbacks, /*bDirty*/true);
+		}
+		else
 		{
 			// Mark particle as dirty so it doesn't get cleared from the dirty particles
-			Info->LastDirtyFrame = CurFrame;
-			return;
+			Info->MarkDirty(CurFrame);
 		}
+		return;
 	}
 
 	const bool bRecordingHistory = !IsResimAndInSync(Handle);
 
 	FDirtyParticleInfo& Info = FindOrAddDirtyObj(Handle);
-	FGeometryParticleStateBase& Latest = Info.AddFrame(CurFrame);
+	Info.MarkDirty(CurFrame);
+	FGeometryParticleStateBase& Latest = Info.GetHistory();
 
 	const FFrameAndPhase FrameAndPhase{ CurFrame, FFrameAndPhase::PostCallbacks };
 
@@ -1579,87 +1716,274 @@ void FRewindData::PushPTDirtyData(TPBDRigidParticleHandle<FReal, 3>& Handle, con
 	}
 }
 
-void FRewindData::CacheCurrentDirtyData(FFrameAndPhase::EParticleHistoryPhase CurrentPhase)
+void FRewindData::CacheKinematicTarget(TPBDRigidParticleHandle<FReal, 3>& Handle)
+{
+	if (!bRewindDataOptimization)
+	{
+		return;
+	}
+
+	FKinematicGeometryParticleHandle* Kinematic = Handle.CastToKinematicParticle();
+	if (!Kinematic || Kinematic->KinematicTarget().GetMode() == EKinematicTargetMode::None)
+	{
+		return;
+	}
+
+	const FFrameAndPhase FrameAndPhase{ CurFrame, FFrameAndPhase::PrePushData };
+
+	FDirtyParticleInfo& Info = FindOrAddDirtyObj(Handle);
+	Info.MarkDirty(CurFrame);
+
+	if (IsResim())
+	{
+		if (!Info.bResimAsFollower)
+		{
+			// Check if particle differ from the currently cached history for this frame, if so, clear the history and mark particle as desynced so it will cache data during resimulation
+			DesyncIfNecessary</*bSkipDynamics=*/false>(Info, FrameAndPhase);
+		}
+
+		if (Kinematic->SyncState() == ESyncState::InSync)
+		{
+			/* No need to cache data in history if the particle is still in sync during resimulation */
+			return;
+		}
+	}
+
+	FGeometryParticleStateBase& Latest = Info.GetHistory();
+
+	const bool bRecordingHistory = !IsResimAndInSync(Handle);
+	if (bRecordingHistory)
+	{
+		if (auto* Data = Latest.KinematicTarget.WriteAccessNonDecreasing(FrameAndPhase, PropertiesPool))
+		{
+			CopyDataFromObject(*Data, *Kinematic);
+		}
+	}
+}
+
+void FRewindData::CacheDirtyParticleData(FGeometryParticleHandle* Particle, const FFrameAndPhase::EParticleHistoryPhase& CurrentPhase, const bool& bDirty)
 {
-	if (bRewindDataOptimization && CurrentPhase != FFrameAndPhase::PostPushData)
+	FDirtyParticleInfo& DirtyParticleInfo = FindOrAddDirtyObj(*Particle);
+	if (bDirty)
 	{
-		// Don't cache any other phase than PostPushData when using optimization
+		DirtyParticleInfo.MarkDirty(CurFrame);
+	}
+
+	if (DirtyParticleInfo.InitializedOnStep > CurFrame)
+	{
+		// Don't cache data for uninitialized particle
+		return;
+	}
+	else if (DirtyParticleInfo.InitializedOnStep == CurFrame && CurrentPhase < FFrameAndPhase::PostCallbacks)
+	{
+		/** Don't cache data in the PrePushData or PostPushData phase on the frame the particle gets initialized
+		* Do cache the data in the PostCallbacks though, with optimization we will convert it from PostCallback state into PrePushData and store it there to have the initial state of the particle */
 		return;
 	}
 
-	const FFrameAndPhase FrameAndPhase { CurFrame, CurrentPhase };
+	// When using optimization we always cache data in the PrePushData phase
+	const FFrameAndPhase FrameAndPhase{ CurFrame, bRewindDataOptimization ? FFrameAndPhase::PrePushData : CurrentPhase };
 
-	auto DirtyPropHelper = [this, FrameAndPhase](auto& Property, const EChaosPropertyFlags PropName, const auto& Object)
+	if (IsResim())
+	{
+		if (!DirtyParticleInfo.bResimAsFollower)
 		{
-			const bool bRecordingHistory = !IsResimAndInSync(Object);
+			// Check if particle differ from the currently cached history for this frame, if so, clear the history and mark particle as desynced so it will cache data during resimulation
+			DesyncIfNecessary</*bSkipDynamics=*/false>(DirtyParticleInfo, FrameAndPhase);
+		}
+
+		if (Particle->SyncState() == ESyncState::InSync)
+		{
+			/* No need to cache data in history if the particle is still in sync during resimulation */
+			return;
+		}
+	}
+
+	const bool bCachePreFromPost = bRewindDataOptimization && CurrentPhase == FFrameAndPhase::PostCallbacks;
+
+	auto DirtyPropHelper = [this, FrameAndPhase](auto& Property, const auto& Particle, const auto& CacheFunction)
+		{
+			const bool bRecordingHistory = !IsResimAndInSync(Particle);
 
 			if (bRecordingHistory || Property.IsClean(FrameAndPhase))
 			{
 				if (auto* Data = Property.WriteAccessNonDecreasing(FrameAndPhase, PropertiesPool))
 				{
-					CopyDataFromObject(*Data, Object);
+					CacheFunction(Particle, *Data);
 				}
 			}
 		};
 
-	for (FDirtyParticleInfo& DirtyParticleInfo : DirtyParticles)
+	FKinematicGeometryParticleHandle* KinematicParticle = Particle->CastToKinematicParticle();
+	FPBDRigidParticleHandle* RigidParticle = KinematicParticle ? KinematicParticle->CastToRigidParticle() : nullptr;;
+
+	FGeometryParticleStateBase& Latest = DirtyParticleInfo.GetHistory();
+
+	DirtyPropHelper(Latest.ParticlePositionRotation, *Particle, [this, &bCachePreFromPost, &RigidParticle](const auto& Particle, FParticlePositionRotation& Data)
 	{
-		if (FGeometryParticleHandle* Geometry = DirtyParticleInfo.GetObjectPtr())
+		if (bCachePreFromPost && RigidParticle && RigidParticle->IsMovingKinematic())
+		{
+			/** FFrameAndPhase::PostCallbacks is called at the end of the physics solve
+			* If this is a kinematic moving object then XR have been updated already in the integrate step via KinematicTarget and the velocity has been updated based on the XR change
+			* Get the pre-solve state by stepping XR back once via the velocities. */
+
+			Data.SetX(RigidParticle->GetX() - (RigidParticle->GetV() * Solver->GetLastDt()));
+			Data.SetR(Chaos::FRotation3::IntegrateRotationWithAngularVelocity(RigidParticle->GetR(), RigidParticle->GetWf(), -Solver->GetLastDt()));
+		}
+		else
+		{
+			// This works during PostCallback since XR is not yet updated from PQ even if it's actually called PostSolve instead of PostCallback
+			CopyDataFromObject(Data, Particle); 
+		}
+	});
+
+	DirtyPropHelper(Latest.NonFrequentData, *Particle, [this](const auto& Particle, FParticleNonFrequentData& Data)
+	{
+		CopyDataFromObject(Data, Particle);
+	});
+
+	if (KinematicParticle)
+	{
+		DirtyPropHelper(Latest.Velocities, *KinematicParticle, [this, &bCachePreFromPost, &RigidParticle](const auto& Particle, FParticleVelocities& Data)
+		{
+			if (bCachePreFromPost && RigidParticle)
+			{
+				Data.SetV(RigidParticle->GetPreV());
+				Data.SetW(RigidParticle->GetPreW());
+			}
+			else
+			{
+				CopyDataFromObject(Data, Particle);
+			}
+		});
+
+		if (!bRewindDataOptimization)
 		{
-			if (IsResim())
+			// Cache kinematic target here if not running with optimization, else it's cached in CacheKinematicTarget()
+			DirtyPropHelper(Latest.KinematicTarget, *KinematicParticle, [this](const auto& Particle, FKinematicTarget& Data)
 			{
-				if (!DirtyParticleInfo.bResimAsFollower)
+				CopyDataFromObject(Data, Particle);
+			});
+		}
+
+		if (RigidParticle)
+		{
+			if (FrameAndPhase.Phase != FFrameAndPhase::PrePushData)
+			{
+				// Don't cache dynamics for PrePushData, it's always zeroed out there. 
+				DirtyPropHelper(Latest.Dynamics, *RigidParticle, [this, &bCachePreFromPost](const auto& Particle, FParticleDynamics& Data)
 				{
-					// Check if particle differ from the currently cached history for this frame, if so, clear the history and mark particle as desynced so it will cache data during resimulation
-					DesyncIfNecessary</*bSkipDynamics=*/false>(DirtyParticleInfo, FrameAndPhase);
-				}
+					CopyDataFromObject(Data, Particle); 
+				});
+			}
 
-				if (Geometry->SyncState() == ESyncState::InSync)
+			DirtyPropHelper(Latest.DynamicsMisc, *RigidParticle, [this, &bCachePreFromPost](const auto& Particle, FParticleDynamicMisc& Data)
+			{
+				CopyDataFromObject(Data, Particle);
+				if (bCachePreFromPost)
 				{
-					/* No need to cache data in history if the particle is still in sync during resimulation */
-					continue;
+					/** FFrameAndPhase::PostCallbacks is called at the end of the physics solve, get the object state from before the solve */
+					Data.SetObjectState(Particle.PreObjectState());
 				}
-			}
+			});
+
+			DirtyPropHelper(Latest.MassProps, *RigidParticle, [this](const auto& Particle, FParticleMassProps& Data)
+			{
+				CopyDataFromObject(Data, Particle);
+			});
+		}
+	}
+}
 
-			FGeometryParticleStateBase& Latest = DirtyParticleInfo.GetHistory();
+void FRewindData::CacheDirtyJointData(FPBDJointConstraintHandle* Joint, const FFrameAndPhase::EParticleHistoryPhase& CurrentPhase, const bool& bDirty)
+{
+	FDirtyJointInfo& DirtyJointInfo = FindOrAddDirtyObj(*Joint);
+	if (bDirty)
+	{
+		DirtyJointInfo.MarkDirty(CurFrame);
+	}
 
-			DirtyPropHelper(Latest.ParticlePositionRotation, EChaosPropertyFlags::XR, *Geometry);
-			DirtyPropHelper(Latest.NonFrequentData, EChaosPropertyFlags::NonFrequentData, *Geometry);
+	if (DirtyJointInfo.InitializedOnStep > CurFrame)
+	{
+		// Don't cache data for uninitialized joint
+		return;
+	}
+	else if (DirtyJointInfo.InitializedOnStep == CurFrame && CurrentPhase < FFrameAndPhase::PostCallbacks)
+	{
+		/** Don't cache data in the PrePushData or PostPushData phase on the frame the joint gets initialized
+		* Do cache the data in the PostCallbacks though, with optimization we will convert it from PostCallback state into PrePushData and store it there to have the initial state of the particle */
+		return;
+	}
 
-			if (auto Kinematic = Geometry->CastToKinematicParticle())
-			{
-				DirtyPropHelper(Latest.Velocities, EChaosPropertyFlags::Velocities, *Kinematic);
-				DirtyPropHelper(Latest.KinematicTarget, EChaosPropertyFlags::KinematicTarget, *Kinematic);
+	// When using optimization we always cache data in the PrePushData phase
+	const FFrameAndPhase FrameAndPhase{ CurFrame, bRewindDataOptimization ? FFrameAndPhase::PrePushData : CurrentPhase };
 
-				if (auto Rigid = Kinematic->CastToRigidParticle())
+	auto DirtyPropHelper = [this, FrameAndPhase](auto& Property, const auto& Object)
+		{
+			const bool bRecordingHistory = !IsResimAndInSync(Object);
+
+			if (bRecordingHistory || Property.IsClean(FrameAndPhase))
+			{
+				if (auto* Data = Property.WriteAccessNonDecreasing(FrameAndPhase, PropertiesPool))
 				{
-					DirtyPropHelper(Latest.Dynamics, EChaosPropertyFlags::Dynamics, *Rigid);
-					DirtyPropHelper(Latest.DynamicsMisc, EChaosPropertyFlags::DynamicMisc, *Rigid);
-					DirtyPropHelper(Latest.MassProps, EChaosPropertyFlags::MassProps, *Rigid);
+					CopyDataFromObject(*Data, Object);
 				}
 			}
+		};
+
+	if (IsResim())
+	{
+		if (!DirtyJointInfo.bResimAsFollower)
+		{
+			// Check if joint differ from the currently cached history for this frame, if so, clear the history and mark joint as desynced so it will cache data during resimulation
+			DesyncIfNecessary</*bSkipDynamics=*/false>(DirtyJointInfo, FrameAndPhase);
+		}
+
+		if (Joint->SyncState() == ESyncState::InSync)
+		{
+			/* No need to cache data in history if the particle is still in sync during resimulation */
+			return;
 		}
 	}
 
-	for (FDirtyJointInfo& DirtyJointInfo : DirtyJoints)
+	FJointStateBase& Latest = DirtyJointInfo.GetHistory();
+
+	DirtyPropHelper(Latest.JointSettings, *Joint);
+}
+
+void FRewindData::CacheCurrentDirtyData(const FFrameAndPhase::EParticleHistoryPhase& CurrentPhase)
+{
+	if (!bRewindDataOptimization)
+	{
+		// The non optimized flow cache data in various other callbacks which will conflict with this
+		return;
+	}
+
+	for (FDirtyParticleInfo& DirtyParticleInfo : DirtyParticles)
 	{
-		if (const FPBDJointConstraintHandle* Joint = DirtyJointInfo.GetObjectPtr())
+		if (FGeometryParticleHandle* Geometry = DirtyParticleInfo.GetObjectPtr())
 		{
-			FJointStateBase& Latest = DirtyJointInfo.GetHistory();
+			CacheDirtyParticleData(Geometry, CurrentPhase, /*bDirty*/false);
+		}
+	}
 
-			DirtyPropHelper(Latest.JointSettings, EChaosPropertyFlags::JointSettings, *Joint);
+	for (FDirtyJointInfo& DirtyJointInfo : DirtyJoints)
+	{
+		if (FPBDJointConstraintHandle* Joint = DirtyJointInfo.GetObjectPtr())
+		{
+			CacheDirtyJointData(Joint, CurrentPhase, /*bDirty*/false);
 		}
 	}
 }
 
 FGeometryParticleState FRewindData::GetPastStateAtFrame(const FGeometryParticleHandle& Handle, int32 Frame, FFrameAndPhase::EParticleHistoryPhase Phase) const
 {
-	return GetPastStateAtFrameImp<FGeometryParticleState>(DirtyParticles, Handle, Frame, Phase);
+	return GetPastStateAtFrameImp<FGeometryParticleState>(DirtyParticles, Handle, Frame, bRewindDataOptimization ? FFrameAndPhase::EParticleHistoryPhase::PrePushData : Phase);
 }
 
 FJointState FRewindData::GetPastJointStateAtFrame(const FPBDJointConstraintHandle& Handle, int32 Frame, FFrameAndPhase::EParticleHistoryPhase Phase) const
 {
-	return GetPastStateAtFrameImp<FJointState>(DirtyJoints, Handle, Frame, Phase);
+	return GetPastStateAtFrameImp<FJointState>(DirtyJoints, Handle, Frame, bRewindDataOptimization ? FFrameAndPhase::EParticleHistoryPhase::PrePushData : Phase);
 }
 
 CHAOS_API int32 bInterpolateTargetGaps = 5;
@@ -1671,17 +1995,17 @@ void FRewindData::SetTargetStateAtFrame(FGeometryParticleHandle& Handle, const i
 	if (bInterpolateTargetGaps)
 	{
 		FDirtyParticleInfo& Info = FindOrAddDirtyObj(Handle);
-		FGeometryParticleStateBase& Latest = Info.GetHistory();
+		FGeometryParticleStateBase& History = Info.GetHistory();
 		FFrameAndPhase FrameAndPhase;
 
-		if (Latest.TargetPositions.GetHeadFrameAndPhase(FrameAndPhase))
+		if (History.TargetPositions.GetHeadFrameAndPhase(FrameAndPhase))
 		{
 			const int32 FrameDiff = Frame - FrameAndPhase.Frame;
 			if (FrameDiff > 1 && FrameDiff <= bInterpolateTargetGaps)
 			{
-				const FParticlePositionRotation* TargetXR = Latest.TargetPositions.Read(FrameAndPhase, PropertiesPool);
-				const FParticleVelocities* TargetVW = Latest.TargetVelocities.Read(FrameAndPhase, PropertiesPool);
-				const FParticleDynamicMisc* TargetDynamic = Latest.TargetStates.Read(FrameAndPhase, PropertiesPool);
+				const FParticlePositionRotation* TargetXR = History.TargetPositions.Read(FrameAndPhase, PropertiesPool);
+				const FParticleVelocities* TargetVW = History.TargetVelocities.Read(FrameAndPhase, PropertiesPool);
+				const FParticleDynamicMisc* TargetDynamic = History.TargetStates.Read(FrameAndPhase, PropertiesPool);
 				if (TargetXR && TargetVW && TargetDynamic)
 				{
 					for (int32 InterpFrame = 1; InterpFrame < FrameDiff; InterpFrame++)
diff --git a/Engine/Source/Runtime/Experimental/Chaos/Public/Chaos/KinematicTargets.h b/Engine/Source/Runtime/Experimental/Chaos/Public/Chaos/KinematicTargets.h
--- a/Engine/Source/Runtime/Experimental/Chaos/Public/Chaos/KinematicTargets.h
+++ b/Engine/Source/Runtime/Experimental/Chaos/Public/Chaos/KinematicTargets.h
@@ -183,6 +183,11 @@ namespace Chaos
 			Mode = Other.KinematicTarget().Mode;
 		}
 
+		static FKinematicTarget ZeroValue()
+		{
+			return FKinematicTarget();
+		}
+
 	private:
 		explicit FKinematicTarget(const FRigidTransform3& InTransform)
 			: Rotation(InTransform.GetRotation())
diff --git a/Engine/Source/Runtime/Experimental/Chaos/Public/Chaos/PBDRigidsEvolution.h b/Engine/Source/Runtime/Experimental/Chaos/Public/Chaos/PBDRigidsEvolution.h
--- a/Engine/Source/Runtime/Experimental/Chaos/Public/Chaos/PBDRigidsEvolution.h
+++ b/Engine/Source/Runtime/Experimental/Chaos/Public/Chaos/PBDRigidsEvolution.h
@@ -275,6 +275,7 @@ public:
 	typedef TFunction<void(const TParticleView<FPBDRigidParticles>&, const FReal)> FUpdatePositionRule;
 	typedef TFunction<void(FPBDRigidParticles&, const FReal, const FReal, const int32)> FKinematicUpdateRule;
 	typedef TFunction<void(TParticleView<FPBDRigidParticles>&)> FCaptureRewindRule;
+	typedef TFunction<void(TParticleView<FPBDRigidParticles>&)> FCaptureRewindKinematicTargetRule;
 
 	CHAOS_API FPBDRigidsEvolutionBase(FPBDRigidsSOAs& InParticles, THandleArray<FChaosPhysicsMaterial>& InSolverPhysicsMaterials, bool InIsSingleThreaded = false);
 	CHAOS_API virtual ~FPBDRigidsEvolutionBase();
@@ -332,7 +333,8 @@ public:
 	void AddForceFunction(FForceRule ForceFunction) { ForceRules.Add(ForceFunction); }
 	void AddImpulseFunction(FForceRule ImpulseFunction) { ImpulseRules.Add(ImpulseFunction); }
 	void SetParticleUpdatePositionFunction(FUpdatePositionRule ParticleUpdate) { ParticleUpdatePosition = ParticleUpdate; }
-	void SetCaptureRewindDataFunction(FCaptureRewindRule Rule){ CaptureRewindData = Rule; }
+	void SetCaptureRewindDataFunction(FCaptureRewindRule Rule) { CaptureRewindData = Rule; }
+	void SetCaptureRewindKinematicTargetFunction(FCaptureRewindKinematicTargetRule Rule) { CaptureRewindKinematicTarget = Rule; }
 
 	TGeometryParticleHandles<FReal, 3>& GetParticleHandles() { return Particles.GetParticleHandles(); }
 	const TGeometryParticleHandles<FReal, 3>& GetParticleHandles() const { return Particles.GetParticleHandles(); }
@@ -914,6 +916,7 @@ protected:
 	FUpdatePositionRule ParticleUpdatePosition;
 	FKinematicUpdateRule KinematicUpdate;
 	FCaptureRewindRule CaptureRewindData;
+	FCaptureRewindKinematicTargetRule CaptureRewindKinematicTarget;
 	TArray<FPBDConstraintContainer*> ConstraintContainers;
 	Private::FPBDIslandManager IslandManager;
 	Private::FPBDIslandGroupManager IslandGroupManager;
diff --git a/Engine/Source/Runtime/Experimental/Chaos/Public/PBDRigidsSolver.h b/Engine/Source/Runtime/Experimental/Chaos/Public/PBDRigidsSolver.h
--- a/Engine/Source/Runtime/Experimental/Chaos/Public/PBDRigidsSolver.h
+++ b/Engine/Source/Runtime/Experimental/Chaos/Public/PBDRigidsSolver.h
@@ -296,7 +296,8 @@ namespace Chaos
 		UE_DEPRECATED(5.5, "Deprecated, use GetUseCollisionResimCache() instead")
 		bool RewindUsesCollisionResimCache() const { return GetUseCollisionResimCache(); }
 
-		CHAOS_API void FinalizeRewindData(const TParticleView<FPBDRigidParticles>& DirtyParticles);
+		UE_DEPRECATED(5.7, "Deprecated, This functionality is now handled by FRewindData::ProcessDirtyPTParticles")
+		CHAOS_API void FinalizeRewindData(const TParticleView<FPBDRigidParticles>& DirtyParticles) { };
 
 		FPerSolverFieldSystem& GetPerSolverField() { return *PerSolverField; }
 		const FPerSolverFieldSystem& GetPerSolverField() const { return *PerSolverField; }
diff --git a/Engine/Source/Runtime/Experimental/Chaos/Public/RewindData.h b/Engine/Source/Runtime/Experimental/Chaos/Public/RewindData.h
--- a/Engine/Source/Runtime/Experimental/Chaos/Public/RewindData.h
+++ b/Engine/Source/Runtime/Experimental/Chaos/Public/RewindData.h
@@ -441,11 +441,11 @@ struct FFrameAndPhase
 	{
 		//The particle state before PushData, server state update, or any sim callbacks are processed 
 		//This is the results of the previous frame before any GT modifications are made in this frame
+		//This is what the server state should be compared against
+		//This is what we rewind to before a resim
 		PrePushData = 0,
 
 		//The particle state after PushData is applied, but before any server state is applied
-		//This is what the server state should be compared against
-		//This is what we rewind to before a resim
 		PostPushData,
 
 		//The particle state after sim callbacks are applied.
@@ -1175,13 +1175,16 @@ struct FGeometryParticleStateBase
 		PreCorrectionXR.SetR(Particle.GetR());
 	}
 
+	/** Setting bNoEntryIsHead = false in TParticlePropertyBuffer allows us to not cache the data each frame and it will only return entries for a specific frame when fetched 
+	* For example, we don't need to cache a kinematic target that is set to None, instead we only cache when it's not set to None and if we don't have an entry for a specific frame we can expect it to be a None kinematic target.
+	* Note that bNoEntryIsHead = true (which is default) will return the latest cached entry closest to the requested frame and phase without returning an earlier entry. */
 	TParticlePropertyBuffer<FParticlePositionRotation,EChaosProperty::XR> ParticlePositionRotation;
 	TParticlePropertyBuffer<FParticleNonFrequentData,EChaosProperty::NonFrequentData> NonFrequentData;
 	TParticlePropertyBuffer<FParticleVelocities,EChaosProperty::Velocities> Velocities;
 	TParticlePropertyBuffer<FParticleDynamics,EChaosProperty::Dynamics, /*bNoEntryIsHead=*/false> Dynamics;
 	TParticlePropertyBuffer<FParticleDynamicMisc,EChaosProperty::DynamicMisc> DynamicsMisc;
 	TParticlePropertyBuffer<FParticleMassProps,EChaosProperty::MassProps> MassProps;
-	TParticlePropertyBuffer<FKinematicTarget, EChaosProperty::KinematicTarget> KinematicTarget;
+	TParticlePropertyBuffer<FKinematicTarget, EChaosProperty::KinematicTarget, /*bNoEntryIsHead=*/false> KinematicTarget;
 
 	TParticlePropertyBuffer<FParticlePositionRotation, EChaosProperty::XR, /*bNoEntryIsHead=*/false> TargetPositions;
 	TParticlePropertyBuffer<FParticleVelocities, EChaosProperty::Velocities, /*bNoEntryIsHead=*/false> TargetVelocities;
@@ -1540,6 +1543,8 @@ public:
 		bRewindDataOptimization = InRewindDataOptimization;
 		LatestTargetFrame = 0;
 		Managers = TCircularBuffer<FFrameManagerInfo>(NumFrames + 1);
+
+		RegisterEvolutionCallbacks();
 	}
 
 	void Init(FPBDRigidsSolver* InSolver, int32 NumFrames, int32 InCurrentFrame)
@@ -1549,8 +1554,14 @@ public:
 		LatestFrame = InCurrentFrame;
 		LatestTargetFrame = 0;
 		Managers = TCircularBuffer<FFrameManagerInfo>(NumFrames + 1);
+
+		RegisterEvolutionCallbacks();
 	}
 
+private:
+	void RegisterEvolutionCallbacks();
+
+public: 
 	int32 Capacity() const { return Managers.Capacity(); }
 	int32 CurrentFrame() const { return CurFrame; }
 	int32 GetLatestFrame() const { return LatestFrame; }
@@ -1693,10 +1704,6 @@ public:
 	/** Call this to mark specific joint as dirty and cache its current data, expected to be called at the start of OnPreSimulate_Internal, it caches data in FFrameAndPhase::PostPushData */
 	void CHAOS_API MarkDirtyJointFromPT(FPBDJointConstraintHandle& Handle);
 
-	/** Caches data for all particles and joints marked dirty in RewindData 
-	* @param CurrentPhase Data gets cached under this phase for the current frame, be sure to call this from the corresponding phase during physics */
-	void CacheCurrentDirtyData(FFrameAndPhase::EParticleHistoryPhase CurrentPhase);
-
 	void CHAOS_API SpawnProxyIfNeeded(FSingleParticlePhysicsProxy& Proxy);
 
 	/** Add input history to the rewind data for future use while resimulating */
@@ -1814,6 +1821,27 @@ private:
 
 	void CHAOS_API AdvanceFrameImp(IResimCacheBase* ResimCache);
 
+	/** Called post-solve from the evolution with all particles that are dirty after solving physics. Note: Before PQ are applied to XR */
+	void ProcessDirtyPTParticles(const TParticleView<TPBDRigidParticles<FReal, 3>>& DirtyPTParticles);
+
+	/** Called during integration just before KinematicTargets are processed with all active kinematic particles */
+	void ProcessDirtyKinematicTargets(const TParticleView<TPBDRigidParticles<FReal, 3>>& ActiveKinematicParticles);
+
+	/** Called during integration just before KinematicTargets are processed */
+	void CacheKinematicTarget(TPBDRigidParticleHandle<FReal, 3>& Rigid);
+
+	/** Called to cache a particles current state into a specified phase of the current frame
+	* @param bDirty will mark the particle as dirty in the rewind history which will continue to cache it for FRewindData::Capacity() amount of frames */
+	void CacheDirtyParticleData(FGeometryParticleHandle* Geometry, const FFrameAndPhase::EParticleHistoryPhase& CurrentPhase, const bool& bDirty);
+
+	/** Called to cache a joints current state into a specified phase of the current frame
+	* @param bDirty will mark the joint as dirty in the rewind history which will continue to cache it for FRewindData::Capacity() amount of frames */
+	void CacheDirtyJointData(FPBDJointConstraintHandle* Joint, const FFrameAndPhase::EParticleHistoryPhase& CurrentPhase, const bool& bDirty);
+
+	/** Caches data for all particles and joints marked dirty in RewindData
+	* @param CurrentPhase Data gets cached under this phase for the current frame, be sure to call this from the corresponding phase during physics */
+	void CacheCurrentDirtyData(const FFrameAndPhase::EParticleHistoryPhase& CurrentPhase);
+
 	struct FFrameManagerInfo
 	{
 		TUniquePtr<IResimCacheBase> ExternalResimCache;
@@ -1871,6 +1899,7 @@ private:
 
 		TObj* GetObjectPtr() const { return ObjPtr; }
 
+		UE_DEPRECATED(5.7, "Deprecated, Use GetHistory() and MarkDirty() individually instead")
 		THistoryType& AddFrame(const int32 Frame)
 		{
 			LastDirtyFrame = Frame;
@@ -1882,7 +1911,7 @@ private:
 			History.ClearEntryAndFuture(FrameAndPhase);
 		}
 
-		const THistoryType& GetHistory() const	//For non-const access use AddFrame
+		const THistoryType& GetHistory() const
 		{
 			return History;
 		}
@@ -1891,6 +1920,11 @@ private:
 		{
 			return History;
 		}
+
+		void MarkDirty(const int32 Frame)
+		{
+			LastDirtyFrame = Frame;
+		}
 	};
 
 	using FDirtyParticleInfo = TDirtyObjectInfo<FGeometryParticleStateBase, FGeometryParticleHandle>;
