Added guards against selecting invalid indices -- return false instead of crashing on attempting to set an invalid selection

diff --git a/Engine/Source/Runtime/Experimental/Dataflow/Core/Private/Dataflow/DataflowSelection.cpp b/Engine/Source/Runtime/Experimental/Dataflow/Core/Private/Dataflow/DataflowSelection.cpp
--- a/Engine/Source/Runtime/Experimental/Dataflow/Core/Private/Dataflow/DataflowSelection.cpp
+++ b/Engine/Source/Runtime/Experimental/Dataflow/Core/Private/Dataflow/DataflowSelection.cpp
@@ -87,14 +87,26 @@ TArray<int32> FDataflowSelection::AsArray() const
 	return SelectionArr;
 }
 
-void FDataflowSelection::SetFromArray(const TArray<int32>& SelectionArr)
+bool FDataflowSelection::SetFromArray(const TArray<int32>& SelectionArr)
 {
 	SelectionArray.Init(false, SelectionArray.Num());
-
+	bool bSuccess = true;
 	for (int32 Elem : SelectionArr)
 	{
-		SelectionArray[Elem] = true;
+		if (SelectionArray.IsValidIndex(Elem))
+		{
+			SelectionArray[Elem] = true;
+		}
+		else
+		{
+			bSuccess = false;
+		}
+	}
+	if (!bSuccess)
+	{
+		UE_LOG(LogChaosDataflow, Warning, TEXT("FDataflowSelection::SetFromArray failed to set invalid indices."));
 	}
+	return bSuccess;
 }
 
 void FDataflowSelection::SetFromArray(const TArray<bool>& SelectionArr)
@@ -183,10 +195,10 @@ bool FDataflowSelection::IsValidForCollection(const FManagedArrayCollection& InC
 	return Num() == InCollection.NumElements(GroupName);
 }
 
-void FDataflowSelection::InitFromArray(const FManagedArrayCollection& InCollection, const TArray<int32>& InSelectionArr)
+bool FDataflowSelection::InitFromArray(const FManagedArrayCollection& InCollection, const TArray<int32>& InSelectionArr)
 {
 	Initialize(InCollection.NumElements(GroupName), false);
-	SetFromArray(InSelectionArr);
+	return SetFromArray(InSelectionArr);
 }
 
 FString FDataflowSelection::ToString()
diff --git a/Engine/Source/Runtime/Experimental/Dataflow/Core/Public/Dataflow/DataflowSelection.h b/Engine/Source/Runtime/Experimental/Dataflow/Core/Public/Dataflow/DataflowSelection.h
--- a/Engine/Source/Runtime/Experimental/Dataflow/Core/Public/Dataflow/DataflowSelection.h
+++ b/Engine/Source/Runtime/Experimental/Dataflow/Core/Public/Dataflow/DataflowSelection.h
@@ -36,9 +36,10 @@ struct FDataflowSelection
 	DATAFLOWCORE_API TArray<int32> AsArrayValidated(const FManagedArrayCollection& InCollection) const;
 	
 	/**
-	* Sets the selection from a sparse array (it only contains the indices of the selected items)
-	*/
-	DATAFLOWCORE_API void SetFromArray(const TArray<int32>& SelectionArr);
+	 * Sets the selection from a sparse array (it only contains the indices of the selected items)
+	 * @return true if all indices were successfully set, false if some indices were invalid and could not be set
+	 */
+	DATAFLOWCORE_API bool SetFromArray(const TArray<int32>& SelectionArr);
 
 	/**
 	* Sets the selection from a dense array (it contains a true/false element for every item)
@@ -55,9 +56,10 @@ struct FDataflowSelection
 	const TBitArray<>& GetBitArray() const { return SelectionArray; };
 
 	/**
-	* Initialize from a Collection and an array
-	*/
-	DATAFLOWCORE_API void InitFromArray(const FManagedArrayCollection& InCollection, const TArray<int32>& InSelectionArr);
+	 * Initialize from a Collection and an array
+	 * @return true if all indices were successfully set, false if some indices were invalid and could not be set
+	 */
+	DATAFLOWCORE_API bool InitFromArray(const FManagedArrayCollection& InCollection, const TArray<int32>& InSelectionArr);
 
 	// @return true if this selection is valid for the given collection -- i.e., if the selection expects the group to have the correct number of elements
 	DATAFLOWCORE_API bool IsValidForCollection(const FManagedArrayCollection& InCollection) const;
