Chaos - add Rewind/Resim support to FPBDRigidsSolver::UnregisterObject for GeometryCollection and ClusterUnion

diff --git a/Engine/Source/Runtime/Experimental/Chaos/Private/PBDRigidsSolver.cpp b/Engine/Source/Runtime/Experimental/Chaos/Private/PBDRigidsSolver.cpp
--- a/Engine/Source/Runtime/Experimental/Chaos/Private/PBDRigidsSolver.cpp
+++ b/Engine/Source/Runtime/Experimental/Chaos/Private/PBDRigidsSolver.cpp
@@ -972,7 +972,9 @@ namespace Chaos
 
 			GeometryCollectionPhysicsProxies_Internal.RemoveSingle(InProxy);
 			InProxy->SyncBeforeDestroy();
-			PendingDestroyGeometryCollectionPhysicsProxy.Add(InProxy);
+
+			const int32 OffsetForRewind = (bResimDelayDestroyObject && MRewindData) ? MRewindData->Capacity() : 0;
+			PendingDestroyGeometryCollectionPhysicsProxy.Add(FPendingProxyDestroyInfo<FGeometryCollectionPhysicsProxy>{ InProxy, GetCurrentFrame() + OffsetForRewind });
 		});
 	}
 
@@ -1031,7 +1033,9 @@ namespace Chaos
 				}
 				ClusterUnionPhysicsProxies_Internal.RemoveSingle(Proxy);
 				Proxy->SyncBeforeDestroy();
-				PendingDestroyClusterUnionProxy.Add(Proxy);
+
+				const int32 OffsetForRewind = (bResimDelayDestroyObject && MRewindData) ? MRewindData->Capacity() : 0;
+				PendingDestroyClusterUnionProxy.Add(FPendingProxyDestroyInfo<FClusterUnionPhysicsProxy>{ Proxy, GetCurrentFrame() + OffsetForRewind });
 			}
 		);
 	}
@@ -1246,30 +1250,38 @@ namespace Chaos
 	void FPBDRigidsSolver::DestroyPendingProxies_Internal()
 	{
 		MarshallingManager.GetCurrentPullData_Internal()->DirtyClusterUnions.Reset();
-		for (FClusterUnionPhysicsProxy* Proxy : PendingDestroyClusterUnionProxy)
+		
+		for (int32 Idx = PendingDestroyClusterUnionProxy.Num() - 1; Idx >= 0; --Idx)
 		{
-			if (!Proxy)
-			{
-				continue;
-			}
+			FPendingProxyDestroyInfo<FClusterUnionPhysicsProxy>& Info = PendingDestroyClusterUnionProxy[Idx];
 
-			// Callback for Unregister PhysicsObject
-			if (PhysicsObjectUnregistrationWatchers.Num() > 0)
+			// Are we on or after the frame where we want to destroy the proxy
+			// Note: GetCurrentFrame() - 1 because this happens at the end of the physics frame after frame number has been incremented already
+			if (Info.DestroyOnStep <= (GetCurrentFrame() - 1) || IsShuttingDown()) 
 			{
-				for (ISimCallbackObject* Callback : PhysicsObjectUnregistrationWatchers)
+				if (Info.Proxy != nullptr)
 				{
-					Callback->OnPhysicsObjectUnregistered_Internal(Proxy->GetPhysicsObjectHandle());
+					// Callback for Unregister PhysicsObject
+					if (PhysicsObjectUnregistrationWatchers.Num() > 0)
+					{
+						for (ISimCallbackObject* Callback : PhysicsObjectUnregistrationWatchers)
+						{
+							Callback->OnPhysicsObjectUnregistered_Internal(Info.Proxy->GetPhysicsObjectHandle());
+						}
+					}
+
+					if (FPBDRigidsEvolutionGBF* Evolution = GetEvolution())
+					{
+						// Remove the cluster union this proxy manages which will in turn also the destroy the necessary particles.
+						Evolution->GetRigidClustering().GetClusterUnionManager().DestroyClusterUnion(Info.Proxy->GetClusterUnionIndex());
+					}
+
+					delete Info.Proxy;
 				}
-			}
 
-			if (FPBDRigidsEvolutionGBF* Evolution = GetEvolution())
-			{
-				// Remove the cluster union this proxy manages which will in turn also the destroy the necessary particles.
-				Evolution->GetRigidClustering().GetClusterUnionManager().DestroyClusterUnion(Proxy->GetClusterUnionIndex());
+				PendingDestroyClusterUnionProxy.RemoveAtSwap(Idx, EAllowShrinking::No);
 			}
-			delete Proxy;
 		}
-		PendingDestroyClusterUnionProxy.Reset();
 
 		// If we have any callback objects watching for particle deregistrations,
 		// send an array of proxies that are about to be deleted.
@@ -1319,30 +1331,39 @@ namespace Chaos
 
 		MarshallingManager.GetCurrentPullData_Internal()->DirtyGeometryCollections.Reset();
 
-		if (!PendingDestroyGeometryCollectionPhysicsProxy.IsEmpty())
+		for (int32 Idx = PendingDestroyGeometryCollectionPhysicsProxy.Num() - 1; Idx >= 0; --Idx)
 		{
-			GetEvolution()->GetRigidClustering().CleanupInternalClustersForProxies(TArrayView<IPhysicsProxyBase*>{reinterpret_cast<IPhysicsProxyBase**>(&PendingDestroyGeometryCollectionPhysicsProxy[0]), PendingDestroyGeometryCollectionPhysicsProxy.Num() });
-			for (auto Proxy : PendingDestroyGeometryCollectionPhysicsProxy)
+			FPendingProxyDestroyInfo<FGeometryCollectionPhysicsProxy>& Info = PendingDestroyGeometryCollectionPhysicsProxy[Idx];
+
+			GetEvolution()->GetRigidClustering().CleanupInternalClustersForProxies(MakeArrayView(reinterpret_cast<IPhysicsProxyBase**>(&Info.Proxy), 1));
+
+			// Are we on or after the frame where we want to destroy the proxy
+			// Note: GetCurrentFrame() - 1 because this happens at the end of the physics frame after frame number has been incremented already
+			if (Info.DestroyOnStep <= (GetCurrentFrame() - 1) || IsShuttingDown())
 			{
-				// Callback for Unregister PhysicsObject
-				if (PhysicsObjectUnregistrationWatchers.Num() > 0)
+				if (Info.Proxy != nullptr)
 				{
-					TArray<FPhysicsObjectHandle> PhysicsObjects = Proxy->GetAllPhysicsObjects();
-					for (FConstPhysicsObjectHandle PhysicsObject : PhysicsObjects)
+					// Callback for Unregister PhysicsObject
+					if (PhysicsObjectUnregistrationWatchers.Num() > 0)
 					{
-						for (ISimCallbackObject* Callback : PhysicsObjectUnregistrationWatchers)
+						TArray<FPhysicsObjectHandle> PhysicsObjects = Info.Proxy->GetAllPhysicsObjects();
+						for (FConstPhysicsObjectHandle PhysicsObject : PhysicsObjects)
 						{
-							Callback->OnPhysicsObjectUnregistered_Internal(PhysicsObject);
+							for (ISimCallbackObject* Callback : PhysicsObjectUnregistrationWatchers)
+							{
+								Callback->OnPhysicsObjectUnregistered_Internal(PhysicsObject);
+							}
 						}
 					}
+
+					// Removing the geometry collection from the solver a bit delayed. This lets the cluster union do its cleanup first before
+					// the geometry collection if they're all being destroyed at the same time.
+					Info.Proxy->OnRemoveFromSolver(this);
+					delete Info.Proxy;
 				}
 
-				// Removing the geometry collection from the solver a bit delayed. This lets the cluster union do its cleanup first before
-				// the geometry collection if they're all being destroyed at the same time.
-				Proxy->OnRemoveFromSolver(this);
-				delete Proxy;
+				PendingDestroyGeometryCollectionPhysicsProxy.RemoveAtSwap(Idx, EAllowShrinking::No);
 			}
-			PendingDestroyGeometryCollectionPhysicsProxy.Reset();
 		}
 	}
 
diff --git a/Engine/Source/Runtime/Experimental/Chaos/Public/PBDRigidsSolver.h b/Engine/Source/Runtime/Experimental/Chaos/Public/PBDRigidsSolver.h
--- a/Engine/Source/Runtime/Experimental/Chaos/Public/PBDRigidsSolver.h
+++ b/Engine/Source/Runtime/Experimental/Chaos/Public/PBDRigidsSolver.h
@@ -434,9 +434,16 @@ namespace Chaos
 			FUniqueIdx UniqueIdx;
 		};
 
+		template<typename ProxyType>
+		struct FPendingProxyDestroyInfo
+		{
+			ProxyType* Proxy;
+			int32 DestroyOnStep;
+		};
+
 		TArray<FPendingDestroyInfo> PendingDestroyPhysicsProxy;
-		TArray<FGeometryCollectionPhysicsProxy*> PendingDestroyGeometryCollectionPhysicsProxy;
-		TArray<FClusterUnionPhysicsProxy*> PendingDestroyClusterUnionProxy;
+		TArray<FPendingProxyDestroyInfo<FGeometryCollectionPhysicsProxy>> PendingDestroyGeometryCollectionPhysicsProxy;
+		TArray<FPendingProxyDestroyInfo<FClusterUnionPhysicsProxy>> PendingDestroyClusterUnionProxy;
 
 		CHAOS_API void ProcessSinglePushedData_Internal(FPushPhysicsData& PushData);
 		CHAOS_API virtual void ProcessPushedData_Internal(FPushPhysicsData& PushData) override;
