Remove unreachable code from loop with only one iteration at most in ISimulationModuleBase::GetFirstChild

diff --git a/Engine/Source/Runtime/Experimental/ChaosVehicles/ChaosVehiclesCore/Private/SimModule/SimulationModuleBase.cpp b/Engine/Source/Runtime/Experimental/ChaosVehicles/ChaosVehiclesCore/Private/SimModule/SimulationModuleBase.cpp
--- a/Engine/Source/Runtime/Experimental/ChaosVehicles/ChaosVehiclesCore/Private/SimModule/SimulationModuleBase.cpp
+++ b/Engine/Source/Runtime/Experimental/ChaosVehicles/ChaosVehiclesCore/Private/SimModule/SimulationModuleBase.cpp
@@ -55,9 +55,9 @@ ISimulationModuleBase* ISimulationModuleBase::GetFirstChild()
 	if (SimModuleTree)
 	{
 		const TSet<int32>& Children = SimModuleTree->GetChildren(SimTreeIndex);
-		for (const int ChildIndex : Children)
+		if (const auto It = Children.CreateConstIterator(); It)
 		{ 
-			return SimModuleTree->AccessSimModule(ChildIndex);
+			return SimModuleTree->AccessSimModule(*It);
 		}
 	}
 	return nullptr;
