Added dataflow nodes for conversion from an index or index array to a dataflow selection

diff --git a/Engine/Source/Runtime/Experimental/Dataflow/Core/Private/Dataflow/DataflowConvertNodes.cpp b/Engine/Source/Runtime/Experimental/Dataflow/Core/Private/Dataflow/DataflowConvertNodes.cpp
--- a/Engine/Source/Runtime/Experimental/Dataflow/Core/Private/Dataflow/DataflowConvertNodes.cpp
+++ b/Engine/Source/Runtime/Experimental/Dataflow/Core/Private/Dataflow/DataflowConvertNodes.cpp
@@ -22,7 +22,9 @@ namespace UE::Dataflow
 		// Needs to be fixed
 //		DATAFLOW_NODE_REGISTER_CREATION_FACTORY(FConvertUObjectConvertibleTypesDataflowNode);
 		DATAFLOW_NODE_REGISTER_CREATION_FACTORY(FConvertSelectionTypesDataflowNode);
-		DATAFLOW_NODE_REGISTER_CREATION_FACTORY(FConvertSelectionTypesToIndicesArrayDataflowNode);
+		DATAFLOW_NODE_REGISTER_CREATION_FACTORY(FConvertSelectionTypesToIndexArrayDataflowNode);
+		DATAFLOW_NODE_REGISTER_CREATION_FACTORY(FConvertIndexToSelectionTypesDataflowNode);
+		DATAFLOW_NODE_REGISTER_CREATION_FACTORY(FConvertIndexArrayToSelectionTypesDataflowNode);
 		DATAFLOW_NODE_REGISTER_CREATION_FACTORY(FConvertVectorArrayTypesDataflowNode);
 		DATAFLOW_NODE_REGISTER_CREATION_FACTORY(FConvertNumericArrayTypesDataflowNode);
 		DATAFLOW_NODE_REGISTER_CREATION_FACTORY(FConvertStringArrayTypesDataflowNode);
@@ -313,14 +315,108 @@ void FConvertSelectionTypesDataflowNode::Evaluate(UE::Dataflow::FContext& Contex
 
 //-----------------------------------------------------------------------------------------------
 
-FConvertSelectionTypesToIndicesArrayDataflowNode::FConvertSelectionTypesToIndicesArrayDataflowNode(const UE::Dataflow::FNodeParameters& InParam, FGuid InGuid)
+namespace UE::Private::DataflowConvertNodesLocal
+{
+	template<typename DataflowSelectionType>
+	bool SetMatchingSelectionHelper(const TArray<int32>& SelectionArr,
+		const FManagedArrayCollection& InCollection, UE::Dataflow::FContext& Context,
+		const FDataflowOutput* Output, const FDataflowSelectionTypes* OutMemberPtr, const FDataflowNode* ThisNode)
+	{
+		if (Output && Output->IsType<DataflowSelectionType>())
+		{
+			DataflowSelectionType NewSelection;
+			NewSelection.InitFromArray(InCollection, SelectionArr);
+			ThisNode->SetValue(Context, MoveTemp(NewSelection), OutMemberPtr);
+			return true;
+		}
+		return false;
+	}
+
+	bool SetMatchingSelection(const TArray<int32>& SelectionArr,
+		const FManagedArrayCollection& InCollection, UE::Dataflow::FContext& Context,
+		const FDataflowOutput* Output, const FDataflowSelectionTypes* OutMemberPtr, const FDataflowNode* ThisNode)
+	{
+		return SetMatchingSelectionHelper<FDataflowTransformSelection>(SelectionArr, InCollection, Context, Output, OutMemberPtr, ThisNode)
+			|| SetMatchingSelectionHelper<FDataflowGeometrySelection>(SelectionArr, InCollection, Context, Output, OutMemberPtr, ThisNode)
+			|| SetMatchingSelectionHelper<FDataflowFaceSelection>(SelectionArr, InCollection, Context, Output, OutMemberPtr, ThisNode)
+			|| SetMatchingSelectionHelper<FDataflowVertexSelection>(SelectionArr, InCollection, Context, Output, OutMemberPtr, ThisNode);
+	}
+}
+
+FConvertIndexArrayToSelectionTypesDataflowNode::FConvertIndexArrayToSelectionTypesDataflowNode(const UE::Dataflow::FNodeParameters& InParam, FGuid InGuid)
+	: FDataflowNode(InParam, InGuid)
+{
+	RegisterInputConnection(&Collection);
+	RegisterInputConnection(&In);
+	RegisterOutputConnection(&Collection, &Collection);
+	RegisterOutputConnection(&Out);
+}
+
+
+void FConvertIndexArrayToSelectionTypesDataflowNode::Evaluate(UE::Dataflow::FContext& Context, const FDataflowOutput* Output) const
+{
+	if (Output->IsA(&Out))
+	{
+		const bool bHasValidCollection = IsConnected(&Collection);
+		if (!bHasValidCollection)
+		{
+			Context.Warning(TEXT("Indices cannot be converted to a valid selection without a connected Collection"), this);
+		}
+		const FManagedArrayCollection& InCollection = GetValue(Context, &Collection);
+		UE::Private::DataflowConvertNodesLocal::SetMatchingSelection(
+			bHasValidCollection ? GetValue(Context, &In) : TArray<int32>(), InCollection, Context, Output, &Out, this);
+	}
+	else if (Output->IsA(&Collection))
+	{
+		SafeForwardInput(Context, &Collection, &Collection);
+	}
+}
+
+FConvertIndexToSelectionTypesDataflowNode::FConvertIndexToSelectionTypesDataflowNode(const UE::Dataflow::FNodeParameters& InParam, FGuid InGuid)
+	: FDataflowNode(InParam, InGuid)
+{
+	RegisterInputConnection(&Collection);
+	RegisterInputConnection(&In);
+	RegisterOutputConnection(&Collection, &Collection);
+	RegisterOutputConnection(&Out);
+}
+
+
+void FConvertIndexToSelectionTypesDataflowNode::Evaluate(UE::Dataflow::FContext& Context, const FDataflowOutput* Output) const
+{
+	if (Output->IsA(&Out))
+	{
+		TArray<int32> InValueArr;
+		const bool bHasValidCollection = IsConnected(&Collection);
+		if (bHasValidCollection)
+		{
+			int32 InValue = GetValue(Context, &In);
+			InValueArr.Add(InValue);
+		}
+		else
+		{
+			Context.Warning(TEXT("Index cannot be converted to a valid selection without a connected Collection"), this);
+		}
+		const FManagedArrayCollection& InCollection = GetValue(Context, &Collection);
+		UE::Private::DataflowConvertNodesLocal::SetMatchingSelection(InValueArr, InCollection, Context, Output, &Out, this);
+	}
+	else if (Output->IsA(&Collection))
+	{
+		SafeForwardInput(Context, &Collection, &Collection);
+	}
+}
+
+
+//-----------------------------------------------------------------------------------------------
+
+FConvertSelectionTypesToIndexArrayDataflowNode::FConvertSelectionTypesToIndexArrayDataflowNode(const UE::Dataflow::FNodeParameters& InParam, FGuid InGuid)
 	: FDataflowNode(InParam, InGuid)
 {
 	RegisterInputConnection(&In);
 	RegisterOutputConnection(&Out);
 }
 
-void FConvertSelectionTypesToIndicesArrayDataflowNode::Evaluate(UE::Dataflow::FContext& Context, const FDataflowOutput* Output) const
+void FConvertSelectionTypesToIndexArrayDataflowNode::Evaluate(UE::Dataflow::FContext& Context, const FDataflowOutput* Output) const
 {
 	if (Output->IsA(&Out))
 	{
diff --git a/Engine/Source/Runtime/Experimental/Dataflow/Core/Public/Dataflow/DataflowConvertNodes.h b/Engine/Source/Runtime/Experimental/Dataflow/Core/Public/Dataflow/DataflowConvertNodes.h
--- a/Engine/Source/Runtime/Experimental/Dataflow/Core/Public/Dataflow/DataflowConvertNodes.h
+++ b/Engine/Source/Runtime/Experimental/Dataflow/Core/Public/Dataflow/DataflowConvertNodes.h
@@ -247,10 +247,10 @@ public:
 * Convert Selection to Index Array
 */
 USTRUCT()
-struct FConvertSelectionTypesToIndicesArrayDataflowNode : public FDataflowNode
+struct FConvertSelectionTypesToIndexArrayDataflowNode : public FDataflowNode
 {
 	GENERATED_USTRUCT_BODY()
-	DATAFLOW_NODE_DEFINE_INTERNAL(FConvertSelectionTypesToIndicesArrayDataflowNode, "ConvertSelectionToIndicesArray", DATAFLOW_CONVERT_NODES_CATEGORY, "")
+	DATAFLOW_NODE_DEFINE_INTERNAL(FConvertSelectionTypesToIndexArrayDataflowNode, "ConvertSelectionToIndexArray", DATAFLOW_CONVERT_NODES_CATEGORY, "")
 
 private:
 
@@ -267,9 +267,74 @@ private:
 	//~ End FDataflowNode interface
 
 public:
-	FConvertSelectionTypesToIndicesArrayDataflowNode(const UE::Dataflow::FNodeParameters& InParam, FGuid InGuid = FGuid::NewGuid());
+	FConvertSelectionTypesToIndexArrayDataflowNode(const UE::Dataflow::FNodeParameters& InParam, FGuid InGuid = FGuid::NewGuid());
 };
 
+/**
+* Convert Index to Selection
+*/
+USTRUCT()
+struct FConvertIndexToSelectionTypesDataflowNode : public FDataflowNode
+{
+	GENERATED_USTRUCT_BODY()
+	DATAFLOW_NODE_DEFINE_INTERNAL(FConvertIndexToSelectionTypesDataflowNode, "ConvertIndexToSelection", DATAFLOW_CONVERT_NODES_CATEGORY, "")
+		DATAFLOW_NODE_RENDER_TYPE("SurfaceRender", FGeometryCollection::StaticType(), "Collection")
+
+public:
+	FConvertIndexToSelectionTypesDataflowNode(const UE::Dataflow::FNodeParameters& InParam, FGuid InGuid = FGuid::NewGuid());
+	
+private:
+
+	/** Collection for the selection */
+	UPROPERTY(meta = (DataflowInput, DataflowOutput, DataflowPassthrough = "Collection", DataflowIntrinsic))
+	FManagedArrayCollection Collection;
+
+	/** Input value */
+	UPROPERTY(meta = (DataflowInput, DataflowIntrinsic))
+	int32 In;
+
+	/** Output value */
+	UPROPERTY(meta = (DataflowOutput))
+	FDataflowSelectionTypes Out;
+
+	//~ Begin FDataflowNode interface
+	virtual void Evaluate(UE::Dataflow::FContext& Context, const FDataflowOutput* Out) const override;
+	//~ End FDataflowNode interface
+};
+
+/**
+* Convert Index Array to Selection
+*/
+USTRUCT()
+struct FConvertIndexArrayToSelectionTypesDataflowNode : public FDataflowNode
+{
+	GENERATED_USTRUCT_BODY()
+	DATAFLOW_NODE_DEFINE_INTERNAL(FConvertIndexArrayToSelectionTypesDataflowNode, "ConvertIndexArrayToSelection", DATAFLOW_CONVERT_NODES_CATEGORY, "")
+	DATAFLOW_NODE_RENDER_TYPE("SurfaceRender", FGeometryCollection::StaticType(), "Collection")
+	
+public:
+	FConvertIndexArrayToSelectionTypesDataflowNode(const UE::Dataflow::FNodeParameters& InParam, FGuid InGuid = FGuid::NewGuid());
+
+private:
+
+	/** Collection for the selection */
+	UPROPERTY(meta = (DataflowInput, DataflowOutput, DataflowPassthrough = "Collection", DataflowIntrinsic))
+	FManagedArrayCollection Collection;
+
+	/** Input value */
+	UPROPERTY(meta = (DataflowInput, DataflowIntrinsic))
+	TArray<int32> In;
+
+	/** Output value */
+	UPROPERTY(meta = (DataflowOutput))
+	FDataflowSelectionTypes Out;
+
+	//~ Begin FDataflowNode interface
+	virtual void Evaluate(UE::Dataflow::FContext& Context, const FDataflowOutput* Out) const override;
+	//~ End FDataflowNode interface
+};
+
+
 /**
 * Convert Vector array types
 */
