Extended GetInstallSize with the option to pin existing chunks and compute the required download size.

#rnx

diff --git a/Engine/Source/Runtime/Experimental/IoStore/OnDemand/Private/OnDemandInstallCache.cpp b/Engine/Source/Runtime/Experimental/IoStore/OnDemand/Private/OnDemandInstallCache.cpp
--- a/Engine/Source/Runtime/Experimental/IoStore/OnDemand/Private/OnDemandInstallCache.cpp
+++ b/Engine/Source/Runtime/Experimental/IoStore/OnDemand/Private/OnDemandInstallCache.cpp
@@ -8,6 +8,7 @@
 #include "Algo/Accumulate.h"
 #include "Algo/Find.h"
 #include "Async/Mutex.h"
+#include "Async/SharedMutex.h"
 #include "Async/UniqueLock.h"
 #include "Async/AsyncFileHandle.h"
 #include "Containers/UnrealString.h"
@@ -1482,6 +1483,11 @@ public:
 
 	// IOnDemandInstallCache
 	virtual bool								IsChunkCached(const FIoHash& ChunkHash) override;
+	virtual bool								TryPinChunks(
+													const FSharedOnDemandContainer& Container,
+													TConstArrayView<int32> EntryIndices,
+													FOnDemandContentHandle ContentHandle,
+													TArray<int32>& OutMissing) override;
 	virtual FIoStatus							PutChunk(FIoBuffer&& Chunk, const FIoHash& ChunkHash) override;
 	virtual FIoStatus							Purge(TSet<FIoHash>&& ChunksToInstall) override;
 	virtual FIoStatus							PurgeAllUnreferenced(bool bDefrag, const uint64* BytesToPurge = nullptr) override;
@@ -1520,7 +1526,8 @@ private:
 	FUniquePendingChunks		PendingChunks;
 	FSharedBackendContext		BackendContext;
 	FIoRequestList				CompletedRequests;
-	UE::FMutex					Mutex;
+	FMutex						Mutex;
+	FSharedMutex				PurgeDefragMutex;
 	uint64						MaxCacheSize{ 0 };
 	uint64						MaxJournalSize;
 
@@ -1978,6 +1985,37 @@ bool FOnDemandInstallCache::IsChunkCached(const FIoHash& ChunkHash)
 	return Loc.IsValid();
 }
 
+bool FOnDemandInstallCache::TryPinChunks(
+	const FSharedOnDemandContainer& Container,
+	TConstArrayView<int32> EntryIndices,
+	FOnDemandContentHandle ContentHandle,
+	TArray<int32>& OutMissing)
+{
+	// Currently we can only pin chunks if we are not purging/defragging the cache 
+	if (PurgeDefragMutex.TryLockShared() == false)
+	{
+		OutMissing = EntryIndices;
+		return false;
+	}
+
+	ON_SCOPE_EXIT { PurgeDefragMutex.UnlockShared(); };
+
+	for (int32 EntryIndex : EntryIndices)
+	{
+		const FOnDemandChunkEntry& ChunkEntry = Container->ChunkEntries[EntryIndex];
+		if (FCasLocation Loc = Cas.FindChunk(ChunkEntry.Hash); Loc.IsValid())
+		{
+			IoStore.AddReference(Container, EntryIndex, ContentHandle);
+		}
+		else
+		{
+			OutMissing.Add(EntryIndex);
+		}
+	}
+
+	return true;
+}
+
 FIoStatus FOnDemandInstallCache::PutChunk(FIoBuffer&& Chunk, const FIoHash& ChunkHash)
 {
 	if (PendingChunks.IsValid() == false)
@@ -2000,6 +2038,8 @@ FIoStatus FOnDemandInstallCache::PutChunk(FIoBuffer&& Chunk, const FIoHash& Chun
 
 FIoStatus FOnDemandInstallCache::Purge(TSet<FIoHash>&& ChunksToInstall)
 {
+	TUniqueLock Lock(PurgeDefragMutex);
+
 	FCasBlockInfoMap	BlockInfo;
 	const uint64		TotalCachedBytes = Cas.GetBlockInfo(BlockInfo);
 
@@ -2090,6 +2130,8 @@ FIoStatus FOnDemandInstallCache::Purge(TSet<FIoHash>&& ChunksToInstall)
 
 FIoStatus FOnDemandInstallCache::PurgeAllUnreferenced(bool bDefrag, const uint64* BytesToPurge /*= nullptr*/)
 {
+	TUniqueLock Lock(PurgeDefragMutex);
+
 	FCasBlockInfoMap	BlockInfo;
 	const uint64		TotalCachedBytes = Cas.GetBlockInfo(BlockInfo);
 
@@ -2157,6 +2199,8 @@ FIoStatus FOnDemandInstallCache::PurgeAllUnreferenced(bool bDefrag, const uint64
 
 FIoStatus FOnDemandInstallCache::DefragAll(const uint64* BytesToFree /*= nullptr*/)
 {
+	TUniqueLock Lock(PurgeDefragMutex);
+
 	FCasBlockInfoMap	BlockInfo;
 	const uint64		TotalCachedBytes = Cas.GetBlockInfo(BlockInfo);
 
diff --git a/Engine/Source/Runtime/Experimental/IoStore/OnDemand/Private/OnDemandInstallCache.h b/Engine/Source/Runtime/Experimental/IoStore/OnDemand/Private/OnDemandInstallCache.h
--- a/Engine/Source/Runtime/Experimental/IoStore/OnDemand/Private/OnDemandInstallCache.h
+++ b/Engine/Source/Runtime/Experimental/IoStore/OnDemand/Private/OnDemandInstallCache.h
@@ -11,9 +11,10 @@ struct FIoContainerHeader;
 namespace UE::IoStore
 {
 
+class FOnDemandContentHandle;
 class FOnDemandIoStore;
-struct FOnDemandContainer;
 struct FOnDemandInstallCacheUsage;
+using FSharedOnDemandContainer = TSharedPtr<struct FOnDemandContainer, ESPMode::ThreadSafe>;
 
 struct FOnDemandInstallCacheStorageUsage
 {
@@ -28,6 +29,11 @@ class IOnDemandInstallCache
 public:
 	virtual										~IOnDemandInstallCache() = default;
 	virtual bool								IsChunkCached(const FIoHash& ChunkHash) = 0;
+	virtual bool								TryPinChunks(
+													const FSharedOnDemandContainer& Container,
+													TConstArrayView<int32> EntryIndices,
+													FOnDemandContentHandle ContentHandle,
+													TArray<int32>& OutMissing) = 0;
 	virtual FIoStatus							PutChunk(FIoBuffer&& Chunk, const FIoHash& ChunkHash) = 0;
 	virtual FIoStatus							Purge(TSet<FIoHash>&& ChunksToInstall) = 0;
 	virtual FIoStatus							PurgeAllUnreferenced(bool bDefrag, const uint64* BytesToPurge = nullptr) = 0;
diff --git a/Engine/Source/Runtime/Experimental/IoStore/OnDemand/Private/OnDemandIoStore.cpp b/Engine/Source/Runtime/Experimental/IoStore/OnDemand/Private/OnDemandIoStore.cpp
--- a/Engine/Source/Runtime/Experimental/IoStore/OnDemand/Private/OnDemandIoStore.cpp
+++ b/Engine/Source/Runtime/Experimental/IoStore/OnDemand/Private/OnDemandIoStore.cpp
@@ -636,7 +636,7 @@ FIoStatus FOnDemandIoStore::Unmount(FStringView MountId)
 	return EIoErrorCode::Ok;
 }
 
-TIoStatusOr<uint64> FOnDemandIoStore::GetInstallSize(const FOnDemandGetInstallSizeArgs& Args) const
+TIoStatusOr<FOnDemandInstallSizeResult> FOnDemandIoStore::GetInstallSize(const FOnDemandGetInstallSizeArgs& Args) const
 {
 	TRACE_CPUPROFILER_EVENT_SCOPE(FOnDemandIoStore::GetInstallSize);
 
@@ -667,13 +667,61 @@ TIoStatusOr<uint64> FOnDemandIoStore::GetInstallSize(const FOnDemandGetInstallSi
 		bIncludeOptionalBulkData,
 		ResolvedChunks,
 		Missing);
+	
+	FOnDemandInstallSizeResult Result;
 
-	const uint64 InstallSize = Algo::TransformAccumulate(
-		ResolvedChunks,
-		[](const FResolvedContainerChunks& R) { return R.TotalSize; },
-		uint64(0));
+	const bool bPin = Args.ContentHandle.IsValid();
+	if (bPin == false)
+	{
+		Result.InstallSize = Algo::TransformAccumulate(
+			ResolvedChunks,
+			[](const FResolvedContainerChunks& R) { return R.TotalSize; },
+			uint64(0));
+
+		return Result;
+	}
+
+	TArray<int32>						Uncached;
+	TArray<int32, TInlineAllocator<64>> Unreferenced;
+
+	Result.DownloadSize.Emplace(0);
+	for (const FResolvedContainerChunks& Resolved : ResolvedChunks)
+	{
+		Result.InstallSize += Resolved.TotalSize;
+
+		Uncached.Reset();
+		Unreferenced.Reset();
+		{
+			TUniqueLock Lock(ContainerMutex);
+			FOnDemandChunkEntryReferences& Refs = Resolved.Container->FindOrAddChunkEntryReferences(*Args.ContentHandle.Handle);
+			for (int32 EntryIndex : Resolved.EntryIndices)
+			{
+				const bool bReferenced = Refs.Indices[EntryIndex];
+				if (bReferenced == false)
+				{
+					Unreferenced.Add(EntryIndex);
+				}
+			}
+		}
+
+		if (Unreferenced.IsEmpty() == false)
+		{
+			const bool bOk = InstallCache->TryPinChunks(
+					Resolved.Container.ToSharedRef(),
+					Unreferenced,
+					Args.ContentHandle,
+					Uncached);
+			UE_CLOG(!bOk, LogIoStoreOnDemand, Log, TEXT("Failed to pin install cache chunks due to purging/defragging"))
 
-	return InstallSize;
+			for (int32 EntryIndex : Uncached)
+			{
+				const FOnDemandChunkEntry& Entry = Resolved.Container->ChunkEntries[EntryIndex];
+				Result.DownloadSize.GetValue() += Entry.GetDiskSize(); 
+			}
+		}
+	}
+
+	return Result;
 }
 
 FIoStatus FOnDemandIoStore::GetInstallSizesByMountId(const FOnDemandGetInstallSizeArgs& Args, TMap<FString, uint64>& OutSizesByMountId) const
@@ -1715,6 +1763,13 @@ void FOnDemandIoStore::GetReferencedContentByHandle(TMap<FOnDemandWeakContentHan
 	}
 }
 
+void FOnDemandIoStore::AddReference(const FSharedOnDemandContainer& Container, int32 EntryIndex, FOnDemandContentHandle ContentHandle)
+{
+	UE::TUniqueLock Lock(ContainerMutex);
+	FOnDemandChunkEntryReferences& Refs = Container->FindOrAddChunkEntryReferences(*ContentHandle.Handle);
+	Refs.Indices[EntryIndex] = true;
+}
+
 void FOnDemandIoStore::CancelInstallRequest(FSharedInternalInstallRequest InstallRequest)
 {
 	check(InstallRequest.IsValid());
diff --git a/Engine/Source/Runtime/Experimental/IoStore/OnDemand/Private/OnDemandIoStore.h b/Engine/Source/Runtime/Experimental/IoStore/OnDemand/Private/OnDemandIoStore.h
--- a/Engine/Source/Runtime/Experimental/IoStore/OnDemand/Private/OnDemandIoStore.h
+++ b/Engine/Source/Runtime/Experimental/IoStore/OnDemand/Private/OnDemandIoStore.h
@@ -335,6 +335,7 @@ public:
 	void								GetReferencedContent(TArray<FSharedOnDemandContainer>& OutContainers, TArray<TBitArray<>>& OutChunkEntryIndices, bool bPackageStore = false);
 	TBitArray<>							GetReferencedContent(const FSharedOnDemandContainer& Container);
 	void								GetReferencedContentByHandle(TMap<FOnDemandWeakContentHandle, TArray<FOnDemandContainerChunkEntryReferences>>& OutReferencesByHandle) const;
+	void								AddReference(const FSharedOnDemandContainer& Container, int32 EntryIndex, FOnDemandContentHandle ContentHandle);
 
 	// IOnDemandIoStore interface
 	virtual FIoStatus					Initialize() override;
@@ -347,7 +348,7 @@ public:
 	virtual void						Defrag(FOnDemandDefragArgs&& Args, FOnDemandDefragCompleted&& OnCompleted) override;
 	virtual void						Verify(FOnDemandVerifyCacheCompleted&& OnCompleted) override;
 	virtual FIoStatus					Unmount(FStringView MountId) override;
-	virtual TIoStatusOr<uint64> 		GetInstallSize(const FOnDemandGetInstallSizeArgs& Args) const override;
+	virtual TIoStatusOr<FOnDemandInstallSizeResult> GetInstallSize(const FOnDemandGetInstallSizeArgs& Args) const override;
 	virtual FIoStatus					GetInstallSizesByMountId(const FOnDemandGetInstallSizeArgs& Args, TMap<FString, uint64>& OutSizesByMountId) const override;
 	virtual FOnDemandCacheUsage			GetCacheUsage(const FOnDemandGetCacheUsageArgs& Args) const override;
 	virtual bool						IsOnDemandStreamingEnabled() const override;
diff --git a/Engine/Source/Runtime/Experimental/IoStore/OnDemandCore/Public/IO/IoStoreOnDemand.h b/Engine/Source/Runtime/Experimental/IoStore/OnDemandCore/Public/IO/IoStoreOnDemand.h
--- a/Engine/Source/Runtime/Experimental/IoStore/OnDemandCore/Public/IO/IoStoreOnDemand.h
+++ b/Engine/Source/Runtime/Experimental/IoStore/OnDemandCore/Public/IO/IoStoreOnDemand.h
@@ -394,6 +394,17 @@ struct FOnDemandGetInstallSizeArgs
 	TArray<FPackageId> PackageIds;
 	/** Options */
 	EOnDemandGetInstallSizeOptions Options = EOnDemandGetInstallSizeOptions::None;
+	/** Pin existing chunk(s) with the specified content handle and compute remaining bytes to download. */
+	FOnDemandContentHandle ContentHandle;
+};
+
+/** Holds information about an install size request. */ 
+struct FOnDemandInstallSizeResult
+{
+	/** The total size of the content. */
+	uint64	InstallSize = 0;
+	/** The total size to download. This value is set when pinning existing chunk(s). */
+	TOptional<uint64> DownloadSize = 0;
 };
 
 /** Options for controlling the behavior of the install size request. */
@@ -539,7 +550,7 @@ public:
 	 * Get the total space in bytes needed to install the specified content.
 	 * @param Args	Arguments for computing required install size.
 	 */
-	virtual TIoStatusOr<uint64> GetInstallSize(const FOnDemandGetInstallSizeArgs& Args) const = 0;
+	virtual TIoStatusOr<FOnDemandInstallSizeResult> GetInstallSize(const FOnDemandGetInstallSizeArgs& Args) const = 0;
 	/**
 	 * Get the total space in bytes needed to install the specified content.
 	 * @param Args	Arguments for computing required install size.
