Stop calling FIndexedCacheStorageManager::Initialize / LoadModuleChecked("IndexedCacheStorage") lazily from FIndexedCacheStorageManager::Get() / IIndexedCacheStorage::Get() :  we may end up trying to load the module outside of the GameThread

Instead, FIndexedCacheStorageManager::Initialize is now called from IIndexedCacheStorageModule::StartupModule.
Change GetOrCreateIndexedCacheStorage to be part of IIndexedCacheStorage instead of IIndexedCacheStorageModule

diff --git a/Engine/Source/Runtime/Experimental/IndexedCacheStorage/Private/IndexedCacheStorage.cpp b/Engine/Source/Runtime/Experimental/IndexedCacheStorage/Private/IndexedCacheStorage.cpp
--- a/Engine/Source/Runtime/Experimental/IndexedCacheStorage/Private/IndexedCacheStorage.cpp
+++ b/Engine/Source/Runtime/Experimental/IndexedCacheStorage/Private/IndexedCacheStorage.cpp
@@ -58,7 +58,7 @@ public:
 #endif
 
 #if (HAS_CUSTOM_INDEXED_CACHE_STORAGE == 0)
-IIndexedCacheStorage& IIndexedCacheStorageModule::GetOrCreateIndexedCacheStorage()
+IIndexedCacheStorage& IIndexedCacheStorage::GetOrCreateIndexedCacheStorage()
 {
 	static FGenericIndexedCacheStorage Instance;
 	return Instance;
@@ -67,9 +67,9 @@ IIndexedCacheStorage& IIndexedCacheStorageModule::GetOrCreateIndexedCacheStorage
 
 IIndexedCacheStorage& IIndexedCacheStorage::Get()
 {
-	return IIndexedCacheStorageModule::Get().GetOrCreateIndexedCacheStorage();
+	return IIndexedCacheStorage::GetOrCreateIndexedCacheStorage();
 }
 
-
 }
-IMPLEMENT_MODULE(Experimental::IIndexedCacheStorageModule, IndexedCacheStorage);
+
+
diff --git a/Engine/Source/Runtime/Experimental/IndexedCacheStorage/Private/IndexedCacheStorage.h b/Engine/Source/Runtime/Experimental/IndexedCacheStorage/Private/IndexedCacheStorage.h
--- a/Engine/Source/Runtime/Experimental/IndexedCacheStorage/Private/IndexedCacheStorage.h
+++ b/Engine/Source/Runtime/Experimental/IndexedCacheStorage/Private/IndexedCacheStorage.h
@@ -5,27 +5,11 @@
 #include "CoreTypes.h"
 #include "Containers/Array.h"
 
-#include "Modules/ModuleInterface.h"
-#include "Modules/ModuleManager.h"
-
 namespace Experimental
 {
 
 class IIndexedCacheStorage;
 
-class IIndexedCacheStorageModule : public IModuleInterface
-{
-public:
-	static IIndexedCacheStorageModule& Get() 
-	{
-		return FModuleManager::Get().LoadModuleChecked<IIndexedCacheStorageModule>("IndexedCacheStorage");
-	}
-
-	IIndexedCacheStorage& GetOrCreateIndexedCacheStorage();
-
-};
-
-
 /**
   * Interface for managing cache storages: some platforms may store user data in different cache storages that may be created / destroyed
   * This allows the title to manage its disk space with a finer granularity. 
@@ -36,6 +20,8 @@ class IIndexedCacheStorage
 public:
 	static IIndexedCacheStorage& Get();
 
+	static IIndexedCacheStorage& GetOrCreateIndexedCacheStorage();
+
 	/**
 	  * Check if indexed cache storage is available on the current platform.
 	  * 
diff --git a/Engine/Source/Runtime/Experimental/IndexedCacheStorage/Private/IndexedCacheStorageManager.cpp b/Engine/Source/Runtime/Experimental/IndexedCacheStorage/Private/IndexedCacheStorageManager.cpp
--- a/Engine/Source/Runtime/Experimental/IndexedCacheStorage/Private/IndexedCacheStorageManager.cpp
+++ b/Engine/Source/Runtime/Experimental/IndexedCacheStorage/Private/IndexedCacheStorageManager.cpp
@@ -12,16 +12,18 @@ DEFINE_LOG_CATEGORY(LogIndexedCacheStorageManager);
 namespace Experimental
 {
 
+static FIndexedCacheStorageManager* GIndexedCacheStorageManager = nullptr;
+
 FIndexedCacheStorageManager& FIndexedCacheStorageManager::Get()
 {
-	static FIndexedCacheStorageManager GIndexedCacheStorageManager;
+	static FIndexedCacheStorageManager Instance;
 	UE_CALL_ONCE([&]
 	{
-		GIndexedCacheStorageManager.Initialize();
+		GIndexedCacheStorageManager = &Instance;
 	}
 	);
 	
-	return GIndexedCacheStorageManager;
+	return Instance;
 }
 
 void FIndexedCacheStorageManager::Initialize()
diff --git a/Engine/Source/Runtime/Experimental/IndexedCacheStorage/Private/IndexedCacheStorageModule.cpp b/Engine/Source/Runtime/Experimental/IndexedCacheStorage/Private/IndexedCacheStorageModule.cpp
new file mode 100644
--- /dev/null
+++ b/Engine/Source/Runtime/Experimental/IndexedCacheStorage/Private/IndexedCacheStorageModule.cpp
@@ -0,0 +1,22 @@
+// Copyright Epic Games, Inc. All Rights Reserved.
+
+#include "IndexedCacheStorageManager.h"
+
+#include "Modules/ModuleManager.h"
+
+namespace Experimental
+{
+
+void FIndexedCacheStorageModule::StartupModule()
+{
+	FIndexedCacheStorageManager::Get().Initialize();
+}
+
+void FIndexedCacheStorageModule::ShutdownModule()
+{
+}
+
+}
+
+IMPLEMENT_MODULE(Experimental::FIndexedCacheStorageModule, IndexedCacheStorage);
+
diff --git a/Engine/Source/Runtime/Experimental/IndexedCacheStorage/Public/IndexedCacheStorageManager.h b/Engine/Source/Runtime/Experimental/IndexedCacheStorage/Public/IndexedCacheStorageManager.h
--- a/Engine/Source/Runtime/Experimental/IndexedCacheStorage/Public/IndexedCacheStorageManager.h
+++ b/Engine/Source/Runtime/Experimental/IndexedCacheStorage/Public/IndexedCacheStorageManager.h
@@ -4,15 +4,24 @@
 
 #include "Containers/Array.h"
 #include "Containers/Map.h"
+#include "Containers/UnrealString.h"
 #include "HAL/Platform.h"
 #include "HAL/CriticalSection.h"
-#include "Containers/UnrealString.h"
+#include "Modules/ModuleInterface.h"
 
 #define UE_API INDEXEDCACHESTORAGE_API
 
 namespace Experimental
 {
 
+class FIndexedCacheStorageModule : public IModuleInterface
+{
+public:
+	/** IModuleInterface implementation */
+	virtual void StartupModule() override;
+	virtual void ShutdownModule() override;
+};
+
 // Manager to handle ref-counting of ICS entries
 class FIndexedCacheStorageManager
 {
@@ -53,6 +62,7 @@ public:
 	UE_API FString GetMountPath(int32 CacheIndex);
 
 private:
+	friend class FIndexedCacheStorageModule;
 
 	void Initialize();
 
