Clean up IAX periodic logging.

#rnx

- Increase the length of time between logging to 60 seconds, up from 30 seconds.
- Moved the initial processing of values to the inner loop so they are only calculated when needed.

diff --git a/Engine/Source/Runtime/Experimental/IoStore/OnDemand/Private/Statistics.cpp b/Engine/Source/Runtime/Experimental/IoStore/OnDemand/Private/Statistics.cpp
--- a/Engine/Source/Runtime/Experimental/IoStore/OnDemand/Private/Statistics.cpp
+++ b/Engine/Source/Runtime/Experimental/IoStore/OnDemand/Private/Statistics.cpp
@@ -27,7 +27,7 @@ LLM_DEFINE_TAG(Ias);
 namespace UE::IoStore
 {
 
-float GIasStatisticsLogInterval = 30.f;
+float GIasStatisticsLogInterval = 60.f;
 static FAutoConsoleVariableRef CVar_StatisticsLogInterval(
 	TEXT("ias.StatisticsLogInterval"),
 	GIasStatisticsLogInterval,
@@ -569,7 +569,13 @@ void FOnDemandIoBackendStats::UpdateCSVValues()
 
 void FOnDemandIoBackendStats::PrintPeriodicLogging()
 {
-	if (GIasStatisticsLogInterval > 0.f)
+	if (GIasStatisticsLogInterval <= 0.f)
+	{
+		return;
+	}
+
+	static double LastLogTime = 0.0;
+	if (double Time = FPlatformTime::Seconds(); Time - LastLogTime > (double)GIasStatisticsLogInterval)
 	{
 		static constexpr float OneOver1024 = 1.0f / 1024.0f;
 
@@ -584,49 +590,45 @@ void FOnDemandIoBackendStats::PrintPeriodicLogging()
 		float CReadKiB = (float)GCacheReadBytes.Get() * OneOver1024;
 		float CRejectedKiB = (float)GCacheRejectBytes.Get() * OneOver1024;
 
-		static double LastLogTime = 0.0;
-		if (double Time = FPlatformTime::Seconds(); Time - LastLogTime > (double)GIasStatisticsLogInterval)
+		if (BackendStatus.IsCacheEnabled())
 		{
-			if (BackendStatus.IsCacheEnabled())
-			{
-				UE_LOG(LogIas, Log, TEXT("CacheStats: CachedKiB=%d, WrittenKiB=%d, ReadKiB=%d, RejectedKiB=%d, Get=%d, Error=%d, Put=%d, PutReject=%d, PutExisting=%d"),
-					(int32)CCachedKiB, (int32)CWrittenKiB, (int32)CReadKiB, (int32)CRejectedKiB, CGetCount, CErrorCount, CPutCount, CPutRejectCount, CPutExistingCount);
-			}
-			else
+			UE_LOG(LogIas, Log, TEXT("CacheStats: CachedKiB=%d, WrittenKiB=%d, ReadKiB=%d, RejectedKiB=%d, Get=%d, Error=%d, Put=%d, PutReject=%d, PutExisting=%d"),
+				(int32)CCachedKiB, (int32)CWrittenKiB, (int32)CReadKiB, (int32)CRejectedKiB, CGetCount, CErrorCount, CPutCount, CPutRejectCount, CPutExistingCount);
+		}
+		else
+		{
+			UE_LOG(LogIas, Log, TEXT("CacheStats: Disabled"));
+		}
+
+		auto HttpLog = [](const TCHAR* Title, EHttpRequestType Type)
 			{
-				UE_LOG(LogIas, Log, TEXT("CacheStats: Disabled"));
-			}
+				UE_LOG(LogIas, Log, TEXT("%s - HttpStats: DownloadedKiB=%d, Get=%d, Retry=%d, Cancel=%d, Error=%d, CurPending=%d, CurDurationMsAvg=%d, CurDurationMsMax=%d"),
+					Title,
+					(int32)((float)GHttpDownloadedBytes[(uint8)Type].Get() * OneOver1024),
+					(int32)GHttpGetCount[(uint8)Type].Get(),
+					(int32)GHttpRetryCount[(uint8)Type].Get(),
+					(int32)GHttpCancelCount[(uint8)Type].Get(),
+					(int32)GHttpErrorCount[(uint8)Type].Get(),
+					(int32)GHttpPendingCount[(uint8)Type].Get(),
+					(int32)GHttpDurationMsAvg[(uint8)Type],
+					GHttpDurationMsMax[(uint8)Type]);
+			};
 
-			auto HttpLog = [](const TCHAR* Title, EHttpRequestType Type)
-				{
-					UE_LOG(LogIas, Log, TEXT("%s - HttpStats: DownloadedKiB=%d, Get=%d, Retry=%d, Cancel=%d, Error=%d, CurPending=%d, CurDurationMsAvg=%d, CurDurationMsMax=%d"),
-						Title,
-						(int32)((float)GHttpDownloadedBytes[(uint8)Type].Get() * OneOver1024),
-						(int32)GHttpGetCount[(uint8)Type].Get(),
-						(int32)GHttpRetryCount[(uint8)Type].Get(),
-						(int32)GHttpCancelCount[(uint8)Type].Get(),
-						(int32)GHttpErrorCount[(uint8)Type].Get(),
-						(int32)GHttpPendingCount[(uint8)Type].Get(),
-						(int32)GHttpDurationMsAvg[(uint8)Type],
-						GHttpDurationMsMax[(uint8)Type]);
-				};
-
-			HttpLog(TEXT("IAS"), EHttpRequestType::Streaming);
-			HttpLog(TEXT("IAD"), EHttpRequestType::Installed);
+		HttpLog(TEXT("IAS"), EHttpRequestType::Streaming);
+		HttpLog(TEXT("IAD"), EHttpRequestType::Installed);
 
 #if COUNTERSTRACE_ENABLED
-			{
-				TUniqueLock Lock(InstallerTraceCounters);
-				UE_LOG(LogIoStoreOnDemand, Log, TEXT("IadStats: InstallCount=%llu, Downloaded=%d KiB, AvgInstallDuration=%dms, AvgCacheHitRatio=%d%%"),
-					InstallerTraceCounters.InstallCount.Get(),
-					(int32(InstallerTraceCounters.DownloadedBytes.Get()) / 1024),
-					int32(InstallerTraceCounters.AvgInstallDurationMs.Get()),
-					int32(InstallerTraceCounters.CacheHitRatio.GetMean() * 100.0));
-			}
+		{
+			TUniqueLock Lock(InstallerTraceCounters);
+			UE_LOG(LogIoStoreOnDemand, Log, TEXT("IadStats: InstallCount=%llu, Downloaded=%d KiB, AvgInstallDuration=%dms, AvgCacheHitRatio=%d%%"),
+				InstallerTraceCounters.InstallCount.Get(),
+				(int32(InstallerTraceCounters.DownloadedBytes.Get()) / 1024),
+				int32(InstallerTraceCounters.AvgInstallDurationMs.Get()),
+				int32(InstallerTraceCounters.CacheHitRatio.GetMean() * 100.0));
+		}
 #endif // COUNTERSTRACE_ENABLED
 
-			LastLogTime = Time;
-		}
+		LastLogTime = Time;
 	}
 }
 
