Move the code for 'ias.skipdevcdns' so that it can be applied to all IAX hostgroups.

#rnx

- Previously we checked this cmdline option when resolving the urls from the distributed endpoint but this would only apply to IAS.
- Now we run this code when resolving the lower level FOnDemandHostGroup itself so that it applies to all systems using IAX.

diff --git a/Engine/Source/Runtime/Experimental/IoStore/OnDemand/Private/DistributionEndpoints.cpp b/Engine/Source/Runtime/Experimental/IoStore/OnDemand/Private/DistributionEndpoints.cpp
--- a/Engine/Source/Runtime/Experimental/IoStore/OnDemand/Private/DistributionEndpoints.cpp
+++ b/Engine/Source/Runtime/Experimental/IoStore/OnDemand/Private/DistributionEndpoints.cpp
@@ -19,10 +19,6 @@
 #include "Statistics.h"
 #include "Tasks/Task.h"
 
-#if !UE_BUILD_SHIPPING
-#include "Internationalization/Regex.h"
-#endif // !UE_BUILD_SHIPPING
-
 namespace UE::IoStore
 {
 
@@ -105,28 +101,6 @@ FDistributionEndpoints::EResult FDistributionEndpoints::ResolveEndpoints(const F
 	while (Loop.Tick(-1))
 		;
 
-#if !UE_BUILD_SHIPPING
-	if (FParse::Param(FCommandLine::Get(), TEXT("Ias.SkipDevCDNs")))
-	{
-		FString ConfigRegex;
-		GConfig->GetString(TEXT("Ias"), TEXT("DevelopmentCDNPattern"), ConfigRegex, GEngineIni);
-
-		if(!ConfigRegex.IsEmpty())
-		{
-			FRegexPattern Pattern(ConfigRegex);
-
-			const int32 OriginaNum = OutServiceUrls.Num();
-			OutServiceUrls.SetNum(Algo::StableRemoveIf(OutServiceUrls, [&Pattern](const FString& Entry)
-			{
-				FRegexMatcher Regex = FRegexMatcher(Pattern, Entry);
-				return Regex.FindNext();
-			}));
-
-			UE_LOG(LogIas, Log, TEXT("Removed %d development CDNs from the distributed endpoint list"), OriginaNum - OutServiceUrls.Num());
-		}
-	}
-#endif //!UE_BUILD_SHIPPING
-
 	UE_CLOG(Result == EResult::Success, LogIas, Log, TEXT("Successfully resolved distributed endpoint '%s' %d urls found"), *DistributionUrl, OutServiceUrls.Num());
 	UE_CLOG(Result != EResult::Success, LogIas, Log, TEXT("Failed to resolve distributed endpoint '%s'"), *DistributionUrl);
 
diff --git a/Engine/Source/Runtime/Experimental/IoStore/OnDemandCore/Private/OnDemandHostGroup.cpp b/Engine/Source/Runtime/Experimental/IoStore/OnDemandCore/Private/OnDemandHostGroup.cpp
--- a/Engine/Source/Runtime/Experimental/IoStore/OnDemandCore/Private/OnDemandHostGroup.cpp
+++ b/Engine/Source/Runtime/Experimental/IoStore/OnDemandCore/Private/OnDemandHostGroup.cpp
@@ -7,6 +7,15 @@
 #include "Containers/StringView.h"
 #include "HAL/IConsoleManager.h"
 
+#if !UE_BUILD_SHIPPING
+#include "IO/IoStoreOnDemand.h"
+#include "Internationalization/Regex.h"
+#include "Logging/StructuredLog.h"
+#include "Misc/CommandLine.h"
+#include "Misc/ConfigCacheIni.h"
+#include "Misc/Parse.h"
+#endif // !UE_BUILD_SHIPPING
+
 namespace UE::IoStore
 {
 
@@ -33,6 +42,58 @@ static bool ValidateUrl(FAnsiStringView Url, FString& Reason)
 	return Url.StartsWith("http") || Url.StartsWith("https");
 }
 
+#if !UE_BUILD_SHIPPING
+
+static bool ShouldSkipUrl(const FString& Url)
+{
+	struct FDevCDNCheck
+	{
+		FDevCDNCheck()
+		{
+			if (FParse::Param(FCommandLine::Get(), TEXT("Ias.SkipDevCDNs")))
+			{
+				FString ConfigRegex;
+				GConfig->GetString(TEXT("Ias"), TEXT("DevelopmentCDNPattern"), ConfigRegex, GEngineIni);
+
+				if (!ConfigRegex.IsEmpty())
+				{
+					Pattern = MakeUnique<FRegexPattern>(MoveTemp(ConfigRegex));
+				}
+			}
+		}
+
+		bool IsDevCNDUrl(const FString& Url) const
+		{
+			if (Pattern == nullptr)
+			{
+				return false;
+			}
+			else
+			{
+				FRegexMatcher Regex = FRegexMatcher(*Pattern, Url);
+				return Regex.FindNext();
+			}
+		}
+	private:
+
+		TUniquePtr<FRegexPattern> Pattern;
+	} static DevCDNCheck;
+
+	return DevCDNCheck.IsDevCNDUrl(Url);
+}
+
+static bool ShouldSkipUrl(const FStringView& Url)
+{
+	return ShouldSkipUrl(FString(Url));
+}
+
+static bool ShouldSkipUrl(FAnsiStringView Url)
+{
+	return ShouldSkipUrl(FString(Url));
+}
+
+#endif // !UE_BUILD_SHIPPING
+
 } // namespace Private
 
 struct FOnDemandHostGroup::FImpl
@@ -152,13 +213,23 @@ TIoStatusOr<FOnDemandHostGroup> FOnDemandHostGroup::Create(TConstArrayView<FAnsi
 			return FIoStatus(EIoErrorCode::InvalidParameter, Reason);
 		}
 
+#if !UE_BUILD_SHIPPING
+		if (Private::ShouldSkipUrl(UrlView))
+		{
+			UE_LOGFMT(LogIas, Log, "Skipping development CDN '{Url}' when creating HostGroup", UrlView);
+			continue;
+		}
+#endif // !UE_BUILD_SHIPPING
+
 		FAnsiString AnsiUrl(Url);
 		if (Config::bForceInsecureHttp)
 		{
 			AnsiUrl.ReplaceInline("https", "http");
 		}
+
 		Impl->HostUrls.Add(MoveTemp(AnsiUrl));
 	}
+
 	Impl->PrimaryIndex = 0; 
 
 	return FOnDemandHostGroup(MoveTemp(Impl)); 
@@ -190,6 +261,14 @@ TIoStatusOr<FOnDemandHostGroup> FOnDemandHostGroup::Create(TConstArrayView<FStri
 			return FIoStatus(EIoErrorCode::InvalidParameter, Reason);
 		}
 
+#if !UE_BUILD_SHIPPING
+		if (Private::ShouldSkipUrl(UrlView))
+		{
+			UE_LOGFMT(LogIas, Log, "Skipping development CDN '{Url}' when creating HostGroup", UrlView);
+			continue;
+		}
+#endif // !UE_BUILD_SHIPPING
+
 		if (Config::bForceInsecureHttp)
 		{
 			AnsiUrl.ReplaceInline("https", "http");
