Changed Iris Multithreaded Polling to use UE::Tasks::TTask and UE::Tasks::Wait which means that only the Polling tasks are executed and any queued GT tasks are not. This fixes a rare crash in GeometryCollectionComponent where the GTSim callbacks would get executed during the Iris Net Tick. Also removed ForegroundTasks flag as the UE::Tasks::Launch API does not support this. Reenabled Multithreaded Polling in the DefaultEngine.ini

#rnx
#tests Tested against Cooked Data and in PIE

diff --git a/Engine/Source/Runtime/Engine/Classes/Engine/NetDriver.h b/Engine/Source/Runtime/Engine/Classes/Engine/NetDriver.h
--- a/Engine/Source/Runtime/Engine/Classes/Engine/NetDriver.h
+++ b/Engine/Source/Runtime/Engine/Classes/Engine/NetDriver.h
@@ -683,10 +683,6 @@ struct FNetDriverReplicationSystemConfig
 	*/
 	UPROPERTY()
 	bool bAllowParallelTasks = false;	
-
-	/** When true, parallel tasks will run on Foreground, Background and Game threads. False means it will not run tasks on the ForegroundThread */
-	UPROPERTY()
-	bool bAllowTasksInForegroundThread = true;
 };
 
 //
diff --git a/Engine/Source/Runtime/Engine/Private/NetDriver.cpp b/Engine/Source/Runtime/Engine/Private/NetDriver.cpp
--- a/Engine/Source/Runtime/Engine/Private/NetDriver.cpp
+++ b/Engine/Source/Runtime/Engine/Private/NetDriver.cpp
@@ -298,7 +298,6 @@ namespace UE::Net::Private
 		}
 
 		OutParams.bAllowParallelTasks = ReplicationSystemConfig.bAllowParallelTasks;
-		OutParams.bAllowTasksInForegroundThread  = ReplicationSystemConfig.bAllowTasksInForegroundThread;
 	}
 #endif // UE_WITH_IRIS
 
diff --git a/Engine/Source/Runtime/Experimental/Iris/Core/Private/Iris/ReplicationSystem/Polling/ObjectPoller.cpp b/Engine/Source/Runtime/Experimental/Iris/Core/Private/Iris/ReplicationSystem/Polling/ObjectPoller.cpp
--- a/Engine/Source/Runtime/Experimental/Iris/Core/Private/Iris/ReplicationSystem/Polling/ObjectPoller.cpp
+++ b/Engine/Source/Runtime/Experimental/Iris/Core/Private/Iris/ReplicationSystem/Polling/ObjectPoller.cpp
@@ -22,6 +22,8 @@
 #include "Net/Core/Trace/NetDebugName.h"
 #include "HAL/IConsoleManager.h"
 
+#include "Tasks/Task.h"
+
 namespace UE::Net::Private
 {
 namespace CVars
@@ -76,7 +78,6 @@ namespace CVars
 		FObjectPoller* Poller = nullptr;
 		FReplicationPollTaskData& PollTaskData;
 		FReplicationSystemInternal* ReplicationSystemInternal = nullptr;
-		ENamedThreads::Type DesiredThreadType = ENamedThreads::Type::AnyNormalThreadNormalTask;
 
 	public:
 		
@@ -84,27 +85,15 @@ namespace CVars
 		FReplicationPollTask(
 			FObjectPoller* InPoller,
 			FReplicationPollTaskData& InPollTaskData,
-			FReplicationSystemInternal* InReplicationSystemInternal,
-			ENamedThreads::Type InDesiredThreadType
+			FReplicationSystemInternal* InReplicationSystemInternal
 		)
 			: Poller(InPoller)
 			, PollTaskData(InPollTaskData)
 			, ReplicationSystemInternal(InReplicationSystemInternal)
-			, DesiredThreadType(InDesiredThreadType)
-		{
-		}
-
-		ENamedThreads::Type GetDesiredThread() const
 		{
-			return DesiredThreadType;
 		}
 
-		ESubsequentsMode::Type GetSubsequentsMode() const
-		{
-			return ESubsequentsMode::FireAndForget;
-		}
-
-		void DoTask(const ENamedThreads::Type CurrentThread, const FGraphEventRef& MyCompletionGraphEvent) const
+		void DoTask() const
 		{
 			check(ReplicationSystemInternal->GetIsInParallelPhase());
 
@@ -155,11 +144,6 @@ namespace CVars
 
 			PollTaskData.NetTypeStats->ReleaseChildNetStatsContext(ChildNetStatsContext);
 		}
-
-		TStatId GetStatId() const
-		{
-			RETURN_QUICK_DECLARE_CYCLE_STAT(ReplicationPollTask, STATGROUP_ThreadPoolAsyncTasks);
-		}
 	};
 
 FObjectPoller::FObjectPoller(const FInitParams& InitParams)
@@ -231,7 +215,7 @@ void FObjectPoller::PollAndCopyObjects(const FNetBitArrayView& InObjectsConsider
 	{
 		ReplicationSystemInternal->StartParallelPhase();
 
-		FGraphEventArray TasksToComplete;
+		TArray<UE::Tasks::TTask<void>> TasksToComplete;
 		TArray<FReplicationPollTaskData> TaskData;
 		TaskData.Init(FReplicationPollTaskData(ObjectsConsideredForPolling), NumPollingTasks);
 
@@ -245,17 +229,21 @@ void FObjectPoller::PollAndCopyObjects(const FNetBitArrayView& InObjectsConsider
 			Data.CurrentTaskIndex = CurrentTaskIndex;
 			Data.BitsPerChunk = BitsPerChunk;
 			Data.ChunkStride = NumPollingTasks * BitsPerChunk;//How many bits should the task skip to get to the next Chunk it owns?
-			
+
 			{
 				IRIS_PROFILER_SCOPE(UObjectPoller_DispatchTask); 
-				ENamedThreads::Type ThreadType = ReplicationSystemInternal->ShouldAllowTasksInForegroundThread() ? ENamedThreads::AnyNormalThreadNormalTask : ENamedThreads::AnyBackgroundThreadNormalTask;
-				FGraphEventRef TaskRef = TGraphTask<FReplicationPollTask>::CreateTask().ConstructAndDispatchWhenReady(this, Data, ReplicationSystemInternal, ThreadType);
-				TasksToComplete.Add(TaskRef);
+				
+				TasksToComplete.Add(UE::Tasks::Launch(UE_SOURCE_LOCATION, [this, &Data]()
+					{
+						FReplicationPollTask NewTask = FReplicationPollTask(this, Data, ReplicationSystemInternal);
+						NewTask.DoTask();
+					}
+				));
 			}
-			
 		}
 	
-		FTaskGraphInterface::Get().WaitUntilTasksComplete(TasksToComplete);
+		//This can also run some of the task list on the current thread (GameThread), but only ones from TaskToComplete
+		UE::Tasks::Wait(TasksToComplete);
 
 		ReplicationSystemInternal->StopParallelPhase();
 	
diff --git a/Engine/Source/Runtime/Experimental/Iris/Core/Private/Iris/ReplicationSystem/ReplicationSystem.cpp b/Engine/Source/Runtime/Experimental/Iris/Core/Private/Iris/ReplicationSystem/ReplicationSystem.cpp
--- a/Engine/Source/Runtime/Experimental/Iris/Core/Private/Iris/ReplicationSystem/ReplicationSystem.cpp
+++ b/Engine/Source/Runtime/Experimental/Iris/Core/Private/Iris/ReplicationSystem/ReplicationSystem.cpp
@@ -93,7 +93,6 @@ public:
 			.MaxReplicationWriterObjectCount = Params.MaxReplicationWriterObjectCount,
 			.bUseRemoteObjectReferences = Params.bUseRemoteObjectReferences,
 			.bAllowParallelTasks = Params.bAllowParallelTasks,
-			.bAllowTasksInForegroundThread = Params.bAllowTasksInForegroundThread ,
 		}))
 	{
 		MetricNameTotalRootObjects = TEXT("TotalSubObjects");
diff --git a/Engine/Source/Runtime/Experimental/Iris/Core/Private/Iris/ReplicationSystem/ReplicationSystemInternal.h b/Engine/Source/Runtime/Experimental/Iris/Core/Private/Iris/ReplicationSystem/ReplicationSystemInternal.h
--- a/Engine/Source/Runtime/Experimental/Iris/Core/Private/Iris/ReplicationSystem/ReplicationSystemInternal.h
+++ b/Engine/Source/Runtime/Experimental/Iris/Core/Private/Iris/ReplicationSystem/ReplicationSystemInternal.h
@@ -39,7 +39,6 @@ struct FReplicationSystemInternalInitParams
 	uint32 MaxReplicationWriterObjectCount;
 	bool bUseRemoteObjectReferences;
 	bool bAllowParallelTasks;
-	bool bAllowTasksInForegroundThread;
 };
 
 class FReplicationSystemInternal
@@ -53,7 +52,6 @@ public:
 	, IrisObjectReferencePackageMap(nullptr)
 	, Id(Params.ReplicationSystemId)
 	, bAllowParallelTasks(Params.bAllowParallelTasks)
-	, bAllowTasksInForegroundThread(Params.bAllowTasksInForegroundThread)
 	{}
 
 	FReplicationProtocolManager& GetReplicationProtocolManager() { return ReplicationProtocolManager; }
@@ -140,8 +138,6 @@ public:
 
 	bool AreParallelTasksAllowed() const { return bAllowParallelTasks; }
 
-	bool ShouldAllowTasksInForegroundThread() const { return bAllowTasksInForegroundThread; }
-
 	/** Used by subsystems such as ObjectPoller to indicate when we're running tasks simultaneously and need to be thread-safe
 		Must be called in-order and exclusively. We currently do not support simultaneous parallel phases (e.g running Write tasks whilst Polling is running) */
 	void StartParallelPhase() 
@@ -204,9 +200,6 @@ public:
 
 	/** Is true whilst a phase is running parallel tasks. If bAllowParallelTasks is false, this can never be true. */
 	bool bIsInParallelPhase = false;
-
-	/** True means parallel tasks will run on Foreground, Background and Game threads. False means it will not run tasks on the ForegroundThread */
-	bool bAllowTasksInForegroundThread = true;
 };
 
 }
diff --git a/Engine/Source/Runtime/Experimental/Iris/Core/Public/Iris/ReplicationSystem/ReplicationSystem.h b/Engine/Source/Runtime/Experimental/Iris/Core/Public/Iris/ReplicationSystem/ReplicationSystem.h
--- a/Engine/Source/Runtime/Experimental/Iris/Core/Public/Iris/ReplicationSystem/ReplicationSystem.h
+++ b/Engine/Source/Runtime/Experimental/Iris/Core/Public/Iris/ReplicationSystem/ReplicationSystem.h
@@ -142,9 +142,6 @@ public:
 		*/
 		bool bAllowParallelTasks = false;
 
-		/** When true, parallel tasks will run on Foreground, Background and Game threads. False means it will not run tasks on the ForegroundThread*/
-		bool bAllowTasksInForegroundThread = true;
-
 		/** Delegate that receives every RPC executed locally. */
 		UE::Net::FForwardNetRPCCallDelegate ForwardNetRPCCallDelegate;
 
