Iris - By default graceful handling of reaching the end of bitstream at unexpected times is now disabled. The cvar net.Iris.ReplicationReader.GracefullyHandleReachingEndOfBitstream can still be enabled if so desired.

diff --git a/Engine/Source/Runtime/Experimental/Iris/Core/Private/Iris/ReplicationSystem/ReplicationReader.cpp b/Engine/Source/Runtime/Experimental/Iris/Core/Private/Iris/ReplicationSystem/ReplicationReader.cpp
--- a/Engine/Source/Runtime/Experimental/Iris/Core/Private/Iris/ReplicationSystem/ReplicationReader.cpp
+++ b/Engine/Source/Runtime/Experimental/Iris/Core/Private/Iris/ReplicationSystem/ReplicationReader.cpp
@@ -143,7 +143,7 @@ static FAutoConsoleVariableRef CVarMaxMustBeMappedHandleArray(
 	MaxMustBeMappedHandleArray,
 	TEXT("Maximum number of unresolved references that will be included with the QueuedBatch timeout error message."));
 
-static bool bGracefullyHandleReachingEndOfBitstream = true;
+static bool bGracefullyHandleReachingEndOfBitstream = false;
 static FAutoConsoleVariableRef CvarGracefullyHandleReachingEndOfBitstream(
 		TEXT("net.Iris.ReplicationReader.GracefullyHandleReachingEndOfBitstream"),
 		bGracefullyHandleReachingEndOfBitstream,
