Added CSV stat for PollAndCopyTask which tracks the work done, rather than the GameThread wall time. This allows us to work out the total amount of PollAndCopy work done, and how much unutilised core time we have waiting for unbalanced tasks.

#tests Tested against Cooked Dedicated Server
#rnx

diff --git a/Engine/Source/Runtime/Experimental/Iris/Core/Private/Iris/ReplicationSystem/Polling/ObjectPoller.cpp b/Engine/Source/Runtime/Experimental/Iris/Core/Private/Iris/ReplicationSystem/Polling/ObjectPoller.cpp
--- a/Engine/Source/Runtime/Experimental/Iris/Core/Private/Iris/ReplicationSystem/Polling/ObjectPoller.cpp
+++ b/Engine/Source/Runtime/Experimental/Iris/Core/Private/Iris/ReplicationSystem/Polling/ObjectPoller.cpp
@@ -15,6 +15,7 @@
 #include "Iris/ReplicationSystem/NetRefHandle.h"
 
 #include "Iris/Core/IrisProfiler.h"
+#include "Iris/Core/IrisCsv.h"
 #include "Iris/Stats/NetStatsContext.h"
 
 #include "Net/Core/DirtyNetObjectTracker/GlobalDirtyNetObjectTracker.h"
@@ -97,6 +98,8 @@ namespace CVars
 		{
 			check(ReplicationSystemInternal->GetIsInParallelPhase());
 
+			IRIS_CSV_PROFILER_SCOPE(Iris, ReplicationPollAndCopyTask);
+
 			IRIS_PROFILER_SCOPE_CONDITIONAL(PollAndCopyPushBased, PollTaskData.bUsePushModel);
 			IRIS_PROFILER_SCOPE_CONDITIONAL(ForcePollAndCopy, !PollTaskData.bUsePushModel);
 
