Chaos - fix collision event loop accessing disabled collision that may have dangling particle and shape pointers in them

#rnx

diff --git a/Engine/Source/Runtime/Experimental/Chaos/Private/EventDefaults.cpp b/Engine/Source/Runtime/Experimental/Chaos/Private/EventDefaults.cpp
--- a/Engine/Source/Runtime/Experimental/Chaos/Private/EventDefaults.cpp
+++ b/Engine/Source/Runtime/Experimental/Chaos/Private/EventDefaults.cpp
@@ -126,6 +126,11 @@ namespace Chaos
 								continue;
 							}
 
+							if (!ContactHandle->IsEnabled() || !ContactHandle->IsInConstraintGraph())
+							{
+								continue;
+							}
+
 							if (NumValidCollisions >= CollisionRule.NumConstraints())
 							{
 								break;
