Iris
- Make iris use ObjectNetSerializer::IsEqual instead of property::Identical when comparing object pointers.

#rnx

diff --git a/Engine/Source/Runtime/Experimental/Iris/Core/Private/Iris/Serialization/ObjectNetSerializer.cpp b/Engine/Source/Runtime/Experimental/Iris/Core/Private/Iris/Serialization/ObjectNetSerializer.cpp
--- a/Engine/Source/Runtime/Experimental/Iris/Core/Private/Iris/Serialization/ObjectNetSerializer.cpp
+++ b/Engine/Source/Runtime/Experimental/Iris/Core/Private/Iris/Serialization/ObjectNetSerializer.cpp
@@ -112,6 +112,7 @@ struct FObjectNetSerializerBase
 	typedef FQuantizedObjectReference QuantizedType;
 
 	static constexpr bool bHasDynamicState = true;
+	static constexpr bool bUseSerializerIsEqual = true;
 
 	static void Serialize(FNetSerializationContext&, const FNetSerializeArgs& Args);
 	static void Deserialize(FNetSerializationContext&, const FNetDeserializeArgs& Args);
@@ -379,8 +380,8 @@ bool FObjectNetSerializerBase<T>::IsEqual(FNetSerializationContext& Context, con
 	}
 	else
 	{
-		const SourceType Value0 = *reinterpret_cast<const SourceType*>(Args.Source0);
-		const SourceType Value1 = *reinterpret_cast<const SourceType*>(Args.Source1);
+		const SourceType& Value0 = *reinterpret_cast<const SourceType*>(Args.Source0);
+		const SourceType& Value1 = *reinterpret_cast<const SourceType*>(Args.Source1);
 
 		return Value0 == Value1;
 	}
