Fixed select in sphere code to use the right geometry index when accessing bounds of a collection
- this prevent crashes in certain situation where there's more transforms than there is geometry in a collection

diff --git a/Engine/Source/Runtime/Experimental/Chaos/Private/GeometryCollection/Facades/CollectionTransformSelectionFacade.cpp b/Engine/Source/Runtime/Experimental/Chaos/Private/GeometryCollection/Facades/CollectionTransformSelectionFacade.cpp
--- a/Engine/Source/Runtime/Experimental/Chaos/Private/GeometryCollection/Facades/CollectionTransformSelectionFacade.cpp
+++ b/Engine/Source/Runtime/Experimental/Chaos/Private/GeometryCollection/Facades/CollectionTransformSelectionFacade.cpp
@@ -874,18 +874,23 @@ namespace GeometryCollection::Facades
 
 			for (int32 TransformIdx = 0; TransformIdx < NumTransforms; ++TransformIdx)
 			{
-				const FBox& BoundingBoxInBoneSpace = BoundingBoxes[TransformIdx];
+				const TManagedArray<int32>& TransformToGeometryIndices = TransformToGeometryIndexAttribute.Get();
+				const int32 GeometryIndex = TransformToGeometryIndices[TransformIdx];
+				if (BoundingBoxes.IsValidIndex(GeometryIndex))
+				{
+					const FBox& BoundingBoxInBoneSpace = BoundingBoxes[GeometryIndex];
 
-				// Transform from BoneSpace to CollectionSpace
-				const FTransform CollectionSpaceTransform = TransformFacade.ComputeCollectionSpaceTransform(TransformIdx);
-				const FBox BoundingBoxInCollectionSpace = BoundingBoxInBoneSpace.TransformBy(CollectionSpaceTransform);
+					// Transform from BoneSpace to CollectionSpace
+					const FTransform CollectionSpaceTransform = TransformFacade.ComputeCollectionSpaceTransform(TransformIdx);
+					const FBox BoundingBoxInCollectionSpace = BoundingBoxInBoneSpace.TransformBy(CollectionSpaceTransform);
 
-				// Transform with specified transform
-				const FBox BoundingBoxInBoxSpace = BoundingBoxInCollectionSpace.InverseTransformBy(InBoxTransform);
+					// Transform with specified transform
+					const FBox BoundingBoxInBoxSpace = BoundingBoxInCollectionSpace.InverseTransformBy(InBoxTransform);
 
-				if (InBox.IsInside(BoundingBoxInBoxSpace))
-				{
-					OutSelection.Add(TransformIdx);
+					if (InBox.IsInside(BoundingBoxInBoxSpace))
+					{
+						OutSelection.Add(TransformIdx);
+					}
 				}
 			}
 		}
@@ -958,18 +963,23 @@ namespace GeometryCollection::Facades
 
 			for (int32 TransformIdx = 0; TransformIdx < NumTransforms; ++TransformIdx)
 			{
-				const FVector CentroidInBoneSpace(Centroids[TransformIdx]);
+				const TManagedArray<int32>& TransformToGeometryIndices = TransformToGeometryIndexAttribute.Get();
+				const int32 GeometryIndex = TransformToGeometryIndices[TransformIdx];
+				if (Centroids.IsValidIndex(GeometryIndex))
+				{
+					const FVector CentroidInBoneSpace(Centroids[GeometryIndex]);
 
-				// Transform from BoneSpace to CollectionSpace
-				const FTransform CollectionSpaceTransform = TransformFacade.ComputeCollectionSpaceTransform(TransformIdx);
-				const FVector CentroidInCollectionSpace = CollectionSpaceTransform.TransformPosition(CentroidInBoneSpace);
+					// Transform from BoneSpace to CollectionSpace
+					const FTransform CollectionSpaceTransform = TransformFacade.ComputeCollectionSpaceTransform(TransformIdx);
+					const FVector CentroidInCollectionSpace = CollectionSpaceTransform.TransformPosition(CentroidInBoneSpace);
 
-				// Transform with specified transform
-				const FVector CentroidInSphereSpace = InSphereTransform.InverseTransformPosition(CentroidInCollectionSpace);
+					// Transform with specified transform
+					const FVector CentroidInSphereSpace = InSphereTransform.InverseTransformPosition(CentroidInCollectionSpace);
 
-				if (InSphere.IsInside(CentroidInSphereSpace))
-				{
-					OutSelection.Add(TransformIdx);
+					if (InSphere.IsInside(CentroidInSphereSpace))
+					{
+						OutSelection.Add(TransformIdx);
+					}
 				}
 			}
 		}
@@ -992,31 +1002,36 @@ namespace GeometryCollection::Facades
 
 			for (int32 TransformIdx = 0; TransformIdx < NumTransforms; ++TransformIdx)
 			{
-				const FBox& BoundingBoxInBoneSpace = BoundingBoxes[TransformIdx];
-
-				int32 NumVerticesInside = 0;
-
-				const TArray<FVector>& VertexInBoneSpaceArr = FBoundsFacade::GetBoundingBoxVertexPositions(BoundingBoxInBoneSpace);
-				for (int32 VertexIdx = 0; VertexIdx < VertexInBoneSpaceArr.Num(); ++VertexIdx)
+				const TManagedArray<int32>& TransformToGeometryIndices = TransformToGeometryIndexAttribute.Get();
+				const int32 GeometryIndex = TransformToGeometryIndices[TransformIdx];
+				if (BoundingBoxes.IsValidIndex(GeometryIndex))
 				{
-					// Transform from BoneSpace to CollectionSpace
-					const FTransform CollectionSpaceTransform = TransformFacade.ComputeCollectionSpaceTransform(TransformIdx);
-					const FVector VertexInCollectionSpace = CollectionSpaceTransform.TransformPosition(VertexInBoneSpaceArr[VertexIdx]);
+					const FBox& BoundingBoxInBoneSpace = BoundingBoxes[GeometryIndex];
 
-					// Transform with specified transform
-					const FVector VertexInSphereSpace = InSphereTransform.InverseTransformPosition(VertexInCollectionSpace);
+					int32 NumVerticesInside = 0;
 
-					if (!InSphere.IsInside(VertexInSphereSpace))
+					const TArray<FVector>& VertexInBoneSpaceArr = FBoundsFacade::GetBoundingBoxVertexPositions(BoundingBoxInBoneSpace);
+					for (int32 VertexIdx = 0; VertexIdx < VertexInBoneSpaceArr.Num(); ++VertexIdx)
 					{
-						break;
-					}
+						// Transform from BoneSpace to CollectionSpace
+						const FTransform CollectionSpaceTransform = TransformFacade.ComputeCollectionSpaceTransform(TransformIdx);
+						const FVector VertexInCollectionSpace = CollectionSpaceTransform.TransformPosition(VertexInBoneSpaceArr[VertexIdx]);
 
-					NumVerticesInside++;
-				}
+						// Transform with specified transform
+						const FVector VertexInSphereSpace = InSphereTransform.InverseTransformPosition(VertexInCollectionSpace);
 
-				if (NumVerticesInside == VertexInBoneSpaceArr.Num())
-				{
-					OutSelection.Add(TransformIdx);
+						if (!InSphere.IsInside(VertexInSphereSpace))
+						{
+							break;
+						}
+
+						NumVerticesInside++;
+					}
+
+					if (NumVerticesInside == VertexInBoneSpaceArr.Num())
+					{
+						OutSelection.Add(TransformIdx);
+					}
 				}
 			}
 		}
