Dataflow : Fix crash when using FVector or FTransform variable nodes or other struct based basic nodes
- Issue comes from thefact that those types are a base struct but do not have a StaticStruct member

diff --git a/Engine/Source/Runtime/Experimental/Dataflow/Core/Public/Dataflow/DataflowContextCache.h b/Engine/Source/Runtime/Experimental/Dataflow/Core/Public/Dataflow/DataflowContextCache.h
--- a/Engine/Source/Runtime/Experimental/Dataflow/Core/Public/Dataflow/DataflowContextCache.h
+++ b/Engine/Source/Runtime/Experimental/Dataflow/Core/Public/Dataflow/DataflowContextCache.h
@@ -26,10 +26,10 @@ namespace UE::Dataflow
 	};
 
 	/** Trait used to select the UStruct cache element path code. */
-	template <typename T, typename = void, typename = void>
+	template <typename T, typename = void>
 	struct TIsReflectedStruct { static constexpr bool Value = false; };
 	template <typename T>
-	struct TIsReflectedStruct<T, std::void_t<decltype(&TBaseStructure<typename TDecay<T>::Type>::Get)>, std::void_t<decltype(&T::StaticStruct)>> { static constexpr bool Value = true; };
+	struct TIsReflectedStruct<T, std::void_t<decltype(&TBaseStructure<typename TDecay<T>::Type>::Get)>> { static constexpr bool Value = true; };
 
 	struct FTimestamp
 	{
