Dataflow : fix issues with drag and dropping duplicate names on input and output subgraph nodes
- also add better support fro arrays from anytypes pins
- add a read only subgraph name to the subgraph call nodes

diff --git a/Engine/Plugins/Experimental/Dataflow/Source/DataflowEditor/Private/Dataflow/DataflowAssetEditUtils.cpp b/Engine/Plugins/Experimental/Dataflow/Source/DataflowEditor/Private/Dataflow/DataflowAssetEditUtils.cpp
--- a/Engine/Plugins/Experimental/Dataflow/Source/DataflowEditor/Private/Dataflow/DataflowAssetEditUtils.cpp
+++ b/Engine/Plugins/Experimental/Dataflow/Source/DataflowEditor/Private/Dataflow/DataflowAssetEditUtils.cpp
@@ -285,6 +285,7 @@ namespace UE::Dataflow
 								const FName UniqueName = GenerateUniqueObjectName(&DataflowAsset, NewSubGraphName);
 								DataflowEdNode->Rename(*UniqueName.ToString(), &EdGraph);
 								CallNode->SetName(UniqueName);
+								CallNode->RefreshSubGraphName();
 							}
 						}
 					}
diff --git a/Engine/Source/Runtime/Experimental/Dataflow/Engine/Private/Dataflow/DataflowSubGraphNodes.cpp b/Engine/Source/Runtime/Experimental/Dataflow/Engine/Private/Dataflow/DataflowSubGraphNodes.cpp
--- a/Engine/Source/Runtime/Experimental/Dataflow/Engine/Private/Dataflow/DataflowSubGraphNodes.cpp
+++ b/Engine/Source/Runtime/Experimental/Dataflow/Engine/Private/Dataflow/DataflowSubGraphNodes.cpp
@@ -6,12 +6,24 @@
 #include "Dataflow/DataflowEdNode.h"
 #include "Dataflow/DataflowObject.h"
 #include "Dataflow/DataflowSubGraph.h"
+#include "Dataflow/DataflowTypePolicy.h"
 #include "Dataflow/DataflowNodeFactory.h"
 #include "GraphEditAction.h"
 #include "UObject/UObjectGlobals.h"
 
 namespace UE::Dataflow::SubGraphNodes::Private
 {
+	FName MakeUniqueNameForPropertyBag(FName OriginalName, const FInstancedPropertyBag& PropertyBag)
+	{
+		FName NewName = OriginalName;
+		int32 SuffixNumber = 1;
+		while (PropertyBag.FindPropertyDescByName(NewName))
+		{
+			NewName.SetNumber(SuffixNumber++);
+		}
+		return NewName;
+	}
+
 	class FContextSubGraph: public FContext
 	{
 	public:
@@ -145,76 +157,87 @@ namespace UE::Dataflow::SubGraphNodes::Private
 		const UE::Dataflow::ISubGraphContextCallback& SubGraphCallback;
 	};
 
-	FPropertyBagPropertyDesc GetPropertyBagPropertyDescFromDataflowConnection(const FDataflowConnection& Connection)
+	FPropertyBagPropertyDesc GetPropertyBagPropertyDescFromDataflowType(const FName Name, const FName Type)
 	{
-		const FProperty* ConnectionProperty = Connection.GetProperty();
-		check(ConnectionProperty);
-
-		// AnyType require special treatment to make sure we give a concrete type
-		if (Connection.IsAnyType())
+		if (Type == UE::Dataflow::GetTypeName<bool>())
 		{
-			EPropertyBagPropertyType PropertyType = EPropertyBagPropertyType::Count;
-
-			if (Connection.IsType<bool>())
-			{
-				PropertyType = EPropertyBagPropertyType::Bool;
-			}
-			else if (Connection.IsType<int32>())
-			{
-				PropertyType = EPropertyBagPropertyType::Int32;
-			}
-			else if (Connection.IsType<int64>())
-			{
-				PropertyType = EPropertyBagPropertyType::Int64;
-			}
-			else if (Connection.IsType<uint32>())
-			{
-				PropertyType = EPropertyBagPropertyType::UInt32;
-			}
-			else if (Connection.IsType<uint64>())
-			{
-				PropertyType = EPropertyBagPropertyType::UInt64;
-			}
-			else if (Connection.IsType<float>())
-			{
-				PropertyType = EPropertyBagPropertyType::Float;
-			}
-			else if (Connection.IsType<double>())
-			{
-				PropertyType = EPropertyBagPropertyType::Double;
-			}
-			else if (Connection.IsType<FName>())
-			{
-				PropertyType = EPropertyBagPropertyType::Name;
-			}
-			else if (Connection.IsType<FString>())
-			{
-				PropertyType = EPropertyBagPropertyType::String;
-			}
-			else if (Connection.IsType<FText>())
+			return FPropertyBagPropertyDesc(Name, EPropertyBagPropertyType::Bool);
+		}
+		else if (Type == UE::Dataflow::GetTypeName<int32>())
+		{
+			return FPropertyBagPropertyDesc(Name, EPropertyBagPropertyType::Int32);
+		}
+		else if (Type == UE::Dataflow::GetTypeName<int64>())
+		{
+			return FPropertyBagPropertyDesc(Name, EPropertyBagPropertyType::Int64);
+		}
+		else if (Type == UE::Dataflow::GetTypeName<uint32>())
+		{
+			return FPropertyBagPropertyDesc(Name, EPropertyBagPropertyType::UInt32);
+		}
+		else if (Type == UE::Dataflow::GetTypeName<uint64>())
+		{
+			return FPropertyBagPropertyDesc(Name, EPropertyBagPropertyType::UInt64);
+		}
+		else if (Type == UE::Dataflow::GetTypeName<float>())
+		{
+			return FPropertyBagPropertyDesc(Name, EPropertyBagPropertyType::Float);
+		}
+		else if (Type == UE::Dataflow::GetTypeName<double>())
+		{
+			return FPropertyBagPropertyDesc(Name, EPropertyBagPropertyType::Double);
+		}
+		else if (Type == UE::Dataflow::GetTypeName<FName>())
+		{
+			return FPropertyBagPropertyDesc(Name, EPropertyBagPropertyType::Name);
+		}
+		else if (Type == UE::Dataflow::GetTypeName<FString>())
+		{
+			return FPropertyBagPropertyDesc(Name, EPropertyBagPropertyType::String);
+		}
+		else if (Type == UE::Dataflow::GetTypeName<FText>())
+		{
+			return FPropertyBagPropertyDesc(Name, EPropertyBagPropertyType::Text);
+		}
+		else
+		{
+			const FString TypeStr(Type.ToString());
+			if (TypeStr.StartsWith("U"))
 			{
-				PropertyType = EPropertyBagPropertyType::Text;
+				// Nothing for now and use thefallback - may need fixing if we find there's an issue with that code path 
 			}
-			else
+			else if (TypeStr.StartsWith("F"))
 			{
-				const FString TypeStr(Connection.GetType().ToString());
-				if (TypeStr.StartsWith("U"))
+				const FString ShortTypeName = TypeStr.RightChop(1);
+				if (const UScriptStruct* ScriptStruct = FindFirstObjectSafe<UScriptStruct>(*ShortTypeName))
 				{
-					// Nothing for now and use thefallback - may need fixing if we find there's an issue with that code path 
-				}
-				else if (TypeStr.StartsWith("F"))
-				{
-					const FString ShortTypeName = TypeStr.RightChop(1);
-					if (const UScriptStruct* ScriptStruct = FindFirstObjectSafe<UScriptStruct>(*ShortTypeName))
-					{
-						return FPropertyBagPropertyDesc(Connection.GetName(), EPropertyBagPropertyType::Struct, ScriptStruct);
-					}
+					return FPropertyBagPropertyDesc(Name, EPropertyBagPropertyType::Struct, ScriptStruct);
 				}
 			}
+			else if (FDataflowArrayTypePolicy::SupportsTypeStatic(Type))
+			{
+				const FName ElementType = FDataflowArrayTypePolicy::GetElementType(Type);
+				FPropertyBagPropertyDesc PropertyDesc = GetPropertyBagPropertyDescFromDataflowType(Name, ElementType);
+				PropertyDesc.ContainerTypes = FPropertyBagContainerTypes{ EPropertyBagContainerType::Array };
+				return PropertyDesc;
+			}
+		}
+		// invalid value 
+		return FPropertyBagPropertyDesc(Name, EPropertyBagPropertyType::Count);
+	}
+
+	FPropertyBagPropertyDesc GetPropertyBagPropertyDescFromDataflowConnection(const FDataflowConnection& Connection)
+	{
+		const FProperty* ConnectionProperty = Connection.GetProperty();
+		check(ConnectionProperty);
 
-			if (PropertyType != EPropertyBagPropertyType::Count)
+		// AnyType require special treatment to make sure we give a concrete type
+		if (Connection.IsAnyType())
+		{
+			const FPropertyBagPropertyDesc Desc = GetPropertyBagPropertyDescFromDataflowType(Connection.GetName(), Connection.GetType());
+			if (Desc.ValueType != EPropertyBagPropertyType::Count)
 			{
-				return FPropertyBagPropertyDesc(Connection.GetName(), PropertyType);
+				return Desc;
 			}
 			else
 			{
@@ -241,7 +264,6 @@ namespace UE::Dataflow::SubGraphNodes::Private
 		// default simply set property desc from the property
 		return FPropertyBagPropertyDesc(Connection.GetName(), Connection.GetProperty());
 	}
-
 }
 
 //////////////////////////////////////////////////////////////////////////////////////////////////////////
@@ -303,7 +325,8 @@ const FDataflowConnection* FDataflowSubGraphInputNode::OnDropConnectionOnNode(co
 	using namespace  UE::Dataflow::SubGraphNodes;
 	if (SupportsDropConnectionOnNode(DroppedConnection.GetType(), DroppedConnection.GetDirection()))
 	{
-		const FPropertyBagPropertyDesc Desc = Private::GetPropertyBagPropertyDescFromDataflowConnection(DroppedConnection);
+		FPropertyBagPropertyDesc Desc = Private::GetPropertyBagPropertyDescFromDataflowConnection(DroppedConnection);
+		Desc.Name = Private::MakeUniqueNameForPropertyBag(Desc.Name, PropertyBag);
 		PropertyBag.AddProperties({ Desc });
 		DynamicConnections.Refresh();
 		return FindOutput(Desc.Name);
@@ -318,7 +341,8 @@ void FDataflowSubGraphInputNode::AddConnectionsTo(TArrayView<FDataflowConnection
 	{
 		if (Connection && SupportsDropConnectionOnNode(Connection->GetType(), Connection->GetDirection()))
 		{
-			const FPropertyBagPropertyDesc Desc = Private::GetPropertyBagPropertyDescFromDataflowConnection(*Connection);
+			FPropertyBagPropertyDesc Desc = Private::GetPropertyBagPropertyDescFromDataflowConnection(*Connection);
+			Desc.Name = Private::MakeUniqueNameForPropertyBag(Desc.Name, PropertyBag);
 			PropertyBag.AddProperties({ Desc });
 		}
 	}
@@ -364,7 +388,8 @@ const FDataflowConnection* FDataflowSubGraphOutputNode::OnDropConnectionOnNode(c
 	using namespace  UE::Dataflow::SubGraphNodes;
 	if (SupportsDropConnectionOnNode(DroppedConnection.GetType(), DroppedConnection.GetDirection()))
 	{
-		const FPropertyBagPropertyDesc Desc = Private::GetPropertyBagPropertyDescFromDataflowConnection(DroppedConnection);
+		FPropertyBagPropertyDesc Desc = Private::GetPropertyBagPropertyDescFromDataflowConnection(DroppedConnection);
+		Desc.Name = Private::MakeUniqueNameForPropertyBag(Desc.Name, PropertyBag);
 		PropertyBag.AddProperties({ Desc });
 		DynamicConnections.Refresh();
 		return FindInput(Desc.Name);
@@ -379,7 +404,8 @@ void FDataflowSubGraphOutputNode::AddConnectionsTo(TArrayView<FDataflowConnectio
 	{
 		if (Connection && SupportsDropConnectionOnNode(Connection->GetType(), Connection->GetDirection()))
 		{
-			const FPropertyBagPropertyDesc Desc = Private::GetPropertyBagPropertyDescFromDataflowConnection(*Connection);
+			FPropertyBagPropertyDesc Desc = Private::GetPropertyBagPropertyDescFromDataflowConnection(*Connection);
+			Desc.Name = Private::MakeUniqueNameForPropertyBag(Desc.Name, PropertyBag);
 			PropertyBag.AddProperties({ Desc });
 		}
 	}
@@ -466,13 +492,27 @@ void FDataflowCallSubGraphNode::SetSubGraphGuid(const FGuid& InSubGraphGuid)
 
 		SubGraphGuid = InSubGraphGuid;
 
-		RegisterHandlers(DataflowAsset->FindSubGraphByGuid(SubGraphGuid));
+		UDataflowSubGraph* NewSubGraph = DataflowAsset->FindSubGraphByGuid(SubGraphGuid);
+		RegisterHandlers(NewSubGraph);
+		SubGraphName = NewSubGraph? NewSubGraph->GetFName(): NAME_None;
 
 		SyncInputsPropertyBagWithSubGraph();
 		SyncOutputsPropertyBagWithSubGraph();
 	}
 }
 
+void FDataflowCallSubGraphNode::RefreshSubGraphName()
+{
+	SubGraphName = NAME_None;
+	if (TStrongObjectPtr<UDataflow> DataflowAsset = DataflowAssetWeakPtr.Pin())
+	{
+		if (UDataflowSubGraph* NewSubGraph = DataflowAsset->FindSubGraphByGuid(SubGraphGuid))
+		{
+			SubGraphName = NewSubGraph->GetFName();
+		}
+	}
+}
+
 bool FDataflowCallSubGraphNode::IsValid() const
 {
 	return SubGraphGuid.IsValid();
diff --git a/Engine/Source/Runtime/Experimental/Dataflow/Engine/Public/Dataflow/DataflowSubGraphNodes.h b/Engine/Source/Runtime/Experimental/Dataflow/Engine/Public/Dataflow/DataflowSubGraphNodes.h
--- a/Engine/Source/Runtime/Experimental/Dataflow/Engine/Public/Dataflow/DataflowSubGraphNodes.h
+++ b/Engine/Source/Runtime/Experimental/Dataflow/Engine/Public/Dataflow/DataflowSubGraphNodes.h
@@ -155,6 +155,8 @@ public:
 	DATAFLOWENGINE_API void SetSubGraphGuid(const FGuid& Guid);
 	const FGuid& GetSubGraphGuid() const { return SubGraphGuid; }
 
+	DATAFLOWENGINE_API void RefreshSubGraphName();
+
 private:
 	virtual void Evaluate(UE::Dataflow::FContext& Context, const FDataflowOutput* Out) const override;
 	virtual void PostSerialize(const FArchive& Ar) override;
@@ -178,13 +180,14 @@ private:
 	void SyncInputsPropertyBagWithSubGraph();
 	void SyncOutputsPropertyBagWithSubGraph();
 
-	const FDataflowSubGraphOutputNode* GetSubGraphOutputNode() const;
-
 	bool IsValid() const;
 
 	UPROPERTY()
 	FGuid SubGraphGuid;
 
+	UPROPERTY(VisibleAnyWhere, Category = Subgraph)
+	FName SubGraphName;
+
 	UPROPERTY()
 	FDataflowDynamicConnections DynamicInputs;
 
