Add Indexed cache storage support for On Demand IOStore
Still disabled, but can be enabled with
[IndexedCacheStorage]
+Storage=(Name="IAD",Index=1,Mount="iad")

diff --git a/Engine/Source/Runtime/Experimental/IoStore/OnDemand/IoStoreOnDemand.Build.cs b/Engine/Source/Runtime/Experimental/IoStore/OnDemand/IoStoreOnDemand.Build.cs
--- a/Engine/Source/Runtime/Experimental/IoStore/OnDemand/IoStoreOnDemand.Build.cs
+++ b/Engine/Source/Runtime/Experimental/IoStore/OnDemand/IoStoreOnDemand.Build.cs
@@ -13,6 +13,7 @@ public class IoStoreOnDemand : ModuleRules
 			new string[] {
 				"Analytics",
 				"CoreUObject",
+				"IndexedCacheStorage",
 				"IoStoreHttpClient",
 				"Json",
 			}
diff --git a/Engine/Source/Runtime/Experimental/IoStore/OnDemand/Private/OnDemandIoStore.cpp b/Engine/Source/Runtime/Experimental/IoStore/OnDemand/Private/OnDemandIoStore.cpp
--- a/Engine/Source/Runtime/Experimental/IoStore/OnDemand/Private/OnDemandIoStore.cpp
+++ b/Engine/Source/Runtime/Experimental/IoStore/OnDemand/Private/OnDemandIoStore.cpp
@@ -17,6 +17,7 @@
 #include "HAL/PlatformFileManager.h"
 #include "HAL/PlatformProcess.h"
 #include "HAL/PlatformTime.h"
+#include "IndexedCacheStorageManager.h"
 #include "IO/IoChunkEncoding.h"
 #include "IO/IoContainerHeader.h"
 #include "IO/IoStatus.h"
@@ -31,6 +32,7 @@
 #include "Misc/Guid.h"
 #include "Misc/PathViews.h"
 #include "Misc/Paths.h"
+#include "Modules/ModuleManager.h"
 #include "OnDemandConfig.h"
 #include "OnDemandHttpClient.h"
 #include "OnDemandHttpThread.h"
@@ -858,7 +860,47 @@ FIoStatus FOnDemandIoStore::InitializeInstallCache()
 		return CacheConfig.Status();
 	}
 
-	InstallCache = MakeOnDemandInstallCache(*this, CacheConfig.ConsumeValueOrDie());
+	FOnDemandInstallCacheConfig OnDemandInstallCacheConfig = CacheConfig.ConsumeValueOrDie();
+
+	// Check if we use ICS for IAD and update the RootDirectory accordingly
+	int32 SelectedCacheIndex = 0;
+	{
+		// OnDemandInstallCacheConfig.DiskQuota is the quota to store the data, but we also need some room for metadata, see GetJournalFilename/GetSnapshotFilename
+		uint64 OnDemandInstallRequestedSize = OnDemandInstallCacheConfig.DiskQuota + 2 * OnDemandInstallCacheConfig.JournalMaxSize + 2 * 32;
+		if (OnDemandInstallRequestedSize > 0)
+		{
+			FModuleManager::Get().LoadModule(TEXT("IndexedCacheStorage"));
+
+			const FString IADStorageName(TEXT("IAD"));
+			SelectedCacheIndex = Experimental::FIndexedCacheStorageManager::Get().GetStorageIndex(IADStorageName);
+
+			if (SelectedCacheIndex > 0)
+			{
+				// When IAD is assigned to an indexed cache, we *have* to get enough disk space.
+				while (true)
+				{
+					if (!Experimental::FIndexedCacheStorageManager::Get().CreateCacheStorage(OnDemandInstallRequestedSize, SelectedCacheIndex))
+					{
+						continue;
+					}
+
+					FString MountPath = Experimental::FIndexedCacheStorageManager::Get().MountCacheStorage(SelectedCacheIndex);
+					if (MountPath.IsEmpty())
+					{
+						UE_LOG(LogIoStoreOnDemand, Error, TEXT("Failed to mount indexed cache storage %d (%s) even though it was created"), SelectedCacheIndex, *IADStorageName);
+						Experimental::FIndexedCacheStorageManager::Get().DestroyCacheStorage(SelectedCacheIndex);
+						continue;
+					}
+
+					FString RelativeDir(MountPath / OnDemandInstallCacheConfig.RootDirectory.Replace(*FPaths::ProjectPersistentDownloadDir(), TEXT("")));
+					OnDemandInstallCacheConfig.RootDirectory = RelativeDir;
+					break;
+				}
+			}
+		}
+	}
+
+	InstallCache = MakeOnDemandInstallCache(*this, OnDemandInstallCacheConfig);
 	if (InstallCache.IsValid())
 	{
 		int32 BackendPriority = -5; // Lower than file (zero) but higher than streaming backend (-10)
@@ -876,6 +918,12 @@ FIoStatus FOnDemandIoStore::InitializeInstallCache()
 	else
 	{
 		UE_LOG(LogIoStoreOnDemand, Error, TEXT("Failed to initialize install cache"));
+
+		if (SelectedCacheIndex > 0)
+		{
+			Experimental::FIndexedCacheStorageManager::Get().DestroyCacheStorage(SelectedCacheIndex);
+		}
+
 		return FIoStatusBuilder(EIoErrorCode::InvalidParameter) << TEXT("Failed to initialize install cache");
 	}
 
