Streamable manager - Fix ensures from trying to cleanup load requests which never got started.
IAD - Mark requests as completed before calling back.  This prevents noop tasks from being enqueued from the callback.

diff --git a/Engine/Source/Runtime/Engine/Classes/Engine/StreamableManager.h b/Engine/Source/Runtime/Engine/Classes/Engine/StreamableManager.h
--- a/Engine/Source/Runtime/Engine/Classes/Engine/StreamableManager.h
+++ b/Engine/Source/Runtime/Engine/Classes/Engine/StreamableManager.h
@@ -632,6 +632,9 @@ private:
 	ECookLoadType CookLoadType;
 #endif
 
+	/** True if this request ever started loading from disk. bLoadCompleted may be set without a load starting if cancels or errors occur during download. */
+	bool bLoadStarted;
+
 	/** True if this request has finished loading. It may still be active, or it may have been released */
 	bool bLoadCompleted;
 
diff --git a/Engine/Source/Runtime/Engine/Private/StreamableManager.cpp b/Engine/Source/Runtime/Engine/Private/StreamableManager.cpp
--- a/Engine/Source/Runtime/Engine/Private/StreamableManager.cpp
+++ b/Engine/Source/Runtime/Engine/Private/StreamableManager.cpp
@@ -519,6 +519,7 @@ FStreamableHandle::FStreamableHandle(FPrivateToken, UE::FSourceLocation&& InLoca
 #if WITH_EDITOR
 	, CookLoadType(ECookLoadType::Unspecified)
 #endif
+	, bLoadStarted(false)
 	, bLoadCompleted(false)
 	, bReleased(false)
 	, bCanceled(false)
@@ -903,7 +904,7 @@ void FStreamableHandle::CancelHandle()
 
 	// Remove from referenced list. If it is stalled then it won't have been registered with
 	// the manager yet
-	if (!bStalled && !DownloadContext)
+	if (bLoadStarted)
 	{
 		for (const FSoftObjectPath& AssetRef : RequestedAssets)
 		{
@@ -966,10 +967,13 @@ void FStreamableHandle::ReleaseHandle()
 
 		TSharedRef<FStreamableHandle> SharedThis = AsShared();
 
-		// Remove from referenced list
-		for (const FSoftObjectPath& AssetRef : RequestedAssets)
+		if (bLoadStarted)
 		{
-			OwningManager->RemoveReferencedAsset(AssetRef, SharedThis);
+			// Remove from referenced list
+			for (const FSoftObjectPath& AssetRef : RequestedAssets)
+			{
+				OwningManager->RemoveReferencedAsset(AssetRef, SharedThis);
+			}
 		}
 
 		// Remove from explicit list
@@ -1935,6 +1939,8 @@ void FStreamableManager::StartHandleRequests(TSharedRef<FStreamableHandle> Handl
 		return;
 	}
 
+	Handle->bLoadStarted = true;
+
 	TRACE_LOADTIME_REQUEST_GROUP_SCOPE(TEXT("StreamableManager - %s"), *Handle->GetDebugName());
 
 	TArray<FStreamable *> ExistingStreamables;
diff --git a/Engine/Source/Runtime/Experimental/IoStore/OnDemand/Private/OnDemandContentInstaller.cpp b/Engine/Source/Runtime/Experimental/IoStore/OnDemand/Private/OnDemandContentInstaller.cpp
--- a/Engine/Source/Runtime/Experimental/IoStore/OnDemand/Private/OnDemandContentInstaller.cpp
+++ b/Engine/Source/Runtime/Experimental/IoStore/OnDemand/Private/OnDemandContentInstaller.cpp
@@ -990,23 +990,23 @@ void FOnDemandContentInstaller::CompleteInstallRequest(FRequest& Request)
 
 	TryExecuteNextRequest();
 
-	FOnDemandInstallRequest::EStatus RequestStatus = FOnDemandInstallRequest::EStatus::None;
+	FOnDemandInstallRequest::EStatus FinalRequestStatus = FOnDemandInstallRequest::EStatus::None;
 	switch (InstallResult.Status.GetErrorCode())
 	{
 		case EIoErrorCode::Ok:
-			RequestStatus = FOnDemandInstallRequest::Ok;
+			FinalRequestStatus = FOnDemandInstallRequest::Ok;
 			break;
 		case EIoErrorCode::Cancelled:
-			RequestStatus = FOnDemandInstallRequest::Cancelled;
+			FinalRequestStatus = FOnDemandInstallRequest::Cancelled;
 			break;
 		default:
-			RequestStatus = FOnDemandInstallRequest::Error;
+			FinalRequestStatus = FOnDemandInstallRequest::Error;
 			break;
 	}
 
 	if (!InstallRequest.OnCompleted)
 	{
-		InstallRequest.Request->Status.store(RequestStatus);
+		InstallRequest.Request->Status.store(FinalRequestStatus);
 		{
 			UE::TUniqueLock Lock(Mutex);
 			RequestAllocator.Destroy(&Request);
@@ -1014,19 +1014,22 @@ void FOnDemandContentInstaller::CompleteInstallRequest(FRequest& Request)
 		return;
 	}
 
+	// Do this before the callback in case the callback triggers any further API calls
+	InstallRequest.Request->Status.store(FOnDemandInstallRequest::PendingCallbacks);
+
 	const bool bCallbackOnGameThread = EnumHasAnyFlags(InstallRequest.Args.Options, EOnDemandInstallOptions::CallbackOnGameThread);
 	if (bCallbackOnGameThread)
 	{
 		ExecuteOnGameThread(
 			UE_SOURCE_LOCATION,
-			[this, &Request, InstallResult = MoveTemp(InstallResult), RequestStatus]() mutable
+			[this, &Request, InstallResult = MoveTemp(InstallResult), FinalRequestStatus]() mutable
 			{
 				FRequest::FInstall& InstallRequest		= Request.AsInstall();
 				FOnDemandInstallCompleted OnCompleted	= MoveTemp(InstallRequest.OnCompleted);
 
 				ensure(InstallRequest.bNotifyingProgressOnGameThread == false);
 				OnCompleted(MoveTemp(InstallResult));
-				InstallRequest.Request->Status.store(RequestStatus);
+				InstallRequest.Request->Status.store(FinalRequestStatus);
 				{
 					UE::TUniqueLock Lock(Mutex);
 					RequestAllocator.Destroy(&Request);
@@ -1037,7 +1040,7 @@ void FOnDemandContentInstaller::CompleteInstallRequest(FRequest& Request)
 	{
 		FOnDemandInstallCompleted OnCompleted = MoveTemp(InstallRequest.OnCompleted);
 		OnCompleted(MoveTemp(InstallResult));
-		InstallRequest.Request->Status.store(RequestStatus);
+		InstallRequest.Request->Status.store(FinalRequestStatus);
 		{
 			UE::TUniqueLock Lock(Mutex);
 			RequestAllocator.Destroy(&Request);
diff --git a/Engine/Source/Runtime/Experimental/IoStore/OnDemandCore/Public/IO/IoStoreOnDemand.h b/Engine/Source/Runtime/Experimental/IoStore/OnDemandCore/Public/IO/IoStoreOnDemand.h
--- a/Engine/Source/Runtime/Experimental/IoStore/OnDemandCore/Public/IO/IoStoreOnDemand.h
+++ b/Engine/Source/Runtime/Experimental/IoStore/OnDemandCore/Public/IO/IoStoreOnDemand.h
@@ -47,6 +47,8 @@ public:
 		None,
 		/** The request is pending. */
 		Pending,
+		/** The request is issuing callbacks. */
+		PendingCallbacks,
 		/** The request completed successfully. */
 		Ok,
 		/** The request was cancelled. */
@@ -62,6 +64,8 @@ public:
 	bool IsNone() const { return EStatus::None == GetStatus(); }
 	/** Returns true if the request is pending. */
 	bool IsPending() const { return EStatus::Pending == GetStatus(); }
+	/** Returns true if the request is issuing callbacks. */
+	bool IsPendingCallbacks() const { return EStatus::PendingCallbacks == GetStatus(); }
 	/** Returns true if the request was successful. */
 	bool IsOk() const { return EStatus::Ok == GetStatus(); }
 	/** Returns true if the request was cancelled. */
