Iris
* Reduced spam from MustBeBlocked if multiple objects are blocked by the same asset

#rnx

diff --git a/Engine/Source/Runtime/Experimental/Iris/Core/Private/Iris/ReplicationSystem/ReplicationReader.cpp b/Engine/Source/Runtime/Experimental/Iris/Core/Private/Iris/ReplicationSystem/ReplicationReader.cpp
--- a/Engine/Source/Runtime/Experimental/Iris/Core/Private/Iris/ReplicationSystem/ReplicationReader.cpp
+++ b/Engine/Source/Runtime/Experimental/Iris/Core/Private/Iris/ReplicationSystem/ReplicationReader.cpp
@@ -2812,7 +2812,7 @@ void FReplicationReader::ProcessQueuedBatches()
 			++PendingBatchData.PendingBatchTryProcessCount;
 
 			// Check how long we've been blocked only every 30 ticks
-			if ((QueuedBatchTimeoutWarningInterval > 0) && (PendingBatchData.PendingBatchTryProcessCount > QueuedBatchTimeoutWarningInterval) && FPlatformProperties::RequiresCookedData())
+			if ((QueuedBatchTimeoutWarningInterval > 0) && (PendingBatchData.PendingBatchTryProcessCount > QueuedBatchTimeoutWarningInterval) && FPlatformProperties::RequiresCookedData() && !FPlatformMisc::IsDebuggerPresent())
 			{
 				const double CurrentTime = FPlatformTime::Seconds();
 
@@ -2821,31 +2821,50 @@ void FReplicationReader::ProcessQueuedBatches()
 				{
 					const double BlockedTime = FPlatformTime::ToSeconds64(FPlatformTime::Cycles64() - PendingBatchData.PendingBatchStartCycles);
 
-					const FString BlockedObjectName = *NetRefHandleManager->PrintObjectFromNetRefHandle(PendingBatchData.Handle);
-
-					UE_LOG(LogIris, Warning, TEXT("FReplicationWriter::ProcessQueuedBatches: Replication is blocked for longer than normal for: %s, QueuedBunches: %i, PendingNetRefHandleResolves: %i, BlockedTime: %f")
-						, *BlockedObjectName
-						, PendingBatchData.QueuedDataChunks.Num()
-						, PendingBatchData.PendingMustBeMappedReferences.Num()
-						, BlockedTime);
-
 					TArray<FNetRefHandle> PendingMustBeMapped;
 					PendingMustBeMapped.Reserve(MaxMustBeMappedHandleArray);
 
-					for (const FNetRefHandle& NetRefHandle : PendingBatchData.PendingMustBeMappedReferences)
+					for (const FNetRefHandle& AsyncLoadingRef : PendingBatchData.PendingMustBeMappedReferences)
 					{
-						UE_LOG(LogIris, Log, TEXT("  NetRefHandle Pending resolve [%s]"), *ObjectReferenceCache->DescribeObjectReference(FObjectReferenceCache::MakeNetObjectReference(NetRefHandle), ResolveContext));
-						if (PendingMustBeMapped.Num() < MaxMustBeMappedHandleArray)
+						if (PendingMustBeMapped.Num() >= MaxMustBeMappedHandleArray)
 						{
-							PendingMustBeMapped.Emplace(NetRefHandle);
+							// Waiting on too many assets already, ignore the rest
+							break;
+						}
+
+						// Ignore this asset if we recently warned about it
+						const double LastWarningTime = BlockedMustBeMappedLastWarningTime.FindOrAdd(AsyncLoadingRef, 0.0);
+						if (LastWarningTime == 0.0 || (CurrentTime - LastWarningTime) > QueuedBatchWarningRepeatTime)
+						{
+							PendingMustBeMapped.Emplace(AsyncLoadingRef);
+
+							BlockedMustBeMappedLastWarningTime[AsyncLoadingRef] = CurrentTime;
 						}
 					}
 
+					// Warn and tell server about the blockage
+					if (!PendingMustBeMapped.IsEmpty())
+					{
+						const FString BlockedObjectName = *NetRefHandleManager->PrintObjectFromNetRefHandle(PendingBatchData.Handle);
+
+						UE_LOG(LogIris, Warning, TEXT("FReplicationWriter::ProcessQueuedBatches: Replication is blocked for longer than normal for: %s, QueuedBunches: %i, PendingNetRefHandleResolves: %i, BlockedTime: %f")
+							, *BlockedObjectName
+							, PendingBatchData.QueuedDataChunks.Num()
+							, PendingBatchData.PendingMustBeMappedReferences.Num()
+							, BlockedTime);
+
+						for (const FNetRefHandle& NetRefHandle : PendingBatchData.PendingMustBeMappedReferences)
+						{
+							UE_LOG(LogIris, Log, TEXT("  NetRefHandle Pending resolve [%s]"), *ObjectReferenceCache->DescribeObjectReference(FObjectReferenceCache::MakeNetObjectReference(NetRefHandle), ResolveContext));
+						}
+
 #if IRIS_CLIENT_PROFILER_ENABLE
-					FClientProfiler::RecordBlockedReplication(*BlockedObjectName, PendingBatchData.PendingMustBeMappedReferences.Num(), (float)BlockedTime);
+						FClientProfiler::RecordBlockedReplication(*BlockedObjectName, PendingBatchData.PendingMustBeMappedReferences.Num(), (float)BlockedTime);
 #endif
 
-					ReplicationBridge->SendErrorWithNetRefHandle(UE::Net::ENetRefHandleError::BlockedByMustBeMapped, PendingBatchData.Handle, Parameters.ConnectionId, PendingMustBeMapped);
+						ReplicationBridge->SendErrorWithNetRefHandle(UE::Net::ENetRefHandleError::BlockedByMustBeMapped, PendingBatchData.Handle, Parameters.ConnectionId, PendingMustBeMapped);
+						
+					}
 
 					PendingBatchData.NextWarningTimeout = CurrentTime + QueuedBatchWarningRepeatTime;
 				}
diff --git a/Engine/Source/Runtime/Experimental/Iris/Core/Private/Iris/ReplicationSystem/ReplicationReader.h b/Engine/Source/Runtime/Experimental/Iris/Core/Private/Iris/ReplicationSystem/ReplicationReader.h
--- a/Engine/Source/Runtime/Experimental/Iris/Core/Private/Iris/ReplicationSystem/ReplicationReader.h
+++ b/Engine/Source/Runtime/Experimental/Iris/Core/Private/Iris/ReplicationSystem/ReplicationReader.h
@@ -281,6 +281,10 @@ private:
 	// We do not expect many objects to be broken
 	TArray<FNetRefHandle> BrokenObjects;
 
+	// Keep track of the last time we warned about an object blocked by this must be mapped reference
+	// Prevents spamming errors for multiple objects all waiting on the same asset
+	TMap<FNetRefHandle, double> BlockedMustBeMappedLastWarningTime;
+
 	// Used during receive and processing of pending batches
 	TArray<FNetRefHandle> TempMustBeMappedReferences;
 
