Iris
* Set async loading priorities to:
* Default = 200
* High = 225
* VeryHigh = 250

diff --git a/Engine/Source/Runtime/Experimental/Iris/Core/Private/Iris/ReplicationSystem/ObjectReferenceTypes.cpp b/Engine/Source/Runtime/Experimental/Iris/Core/Private/Iris/ReplicationSystem/ObjectReferenceTypes.cpp
--- a/Engine/Source/Runtime/Experimental/Iris/Core/Private/Iris/ReplicationSystem/ObjectReferenceTypes.cpp
+++ b/Engine/Source/Runtime/Experimental/Iris/Core/Private/Iris/ReplicationSystem/ObjectReferenceTypes.cpp
@@ -6,21 +6,21 @@
 
 namespace UE::Net::Private
 {
-	static int32 IrisDefaultAsyncLoadingPriority = 0;
+	static int32 IrisDefaultAsyncLoadingPriority = 200;
 	static FAutoConsoleVariableRef CVarIrisDefaultAsyncLoadingPriority(
 		TEXT("net.iris.AsyncLoading.DefaultPriority"),
 		IrisDefaultAsyncLoadingPriority,
 		TEXT("The default priority to use when async loading packages referenced by replicated properties.")
 		);
 
-	static int32 IrisHighAsyncLoadingPriority = 50;
+	static int32 IrisHighAsyncLoadingPriority = 225;
 	static FAutoConsoleVariableRef CVarIrisHighAsyncLoadingPriority(
 		TEXT("net.iris.AsyncLoading.HighPriority"),
 		IrisHighAsyncLoadingPriority,
 		TEXT("The high priority value to use when async loading packages referenced by replicated properties.")
 		);
 
-	static int32 IrisVeryHighAsyncLoadingPriority = 100;
+	static int32 IrisVeryHighAsyncLoadingPriority = 250;
 	static FAutoConsoleVariableRef CVarIrisVeryHighAsyncLoadingPriority(
 		TEXT("net.iris.AsyncLoading.VeryHighPriority"),
 		IrisVeryHighAsyncLoadingPriority,
