Optionally clear input state on physics thread once it has been consumed, to prevent same input being consumed during multiple simulation steps when GT is running very slowly

diff --git a/Engine/Plugins/Experimental/ChaosModularVehicle/Source/ChaosModularVehicleEngine/Private/ChaosModularVehicle/ModularVehicleSimulationCU.cpp b/Engine/Plugins/Experimental/ChaosModularVehicle/Source/ChaosModularVehicleEngine/Private/ChaosModularVehicle/ModularVehicleSimulationCU.cpp
--- a/Engine/Plugins/Experimental/ChaosModularVehicle/Source/ChaosModularVehicleEngine/Private/ChaosModularVehicle/ModularVehicleSimulationCU.cpp
+++ b/Engine/Plugins/Experimental/ChaosModularVehicle/Source/ChaosModularVehicleEngine/Private/ChaosModularVehicle/ModularVehicleSimulationCU.cpp
@@ -139,6 +139,10 @@ void FModularVehicleSimulation::SimulateModuleTree(UWorld* InWorld, float DeltaS
 		{
 			VehicleSimulationCallback.Broadcast(DeltaSeconds, InputData, SimInputData, Proxy, RootParticle, SimModuleTree.Get());
 		}
+
+		// Clear those Inputs that we don't want to remain set if the physics simulation thread ticks more frames than the GT
+		InputData.PhysicsInputs.NetworkInputs.VehicleInputs.Container.ClearConsumedInputs();
+		InputData.PhysicsInputs.StateInputs.StateInputContainer.ClearConsumedInputs();
 	}
 
 }
diff --git a/Engine/Source/Runtime/Experimental/ChaosVehicles/ChaosVehiclesCore/Private/SimModule/ModuleInput.cpp b/Engine/Source/Runtime/Experimental/ChaosVehicles/ChaosVehiclesCore/Private/SimModule/ModuleInput.cpp
--- a/Engine/Source/Runtime/Experimental/ChaosVehicles/ChaosVehiclesCore/Private/SimModule/ModuleInput.cpp
+++ b/Engine/Source/Runtime/Experimental/ChaosVehicles/ChaosVehiclesCore/Private/SimModule/ModuleInput.cpp
@@ -252,7 +252,8 @@ void FModuleInputValue::Serialize(FArchive& Ar, UPackageMap* Map, bool& bOutSucc
 	Ar.SerializeBits(&bIsNonZero, 1);
 	Ar.SerializeBits(&ValueType, 3);
 	Ar.SerializeBits(&bApplyInputDecay, 1);
-	
+	Ar.SerializeBits(&bClearAfterConsumed, 1);
+
 
 	if (bIsNonZero)
 	{
@@ -304,6 +305,7 @@ void FModuleInputValue::NetSerialize(FArchive& Ar, UPackageMap* Map, bool& bOutS
 	{
 		Ar.SerializeBits(&ValueType, 3);
 		Ar.SerializeBits(&bApplyInputDecay, 1);
+		Ar.SerializeBits(&bClearAfterConsumed, 1);
 	}
 
 	if (bIsNonZero)
@@ -366,12 +368,20 @@ void FModuleInputValue::DeltaNetSerialize(FArchive& Ar, UPackageMap* Map, bool&
 		Value = PreviousInputValue.Value;
 		ValueInt = PreviousInputValue.ValueInt;
 		bApplyInputDecay = PreviousInputValue.bApplyInputDecay;
+		bClearAfterConsumed = PreviousInputValue.bClearAfterConsumed;
 		return;
 	}
 
 	NetSerialize(Ar, Map, bOutSuccess, InInputQuantizationType);
 }
 
+/** Set only the value without overriding the other properties in FModuleInputValue */
+void FModuleInputValue::Set(const FModuleInputValue& In)
+{
+	ValueInt = In.ValueInt;
+	Value += In.Value;
+}
+
 void FModuleInputValue::Lerp(const FModuleInputValue& Min, const FModuleInputValue& Max, float Alpha)
 {
 	switch (GetValueType())
@@ -557,6 +567,7 @@ void FModuleInputContainer::Initialize(TArray<FModuleInputSetup>& SetupData, FIn
 	{
 		int Index = AddInput(Setup.Type);
 		InputValues[Index].SetApplyInputDecay(Setup.bApplyInputDecay);
+		InputValues[Index].SetClearAfterConsumed(Setup.bClearAfterConsumed);
 		NameMapOut.Add(Setup.Name, Index);
 	}
 }
@@ -646,6 +657,17 @@ void FModuleInputContainer::Decay(const float DecayAmount)
 	}
 }
 
+void FModuleInputContainer::ClearConsumedInputs()
+{
+	for (int I = 0; I < InputValues.Num(); I++)
+	{
+		if (InputValues[I].ShouldClearAfterConsumed())
+		{
+			InputValues[I].Reset();
+		}
+	}
+}
+
 void FInputInterface::SetValue(const FName& InName, const FModuleInputValue& InValue, bool Quantize /*= true*/)
 {
 	if (ValueContainer.GetNumInputs() > 0)
@@ -764,4 +786,4 @@ bool FInputInterface::InputsNonZero() const
 	}
 
 	return false;
-}
\ No newline at end of file
+}
diff --git a/Engine/Source/Runtime/Experimental/ChaosVehicles/ChaosVehiclesCore/Public/SimModule/ModuleInput.h b/Engine/Source/Runtime/Experimental/ChaosVehicles/ChaosVehiclesCore/Public/SimModule/ModuleInput.h
--- a/Engine/Source/Runtime/Experimental/ChaosVehicles/ChaosVehiclesCore/Public/SimModule/ModuleInput.h
+++ b/Engine/Source/Runtime/Experimental/ChaosVehicles/ChaosVehiclesCore/Public/SimModule/ModuleInput.h
@@ -356,6 +356,8 @@ public:
 	UE_API void NetSerialize(FArchive& Ar, UPackageMap* Map, bool& bOutSuccess, EModuleInputQuantizationType InInputQuantizationType);
 	UE_API void DeltaNetSerialize(FArchive& Ar, UPackageMap* Map, bool& bOutSuccess,const FModuleInputValue& PreviousInputValue, EModuleInputQuantizationType InInputQuantizationType);
 
+	UE_API void Set(const FModuleInputValue& In);
+
 	UE_API void Lerp(const FModuleInputValue& Min, const FModuleInputValue& Max, float Alpha);
 
 	UE_API void Merge(const FModuleInputValue& From);
@@ -371,6 +373,10 @@ public:
 	void SetApplyInputDecay(const bool bInApplyInputDecay) { bApplyInputDecay = bInApplyInputDecay; }
 	bool ShouldApplyInputDecay() const { return  bApplyInputDecay; }
 
+	/** After simulation first consumes the input, clear the value doesn't persist across multiple async physics ticks */
+	void SetClearAfterConsumed(const bool bInClearAfterConsumed) { bClearAfterConsumed = bInClearAfterConsumed; }
+	bool ShouldClearAfterConsumed() const { return  bClearAfterConsumed; }
+
 protected:
 	void SetMagnitude(float NewSize);
 
@@ -385,6 +391,10 @@ protected:
 
 	UPROPERTY()
 	bool bApplyInputDecay = false;
+
+	UPROPERTY()
+	bool bClearAfterConsumed = false;
+
 };
 
 // Supported getter specializations
@@ -514,12 +524,14 @@ struct FModuleInputSetup
 	{
 		Type = EModuleInputValueType::MBoolean;
 		bApplyInputDecay = false;
+		bClearAfterConsumed = false;
 	}
 
 	FModuleInputSetup(const FName& InName, const EModuleInputValueType& InType)
 		: Name(InName)
 		, Type(InType)
 		, bApplyInputDecay(false)
+		, bClearAfterConsumed(false)
 	{
 	}
 
@@ -537,6 +549,10 @@ struct FModuleInputSetup
 	/** During physics resimulation, apply decay for this input while it's being extrapolated */
 	UPROPERTY(EditAnywhere, Category = VehicleInput)
 	bool bApplyInputDecay;
+
+	/** After simulation first consumes the input, clear the value doesn't persist across multiple async physics ticks */
+	UPROPERTY(EditAnywhere, Category = VehicleInput)
+	bool bClearAfterConsumed;
 };
 
 class FScopedModuleInputInitializer
@@ -587,11 +603,11 @@ public:
 	{
 		if (Quantize)
 		{
-			InputValues[Index] = InValue.ReturnQuantized(InInputQuantizationType);
+			InputValues[Index].Set(InValue.ReturnQuantized(InInputQuantizationType));
 		}
 		else
 		{
-			InputValues[Index] = InValue;
+			InputValues[Index].Set(InValue);
 		}
 	}
 
@@ -638,6 +654,9 @@ public:
 	/**  Decay input during resimulation by DecayAmount which increases over resimulation frames from 0.0 -> 1.0 when the input is being reused */
 	UE_API void Decay(const float DecayAmount);
 
+	/** Clear input once it has been consumed if input is marked as sutch. Prevents same input being fed into simulation when the PT tick is running multiple times per GT tick */
+	UE_API void ClearConsumedInputs();
+
 	TArray<FModuleInputValue>& AccessInputValues() { return InputValues; }
 
 private:
@@ -760,4 +779,5 @@ public:
 	EModuleInputQuantizationType InputQuantizationType;
 };
 
-#undef UE_API 
\ No newline at end of file
+#undef UE_API 
+
