Enable Parallel IO Operations on Indexed Cache Storage
- CreateCacheStorage / DestroyCacheStorage will now wait until a cache has been unmounted
- This is done to handle interleaved thread scenarios:
+ Thread A mounts the cache and starts processing files.
+ Thread B requests cache destruction while A is still active.
+ WaitForPredicateWhileLocked in DestroyCacheStorage allows B to release the lock and wait on a condition variable until A completes unmounting and calls NotifyWaitersWhileLocked.
- Mounting remains reference-counted to prevent redundant mounts.
- Remove the global lock since the array of entries is immutable after initialization. This means that GetMountPath will not wait for an IO operation to finish
- Move the array as a const array in FIndexedCacheStorageMetaData to make sure we do not change it by mistake after initialization

diff --git a/Engine/Source/Runtime/Experimental/IndexedCacheStorage/Private/IndexedCacheStorageManager.cpp b/Engine/Source/Runtime/Experimental/IndexedCacheStorage/Private/IndexedCacheStorageManager.cpp
--- a/Engine/Source/Runtime/Experimental/IndexedCacheStorage/Private/IndexedCacheStorageManager.cpp
+++ b/Engine/Source/Runtime/Experimental/IndexedCacheStorage/Private/IndexedCacheStorageManager.cpp
@@ -1,10 +1,16 @@
 // Copyright Epic Games, Inc. All Rights Reserved.
 
 #include "IndexedCacheStorageManager.h"
+#include "IndexedCacheStorage.h"
+
+#include "Async/Mutex.h"
+#include "Async/UniqueLock.h"
+#include "Experimental/Async/ConditionVariable.h"
 #include "HAL/PlatformMisc.h"
 #include "Misc/ConfigCacheIni.h"
+#include "HAL/PlatformStackWalk.h"
+#include "Misc/ScopeExit.h"
 #include "Misc/ScopeRWLock.h"
-#include "IndexedCacheStorage.h"
 
 DECLARE_LOG_CATEGORY_EXTERN(LogIndexedCacheStorageManager, Display, All);
 DEFINE_LOG_CATEGORY(LogIndexedCacheStorageManager);
@@ -12,6 +18,269 @@ DEFINE_LOG_CATEGORY(LogIndexedCacheStorageManager);
 namespace Experimental
 {
 
+class FCacheStorageMetaData
+{
+public:
+	enum class ECacheStorageState : uint8
+	{
+		Deleted,
+		Created,
+		Mounted,
+		Count
+	};
+
+	const TCHAR* LexToString(ECacheStorageState Val)
+	{
+		static const TCHAR* Strings[] =
+		{
+			TEXT("ECacheStorageState::Deleted"),
+			TEXT("ECacheStorageState::Created"),
+			TEXT("ECacheStorageState::Mounted"),
+		};
+		static_assert(int32(ECacheStorageState::Count) == UE_ARRAY_COUNT(Strings), "");
+
+		return Strings[int32(Val)];
+	}
+	
+	FCacheStorageMetaData(int32 InCacheIndex, const FString& InMountName, bool bCacheExists)
+	: MountName(InMountName)
+	, CacheIndex(InCacheIndex)
+	{
+		CacheStorageState = bCacheExists ? ECacheStorageState::Created : ECacheStorageState::Deleted;
+	}
+
+	bool IsValid() const
+	{
+		return !MountName.IsEmpty();
+	}
+
+	const FString& GetMountName() const
+	{
+		return MountName;
+	}
+
+	FString GetMountPath() const
+	{
+		FString MountedPath;
+		IIndexedCacheStorage::Get().GetCacheStorageMountPath(MountedPath, MountName);
+		return MountedPath;
+	}
+
+	void WaitForPredicateWhileLocked(const TFunction<bool()>& Predicate)
+	{
+		while (!Predicate())
+		{
+			ConditionVariable.Wait(Mutex);
+		}
+	}
+
+	void NotifyWaitersWhileLocked()
+	{
+		ConditionVariable.NotifyAll();
+	}
+
+	bool CreateCacheStorage(uint64 RequestNumberOfBytes)
+	{
+		UE::TUniqueLock Lock(Mutex);
+		 
+		UE_LOG(LogIndexedCacheStorageManager, Display, TEXT("CreateCacheStorage request for %ull bytes at cache %d. Current state %s"), RequestNumberOfBytes, CacheIndex, LexToString(CacheStorageState));
+
+		WaitForPredicateWhileLocked([this]()->bool 
+		{ 
+			return CacheStorageState != ECacheStorageState::Mounted;
+		});;
+
+		check(MountRefCount == 0);
+
+		// Do not uncomment the code below: even if the cache is in 'Created' state, we might want to handle the case of a reallocation and 
+		// the lower level takes care of either reuse the one we asked for (same size), or destroy and create a new one
+		//if (CacheStorageState == ECacheStorageState::Created)
+		//{
+		//	return true;
+		//}
+
+		UE_LOG(LogIndexedCacheStorageManager, Display, TEXT("CreateCacheStorage creating cache %d with %ull bytes"), CacheIndex, RequestNumberOfBytes);
+		bool bCacheExists = IIndexedCacheStorage::Get().CreateCacheStorage(RequestNumberOfBytes, CacheIndex);
+		
+		if (bCacheExists)
+		{
+			CacheStorageState = ECacheStorageState::Created;
+			UE_LOG(LogIndexedCacheStorageManager, Display, TEXT("CreateCacheStorage created cache %d at %s. Current state %s"), CacheIndex, *MountName, LexToString(CacheStorageState));
+		}
+		else
+		{
+			CacheStorageState = ECacheStorageState::Deleted;
+			UE_LOG(LogIndexedCacheStorageManager, Error, TEXT("CreateCacheStorage failed to create cache %d at %s. Current state %s"), CacheIndex, *MountName, LexToString(CacheStorageState));
+		}
+
+		NotifyWaitersWhileLocked();
+
+		return bCacheExists;
+	}
+
+	void DestroyCacheStorage()
+	{
+		UE::TUniqueLock Lock(Mutex);
+		 
+		UE_LOG(LogIndexedCacheStorageManager, Display, TEXT("DestroyCacheStorage request at cache %d. Current state %s"), CacheIndex, LexToString(CacheStorageState));
+
+		WaitForPredicateWhileLocked([this]()->bool 
+		{ 
+			return CacheStorageState != ECacheStorageState::Mounted;
+		});;
+
+		check(MountRefCount == 0);
+
+		if (CacheStorageState == ECacheStorageState::Deleted)
+		{
+			UE_LOG(LogIndexedCacheStorageManager, Display, TEXT("DestroyCacheStorage skipped deletion of cache %d at %s"), CacheIndex, *MountName);
+			return;
+		}
+
+		UE_LOG(LogIndexedCacheStorageManager, Display, TEXT("DestroyCacheStorage deleting cache %d"), CacheIndex);
+		bool bCachedDestroyed = IIndexedCacheStorage::Get().DestroyCacheStorage(CacheIndex);
+		
+		if (bCachedDestroyed)
+		{
+			CacheStorageState = ECacheStorageState::Deleted;
+			UE_LOG(LogIndexedCacheStorageManager, Display, TEXT("DestroyCacheStorage deleted cache %d at %s. Current state %s"), CacheIndex, *MountName, LexToString(CacheStorageState));
+		}
+		else
+		{
+			CacheStorageState = ECacheStorageState::Created;
+			UE_LOG(LogIndexedCacheStorageManager, Error, TEXT("DestroyCacheStorage failed to delete cache %d at %s. Current state %s"), CacheIndex, *MountName, LexToString(CacheStorageState));
+		}
+
+		NotifyWaitersWhileLocked();
+	}
+
+	uint64 GetCacheStorageInfo()
+	{
+		UE::TUniqueLock Lock(Mutex);
+
+		UE_LOG(LogIndexedCacheStorageManager, Display, TEXT("GetCacheStorageInfo request at cache %d. Current state %s"), CacheIndex, LexToString(CacheStorageState));
+
+		uint64 ExistingSpace = 0;
+
+		if (CacheStorageState == ECacheStorageState::Deleted)
+		{
+			UE_LOG(LogIndexedCacheStorageManager, Display, TEXT("GetCacheStorageInfo cache %d is empty"), CacheIndex, ExistingSpace);
+			return ExistingSpace;
+		}
+		
+		UE_LOG(LogIndexedCacheStorageManager, Display, TEXT("GetCacheStorageInfo cache %d"), CacheIndex);
+		IIndexedCacheStorage::Get().GetCacheStorageInfo(CacheIndex, ExistingSpace);
+		UE_LOG(LogIndexedCacheStorageManager, Display, TEXT("GetCacheStorageInfo cache %d is %llu bytes"), CacheIndex, ExistingSpace);
+		
+		NotifyWaitersWhileLocked();
+
+		return ExistingSpace;
+	}
+
+	FString MountCacheStorage()
+	{
+		UE::TUniqueLock Lock(Mutex);
+
+		UE_LOG(LogIndexedCacheStorageManager, Display, TEXT("MountCacheStorage request at cache %d. Current state %s"), CacheIndex, LexToString(CacheStorageState));
+
+		if (CacheStorageState == ECacheStorageState::Deleted)
+		{
+			check(MountRefCount == 0);
+			UE_LOG(LogIndexedCacheStorageManager, Display, TEXT("MountCacheStorage skip mount for cache %d deleted"), CacheIndex);
+			return FString();
+		}
+
+		FString MountedPath = GetMountPath();
+		if (CacheStorageState == ECacheStorageState::Mounted)
+		{
+			check(MountRefCount > 0);
+			UE_LOG(LogIndexedCacheStorageManager, Display, TEXT("MountCacheStorage skip mount for cache %d already mounted (%d)"), CacheIndex, MountRefCount);
+			++MountRefCount;
+			return MountedPath;
+		}
+
+		check(MountRefCount == 0);
+		check(CacheStorageState == ECacheStorageState::Created);
+
+		UE_LOG(LogIndexedCacheStorageManager, Display, TEXT("MountCacheStorage mounting cache %d at %s"), CacheIndex, *MountedPath);
+		bool bMounted = IIndexedCacheStorage::Get().MountCacheStorage(MountedPath, CacheIndex, MountName);
+		
+		check(MountRefCount == 0);
+
+		if (bMounted)
+		{
+			CacheStorageState = ECacheStorageState::Mounted;
+			MountRefCount = 1;
+			UE_LOG(LogIndexedCacheStorageManager, Display, TEXT("MountCacheStorage mounted cache %d at %s. Current state %s"), CacheIndex, *MountedPath, LexToString(CacheStorageState));
+		}
+		else
+		{
+			CacheStorageState = ECacheStorageState::Created;
+			UE_LOG(LogIndexedCacheStorageManager, Error, TEXT("MountCacheStorage failed to mount cache %d at %s. Current state %s"), CacheIndex, *MountedPath, LexToString(CacheStorageState));
+		}
+
+		NotifyWaitersWhileLocked();
+		return (bMounted) ? MountedPath : FString();
+	}
+
+
+	void UnmountCacheStorage()
+	{
+		UE::TUniqueLock Lock(Mutex);
+
+		UE_LOG(LogIndexedCacheStorageManager, Display, TEXT("UnmountCacheStorage request at cache %d. Current state %s"), CacheIndex, LexToString(CacheStorageState));
+
+		if (CacheStorageState == ECacheStorageState::Deleted || CacheStorageState == ECacheStorageState::Created)
+		{
+			check(MountRefCount == 0);
+			UE_LOG(LogIndexedCacheStorageManager, Display, TEXT("UnmountCacheStorage skip unmount for cache %d already unmounted"), CacheIndex);
+			return;
+		}
+
+		check(CacheStorageState == ECacheStorageState::Mounted);
+		check(MountRefCount > 0);
+
+		if (MountRefCount > 1)
+		{
+			--MountRefCount;
+			UE_LOG(LogIndexedCacheStorageManager, Display, TEXT("UnmountCacheStorage cache %d still mounted (%d)"), CacheIndex, MountRefCount);
+			return;
+		}
+
+		FString MountedPath = GetMountPath();
+		UE_LOG(LogIndexedCacheStorageManager, Display, TEXT("UnmountCacheStorage unmounting cache %d at %s"), CacheIndex, *MountedPath);
+		IIndexedCacheStorage::Get().UnmountCacheStorage(*MountName);
+		CacheStorageState = ECacheStorageState::Created;
+		UE_LOG(LogIndexedCacheStorageManager, Display, TEXT("UnmountCacheStorage unmounted cache %d at %s. Current state %s"), CacheIndex, *MountedPath, LexToString(CacheStorageState));
+
+		check(MountRefCount == 1);
+		MountRefCount = 0;
+		NotifyWaitersWhileLocked();
+	}
+
+private:
+	const FString MountName;
+	const int32 CacheIndex;
+	int32 MountRefCount = 0;
+	ECacheStorageState CacheStorageState = ECacheStorageState::Deleted;
+	UE::FMutex Mutex;
+	UE::FConditionVariable ConditionVariable;
+
+};
+
+using FCacheStorageMetaDataArray = TArray<TUniquePtr<FCacheStorageMetaData> >;
+
+struct FIndexedCacheStorageMetaData
+{
+	FIndexedCacheStorageMetaData(FCacheStorageMetaDataArray&& InCacheStorageMetaDataArray)
+	: IndexedCacheEntriesMetaData(MoveTemp(InCacheStorageMetaDataArray))
+	{
+	}
+
+	const FCacheStorageMetaDataArray IndexedCacheEntriesMetaData;
+
+};
+
 static FIndexedCacheStorageManager* GIndexedCacheStorageManager = nullptr;
 
 FIndexedCacheStorageManager& FIndexedCacheStorageManager::Get()
@@ -29,8 +298,21 @@ FIndexedCacheStorageManager& FIndexedCacheStorageManager::Get()
 void FIndexedCacheStorageManager::Initialize()
 {
 	uint32 NumIndexedCacheEntries = IIndexedCacheStorage::Get().GetNumIndexedCacheStorages();
+	FCacheStorageMetaDataArray IndexedCacheEntriesMetaData;
 	IndexedCacheEntriesMetaData.SetNum(NumIndexedCacheEntries);
 
+	TArray<int32> CacheStorageIndices;
+	IIndexedCacheStorage::Get().EnumerateCacheStorages(CacheStorageIndices);
+
+	TArray<bool> CacheStorageAlreadyExists;
+	CacheStorageAlreadyExists.SetNumZeroed(NumIndexedCacheEntries);
+
+	for (int32 CacheStorageIndex : CacheStorageIndices)
+	{
+		CacheStorageAlreadyExists[CacheStorageIndex] = true;
+	}
+
+
 	// Takes on the pattern
 	// (Name="CacheStorageName",Index=CacheStorageIndex,Mount="CacheStorageMount")
 	TArray<FString> StorageConfigs;
@@ -89,30 +371,26 @@ void FIndexedCacheStorageManager::Initialize()
 				UE_LOG(LogIndexedCacheStorageManager, Error, TEXT("Found an existing entry with name %s in [IndexedCacheStorage]:Storage"), *StorageName);
 				continue;
 			}
-			
-			FCacheStorageMetaData& CacheStorageMetaData = IndexedCacheEntriesMetaData[StorageIndex];
+			ExistingStorageIndex = StorageIndex;
 
-			if (!CacheStorageMetaData.MountName.IsEmpty())
+			if (IndexedCacheEntriesMetaData[StorageIndex] == nullptr)
 			{
-				UE_LOG(LogIndexedCacheStorageManager, Error, TEXT("Found an existing entry with at index %d for name %s in [IndexedCacheStorage]:Storage"), StorageIndex, *StorageName);
-				continue;
+				IndexedCacheEntriesMetaData[StorageIndex] = MakeUnique<FCacheStorageMetaData>(StorageIndex, StorageMount, CacheStorageAlreadyExists[StorageIndex]);
+				++RegisteredStorages;
+			}
+			else
+			{
+				const FString& ExistingName = IndexedCacheEntriesMetaData[StorageIndex]->GetMountName();
+				UE_LOG(LogIndexedCacheStorageManager, Error, TEXT("Found an existing entry named %s with at cache %d for name %s in [IndexedCacheStorage]:Storage"), *ExistingName, StorageIndex, *StorageName);
 			}
-
-			ExistingStorageIndex = StorageIndex;
-			CacheStorageMetaData.MountName = StorageMount;
-			++RegisteredStorages;
 		}
 	}
-
-	TArray<int32> CacheStorageIndices;
-	IIndexedCacheStorage::Get().EnumerateCacheStorages(CacheStorageIndices);
-	TSet<int32> CacheStorageIndicesSet;
-	for (int32 CacheStorageIndex : CacheStorageIndices)
+	
+	bHasRegisteredEntries = (RegisteredStorages > 0);
+	if (bHasRegisteredEntries)
 	{
-		IndexedCacheEntriesMetaData[CacheStorageIndex].bCacheExists = true;
+		IndexedCacheStorageMetaData = MakeUnique<FIndexedCacheStorageMetaData>(MoveTemp(IndexedCacheEntriesMetaData));
 	}
-	bHasRegisteredEntries = (RegisteredStorages > 0);
-
 }
 
 bool FIndexedCacheStorageManager::SupportsIndexedCacheStorage()
@@ -122,17 +400,12 @@ bool FIndexedCacheStorageManager::SupportsIndexedCacheStorage()
 
 int32 FIndexedCacheStorageManager::GetStorageIndex(const FString& StorageName)
 {
-	FRWScopeLock _(IndexedCacheEntriesMetaDataLock, SLT_ReadOnly);
-
 	int32* StorageIndex = StorageNameToStorageIndex.Find(StorageName);
-
 	return StorageIndex ? *StorageIndex : 0;
 }
 
 void FIndexedCacheStorageManager::EnumerateCacheStorages(TArray<FString>& OutCacheStorageNames)
 {
-	FRWScopeLock _(IndexedCacheEntriesMetaDataLock, SLT_ReadOnly);
-
 	OutCacheStorageNames.Reserve(StorageNameToStorageIndex.Num());
 	for (const TPair<FString, int32>& StorageNameToStorageIndexIter : StorageNameToStorageIndex)
 	{
@@ -140,148 +413,106 @@ void FIndexedCacheStorageManager::EnumerateCacheStorages(TArray<FString>& OutCac
 	}
 }
 
-
 bool FIndexedCacheStorageManager::CreateCacheStorage(uint64 RequestNumberOfBytes, int32 CacheIndex)
 {
-	FRWScopeLock _(IndexedCacheEntriesMetaDataLock, SLT_Write);
-	if (CacheIndex <= 0 && !ensure(IndexedCacheEntriesMetaData.IsValidIndex(CacheIndex)))
-	{
-		UE_LOG(LogIndexedCacheStorageManager, Error, TEXT("CreateCacheStorage invalid index %d"), CacheIndex);
-		return false;
-	}
-
-	if (IndexedCacheEntriesMetaData[CacheIndex].MountName.IsEmpty())
+	FCacheStorageMetaData* CacheStorageMetaData = GetCacheStorageMetaData(CacheIndex);
+	if (CacheStorageMetaData == nullptr)
 	{
-		UE_LOG(LogIndexedCacheStorageManager, Error, TEXT("CreateCacheStorage invalid empty mount name"));
+		UE_LOG(LogIndexedCacheStorageManager, Error, TEXT("CreateCacheStorage skip invalid cache %d"), CacheIndex);
 		return false;
 	}
 
-	ensure(IndexedCacheEntriesMetaData[CacheIndex].MountRefCount == 0);
-	IndexedCacheEntriesMetaData[CacheIndex].bCacheExists = IIndexedCacheStorage::Get().CreateCacheStorage(RequestNumberOfBytes, CacheIndex);
-
-	UE_LOG(LogIndexedCacheStorageManager, Display, TEXT("CreateCacheStorage %ull bytes at index %d: %d"), RequestNumberOfBytes, CacheIndex, IndexedCacheEntriesMetaData[CacheIndex].bCacheExists);
-
-	return IndexedCacheEntriesMetaData[CacheIndex].bCacheExists;
+	return CacheStorageMetaData->CreateCacheStorage(RequestNumberOfBytes);
 }
 
 void FIndexedCacheStorageManager::DestroyCacheStorage(int32 CacheIndex)
 {
-	FRWScopeLock _(IndexedCacheEntriesMetaDataLock, SLT_Write);
-
-	if (CacheIndex <= 0 && !ensure(IndexedCacheEntriesMetaData.IsValidIndex(CacheIndex)))
+	FCacheStorageMetaData* CacheStorageMetaData = GetCacheStorageMetaData(CacheIndex);
+	if (CacheStorageMetaData == nullptr)
 	{
+		UE_LOG(LogIndexedCacheStorageManager, Error, TEXT("DestroyCacheStorage skip invalid cache %d"), CacheIndex);
 		return;
 	}
 
-	if (IndexedCacheEntriesMetaData[CacheIndex].MountName.IsEmpty() || !IndexedCacheEntriesMetaData[CacheIndex].bCacheExists)
-	{
-		return;
-	}
-
-	ensure(IndexedCacheEntriesMetaData[CacheIndex].MountRefCount == 0);
-	UE_LOG(LogIndexedCacheStorageManager, Display, TEXT("DestroyCacheStorage  index %d"), CacheIndex);
-	IIndexedCacheStorage::Get().DestroyCacheStorage(CacheIndex);
-	IndexedCacheEntriesMetaData[CacheIndex].bCacheExists = false;
+	CacheStorageMetaData->DestroyCacheStorage();
 }
 
 uint64 FIndexedCacheStorageManager::GetCacheStorageCapacity(int32 CacheIndex)
 {
-	FRWScopeLock _(IndexedCacheEntriesMetaDataLock, SLT_ReadOnly);
-
-	if (CacheIndex <= 0 && !ensure(IndexedCacheEntriesMetaData.IsValidIndex(CacheIndex)))
-	{
-		return 0;
-	}
-
-	if (IndexedCacheEntriesMetaData[CacheIndex].MountName.IsEmpty() || !IndexedCacheEntriesMetaData[CacheIndex].bCacheExists)
+	FCacheStorageMetaData* CacheStorageMetaData = GetCacheStorageMetaData(CacheIndex);
+	if (CacheStorageMetaData == nullptr)
 	{
+		UE_LOG(LogIndexedCacheStorageManager, Warning, TEXT("GetCacheStorageCapacity skip invalid cache %d"), CacheIndex);
 		return 0;
 	}
 
-	uint64 ExistingSpace = 0;
-	IIndexedCacheStorage::Get().GetCacheStorageInfo(CacheIndex, ExistingSpace);
-	return ExistingSpace;
+	return CacheStorageMetaData->GetCacheStorageInfo();
 }
 
 FString FIndexedCacheStorageManager::MountCacheStorage(int32 CacheIndex)
 {
-	FRWScopeLock _(IndexedCacheEntriesMetaDataLock, SLT_Write);
-
-	FString ActualMountName;
-	if (!IndexedCacheEntriesMetaData.IsValidIndex(CacheIndex) || !IndexedCacheEntriesMetaData[CacheIndex].bCacheExists)
-	{
-		UE_LOG(LogIndexedCacheStorageManager, Display, TEXT("MountCacheStorage skip mount for cache %d, exists=%d"), CacheIndex, IndexedCacheEntriesMetaData[CacheIndex].bCacheExists);
-		return ActualMountName;
-	}
-
-	check(IndexedCacheEntriesMetaData[CacheIndex].MountRefCount >= 0);
-
-	FString MountedPath;
-	if (IndexedCacheEntriesMetaData[CacheIndex].MountRefCount == 0)
+	FCacheStorageMetaData* CacheStorageMetaData = GetCacheStorageMetaData(CacheIndex);
+	if (CacheStorageMetaData == nullptr)
 	{
-		bool bMounted = IIndexedCacheStorage::Get().MountCacheStorage(MountedPath, CacheIndex, IndexedCacheEntriesMetaData[CacheIndex].MountName);
-		if (!bMounted)
-		{
-			UE_LOG(LogIndexedCacheStorageManager, Error, TEXT("MountCacheStorage failed to mount cache %d"), CacheIndex);
-			return ActualMountName;
-		}
-		UE_LOG(LogIndexedCacheStorageManager, Display, TEXT("MountCacheStorage mounted cache %d at %s"), CacheIndex, *MountedPath);
-	}
-	else
-	{
-		IIndexedCacheStorage::Get().GetCacheStorageMountPath(MountedPath, IndexedCacheEntriesMetaData[CacheIndex].MountName);
-		UE_LOG(LogIndexedCacheStorageManager, Display, TEXT("MountCacheStorage cache %d already mounted %d times at %s"), CacheIndex, IndexedCacheEntriesMetaData[CacheIndex].MountRefCount, *MountedPath);
+		UE_LOG(LogIndexedCacheStorageManager, Warning, TEXT("MountCacheStorage skip invalid cache %d"), CacheIndex);
+		return FString();
 	}
 
-	check(!MountedPath.IsEmpty());
-
-	++IndexedCacheEntriesMetaData[CacheIndex].MountRefCount;
-
-	return MountedPath;
+	return CacheStorageMetaData->MountCacheStorage();
 }
 
 void FIndexedCacheStorageManager::UnmountCacheStorage(int32 CacheIndex)
 {
-	FRWScopeLock _(IndexedCacheEntriesMetaDataLock, SLT_Write);
-
-	FString ActualMountName;
-	if (!IndexedCacheEntriesMetaData.IsValidIndex(CacheIndex) || IndexedCacheEntriesMetaData[CacheIndex].MountRefCount == 0)
+	FCacheStorageMetaData* CacheStorageMetaData = GetCacheStorageMetaData(CacheIndex);
+	if (CacheStorageMetaData == nullptr)
 	{
-		UE_LOG(LogIndexedCacheStorageManager, Display, TEXT("UnmountCacheStorage skip unmount for cache %d mounted %d times"), CacheIndex, IndexedCacheEntriesMetaData[CacheIndex].MountRefCount);
+		UE_LOG(LogIndexedCacheStorageManager, Warning, TEXT("UnmountCacheStorage skip invalid cache %d"), CacheIndex);
 		return;
 	}
 
-	--IndexedCacheEntriesMetaData[CacheIndex].MountRefCount;
-	if (IndexedCacheEntriesMetaData[CacheIndex].MountRefCount == 0)
-	{
-		UE_LOG(LogIndexedCacheStorageManager, Display, TEXT("UnmountCacheStorage cache %d at %s"), CacheIndex, *IndexedCacheEntriesMetaData[CacheIndex].MountName);
-		IIndexedCacheStorage::Get().UnmountCacheStorage(*IndexedCacheEntriesMetaData[CacheIndex].MountName);
-	}
-	else
-	{
-		UE_LOG(LogIndexedCacheStorageManager, Display, TEXT("UnmountCacheStorage cache %d still mounted %d times"), CacheIndex, IndexedCacheEntriesMetaData[CacheIndex].MountRefCount);
-	}
+	return CacheStorageMetaData->UnmountCacheStorage();
 }
 
-
 FString FIndexedCacheStorageManager::GetMountName(int32 CacheIndex)
 {
-	FRWScopeLock _(IndexedCacheEntriesMetaDataLock, SLT_ReadOnly);
-	if (!IndexedCacheEntriesMetaData.IsValidIndex(CacheIndex))
+	const FCacheStorageMetaData* CacheStorageMetaData = GetCacheStorageMetaData(CacheIndex);
+	if (CacheStorageMetaData == nullptr)
 	{
+		UE_LOG(LogIndexedCacheStorageManager, Warning, TEXT("GetMountName skip invalid cache %d"), CacheIndex);
 		return FString();
 	}
 
-	return IndexedCacheEntriesMetaData[CacheIndex].MountName;
-
+	return CacheStorageMetaData->GetMountName();
 }
 
 FString FIndexedCacheStorageManager::GetMountPath(int32 CacheIndex)
 {
-	FString MountName = GetMountName(CacheIndex);
-	FString BundleMountPath; 
-	IIndexedCacheStorage::Get().GetCacheStorageMountPath(BundleMountPath, MountName);
-	return BundleMountPath;
+	const FCacheStorageMetaData* CacheStorageMetaData = GetCacheStorageMetaData(CacheIndex);
+	if (CacheStorageMetaData == nullptr)
+	{
+		UE_LOG(LogIndexedCacheStorageManager, Warning, TEXT("GetMountPath skip invalid cache %d"), CacheIndex);
+		return FString();
+	}
+
+	return CacheStorageMetaData->GetMountPath();
+}
+
+FCacheStorageMetaData* FIndexedCacheStorageManager::GetCacheStorageMetaData(int32 CacheIndex)
+{
+	// We must skip index 0 since it's mapped to GamePersistentDownloadDir()
+	bool bValid = CacheIndex > 0 && IndexedCacheStorageMetaData && CacheIndex < IndexedCacheStorageMetaData->IndexedCacheEntriesMetaData.Num();
+	if (!bValid)
+	{
+		return nullptr;
+	}
+
+	FCacheStorageMetaData* OutCacheStorageMetaData = IndexedCacheStorageMetaData->IndexedCacheEntriesMetaData[CacheIndex].Get();
+	if (OutCacheStorageMetaData && OutCacheStorageMetaData->IsValid())
+	{
+		return OutCacheStorageMetaData;
+	}
+
+	return nullptr;
 }
 
 }
diff --git a/Engine/Source/Runtime/Experimental/IndexedCacheStorage/Public/IndexedCacheStorageManager.h b/Engine/Source/Runtime/Experimental/IndexedCacheStorage/Public/IndexedCacheStorageManager.h
--- a/Engine/Source/Runtime/Experimental/IndexedCacheStorage/Public/IndexedCacheStorageManager.h
+++ b/Engine/Source/Runtime/Experimental/IndexedCacheStorage/Public/IndexedCacheStorageManager.h
@@ -22,6 +22,8 @@ public:
 	virtual void ShutdownModule() override;
 };
 
+class FCacheStorageMetaData;
+
 // Manager to handle ref-counting of ICS entries
 class FIndexedCacheStorageManager
 {
@@ -65,16 +67,9 @@ private:
 	friend class FIndexedCacheStorageModule;
 
 	void Initialize();
+	FCacheStorageMetaData* GetCacheStorageMetaData(int32 CacheIndex);
 
-	FRWLock IndexedCacheEntriesMetaDataLock;
-
-	struct FCacheStorageMetaData
-	{
-		FString MountName;
-		int32 MountRefCount = 0;
-		bool bCacheExists = false;
-	};
-	TArray<FCacheStorageMetaData> IndexedCacheEntriesMetaData;
+	TUniquePtr<struct FIndexedCacheStorageMetaData> IndexedCacheStorageMetaData;
 	TMap<FString, int32> StorageNameToStorageIndex;
 	bool bHasRegisteredEntries = false;
 };
