Modular Vehicle : CIS Fix, typo introduced CL 43421136

diff --git a/Engine/Source/Runtime/Experimental/ChaosVehicles/ChaosVehiclesCore/Private/SimModule/ModuleInput.cpp b/Engine/Source/Runtime/Experimental/ChaosVehicles/ChaosVehiclesCore/Private/SimModule/ModuleInput.cpp
--- a/Engine/Source/Runtime/Experimental/ChaosVehicles/ChaosVehiclesCore/Private/SimModule/ModuleInput.cpp
+++ b/Engine/Source/Runtime/Experimental/ChaosVehicles/ChaosVehiclesCore/Private/SimModule/ModuleInput.cpp
@@ -379,7 +379,7 @@ void FModuleInputValue::DeltaNetSerialize(FArchive& Ar, UPackageMap* Map, bool&
 void FModuleInputValue::Set(const FModuleInputValue& In)
 {
 	ValueInt = In.ValueInt;
-	Value += In.Value;
+	Value = In.Value;
 }
 
 void FModuleInputValue::Lerp(const FModuleInputValue& Min, const FModuleInputValue& Max, float Alpha)
