Iris
- Crashfix
- Fixed for servercrash when we tried to Serialize invalid FString
- Early out when trying to write BitStreams with bitcount 0

diff --git a/Engine/Source/Runtime/Experimental/Iris/Core/Private/Iris/Serialization/NetBitStreamWriter.cpp b/Engine/Source/Runtime/Experimental/Iris/Core/Private/Iris/Serialization/NetBitStreamWriter.cpp
--- a/Engine/Source/Runtime/Experimental/Iris/Core/Private/Iris/Serialization/NetBitStreamWriter.cpp
+++ b/Engine/Source/Runtime/Experimental/Iris/Core/Private/Iris/Serialization/NetBitStreamWriter.cpp
@@ -115,7 +115,7 @@ void FNetBitStreamWriter::WriteBitStream(const uint32* InSrc, uint32 SrcBitOffse
 {
 	UE_NETBITSTREAMWRITER_CHECK((!bIsInvalid) & (!bHasSubstream));
 
-	if (OverflowBitCount != 0)
+	if ((OverflowBitCount != 0) || (BitCount == 0))
 	{
 		return;
 	}
diff --git a/Engine/Source/Runtime/Experimental/Iris/Core/Private/Iris/Serialization/StringNetSerializerUtils.cpp b/Engine/Source/Runtime/Experimental/Iris/Core/Private/Iris/Serialization/StringNetSerializerUtils.cpp
--- a/Engine/Source/Runtime/Experimental/Iris/Core/Private/Iris/Serialization/StringNetSerializerUtils.cpp
+++ b/Engine/Source/Runtime/Experimental/Iris/Core/Private/Iris/Serialization/StringNetSerializerUtils.cpp
@@ -77,8 +77,8 @@ void FStringNetSerializerBase::Quantize(FNetSerializationContext& Context, const
 {
 	QuantizedType& Target = *reinterpret_cast<QuantizedType*>(Args.Target);
 
-	Target.bIsEncoded = !FCString::IsPureAnsi(*Source);
-	if (Target.bIsEncoded)
+	const bool bIsEncoded = !FCString::IsPureAnsi(*Source, Source.Len());
+	if (bIsEncoded)
 	{
 		const uint32 StringLength = Source.Len() + 1U;
 		constexpr uint32 MaxStringLength = 65535U/3U;
@@ -88,6 +88,7 @@ void FStringNetSerializerBase::Quantize(FNetSerializationContext& Context, const
 			return;
 		}
 		FStringNetSerializerUtils::AdjustArraySize<QuantizedType, uint8>(Context, Target, static_cast<uint16>(3U*StringLength));
+		Target.bIsEncoded = true;
 
 		uint32 OutDestLen = 0;
 		uint8* EncodingBuffer = static_cast<uint8*>(Target.ElementStorage);
@@ -103,6 +104,8 @@ void FStringNetSerializerBase::Quantize(FNetSerializationContext& Context, const
 			return;
 		}
 		FStringNetSerializerUtils::AdjustArraySize<QuantizedType, uint8>(Context, Target, static_cast<uint16>(StringLength > 0 ? StringLength + 1 : 0));
+		Target.bIsEncoded = false;
+
 		if (StringLength > 0)
 		{
 			const TCHAR* SourceString = GetData(Source);
