Chaos - fix order of initialization error in FPBDRigidsSolver

diff --git a/Engine/Source/Runtime/Experimental/Chaos/Private/PBDRigidsSolver.cpp b/Engine/Source/Runtime/Experimental/Chaos/Private/PBDRigidsSolver.cpp
--- a/Engine/Source/Runtime/Experimental/Chaos/Private/PBDRigidsSolver.cpp
+++ b/Engine/Source/Runtime/Experimental/Chaos/Private/PBDRigidsSolver.cpp
@@ -657,6 +657,11 @@ namespace Chaos
 		, bIsFloorAnalytic(false)
 		, FloorHeight(0.f)
 		, bIsDeterministic(false)
+		, QueryMaterials_External()
+		, QueryMaterialMasks_External()
+		, SimMaterials()
+		, SimMaterialMasks()
+		, UniqueIndices()
 		, Particles(UniqueIndices)
 		, MEvolution(new FPBDRigidsEvolution(Particles, SimMaterials, &MidPhaseModifiers, &CCDModifiers, &StrainModifiers, &ContactModifiers, BufferingModeIn == EMultiBufferMode::Single))
 		, MEventManager(new FEventManager(BufferingModeIn))
diff --git a/Engine/Source/Runtime/Experimental/Chaos/Public/PBDRigidsSolver.h b/Engine/Source/Runtime/Experimental/Chaos/Public/PBDRigidsSolver.h
--- a/Engine/Source/Runtime/Experimental/Chaos/Public/PBDRigidsSolver.h
+++ b/Engine/Source/Runtime/Experimental/Chaos/Public/PBDRigidsSolver.h
@@ -391,6 +391,18 @@ namespace Chaos
 		FReal FloorHeight;
 		bool bIsDeterministic;
 
+		// Physics material mirrors for the solver. These should generally stay in sync with the global material list from
+		// the game thread. This data is read only in the solver as we should never need to update it here. External threads can
+		// Enqueue commands to change parameters.
+		//
+		// There are two copies here to enable SQ to lock only the solvers that it needs to handle the material access during a query
+		// instead of having to lock the entire physics state of the runtime.
+		// NOTE: Must be before MEvolution because they are passed to it in the constructor
+		THandleArray<FChaosPhysicsMaterial> QueryMaterials_External;
+		THandleArray<FChaosPhysicsMaterialMask> QueryMaterialMasks_External;
+		THandleArray<FChaosPhysicsMaterial> SimMaterials;
+		THandleArray<FChaosPhysicsMaterialMask> SimMaterialMasks;
+
 		FParticleUniqueIndicesMultithreaded UniqueIndices;
 		FParticlesType Particles;
 		TUniquePtr<FPBDRigidsEvolution> MEvolution;
@@ -413,19 +425,6 @@ namespace Chaos
 
 		FSolverSerializer Serializer;
 
-
-		// Physics material mirrors for the solver. These should generally stay in sync with the global material list from
-		// the game thread. This data is read only in the solver as we should never need to update it here. External threads can
-		// Enqueue commands to change parameters.
-		//
-		// There are two copies here to enable SQ to lock only the solvers that it needs to handle the material access during a query
-		// instead of having to lock the entire physics state of the runtime.
-		
-		THandleArray<FChaosPhysicsMaterial> QueryMaterials_External;
-		THandleArray<FChaosPhysicsMaterialMask> QueryMaterialMasks_External;
-		THandleArray<FChaosPhysicsMaterial> SimMaterials;
-		THandleArray<FChaosPhysicsMaterialMask> SimMaterialMasks;
-
 		struct FPendingDestroyInfo
 		{
 			FSingleParticlePhysicsProxy* Proxy;
