Iris - allow LastResortPropertyNetSerializer quantized states larger than 65535 bits, and make the max size overridable per-struct/config.

diff --git a/Engine/Source/Runtime/Experimental/Iris/Core/Private/Iris/ReplicationState/PropertyNetSerializerInfoRegistry.cpp b/Engine/Source/Runtime/Experimental/Iris/Core/Private/Iris/ReplicationState/PropertyNetSerializerInfoRegistry.cpp
--- a/Engine/Source/Runtime/Experimental/Iris/Core/Private/Iris/ReplicationState/PropertyNetSerializerInfoRegistry.cpp
+++ b/Engine/Source/Runtime/Experimental/Iris/Core/Private/Iris/ReplicationState/PropertyNetSerializerInfoRegistry.cpp
@@ -85,6 +85,18 @@ struct FDefaultArrayPropertyNetSerializerInfo : public FPropertyNetSerializerInf
 UE_NET_IMPLEMENT_NETSERIALIZER_INFO(FDefaultArrayPropertyNetSerializerInfo);
 
 /** Last resort for properties that have no specialized serializer. An example is TextProperty. */
+FLastResortPropertyNetSerializerInfo::FLastResortPropertyNetSerializerInfo()
+	: FPropertyNetSerializerInfo()
+	, MaxQuantizedSizeBits(FLastResortPropertyNetSerializerConfig::DefaultMaxQuantizedSizeBits)
+{
+}
+
+FLastResortPropertyNetSerializerInfo::FLastResortPropertyNetSerializerInfo(const uint32 InMaxQuantizedSizeBits)
+	: FPropertyNetSerializerInfo()
+	, MaxQuantizedSizeBits(InMaxQuantizedSizeBits)
+{
+}
+
 const FFieldClass* FLastResortPropertyNetSerializerInfo::GetPropertyTypeClass() const
 {
 	return FProperty::StaticClass();
@@ -110,6 +122,7 @@ FNetSerializerConfig* FLastResortPropertyNetSerializerInfo::BuildNetSerializerCo
 	FLastResortPropertyNetSerializerConfig* Config = new (NetSerializerConfigBuffer) FLastResortPropertyNetSerializerConfig();
 	InitLastResortPropertyNetSerializerConfigFromProperty(*Config, Property);
 	Config->bExcludeFromDefaultStateHash = bExcludeFromDefaultStateHash;
+	Config->MaxQuantizedSizeBits = MaxQuantizedSizeBits;
 
 	return Config;
 }
diff --git a/Engine/Source/Runtime/Experimental/Iris/Core/Private/Iris/Serialization/InternalNetSerializers.h b/Engine/Source/Runtime/Experimental/Iris/Core/Private/Iris/Serialization/InternalNetSerializers.h
--- a/Engine/Source/Runtime/Experimental/Iris/Core/Private/Iris/Serialization/InternalNetSerializers.h
+++ b/Engine/Source/Runtime/Experimental/Iris/Core/Private/Iris/Serialization/InternalNetSerializers.h
@@ -71,11 +71,16 @@ struct FLastResortPropertyNetSerializerConfig : public FNetSerializerConfig
 	GENERATED_BODY()
 
 public:
+	static constexpr uint32 DefaultMaxQuantizedSizeBits = 65535;
+
 	UPROPERTY()
 	TFieldPath<FProperty> Property;
 
 	UPROPERTY()
 	bool bExcludeFromDefaultStateHash = false;
+
+	UPROPERTY()
+	uint32 MaxQuantizedSizeBits = DefaultMaxQuantizedSizeBits;
 };
 
 // ENetRole. With role swapping at deserialization
diff --git a/Engine/Source/Runtime/Experimental/Iris/Core/Private/Iris/Serialization/LastResortPropertyNetSerializer.cpp b/Engine/Source/Runtime/Experimental/Iris/Core/Private/Iris/Serialization/LastResortPropertyNetSerializer.cpp
--- a/Engine/Source/Runtime/Experimental/Iris/Core/Private/Iris/Serialization/LastResortPropertyNetSerializer.cpp
+++ b/Engine/Source/Runtime/Experimental/Iris/Core/Private/Iris/Serialization/LastResortPropertyNetSerializer.cpp
@@ -23,9 +23,9 @@ struct FFLastResortPropertyNetSerializerQuantizedType
 	FIrisPackageMapExportsQuantizedType QuantizedExports;
 
 	// How many bytes the current allocation can hold.
-	uint16 ByteCapacity = 0;
+	uint32 ByteCapacity = 0;
 	// How many bits are valid
-	uint16 BitCount = 0;
+	uint32 BitCount = 0;
 	void* Storage = nullptr;
 };
 
@@ -70,9 +70,9 @@ private:
 	static constexpr uint32 AllocationAlignment = 4U;
 
 	static void FreeDynamicStateInternal(FNetSerializationContext&, QuantizedType& Value);
-	static void GrowDynamicStateInternal(FNetSerializationContext&, QuantizedType& Value, uint16 NewBitCount);
-	static void ShrinkDynamicStateInternal(FNetSerializationContext&, QuantizedType& Value, uint16 NewBitCount);
-	static void AdjustStorageSize(FNetSerializationContext&, QuantizedType& Value, uint16 NewBitCount);
+	static void GrowDynamicStateInternal(FNetSerializationContext&, QuantizedType& Value, uint32 NewBitCount);
+	static void ShrinkDynamicStateInternal(FNetSerializationContext&, QuantizedType& Value, uint32 NewBitCount);
+	static void AdjustStorageSize(FNetSerializationContext&, QuantizedType& Value, uint32 NewBitCount);
 };
 UE_NET_IMPLEMENT_SERIALIZER_INTERNAL(FLastResortPropertyNetSerializer);
 
@@ -131,13 +131,15 @@ void FLastResortPropertyNetSerializer::Deserialize(FNetSerializationContext& Con
 
 	// Read the data
 	const uint32 NewBitCount = ReadPackedUint32(Reader);
-	if (NewBitCount > 65535U)
+	if (!ensureMsgf(NewBitCount <= Config->MaxQuantizedSizeBits,
+		TEXT("FLastResortPropertyNetSerializer::Deserialize data size of %u bits exceeds maximum of %u bits for property %s."),
+		NewBitCount, Config->MaxQuantizedSizeBits, *Config->Property.Get()->GetName()))
 	{
 		Context.SetError(GNetError_ArraySizeTooLarge);
 		return;
 	}
 
-	AdjustStorageSize(Context, Value, static_cast<uint16>(NewBitCount));
+	AdjustStorageSize(Context, Value, NewBitCount);
 
 	Reader->ReadBitStream(static_cast<uint32*>(Value.Storage), NewBitCount);
 }
@@ -167,7 +169,16 @@ void FLastResortPropertyNetSerializer::Quantize(FNetSerializationContext& Contex
 	Property->NetSerializeItem(Archive, PackageMap, reinterpret_cast<void*>(Args.Source));
 
 	const uint64 BitCount = Archive.GetNumBits();
-	if (BitCount > 65535U)
+
+	if (!ensureMsgf(BitCount <= Config->MaxQuantizedSizeBits,
+		TEXT("FLastResortPropertyNetSerializer::Quantize: data size of %llu bits exceeds maximum of %u bits for property %s."),
+		BitCount, Config->MaxQuantizedSizeBits, *Config->Property.Get()->GetName()))
+	{
+		Context.SetError(GNetError_ArraySizeTooLarge);
+		return;
+	}
+
+	if (!ensureMsgf(!Archive.IsError(), TEXT("FLastResortPropertyNetSerializer::Quantize: NetBitWriter archive error in NetSerializeItem for property %s."), *Config->Property.Get()->GetName()))
 	{
 		Context.SetError(GNetError_ArraySizeTooLarge);
 		return;
@@ -177,7 +188,7 @@ void FLastResortPropertyNetSerializer::Quantize(FNetSerializationContext& Contex
 	FIrisPackageMapExportsUtil::Quantize(Context, PackageMapExports, NetTokensPendingExport, Value.QuantizedExports);
 
 	// Deal with serialized data
-	AdjustStorageSize(Context, Value, static_cast<uint16>(BitCount));
+	AdjustStorageSize(Context, Value, static_cast<uint32>(BitCount));
 	if (BitCount > 0)
 	{
 		FMemory::Memcpy(Value.Storage, Archive.GetData(), (BitCount + 7U)/8U);
@@ -289,11 +300,11 @@ void FLastResortPropertyNetSerializer::FreeDynamicStateInternal(FNetSerializatio
 	Value.Storage = 0;
 }
 
-void FLastResortPropertyNetSerializer::GrowDynamicStateInternal(FNetSerializationContext& Context, QuantizedType& Value, uint16 NewBitCount)
+void FLastResortPropertyNetSerializer::GrowDynamicStateInternal(FNetSerializationContext& Context, QuantizedType& Value, uint32 NewBitCount)
 {
 	checkSlow(NewBitCount > Value.BitCount);
 
-	const uint16 ByteCount = static_cast<uint16>(Align((NewBitCount + 7U)/8U, AllocationAlignment));
+	const uint32 ByteCount = Align((NewBitCount + 7U)/8U, AllocationAlignment);
 
 	// We don't support delta compression for the unknown contents of the bits so we don't need to copy the old data.
 	Context.GetInternalContext()->Free(Value.Storage);
@@ -309,7 +320,7 @@ void FLastResortPropertyNetSerializer::GrowDynamicStateInternal(FNetSerializatio
 	Value.Storage = Storage;
 }
 
-void FLastResortPropertyNetSerializer::AdjustStorageSize(FNetSerializationContext& Context, QuantizedType& Value, uint16 NewBitCount)
+void FLastResortPropertyNetSerializer::AdjustStorageSize(FNetSerializationContext& Context, QuantizedType& Value, uint32 NewBitCount)
 {
 	const uint32 NewByteCapacity = Align((NewBitCount + 7U)/8U, AllocationAlignment);
 	if (NewByteCapacity == 0)
diff --git a/Engine/Source/Runtime/Experimental/Iris/Core/Public/Iris/ReplicationState/PropertyNetSerializerInfoRegistry.h b/Engine/Source/Runtime/Experimental/Iris/Core/Public/Iris/ReplicationState/PropertyNetSerializerInfoRegistry.h
--- a/Engine/Source/Runtime/Experimental/Iris/Core/Public/Iris/ReplicationState/PropertyNetSerializerInfoRegistry.h
+++ b/Engine/Source/Runtime/Experimental/Iris/Core/Public/Iris/ReplicationState/PropertyNetSerializerInfoRegistry.h
@@ -156,6 +156,9 @@ private:
 
 struct FLastResortPropertyNetSerializerInfo : public FPropertyNetSerializerInfo
 {
+	IRISCORE_API FLastResortPropertyNetSerializerInfo();
+	IRISCORE_API explicit FLastResortPropertyNetSerializerInfo(const uint32 InMaxQuantizedSizeBits);
+
 	IRISCORE_API virtual const FFieldClass* GetPropertyTypeClass() const override;
 	IRISCORE_API virtual bool IsSupported(const FProperty* Property) const override;
 	IRISCORE_API virtual const FNetSerializer* GetNetSerializer(const FProperty* Property) const override final;
@@ -163,6 +166,9 @@ struct FLastResortPropertyNetSerializerInfo : public FPropertyNetSerializerInfo
 	IRISCORE_API virtual FNetSerializerConfig* BuildNetSerializerConfig(void* NetSerializerConfigBuffer, const FProperty* Property) const override final;
 
 protected:
+	// Maximum size of quantized state in bits
+	uint32 MaxQuantizedSizeBits;
+
 	// If set this class should not be included in the default state hash.
 	bool bExcludeFromDefaultStateHash = false;
 };
@@ -175,6 +181,12 @@ struct FNamedStructLastResortPropertyNetSerializerInfo final : public FLastResor
 		StructName = InPropertyFName;
 	}
 
+	FNamedStructLastResortPropertyNetSerializerInfo(const FName InPropertyFName, const uint32 InMaxBits)
+	: FLastResortPropertyNetSerializerInfo(InMaxBits)
+	{
+		StructName = InPropertyFName;
+	}
+
 	IRISCORE_API virtual const FFieldClass* GetPropertyTypeClass() const override;
 	IRISCORE_API virtual bool IsSupported(const FProperty* Property) const override;
 };
@@ -224,6 +236,10 @@ const UE::Net::FPropertyNetSerializerInfo& GetPropertyNetSerializerInfo_##Name()
 #define UE_NET_IMPLEMENT_NAMED_STRUCT_LASTRESORT_NETSERIALIZER_INFO(StructName) \
 const UE::Net::FPropertyNetSerializerInfo& GetPropertyNetSerializerInfo_##StructName() { static UE::Net::FNamedStructLastResortPropertyNetSerializerInfo StaticInstance(StructName); return StaticInstance; };
 
+// Implement FPropertyNetSerializerInfo for cases where LastResortPropertyNetSerializer is needed for struct with custom serialization, with quantized state size override.
+#define UE_NET_IMPLEMENT_NAMED_STRUCT_LASTRESORT_NETSERIALIZER_INFO_WITH_SIZE_OVERRIDE(StructName, MaxBits) \
+const UE::Net::FPropertyNetSerializerInfo& GetPropertyNetSerializerInfo_##StructName() { static UE::Net::FNamedStructLastResortPropertyNetSerializerInfo StaticInstance(StructName, MaxBits); return StaticInstance; };
+
 // Register 
 #define UE_NET_REGISTER_NETSERIALIZER_INFO(Name) \
 UE::Net::FPropertyNetSerializerInfoRegistry::Register(&GetPropertyNetSerializerInfo_##Name());
@@ -259,3 +275,9 @@ static F##Name##NetSerializerRegistryDelegates Name##NetSerializerRegistryDelega
 	UE_NET_IMPLEMENT_NAMED_STRUCT_LASTRESORT_NETSERIALIZER_INFO(PropertyNetSerializerRegistry_NAME_##Name); \
 	UE_NET_IMPLEMENT_NETSERIALIZER_REGISTRY_DELEGATES(Name)
 
+// Utility that can be used to forward serialization of a Struct to a last resort net serializer with a maximum quantized size
+#define UE_NET_IMPLEMENT_NAMED_STRUCT_LASTRESORT_NETSERIALIZER_AND_REGISTRY_DELEGATES_WITH_SIZE_OVERRIDE(Name, MaxBits) \
+	UE_NET_IS_DECLARED_TYPE(Name); \
+	static const FName PropertyNetSerializerRegistry_NAME_##Name( PREPROCESSOR_TO_STRING(Name) ); \
+	UE_NET_IMPLEMENT_NAMED_STRUCT_LASTRESORT_NETSERIALIZER_INFO_WITH_SIZE_OVERRIDE(PropertyNetSerializerRegistry_NAME_##Name, MaxBits); \
+	UE_NET_IMPLEMENT_NETSERIALIZER_REGISTRY_DELEGATES(Name)
