Allow zero size being passed to FNetBitStreamWriter's InitBytes method.

diff --git a/Engine/Source/Runtime/Experimental/Iris/Core/Private/Iris/Serialization/NetBitStreamWriter.cpp b/Engine/Source/Runtime/Experimental/Iris/Core/Private/Iris/Serialization/NetBitStreamWriter.cpp
--- a/Engine/Source/Runtime/Experimental/Iris/Core/Private/Iris/Serialization/NetBitStreamWriter.cpp
+++ b/Engine/Source/Runtime/Experimental/Iris/Core/Private/Iris/Serialization/NetBitStreamWriter.cpp
@@ -42,7 +42,7 @@ void FNetBitStreamWriter::InitBytes(void* InBuffer, uint32 ByteCount)
 {
 	UE_NETBITSTREAMWRITER_CHECK(InBuffer != nullptr);
 	UE_NETBITSTREAMWRITER_CHECKF((uintptr_t(InBuffer) & 3) == 0, TEXT("Buffer needs to be 4-byte aligned."));
-	UE_NETBITSTREAMWRITER_CHECKF((ByteCount >= 4) && (ByteCount & 3) == 0, TEXT("Buffer capacity needs to be a multiple of 4 and at least 4."));
+	UE_NETBITSTREAMWRITER_CHECKF((ByteCount & 3) == 0, TEXT("Buffer capacity needs to be a multiple of 4."));
 	UE_NETBITSTREAMWRITER_CHECK(!bHasSubstream && !bIsSubstream);
 
 	Buffer = static_cast<uint32*>(InBuffer);
