Added support for sending headers with the multi endpoint HTTP client.

#rnx

diff --git a/Engine/Source/Runtime/Experimental/IoStore/OnDemand/Private/OnDemandHttpClient.cpp b/Engine/Source/Runtime/Experimental/IoStore/OnDemand/Private/OnDemandHttpClient.cpp
--- a/Engine/Source/Runtime/Experimental/IoStore/OnDemand/Private/OnDemandHttpClient.cpp
+++ b/Engine/Source/Runtime/Experimental/IoStore/OnDemand/Private/OnDemandHttpClient.cpp
@@ -200,21 +200,27 @@ TIoStatusOr<FMultiEndpointHttpClientResponse> FMultiEndpointHttpClient::Get(FAns
 	return FIoStatus(EIoErrorCode::ReadError, Reason.ToView());
 }
 
-void FMultiEndpointHttpClient::Get(const FOnDemandHostGroup& Endpoint, FAnsiStringView RelativeUrl, FOnHttpResponse&& OnResponse)
+void FMultiEndpointHttpClient::Get(const FOnDemandHostGroup& HostGroup, FAnsiStringView RelativeUrl, FOnHttpResponse&& OnResponse)
 {
 	FIoOffsetAndLength DefaultChunkRange(0,0);
-	Get(Endpoint, RelativeUrl, DefaultChunkRange, MoveTemp(OnResponse));
+	Get(HostGroup, RelativeUrl, DefaultChunkRange, MoveTemp(OnResponse));
 }
 
-void FMultiEndpointHttpClient::Get(const FOnDemandHostGroup& Endpoint, FAnsiStringView RelativeUrl, const FIoOffsetAndLength& ChunkRange, FOnHttpResponse&& OnResponse)
+void FMultiEndpointHttpClient::Get(const FOnDemandHostGroup& HostGroup, FAnsiStringView RelativeUrl, const FIoOffsetAndLength& ChunkRange, FOnHttpResponse&& OnResponse)
+{
+	Get(HostGroup, RelativeUrl, ChunkRange, TArray<FAnsiString>(), MoveTemp(OnResponse));
+}
+
+void FMultiEndpointHttpClient::Get(const FOnDemandHostGroup& HostGroup, FAnsiStringView RelativeUrl, const FIoOffsetAndLength& ChunkRange, TArray<FAnsiString>&& Headers, FOnHttpResponse&& OnResponse)
 {
 	UE_CHECK_THREADSAFETY(FMultiEndpointHttpClient::Get);
 
-	FConnection& Connection = GetConnection(Endpoint);
+	FConnection& Connection = GetConnection(HostGroup);
 
 	IssueRequest(FRequest
 	{
 		.OnResponse		= MoveTemp(OnResponse),
+		.Headers		= MoveTemp(Headers),
 		.RelativeUrl	= FAnsiString(RelativeUrl),
 		.Range			= ChunkRange,
 		.Connection		= Connection,
@@ -293,6 +299,34 @@ void FMultiEndpointHttpClient::IssueRequest(FRequest&& Request)
 	FAnsiStringView Url				= Request.RelativeUrl;
 	FConnectionPool& ConnectionPool = *Request.Connection.Pools[Request.Connection.CurrentHost];
 
+	FEventLoop::FRequestParams RequestParams = FEventLoop::FRequestParams
+	{
+		.ContentSizeEst = uint32(Request.Range.GetLength()),
+		.bAutoRedirect	= Config.Redirects == EHttpRedirects::Follow,
+		.bAllowChunked	= Config.bAllowChunkedTransfer
+	};
+
+	UE::IoStore::HTTP::FRequest HttpRequest = EventLoop.Get(Url, ConnectionPool, &RequestParams);
+
+#if !UE_BUILD_TEST
+	HttpRequest.Header("Pragma", "akamai-x-cache-on");
+#endif
+
+	check(Request.Headers.IsEmpty() || ((Request.Headers.Num() % 2) == 0));
+	if (!Request.Headers.IsEmpty())
+	{
+		for (int32 Idx = 0; Idx < Request.Headers.Num(); Idx += 2)
+		{
+			HttpRequest.Header(Request.Headers[Idx], Request.Headers[Idx + 1]);
+		}
+	}
+
+	if (Request.Range.GetOffset() > 0 || Request.Range.GetLength() > 0)
+	{
+		HttpRequest.Header(ANSITEXTVIEW("Range"),
+			WriteToAnsiString<64>(ANSITEXTVIEW("bytes="), Request.Range.GetOffset(), ANSITEXTVIEW("-"), Request.Range.GetOffset() + Request.Range.GetLength() - 1));
+	}
+
 	auto HttpSink = [this, Request = MoveTemp(Request)](const FTicketStatus& TicketStatus) mutable
 	{
 		switch (TicketStatus.GetId())
@@ -385,25 +419,6 @@ void FMultiEndpointHttpClient::IssueRequest(FRequest&& Request)
 		}
 	};
 
-	FEventLoop::FRequestParams RequestParams = 
-		{
-			.ContentSizeEst = uint32(Request.Range.GetLength()),
-			.bAutoRedirect = Config.Redirects == EHttpRedirects::Follow,
-			.bAllowChunked = Config.bAllowChunkedTransfer
-		};
-
-	UE::IoStore::HTTP::FRequest HttpRequest = EventLoop.Get(Url, ConnectionPool, &RequestParams);
-
-#if !UE_BUILD_TEST
-	HttpRequest.Header("Pragma", "akamai-x-cache-on");
-#endif
-
-	if (Request.Range.GetOffset() > 0 || Request.Range.GetLength() > 0)
-	{
-		HttpRequest.Header(ANSITEXTVIEW("Range"),
-			WriteToAnsiString<64>(ANSITEXTVIEW("bytes="), Request.Range.GetOffset(), ANSITEXTVIEW("-"), Request.Range.GetOffset() + Request.Range.GetLength() - 1));
-	}
-
 	EventLoop.Send(MoveTemp(HttpRequest), MoveTemp(HttpSink));
 }
 
diff --git a/Engine/Source/Runtime/Experimental/IoStore/OnDemand/Private/OnDemandHttpClient.h b/Engine/Source/Runtime/Experimental/IoStore/OnDemand/Private/OnDemandHttpClient.h
--- a/Engine/Source/Runtime/Experimental/IoStore/OnDemand/Private/OnDemandHttpClient.h
+++ b/Engine/Source/Runtime/Experimental/IoStore/OnDemand/Private/OnDemandHttpClient.h
@@ -108,6 +108,7 @@ public:
 
 	void Get(const FOnDemandHostGroup& HostGroup, FAnsiStringView RelativeUrl, FOnHttpResponse&& OnResponse);
 	void Get(const FOnDemandHostGroup& HostGroup, FAnsiStringView RelativeUrl, const FIoOffsetAndLength& ChunkRange, FOnHttpResponse&& OnResponse);
+	void Get(const FOnDemandHostGroup& HostGroup, FAnsiStringView RelativeUrl, const FIoOffsetAndLength& ChunkRange, TArray<FAnsiString>&& Headers, FOnHttpResponse&& OnResponse);
 
 	bool Tick(int32 WaitTimeMs, uint32 MaxKiBPerSecond);
 	bool Tick() { return Tick(-1, 0); }
@@ -127,6 +128,7 @@ private:
 	struct FRequest
 	{
 		FOnHttpResponse		OnResponse;
+		TArray<FAnsiString>	Headers;
 		FAnsiString			RelativeUrl;
 		FIoOffsetAndLength	Range;
 		FConnection&		Connection;
