IAX analytics now record the number of times that the game has been suspended so that we can filter out misleading data.

#rnx

- If there are inflight IAX http requests when the game is suspended we run the risk of the request being recorded with a very long total request time which can unfairly skew the analytics.
- Note that FCoreDelegates::ApplicationWillEnterBackgroundDelegate is not reliable on all platforms so we cannot be sure when the game is suspended but we can reliably track when the process is resumed which is good enough for our purposes.
- Since we cannot easily track which specific requests may have been inflight during a suspension we will instead mark how many times the game was resumed during the analytics period. This will allow us to both track how often this is actually happening and remove the potentially incorrect analytics from our totals if desired.
- Moved some existing delegate handles from global space to be members of FOnDemandIoBackendStats.
- The new analytics is called 'IasAppResumeCount' and will be reset each time the analytics are sent just like our other values.

diff --git a/Engine/Source/Runtime/Experimental/IoStore/OnDemand/Private/Statistics.cpp b/Engine/Source/Runtime/Experimental/IoStore/OnDemand/Private/Statistics.cpp
--- a/Engine/Source/Runtime/Experimental/IoStore/OnDemand/Private/Statistics.cpp
+++ b/Engine/Source/Runtime/Experimental/IoStore/OnDemand/Private/Statistics.cpp
@@ -292,6 +292,7 @@ FCounterAtomicInt		GCacheReadBytes(TEXT("Ias/CacheReadBytes"), TraceCounterDispl
 FCounterAtomicInt		GCacheRejectBytes(TEXT("Ias/CachePutRejectBytes"), TraceCounterDisplayHint_Memory);
 
 // http stats
+int64					GAppResumeCount = 0;
 bool					GHttpDistributedEndpointResolved = false;
 FCounterInt				GHttpConnectCount(TEXT("Ias/HttpConnectCount"), TraceCounterDisplayHint_None);
 FCounterInt				GHttpDisconnectCount(TEXT("Ias/HttpDisconnectCount"), TraceCounterDisplayHint_None);
@@ -482,8 +483,6 @@ static FInstallCacheAnalytics InstallCacheAnalytics;
 
 ////////////////////////////////////////////////////////////////////////////////
 static FOnDemandIoBackendStats* GStatistics = nullptr;
-static FDelegateHandle GStatisticsEndFrameDelegateHandle;
-static FDelegateHandle GStatisticsOnScreenDelegateHandle;
 
 FOnDemandIoBackendStats::FOnDemandIoBackendStats(FBackendStatus& InStatus)
 	: BackendStatus(InStatus)
@@ -492,14 +491,17 @@ FOnDemandIoBackendStats::FOnDemandIoBackendStats(FBackendStatus& InStatus)
 
 	check(GStatistics == nullptr);
 	GStatistics = this;
-	GStatisticsEndFrameDelegateHandle = FCoreDelegates::OnEndFrame.AddLambda([this]()
+
+	OnApplicationResumeHandle = FCoreDelegates::ApplicationHasEnteredForegroundDelegate.AddRaw(this, &FOnDemandIoBackendStats::OnApplicationResume);
+
+	EndFrameDelegateHandle = FCoreDelegates::OnEndFrame.AddLambda([this]()
 	{
 		UpdateCSVValues();
 		PrintPeriodicLogging();
 	});
 
 #if UE_ENABLE_ONSCREEN_STATISTICS
-	GStatisticsOnScreenDelegateHandle = FCoreDelegates::OnGetOnScreenMessages.AddLambda(
+	OnScreenDelegateHandle = FCoreDelegates::OnGetOnScreenMessages.AddLambda(
 		[this] (FCoreDelegates::FSeverityMessageMap& OutMessages)
 		{
 			if (!GIasDisplayOnScreenStatistics)
@@ -514,8 +516,10 @@ FOnDemandIoBackendStats::FOnDemandIoBackendStats(FBackendStatus& InStatus)
 
 FOnDemandIoBackendStats::~FOnDemandIoBackendStats()
 {
-	FCoreDelegates::OnEndFrame.Remove(GStatisticsEndFrameDelegateHandle);
-	FCoreDelegates::OnGetOnScreenMessages.Remove(GStatisticsOnScreenDelegateHandle);
+	FCoreDelegates::OnEndFrame.Remove(OnApplicationResumeHandle);
+	FCoreDelegates::OnEndFrame.Remove(EndFrameDelegateHandle);
+	FCoreDelegates::OnGetOnScreenMessages.Remove(OnScreenDelegateHandle);
+
 	GStatistics = nullptr;
 }
 
@@ -781,9 +785,12 @@ void FOnDemandIoBackendStats::ReportGeneralAnalytics(TArray<FAnalyticsEventAttri
 	// not since with the current system we are unable to specify our own analytics payload.
 
 	AppendAnalyticsEventAttributeArray(OutAnalyticsArray
+		, TEXT("IasAppResumeCount"), GAppResumeCount
 		, TEXT("IasHttpDistributedEndpointResolved"), GHttpDistributedEndpointResolved
 		, TEXT("IasHttpHasEverConnected"), GHttpConnectCount.Get() > 0 // Report if the system has ever actually managed to make a connection
 	);
+
+	GAppResumeCount = 0;
 }
 
 void FOnDemandIoBackendStats::ReportEndPointAnalytics(TArray<FAnalyticsEventAttribute>& OutAnalyticsArray) const
@@ -1204,6 +1211,11 @@ void FOnDemandIoBackendStats::OnHttpRequestCompleted()
 	--GHttpNumConcurrentRequests;
 }
 
+void FOnDemandIoBackendStats::OnApplicationResume()
+{
+	++GAppResumeCount;
+}
+
 ////////////////////////////////////////////////////////////////////////////////
 void FOnDemandContentInstallerStats::OnRequestEnqueued()
 {
diff --git a/Engine/Source/Runtime/Experimental/IoStore/OnDemand/Private/Statistics.h b/Engine/Source/Runtime/Experimental/IoStore/OnDemand/Private/Statistics.h
--- a/Engine/Source/Runtime/Experimental/IoStore/OnDemand/Private/Statistics.h
+++ b/Engine/Source/Runtime/Experimental/IoStore/OnDemand/Private/Statistics.h
@@ -89,11 +89,17 @@ private:
 
 #if IAS_WITH_STATISTICS
 
+	void OnApplicationResume();
+
 	void UpdateCSVValues();
 	void PrintPeriodicLogging();
 
 	void PrintOnScreenStatistics(FCoreDelegates::FSeverityMessageMap& OutMessages);
 
+	FDelegateHandle OnApplicationResumeHandle;
+	FDelegateHandle EndFrameDelegateHandle;
+	FDelegateHandle OnScreenDelegateHandle;
+
 	const FBackendStatus& BackendStatus;
 #endif //IAS_WITH_STATISTICS
 };
