Chaos - fix long-standing VBD friction bug
- also fix some comments

The impulse from the current velocity iteration was being double-counted when determining the friction cone limit.

#rnx

diff --git a/Engine/Source/Runtime/Experimental/Chaos/Public/Chaos/Collision/PBDCollisionSolver.h b/Engine/Source/Runtime/Experimental/Chaos/Public/Chaos/Collision/PBDCollisionSolver.h
--- a/Engine/Source/Runtime/Experimental/Chaos/Public/Chaos/Collision/PBDCollisionSolver.h
+++ b/Engine/Source/Runtime/Experimental/Chaos/Public/Chaos/Collision/PBDCollisionSolver.h
@@ -60,9 +60,9 @@ namespace Chaos
 			FSolverVec3 ContactTangentUAngular1;		// InvI1.R1xU
 			FSolverVec3 ContactTangentVAngular1;		// InvI1.R1xV
 
-			FSolverReal ContactMassNormal;				// 1/[NxR0.InvI0.R0xN + NxR1.InvI0.R1xN + InvM0 + InvM1]
-			FSolverReal ContactMassTangentU;			// 1/[UxR0.InvI0.R0xU + UxR1.InvI0.R1xU + InvM0 + InvM1]
-			FSolverReal ContactMassTangentV;			// 1/[VxR0.InvI0.R0xV + VxR1.InvI0.R1xV + InvM0 + InvM1]
+			FSolverReal ContactMassNormal;				// 1/[NxR0.InvI0.R0xN + InvM0 + NxR1.InvI1.R1xN + InvM1]
+			FSolverReal ContactMassTangentU;			// 1/[UxR0.InvI0.R0xU + InvM0 + UxR1.InvI1.R1xU + InvM1]
+			FSolverReal ContactMassTangentV;			// 1/[VxR0.InvI0.R0xV + InvM0 + VxR1.InvI1.R1xV + InvM1]
 
 			FSolverReal NetPushOutNormal;
 			FSolverReal NetPushOutTangentU;
@@ -646,7 +646,6 @@ namespace Chaos
 
 				const FSolverReal TotalImpulseNormal
 					= ManifoldPoint.NetImpulseNormal
-					+ ImpulseNormal
 					+ (ManifoldPoint.NetSoftPushOutNormal
 					+  ManifoldPoint.NetPushOutNormal) / Dt;
 
@@ -1091,8 +1090,6 @@ namespace Chaos
 				SolveVelocityAverage(Dt);
 			}
 
-			// NOTE: this dynamic friction implementation is iteration-count sensitive
-			// @todo(chaos): fix iteration count dependence of dynamic friction
 			const FSolverReal DynamicFriction = (bApplyDynamicFriction && (Dt > 0) && CVars::bChaos_PBDCollisionSolver_Velocity_FrictionEnabled) ? State.VelocityFriction : FSolverReal(0);
 
 			for (int32 PointIndex = 0; PointIndex < NumManifoldPoints(); ++PointIndex)
