Chaos - fix collision events fro probes

Partial undo of an over-zealous validity check in the collision event data gather step. Probe collisions are not in the constraint graph, so we can't use that as a reason to reject collisions. Must rely on the Enabled state only.

diff --git a/Engine/Source/Runtime/Experimental/Chaos/Private/EventDefaults.cpp b/Engine/Source/Runtime/Experimental/Chaos/Private/EventDefaults.cpp
--- a/Engine/Source/Runtime/Experimental/Chaos/Private/EventDefaults.cpp
+++ b/Engine/Source/Runtime/Experimental/Chaos/Private/EventDefaults.cpp
@@ -126,7 +126,7 @@ namespace Chaos
 								continue;
 							}
 
-							if (!ContactHandle->IsEnabled() || !ContactHandle->IsInConstraintGraph())
+							if (!ContactHandle->IsEnabled())
 							{
 								continue;
 							}
