Iris
- Always call ReplicationSystem::NetUpdate for TickFlush regardless of if we have any connections or not on server.
- Now performs mininal update.

diff --git a/Engine/Plugins/Runtime/ReplicationSystemTestPlugin/Source/Private/Tests/ReplicationSystem/Filtering/TestFiltering.cpp b/Engine/Plugins/Runtime/ReplicationSystemTestPlugin/Source/Private/Tests/ReplicationSystem/Filtering/TestFiltering.cpp
--- a/Engine/Plugins/Runtime/ReplicationSystemTestPlugin/Source/Private/Tests/ReplicationSystem/Filtering/TestFiltering.cpp
+++ b/Engine/Plugins/Runtime/ReplicationSystemTestPlugin/Source/Private/Tests/ReplicationSystem/Filtering/TestFiltering.cpp
@@ -975,6 +975,9 @@ UE_NET_TEST_FIXTURE(FTestFilteringFixture, DynamicFilterAddObjectAndRemoveObject
 
 	Server->DestroyObject(ServerObject);
 
+	// Late join a client and verify that everything works as expected.
+	FReplicationSystemTestClient* Client = CreateClient();
+
 	Server->NetUpdate();
 	// Filter needs to be cleared now.
 	{
diff --git a/Engine/Plugins/Runtime/ReplicationSystemTestPlugin/Source/Private/Tests/ReplicationSystem/ReplicationSystemServerClientTestFixture.cpp b/Engine/Plugins/Runtime/ReplicationSystemTestPlugin/Source/Private/Tests/ReplicationSystem/ReplicationSystemServerClientTestFixture.cpp
--- a/Engine/Plugins/Runtime/ReplicationSystemTestPlugin/Source/Private/Tests/ReplicationSystem/ReplicationSystemServerClientTestFixture.cpp
+++ b/Engine/Plugins/Runtime/ReplicationSystemTestPlugin/Source/Private/Tests/ReplicationSystem/ReplicationSystemServerClientTestFixture.cpp
@@ -156,6 +156,10 @@ void FReplicationSystemTestNode::Setup(bool bIsServer, const TCHAR* Name, const
 		Params.NetObjectListGrowCount = ParamsOverride->NetObjectListGrowCount > 0 ? ParamsOverride->NetObjectListGrowCount : Params.NetObjectListGrowCount;
 		Params.bUseRemoteObjectReferences = ParamsOverride->bUseRemoteObjectReferences;
 		Params.bAllowParallelTasks = ParamsOverride->bAllowParallelTasks;
+		if (ParamsOverride->bAllowMinimalUpdateIfNoConnections.IsSet())
+		{
+			Params.bAllowMinimalUpdateIfNoConnections = ParamsOverride->bAllowMinimalUpdateIfNoConnections.GetValue();
+		}
 	}
 
 	LOG_SCOPE_VERBOSITY_OVERRIDE(LogIris, ELogVerbosity::Error);
diff --git a/Engine/Plugins/Runtime/ReplicationSystemTestPlugin/Source/Private/Tests/ReplicationSystem/ReplicationSystemServerClientTestFixture.h b/Engine/Plugins/Runtime/ReplicationSystemTestPlugin/Source/Private/Tests/ReplicationSystem/ReplicationSystemServerClientTestFixture.h
--- a/Engine/Plugins/Runtime/ReplicationSystemTestPlugin/Source/Private/Tests/ReplicationSystem/ReplicationSystemServerClientTestFixture.h
+++ b/Engine/Plugins/Runtime/ReplicationSystemTestPlugin/Source/Private/Tests/ReplicationSystem/ReplicationSystemServerClientTestFixture.h
@@ -17,6 +17,7 @@
 #include "Iris/ReplicationSystem/ReplicationSystem.h"
 
 #include "Net/Core/Misc/ResizableCircularQueue.h"
+#include "Misc/Optional.h"
 
 namespace UE::Net
 {
@@ -101,6 +102,7 @@ public:
 		uint32 NetObjectListGrowCount = 0;
 		bool bUseRemoteObjectReferences = false;
 		bool bAllowParallelTasks = false;
+		TOptional<bool> bAllowMinimalUpdateIfNoConnections;
 	};
 
 	enum EDelaySetup { DelaySetup=0 };
diff --git a/Engine/Plugins/Runtime/ReplicationSystemTestPlugin/Source/Private/Tests/ReplicationSystem/ReplicationSystemTestFixture.h b/Engine/Plugins/Runtime/ReplicationSystemTestPlugin/Source/Private/Tests/ReplicationSystem/ReplicationSystemTestFixture.h
--- a/Engine/Plugins/Runtime/ReplicationSystemTestPlugin/Source/Private/Tests/ReplicationSystem/ReplicationSystemTestFixture.h
+++ b/Engine/Plugins/Runtime/ReplicationSystemTestPlugin/Source/Private/Tests/ReplicationSystem/ReplicationSystemTestFixture.h
@@ -58,6 +58,9 @@ protected:
 		Params.bAllowObjectReplication = true;
 		Params.NetTokenStore =  NetTokenStore.Get();
 
+		// The types of tests using this type of fixture expects a full update
+		Params.bAllowMinimalUpdateIfNoConnections = false;
+
 		// In a testing environment without configs the creation of the ReplicationSystem can be quite spammy
 		ELogVerbosity::Type IrisLogVerbosity = UE_GET_LOG_VERBOSITY(LogIris);
 		LogIris.SetVerbosity(ELogVerbosity::Error);
diff --git a/Engine/Plugins/Runtime/ReplicationSystemTestPlugin/Source/Private/Tests/ReplicationSystem/TestObjectReferences.cpp b/Engine/Plugins/Runtime/ReplicationSystemTestPlugin/Source/Private/Tests/ReplicationSystem/TestObjectReferences.cpp
--- a/Engine/Plugins/Runtime/ReplicationSystemTestPlugin/Source/Private/Tests/ReplicationSystem/TestObjectReferences.cpp
+++ b/Engine/Plugins/Runtime/ReplicationSystemTestPlugin/Source/Private/Tests/ReplicationSystem/TestObjectReferences.cpp
@@ -145,6 +145,9 @@ public:
 
 	virtual void SetUp() override
 	{
+		// These tests rely on full update being executed.
+		ServerParamsOverride.bAllowMinimalUpdateIfNoConnections = false;
+
 		FReplicationSystemServerClientTestFixture::SetUp();
 
 		// Create objects
diff --git a/Engine/Source/Runtime/Engine/Private/NetDriver.cpp b/Engine/Source/Runtime/Engine/Private/NetDriver.cpp
--- a/Engine/Source/Runtime/Engine/Private/NetDriver.cpp
+++ b/Engine/Source/Runtime/Engine/Private/NetDriver.cpp
@@ -1054,7 +1054,7 @@ void UNetDriver::TickFlush(float DeltaSeconds)
 	}
 	FSimpleScopeSecondsCounter ScopedTimer(GTickFlushGameDriverTimeSeconds, bEnableTimer);
 
-	if (IsServer() && ClientConnections.Num() > 0 && !bSkipServerReplicateActors)
+	if (IsServer() && !bSkipServerReplicateActors)
 	{
 		// Update all clients.
 #if WITH_SERVER_CODE
@@ -1069,6 +1069,7 @@ void UNetDriver::TickFlush(float DeltaSeconds)
 		}
 		else
 #endif // UE_WITH_IRIS
+		if (ClientConnections.Num() > 0)
 		{
 #if UE_WITH_REMOTE_OBJECT_HANDLE
 			if (bReplicateTransactionally)
@@ -1104,7 +1105,7 @@ void UNetDriver::TickFlush(float DeltaSeconds)
 	ProcessQueuedBunchesCurrentFrameMilliseconds = 0.0f;
 
 	// Poll all sockets.
-	if( ServerConnection )
+	if (ServerConnection)
 	{
 		// Queue client voice packets in the server's voice channel
 		ProcessLocalClientPackets();
@@ -1534,7 +1535,7 @@ void UNetDriver::PostTickFlush()
 	{
 		TGuardValue<bool> GuardInNetTick(bInTick, true);
 #if UE_WITH_IRIS
-		if (ReplicationSystem && (ServerConnection != nullptr || ClientConnections.Num() > 0))
+		if (ReplicationSystem && (ServerConnection != nullptr || IsServer()))
 		{
 			ReplicationSystem->PostSendUpdate();
 		}
diff --git a/Engine/Source/Runtime/Experimental/Iris/Core/Private/Iris/ReplicationSystem/ReplicationSystem.cpp b/Engine/Source/Runtime/Experimental/Iris/Core/Private/Iris/ReplicationSystem/ReplicationSystem.cpp
--- a/Engine/Source/Runtime/Experimental/Iris/Core/Private/Iris/ReplicationSystem/ReplicationSystem.cpp
+++ b/Engine/Source/Runtime/Experimental/Iris/Core/Private/Iris/ReplicationSystem/ReplicationSystem.cpp
@@ -80,6 +80,7 @@ public:
 	FName MetricNameMaxPendingDependentObjectCount;
 	FName MetricNameAvgHugeObjectSendQueue;
 	FName MetricNameMaxHugeObjectSendQueue;
+	bool bHadValidConnectionsLastUpdate = false;
 
 	explicit FReplicationSystemImpl(UReplicationSystem* InReplicationSystem, const UReplicationSystem::FReplicationSystemParams& Params)
 	: ReplicationSystem(InReplicationSystem)
@@ -93,6 +94,7 @@ public:
 			.MaxReplicationWriterObjectCount = Params.MaxReplicationWriterObjectCount,
 			.bUseRemoteObjectReferences = Params.bUseRemoteObjectReferences,
 			.bAllowParallelTasks = Params.bAllowParallelTasks,
+			.bAllowMinimalUpdateIfNoConnections = Params.bAllowMinimalUpdateIfNoConnections,
 		}))
 	{
 		MetricNameTotalRootObjects = TEXT("TotalSubObjects");
@@ -599,6 +601,23 @@ public:
 		NetBlobManager.ResetNetObjectAttachmentSendQueue();
 	}
 
+	bool ShouldUpdateForConnections()
+	{
+		if (ReplicationSystemInternal.GetInitParams().bAllowMinimalUpdateIfNoConnections)
+		{
+			// Make sure that we do a complete update after closing the last connection
+			const bool bHasValidConnections = ReplicationSystemInternal.GetConnections().GetValidConnections().IsAnyBitSet();
+			const bool bUpdateConnections = bHadValidConnectionsLastUpdate || bHasValidConnections;
+			bHadValidConnectionsLastUpdate = bHasValidConnections;
+
+			return bUpdateConnections;
+		}
+		else
+		{
+			return true;
+		}
+	}
+
 	void AddConnection(uint32 ConnectionId)
 	{
 		LLM_SCOPE_BYTAG(IrisConnection);
@@ -879,6 +898,9 @@ void UReplicationSystem::NetUpdate(float DeltaSeconds)
 	UDataStream::FUpdateParameters DataStreamUpdateParams = { .UpdateType = UDataStream::EUpdateType::PreSendUpdate };
 	Impl->UpdateDataStreams(DataStreamUpdateParams);
 
+	// If we no longer have any valid connections we can skip part of the update.
+	const bool bUpdateConnectionSpecifics = Impl->ShouldUpdateForConnections();
+
 	if (bAllowObjectReplication)
 	{
 		UE_NET_TRACE_FRAME_STATSCOUNTER(GetId(), ReplicationSystem.ReplicatedObjectCount, InternalSys.GetNetRefHandleManager().GetActiveObjectCount(), ENetTraceVerbosity::Verbose);
@@ -889,39 +911,42 @@ void UReplicationSystem::NetUpdate(float DeltaSeconds)
 		// Refresh the dirty objects we were told about.
 		Impl->UpdateDirtyObjectList();
 
-		// Update world locations. We need this to happen before both filtering and prioritization.
-		Impl->UpdateWorldLocations();
+		if (bUpdateConnectionSpecifics)
+		{
+			// Update world locations. We need this to happen before both filtering and prioritization.
+			Impl->UpdateWorldLocations();
 
-		// Update filters, reduce the top-level scoped object list and set each connection's scope.
-		Impl->UpdateFiltering();
+			// Update filters, reduce the top-level scoped object list and set each connection's scope.
+			Impl->UpdateFiltering();
 
-		// Invoke any operations we need to do before copying state data
-		Impl->CallPreSendUpdate(DeltaSeconds);
+			// Invoke any operations we need to do before copying state data
+			Impl->CallPreSendUpdate(DeltaSeconds);
 
-		// Finalize the dirty list with objects set dirty during the poll phase
-		Impl->UpdateDirtyListPostPoll();
+			// Finalize the dirty list with objects set dirty during the poll phase
+			Impl->UpdateDirtyListPostPoll();
 
-		// Update conditionals
-		Impl->UpdateConditionals();
+			// Update conditionals
+			Impl->UpdateConditionals();
 
-		// Quantize dirty state data. We need this to happen before both filtering and prioritization
-		Impl->QuantizeDirtyStateData();
+			// Quantize dirty state data. We need this to happen before both filtering and prioritization
+			Impl->QuantizeDirtyStateData();
 
-		// We must process all attachments to objects going out of scope before we update the scope
-		Impl->ProcessNetObjectAttachmentSendQueue(FNetBlobManager::EProcessMode::ProcessObjectsGoingOutOfScope);
+			// We must process all attachments to objects going out of scope before we update the scope
+			Impl->ProcessNetObjectAttachmentSendQueue(FNetBlobManager::EProcessMode::ProcessObjectsGoingOutOfScope);
 
-		// Update scope for all connections
-		Impl->UpdateObjectScopes();
+			// Update scope for all connections
+			Impl->UpdateObjectScopes();
 
-		// Propagate dirty changes to all connections
-		Impl->PropagateDirtyChanges();
+			// Propagate dirty changes to all connections
+			Impl->PropagateDirtyChanges();
+		}
 	}
 
 	// Forward attachments to the connections after scope update
 	Impl->ProcessNetObjectAttachmentSendQueue(FNetBlobManager::EProcessMode::ProcessObjectsInScope);
 	Impl->ResetNetObjectAttachmentSendQueue();
 
-	if (bAllowObjectReplication)
+	if (bAllowObjectReplication && bUpdateConnectionSpecifics)
 	{
 		// Update object priorities
 		Impl->UpdatePrioritization(ReplicatingConnections);
diff --git a/Engine/Source/Runtime/Experimental/Iris/Core/Private/Iris/ReplicationSystem/ReplicationSystemInternal.h b/Engine/Source/Runtime/Experimental/Iris/Core/Private/Iris/ReplicationSystem/ReplicationSystemInternal.h
--- a/Engine/Source/Runtime/Experimental/Iris/Core/Private/Iris/ReplicationSystem/ReplicationSystemInternal.h
+++ b/Engine/Source/Runtime/Experimental/Iris/Core/Private/Iris/ReplicationSystem/ReplicationSystemInternal.h
@@ -39,6 +39,7 @@ struct FReplicationSystemInternalInitParams
 	uint32 MaxReplicationWriterObjectCount;
 	bool bUseRemoteObjectReferences;
 	bool bAllowParallelTasks;
+	bool bAllowMinimalUpdateIfNoConnections;
 };
 
 class FReplicationSystemInternal
diff --git a/Engine/Source/Runtime/Experimental/Iris/Core/Public/Iris/ReplicationSystem/ReplicationSystem.h b/Engine/Source/Runtime/Experimental/Iris/Core/Public/Iris/ReplicationSystem/ReplicationSystem.h
--- a/Engine/Source/Runtime/Experimental/Iris/Core/Public/Iris/ReplicationSystem/ReplicationSystem.h
+++ b/Engine/Source/Runtime/Experimental/Iris/Core/Public/Iris/ReplicationSystem/ReplicationSystem.h
@@ -142,6 +142,9 @@ public:
 		*/
 		bool bAllowParallelTasks = false;
 
+		/** If true we will only do a minimal update if there are no valid connections */
+		bool bAllowMinimalUpdateIfNoConnections = true;
+
 		/** Delegate that receives every RPC executed locally. */
 		UE::Net::FForwardNetRPCCallDelegate ForwardNetRPCCallDelegate;
 
