Fix but where IAS statistics were not correctly unregistering from a delegate.

#rnx

- The delegate handle 'OnApplicationResumeHandle' was being removed from the wrong delegate (FCoreDelegates::OnEndFrame) when the statistics object is being destroyed leaving it still registered to FCoreDelegates::ApplicationHasEnteredForegroundDelegate.
- It is not likely that this delegate will be invoked during shutdown but is still an obvious error and should be fixed.

diff --git a/Engine/Source/Runtime/Experimental/IoStore/OnDemand/Private/Statistics.cpp b/Engine/Source/Runtime/Experimental/IoStore/OnDemand/Private/Statistics.cpp
--- a/Engine/Source/Runtime/Experimental/IoStore/OnDemand/Private/Statistics.cpp
+++ b/Engine/Source/Runtime/Experimental/IoStore/OnDemand/Private/Statistics.cpp
@@ -516,7 +516,7 @@ FOnDemandIoBackendStats::FOnDemandIoBackendStats(FBackendStatus& InStatus)
 
 FOnDemandIoBackendStats::~FOnDemandIoBackendStats()
 {
-	FCoreDelegates::OnEndFrame.Remove(OnApplicationResumeHandle);
+	FCoreDelegates::ApplicationHasEnteredForegroundDelegate.Remove(OnApplicationResumeHandle);
 	FCoreDelegates::OnEndFrame.Remove(EndFrameDelegateHandle);
 	FCoreDelegates::OnGetOnScreenMessages.Remove(OnScreenDelegateHandle);
 
