Iris
* Remove hardcoded dependency on Object->GetArchetype() being the object used to instantiate an object.
* Now a NetObjectFactory can return the original template of an instance and it will be used to build the protocol and default state for that object

#rnx

diff --git a/Engine/Source/Runtime/Engine/Private/Net/Iris/ReplicationSystem/NetSubObjectFactory.cpp b/Engine/Source/Runtime/Engine/Private/Net/Iris/ReplicationSystem/NetSubObjectFactory.cpp
--- a/Engine/Source/Runtime/Engine/Private/Net/Iris/ReplicationSystem/NetSubObjectFactory.cpp
+++ b/Engine/Source/Runtime/Engine/Private/Net/Iris/ReplicationSystem/NetSubObjectFactory.cpp
@@ -26,6 +26,7 @@
 
 #if UE_WITH_IRIS
 
+//$IRIS TODO: Temp cvar can be removed later
 static bool bUseArchetypeInSubObjectCreation = true;
 static FAutoConsoleVariableRef CVarUseArchetype(
 	TEXT("net.Iris.SubObjectSupportArchetype"),
@@ -60,7 +61,7 @@ TUniquePtr<UE::Net::FNetObjectCreationHeader> UNetSubObjectFactory::CreateAndFil
 	check(SubObject->NeedsLoadForClient() ); // We have no business sending this unless the client can load
 	check(SubObject->GetClass()->NeedsLoadForClient());	// We have no business sending this unless the client can load
 
-	// Set the archetype to spawn the object from
+	// Set the template to clone the object from
 	{
 		UObject* Archetype = nullptr;
 		if (bUseArchetypeInSubObjectCreation)
@@ -72,9 +73,9 @@ TUniquePtr<UE::Net::FNetObjectCreationHeader> UNetSubObjectFactory::CreateAndFil
 			Archetype = SubObject->GetClass();
 		}
 
-		SubObjectHeader->ArchetypeReference = Bridge->GetOrCreateObjectReference(Archetype);
+		SubObjectHeader->TemplateReference = Bridge->GetOrCreateObjectReference(Archetype);
 
-		ensureMsgf(SubObjectHeader->ArchetypeReference.GetRefHandle().IsStatic(), TEXT("Archetype %s for subobject %s (%s) should be a stable name otherwise the client can't use it"), 
+		ensureMsgf(SubObjectHeader->TemplateReference.GetRefHandle().IsStatic(), TEXT("Archetype %s for subobject %s (%s) should be a stable name otherwise the client can't use it"),
 			*GetNameSafe(Archetype), *Bridge->PrintObjectFromNetRefHandle(Handle), *GetNameSafe(SubObject->GetClass()->GetAuthoritativeClass()));
 	}
 
@@ -157,8 +158,7 @@ UNetSubObjectFactory::FInstantiateResult UNetSubObjectFactory::InstantiateReplic
 
 		UE_LOG(LogIris, Verbose, TEXT("UNetSubObjectFactory::InstantiateReplicatedObjectFromHeader %s: Found static or stable name SubObject using path %s"), *Context.Handle.ToString(), ToCStr(SubObject->GetPathName()));
 
-		FInstantiateResult Result { .Instance = SubObject };
-		return Result;
+		return FInstantiateResult{ .Instance = SubObject };
 	}
 
 	// For dynamic objects we have to spawn them
@@ -191,11 +191,11 @@ UNetSubObjectFactory::FInstantiateResult UNetSubObjectFactory::InstantiateReplic
 	}
 
 	// Find the archetype of the subobject
-	UObject* Archetype = Bridge->ResolveObjectReference(SubObjectHeader->ArchetypeReference, Context.ResolveContext);
-	if (!Archetype)
+	UObject* Template = Bridge->ResolveObjectReference(SubObjectHeader->TemplateReference, Context.ResolveContext);
+	if (!Template)
 	{
-		UE_LOG(LogIris, Error, TEXT("UNetSubObjectFactory::InstantiateReplicatedObjectFromHeader Unable to resolve archetype: %s, cannot instantiate dynamic subobject %s"), *Bridge->DescribeObjectReference(SubObjectHeader->ArchetypeReference, Context.ResolveContext), *Context.Handle.ToString());
-		ensure(Archetype);
+		UE_LOG(LogIris, Error, TEXT("UNetSubObjectFactory::InstantiateReplicatedObjectFromHeader Unable to resolve template: %s, cannot instantiate dynamic subobject %s"), *Bridge->DescribeObjectReference(SubObjectHeader->TemplateReference, Context.ResolveContext), *Context.Handle.ToString());
+		ensure(Template);
 		return FInstantiateResult();
 	}
 
@@ -204,23 +204,26 @@ UNetSubObjectFactory::FInstantiateResult UNetSubObjectFactory::InstantiateReplic
 	 UClass* SubObjClass = nullptr;
 
 	// Only if bUseArchetypeInSubObjectCreation is false on the server. Remove along with the cvar later
-	if (Archetype->IsA<UClass>())
+	if (Template->IsA<UClass>())
 	{
-		SubObjClass = CastChecked<UClass>(Archetype);
+		SubObjClass = CastChecked<UClass>(Template);
 		SubObj = NewObject<UObject>(OuterObject, SubObjClass);
 	}
 	else
 	{
-		SubObjClass = Archetype->GetClass();
-		SubObj = NewObject<UObject>(OuterObject, SubObjClass, NAME_None, RF_NoFlags, Archetype);
+		SubObjClass = Template->GetClass();
+		SubObj = NewObject<UObject>(OuterObject, SubObjClass, NAME_None, RF_NoFlags, Template);
 	}
-	
 
 	// Sanity check some things
 	checkf(SubObj != nullptr, TEXT("UNetSubObjectFactory::InstantiateReplicatedObjectFromHeader: Subobject is NULL after instantiating. Class: %s, Outer %s, Actor %s"), *GetNameSafe(SubObjClass), *GetNameSafe(OuterObject), *GetNameSafe(RootObject));
 	checkf(!OuterObject || SubObj->IsIn(OuterObject), TEXT("UNetSubObjectFactory::InstantiateReplicatedObjectFromHeader: Subobject is not in Outer. SubObject: %s, Outer %s, Actor %s"), *SubObj->GetName(), *GetNameSafe(OuterObject), *GetNameSafe(RootObject));
 
-	FInstantiateResult Result { .Instance = SubObj };
+	FInstantiateResult Result
+	{ 
+		.Instance = SubObj,
+		.Template = Template, // Force the template to be used as the default init state source. It's possible GetArchetype() will be different then the template
+	};
 
 	// We must defer call OnSubObjectCreatedFromReplication after the state has been applied to the owning actor in order to behave like old system.
 	Result.Flags |= EReplicationBridgeCreateNetRefHandleResultFlags::ShouldCallSubObjectCreatedFromReplication;
@@ -320,13 +323,13 @@ bool FNetStaticSubObjectCreationHeader::Deserialize(const FCreationHeaderContext
 FString FNetDynamicSubObjectCreationHeader::ToString() const
 {
 	return FString::Printf(TEXT("\n\tFNetDynamicSubObjectCreationHeader (ProtocolId:0x%x):\n\t"
-								"ArchetypeReference=%s\n\t"
+								"TemplateReference=%s\n\t"
 								"OuterReference=%s\n\t"
 								"bUsePersistenLevel=%u\n\t"
 								"bOuterIsTransientLevel=%u\n\t"
 								"bOuterIsRootObject=%u\n\t"),
 								GetProtocolId(),
-								*ArchetypeReference.ToString(),
+								*TemplateReference.ToString(),
 								*OuterReference.ToString(),
 								bUsePersistentLevel,
 								bOuterIsTransientLevel,
@@ -338,7 +341,7 @@ bool FNetDynamicSubObjectCreationHeader::Serialize(const FCreationHeaderContext&
 {
 	FNetBitStreamWriter* Writer = Context.Serialization.GetBitStreamWriter();
 
-	WriteFullNetObjectReference(Context.Serialization, ArchetypeReference);
+	WriteFullNetObjectReference(Context.Serialization, TemplateReference);
 
 	if (!Writer->WriteBool(bOuterIsTransientLevel))
 	{
@@ -355,7 +358,7 @@ bool FNetDynamicSubObjectCreationHeader::Deserialize(const FCreationHeaderContex
 {
 	FNetBitStreamReader* Reader = Context.Serialization.GetBitStreamReader();
 
-	ReadFullNetObjectReference(Context.Serialization, ArchetypeReference);
+	ReadFullNetObjectReference(Context.Serialization, TemplateReference);
 
 	bOuterIsTransientLevel = Reader->ReadBool();
 	if (!bOuterIsTransientLevel)
diff --git a/Engine/Source/Runtime/Engine/Public/Net/Iris/ReplicationSystem/NetSubObjectFactory.h b/Engine/Source/Runtime/Engine/Public/Net/Iris/ReplicationSystem/NetSubObjectFactory.h
--- a/Engine/Source/Runtime/Engine/Public/Net/Iris/ReplicationSystem/NetSubObjectFactory.h
+++ b/Engine/Source/Runtime/Engine/Public/Net/Iris/ReplicationSystem/NetSubObjectFactory.h
@@ -104,7 +104,7 @@ public:
 
 	virtual FString ToString() const override;
 
-	FNetObjectReference ArchetypeReference;
+	FNetObjectReference TemplateReference;
 	FNetObjectReference OuterReference;
 	uint8 bUsePersistentLevel : 1 = false;
 	uint8 bOuterIsTransientLevel : 1 = false;	// When set the OuterReference was not sent because the Outer is the default transient level.
diff --git a/Engine/Source/Runtime/Experimental/Iris/Core/Private/Iris/Core/IrisDebugging.cpp b/Engine/Source/Runtime/Experimental/Iris/Core/Private/Iris/Core/IrisDebugging.cpp
--- a/Engine/Source/Runtime/Experimental/Iris/Core/Private/Iris/Core/IrisDebugging.cpp
+++ b/Engine/Source/Runtime/Experimental/Iris/Core/Private/Iris/Core/IrisDebugging.cpp
@@ -506,10 +506,10 @@ void NetObjectProtocolReferencesToString(FStringBuilderBase& StringBuilder, FRep
 	const FReplicationProtocolManager& ProtocolManager = ReplicationSystemInternal->GetReplicationProtocolManager();
 
 	// As there might be multiple protocols sharing the same identifier
-	ProtocolManager.ForEachProtocol(ProtocolId, [&StringBuilder, &NetRefHandleManager](const FReplicationProtocol* Protocol, const UObject* ArchetypeOrCDOUsedAsKey)
+	ProtocolManager.ForEachProtocol(ProtocolId, [&StringBuilder, &NetRefHandleManager](const FReplicationProtocol* Protocol, const FObjectKey TemplateKey)
 		{
 			StringBuilder << TEXT("Protocol: ") << ToCStr(Protocol->DebugName);
-			StringBuilder.Appendf(TEXT("Id: 0x%x Created From : 0x%p"), Protocol->ProtocolIdentifier, ArchetypeOrCDOUsedAsKey) << TEXT(" Used by : \n");
+			StringBuilder.Appendf(TEXT("Id: 0x%x Created From : 0x%p"), Protocol->ProtocolIdentifier, TemplateKey.ResolveObjectPtrEvenIfGarbage()) << TEXT(" Used by : \n");
 
 			auto FindMatchingProtocol = [&StringBuilder, &Protocol, &NetRefHandleManager](uint32 InternalObjectIndex)
 			{
diff --git a/Engine/Source/Runtime/Experimental/Iris/Core/Private/Iris/ReplicationState/ReplicationStateDescriptorBuilder.cpp b/Engine/Source/Runtime/Experimental/Iris/Core/Private/Iris/ReplicationState/ReplicationStateDescriptorBuilder.cpp
--- a/Engine/Source/Runtime/Experimental/Iris/Core/Private/Iris/ReplicationState/ReplicationStateDescriptorBuilder.cpp
+++ b/Engine/Source/Runtime/Experimental/Iris/Core/Private/Iris/ReplicationState/ReplicationStateDescriptorBuilder.cpp
@@ -2475,15 +2475,7 @@ namespace UE::Net
 {
 
 FReplicationStateDescriptorBuilder::FParameters::FParameters()
-: DescriptorRegistry(nullptr),
-  ReplicationSystem(nullptr),
-  DefaultStateSource(nullptr),
-  IncludeSuper(1U),
-  GetLifeTimeProperties(1U),
-  EnableFastArrayHandling(1U),
-  AllowFastArrayWithExtraReplicatedProperties(0U),
-  SkipCheckForCustomNetSerializerForStruct(0U),
-  SinglePropertyIndex(-1)
+	: DescriptorRegistry(nullptr)
 {
 }
 
@@ -2754,14 +2746,15 @@ SIZE_T FReplicationStateDescriptorBuilder::CreateDescriptorsForClass(FResult& Cr
 	// Just to make sure we don't mess up the parameters when building descriptor for single property
 	//check(Parameters.SinglePropertyIndex == -1 || Parameters.DescriptorRegistry == nullptr)
 
-	const UObject* ObjectClassOrArchetypeUsedAsKey = Parameters.DefaultStateSource ? Parameters.DefaultStateSource : InObjectClass;
+	// Use the template as the default key or the class CDO if there is no template.
+	const UObject* DefaultStateSource = Parameters.DefaultStateSource ? Parameters.DefaultStateSource : InObjectClass->GetDefaultObject();
 
 	// Check registry first to see if we already have created descriptors for this class
 	if (Parameters.DescriptorRegistry)
 	{
-		if (const FReplicationStateDescriptorRegistry::FDescriptors* Result = Parameters.DescriptorRegistry->Find(ObjectClassOrArchetypeUsedAsKey, InObjectClass))
+		if (const FReplicationStateDescriptorRegistry::FDescriptors* Result = Parameters.DescriptorRegistry->Find(DefaultStateSource, InObjectClass))
 		{
-			UE_LOG_DESCRIPTORBUILDER(VeryVerbose, TEXT("Found existing descriptors for Class %s : ArcheTypeOrKey: %s ptr: 0x%p"), *InObjectClass->GetName(), *ObjectClassOrArchetypeUsedAsKey->GetPathName(), ObjectClassOrArchetypeUsedAsKey);
+			UE_LOG_DESCRIPTORBUILDER(VeryVerbose, TEXT("Found existing descriptors for Class %s : DescriptorKey: %s ptr: 0x%p"), *InObjectClass->GetName(), *DefaultStateSource->GetPathName(), DefaultStateSource);
 			CreatedDescriptors.Insert(*Result, 0);
 			return CreatedDescriptors.Num();
 		}
@@ -2977,7 +2970,7 @@ SIZE_T FReplicationStateDescriptorBuilder::CreateDescriptorsForClass(FResult& Cr
 	// Build state descriptors
 	FLazyGetPathNameHelper LazyPathName;
 
-	const uint8* DefaultStateSourceData = reinterpret_cast<const uint8*>(Parameters.DefaultStateSource ? Parameters.DefaultStateSource : InObjectClass->GetDefaultObject());
+	const uint8* DefaultStateSourceData = reinterpret_cast<const uint8*>(DefaultStateSource);
 	for (uint32 BuilderTypeIndex = 0; BuilderTypeIndex < BuilderTypeCount; ++BuilderTypeIndex)
 	{
 		if (Builders[BuilderTypeIndex].HasDataToBuild())
@@ -3049,8 +3042,8 @@ SIZE_T FReplicationStateDescriptorBuilder::CreateDescriptorsForClass(FResult& Cr
 	// If we have a registry, register the created descriptors
 	if (bRegisterCreatedDescriptors && CreatedDescriptors.Num())
 	{
-		UE_LOG_DESCRIPTORBUILDER(VeryVerbose, TEXT("Registering descriptors for  Class %s : ArcheTypeOrKey: %s ptr: 0x%p"), *InObjectClass->GetName(), *ObjectClassOrArchetypeUsedAsKey->GetPathName(), ObjectClassOrArchetypeUsedAsKey);
-		Parameters.DescriptorRegistry->Register(ObjectClassOrArchetypeUsedAsKey, InObjectClass, CreatedDescriptors);
+		UE_LOG_DESCRIPTORBUILDER(VeryVerbose, TEXT("Registering descriptors for Class %s : DescriptorKey: %s ptr: 0x%p"), *InObjectClass->GetName(), *DefaultStateSource->GetPathName(), DefaultStateSource);
+		Parameters.DescriptorRegistry->Register(DefaultStateSource, InObjectClass, CreatedDescriptors);
 	}
 	
 	return CreatedDescriptors.Num();
diff --git a/Engine/Source/Runtime/Experimental/Iris/Core/Private/Iris/ReplicationState/ReplicationStateDescriptorRegistry.cpp b/Engine/Source/Runtime/Experimental/Iris/Core/Private/Iris/ReplicationState/ReplicationStateDescriptorRegistry.cpp
--- a/Engine/Source/Runtime/Experimental/Iris/Core/Private/Iris/ReplicationState/ReplicationStateDescriptorRegistry.cpp
+++ b/Engine/Source/Runtime/Experimental/Iris/Core/Private/Iris/ReplicationState/ReplicationStateDescriptorRegistry.cpp
@@ -24,18 +24,17 @@ void FReplicationStateDescriptorRegistry::Init(const FReplicationStateDescriptor
 	ProtocolManager = Params.ProtocolManager;
 }
 
-void FReplicationStateDescriptorRegistry::Register(const FFieldVariant& Object, const UObject* ObjectForPruning, const FDescriptors& Descriptors)
+void FReplicationStateDescriptorRegistry::Register(const FFieldVariant& DescriptorKey, const UObject* ObjectForPruning, const FDescriptors& Descriptors)
 {
-	const FRegisteredDescriptors* Entry = RegisteredDescriptorsMap.Find(Object);
-	
 	check(ObjectForPruning != nullptr);
 
-	if (Entry)
+	// Make sure the object isn't already registered
+	if (const FRegisteredDescriptors* Entry = RegisteredDescriptorsMap.Find(DescriptorKey))
 	{
 		// We do not want to overwrite descriptors for valid objects.
 		if (Entry->WeakPtrForPruning.Get() == ObjectForPruning)
 		{
-			if (Entry->OwnerKey == FObjectKey(GetObjectForPruning(Object)))
+			if (Entry->OwnerKey == FObjectKey(GetObjectForPruning(DescriptorKey)))
 			{
 				checkf(false, TEXT("FReplicationStateDescriptorRegistry::Trying to register descriptors for the same UObject %s"), ToCStr(ObjectForPruning->GetName()));
 				return;
@@ -43,32 +42,31 @@ void FReplicationStateDescriptorRegistry::Register(const FFieldVariant& Object,
 		}
 
 		// We found an invalid entry, invalidate it before registering new descriptors.
-		UE_LOG(LogIris, VeryVerbose, TEXT("FReplicationStateDescriptorRegistry::Register invalidate descriptors for ptr: 0x%p"), Object.GetRawPointer());
+		UE_LOG(LogIris, VeryVerbose, TEXT("FReplicationStateDescriptorRegistry::Register invalidate descriptors for ptr: 0x%p"), DescriptorKey.GetRawPointer());
+
 		// Notify protocol manager about pruned descriptors
 		InvalidateDescriptors(Entry->Descriptors);
-		RegisteredDescriptorsMap.Remove(Object);
+		RegisteredDescriptorsMap.Remove(DescriptorKey);
 	}
 
-	// We found an invalid entry, invalidate it before registering new descriptors.
 	FRegisteredDescriptors NewEntry;
-	NewEntry.OwnerKey = FObjectKey(GetObjectForPruning(Object));
+	NewEntry.OwnerKey = FObjectKey(GetObjectForPruning(DescriptorKey));
 	NewEntry.WeakPtrForPruning = TWeakObjectPtr<const UObject>(ObjectForPruning);
 	NewEntry.Descriptors = Descriptors;
-	RegisteredDescriptorsMap.Add(Object, NewEntry);
+	RegisteredDescriptorsMap.Add(DescriptorKey, NewEntry);
 }
 
-void FReplicationStateDescriptorRegistry::Register(const FFieldVariant& Object, const UObject* ObjectForPruning, const TRefCountPtr<const FReplicationStateDescriptor>& Descriptor)
+void FReplicationStateDescriptorRegistry::Register(const FFieldVariant& DescriptorKey, const UObject* ObjectForPruning, const TRefCountPtr<const FReplicationStateDescriptor>& Descriptor)
 {
-	const FRegisteredDescriptors* Entry = RegisteredDescriptorsMap.Find(Object);
-
 	check(ObjectForPruning != nullptr);
 
-	if (Entry)
+	// Make sure the descriptor isn't already registered
+	if (const FRegisteredDescriptors* Entry = RegisteredDescriptorsMap.Find(DescriptorKey))
 	{
 		// We do not want to overwrite descriptors for valid objects.
 		if (Entry->WeakPtrForPruning.Get() == ObjectForPruning)
 		{
-			if (Entry->OwnerKey == FObjectKey(GetObjectForPruning(Object)))
+			if (Entry->OwnerKey == FObjectKey(GetObjectForPruning(DescriptorKey)))
 			{
 				checkf(false, TEXT("FReplicationStateDescriptorRegistry::Trying to register descriptor for the same UObject %s"), ToCStr(ObjectForPruning->GetName()));
 				return;
@@ -76,14 +74,14 @@ void FReplicationStateDescriptorRegistry::Register(const FFieldVariant& Object,
 		}
 
 		// We found an invalid entry, invalidate it before registering new descriptors.
-		UE_LOG(LogIris, VeryVerbose, TEXT("FReplicationStateDescriptorRegistry::Register invalidate descriptor for ptr: 0x%p"), Object.GetRawPointer());
+		UE_LOG(LogIris, VeryVerbose, TEXT("FReplicationStateDescriptorRegistry::Register invalidate descriptor for ptr: 0x%p"), DescriptorKey.GetRawPointer());
 		// Notify protocol manager about pruned descriptors
 		InvalidateDescriptors(Entry->Descriptors);
-		RegisteredDescriptorsMap.Remove(Object);
+		RegisteredDescriptorsMap.Remove(DescriptorKey);
 	}
 
-	FRegisteredDescriptors& NewEntry = RegisteredDescriptorsMap.Emplace(Object);
-	NewEntry.OwnerKey = FObjectKey(GetObjectForPruning(Object));
+	FRegisteredDescriptors& NewEntry = RegisteredDescriptorsMap.Emplace(DescriptorKey);
+	NewEntry.OwnerKey = FObjectKey(GetObjectForPruning(DescriptorKey));
 	NewEntry.WeakPtrForPruning = TWeakObjectPtr<const UObject>(ObjectForPruning);
 	NewEntry.Descriptors.Add(Descriptor);
 }
diff --git a/Engine/Source/Runtime/Experimental/Iris/Core/Private/Iris/ReplicationState/ReplicationStateDescriptorRegistry.h b/Engine/Source/Runtime/Experimental/Iris/Core/Private/Iris/ReplicationState/ReplicationStateDescriptorRegistry.h
--- a/Engine/Source/Runtime/Experimental/Iris/Core/Private/Iris/ReplicationState/ReplicationStateDescriptorRegistry.h
+++ b/Engine/Source/Runtime/Experimental/Iris/Core/Private/Iris/ReplicationState/ReplicationStateDescriptorRegistry.h
@@ -44,16 +44,16 @@ public:
 	const FDescriptors* Find(const FFieldVariant& Object) const { return Find(Object, GetObjectForPruning(Object)); }
 
 	/** Register created Descriptors for ClassObject */
-	void Register(const FFieldVariant& ObjectKey, const UObject* ObjectForPruning, const FDescriptors& Descriptors);
+	void Register(const FFieldVariant& DescriptorKey, const UObject* ObjectForPruning, const FDescriptors& Descriptors);
 
 	/** Register created Descriptors for ClassObject */
-	void Register(const FFieldVariant& ObjectKey, const FDescriptors& Descriptors) { Register(ObjectKey, GetObjectForPruning(ObjectKey), Descriptors); }
+	void Register(const FFieldVariant& DescriptorKey, const FDescriptors& Descriptors) { Register(DescriptorKey, GetObjectForPruning(DescriptorKey), Descriptors); }
 
 	/** Register a single Descriptor for ClassObject */
-	void Register(const FFieldVariant& ObjectKey, const UObject* ObjectForPruning, const TRefCountPtr<const FReplicationStateDescriptor>& Descriptor);
+	void Register(const FFieldVariant& DescriptorKey, const UObject* ObjectForPruning, const TRefCountPtr<const FReplicationStateDescriptor>& Descriptor);
 
 	/** Register a single Descriptor for ClassObject */
-	void Register(const FFieldVariant& ObjectKey, const TRefCountPtr<const FReplicationStateDescriptor>& Descriptor) { Register(ObjectKey, GetObjectForPruning(ObjectKey), Descriptor); }
+	void Register(const FFieldVariant& DescriptorKey, const TRefCountPtr<const FReplicationStateDescriptor>& Descriptor) { Register(DescriptorKey, GetObjectForPruning(DescriptorKey), Descriptor); }
 
 	/** Invoked on PostGarbageCollect to prune descriptors for stale weak ptrs */
 	void PruneStaleDescriptors();
diff --git a/Engine/Source/Runtime/Experimental/Iris/Core/Private/Iris/ReplicationSystem/ObjectReplicationBridge.cpp b/Engine/Source/Runtime/Experimental/Iris/Core/Private/Iris/ReplicationSystem/ObjectReplicationBridge.cpp
--- a/Engine/Source/Runtime/Experimental/Iris/Core/Private/Iris/ReplicationSystem/ObjectReplicationBridge.cpp
+++ b/Engine/Source/Runtime/Experimental/Iris/Core/Private/Iris/ReplicationSystem/ObjectReplicationBridge.cpp
@@ -138,8 +138,10 @@ namespace UE::Net::Private
 {
 	void CallRegisterReplicationFragments(UObject* Object, FFragmentRegistrationContext& Context, EFragmentRegistrationFlags RegistrationFlags)
 	{
+		// Try the user implementation first
 		Object->RegisterReplicationFragments(Context, RegistrationFlags);
 
+		// If no user implementations exist, create fragments automatically
 		if (!Context.WasRegistered())
 		{
 			FReplicationFragmentUtil::CreateAndRegisterFragmentsForObject(Object, Context, RegistrationFlags);
@@ -358,7 +360,7 @@ UE::Net::FNetRefHandle UObjectReplicationBridge::StartReplicatingNetObject(UObje
 
 	IRIS_PROFILER_SCOPE(StartReplicatingNetObject);
 	
-	FFragmentRegistrationContext FragmentRegistrationContext(GetReplicationStateDescriptorRegistry(), GetReplicationSystem(), Traits);
+	FFragmentRegistrationContext FragmentRegistrationContext(GetReplicationStateDescriptorRegistry(), GetReplicationSystem(), Traits, Instance);
 
 	// For everything derived from UObject we can call the virtual function RegisterReplicationFragments	
 	CallRegisterReplicationFragments(Instance, FragmentRegistrationContext, EFragmentRegistrationFlags::None);
@@ -388,8 +390,15 @@ UE::Net::FNetRefHandle UObjectReplicationBridge::StartReplicatingNetObject(UObje
 	}
 #endif
 
-	// We currently identify protocols by local archetype or CDO pointer and verified the protocol id received from server and the hash of the default state
-	const UObject* ArchetypeOrCDOUsedAsKey = Instance->GetArchetype();
+	// Some fragments may be bound to a default source (e.g. template or archetype or cdo)
+	// Find the state source used for this instance and use it to bind the protocol to.
+    const UObject* TemplateKey = FFragmentRegistrationContextPrivateAccessor::GetDefaultStateSource(FragmentRegistrationContext);
+	
+	// Fallback to using the archetype or cdo here if the main object register any state source. This only happens for special unit tests
+	if (!TemplateKey)
+	{
+		TemplateKey = Instance->GetArchetype();
+	}
 
 	// Create Protocols
 	FReplicationProtocolManager* ProtocolManager = GetReplicationProtocolManager();
@@ -397,10 +406,14 @@ UE::Net::FNetRefHandle UObjectReplicationBridge::StartReplicatingNetObject(UObje
 	FReplicationInstanceProtocolPtr InstanceProtocol(ProtocolManager->CreateInstanceProtocol(RegisteredFragments, FragmentRegistrationContext.GetFragmentTraits()));
 		
 	FReplicationProtocolIdentifier ProtocolIdentifier = FReplicationProtocolManager::CalculateProtocolIdentifier(RegisteredFragments);
-	const FReplicationProtocol* ReplicationProtocol = ProtocolManager->GetReplicationProtocol(ProtocolIdentifier, ArchetypeOrCDOUsedAsKey);
+	const FReplicationProtocol* ReplicationProtocol = ProtocolManager->GetReplicationProtocol(ProtocolIdentifier, TemplateKey);
 	if (!ReplicationProtocol)
 	{
-		FCreateReplicationProtocolParameters CreateProtocolParams {.ArchetypeOrCDOUsedAsKey = ArchetypeOrCDOUsedAsKey, .TypeStatsIndex = GetTypeStatsIndex(Instance->GetClass()) };
+		const FCreateReplicationProtocolParameters CreateProtocolParams
+		{
+			.TemplateKey = TemplateKey,
+			.TypeStatsIndex = GetTypeStatsIndex(Instance->GetClass()) 
+		};
 		ReplicationProtocol = ProtocolManager->CreateReplicationProtocol(ProtocolIdentifier, RegisteredFragments, *Instance->GetClass()->GetName(), CreateProtocolParams);
 	}
 #if UE_IRIS_VALIDATE_PROTOCOLS
@@ -975,10 +988,6 @@ FReplicationBridgeCreateNetRefHandleResult UObjectReplicationBridge::CreateNetRe
 
 	const FReplicationProtocolIdentifier ReceivedProtocolId = Header->GetProtocolId();	
 
-	// Currently remote objects can only receive replicated data
-	FFragmentRegistrationContext FragmentRegistrationContext(GetReplicationStateDescriptorRegistry(), GetReplicationSystem(), EReplicationFragmentTraits::CanReceive);
-	FReplicationProtocolManager* ProtocolManager = GetReplicationProtocolManager();
-
 	if (UE_LOG_ACTIVE(LogIris, Verbose))
 	{
 		if (RootObjectOfSubObject.IsValid())
@@ -1024,29 +1033,50 @@ FReplicationBridgeCreateNetRefHandleResult UObjectReplicationBridge::CreateNetRe
 	InstancePtr = Result.Instance;
 	CreateResult.Flags |= Result.Flags;
 
+	// Currently remote objects can only receive replicated data
+	FFragmentRegistrationContext FragmentRegistrationContext(GetReplicationStateDescriptorRegistry(), GetReplicationSystem(), EReplicationFragmentTraits::CanReceive, InstancePtr);
+
+	// Set the optional template provided by the Factory
+	FFragmentRegistrationContextPrivateAccessor::SetTemplate(FragmentRegistrationContext, Result.Template);
+
 	// Register all fragments
 	CallRegisterReplicationFragments(InstancePtr, FragmentRegistrationContext, EFragmentRegistrationFlags::None);
 
-	const FReplicationFragments& RegisteredFragments = FFragmentRegistrationContextPrivateAccessor::GetReplicationFragments(FragmentRegistrationContext);
+	// Find the final template assigned to the instance during fragment creation
+	const UObject* TemplateKey = FFragmentRegistrationContextPrivateAccessor::GetDefaultStateSource(FragmentRegistrationContext);
 
-	// We currently identify protocols by local archetype or CDO pointer and verified the protocol id received from server
-	// We also should verify the default state that we use for delta compression
-	const UObject* ArchetypeOrCDOUsedAsKey = InstancePtr->GetArchetype();
+	// If no templates were registered
+	if (!TemplateKey)
+	{
+		// Use the template provided by the factory or the object's archetype
+		TemplateKey = Result.Template ? Result.Template : InstancePtr->GetArchetype();
+	}
 
-	// Create Protocols
+	FReplicationProtocolManager* ProtocolManager = GetReplicationProtocolManager();
+	const FReplicationFragments& RegisteredFragments = FFragmentRegistrationContextPrivateAccessor::GetReplicationFragments(FragmentRegistrationContext);
+
+	// Create Protocol instance
 	FReplicationInstanceProtocolPtr InstanceProtocol(ProtocolManager->CreateInstanceProtocol(RegisteredFragments, FragmentRegistrationContext.GetFragmentTraits()));
 
-	// See if the protocol already is known
-	const FReplicationProtocol* ReplicationProtocol = ProtocolManager->GetReplicationProtocol(ReceivedProtocolId, ArchetypeOrCDOUsedAsKey);
+	// See if the protocol is known
+	const FReplicationProtocol* ReplicationProtocol = ProtocolManager->GetReplicationProtocol(ReceivedProtocolId, TemplateKey);
 	if (!ReplicationProtocol)
 	{
-		FCreateReplicationProtocolParameters CreateProtocolParams {.ArchetypeOrCDOUsedAsKey = ArchetypeOrCDOUsedAsKey, .bValidateProtocolId = true};
+		// Create the protocol and compare the received hash ID to make sure it is identical to the one on the server
+		const FCreateReplicationProtocolParameters CreateProtocolParams 
+		{
+			.TemplateKey = TemplateKey, 
+			.bValidateProtocolId = true, // Make sure the local protocol hash matches the one sent by the server
+		};
 		ReplicationProtocol = ProtocolManager->CreateReplicationProtocol(ReceivedProtocolId, RegisteredFragments, *(InstancePtr->GetClass()->GetName()), CreateProtocolParams);
 	}
 	else
 	{
-		constexpr bool bDoNotLogErrors = false; // Don't log errors because it would spam for every individual object of the same class. 
-		bool bIsValid = FReplicationProtocolManager::ValidateReplicationProtocol(ReplicationProtocol, RegisteredFragments, bDoNotLogErrors);
+		// Don't log errors because it would spam for every individual object of the same class. 
+		constexpr bool bDoNotLogErrors = false;
+
+		// Confirm that the found protocol works with the fragments
+		const bool bIsValid = FReplicationProtocolManager::ValidateReplicationProtocol(ReplicationProtocol, RegisteredFragments, bDoNotLogErrors);
 		if (!bIsValid)
 		{
 			ReplicationProtocol = nullptr;
@@ -1055,11 +1085,11 @@ FReplicationBridgeCreateNetRefHandleResult UObjectReplicationBridge::CreateNetRe
 
 	if (!ReplicationProtocol)
 	{
-		UE_LOG_BRIDGEID(LogIris, Error, TEXT("Protocol mismatch prevents binding %s to instanced object %s (CDO: %s)."), *WantedNetHandle.ToString(), *GetNameSafe(InstancePtr), *GetPathNameSafe(ArchetypeOrCDOUsedAsKey));
+		UE_LOG_BRIDGEID(LogIris, Error, TEXT("Protocol mismatch prevents binding %s to instanced object %s (Template: %s)."), *WantedNetHandle.ToString(), *GetNameSafe(InstancePtr), *GetPathNameSafe(TemplateKey));
 
 		if (UE_LOG_ACTIVE(LogIris, Error))
 		{
-			UE::Net::Private::ObjectBridgeDebugging::RemoteProtocolMismatchDetected(ArchetypesAlreadyPrinted, ReplicationSystem, Context.ConnectionId, RegisteredFragments, ArchetypeOrCDOUsedAsKey, InstancePtr);
+			UE::Net::Private::ObjectBridgeDebugging::RemoteProtocolMismatchDetected(ArchetypesAlreadyPrinted, ReplicationSystem, Context.ConnectionId, RegisteredFragments, TemplateKey, InstancePtr);
 		}
 
 		FIrisDelegates::GetCriticalErrorDetectedDelegate().Broadcast(ReplicationSystem);
@@ -1075,16 +1105,22 @@ FReplicationBridgeCreateNetRefHandleResult UObjectReplicationBridge::CreateNetRe
 		.ReplicationProtocol = ReplicationProtocol,
 	};
 
-	FNetRefHandle Handle = InternalCreateNetObjectFromRemote(WantedNetHandle, CreateNetObjectParams);
-	CreateResult.NetRefHandle = Handle;
-	if (Handle.IsValid())
+	CreateResult.NetRefHandle = InternalCreateNetObjectFromRemote(WantedNetHandle, CreateNetObjectParams);
+
+	if (CreateResult.NetRefHandle.IsValid())
 	{
-		RegisterRemoteInstance(Handle, InstancePtr, ReplicationProtocol, InstanceProtocol.Get(), Context.ConnectionId);
+		RegisterRemoteInstance(CreateResult.NetRefHandle, InstancePtr, ReplicationProtocol, InstanceProtocol.Get(), Context.ConnectionId);
 
 		// Release instance protocol from the uniquePtr as it is now successfully bound to the handle
 		(void)InstanceProtocol.Release();
 
-		Factory->PostInstantiation(UNetObjectFactory::FPostInstantiationContext{ .Instance = InstancePtr, .Header = Header.Get(), .ConnectionId = Context.ConnectionId });
+		const UNetObjectFactory::FPostInstantiationContext PostContext
+		{ 
+			.Instance = InstancePtr, 
+			.Header = Header.Get(), 
+			.ConnectionId = Context.ConnectionId 
+		};
+		Factory->PostInstantiation(PostContext);
 	}
 
 	return CreateResult;
diff --git a/Engine/Source/Runtime/Experimental/Iris/Core/Private/Iris/ReplicationSystem/ObjectReplicationBridgeDebugging.cpp b/Engine/Source/Runtime/Experimental/Iris/Core/Private/Iris/ReplicationSystem/ObjectReplicationBridgeDebugging.cpp
--- a/Engine/Source/Runtime/Experimental/Iris/Core/Private/Iris/ReplicationSystem/ObjectReplicationBridgeDebugging.cpp
+++ b/Engine/Source/Runtime/Experimental/Iris/Core/Private/Iris/ReplicationSystem/ObjectReplicationBridgeDebugging.cpp
@@ -1107,9 +1107,10 @@ void UObjectReplicationBridge::PrintPushBasedStatuses() const
 	};
 
 	TArray<FPushBasedInfo> PushBasedInfos;
-	ProtocolManager->ForEachProtocol([&](const FReplicationProtocol* Protocol, const UObject* ArchetypeOrCDOUsedAsKey)
+	ProtocolManager->ForEachProtocol([&](const FReplicationProtocol* Protocol, const FObjectKey TemplateKey)
 	{
-		if (!ArchetypeOrCDOUsedAsKey)
+		UObject* Template = TemplateKey.ResolveObjectPtr();
+		if (!Template)
 		{
 			return;
 		}
@@ -1118,12 +1119,12 @@ void UObjectReplicationBridge::PrintPushBasedStatuses() const
 		{
 			if (!EnumHasAnyFlags(StateDescriptor->Traits, EReplicationStateTraits::HasPushBasedDirtiness))
 			{
-				PushBasedInfos.Add({ArchetypeOrCDOUsedAsKey->GetClass(), Protocol->GetRefCount(), false});
+				PushBasedInfos.Add({ Template->GetClass(), Protocol->GetRefCount(), false});
 				return;
 			}
 		}
 
-		PushBasedInfos.Add({ArchetypeOrCDOUsedAsKey->GetClass(), Protocol->GetRefCount(), true});
+		PushBasedInfos.Add({ Template->GetClass(), Protocol->GetRefCount(), true});
 	});
 
 	// Print by push-based status (not push-based first), then by ref count, then by name.
diff --git a/Engine/Source/Runtime/Experimental/Iris/Core/Private/Iris/ReplicationSystem/ReplicationBridge.cpp b/Engine/Source/Runtime/Experimental/Iris/Core/Private/Iris/ReplicationSystem/ReplicationBridge.cpp
--- a/Engine/Source/Runtime/Experimental/Iris/Core/Private/Iris/ReplicationSystem/ReplicationBridge.cpp
+++ b/Engine/Source/Runtime/Experimental/Iris/Core/Private/Iris/ReplicationSystem/ReplicationBridge.cpp
@@ -380,8 +380,14 @@ void UReplicationBridge::Initialize(UReplicationSystem* InReplicationSystem)
 	// Create destruction info protocol
 	{
 		const FReplicationFragments RegisteredFragments;
-		FCreateReplicationProtocolParameters CreateProtocolParams {.bValidateProtocolId = false, .TypeStatsIndex =  GetReplicationSystem()->GetReplicationSystemInternal()->GetNetTypeStats().GetOrCreateTypeStats(FName("DestructionInfo"))};
-		DestructionInfoProtocol = ReplicationProtocolManager->CreateReplicationProtocol(FReplicationProtocolManager::CalculateProtocolIdentifier(RegisteredFragments), RegisteredFragments, TEXT("InternalDestructionInfo"), CreateProtocolParams);
+		const FReplicationProtocolIdentifier ProtocolId = FReplicationProtocolManager::CalculateProtocolIdentifier(RegisteredFragments);
+		const FCreateReplicationProtocolParameters CreateProtocolParams
+		{
+			.bValidateProtocolId = false,
+			.bHasTemplateKey = false, // Destruction info protocol doesn't need a TemplateKey
+			.TypeStatsIndex =  GetReplicationSystem()->GetReplicationSystemInternal()->GetNetTypeStats().GetOrCreateTypeStats(FName("DestructionInfo")),
+		};
+		DestructionInfoProtocol = ReplicationProtocolManager->CreateReplicationProtocol(ProtocolId, RegisteredFragments, TEXT("InternalDestructionInfo"), CreateProtocolParams);
 		// Explicit refcount
 		DestructionInfoProtocol->AddRef();
 	}
diff --git a/Engine/Source/Runtime/Experimental/Iris/Core/Private/Iris/ReplicationSystem/ReplicationFragmentInternal.h b/Engine/Source/Runtime/Experimental/Iris/Core/Private/Iris/ReplicationSystem/ReplicationFragmentInternal.h
--- a/Engine/Source/Runtime/Experimental/Iris/Core/Private/Iris/ReplicationSystem/ReplicationFragmentInternal.h
+++ b/Engine/Source/Runtime/Experimental/Iris/Core/Private/Iris/ReplicationSystem/ReplicationFragmentInternal.h
@@ -8,9 +8,49 @@ namespace UE::Net::Private
 {
 	struct FFragmentRegistrationContextPrivateAccessor
 	{
-		static const FReplicationFragments& GetReplicationFragments(const FFragmentRegistrationContext& Context) { return Context.Fragments; }
-		static const Private::FReplicationStateDescriptorRegistry* GetReplicationStateRegistry(const FFragmentRegistrationContext& Context) { return Context.ReplicationStateRegistry; }
-		static Private::FReplicationStateDescriptorRegistry* GetReplicationStateRegistry(FFragmentRegistrationContext& Context) { return Context.ReplicationStateRegistry; }
-		static UReplicationSystem* GetReplicationSystem(const FFragmentRegistrationContext& Context) { return Context.ReplicationSystem; }
+		static const FReplicationFragments& GetReplicationFragments(const FFragmentRegistrationContext& Context)
+		{ 
+			return Context.Fragments;
+		}
+		
+		static const Private::FReplicationStateDescriptorRegistry* GetReplicationStateRegistry(const FFragmentRegistrationContext& Context)
+		{ 
+			return Context.ReplicationStateRegistry;
+		}
+		
+		static Private::FReplicationStateDescriptorRegistry* GetReplicationStateRegistry(FFragmentRegistrationContext& Context)
+		{ 
+			return Context.ReplicationStateRegistry;
+		}
+		
+		static UReplicationSystem* GetReplicationSystem(const FFragmentRegistrationContext& Context)
+		{ 
+			return Context.ReplicationSystem;
+		}
+
+		static bool IsMainObject(const FFragmentRegistrationContext& Context, UObject* ObjectInstance)
+		{
+			return Context.MainObjectInstance == ObjectInstance;
+		}
+
+		static void SetDefaultStateSource(FFragmentRegistrationContext& Context, const UObject* DefaultStateSource)
+		{
+			Context.MainObjectDefaultStateSource = DefaultStateSource;
+		}
+
+		static const UObject* GetDefaultStateSource(const FFragmentRegistrationContext& Context)
+		{
+			return Context.MainObjectDefaultStateSource;
+		}
+
+		static void SetTemplate(FFragmentRegistrationContext& Context, const UObject* Template)
+		{
+			Context.Template = Template;
+		}
+
+		static const UObject* GetTemplate(const FFragmentRegistrationContext& Context)
+		{
+			return Context.Template;
+		}
 	};
 }
diff --git a/Engine/Source/Runtime/Experimental/Iris/Core/Private/Iris/ReplicationSystem/ReplicationFragmentUtil.cpp b/Engine/Source/Runtime/Experimental/Iris/Core/Private/Iris/ReplicationSystem/ReplicationFragmentUtil.cpp
--- a/Engine/Source/Runtime/Experimental/Iris/Core/Private/Iris/ReplicationSystem/ReplicationFragmentUtil.cpp
+++ b/Engine/Source/Runtime/Experimental/Iris/Core/Private/Iris/ReplicationSystem/ReplicationFragmentUtil.cpp
@@ -14,36 +14,77 @@ namespace UE::Net
 
 uint32 FReplicationFragmentUtil::CreateAndRegisterFragmentsForObject(UObject* Object, FFragmentRegistrationContext& Context, EFragmentRegistrationFlags RegistrationFlags, TArray<FReplicationFragment*>* OutCreatedFragments)
 {
+	const FCreateFragmentParams CreateParams
+	{
+		.ObjectInstance = Object,
+		.RegistrationFlags = RegistrationFlags,
+	};
+
+	return CreateAndRegisterFragmentsForObject(CreateParams, Context, OutCreatedFragments);
+}
+
+uint32 FReplicationFragmentUtil::CreateAndRegisterFragmentsForObject(const FCreateFragmentParams& Params, FFragmentRegistrationContext& Context, TArray<FReplicationFragment*>* OutCreatedFragments)
+{
+	using namespace UE::Net::Private;
+
 	IRIS_PROFILER_SCOPE(FReplicationFragmentUtil_CreateAndRegisterFragmentsForObject);
 
 	FReplicationStateDescriptorBuilder::FParameters BuilderParameters;
-	BuilderParameters.DescriptorRegistry = UE::Net::Private::FFragmentRegistrationContextPrivateAccessor::GetReplicationStateRegistry(Context);
-	BuilderParameters.ReplicationSystem = UE::Net::Private::FFragmentRegistrationContextPrivateAccessor::GetReplicationSystem(Context);
+	BuilderParameters.DescriptorRegistry = FFragmentRegistrationContextPrivateAccessor::GetReplicationStateRegistry(Context);
+	BuilderParameters.ReplicationSystem = FFragmentRegistrationContextPrivateAccessor::GetReplicationSystem(Context);
+
+	UClass* ObjectClass = Params.ObjectInstance->GetClass();
+
+	const bool bIsMainObject = FFragmentRegistrationContextPrivateAccessor::IsMainObject(Context, Params.ObjectInstance);
 
-	// By default we use the archetype for our default values, but if InitializeDefaultStateFromClassDefaults is set we 
-	// use the class default instead
-	if (!EnumHasAnyFlags(RegistrationFlags, EFragmentRegistrationFlags::InitializeDefaultStateFromClassDefaults))
+	// Find the state source needed by some fragments
 	{
-		BuilderParameters.DefaultStateSource = Object->GetArchetype();
+		const UObject* DefaultStateSource = nullptr;
 
-		if (BuilderParameters.DefaultStateSource == nullptr)
+		// Some objects want to always build the init state from the CDO
+		if (EnumHasAnyFlags(Params.RegistrationFlags, EFragmentRegistrationFlags::InitializeDefaultStateFromClassDefaults))
+		{
+			DefaultStateSource = ObjectClass->GetDefaultObject();
+		}
+		// For the main object look if a template was provided
+		else if (bIsMainObject)
 		{
-			UE_LOG(LogIris, Error, TEXT("FPropertyReplicationFragment::CreateAndRegisterFragmentsForObject: Invalid object archetype for object %s, default state will be used from CDO:"), *GetFullNameSafe(Object));
+			DefaultStateSource = FFragmentRegistrationContextPrivateAccessor::GetTemplate(Context);
 		}
+	
+		// Default to use the archetype otherwise
+		if (DefaultStateSource == nullptr)
+		{
+			DefaultStateSource = Params.ObjectInstance->GetArchetype();
+		}
+
+		if (!IsValid(DefaultStateSource))
+		{
+			UE_LOG(LogIris, Error, TEXT("FPropertyReplicationFragment::CreateAndRegisterFragmentsForObject: Invalid object archetype for object %s, default state will use the CDO"), *GetFullNameSafe(Params.ObjectInstance));
+			DefaultStateSource = ObjectClass->GetDefaultObject();
+		}
+
+		if (bIsMainObject)
+		{
+			// Store the state source so the bridge can retrieve it
+			FFragmentRegistrationContextPrivateAccessor::SetDefaultStateSource(Context, DefaultStateSource);
+		}
+
+		BuilderParameters.DefaultStateSource = DefaultStateSource;
 	}
 
 	// Pass-on that we allow FastAarrays with extra replicated properties for this object.
 	// NOTE: this is not perfect as it allows this for all FastArray properties of the object but as there 
 	// is further validation in the actual ReplicationStateFragment implementation for FastArrays it is good enough.
-	if (EnumHasAnyFlags(RegistrationFlags, EFragmentRegistrationFlags::AllowFastArraysWithAdditionalProperties))
+	if (EnumHasAnyFlags(Params.RegistrationFlags, EFragmentRegistrationFlags::AllowFastArraysWithAdditionalProperties))
 	{
 		BuilderParameters.AllowFastArrayWithExtraReplicatedProperties = 1U;
 	}
 
 	FReplicationStateDescriptorBuilder::FResult Result;
-	FReplicationStateDescriptorBuilder::CreateDescriptorsForClass(Result, Object->GetClass(), BuilderParameters);
+	FReplicationStateDescriptorBuilder::CreateDescriptorsForClass(Result, ObjectClass, BuilderParameters);
 
-	const bool bRegisterFunctionsOnly = EnumHasAnyFlags(RegistrationFlags, EFragmentRegistrationFlags::RegisterRPCsOnly);
+	const bool bRegisterFunctionsOnly = EnumHasAnyFlags(Params.RegistrationFlags, EFragmentRegistrationFlags::RegisterRPCsOnly);
 
 	uint32 NumCreatedReplicationFragments = 0U;
 	// create fragments and let the instance protocol own them.
@@ -58,11 +99,11 @@ uint32 FReplicationFragmentUtil::CreateAndRegisterFragmentsForObject(UObject* Ob
 		// If descriptor provides CreateAndRegisterReplicationFragment function we use that to instantiate fragment
 		if (Desc->CreateAndRegisterReplicationFragmentFunction)
 		{
-			Fragment = Desc->CreateAndRegisterReplicationFragmentFunction(Object, Desc.GetReference(), Context);
+			Fragment = Desc->CreateAndRegisterReplicationFragmentFunction(Params.ObjectInstance, Desc.GetReference(), Context);
 		}
 		else
 		{
-			Fragment = FPropertyReplicationFragment::CreateAndRegisterFragment(Object, Desc.GetReference(), Context);
+			Fragment = FPropertyReplicationFragment::CreateAndRegisterFragment(Params.ObjectInstance, Desc.GetReference(), Context);
 		}
 		if (Fragment && OutCreatedFragments)
 		{
diff --git a/Engine/Source/Runtime/Experimental/Iris/Core/Private/Iris/ReplicationSystem/ReplicationProtocolManager.cpp b/Engine/Source/Runtime/Experimental/Iris/Core/Private/Iris/ReplicationSystem/ReplicationProtocolManager.cpp
--- a/Engine/Source/Runtime/Experimental/Iris/Core/Private/Iris/ReplicationSystem/ReplicationProtocolManager.cpp
+++ b/Engine/Source/Runtime/Experimental/Iris/Core/Private/Iris/ReplicationSystem/ReplicationProtocolManager.cpp
@@ -225,13 +225,12 @@ bool FReplicationProtocolManager::ValidateReplicationProtocol(const FReplication
 	return bResult;
 }
 
-const FReplicationProtocol* FReplicationProtocolManager::GetReplicationProtocol(FReplicationProtocolIdentifier ProtocolId, const UObject* ArchetypeOrCDOUsedAsKey)
+const FReplicationProtocol* FReplicationProtocolManager::GetReplicationProtocol(FReplicationProtocolIdentifier ProtocolId, FObjectKey TemplateKey)
 {
-	// Find protocols using the descriptor being invalidated
 	for (auto It = RegisteredProtocols.CreateConstKeyIterator(ProtocolId); It; ++It)
 	{
 		const FRegisteredProtocolInfo& Info = It.Value();
-		if (Info.ArchetypeOrCDOUsedAsKey == ArchetypeOrCDOUsedAsKey)
+		if (Info.TemplateKey == TemplateKey)
 		{
 			return Info.Protocol;
 		}
@@ -274,7 +273,11 @@ const FReplicationProtocol* FReplicationProtocolManager::CreateReplicationProtoc
 		}
 	}
 
-	check(GetReplicationProtocol(ProtocolId, Params.ArchetypeOrCDOUsedAsKey) == nullptr);
+	// Template key should be valid unless the protocol doesn't use one.
+	check(Params.bHasTemplateKey == false || IsValid(Params.TemplateKey));
+
+	// Protocol should not exist already
+	check(GetReplicationProtocol(ProtocolId, Params.TemplateKey) == nullptr);
 
 	// Create the protocol
 	const uint32 FragmentCount = Fragments.Num();
@@ -505,7 +508,7 @@ const FReplicationProtocol* FReplicationProtocolManager::CreateReplicationProtoc
 
 	// Register protocol
 	FRegisteredProtocolInfo Info;
-	Info.ArchetypeOrCDOUsedAsKey = Params.ArchetypeOrCDOUsedAsKey;
+	Info.TemplateKey = Params.TemplateKey;
 	Info.Protocol = Protocol;
 
 	RegisteredProtocols.AddUnique(ProtocolId, Info);
diff --git a/Engine/Source/Runtime/Experimental/Iris/Core/Private/Iris/ReplicationSystem/ReplicationProtocolManager.h b/Engine/Source/Runtime/Experimental/Iris/Core/Private/Iris/ReplicationSystem/ReplicationProtocolManager.h
--- a/Engine/Source/Runtime/Experimental/Iris/Core/Private/Iris/ReplicationSystem/ReplicationProtocolManager.h
+++ b/Engine/Source/Runtime/Experimental/Iris/Core/Private/Iris/ReplicationSystem/ReplicationProtocolManager.h
@@ -4,6 +4,8 @@
 
 #include "CoreTypes.h"
 #include "Containers/Map.h"
+#include "UObject/ObjectKey.h"
+
 #include "Iris/ReplicationState/ReplicationStateDescriptor.h"
 #include "Iris/ReplicationSystem/ReplicationFragment.h"
 #include "Iris/ReplicationSystem/ReplicationProtocol.h"
@@ -13,8 +15,9 @@ namespace UE::Net::Private
 
 struct FCreateReplicationProtocolParameters
 {
-	const UObject* ArchetypeOrCDOUsedAsKey = nullptr;
+	const UObject* TemplateKey = nullptr;
 	bool bValidateProtocolId = false;
+	bool bHasTemplateKey = true;
 	int32 TypeStatsIndex = INDEX_NONE;
 };
 
@@ -27,7 +30,7 @@ public:
 	IRISCORE_API const FReplicationProtocol* CreateReplicationProtocol(const FReplicationProtocolIdentifier ProtocolId, const FReplicationFragments& Fragments, const TCHAR* DebugName, const FCreateReplicationProtocolParameters& Params = FCreateReplicationProtocolParameters());
 
 	/* Get an existing replication protocol */
-	IRISCORE_API const FReplicationProtocol* GetReplicationProtocol(FReplicationProtocolIdentifier ProtocolId, const UObject* ArchetypeOrCDOUsedAsKey = nullptr);
+	IRISCORE_API const FReplicationProtocol* GetReplicationProtocol(FReplicationProtocolIdentifier ProtocolId, FObjectKey TemplateKey);
 
 	/** Iterate over all protocols matching the ProtocolId, mostly used by debug functionality with no sideeffects */
 	template<typename T>
@@ -69,9 +72,9 @@ private:
 	struct FRegisteredProtocolInfo
 	{
 		const FReplicationProtocol* Protocol;
-		const UObject* ArchetypeOrCDOUsedAsKey;
+		FObjectKey TemplateKey;
 	
-		bool operator==(const FRegisteredProtocolInfo& Other) const { return Protocol == Other.Protocol && ArchetypeOrCDOUsedAsKey == Other.ArchetypeOrCDOUsedAsKey; };
+		bool operator==(const FRegisteredProtocolInfo& Other) const { return Protocol == Other.Protocol && TemplateKey == Other.TemplateKey; };
 	};
 
 	TMultiMap<FReplicationProtocolIdentifier, FRegisteredProtocolInfo> RegisteredProtocols;
@@ -89,7 +92,7 @@ void FReplicationProtocolManager::ForEachProtocol(FReplicationProtocolIdentifier
 	for (auto It = RegisteredProtocols.CreateConstKeyIterator(ProtocolId); It; ++It)
 	{
 		const FRegisteredProtocolInfo& Info = It.Value();
-		Functor(Info.Protocol, Info.ArchetypeOrCDOUsedAsKey);
+		Functor(Info.Protocol, Info.TemplateKey);
 	}
 }
 
@@ -99,7 +102,7 @@ void FReplicationProtocolManager::ForEachProtocol(T&& Functor) const
 	for (auto It = RegisteredProtocols.CreateConstIterator(); It; ++It)
 	{
 		const FRegisteredProtocolInfo& Info = It.Value();
-		Functor(Info.Protocol, Info.ArchetypeOrCDOUsedAsKey);
+		Functor(Info.Protocol, Info.TemplateKey);
 	}
 }
 
diff --git a/Engine/Source/Runtime/Experimental/Iris/Core/Public/Iris/ReplicationState/ReplicationStateDescriptorBuilder.h b/Engine/Source/Runtime/Experimental/Iris/Core/Public/Iris/ReplicationState/ReplicationStateDescriptorBuilder.h
--- a/Engine/Source/Runtime/Experimental/Iris/Core/Public/Iris/ReplicationState/ReplicationStateDescriptorBuilder.h
+++ b/Engine/Source/Runtime/Experimental/Iris/Core/Public/Iris/ReplicationState/ReplicationStateDescriptorBuilder.h
@@ -25,27 +25,31 @@ public:
 		IRISCORE_API FParameters();
 
 		Private::FReplicationStateDescriptorRegistry* DescriptorRegistry;
-		UReplicationSystem* ReplicationSystem;
-		// If provided we use this to source our default state, if not provided we source default state from the CDO instead
-		UObject* DefaultStateSource;
+		UReplicationSystem* ReplicationSystem = nullptr;
+
+		// The template to build the state source from
+		const UObject* DefaultStateSource = nullptr;
+
 		 // Include super class when building descriptor.
-		uint32 IncludeSuper : 1;
+		bool IncludeSuper : 1 = true;
 		// Get the lifetime properties for conditionals if applicable. This applies only to classes.
-		uint32 GetLifeTimeProperties : 1;
+		bool GetLifeTimeProperties : 1 = true;
 		// If EnableFastArrayHandling is true and the struct inherits from FFastArraySerializer then special logic will be applied which allows it to be bound to a FastArrayReplicationFragment.
-		uint32 EnableFastArrayHandling : 1;
+		bool EnableFastArrayHandling : 1 = true;
 		// If AllowFastArrayWithExtraReplicatedProperties is set to true, we will allow building descriptors for fastarrays with more than a single property.
-		uint32 AllowFastArrayWithExtraReplicatedProperties : 1;
+		bool AllowFastArrayWithExtraReplicatedProperties : 1 = false;
 		// In SkipCheckForCustomNetSerializer is true the descriptor will be built using the underlying representation even if a CustomNetSerializer is registered for the struct
-		uint32 SkipCheckForCustomNetSerializerForStruct : 1;
+		bool SkipCheckForCustomNetSerializerForStruct : 1 = false;
 		 // If SinglePropertyIndex != -1 we will create a state that only includes the specified property.
-		int32 SinglePropertyIndex;
+		int32 SinglePropertyIndex = -1;
 	};
 
 	typedef TArray<TRefCountPtr<const FReplicationStateDescriptor>> FResult;
 
 	/**
-		Create ReplicationDescriptor(s) for Class by processing reflection information at runtime for existing replicated properties.
+		Create the descriptors from an optional template or the class CDO.
+
+		Done by processing reflection information at runtime for existing replicated properties.
 
 		Since we want to allow ReplicationStates to be shared as much as possible between different connection we do split them up based on
 		feature requirements, filtering, conditionals etc. We also keep Init states separate.
diff --git a/Engine/Source/Runtime/Experimental/Iris/Core/Public/Iris/ReplicationSystem/NetObjectFactory.h b/Engine/Source/Runtime/Experimental/Iris/Core/Public/Iris/ReplicationSystem/NetObjectFactory.h
--- a/Engine/Source/Runtime/Experimental/Iris/Core/Public/Iris/ReplicationSystem/NetObjectFactory.h
+++ b/Engine/Source/Runtime/Experimental/Iris/Core/Public/Iris/ReplicationSystem/NetObjectFactory.h
@@ -180,6 +180,8 @@ public:
 	{
 		/** The instantiated object represented by the header */
 		UObject* Instance = nullptr;
+		/** The template used to instantiate the object with. Only set this when the template is different from the object's archetype. */
+		UObject* Template = nullptr;
 		/** Flags to pass back to the bridge */
 		EReplicationBridgeCreateNetRefHandleResultFlags Flags = EReplicationBridgeCreateNetRefHandleResultFlags::None;
 	};
diff --git a/Engine/Source/Runtime/Experimental/Iris/Core/Public/Iris/ReplicationSystem/ReplicationFragment.h b/Engine/Source/Runtime/Experimental/Iris/Core/Public/Iris/ReplicationSystem/ReplicationFragment.h
--- a/Engine/Source/Runtime/Experimental/Iris/Core/Public/Iris/ReplicationSystem/ReplicationFragment.h
+++ b/Engine/Source/Runtime/Experimental/Iris/Core/Public/Iris/ReplicationSystem/ReplicationFragment.h
@@ -3,7 +3,10 @@
 #pragma once
 
 #include "CoreTypes.h"
+
 #include "Containers/Array.h"
+#include "Containers/Map.h"
+
 #include "Misc/CoreMiscDefines.h"
 #include "Misc/EnumClassFlags.h"
 
@@ -240,9 +243,10 @@ typedef TArray<FReplicationFragmentInfo, TInlineAllocator<32>> FReplicationFragm
 class FFragmentRegistrationContext
 {
 public:
-	explicit FFragmentRegistrationContext(Private::FReplicationStateDescriptorRegistry* InReplicationStateRegistry, UReplicationSystem* InReplicationSystem, const EReplicationFragmentTraits InFragmentTraits)
+	explicit FFragmentRegistrationContext(Private::FReplicationStateDescriptorRegistry* InReplicationStateRegistry, UReplicationSystem* InReplicationSystem, const EReplicationFragmentTraits InFragmentTraits, UObject* MainObject)
 		: ReplicationStateRegistry(InReplicationStateRegistry)
 		, ReplicationSystem(InReplicationSystem)
+		, MainObjectInstance(MainObject)
 		, FragmentTraits(InFragmentTraits)
 	{
 	}
@@ -251,7 +255,10 @@ public:
 	const EReplicationFragmentTraits GetFragmentTraits() const { return FragmentTraits; }
 
 	/** Register ReplicationFragment */
-	void RegisterReplicationFragment(FReplicationFragment* Fragment, const FReplicationStateDescriptor* Descriptor, void* SrcReplicationStateBuffer) { Fragments.Add({ Descriptor, SrcReplicationStateBuffer, Fragment }); }
+	void RegisterReplicationFragment(FReplicationFragment* Fragment, const FReplicationStateDescriptor* Descriptor, void* SrcReplicationStateBuffer)
+	{ 
+		Fragments.Add({ Descriptor, SrcReplicationStateBuffer, Fragment }); 
+	}
 
 	/** Call this when you have a netobject that is replicated but will never contain any RPCs or replicated properties. Prevents the registration code from complaining about potential errors. */
 	void SetIsFragmentlessNetObject(bool bIsFragmentless) { bIsAFragmentlessNetObject = bIsFragmentless; }
@@ -271,11 +278,24 @@ public:
 	friend Private::FFragmentRegistrationContextPrivateAccessor;
 
 private:
+
 	FReplicationFragments Fragments;
 	Private::FReplicationStateDescriptorRegistry* ReplicationStateRegistry;
 	UReplicationSystem* ReplicationSystem;
+
+	/** The main object used to pull most fragments from. Used as the basic for the default state source */
+	UObject* MainObjectInstance = nullptr;
+
+	/** The externally provided template used to clone the main object instance from */
+	const UObject* Template = nullptr;
+
+	/** The final reference we used to build the default state for the main object. Since Template is optional this lets callers know what was used from a template, archetype or cdo. */
+	const UObject* MainObjectDefaultStateSource = nullptr;
+
 	const EReplicationFragmentTraits FragmentTraits;
+
 	bool bIsAFragmentlessNetObject = false;
+
 };
 
 }
diff --git a/Engine/Source/Runtime/Experimental/Iris/Core/Public/Iris/ReplicationSystem/ReplicationFragmentUtil.h b/Engine/Source/Runtime/Experimental/Iris/Core/Public/Iris/ReplicationSystem/ReplicationFragmentUtil.h
--- a/Engine/Source/Runtime/Experimental/Iris/Core/Public/Iris/ReplicationSystem/ReplicationFragmentUtil.h
+++ b/Engine/Source/Runtime/Experimental/Iris/Core/Public/Iris/ReplicationSystem/ReplicationFragmentUtil.h
@@ -4,29 +4,37 @@
 #include "CoreTypes.h"
 #include "Containers/ContainersFwd.h"
 
+#include "Iris/ReplicationSystem/ReplicationFragment.h"
+
 class UObject;
 
 namespace UE::Net
 {
 
-class FFragmentRegistrationContext;
-class FReplicationFragment;
-enum class EFragmentRegistrationFlags : uint32;
-
-/**
-* FPropertyReplicationFragment - used to bind PropertyReplicationStates to their owner
-*/
+/** FPropertyReplicationFragment - used to bind PropertyReplicationStates to their owner */
 class FReplicationFragmentUtil
 {
 public:
 
-	/** 
+	/** Simple version of CreateAndRegisterFragmentsForObject */
+	IRISCORE_API static uint32 CreateAndRegisterFragmentsForObject(UObject* Object, FFragmentRegistrationContext& Context, EFragmentRegistrationFlags RegistrationFlags, TArray<FReplicationFragment*>* OutCreatedFragments = nullptr);
+
+
+	struct FCreateFragmentParams
+	{
+		/** The replicated object to create fragments for */
+		UObject* ObjectInstance = nullptr;
+		/** Flags to control how the fragments should be created */
+		EFragmentRegistrationFlags RegistrationFlags = EFragmentRegistrationFlags::None;
+	};
+
+	/**
 	* Create and register all property replication Fragments for the provided object, Descriptors will be created based on the Class of the object
 	* Lifetime of created fragments will be managed by the ReplicationSystem
 	* If the OutCreatedFragments are provided pointers to the created fragments will be added to the provided array
 	* Returns the number of created fragments
 	*/
-	IRISCORE_API static uint32 CreateAndRegisterFragmentsForObject(UObject* Object, FFragmentRegistrationContext& Context, EFragmentRegistrationFlags RegistrationFlags, TArray<FReplicationFragment*>* OutCreatedFragments = nullptr);
+	IRISCORE_API static uint32 CreateAndRegisterFragmentsForObject(const FCreateFragmentParams& Params, FFragmentRegistrationContext& Context, TArray<FReplicationFragment*>* OutCreatedFragments = nullptr);
 };
 
 }
