Wait for IAX host group task to end to avoid use-after-free.

#rnx

diff --git a/Engine/Source/Runtime/Experimental/IoStore/OnDemand/Private/OnDemandIoDispatcherBackend.cpp b/Engine/Source/Runtime/Experimental/IoStore/OnDemand/Private/OnDemandIoDispatcherBackend.cpp
--- a/Engine/Source/Runtime/Experimental/IoStore/OnDemand/Private/OnDemandIoDispatcherBackend.cpp
+++ b/Engine/Source/Runtime/Experimental/IoStore/OnDemand/Private/OnDemandIoDispatcherBackend.cpp
@@ -696,6 +696,7 @@ private:
 	FIASHostGroup DefaultHostGroup;
 	FDelegateHandle OnThreadTickIdleHandle;
 	mutable FRWLock Lock;
+	UE::Tasks::FTask HostGroupTask;
 	std::atomic_uint32_t InflightCacheRequestCount{0};
 	std::atomic_bool bStopRequested{false};
 
@@ -797,6 +798,9 @@ void FOnDemandIoBackend::Shutdown()
 
 	UE_LOG(LogIas, Log, TEXT("Shutting down on demand I/O dispatcher backend"));
 
+	bStopRequested.store(true, std::memory_order_relaxed);
+	HostGroupTask.Wait();
+
 	if (OnThreadTickIdleHandle.IsValid())
 	{
 		HttpClient.OnTickIdle().Remove(OnThreadTickIdleHandle);
@@ -826,7 +830,7 @@ void FOnDemandIoBackend::OnThreadTickIdle()
 
 void FOnDemandIoBackend::BeginCreateDefaultHostGroup()
 {
-	UE::Tasks::Launch(UE_SOURCE_LOCATION, [this]()
+	HostGroupTask = UE::Tasks::Launch(UE_SOURCE_LOCATION, [this]()
 		{
 			if (CreateDefaultHostGroup())
 			{
