HTTP I/O dispatcher

* extracted a generic HTTP interface that mimics the I/O dispatcher API and uses the IAS cache as general purpose HTTP cache
* reimplemented a new I/O dispatcher backend that uses the new HTTP I/O dispatcher
* updated the content installer to use the new HTTP I/O dispatcher if enabled

Note: The HTTP dispatcher enables merging range requests going to the same endpoint resource which will
enable faster downloads if we decide to merge related I/O store chunks into bigger CDN resources.
Currenlty there's no limit to how many requests get merged since the resources i.e. I/O store chunks is relatively small. However
this can be further extend to be limitied to X number of bytes.
Lane tracing and some of the debug commands has been left out until this change has proven to work.

#rnx

diff --git a/Engine/Source/Runtime/Experimental/IoStore/OnDemand/Private/AtomicFlags.h b/Engine/Source/Runtime/Experimental/IoStore/OnDemand/Private/AtomicFlags.h
new file mode 100644
--- /dev/null
+++ b/Engine/Source/Runtime/Experimental/IoStore/OnDemand/Private/AtomicFlags.h
@@ -0,0 +1,46 @@
+// Copyright Epic Games, Inc. All Rights Reserved.
+
+#pragma once
+
+#include <atomic>
+
+namespace UE
+{
+
+template <typename FlagType>
+class TAtomicFlags
+{
+	using UnderlyingType = __underlying_type(FlagType);
+
+public:
+	FlagType Add(FlagType FlagsToAdd, std::memory_order Order = std::memory_order_seq_cst)
+	{
+		return (FlagType)Flags.fetch_or((UnderlyingType)FlagsToAdd, Order);
+	}
+
+	FlagType Remove(FlagType FlagsToAdd, std::memory_order Order = std::memory_order_seq_cst)
+	{
+		return (FlagType)Flags.fetch_and(~((UnderlyingType)FlagsToAdd), Order);
+	}
+
+	bool HasAny(FlagType Contains, std::memory_order Order = std::memory_order_seq_cst)
+	{
+		return (Flags.load(Order) & (UnderlyingType)Contains) != 0;
+	}
+
+	FlagType Get(std::memory_order Order = std::memory_order_seq_cst) const
+	{
+		return (FlagType)Flags.load(Order);
+	}
+
+	bool TrySet(FlagType FlagsToSet, std::memory_order Order = std::memory_order_seq_cst)
+	{
+		UnderlyingType Expected = Flags.load(Order);
+		return Flags.compare_exchange_strong(Expected, (UnderlyingType)FlagsToSet);
+	}
+
+private:
+	std::atomic<UnderlyingType> Flags{0};
+};
+
+} // namespace UE
diff --git a/Engine/Source/Runtime/Experimental/IoStore/OnDemand/Private/OnDemandContentInstaller.cpp b/Engine/Source/Runtime/Experimental/IoStore/OnDemand/Private/OnDemandContentInstaller.cpp
--- a/Engine/Source/Runtime/Experimental/IoStore/OnDemand/Private/OnDemandContentInstaller.cpp
+++ b/Engine/Source/Runtime/Experimental/IoStore/OnDemand/Private/OnDemandContentInstaller.cpp
@@ -354,7 +354,11 @@ void FOnDemandContentInstaller::UpdateInstallRequestPriority(FSharedInternalInst
 			{
 				for (FChunkHttpRequestHandle& PendingHttpRequest : Install.HttpRequestHandles)
 				{
-					if (PendingHttpRequest.Handle != nullptr)
+					if (PendingHttpRequest.IoHandle.IsValid())
+					{
+						PendingHttpRequest.IoHandle.UpdatePriorty(Request.Priority);
+					}
+					else if (PendingHttpRequest.Handle != nullptr)
 					{
 						HttpClient.ReprioritizeRequest(PendingHttpRequest.Handle, Request.Priority);
 					}
@@ -694,32 +698,75 @@ void FOnDemandContentInstaller::ExecuteInstallRequest(FRequest& Request, bool bR
 
 	NotifyInstallProgress(Request);
 
-	for (FChunkHttpRequestHandle& HttpRequest : InstallRequest.HttpRequestHandles)
+	if (FHttpIoDispatcher::IsInitialized())
 	{
-		TRACE_CPUPROFILER_EVENT_SCOPE(FOnDemandContentInstaller::IssueRequest);
-		FSharedOnDemandContainer& Container		= InstallRequest.ResolvedChunks[HttpRequest.ContainerIndex].Container;
-		const FOnDemandChunkEntry& ChunkEntry	= Container->ChunkEntries[HttpRequest.EntryIndex];
-
-		HttpRequest.Handle = HttpClient.IssueRequest(
-			FOnDemandChunkInfo(Container, ChunkEntry),
-			FIoOffsetAndLength(),
-			Request.Priority,
-			[this, &Request, &HttpRequest](uint32 HttpStatusCode, FStringView ErrorReason, FIoBuffer&& Chunk)
-			{
-				TRACE_CPUPROFILER_EVENT_SCOPE(FOnDemandContentInstaller::Callback);
+		const int32 HttpRetryCount = 2;
+		FIoHttpBatch Batch = FHttpIoDispatcher::NewBatch();
+		for (FChunkHttpRequestHandle& HttpRequest : InstallRequest.HttpRequestHandles)
+		{
+			TRACE_CPUPROFILER_EVENT_SCOPE(FOnDemandContentInstaller::IssueRequest);
+			FSharedOnDemandContainer& Container		= InstallRequest.ResolvedChunks[HttpRequest.ContainerIndex].Container;
+			FOnDemandChunkInfo ChunkInfo			= FOnDemandChunkInfo(Container, Container->ChunkEntries[HttpRequest.EntryIndex]);
+			const uint8 HttpCategory				= 1; // see FOnDemandHttpThread
+
+			// Setting HTTP flags to None to not use the HTTP cache
+			FIoHttpOptions HttpOptions = FIoHttpOptions(Request.Priority, HttpRetryCount, EIoHttpFlags::None);
+			HttpOptions.SetCategory(HttpCategory);
+
+			HttpRequest.IoHandle = Batch.Get(
+				ChunkInfo.HostGroupName(),
+				ChunkInfo.RelativeUrl(),
+				FIoHttpHeaders(),
+				HttpOptions,
+				ChunkInfo.ChunkEntry().Hash,
+			[this, &Request, &HttpRequest](FIoHttpResponse&& HttpResponse)
+			{ 
 				InstallerPipe.Launch(
 					TEXT("ProcessIoStoreOnDemandDownloadedChunk"),
-					[this, &Request, &HttpRequest, HttpStatusCode, ErrorReason = FString(ErrorReason), Chunk = MoveTemp(Chunk)]() mutable
+					[this, &Request, &HttpRequest, HttpResponse = MoveTemp(HttpResponse)]() mutable
 					{
+						FIoBuffer Chunk = HttpResponse.GetBody();
 						ProcessDownloadedChunk(
 							Request,
 							HttpRequest,
-							HttpStatusCode,
-							MoveTemp(ErrorReason),
+							HttpResponse.GetStatusCode(),
+							FString(),
 							MoveTemp(Chunk));
 					},
 					UE::Tasks::ETaskPriority::BackgroundLow);
-			}, EHttpRequestType::Installed);
+			});
+		}
+		Batch.Issue();
+	}
+	else
+	{
+		for (FChunkHttpRequestHandle& HttpRequest : InstallRequest.HttpRequestHandles)
+		{
+			TRACE_CPUPROFILER_EVENT_SCOPE(FOnDemandContentInstaller::IssueRequest);
+			FSharedOnDemandContainer& Container		= InstallRequest.ResolvedChunks[HttpRequest.ContainerIndex].Container;
+			const FOnDemandChunkEntry& ChunkEntry	= Container->ChunkEntries[HttpRequest.EntryIndex];
+
+			HttpRequest.Handle = HttpClient.IssueRequest(
+				FOnDemandChunkInfo(Container, ChunkEntry),
+				FIoOffsetAndLength(),
+				Request.Priority,
+				[this, &Request, &HttpRequest](uint32 HttpStatusCode, FStringView ErrorReason, FIoBuffer&& Chunk)
+				{
+					TRACE_CPUPROFILER_EVENT_SCOPE(FOnDemandContentInstaller::Callback);
+					InstallerPipe.Launch(
+						TEXT("ProcessIoStoreOnDemandDownloadedChunk"),
+						[this, &Request, &HttpRequest, HttpStatusCode, ErrorReason = FString(ErrorReason), Chunk = MoveTemp(Chunk)]() mutable
+						{
+							ProcessDownloadedChunk(
+								Request,
+								HttpRequest,
+								HttpStatusCode,
+								MoveTemp(ErrorReason),
+								MoveTemp(Chunk));
+						},
+						UE::Tasks::ETaskPriority::BackgroundLow);
+				}, EHttpRequestType::Installed);
+		}
 	}
 	
 	InstallRequest.bHttpRequestsIssued = true;
@@ -877,7 +924,12 @@ void FOnDemandContentInstaller::ProcessDownloadedChunk(
 			int32 NumCancelled = 0;
 			for (FChunkHttpRequestHandle& PendingHttpRequest : InstallRequest.HttpRequestHandles)
 			{
-				if (PendingHttpRequest.Handle != nullptr)
+				if (PendingHttpRequest.IoHandle.IsValid())
+				{
+					PendingHttpRequest.IoHandle.Cancel();
+					++NumCancelled;
+				}
+				else if (PendingHttpRequest.Handle != nullptr)
 				{
 					HttpClient.CancelRequest(PendingHttpRequest.Handle);
 					++NumCancelled;
diff --git a/Engine/Source/Runtime/Experimental/IoStore/OnDemand/Private/OnDemandContentInstaller.h b/Engine/Source/Runtime/Experimental/IoStore/OnDemand/Private/OnDemandContentInstaller.h
--- a/Engine/Source/Runtime/Experimental/IoStore/OnDemand/Private/OnDemandContentInstaller.h
+++ b/Engine/Source/Runtime/Experimental/IoStore/OnDemand/Private/OnDemandContentInstaller.h
@@ -8,6 +8,7 @@
 #include "IO/PackageId.h"
 #include "IO/IoStoreOnDemandInternals.h"
 #include "IO/IoAllocators.h"
+#include "IO/HttpIoDispatcher.h"
 #include "Misc/TVariant.h"
 #include "OnDemandIoStore.h"
 #include "Tasks/Pipe.h"
@@ -59,9 +60,10 @@ class FOnDemandContentInstaller
 {
 	struct FChunkHttpRequestHandle
 	{
-		void*	Handle = nullptr;
-		int32	ContainerIndex = INDEX_NONE;
-		int32	EntryIndex = INDEX_NONE;
+		void*			Handle = nullptr;
+		FIoHttpRequest	IoHandle;
+		int32			ContainerIndex = INDEX_NONE;
+		int32			EntryIndex = INDEX_NONE;
 	};
 
 	struct FRequest
diff --git a/Engine/Source/Runtime/Experimental/IoStore/OnDemand/Private/OnDemandHttpIoDispacher.cpp b/Engine/Source/Runtime/Experimental/IoStore/OnDemand/Private/OnDemandHttpIoDispacher.cpp
new file mode 100644
--- /dev/null
+++ b/Engine/Source/Runtime/Experimental/IoStore/OnDemand/Private/OnDemandHttpIoDispacher.cpp
@@ -0,0 +1,1057 @@
+// Copyright Epic Games, Inc. All Rights Reserved.
+
+#include "OnDemandHttpIoDispatcher.h"
+
+#include "Async/Mutex.h"
+#include "Async/UniqueLock.h"
+#include "AtomicFlags.h"
+#include "Containers/AnsiString.h"
+#include "Experimental/ConcurrentLinearAllocator.h"
+#include "HAL/Platform.h"
+#include "HAL/Runnable.h"
+#include "HAL/RunnableThread.h"
+#include "IasCache.h"
+#include "IasHostGroup.h"
+#include "IO/IoContainers.h"
+#include "IO/IoStoreOnDemand.h"
+#include "IO/HttpIoDispatcher.h"
+#include "Logging/LogVerbosity.h"
+#include "Misc/CommandLine.h"
+#include "Misc/EnumClassFlags.h"
+#include "Misc/Fork.h"
+#include "Misc/SingleThreadRunnable.h"
+#include "OnDemandHttpClient.h"
+#include "OnDemandHttpThread.h" // EHttpRequestType
+#include "Statistics.h"
+
+#include <atomic>
+
+namespace UE::IoStore
+{
+////////////////////////////////////////////////////////////////////////////////
+extern int32 GIasHttpTimeOutMs;
+extern int32 GOnDemandBackendThreadPriorityIndex; 
+
+bool GHttpCacheEnabled = true;
+static FAutoConsoleVariableRef CVar_HttpCacheEnabled(
+	TEXT("iax.HttpCacheEnabled"),
+	GHttpCacheEnabled,
+	TEXT("")
+);
+
+namespace HttpIoDispatcher
+{
+////////////////////////////////////////////////////////////////////////////////
+static FIoHash GetCacheKey(const FIoHash& ChunkHash, const FIoHttpRange& Range)
+{
+	FIoHashBuilder HashBuilder;
+	HashBuilder.Update(ChunkHash.GetBytes(), sizeof(FIoHash::ByteArray));
+	HashBuilder.Update(&Range, sizeof(FIoHttpRange));
+
+	return HashBuilder.Finalize();
+}
+
+////////////////////////////////////////////////////////////////////////////////
+static FIoHash GetCacheKey(const FName& HostGroup, const FIoRelativeUrl Url, const FIoHttpHeaders& Headers, const FIoHttpOptions& Options)
+{
+	FIoHashBuilder HashBuilder;
+
+	TCHAR NameString[NAME_SIZE];
+	HostGroup.GetPlainNameString(NameString);
+	HashBuilder.Update(NameString, NAME_SIZE);
+	HashBuilder.Update(Url.ToString(), Url.Len());
+	
+	TConstArrayView<FAnsiString> HeadersView = Headers.ToArrayView();
+	for (int32 Idx = 0; Idx < HeadersView.Num(); Idx += 2)
+	{
+		const FAnsiString Name	= HeadersView[Idx];
+		const FAnsiString Value = HeadersView[Idx + 1];
+		HashBuilder.Update(*Name, Name.Len());
+		HashBuilder.Update(*Value, Value.Len());
+	}
+
+	HashBuilder.Update(&Options.GetRange(), sizeof(FIoHttpRange));
+	return HashBuilder.Finalize();
+}
+
+////////////////////////////////////////////////////////////////////////////////
+struct FHttpRequestBlockAllocationTag : FDefaultBlockAllocationTag
+{
+	static constexpr uint32			BlockSize = 4 * 1024;
+	static constexpr bool			AllowOversizedBlocks = false;
+	static constexpr bool			RequiresAccurateSize = false;
+	static constexpr bool			InlineBlockAllocation = true;
+	static constexpr const char*	TagName = "IoHttpRequestLinear";
+
+	using Allocator = TBlockAllocationCache<BlockSize, FAlignedAllocator>;
+};
+
+////////////////////////////////////////////////////////////////////////////////
+enum class EHttpRequestFlags : uint32
+{
+	None				= 0,
+	Issued				= (1 << 0),
+	CacheInflight		= (1 << 1),
+	HttpQueued			= (1 << 2),
+	HttpInflight		= (1 << 3),
+	CancelRequested		= (1 << 4),
+	Completing			= (1 << 5),
+	Completed			= (1 << 6)
+};
+ENUM_CLASS_FLAGS(EHttpRequestFlags);
+
+////////////////////////////////////////////////////////////////////////////////
+struct FHttpRequest final
+	: TConcurrentLinearObject<FHttpRequest, FHttpRequestBlockAllocationTag>
+	, TIntrusiveListElement<FHttpRequest>
+{
+	using FFlags				= TAtomicFlags<EHttpRequestFlags>;
+	static uint32				NextSeqNo;
+
+	EHttpRequestType			GetHttpRequestType() const { return Options.GetCategory() == 1 ? EHttpRequestType::Installed : EHttpRequestType::Streaming; }
+	void						AddRef() { RefCount.fetch_add(1, std::memory_order_relaxed); }
+	inline void					Release();
+
+	FIoHash						ChunkHash;
+	FIoHash						CacheKey; //TODO make this smaller the IAS cache only uses 64 bit cache keys
+	FIoHttpRequestCompleted 	OnCompleted;
+	FIoHttpHeaders				Headers;
+	FIoHttpOptions				Options;
+	FIoBuffer					Buffer;
+	FIoRelativeUrl				RelativeUrl;
+	FHttpRequest*				Next = nullptr;
+	FName						HostGroupName;
+	EIoErrorCode				CacheStatus = EIoErrorCode::Unknown;
+	FFlags						Flags;
+	std::atomic<int32>			RefCount{0};
+	std::atomic<EIoErrorCode>	CompletionStatus{EIoErrorCode::Unknown};
+	uint32						SeqNo = ++NextSeqNo;
+};
+using FHttpRequestList = TIntrusiveList<FHttpRequest>;
+
+////////////////////////////////////////////////////////////////////////////////
+void FHttpRequest::Release()
+{
+	if (RefCount.fetch_sub(1, std::memory_order_acq_rel) == 1)
+	{
+		delete this;
+	}
+}
+
+////////////////////////////////////////////////////////////////////////////////
+static bool HttpRequestSortPredicate(const FHttpRequest& LHS, const FHttpRequest& RHS)
+{
+	if (LHS.Options.GetPriority() == RHS.Options.GetPriority())
+	{
+		return LHS.SeqNo < RHS.SeqNo;
+	}
+
+	return LHS.Options.GetPriority() > RHS.Options.GetPriority();
+}
+
+uint32 FHttpRequest::NextSeqNo = 0;
+
+////////////////////////////////////////////////////////////////////////////////
+FIoHttpRange GetTotalRange(FHttpRequest* Request)
+{
+	FIoHttpRange TotalRange;
+	for (FHttpRequest* It = Request; It != nullptr; It = It->Next)
+	{
+		// If the range is invalid we fetch the entire resource
+		if (It->Options.GetRange().IsValid() == false)
+		{
+			return FIoHttpRange();
+		}
+		TotalRange += It->Options.GetRange();
+	}
+
+	return TotalRange;
+}
+
+////////////////////////////////////////////////////////////////////////////////
+void DebugPrintRange(FHttpRequest* Requests, const UE::FIoHttpRange& TotalRange)
+{
+#if !UE_BUILD_SHIPPING 
+	if (Requests->Next == nullptr)
+	{
+		return; 
+	}
+
+	int32 Count = 0;
+	TStringBuilder<256> Sb;
+	for (FHttpRequest* It = Requests; It != nullptr; It = It->Next)
+	{
+		const UE::FIoHttpRange& Range = It->Options.GetRange();
+		Sb << TEXT("[") << Range.GetMin() << TEXT("-") << Range.GetMax() << TEXT("] ");
+		++Count;
+	}
+
+	Sb << TEXT("-> ") << TEXT("[") << TotalRange.GetMin() << TEXT("-") << TotalRange.GetMax() << TEXT("]");
+	UE_LOG(LogHttpIoDispatcher, VeryVerbose, TEXT("Merging %d requests, %s, total size = %.2lf KiB"),
+		Count, Sb.ToString(), double(TotalRange.GetSize()) / 1024.0);
+#endif
+}
+
+////////////////////////////////////////////////////////////////////////////////
+class FHttpQueue
+{
+public:
+	void			Reprioritize(FHttpRequest* Request, int32 NewPriority);
+	bool			Cancel(FHttpRequest* Request);
+	void			Enqueue(FHttpRequest* Request);
+	FHttpRequest*	Dequeue();
+
+private:
+	TArray<FHttpRequest*>			Heap;
+	TMap<FIoHash, FHttpRequestList>	ByChunkKey;
+	FMutex							Mutex;
+	bool							bReprioritize = false;
+};
+
+////////////////////////////////////////////////////////////////////////////////
+void FHttpQueue::Reprioritize(FHttpRequest* Request, int32 NewPriority)
+{
+	TUniqueLock Lock(Mutex);
+
+	const bool bGenericHttpRequest = Request->ChunkHash.IsZero();
+	if (bGenericHttpRequest)
+	{
+		if (int32 Index = Heap.Find(Request); Index != INDEX_NONE)
+		{
+			Request->Options.SetPriority(NewPriority);
+			bReprioritize = true;
+		}
+		return;
+	}
+
+	FHttpRequestList* List = ByChunkKey.Find(Request->ChunkHash);
+	if (List == nullptr)
+	{
+		ensure(Heap.Find(Request) == INDEX_NONE);
+		return;
+	}
+
+	for (FHttpRequest& Existing : *List)
+	{
+		if (NewPriority > Existing.Options.GetPriority())
+		{
+			Existing.Options.SetPriority(NewPriority);
+			bReprioritize = true;
+		}
+	}
+}
+
+bool FHttpQueue::Cancel(FHttpRequest* Request)
+{
+	TUniqueLock Lock(Mutex);
+
+	const bool bGenericHttpRequest = Request->ChunkHash.IsZero();
+	if (bGenericHttpRequest)
+	{
+		ensure(Request->Next == nullptr);
+		if (int32 Index = Heap.Find(Request); Index != INDEX_NONE)
+		{
+			Heap.HeapRemoveAt(Index, HttpRequestSortPredicate);
+			FOnDemandIoBackendStats::Get()->OnHttpDequeue(Request->GetHttpRequestType());
+			FOnDemandIoBackendStats::Get()->OnHttpCancel(Request->GetHttpRequestType());
+			return true;
+		}
+		return false;
+	}
+
+	FHttpRequestList* List = ByChunkKey.Find(Request->ChunkHash);
+	if (List == nullptr)
+	{
+		ensure(Heap.Find(Request) == INDEX_NONE);
+		return false;
+	}
+
+	const bool bRemoved = List->Remove(Request);
+	ensure(bRemoved);
+	ensure(Request->Next == nullptr);
+
+	// If the request was the first for this chunk we push any of the remaining requests to the heap
+	if (int32 Index = Heap.Find(Request); Index != INDEX_NONE)
+	{
+		Heap.HeapRemoveAt(Index, HttpRequestSortPredicate);
+		FOnDemandIoBackendStats::Get()->OnHttpDequeue(Request->GetHttpRequestType());
+		FOnDemandIoBackendStats::Get()->OnHttpCancel(Request->GetHttpRequestType());
+
+		if (!List->IsEmpty())
+		{
+			FHttpRequest* Head = List->PeekHead();
+			ensure(Heap.Find(Head) == INDEX_NONE);
+			Heap.HeapPush(Head, HttpRequestSortPredicate);
+			FOnDemandIoBackendStats::Get()->OnHttpEnqueue(Head->GetHttpRequestType());
+		}
+	}
+
+	if (List->IsEmpty())
+	{
+		ByChunkKey.Remove(Request->ChunkHash);
+	}
+
+	return true;
+}
+
+void FHttpQueue::Enqueue(FHttpRequest* Request)
+{
+	TUniqueLock Lock(Mutex);
+
+	Request->Flags.Remove(EHttpRequestFlags::CacheInflight);
+
+	const bool bGenericHttpRequest = Request->ChunkHash.IsZero();
+	if (bGenericHttpRequest)
+	{
+		Heap.HeapPush(Request, HttpRequestSortPredicate);
+		Request->Flags.Add(EHttpRequestFlags::HttpQueued);
+		FOnDemandIoBackendStats::Get()->OnHttpEnqueue(Request->GetHttpRequestType());
+		return;
+	}
+
+	if (FHttpRequestList* List = ByChunkKey.Find(Request->ChunkHash))
+	{
+		ensure(List->IsEmpty() == false);
+		FHttpRequest& Head = *List->PeekHead();
+		if (Request->Options.GetPriority() > Head.Options.GetPriority())
+		{
+			Head.Options.SetPriority(Request->Options.GetPriority());
+			bReprioritize = true;
+		}
+		else
+		{
+			Request->Options.SetPriority(Head.Options.GetPriority());
+		}
+		List->AddOrInsertBefore(Request, [](const FHttpRequest& New, const FHttpRequest& Existing)
+		{
+			return New.Options.GetRange().GetMin() < Existing.Options.GetRange().GetMin();
+		});
+	}
+	else
+	{
+		Heap.HeapPush(Request, HttpRequestSortPredicate);
+		ByChunkKey.Add(Request->ChunkHash, FHttpRequestList(Request));
+		FOnDemandIoBackendStats::Get()->OnHttpEnqueue(Request->GetHttpRequestType());
+	}
+
+	Request->Flags.Add(EHttpRequestFlags::HttpQueued);
+}
+
+FHttpRequest* FHttpQueue::Dequeue()
+{
+	TUniqueLock Lock(Mutex);
+
+	if (Heap.IsEmpty())
+	{
+		ensure(ByChunkKey.IsEmpty());
+		return nullptr;
+	}
+
+	if (bReprioritize)
+	{
+		bReprioritize = false;
+		Heap.Heapify(HttpRequestSortPredicate);
+	}
+
+	FHttpRequest* Next = nullptr;
+	Heap.HeapPop(Next, HttpRequestSortPredicate, EAllowShrinking::No);
+	if (Heap.IsEmpty())
+	{
+		Heap.Empty(16);
+	}
+
+	if (Next == nullptr)
+	{
+		return Next;
+	}
+
+	//TODO: Return N request up to some range limit
+	if (FHttpRequestList* List = ByChunkKey.Find(Next->ChunkHash))
+	{
+		ByChunkKey.Remove(Next->ChunkHash);
+		// The first request that was enqueued for the chunk is not necessarily the first request in the list
+		Next = List->PeekHead();
+		if (ByChunkKey.IsEmpty())
+		{
+			ByChunkKey.Empty(16);
+		}
+	}
+
+	for (FHttpRequest* It = Next; It != nullptr; It = It->Next)
+	{
+		It->Flags.Add(EHttpRequestFlags::HttpInflight);
+		It->Flags.Remove(EHttpRequestFlags::HttpQueued);
+	}
+
+	if (Next != nullptr)
+	{
+		FOnDemandIoBackendStats::Get()->OnHttpDequeue(Next->GetHttpRequestType());
+	}
+
+	return Next;
+}
+
+////////////////////////////////////////////////////////////////////////////////
+class FHttpDispatcher final
+	: public FRunnable
+	, public FSingleThreadRunnable
+	, public IOnDemandHttpIoDispatcher
+{
+public:
+									FHttpDispatcher(TUniquePtr<IIasCache>&& Cache);
+									~FHttpDispatcher();
+
+private:
+	// FRunnable
+	virtual bool					Init() override;
+	virtual uint32					Run() override;
+	virtual void					Stop() override;
+	virtual void					Exit() override;
+	//FSingleThreadRunnable
+	virtual void					Tick() override;
+	// IHttpIoDispatcher
+	virtual void					Shutdown() override;
+	virtual FIoStatus				RegisterHostGroup(const FName& HostGroupName, TConstArrayView<FAnsiString> HostNames, FAnsiStringView TestUrl) override;
+	virtual bool					IsHostGroupRegistered(const FName& HostGroup) override;
+	virtual bool					IsHostGroupOk(const FName& HostGroup) override;
+	virtual FHostGroupRegistered&	OnHostGroupRegistered() override { return HostGroupRegisterd; }
+	virtual FIoHttpRequestHandle	CreateRequest(
+										FIoHttpRequestHandle& First,
+										FIoHttpRequestHandle& Last,
+										const FName& HostGroupName,
+										const FIoRelativeUrl& RelativeUrl,
+										const FIoHttpOptions& Options,
+										FIoHttpHeaders&& Headers,
+										FIoHttpRequestCompleted&& OnCompleted,
+										const FIoHash* ChunkHash) override;
+	virtual void					IssueRequest(FIoHttpRequestHandle RequestHandle) override;
+	virtual void					ReleaseRequest(FIoHttpRequestHandle Handle) override;
+	virtual void					CancelRequest(FIoHttpRequestHandle Handle) override;
+	virtual void					UpdateRequestPriority(FIoHttpRequestHandle Handle, int32 NewPriority) override;
+	virtual EIoErrorCode			GetRequestStatus(FIoHttpRequestHandle Handle) override;
+
+	bool							TryReadFromCache(FHttpRequest* Request);
+	bool							ProcessHttpRequests();
+	void							CompleteHttpRequest(FHttpRequest& Request, FMultiEndpointHttpClientResponse&& HttpResponse, FIASHostGroup& HostGroup, const FIoHttpRange& TotalRange);
+	void							CompleteRequest(FHttpRequest& Request, FIoBuffer& Body, const FIoHttpRange& TotalRange, uint32 StatusCode, bool bCached);
+
+	TUniquePtr<IIasCache>					Cache;
+	TUniquePtr<FMultiEndpointHttpClient>	HttpClient;
+	TUniquePtr<FRunnableThread>				Thread;
+	FHttpQueue								HttpQueue;
+	FHostGroupRegistered					HostGroupRegisterd;
+	FEventRef								WakeUp;
+	FMutex									Mutex;
+	std::atomic_bool						bStopRequested{false};
+};
+
+static TUniquePtr<FHttpDispatcher> HttpDispatcher;
+
+////////////////////////////////////////////////////////////////////////////////
+FHttpDispatcher::FHttpDispatcher(TUniquePtr<IIasCache>&& InCache)
+	: Cache(MoveTemp(InCache))
+{
+	static constexpr int32 ThreadPriorities[5] =
+	{
+		EThreadPriority::TPri_Lowest,
+		EThreadPriority::TPri_BelowNormal,
+		EThreadPriority::TPri_SlightlyBelowNormal,
+		EThreadPriority::TPri_Normal,
+		EThreadPriority::TPri_AboveNormal
+	};
+
+	const int32 ThreadPriorityIndex			= FMath::Clamp(GOnDemandBackendThreadPriorityIndex, 0, (int32)UE_ARRAY_COUNT(ThreadPriorities) - 1);
+	EThreadPriority DesiredThreadPriority	= (EThreadPriority)ThreadPriorities[ThreadPriorityIndex];
+	EThreadPriority ThreadPriority			= EThreadPriority::TPri_Normal;
+	const uint32 StackSize					= 0; // Use default stack size
+	const uint64 ThreadAffinityMask			= FGenericPlatformAffinity::GetNoAffinityMask();
+	const EThreadCreateFlags CreateFlags	= EThreadCreateFlags::None;
+	const bool bAllowPreFork				= FParse::Param(FCommandLine::IsInitialized() ? FCommandLine::Get() : TEXT(""), TEXT("-Ias.EnableHttpThreadPreFork"));
+
+	Thread.Reset(FForkProcessHelper::CreateForkableThread(this, TEXT("IoService.Http"), StackSize, ThreadPriority, ThreadAffinityMask, CreateFlags, bAllowPreFork));
+};
+
+FHttpDispatcher::~FHttpDispatcher()
+{
+	Thread.Reset();
+}
+
+bool FHttpDispatcher::Init()
+{
+	HttpClient = FMultiEndpointHttpClient::Create(FMultiEndpointHttpClientConfig
+	{
+		.MaxConnectionCount		= GIasHttpConnectionCount,
+		.ReceiveBufferSize		= GIasHttpRecvBufKiB >= 0 ? GIasHttpRecvBufKiB << 10 : -1,
+		.MaxRetryCount			= GIasHttpRetryCount,
+		.TimeoutMs				= GIasHttpFailTimeOutMs,
+		.Redirects				= EHttpRedirects::Follow,
+		.bAllowChunkedTransfer	= true,
+		.LogCategory			= &LogHttpIoDispatcher,
+		.LogVerbosity			= ELogVerbosity::VeryVerbose
+	});
+
+	return HttpClient.IsValid();
+}
+
+uint32 FHttpDispatcher::Run()
+{
+	check(HttpClient.IsValid());
+
+	for (;;)
+	{
+		FHostGroupManager::Get().Tick(GIasHttpTimeOutMs, bStopRequested);
+
+		// TODO: It would be better to only update connections as they need it, consider doing this on
+		// hostgroup connect/disconnect events.
+		HttpClient->UpdateConnections();
+
+		if (ProcessHttpRequests())
+		{
+			continue;
+		}
+
+		if (bStopRequested.load(std::memory_order_relaxed))
+		{
+			break;
+		}
+
+		WakeUp->Wait(FHostGroupManager::Get().GetNumDisconnctedHosts() > 0 ? GIasHttpHealthCheckWaitTime : MAX_uint32);
+	}
+
+	return 0;
+}
+
+void FHttpDispatcher::Stop()
+{
+	bool bExpected = false;
+	if (bStopRequested.compare_exchange_strong(bExpected, true))
+	{
+		WakeUp->Trigger();
+	}
+}
+
+void FHttpDispatcher::Exit()
+{
+}
+
+void FHttpDispatcher::Tick()
+{
+	for (;;)
+	{
+		FHostGroupManager::Get().Tick(GIasHttpTimeOutMs, bStopRequested);
+
+		// TODO: It would be better to only update connections as they need it, consider doing this on
+		// hostgroup connect/disconnect events.
+		HttpClient->UpdateConnections();
+
+		if (ProcessHttpRequests() == false)
+		{
+			break;
+		}
+	}
+}
+
+void FHttpDispatcher::Shutdown()
+{
+	bool bExpected = false;
+	if (bStopRequested.compare_exchange_strong(bExpected, true))
+	{
+		WakeUp->Trigger();
+	}
+}
+
+FIoStatus FHttpDispatcher::RegisterHostGroup(const FName& HostGroupName, TConstArrayView<FAnsiString> HostNames, FAnsiStringView TestUrl)
+{
+	FIASHostGroup Existing = FHostGroupManager::Get().Find(HostGroupName);
+	if (Existing.GetName() == HostGroupName)
+	{
+		return FIoStatus::Ok;
+	}
+
+	TIoStatusOr<FIASHostGroup> NewGroup = FHostGroupManager::Get().Register(HostGroupName, TestUrl);
+	if (NewGroup.IsOk() == false)
+	{
+		return NewGroup.Status();
+	}
+
+#if !UE_BUILD_SHIPPING
+	TArray<FAnsiString> HostsOverride;
+	if (HostGroupName == FOnDemandHostGroup::DefaultName)
+	{
+		FString HostValue;
+		if (FParse::Value(FCommandLine::Get(), TEXT("Iax.DefaultHostGroup="), HostValue))
+		{
+			TArray<FString> HostArray;
+			HostValue.ParseIntoArray(HostArray, TEXT(","), true);
+			for (FString& Host : HostArray)
+			{
+				Host.TrimStartAndEndInline();
+				if (!Host.StartsWith(TEXT("http")))
+				{
+					Host = TEXT("http://") + Host;
+				}
+				HostsOverride.Add(FAnsiString(StringCast<ANSICHAR>(*Host)));
+			}
+			HostNames = HostsOverride; 
+		}
+	}
+#endif
+
+	if (FIoStatus Status = NewGroup.ConsumeValueOrDie().Resolve(HostNames); !Status.IsOk())
+	{
+		UE_LOG(LogHttpIoDispatcher, Warning, TEXT("Failed to create host group '%s'"), *HostGroupName.ToString());
+		return Status;
+	}
+
+	UE_LOG(LogHttpIoDispatcher, Log, TEXT("Registered new host group '%s'"), *HostGroupName.ToString());
+	for (const FAnsiString& HostName : HostNames) 
+	{
+		UE_LOG(LogHttpIoDispatcher, Log, TEXT("\t* %s"), *FString(HostName));
+	}
+
+	if (HostGroupRegisterd.IsBound())
+	{
+		HostGroupRegisterd.Broadcast(HostGroupName);
+	}
+
+	return FIoStatus::Ok;
+}
+
+bool FHttpDispatcher::IsHostGroupRegistered(const FName& HostGroup)
+{
+	FIASHostGroup Existing = FHostGroupManager::Get().Find(HostGroup);
+	return Existing.GetName() == HostGroup;
+}
+
+bool FHttpDispatcher::IsHostGroupOk(const FName& HostGroup)
+{
+	FIASHostGroup Existing = FHostGroupManager::Get().Find(HostGroup);
+	return Existing.IsConnected();
+}
+
+FIoHttpRequestHandle FHttpDispatcher::CreateRequest(
+	FIoHttpRequestHandle& First,
+	FIoHttpRequestHandle& Last,
+	const FName& HostGroupName,
+	const FIoRelativeUrl& RelativeUrl,
+	const FIoHttpOptions& Options,
+	FIoHttpHeaders&& Headers,
+	FIoHttpRequestCompleted&& OnCompleted,
+	const FIoHash* ChunkHash)
+	
+{
+	FHttpRequest* Request = new FHttpRequest();
+	Request->OnCompleted	= MoveTemp(OnCompleted);
+	Request->Headers		= MoveTemp(Headers);
+	Request->Options		= Options;
+	Request->RelativeUrl	= RelativeUrl;
+	Request->HostGroupName	= HostGroupName;
+	check(Request->Next == nullptr);
+
+	if (ChunkHash != nullptr)
+	{
+		Request->ChunkHash	= *ChunkHash;
+		Request->CacheKey	= GetCacheKey(*ChunkHash, Options.GetRange());
+	}
+	else
+	{
+		Request->ChunkHash	= FIoHash::Zero;
+		Request->CacheKey	= GetCacheKey(HostGroupName, Request->RelativeUrl, Request->Headers, Request->Options);
+	}
+
+	if (Last == 0)
+	{
+		ensure(First == 0);
+		First = UPTRINT(Request);
+	}
+	else
+	{
+		FHttpRequest* BatchLast = reinterpret_cast<FHttpRequest*>(Last);
+		BatchLast->Next = Request;
+	}
+	Last = UPTRINT(Request);
+
+	Request->AddRef();
+	return reinterpret_cast<FIoHttpRequestHandle>(Request);
+}
+
+void FHttpDispatcher::IssueRequest(FIoHttpRequestHandle RequestHandle)
+{
+	ensure(RequestHandle != 0);
+	if (RequestHandle == 0)
+	{
+		return;
+	}
+
+	bool bWakeUp				= false;
+	FHttpRequest* NextRequest	= reinterpret_cast<FHttpRequest*>(RequestHandle);
+
+	while (NextRequest != nullptr)
+	{
+		FHttpRequest* Request	= NextRequest;
+		NextRequest				= Request->Next;
+		Request->Next			= nullptr;
+
+		Request->AddRef();
+		Request->Flags.Add(EHttpRequestFlags::Issued);
+		if (TryReadFromCache(Request) == false)
+		{
+			HttpQueue.Enqueue(Request);
+			bWakeUp = true;
+		}
+	}
+
+	if (bWakeUp)
+	{
+		WakeUp->Trigger();
+	}
+}
+
+void FHttpDispatcher::CancelRequest(FIoHttpRequestHandle Handle)
+{
+	ensure(Handle != 0);
+	if (Handle == 0)
+	{
+		return;
+	}
+
+	FHttpRequest* Request = reinterpret_cast<FHttpRequest*>(Handle);
+
+	const EHttpRequestFlags CurrentFlags = Request->Flags.Get();
+	if (EnumHasAnyFlags(CurrentFlags, EHttpRequestFlags::CancelRequested | EHttpRequestFlags::Completing | EHttpRequestFlags::Completed))
+	{
+		return;
+	}
+
+	if (Request->Flags.TrySet(CurrentFlags | EHttpRequestFlags::CancelRequested) == false)
+	{
+		return;
+	}
+
+	UE_LOG(LogHttpIoDispatcher, Verbose, TEXT("Cancelling request, SeqNo=%u"), Request->SeqNo);
+
+	if (EnumHasAnyFlags(CurrentFlags, EHttpRequestFlags::CacheInflight))
+	{
+		Cache->Cancel(Request->Buffer);
+	}
+
+	if (HttpQueue.Cancel(Request))
+	{
+		UE::Tasks::Launch(UE_SOURCE_LOCATION, [this, Request]
+		{
+			CompleteRequest(*Request, Request->Buffer, FIoHttpRange(), 0, false);
+		}, UE::Tasks::ETaskPriority::BackgroundLow);
+	}
+}
+
+void FHttpDispatcher::UpdateRequestPriority(FIoHttpRequestHandle Handle, int32 NewPriority)
+{
+	ensure(Handle != 0);
+	if (Handle == 0)
+	{
+		return;
+	}
+
+	FHttpRequest* Request = reinterpret_cast<FHttpRequest*>(Handle);
+	UE_LOG(LogHttpIoDispatcher, Verbose, TEXT("Updating request priority, SeqNo=%u, Priority=%d, NewPriority=%d"),
+		Request->SeqNo, Request->Options.GetPriority(), NewPriority);
+	HttpQueue.Reprioritize(Request, NewPriority);
+}
+
+EIoErrorCode FHttpDispatcher::GetRequestStatus(FIoHttpRequestHandle Handle)
+{
+	if (Handle != 0)
+	{
+		return reinterpret_cast<FHttpRequest*>(Handle)->CompletionStatus.load(std::memory_order_relaxed);
+	}
+
+	return EIoErrorCode::InvalidCode;
+}
+
+bool FHttpDispatcher::TryReadFromCache(FHttpRequest* Request)
+{
+	if (!GHttpCacheEnabled || !Cache.IsValid())
+	{
+		return false;
+	}
+
+	if (EnumHasAnyFlags(Request->Options.GetFlags(), EIoHttpFlags::ReadCache) == false)
+	{
+		return false;
+	}
+
+	Request->Flags.Add(EHttpRequestFlags::CacheInflight);
+	const EIoErrorCode CacheStatus = Cache->Get(Request->CacheKey, Request->Buffer);
+	if (CacheStatus == EIoErrorCode::Ok)
+	{
+		FOnDemandIoBackendStats::Get()->OnCacheGet(Request->Buffer.GetSize());
+		UE::Tasks::Launch(UE_SOURCE_LOCATION, [this, Request]
+		{
+			CompleteRequest(*Request, Request->Buffer, FIoHttpRange(), 200, true);
+		}, UE::Tasks::ETaskPriority::BackgroundLow);
+
+		return true;
+	}
+
+	if (CacheStatus == EIoErrorCode::FileNotOpen)
+	{
+		UE::Tasks::FTaskEvent OnReadyEvent(TEXT("HttpCacheReadComplete"));
+		Launch(UE_SOURCE_LOCATION, [this, Request]
+		{
+			if (Request->CacheStatus == EIoErrorCode::Ok)
+			{
+				FOnDemandIoBackendStats::Get()->OnCacheGet(Request->Buffer.GetSize());
+				CompleteRequest(*Request, Request->Buffer, FIoHttpRange(), 200, true);
+			}
+			else if (Request->CacheStatus == EIoErrorCode::Cancelled)
+			{
+				CompleteRequest(*Request, Request->Buffer, FIoHttpRange(), 0, false);
+			}
+			else
+			{
+				if (Request->CacheStatus == EIoErrorCode::ReadError)
+				{
+					FOnDemandIoBackendStats::Get()->OnCacheError();
+				}
+
+				HttpQueue.Enqueue(Request);
+				WakeUp->Trigger();
+			}
+		}, OnReadyEvent, UE::Tasks::ETaskPriority::BackgroundLow);
+
+		Cache->Materialize(Request->CacheKey, Request->Buffer, Request->CacheStatus, MoveTemp(OnReadyEvent));
+
+		return true;
+	}
+
+	return false;
+}
+
+void FHttpDispatcher::ReleaseRequest(FIoHttpRequestHandle Handle)
+{
+	ensure(Handle != 0);
+	FHttpRequest* NextRequest = reinterpret_cast<FHttpRequest*>(Handle);
+	while (NextRequest != nullptr)
+	{
+		FHttpRequest* Request	= NextRequest;
+		NextRequest				= Request->Next;
+		Request->Next			= nullptr;
+
+#if !UE_BUILD_SHIPPING 
+		const EHttpRequestFlags Flags = Request->Flags.Get();
+		ensureMsgf(Flags == EHttpRequestFlags::None || EnumHasAnyFlags(Flags, EHttpRequestFlags::Completing | EHttpRequestFlags::Completed),
+			TEXT("The HTTP request handle must not be released before the completion callback has been triggered."));
+#endif
+		Request->Release();
+	}
+}
+
+bool FHttpDispatcher::ProcessHttpRequests()
+{
+	const int32 MaxConcurrentRequests = FMath::Min(GIasHttpConcurrentRequests, 32);
+	int32		NumConcurrentRequests = 0;
+
+	FHttpRequest* NextHttpRequest = HttpQueue.Dequeue();
+	if (NextHttpRequest == nullptr)
+	{
+		return false;
+	}
+
+	TAnsiStringBuilder<128> Url;
+	TAnsiStringBuilder<41>	HashString;
+
+	while (NextHttpRequest)
+	{
+		while (NextHttpRequest)
+		{
+			FHttpRequest* HttpRequest	= NextHttpRequest;
+			NextHttpRequest				= HttpQueue.Dequeue();
+
+			ensure(!HttpRequest->Flags.HasAny(EHttpRequestFlags::HttpQueued | EHttpRequestFlags::CacheInflight));
+			FOnDemandIoBackendStats::Get()->OnHttpRequestStarted();
+
+			FIASHostGroup HostGroup = FHostGroupManager::Get().Find(HttpRequest->HostGroupName);
+			if (HostGroup.IsConnected() == false)
+			{
+				// Technically this request is being skipped because of a pre-existing error. It is not
+				// an error itself and it is not being canceled by higher level code. However we do not
+				// currently have a statistic for that and we have to call one of the existing types in
+				// order to correctly reduce the pending count.
+				FOnDemandIoBackendStats::Get()->OnHttpCancel(HttpRequest->GetHttpRequestType());
+				CompleteRequest(*HttpRequest, HttpRequest->Buffer, FIoHttpRange(), 0, false);
+				continue;
+			}
+
+			Url.Reset();
+			HashString.Reset();
+
+			TArray<FAnsiString> Headers;
+			FIoHttpRange		TotalRange;
+
+			const bool bGenericHttpRequest = HttpRequest->ChunkHash.IsZero();
+			if (bGenericHttpRequest)
+			{
+				ensure(HttpRequest->Next == nullptr);
+				Url << *HttpRequest->RelativeUrl;
+				Headers = MoveTemp(HttpRequest->Headers).ToArray();
+			}
+			else
+			{
+				// I/O store chunk request 
+				HashString << HttpRequest->ChunkHash;
+				Url << *HttpRequest->RelativeUrl << "/" << HashString.ToView().Left(2) << "/" << HashString << ANSITEXTVIEW(".iochunk");
+				TotalRange = GetTotalRange(HttpRequest);
+				//DebugPrintRange(HttpRequest, TotalRange);
+			}
+
+			NumConcurrentRequests++;
+			HttpClient->Get(HostGroup.GetUnderlyingHostGroup(), Url, TotalRange.ToOffsetAndLength(), MoveTemp(Headers),
+				[this, HttpRequest, TotalRange, &NumConcurrentRequests, HostGroup = MoveTemp(HostGroup)]
+				(FMultiEndpointHttpClientResponse&& HttpResponse) mutable
+				{
+					NumConcurrentRequests--;
+					UE::Tasks::Launch(
+						UE_SOURCE_LOCATION,
+						[this, HttpRequest, HttpResponse = MoveTemp(HttpResponse), TotalRange, HostGroup = MoveTemp(HostGroup)]() mutable
+						{
+							CompleteHttpRequest(*HttpRequest, MoveTemp(HttpResponse), HostGroup, TotalRange);
+						}, UE::Tasks::ETaskPriority::BackgroundLow);
+				});
+
+			if (NumConcurrentRequests >= MaxConcurrentRequests)
+			{
+				TRACE_CPUPROFILER_EVENT_SCOPE(FOnDemandHttpThread::TickHttpSaturated);
+				while (NumConcurrentRequests >= MaxConcurrentRequests) //-V654
+				{
+					HttpClient->Tick(MAX_uint32, GIasHttpRateLimitKiBPerSecond); 
+				}
+			}
+
+			if (!NextHttpRequest)
+			{
+				NextHttpRequest = HttpQueue.Dequeue();
+			}
+		} // Inner
+		
+		{
+			// Keep processing pending connections until all requests are completed or a new one is issued
+			TRACE_CPUPROFILER_EVENT_SCOPE(FOnDemandHttpThread::TickHttp);
+			while (HttpClient->Tick(GIasHttpPollTimeoutMs, GIasHttpRateLimitKiBPerSecond))
+			{
+				if (!NextHttpRequest)
+				{
+					NextHttpRequest = HttpQueue.Dequeue();
+				}
+
+				if (NextHttpRequest)
+				{
+					break;
+				}
+			}
+		}
+	} // Outer
+
+	return true;
+}
+
+void FHttpDispatcher::CompleteHttpRequest(FHttpRequest& Request, FMultiEndpointHttpClientResponse&& HttpResponse, FIASHostGroup& HostGroup, const FIoHttpRange& TotalRange)
+{
+	FOnDemandIoBackendStats::Get()->OnHttpRequestCompleted();
+
+	if (HttpResponse.RetryCount > 0)
+	{
+		FOnDemandIoBackendStats::Get()->OnHttpRetry(Request.GetHttpRequestType());
+	}
+
+	if (HttpResponse.IsOk())
+	{
+		HostGroup.OnSuccessfulResponse();
+		FOnDemandIoBackendStats::Get()->OnHttpCdnCacheReply(Request.GetHttpRequestType(), HttpResponse.CDNCacheStatus);
+		FOnDemandIoBackendStats::Get()->OnHttpGet(Request.GetHttpRequestType(), HttpResponse.Body.DataSize(), HttpResponse.Sample.GetTotalMs());
+	}
+	else
+	{
+		FOnDemandIoBackendStats::Get()->OnHttpError(Request.GetHttpRequestType());
+		if (HostGroup.OnFailedResponse())
+		{
+			// A disconnect was triggered
+			FOnDemandIoBackendStats::Get()->OnHttpDisconnected();
+		}
+	}
+
+	CompleteRequest(Request, HttpResponse.Body, TotalRange, HttpResponse.StatusCode, false);
+}
+
+void FHttpDispatcher::CompleteRequest(FHttpRequest& Request, FIoBuffer& Body, const FIoHttpRange& TotalRange, uint32 StatusCode, bool bCached)
+{
+	FHttpRequest* NextRequest = &Request;
+	while (NextRequest != nullptr)
+	{
+		FHttpRequest* ToComplete	= NextRequest;
+		NextRequest					= ToComplete->Next;
+		ToComplete->Next			= nullptr;
+
+		ensure(!ToComplete->Flags.HasAny(EHttpRequestFlags::Completing | EHttpRequestFlags::Completed));
+		ToComplete->Flags.Add(EHttpRequestFlags::Completing);
+		ToComplete->Flags.Remove(EHttpRequestFlags::CacheInflight | EHttpRequestFlags::HttpInflight);
+
+		const FIoHttpRange& Range		= ToComplete->Options.GetRange();
+		EIoErrorCode CompletionStatus	= StatusCode > 199 && StatusCode < 300 ? EIoErrorCode::Ok : EIoErrorCode::ReadError;
+		FIoBuffer PartialBody			= Body;
+
+		if (ToComplete->Flags.HasAny(EHttpRequestFlags::CancelRequested))
+		{
+			CompletionStatus	= EIoErrorCode::Cancelled;
+			PartialBody			= FIoBuffer();
+			StatusCode			= 409; // Conflict
+		}
+
+		// If the request is NOT cached and specified a range we retrieve the partial body
+		if (!bCached && CompletionStatus == EIoErrorCode::Ok && Range.IsValid())
+		{
+			uint64			Offset	= Range.GetMin();
+			const uint64	Size	= Range.GetSize();
+
+			// If the total range is NOT valid we fetched the entire resource
+			if (TotalRange.IsValid())
+			{
+				ensure(Range.GetMin() >= TotalRange.GetMin());
+				ensure(Range.GetMax() <= TotalRange.GetMax());
+				Offset -= TotalRange.GetMin();
+			}
+
+			PartialBody = FIoBuffer(Body.GetView().Mid(Offset, Size), Body);
+		}
+
+		const bool bWriteCache = EnumHasAnyFlags(ToComplete->Options.GetFlags(), EIoHttpFlags::WriteCache)
+			&& !bCached && CompletionStatus == EIoErrorCode::Ok
+			&& GHttpCacheEnabled && Cache.IsValid()
+			&& PartialBody.GetSize() > 0;
+
+		if (bWriteCache)
+		{
+			if (FIoStatus Status = Cache->Put(ToComplete->CacheKey, PartialBody); !Status.IsOk())
+			{
+				UE_LOG(LogHttpIoDispatcher, Error, TEXT("Put HTTP cache failed, reason '%s'"), *Status.ToString());
+			}
+		}
+
+		const EIoHttpResponseFlags Flags	= bCached ? EIoHttpResponseFlags::Cached : EIoHttpResponseFlags::None;
+		FIoHttpRequestCompleted OnCompleted = MoveTemp(ToComplete->OnCompleted);
+		OnCompleted(FIoHttpResponse(PartialBody, CompletionStatus, StatusCode, Flags));
+
+		ToComplete->Flags.Add(EHttpRequestFlags::Completed);
+		ToComplete->CompletionStatus.store(CompletionStatus, std::memory_order_seq_cst);
+		ToComplete->Release();
+	}
+}
+
+} // namespace UE::IoStore::HttpIoDispatcher
+
+///////////////////////////////////////////////////////////////////////////////
+TSharedPtr<IOnDemandHttpIoDispatcher> MakeOnDemanHttpIoDispatcher(TUniquePtr<class IIasCache>&& Cache)
+{
+	return MakeShared<HttpIoDispatcher::FHttpDispatcher>(MoveTemp(Cache));
+}
+
+} // UE::IoStore
diff --git a/Engine/Source/Runtime/Experimental/IoStore/OnDemand/Private/OnDemandHttpIoDispatcher.h b/Engine/Source/Runtime/Experimental/IoStore/OnDemand/Private/OnDemandHttpIoDispatcher.h
new file mode 100644
--- /dev/null
+++ b/Engine/Source/Runtime/Experimental/IoStore/OnDemand/Private/OnDemandHttpIoDispatcher.h
@@ -0,0 +1,25 @@
+// Copyright Epic Games, Inc. All Rights Reserved.
+
+#pragma once
+
+#include "IO/HttpIoDispatcher.h"
+#include "Templates/SharedPointer.h"
+
+#define UE_API IOSTOREONDEMANDCORE_API
+
+
+namespace UE::IoStore
+{
+
+class IOnDemandHttpIoDispatcher
+	: public IHttpIoDispatcher
+{
+public:
+	virtual ~IOnDemandHttpIoDispatcher() = default;
+};
+
+TSharedPtr<IOnDemandHttpIoDispatcher> MakeOnDemanHttpIoDispatcher(TUniquePtr<class IIasCache>&& Cache);
+
+} // namespace UE::IoStore
+
+#undef UE_API
diff --git a/Engine/Source/Runtime/Experimental/IoStore/OnDemand/Private/OnDemandHttpIoDispatcherBackend.cpp b/Engine/Source/Runtime/Experimental/IoStore/OnDemand/Private/OnDemandHttpIoDispatcherBackend.cpp
new file mode 100644
--- /dev/null
+++ b/Engine/Source/Runtime/Experimental/IoStore/OnDemand/Private/OnDemandHttpIoDispatcherBackend.cpp
@@ -0,0 +1,362 @@
+// Copyright Epic Games, Inc. All Rights Reserved.
+
+#include "OnDemandHttpIoDispatcherBackend.h"
+
+#include "Containers/AnsiString.h"
+#include "HAL/PlatformTime.h"
+#include "IO/HttpIoDispatcher.h"
+#include "IO/IoAllocators.h"
+#include "IO/IoOffsetLength.h"
+#include "IO/IoStatus.h"
+#include "IO/IoStoreOnDemand.h"
+#include "OnDemandBackendStatus.h"
+#include "OnDemandIoStore.h"
+#include "OnDemandHttpThread.h" // EHttpRequestType
+#include "Statistics.h"
+
+namespace UE::IoStore
+{
+////////////////////////////////////////////////////////////////////////////////
+extern int32 GIasHttpRangeRequestMinSizeKiB;
+
+////////////////////////////////////////////////////////////////////////////////
+class FHttpIoDispatcherBackend
+	: public IOnDemandHttpIoDispatcherBackend
+{
+	using FSharedBackendCtx = TSharedPtr<const FIoDispatcherBackendContext>;
+
+	struct FResolvedRequest
+	{
+		FResolvedRequest(FIoRequestImpl& InRequest, FOnDemandChunkInfo&& InChunkInfo, const FIoOffsetAndLength& InChunkRange)
+			: DispatcherRequest(InRequest)
+			, ChunkInfo(MoveTemp(InChunkInfo))
+			, ChunkRange(InChunkRange)
+			, StartTimeCycles(FPlatformTime::Cycles64())
+		{
+			check(DispatcherRequest.BackendData == nullptr);
+			DispatcherRequest.BackendData = this;
+		}
+
+		static FResolvedRequest* Get(FIoRequestImpl& DispatcherRequest)
+		{
+			return reinterpret_cast<FResolvedRequest*>(DispatcherRequest.BackendData);
+		}
+
+		static FResolvedRequest* Create(FIoRequestImpl& DispatcherRequest, FOnDemandChunkInfo&& ChunkInfo, const FIoOffsetAndLength& ChunkRange)
+		{
+			return Allocator.Construct(DispatcherRequest, MoveTemp(ChunkInfo), ChunkRange);
+		}
+
+		static void Destroy(FResolvedRequest* ResolvedRequest)
+		{
+			check(ResolvedRequest->DispatcherRequest.BackendData == ResolvedRequest);
+			ResolvedRequest->DispatcherRequest.BackendData = nullptr;
+			Allocator.Destroy(ResolvedRequest);
+		}
+
+		using FAllocator		= TSingleThreadedSlabAllocator<FResolvedRequest, 64>;
+		static FAllocator		Allocator;
+
+		FIoRequestImpl&			DispatcherRequest;
+		FOnDemandChunkInfo		ChunkInfo;
+		FIoOffsetAndLength		ChunkRange;
+		FIoHttpRequest			HttpRequest;
+		uint64					StartTimeCycles;
+	};
+
+public:
+											FHttpIoDispatcherBackend(FOnDemandIoStore& IoStore);
+	//IOnDemandHttpIoDispatcherBackend
+	virtual void							SetOptionalBulkDataEnabled(bool bEnabled) override;
+	virtual void 							ReportAnalytics(TArray<FAnalyticsEventAttribute>& OutAnalyticsArray) const override;
+	// I/O dispatcher backend
+	virtual void							Initialize(TSharedRef<const FIoDispatcherBackendContext> Context) override;
+	virtual void							Shutdown() override;
+	virtual void							ResolveIoRequests(FIoRequestList Requests, FIoRequestList& OutUnresolved) override;
+	virtual FIoRequestImpl*					GetCompletedIoRequests() override;
+	virtual void							CancelIoRequest(FIoRequestImpl* Request) override;
+	virtual void							UpdatePriorityForIoRequest(FIoRequestImpl* Request) override;
+	virtual bool							DoesChunkExist(const FIoChunkId& ChunkId) const override;
+	virtual bool							DoesChunkExist(const FIoChunkId& ChunkId, const FIoOffsetAndLength& ChunkRange) const override; 
+	virtual TIoStatusOr<uint64> 			GetSizeForChunk(const FIoChunkId& ChunkId) const override;
+	virtual TIoStatusOr<FIoMappedRegion>	OpenMapped(const FIoChunkId& ChunkId, const FIoReadOptions& Options) override { return FIoStatus(EIoErrorCode::NotFound); }
+	virtual const TCHAR*					GetName() const { return TEXT("HttpIoDispatcherBackend"); }
+
+private:
+	void									CompleteRequest(FResolvedRequest& ResolvedRequst, FIoHttpResponse&& HttpResponse);
+	static void								LogRequest(FResolvedRequest& ResolvedRequest, FIoRequestImpl& DispacherRequest, bool bDecoded, bool bCached, uint64 DurationMs);
+
+	FBackendStatus			BackendStatus;
+	FOnDemandIoBackendStats	Stats;
+	FOnDemandIoStore&		IoStore;
+	FIoRequestList			CompletedRequests;
+	FMutex					Mutex;
+	FSharedBackendCtx		BackendContext;
+};
+
+FHttpIoDispatcherBackend::FResolvedRequest::FAllocator FHttpIoDispatcherBackend::FResolvedRequest::Allocator;
+
+////////////////////////////////////////////////////////////////////////////////
+FHttpIoDispatcherBackend::FHttpIoDispatcherBackend(FOnDemandIoStore& InIoStore)
+	: Stats(BackendStatus)
+	, IoStore(InIoStore)
+{
+	BackendStatus.SetHttpEnabled(true);
+	BackendStatus.SetCacheEnabled(true);
+}
+
+void FHttpIoDispatcherBackend::SetOptionalBulkDataEnabled(bool bEnabled)
+{
+	BackendStatus.SetHttpOptionalBulkEnabled(bEnabled);
+}
+
+void FHttpIoDispatcherBackend::ReportAnalytics(TArray<FAnalyticsEventAttribute>& OutAnalyticsArray) const
+{
+	Stats.ReportGeneralAnalytics(OutAnalyticsArray);
+	Stats.ReportEndPointAnalytics(OutAnalyticsArray);
+}
+
+void FHttpIoDispatcherBackend::Initialize(TSharedRef<const FIoDispatcherBackendContext> Context)
+{
+	BackendContext = Context;
+}
+
+void FHttpIoDispatcherBackend::Shutdown()
+{
+	BackendContext.Reset();
+}
+
+void FHttpIoDispatcherBackend::ResolveIoRequests(FIoRequestList Requests, FIoRequestList& OutUnresolved)
+{
+	const int32 HttpRetryCount = 2;
+
+	FIoHttpBatch Batch = FHttpIoDispatcher::NewBatch();
+	while (FIoRequestImpl* Request = Requests.PopHead())
+	{
+		FOnDemandChunkInfo ChunkInfo = IoStore.GetStreamingChunkInfo(Request->ChunkId);
+		if (!ChunkInfo.IsValid())
+		{
+			OutUnresolved.AddTail(Request);
+			continue;
+		}
+
+		if (FHttpIoDispatcher::IsHostGroupOk(ChunkInfo.HostGroupName()) == false)
+		{
+			OutUnresolved.AddTail(Request);
+			continue;
+		}
+
+		if (BackendStatus.IsHttpEnabled(Request->ChunkId.GetChunkType()) == false)
+		{
+			OutUnresolved.AddTail(Request);
+			continue;
+		}
+
+		const uint64 ResolvedSize = FMath::Min<uint64>(Request->Options.GetSize(), ChunkInfo.RawSize() - Request->Options.GetOffset());
+		Request->Options.SetRange(Request->Options.GetOffset(), ResolvedSize);
+		
+		FIoOffsetAndLength ChunkRange;
+		const uint64 ChunkSize = ChunkInfo.ChunkEntry().GetDiskSize(); // AES aligned encoded size
+		if (ChunkSize <= (uint64(GIasHttpRangeRequestMinSizeKiB) << 10))
+		{
+			ChunkRange = FIoOffsetAndLength(0, ChunkSize);
+		}
+		else
+		{
+			ChunkRange = FIoChunkEncoding::GetChunkRange(
+				ChunkInfo.RawSize(),
+				ChunkInfo.BlockSize(),
+				ChunkInfo.Blocks(),
+				Request->Options.GetOffset(),
+				ResolvedSize).ConsumeValueOrDie();
+		}
+
+		const FIoHttpRange HttpRange = FIoHttpRange::FromOffsetAndLength(ChunkRange);
+		FResolvedRequest* ResolvedRequest = FResolvedRequest::Create(*Request, MoveTemp(ChunkInfo), ChunkRange);
+		ResolvedRequest->HttpRequest = Batch.Get(
+			ResolvedRequest->ChunkInfo.HostGroupName(),
+			ResolvedRequest->ChunkInfo.RelativeUrl(),
+			FIoHttpHeaders(),
+			FIoHttpOptions(Request->Priority, HttpRetryCount, EIoHttpFlags::Default, HttpRange),
+			ResolvedRequest->ChunkInfo.ChunkEntry().Hash,
+			[this, ResolvedRequest](FIoHttpResponse&& HttpResponse)
+			{
+				// Callbacks are always triggered from the task pool
+				CompleteRequest(*ResolvedRequest, MoveTemp(HttpResponse));
+			});
+	}
+
+	Batch.Issue();
+}
+
+FIoRequestImpl* FHttpIoDispatcherBackend::GetCompletedIoRequests()
+{
+	FIoRequestList Out;
+	{
+		TUniqueLock Lock(Mutex);
+		Out = MoveTemp(CompletedRequests);
+		CompletedRequests = FIoRequestList();
+	}
+
+	for (FIoRequestImpl& R : Out)
+	{
+		FResolvedRequest::Destroy(reinterpret_cast<FResolvedRequest*>(R.BackendData));
+		check(R.BackendData == nullptr);
+	}
+
+	return Out.GetHead();
+}
+
+void FHttpIoDispatcherBackend::CancelIoRequest(FIoRequestImpl* Request)
+{
+	if (FResolvedRequest* Resolved = FResolvedRequest::Get(*Request))
+	{
+		Resolved->HttpRequest.Cancel();
+	}
+}
+
+void FHttpIoDispatcherBackend::UpdatePriorityForIoRequest(FIoRequestImpl* Request)
+{
+	if (FResolvedRequest* Resolved = FResolvedRequest::Get(*Request))
+	{
+		Resolved->HttpRequest.UpdatePriorty(Request->Priority);
+	}
+}
+
+bool FHttpIoDispatcherBackend::DoesChunkExist(const FIoChunkId& ChunkId) const
+{
+	const TIoStatusOr<uint64> ChunkSize = GetSizeForChunk(ChunkId);
+	return ChunkSize.IsOk();
+}
+
+bool FHttpIoDispatcherBackend::DoesChunkExist(const FIoChunkId& ChunkId, const FIoOffsetAndLength& ChunkRange) const
+{
+	return DoesChunkExist(ChunkId);
+}
+
+TIoStatusOr<uint64> FHttpIoDispatcherBackend::GetSizeForChunk(const FIoChunkId& ChunkId) const
+{
+	const FOnDemandChunkInfo ChunkInfo = IoStore.GetStreamingChunkInfo(ChunkId);
+	if (!ChunkInfo.IsValid() || !BackendStatus.IsHttpEnabled(ChunkId.GetChunkType()))
+	{
+		return FIoStatus(EIoErrorCode::UnknownChunkID);
+	}
+
+	return ChunkInfo.RawSize();
+}
+
+void FHttpIoDispatcherBackend::CompleteRequest(FResolvedRequest& ResolvedRequest, FIoHttpResponse&& HttpResponse)
+{
+	FIoBuffer Chunk						= HttpResponse.GetBody();
+	FOnDemandChunkInfo& ChunkInfo		= ResolvedRequest.ChunkInfo;
+	FIoRequestImpl& DispatcherRequest	= ResolvedRequest.DispatcherRequest;
+
+	bool bDecoded = false;
+	if (HttpResponse.IsOk() && Chunk.GetSize() > 0)
+	{
+		FIoChunkDecodingParams Params;
+		Params.EncryptionKey				= ChunkInfo.EncryptionKey();
+		Params.CompressionFormat			= ChunkInfo.CompressionFormat();
+		Params.BlockSize					= ChunkInfo.BlockSize();
+		Params.TotalRawSize					= ChunkInfo.RawSize();
+		Params.EncodedBlockSize				= ChunkInfo.Blocks();
+		Params.BlockHash					= ChunkInfo.BlockHashes();
+		Params.EncodedOffset				= ResolvedRequest.ChunkRange.GetOffset();
+		Params.RawOffset					= DispatcherRequest.Options.GetOffset();
+
+		DispatcherRequest.CreateBuffer(DispatcherRequest.Options.GetSize());
+
+		const EIoDecodeFlags Options = EnumHasAnyFlags(DispatcherRequest.Options.GetFlags(), EIoReadOptionsFlags::HardwareTargetBuffer)
+			? EIoDecodeFlags::WriteOnly : EIoDecodeFlags::None;
+		bDecoded = FIoChunkEncoding::Decode(Params, Chunk.GetView(), DispatcherRequest.GetBuffer().GetMutableView(), Options);
+	}
+
+	const uint64 DurationMs = ResolvedRequest.StartTimeCycles > 0
+		? uint64(FPlatformTime::ToMilliseconds64(FPlatformTime::Cycles64() - ResolvedRequest.StartTimeCycles))
+		: 0;
+
+	if (bDecoded)
+	{
+		Stats.OnIoRequestComplete(DispatcherRequest.GetBuffer().GetSize(), DurationMs);
+	}
+	else
+	{
+		DispatcherRequest.SetFailed();
+		Stats.OnIoRequestError();
+		if (HttpResponse.IsCached())
+		{
+			Stats.OnCacheDecodeError();
+		}
+		else
+		{
+			Stats.OnHttpDecodeError(EHttpRequestType::Streaming);
+		}
+	}
+	
+	LogRequest(ResolvedRequest, DispatcherRequest, bDecoded, HttpResponse.IsCached(), DurationMs);
+
+	{
+		TUniqueLock Lock(Mutex);
+		CompletedRequests.AddTail(&DispatcherRequest);
+	}
+
+	BackendContext->WakeUpDispatcherThreadDelegate.Execute();
+}
+
+void FHttpIoDispatcherBackend::LogRequest(FResolvedRequest& ResolvedRequest, FIoRequestImpl& DispacherRequest, bool bDecoded, bool bCached, uint64 DurationMs)
+{
+	const TCHAR* Prefix = [bDecoded, bCached]() -> const TCHAR*
+	{
+		if (bDecoded)
+		{
+			return bCached ? TEXT("io-cache") : TEXT("io-http ");
+		}
+		return bCached ? TEXT("io-cache-error") : TEXT("io-http-error ");
+	}();
+
+	auto PrioToString = [](int32 Prio) -> const TCHAR*
+	{
+		if (Prio < IoDispatcherPriority_Low)
+		{
+			return TEXT("Min");
+		}
+		if (Prio < IoDispatcherPriority_Medium)
+		{
+			return TEXT("Low");
+		}
+		if (Prio < IoDispatcherPriority_High)
+		{
+			return TEXT("Medium");
+		}
+		if (Prio < IoDispatcherPriority_Max)
+		{
+			return TEXT("High");
+		}
+
+		return TEXT("Max");
+	};
+
+	FOnDemandChunkInfo& ChunkInfo			= ResolvedRequest.ChunkInfo;
+	const uint64 UncompressedSize			= bDecoded ? DispacherRequest.GetBuffer().GetSize() : 0;
+	const FIoOffsetAndLength& ChunkRange	= ResolvedRequest.ChunkRange;
+	const uint64 ChunkSize					= ChunkInfo.ChunkEntry().GetDiskSize();
+
+	UE_LOG(LogIas, VeryVerbose, TEXT("%s: %5" UINT64_FMT "ms %5" UINT64_FMT "KiB[%8" UINT64_FMT "] % s: % s | Range: %" UINT64_FMT "-%" UINT64_FMT "/%" UINT64_FMT " (%.2f%%) | Prio: %s"),
+		Prefix,
+		DurationMs,
+		UncompressedSize >> 10,
+		DispacherRequest.Options.GetOffset(),
+		*LexToString(DispacherRequest.ChunkId),
+		*LexToString(ChunkInfo.Hash()),
+		ChunkRange.GetOffset(), (ChunkRange.GetOffset() + ChunkRange.GetLength() - 1), ChunkSize,
+		100.0f * (float(ChunkRange.GetLength()) / float(ChunkSize)),
+		PrioToString(DispacherRequest.Priority));
+}
+
+////////////////////////////////////////////////////////////////////////////////
+TSharedPtr<IOnDemandHttpIoDispatcherBackend> MakeOnDemandHttpIoDispatcherBackend(FOnDemandIoStore& IoStore)
+{
+	return TSharedPtr<IOnDemandHttpIoDispatcherBackend>(new FHttpIoDispatcherBackend(IoStore));
+}
+
+} // namespace UE::IoStore
diff --git a/Engine/Source/Runtime/Experimental/IoStore/OnDemand/Private/OnDemandHttpIoDispatcherBackend.h b/Engine/Source/Runtime/Experimental/IoStore/OnDemand/Private/OnDemandHttpIoDispatcherBackend.h
new file mode 100644
--- /dev/null
+++ b/Engine/Source/Runtime/Experimental/IoStore/OnDemand/Private/OnDemandHttpIoDispatcherBackend.h
@@ -0,0 +1,24 @@
+// Copyright Epic Games, Inc. All Rights Reserved.
+
+#pragma once
+
+#include "IO/IoDispatcherBackend.h"
+#include "Templates/SharedPointer.h"
+
+struct FAnalyticsEventAttribute;
+
+namespace UE::IoStore
+{
+
+class IOnDemandHttpIoDispatcherBackend
+	: public IIoDispatcherBackend
+{
+public:
+	virtual ~IOnDemandHttpIoDispatcherBackend() = default;
+	virtual void SetOptionalBulkDataEnabled(bool bEnabled) = 0;
+	virtual void ReportAnalytics(TArray<FAnalyticsEventAttribute>& OutAnalyticsArray) const = 0;
+};
+
+TSharedPtr<IOnDemandHttpIoDispatcherBackend> MakeOnDemandHttpIoDispatcherBackend(class FOnDemandIoStore& IoStore);
+
+} // namespace UE::IoStore
diff --git a/Engine/Source/Runtime/Experimental/IoStore/OnDemand/Private/OnDemandHttpThread.h b/Engine/Source/Runtime/Experimental/IoStore/OnDemand/Private/OnDemandHttpThread.h
--- a/Engine/Source/Runtime/Experimental/IoStore/OnDemand/Private/OnDemandHttpThread.h
+++ b/Engine/Source/Runtime/Experimental/IoStore/OnDemand/Private/OnDemandHttpThread.h
@@ -24,6 +24,17 @@ class FMultiEndpointHttpClient;
 struct FHttpRequest;
 struct FOnDemandChunkInfo;
 
+extern int32 GIasHttpHealthCheckWaitTime;
+extern int32 GOnDemandBackendThreadPriorityIndex;
+extern int32 GIasHttpConnectionCount;
+extern int32 GIasHttpRecvBufKiB;
+extern int32 GIasHttpRetryCount;
+extern int32 GIasHttpFailTimeOutMs; 
+extern int32 GIasHttpConcurrentRequests; 
+extern int32 GIasHttpRateLimitKiBPerSecond; 
+extern int32 GIasHttpPollTimeoutMs; 
+extern bool GIasHttpAllowChunkedXfer; 
+
 /** Describes the system making the request */
 enum class EHttpRequestType : uint8
 {
diff --git a/Engine/Source/Runtime/Experimental/IoStore/OnDemand/Private/OnDemandIoStore.cpp b/Engine/Source/Runtime/Experimental/IoStore/OnDemand/Private/OnDemandIoStore.cpp
--- a/Engine/Source/Runtime/Experimental/IoStore/OnDemand/Private/OnDemandIoStore.cpp
+++ b/Engine/Source/Runtime/Experimental/IoStore/OnDemand/Private/OnDemandIoStore.cpp
@@ -35,6 +35,8 @@
 #include "Modules/ModuleManager.h"
 #include "OnDemandConfig.h"
 #include "OnDemandHttpClient.h"
+#include "OnDemandHttpIoDispatcher.h"
+#include "OnDemandHttpIoDispatcherBackend.h"
 #include "OnDemandHttpThread.h"
 #include "OnDemandInstallCache.h"
 #include "OnDemandContentInstaller.h"
@@ -48,6 +50,14 @@
 #endif
 
 ///////////////////////////////////////////////////////////////////////////////
+bool GHttpIoDispatcherEnabled = false;
+static FAutoConsoleVariableRef CVar_HttpIoDispatcherEnabled(
+	TEXT("iax.HttpIoDispatcherEnabled"),
+	GHttpIoDispatcherEnabled,
+	TEXT("Whether the HTTP I/O dispatcher should be enabled."),
+	ECVF_ReadOnly
+);
+
 bool GIoStoreOnDemandTreatMissingChunksAsError = true;
 static FAutoConsoleVariableRef CVar_IoStoreOnDemandTreatMissingChunksAsError(
 	TEXT("iostore.TreatMissingOnDemandChunksAsError"),
@@ -240,6 +250,12 @@ FOnDemandIoStore::~FOnDemandIoStore()
 
 	Installer.Reset();
 	HttpClient.Reset();
+
+	if (FHttpIoDispatcher::IsInitialized())
+	{
+		FHttpIoDispatcher::OnHostGroupRegistered().RemoveAll(this);
+		FHttpIoDispatcher::Shutdown();
+	}
 }
 
 FIoStatus FOnDemandIoStore::Initialize()
@@ -353,6 +369,51 @@ FIoStatus FOnDemandIoStore::Initialize()
 		UE_LOG(LogIoStoreOnDemand, Log, TEXT("Install cache disabled"));
 	}
 
+	const bool bHttpIoDispatcherEnabled = (GHttpIoDispatcherEnabled || FParse::Param(FCommandLine::Get(), TEXT("HttpIoDispatcher")))
+		&& !FParse::Param(FCommandLine::Get(), TEXT("NoHttpIoDispatcher"));
+	if (bHttpIoDispatcherEnabled)
+	{
+		TUniquePtr<IIasCache> HttpCache;
+		FIasCacheConfig CacheConfig = Config::GetStreamingCacheConfig(FCommandLine::Get());
+
+		if (CacheConfig.DiskQuota > 0)
+		{
+			if (FPaths::HasProjectPersistentDownloadDir())
+			{
+				FString CacheDir = FPaths::ProjectPersistentDownloadDir();
+				HttpCache = MakeIasCache(*CacheDir, CacheConfig);
+
+				UE_CLOG(!HttpCache.IsValid(), LogIoStoreOnDemand, Warning, TEXT("HTTP cache disabled"));
+			}
+			else
+			{
+				UE_LOG(LogIoStoreOnDemand, Warning, TEXT("HTTP cache disabled - streaming only (project has no persistent download dir enabled for this platform)"));
+			}
+		}
+		else
+		{
+			UE_LOG(LogIoStoreOnDemand, Log, TEXT("HTTP cache disabled - streaming only (zero-quota)"));
+		}
+
+		TSharedPtr<IOnDemandHttpIoDispatcher> HttpDispatcher = MakeOnDemanHttpIoDispatcher(MoveTemp(HttpCache));
+		if (FIoStatus Status = FHttpIoDispatcher::Initialize(HttpDispatcher); !Status.IsOk())
+		{
+			UE_LOG(LogIoStoreOnDemand, Error, TEXT("Failed to initialize HTTP I/O dispatcher"));
+		}
+		
+		int32 HttpBackendPriority = -10;
+#if !UE_BUILD_SHIPPING
+		if (FParse::Param(FCommandLine::Get(), TEXT("Ias")))
+		{
+			UE_LOG(LogIoStoreOnDemand, Log, TEXT("Setting HTTP backend priority higher than file system backend"));
+			HttpBackendPriority = 10;
+		}
+#endif
+		HttpIoBackend = MakeOnDemandHttpIoDispatcherBackend(*this);
+		FIoDispatcher::Get().Mount(HttpIoBackend.ToSharedRef(), HttpBackendPriority);
+		FHttpIoDispatcher::OnHostGroupRegistered().AddRaw(this, &FOnDemandIoStore::OnHostGroupRegistered);
+	}
+
 	return InstallCacheStatus;
 }
 
@@ -397,9 +458,12 @@ FIoStatus FOnDemandIoStore::InitializePostHotfix()
 		}
 	}
 
-	if (FIoStatus Status = InitializeStreamingBackend(EndpointConfig); Status.IsOk() == false)
+	if (FHttpIoDispatcher::IsInitialized() == false)
 	{
-		return Status;
+		if (FIoStatus Status = InitializeStreamingBackend(EndpointConfig); Status.IsOk() == false)
+		{
+			return Status;
+		}
 	}
 
 	bool bUsePerContainerTocsConfigValue = false;
@@ -420,12 +484,23 @@ FIoStatus FOnDemandIoStore::InitializePostHotfix()
 	TOptional<FOnDemandMountArgs> MountArgs;
 	if (EndpointConfig.TocFilePath.IsEmpty() == false)
 	{
+		FOnDemandHostGroup HostGroup;
+		if (EndpointConfig.ServiceUrls.IsEmpty() == false)
+		{
+			TIoStatusOr<FOnDemandHostGroup> HostGroupStatus = FOnDemandHostGroup::Create(EndpointConfig.ServiceUrls);
+			if (HostGroupStatus.IsOk())
+			{
+				HostGroup = HostGroupStatus.ConsumeValueOrDie();
+			}
+		}
+		
 		if (bUsePerContainerTocs == false)
 		{
 			MountArgs.Emplace(FOnDemandMountArgs
 			{
 				.MountId = EndpointConfig.TocFilePath,
 				.FilePath = EndpointConfig.TocFilePath,
+				.HostGroup = HostGroup,
 				.Options = EOnDemandMountOptions::StreamOnDemand
 			});
 		}
@@ -1190,10 +1265,6 @@ FIoStatus FOnDemandIoStore::TickMountRequest(FMountRequest& MountRequest)
 				UE_LOG(LogIoStoreOnDemand, Log, TEXT("Overriding container flag 'StreamOnDemand' from mount option"));
 				EnumAddFlags(Container->Flags, EOnDemandContainerFlags::StreamOnDemand);
 			}
-
-			//TODO: Remove this special case for streaming containers
-			Container->HostGroupName	= FName("DefaultOnDemand");
-			Container->HostGroup		= FHostGroupManager::Get().Find(Container->HostGroupName);
 		}
 		else if (EnumHasAnyFlags(Args.Options, EOnDemandMountOptions::InstallOnDemand))
 		{
@@ -1207,75 +1278,21 @@ FIoStatus FOnDemandIoStore::TickMountRequest(FMountRequest& MountRequest)
 			{
 				EnumAddFlags(Container->Flags, EOnDemandContainerFlags::WithSoftReferences);
 			}
+		}
 
-			UE_CLOG(Container->HostGroupName != Args.HostGroupName, LogIoStoreOnDemand, Log, TEXT("Overriding container host group '%s' with '%s' from mount option"),
-				*Container->HostGroupName.ToString(), *Args.HostGroupName.ToString());
-
-			Container->HostGroupName	= Args.HostGroupName;
-			Container->HostGroup		= FHostGroupManager::Get().Find(Container->HostGroupName);
-
-			if (Container->HostGroup.IsResolved() == false && Args.HostGroup.IsEmpty() == false)
-			{
-				TIoStatusOr<FIASHostGroup> NewHostGroup = FHostGroupManager::Get().Register(Container->HostGroupName, Container->GetTestUrl());
-				if (NewHostGroup.IsOk())
-				{
-					Container->HostGroup = NewHostGroup.ConsumeValueOrDie();
-					TConstArrayView<FAnsiString> Hosts = Args.HostGroup.Hosts();
-#if !UE_BUILD_SHIPPING
-					TArray<FAnsiString> HostsOverride;
-					if (Container->HostGroupName == FOnDemandHostGroup::DefaultName)
-					{
-						FString HostValue;
-						if (FParse::Value(FCommandLine::Get(), TEXT("Iax.DefaultHostGroup="), HostValue))
-						{
-							TArray<FString> HostArray;
-							HostValue.ParseIntoArray(HostArray, TEXT(","), true);
-							for (FString& Host : HostArray)
-							{
-								Host.TrimStartAndEndInline();
-								HostsOverride.Add(FAnsiString(StringCast<ANSICHAR>(*Host)));
-							}
-							Hosts = HostsOverride; 
-						}
-					}
-#endif
-					if (FIoStatus Status = Container->HostGroup.Resolve(Hosts); !Status.IsOk())
-					{
-						return Status;
-					}
-
-					{
-						const FOnDemandHostGroup& HostGroup = Container->HostGroup.GetUnderlyingHostGroup();
-						Hosts								= HostGroup.Hosts();
-						TStringBuilder<128> Sb;
-
-						UE_LOG(LogIoStoreOnDemand, Log, TEXT("Created new host group '%s'"), *Container->HostGroupName.ToString());
-						for (int32 HostIndex = 0; HostIndex < Hosts.Num(); ++HostIndex) 
-						{
-							Sb.Reset();
-							Sb << Hosts[HostIndex];
-							if (HostIndex == HostGroup.PrimaryHostIndex())
-							{
-								Sb << TEXT(" (primary)");
-							}
-							UE_LOG(LogIoStoreOnDemand, Log, TEXT("\t%s"), Sb.ToString()); 
-						}
-					}
-				}
-				else
-				{
-					return NewHostGroup.Status();
-				}
-			}
-
-			if (Container->HostGroup.IsResolved() == false)
+		if (const FIoStatus Status = SetupHostGroup(Container, Args); !Status.IsOk())
+		{
+			if (Status.GetErrorCode() == EIoErrorCode::PendingHostGroup)
 			{
-				//TODO: Handle new host group being registered
 				EnumAddFlags(Container->Flags, EOnDemandContainerFlags::PendingHostGroup);
 				MountResult = EIoErrorCode::PendingHostGroup;
 				UE_LOG(LogIoStoreOnDemand, Log, TEXT("Deferring container '%s' until host group '%s' becomes available"),
 					*Container->Name, *Container->HostGroupName.ToString());
 			}
+			else
+			{
+				return Status;
+			}
 		}
 
 		if (EnumHasAnyFlags(Container->Flags, EOnDemandContainerFlags::Encrypted) &&
@@ -1434,6 +1451,7 @@ void FOnDemandIoStore::OnEncryptionKeyAdded(const FGuid& Id, const FAES::FAESKey
 				*Container->Name, Container->ChunkEntries.Num(), *LexToString(Container->Flags));
 
 		EnumRemoveFlags(Container->Flags, EOnDemandContainerFlags::PendingEncryptionKey);
+		EnumAddFlags(Container->Flags, EOnDemandContainerFlags::Mounted);
 		
 		Container->EncryptionKeyGuid.Empty();
 
@@ -1459,6 +1477,7 @@ void FOnDemandIoStore::OnEncryptionKeyAdded(const FGuid& Id, const FAES::FAESKey
 				}
 				else
 				{
+					EnumRemoveFlags(Container->Flags, EOnDemandContainerFlags::Mounted);
 					UE_LOG(LogIoStoreOnDemand, Error, TEXT("Failed to deserialize header when mounting container '%s', Entries=%d, Flags='%s'"),
 						*Container->Name, Container->ChunkEntries.Num(), *LexToString(Container->Flags));
 				}
@@ -1474,6 +1493,23 @@ void FOnDemandIoStore::OnEncryptionKeyAdded(const FGuid& Id, const FAES::FAESKey
 	}
 }
 
+void FOnDemandIoStore::OnHostGroupRegistered(const FName& HostGroup)
+{
+	TUniqueLock Lock(ContainerMutex);
+	for (FSharedOnDemandContainer& Container : Containers)
+	{
+		if (!EnumHasAnyFlags(Container->Flags, EOnDemandContainerFlags::PendingHostGroup) || Container->HostGroupName != HostGroup)
+		{
+			continue;
+		}
+
+		EnumRemoveFlags(Container->Flags, EOnDemandContainerFlags::PendingHostGroup);
+		EnumAddFlags(Container->Flags, EOnDemandContainerFlags::Mounted);
+		UE_LOG(LogIoStoreOnDemand, Log, TEXT("Mounting container '%s', Entries=%d, Flags='%s', HostGroup='%s'"),
+			*Container->Name, Container->ChunkEntries.Num(), *LexToString(Container->Flags), *HostGroup.ToString());
+	}
+}
+
 void FOnDemandIoStore::CreateContainersFromToc(
 	FStringView MountId,
 	FStringView TocPath,
@@ -1533,6 +1569,7 @@ void FOnDemandIoStore::CreateContainersFromToc(
 		Container->BlockHashes			= MoveTemp(ContainerEntry.BlockHashes);
 		Container->ContainerId			= ContainerEntry.ContainerId;
 		Container->HostGroupName		= Toc.Header.HostGroupName.IsEmpty() ? FOnDemandHostGroup::DefaultName : FName(Toc.Header.HostGroupName);
+		Container->RelativeUrl			= UE::FIoRelativeUrl::From(Container->ChunksDirectory);
 
 		Container->CompressionFormats.Reserve(1);
 		Container->CompressionFormats.Add(CompressionFormat);
@@ -1603,6 +1640,108 @@ void FOnDemandIoStore::CreateContainersFromToc(
 	}
 }
 
+FIoStatus FOnDemandIoStore::SetupHostGroup(const FSharedOnDemandContainer& Container, const FOnDemandMountArgs& MountArgs)
+{
+	if (FHttpIoDispatcher::IsInitialized())
+	{
+		UE_CLOG(Container->HostGroupName != MountArgs.HostGroupName, LogIoStoreOnDemand, Log,
+			TEXT("Overriding container host group '%s' with '%s' from mount option"),
+			*Container->HostGroupName.ToString(), *MountArgs.HostGroupName.ToString());
+
+		Container->HostGroupName = MountArgs.HostGroupName;
+		if (FHttpIoDispatcher::IsHostGroupRegistered(Container->HostGroupName))
+		{
+			return FIoStatus::Ok;
+		}
+
+		if (MountArgs.HostGroup.IsEmpty() == false)
+		{
+			return FHttpIoDispatcher::RegisterHostGroup(Container->HostGroupName, MountArgs.HostGroup.Hosts(), Container->GetTestUrl());
+		}
+
+		return FIoStatus(EIoErrorCode::PendingHostGroup); 
+	}
+
+	// This path will be removed once the HTTP I/O dispatcher is on by default
+	if (EnumHasAnyFlags(Container->Flags, EOnDemandContainerFlags::StreamOnDemand))
+	{
+		//TODO: Remove this special case for streaming containers
+		Container->HostGroupName	= FName("DefaultOnDemand");
+		Container->HostGroup		= FHostGroupManager::Get().Find(Container->HostGroupName);
+
+		return FIoStatus::Ok;
+	}
+
+	UE_CLOG(Container->HostGroupName != MountArgs.HostGroupName, LogIoStoreOnDemand, Log,
+		TEXT("Overriding container host group '%s' with '%s' from mount option"),
+		*Container->HostGroupName.ToString(), *MountArgs.HostGroupName.ToString());
+
+	Container->HostGroupName	= MountArgs.HostGroupName;
+	Container->HostGroup		= FHostGroupManager::Get().Find(Container->HostGroupName);
+
+	if (Container->HostGroup.IsResolved())
+	{
+		return FIoStatus::Ok;
+	}
+
+	if (MountArgs.HostGroup.IsEmpty())
+	{
+		return FIoStatus(EIoErrorCode::PendingHostGroup);
+	}
+	
+	TIoStatusOr<FIASHostGroup> NewHostGroup = FHostGroupManager::Get().Register(Container->HostGroupName, Container->GetTestUrl());
+	if (NewHostGroup.IsOk() == false)
+	{
+		return NewHostGroup.Status();
+	}
+
+	Container->HostGroup = NewHostGroup.ConsumeValueOrDie();
+
+	TConstArrayView<FAnsiString> Hosts = MountArgs.HostGroup.Hosts();
+#if !UE_BUILD_SHIPPING
+	TArray<FAnsiString> HostsOverride;
+	if (Container->HostGroupName == FOnDemandHostGroup::DefaultName)
+	{
+		FString HostValue;
+		if (FParse::Value(FCommandLine::Get(), TEXT("Iax.DefaultHostGroup="), HostValue))
+		{
+			TArray<FString> HostArray;
+			HostValue.ParseIntoArray(HostArray, TEXT(","), true);
+			for (FString& Host : HostArray)
+			{
+				Host.TrimStartAndEndInline();
+				HostsOverride.Add(FAnsiString(StringCast<ANSICHAR>(*Host)));
+			}
+			Hosts = HostsOverride; 
+		}
+	}
+#endif
+	if (FIoStatus Status = Container->HostGroup.Resolve(Hosts); Status.IsOk() == false)
+	{
+		return Status;
+	}
+
+	{
+		const FOnDemandHostGroup& HostGroup = Container->HostGroup.GetUnderlyingHostGroup();
+		Hosts								= HostGroup.Hosts();
+		TStringBuilder<128> Sb;
+
+		UE_LOG(LogIoStoreOnDemand, Log, TEXT("Created new host group '%s'"), *Container->HostGroupName.ToString());
+		for (int32 HostIndex = 0; HostIndex < Hosts.Num(); ++HostIndex) 
+		{
+			Sb.Reset();
+			Sb << Hosts[HostIndex];
+			if (HostIndex == HostGroup.PrimaryHostIndex())
+			{
+				Sb << TEXT(" (primary)");
+			}
+			UE_LOG(LogIoStoreOnDemand, Log, TEXT("\t%s"), Sb.ToString()); 
+		}
+	}
+
+	return FIoStatus::Ok;
+}
+
 TArray<FSharedOnDemandContainer> FOnDemandIoStore::GetContainers(EOnDemandContainerFlags ContainerFlags) const
 {
 	TArray<FSharedOnDemandContainer> Out;
@@ -1909,6 +2048,11 @@ void FOnDemandIoStore::SetStreamingOptions(EOnDemandStreamingOptions Options)
 	{
 		StreamingBackend->SetBulkOptionalEnabled(EnumHasAnyFlags(Options, EOnDemandStreamingOptions::OptionalBulkDataDisabled) == false);
 	}
+
+	if (HttpIoBackend.IsValid())
+	{
+		HttpIoBackend->SetOptionalBulkDataEnabled(EnumHasAnyFlags(Options, EOnDemandStreamingOptions::OptionalBulkDataDisabled) == false);
+	}
 }
 
 void FOnDemandIoStore::ReportAnalytics(TArray<FAnalyticsEventAttribute>& OutAnalyticsArray) const
@@ -1922,6 +2066,11 @@ void FOnDemandIoStore::ReportAnalytics(TArray<FAnalyticsEventAttribute>& OutAnal
 	{
 		return StreamingBackend->ReportAnalytics(OutAnalyticsArray);
 	}
+
+	if (HttpIoBackend.IsValid())
+	{
+		HttpIoBackend->ReportAnalytics(OutAnalyticsArray);
+	}
 }
 
 TUniquePtr<IAnalyticsRecording> FOnDemandIoStore::StartAnalyticsRecording() const
@@ -1985,4 +2134,3 @@ FString LexToString(UE::IoStore::EOnDemandContainerFlags Flags)
 	Sb << Flags;
 	return FString::ConstructFromPtrSize(Sb.ToString(), Sb.Len());
 }
-
diff --git a/Engine/Source/Runtime/Experimental/IoStore/OnDemand/Private/OnDemandIoStore.h b/Engine/Source/Runtime/Experimental/IoStore/OnDemand/Private/OnDemandIoStore.h
--- a/Engine/Source/Runtime/Experimental/IoStore/OnDemand/Private/OnDemandIoStore.h
+++ b/Engine/Source/Runtime/Experimental/IoStore/OnDemand/Private/OnDemandIoStore.h
@@ -7,6 +7,7 @@
 #include "Async/Mutex.h"
 #include "Containers/AnsiString.h"
 #include "Containers/BitArray.h"
+#include "IO/HttpIoDispatcher.h"
 #include "IO/IoBuffer.h"
 #include "IO/IoChunkEncoding.h"
 #include "IO/IoChunkId.h"
@@ -34,6 +35,7 @@ enum class ELogOnDemandCacheUsage : uint8;
 using FSharedStreamingBackend		= TSharedPtr<class IOnDemandIoDispatcherBackend>;
 using FSharedPackageStoreBackend	= TSharedPtr<class IOnDemandPackageStoreBackend>;
 using FSharedInstallCache			= TSharedPtr<class IOnDemandInstallCache>;
+using FSharedHttpIoBackend			= TSharedPtr<class IOnDemandHttpIoDispatcherBackend>;
 
 ///////////////////////////////////////////////////////////////////////////////
 enum class EOnDemandContainerFlags : uint8
@@ -117,6 +119,7 @@ struct FOnDemandContainer
 	TArrayView<FIoChunkId>					ChunkIds;
 	TArrayView<FOnDemandChunkEntry> 		ChunkEntries;
 	TArray<FOnDemandChunkEntryReferences>	ChunkEntryReferences;
+	UE::FIoRelativeUrl						RelativeUrl;
 	FIoContainerId							ContainerId;
 	uint32									BlockSize = 0;
 	FName									HostGroupName = FOnDemandHostGroup::DefaultName;
@@ -241,8 +244,10 @@ struct FOnDemandChunkInfo
 	inline TConstArrayView<uint32>			Blocks() const;
 	inline TConstArrayView<FIoBlockHash>	BlockHashes() const;
 	FAnsiStringView							ChunksDirectory() const { return SharedContainer->ChunksDirectory; }
+	UE::FIoRelativeUrl						RelativeUrl() const { return SharedContainer->RelativeUrl; }
 	const FOnDemandChunkEntry&				ChunkEntry() const { return Entry; }
 	inline const FIASHostGroup&				HostGroup() const;
+	FName									HostGroupName() const { return SharedContainer->HostGroupName; }
 
 	bool									IsValid() const { return SharedContainer.IsValid(); }
 	operator								bool() const { return IsValid(); }
@@ -376,18 +381,21 @@ private:
 	FIoStatus							TickMountRequest(FMountRequest& MountRequest);
 	void								CompleteMountRequest(FMountRequest& MountRequest, FOnDemandMountResult&& MountResult);
 	void								OnEncryptionKeyAdded(const FGuid& Id, const FAES::FAESKey& Key);
+	void								OnHostGroupRegistered(const FName& HostGroup);
 	static void							CreateContainersFromToc(
 											FStringView MountId,
 											FStringView TocPath,
 											FOnDemandToc& Toc,
 											const TArray<FString>& ExistingContainers,
 											TArray<FSharedOnDemandContainer>& Out);
+	static FIoStatus						SetupHostGroup(const FSharedOnDemandContainer& Container, const FOnDemandMountArgs& MountArgs);
 
 	TUniquePtr<FOnDemandHttpThread>			HttpClient;
 	TUniquePtr<FOnDemandContentInstaller>	Installer;
 	FSharedInstallCache						InstallCache;
 	FSharedPackageStoreBackend				PackageStoreBackend;
 	FSharedStreamingBackend					StreamingBackend;
+	FSharedHttpIoBackend					HttpIoBackend;
 	FDelegateHandle							OnMountPakHandle;
 	FDelegateHandle							OnServerPostForkHandle;
 	TArray<FSharedOnDemandContainer>		Containers;
diff --git a/Engine/Source/Runtime/Experimental/IoStore/OnDemandCore/Private/HttpIoDispatcher.cpp b/Engine/Source/Runtime/Experimental/IoStore/OnDemandCore/Private/HttpIoDispatcher.cpp
new file mode 100644
--- /dev/null
+++ b/Engine/Source/Runtime/Experimental/IoStore/OnDemandCore/Private/HttpIoDispatcher.cpp
@@ -0,0 +1,309 @@
+// Copyright Epic Games, Inc. All Rights Reserved.
+
+#include "IO/HttpIoDispatcher.h"
+#include "Containers/AnsiString.h"
+#include "Misc/StringBuilder.h"
+
+DEFINE_LOG_CATEGORY(LogHttpIoDispatcher);
+
+namespace UE
+{
+
+////////////////////////////////////////////////////////////////////////////////
+TSharedPtr<IHttpIoDispatcher> GHttpIoDispatcher;
+const FIoHttpOptions FIoHttpOptions::Default = FIoHttpOptions();
+
+////////////////////////////////////////////////////////////////////////////////
+FIoHttpHeaders& FIoHttpHeaders::Add(FAnsiString&& HeaderName, FAnsiString&& HeaderValue)
+{
+	Headers.Add(MoveTemp(HeaderName));
+	Headers.Add(MoveTemp(HeaderValue));
+	return *this;
+}
+
+FIoHttpHeaders& FIoHttpHeaders::Add(FAnsiStringView HeaderName, FAnsiStringView HeaderValue)
+{
+	Headers.Emplace(HeaderName);
+	Headers.Emplace(HeaderValue);
+	return *this;
+}
+
+TArray<FAnsiString> FIoHttpHeaders::ToArray() &&
+{
+	return MoveTemp(Headers);
+}
+
+FIoHttpHeaders FIoHttpHeaders::Create(FAnsiString&& HeaderName, FAnsiString&& HeaderValue)
+{
+	FIoHttpHeaders Headers;
+	Headers.Add(MoveTemp(HeaderName), MoveTemp(HeaderValue));
+	return Headers;
+}
+
+FIoHttpHeaders FIoHttpHeaders::Create(FAnsiStringView HeaderName, FAnsiStringView HeaderValue)
+{
+	FIoHttpHeaders Headers;
+	Headers.Add(HeaderName, HeaderValue);
+	return Headers;
+}
+
+////////////////////////////////////////////////////////////////////////////////
+FIoRelativeUrl FIoRelativeUrl::From(FAnsiStringView Url)
+{
+	if (Url.IsEmpty())
+	{
+		return FIoRelativeUrl();
+	}
+
+	if (Url[0] != ANSICHAR('/'))
+	{
+		TAnsiStringBuilder<128> Sb;
+		Sb << "/" << Url;
+		return FIoRelativeUrl(Sb.ToString());
+	}
+
+	return FIoRelativeUrl(Url);
+}
+
+////////////////////////////////////////////////////////////////////////////////
+FIoHttpRequest::~FIoHttpRequest()
+{
+	Release();
+}
+
+void FIoHttpRequest::Cancel()
+{
+	if (Handle != 0)
+	{
+		return GHttpIoDispatcher->CancelRequest(Handle);
+	}
+}
+
+void FIoHttpRequest::UpdatePriorty(int32 NewPriority)
+{
+	if (Handle != 0)
+	{
+		return GHttpIoDispatcher->UpdateRequestPriority(Handle, NewPriority);
+	}
+}
+
+EIoErrorCode FIoHttpRequest::Status() const
+{
+	if (Handle != 0)
+	{
+		return GHttpIoDispatcher->GetRequestStatus(Handle);
+	}
+
+	return EIoErrorCode::InvalidCode;
+}
+
+void FIoHttpRequest::Release()
+{
+	if (Handle != 0)
+	{
+		GHttpIoDispatcher->ReleaseRequest(Handle);
+		Handle = 0;
+	}
+}
+
+FIoHttpRequest&	FIoHttpRequest::operator=(FIoHttpRequest&& Other)
+{
+	Release();
+	Handle = Other.Handle;
+	Other.Handle = 0;
+	return *this;
+}
+
+////////////////////////////////////////////////////////////////////////////////
+FIoHttpBatch::~FIoHttpBatch()
+{
+	if (First != 0)
+	{
+		ensure(Last != 0);
+		GHttpIoDispatcher->ReleaseRequest(First);
+		First = Last = 0;
+	}
+}
+
+FIoHttpRequest FIoHttpBatch::Get(
+	const FName& HostGroup,
+	const FIoRelativeUrl& RelativeUrl,
+	FIoHttpHeaders&& Headers,
+	const FIoHttpOptions& Options,
+	const FIoHash& ChunkHash,
+	FIoHttpRequestCompleted&& OnCompleted)
+{
+	FIoHttpRequestHandle Handle = GHttpIoDispatcher->CreateRequest(
+		First,
+		Last,
+		HostGroup,
+		RelativeUrl,
+		Options,
+		MoveTemp(Headers),
+		MoveTemp(OnCompleted),
+		&ChunkHash);
+
+	return FIoHttpRequest(Handle);
+}
+
+FIoHttpRequest FIoHttpBatch::Get(
+	const FName& HostGroup,
+	const FIoRelativeUrl& RelativeUrl,
+	FIoHttpHeaders&& Headers,
+	const FIoHttpOptions& Options,
+	FIoHttpRequestCompleted&& OnCompleted)
+{
+	FIoHttpRequestHandle Handle = GHttpIoDispatcher->CreateRequest(
+		First,
+		Last,
+		HostGroup,
+		RelativeUrl,
+		Options,
+		MoveTemp(Headers),
+		MoveTemp(OnCompleted));
+
+	return FIoHttpRequest(Handle);
+}
+
+FIoHttpRequest FIoHttpBatch::Get(
+	const FName& HostGroup,
+	const FIoRelativeUrl& RelativeUrl,
+	FIoHttpHeaders&& Headers,
+	FIoHttpRequestCompleted&& OnCompleted)
+{
+	return Get(HostGroup, RelativeUrl, MoveTemp(Headers), FIoHttpOptions::Default, MoveTemp(OnCompleted));
+}
+
+FIoHttpRequest FIoHttpBatch::Get(
+	const FName& HostGroup,
+	const FIoRelativeUrl& RelativeUrl,
+	FIoHttpRequestCompleted&& OnCompleted)
+{
+	return Get(HostGroup, RelativeUrl, FIoHttpHeaders(), FIoHttpOptions::Default, MoveTemp(OnCompleted));
+}
+
+void FIoHttpBatch::Issue()
+{
+	if (First != 0)
+	{
+		ensure(Last != 0);
+		GHttpIoDispatcher->IssueRequest(First);
+		First = Last = 0;
+	}
+}
+
+////////////////////////////////////////////////////////////////////////////////
+bool FHttpIoDispatcher::IsInitialized()
+{
+	return GHttpIoDispatcher.IsValid();
+}
+
+FIoStatus FHttpIoDispatcher::Initialize(TSharedPtr<IHttpIoDispatcher> Dispatcher)
+{
+	if (GHttpIoDispatcher.IsValid())
+	{
+		return FIoStatus(EIoErrorCode::InvalidCode);
+	}
+
+	UE_LOG(LogHttpIoDispatcher, Log, TEXT("Initializing HTTP I/O dispatcher"));
+	GHttpIoDispatcher = Dispatcher;
+
+	return FIoStatus::Ok;
+}
+
+FIoStatus FHttpIoDispatcher::Shutdown()
+{
+	if (GHttpIoDispatcher.IsValid() == false)
+	{
+		return FIoStatus::Invalid;
+	}
+
+	UE_LOG(LogHttpIoDispatcher, Log, TEXT("Shutting down HTTP I/O dispatcher"));
+	GHttpIoDispatcher->Shutdown();
+	GHttpIoDispatcher.Reset();
+
+	return FIoStatus::Ok;
+}
+
+FIoStatus FHttpIoDispatcher::RegisterHostGroup(const FName& HostGroup, TConstArrayView<FAnsiString> HostNames, FAnsiStringView TestUrl)
+{
+	if (GHttpIoDispatcher.IsValid())
+	{
+		return GHttpIoDispatcher->RegisterHostGroup(HostGroup, HostNames, TestUrl);
+	}
+
+	return FIoStatus(EIoErrorCode::InvalidCode);
+}
+
+FIoStatus FHttpIoDispatcher::RegisterHostGroup(const FName& HostGroup, FAnsiStringView HostName, FAnsiStringView TestUrl)
+{
+	FAnsiString HostNameString(HostName);
+	TConstArrayView<FAnsiString> HostNames(&HostNameString, 1);
+	return RegisterHostGroup(HostGroup, HostNames, TestUrl);
+}
+
+bool FHttpIoDispatcher::IsHostGroupRegistered(const FName& HostGroup)
+{
+	return GHttpIoDispatcher.IsValid() ? GHttpIoDispatcher->IsHostGroupRegistered(HostGroup) : false;
+}
+
+bool FHttpIoDispatcher::IsHostGroupOk(const FName& HostGroup)
+{
+	return GHttpIoDispatcher.IsValid() ? GHttpIoDispatcher->IsHostGroupOk(HostGroup) : false;
+}
+
+FIoHttpBatch FHttpIoDispatcher::NewBatch()
+{
+	ensure(GHttpIoDispatcher.IsValid());
+	return FIoHttpBatch();
+}
+
+FIoHttpRequest FHttpIoDispatcher::Get(
+	const FName& HostGroup,
+	const FIoRelativeUrl& RelativeUrl,
+	FIoHttpHeaders&& Headers,
+	const FIoHttpOptions& Options,
+	FIoHttpRequestCompleted&& OnCompleted)
+{
+	ensure(GHttpIoDispatcher.IsValid());
+
+	FIoHttpRequestHandle First = 0;
+	FIoHttpRequestHandle Last = 0;
+
+	FIoHttpRequestHandle Handle = GHttpIoDispatcher->CreateRequest(
+		First,
+		Last,
+		HostGroup,
+		RelativeUrl,
+		Options,
+		MoveTemp(Headers),
+		MoveTemp(OnCompleted));
+
+	GHttpIoDispatcher->IssueRequest(First);
+	return FIoHttpRequest(First);
+}
+
+FIoHttpRequest FHttpIoDispatcher::Get(
+	const FName& HostGroup,
+	const FIoRelativeUrl& RelativeUrl,
+	FIoHttpHeaders&& Headers,
+	FIoHttpRequestCompleted&& OnCompleted)
+{
+	return Get(HostGroup, RelativeUrl, MoveTemp(Headers), FIoHttpOptions::Default, MoveTemp(OnCompleted));
+}
+
+FIoHttpRequest FHttpIoDispatcher::Get(
+	const FName& HostGroup,
+	const FIoRelativeUrl& RelativeUrl,
+	FIoHttpRequestCompleted&& OnCompleted)
+{
+	return Get(HostGroup, RelativeUrl, FIoHttpHeaders(), FIoHttpOptions::Default, MoveTemp(OnCompleted));
+}
+
+FHttpIoDispatcher::FHostGroupRegistered& FHttpIoDispatcher::OnHostGroupRegistered()
+{
+	ensure(GHttpIoDispatcher.IsValid());
+	return GHttpIoDispatcher->OnHostGroupRegistered();
+}
+
+} // namespace UE
diff --git a/Engine/Source/Runtime/Experimental/IoStore/OnDemandCore/Public/IO/HttpIoDispatcher.h b/Engine/Source/Runtime/Experimental/IoStore/OnDemandCore/Public/IO/HttpIoDispatcher.h
new file mode 100644
--- /dev/null
+++ b/Engine/Source/Runtime/Experimental/IoStore/OnDemandCore/Public/IO/HttpIoDispatcher.h
@@ -0,0 +1,377 @@
+// Copyright Epic Games, Inc. All Rights Reserved.
+
+#pragma once
+
+#include "Containers/AnsiString.h"
+#include "Containers/SharedString.h"
+#include "Delegates/Delegate.h"
+#include "IO/IoBuffer.h"
+#include "IO/IoHash.h"
+#include "IO/IoOffsetLength.h"
+#include "IO/IoStatus.h"
+#include "Misc/EnumClassFlags.h"
+#include "Templates/Function.h"
+#include "Templates/SharedPointer.h"
+
+#define UE_API IOSTOREONDEMANDCORE_API
+
+UE_API DECLARE_LOG_CATEGORY_EXTERN(LogHttpIoDispatcher, Log, All);
+
+namespace UE
+{
+
+/** Internal request handle type. */
+using FIoHttpRequestHandle = UPTRINT;
+
+/** Flags for controlling the behavior of an HTTP request. */
+enum class EIoHttpFlags
+{
+	/** No additional flags. */
+	None			= 0,
+	/** Whether to read from the HTTP cache. */ 
+	ReadCache		= (1 << 0),
+	/** Whether to write to the HTTP cache. */ 
+	WriteCache		= (1 << 1),
+	/** Whether to read and write to the HTTP cache. */ 
+	ReadWriteCache	= ReadCache | WriteCache,
+	/** Default flags . */
+	Default			= ReadWriteCache
+};
+ENUM_CLASS_FLAGS(EIoHttpFlags);
+
+/** Represents a range within a resource. */
+struct FIoHttpRange
+{
+	UE_API						FIoHttpRange() = default;
+	explicit inline				FIoHttpRange(uint32 InMin, uint32 InMax);
+
+	bool						IsValid()	const	{ return Min <= Max; }
+	uint32						GetMin()	const	{ return Min; }
+	uint32						GetMax()	const	{ return Max; }
+	uint32						GetSize()	const	{ return (Max - Min); }
+	inline FIoHttpRange&		Expand(const FIoHttpRange& Other);
+
+	inline FIoOffsetAndLength	ToOffsetAndLength() const;
+	static inline FIoHttpRange	FromOffsetAndLength(const FIoOffsetAndLength& OffsetAndLength);
+
+	bool						operator==(const FIoHttpRange& Other)	{ return Min == Other.Min && Max == Other.Max; }
+	bool						operator!=(const FIoHttpRange& Other)	{ return Min != Other.Min || Max != Other.Max; }
+	FIoHttpRange&				operator+(const FIoHttpRange& Other)	{ return Expand(Other); }
+	FIoHttpRange&				operator+=(const FIoHttpRange& Other)	{ return Expand(Other); }
+
+private:
+	uint32 Min = MAX_uint32;
+	uint32 Max = MIN_uint32;
+};
+
+FIoHttpRange::FIoHttpRange(uint32 InMin, uint32 InMax)
+	: Min(InMin)
+	, Max(InMax)
+{
+}
+
+FIoHttpRange& FIoHttpRange::Expand(const FIoHttpRange& Other)
+{
+	if (IsValid())
+	{
+		Min = FMath::Min(Min, Other.Min);
+		Max = FMath::Max(Max, Other.Max);
+	}
+	else
+	{
+		Min = Other.Min;
+		Max = Other.Max;
+	}
+
+	return *this;
+}
+
+FIoOffsetAndLength FIoHttpRange::ToOffsetAndLength() const
+{
+	return IsValid() ? FIoOffsetAndLength(Min, Max - Min) : FIoOffsetAndLength(0, 0);
+}
+
+FIoHttpRange FIoHttpRange::FromOffsetAndLength(const FIoOffsetAndLength& OffsetAndLength)
+{
+	const uint64 End = OffsetAndLength.GetOffset() + OffsetAndLength.GetLength();
+
+	if (OffsetAndLength.GetOffset() > MAX_uint32 || End > MAX_uint32)
+	{
+		return FIoHttpRange();
+	}
+
+	return FIoHttpRange(IntCastChecked<uint32>(OffsetAndLength.GetOffset()), IntCastChecked<uint32>(End));
+}
+
+/** Options for controlling the behavior of an HTTP request. */
+struct FIoHttpOptions
+{
+public:
+	inline						FIoHttpOptions(int32 InPriority, int32 InRetryCount, EIoHttpFlags InFlags, const FIoHttpRange& InRange);
+	inline						FIoHttpOptions(int32 InPriority, int32 InRetryCount, EIoHttpFlags InFlags);
+	inline						FIoHttpOptions(int32 InPriority, int32 InRetryCount);
+	inline explicit				FIoHttpOptions(int32 InPriority);
+	UE_API						FIoHttpOptions() = default;
+
+	const FIoHttpRange&			GetRange()		const { return Range; }
+	int32						GetPriority()	const { return Priority; }
+	int32						GetRetryCount()	const { return RetryCount; }
+	EIoHttpFlags				GetFlags()		const { return Flags; }
+	uint8						GetCategory()	const { return Category; }
+
+	void						SetRange(const FIoHttpRange InRange)	{ Range = InRange; }
+	void						SetPriority(int32 InPriority)			{ Priority = InPriority; }
+	void						SetRetryCount(int32 InRetryCount)		{ RetryCount = InRetryCount; }
+	void						SetCategory(uint8 InCategory)			{ Category = InCategory; }
+
+	UE_API static const FIoHttpOptions Default;
+
+private:
+	FIoHttpRange	Range;
+	int32			Priority = 0;
+	int32			RetryCount = 0;
+	EIoHttpFlags	Flags = EIoHttpFlags::Default;
+	uint8			Category = 0;
+};
+
+FIoHttpOptions::FIoHttpOptions(int32 InPriority, int32 InRetryCount, EIoHttpFlags InFlags, const FIoHttpRange& InRange)
+	: Range(InRange)
+	, Priority(InPriority)
+	, RetryCount(InRetryCount)
+	, Flags(InFlags)
+{
+}
+
+FIoHttpOptions::FIoHttpOptions(int32 InPriority, int32 InRetryCount, EIoHttpFlags InFlags)
+	: Priority(InPriority)
+	, RetryCount(InRetryCount)
+	, Flags(InFlags)
+{
+}
+
+FIoHttpOptions::FIoHttpOptions(int32 InPriority, int32 InRetryCount)
+	: Priority(InPriority)
+	, RetryCount(InRetryCount)
+{
+}
+
+FIoHttpOptions::FIoHttpOptions(int32 InPriority)
+	: Priority(InPriority)
+{
+}
+
+/** HTTP headers. */
+class FIoHttpHeaders
+{
+public:
+	UE_API							FIoHttpHeaders() = default;
+	UE_API							FIoHttpHeaders(const FIoHttpHeaders&) = default;
+	UE_API							FIoHttpHeaders(FIoHttpHeaders&&) = default;
+	UE_API FIoHttpHeaders&			Add(FAnsiString&& HeaderName, FAnsiString&& HeaderValue);
+	UE_API FIoHttpHeaders&			Add(FAnsiStringView HeaderName, FAnsiStringView HeaderValue);
+	UE_API TArray<FAnsiString>		ToArray() &&;
+	TConstArrayView<FAnsiString>	ToArrayView() const { return Headers; }
+
+	UE_API static FIoHttpHeaders	Create(FAnsiString&& HeaderName, FAnsiString&& HeaderValue);
+	UE_API static FIoHttpHeaders	Create(FAnsiStringView HeaderName, FAnsiStringView HeaderValue);
+	
+	FIoHttpHeaders&					operator=(const FIoHttpHeaders& Other)	{ Headers = Other.Headers; return *this; }
+	FIoHttpHeaders&					operator=(FIoHttpHeaders&& Other)		{ Headers = MoveTemp(Other.Headers); return *this; }
+
+private:
+	//TODO: Make this better
+	TArray<FAnsiString> Headers;
+};
+
+/** A relative URL from the host. */
+class FIoRelativeUrl
+{
+public:
+	UE_API							FIoRelativeUrl() = default;
+	bool							IsEmpty() const		{ return Url.IsEmpty(); }
+	int32							Len() const			{ return Url.Len(); }
+	const ANSICHAR*					ToString() const	{ return *Url; }
+	UE_API static FIoRelativeUrl	From(FAnsiStringView Url); 
+
+	const ANSICHAR*					operator*() const	{ return *Url; }
+
+private:
+	explicit FIoRelativeUrl(FAnsiStringView InUrl)
+		: Url(InUrl) { }
+
+	TSharedString<ANSICHAR> Url;
+};
+
+/**
+ * An HTTP request handle.
+ *
+ * The handle needs to be keept alive until the completion callback has been triggered.
+ */
+class FIoHttpRequest
+{
+public:
+	UE_API					FIoHttpRequest() = default;
+	explicit 				FIoHttpRequest(FIoHttpRequestHandle InHandle)
+								: Handle(InHandle) { }
+							FIoHttpRequest(FIoHttpRequest&& Other)
+								: Handle(Other.Handle) { Other.Handle = 0; }
+							FIoHttpRequest(const FIoHttpRequest&) = delete;
+	UE_API					~FIoHttpRequest();
+
+	bool					IsValid() const { return Handle != 0; }
+	UE_API void				Cancel();
+	UE_API void				UpdatePriorty(int32 NewPriority);
+	UE_API EIoErrorCode		Status() const;
+
+	FIoHttpRequest&			operator=(const FIoHttpRequest&) = delete;
+	UE_API FIoHttpRequest&	operator=(FIoHttpRequest&& Other);
+
+private:
+	void					Release();
+	FIoHttpRequestHandle	Handle = 0;
+};
+
+/** Flags describing a HTTP response. */
+enum class EIoHttpResponseFlags : uint8
+{
+	/** No additional flags. */
+	None			= 0,
+	/** Whether the response was cached. */ 
+	Cached			= (1 << 0),
+};
+ENUM_CLASS_FLAGS(EIoHttpResponseFlags);
+
+/** A HTTP response. */
+class FIoHttpResponse
+{
+public:
+	UE_API					FIoHttpResponse() = default;
+	inline					FIoHttpResponse(const FIoBuffer& InBody, EIoErrorCode InErrorCode, uint32 InStatusCode, EIoHttpResponseFlags InFlags);
+							FIoHttpResponse(EIoErrorCode InErrorCode, uint32 InStatusCode)
+								: ErrorCode(InErrorCode)
+								, StatusCode(InStatusCode) { }
+
+	const FIoBuffer&		GetBody()		const { return Body; }
+	EIoErrorCode			GetErrorCode()	const { return ErrorCode; }
+	uint32					GetStatusCode()	const { return StatusCode; }
+	EIoHttpResponseFlags	GetFlags()		const { return Flags; }
+	bool					IsOk()			const { return ErrorCode == EIoErrorCode::Ok && StatusCode > 199 && StatusCode < 300; }
+	bool					IsCancelled()	const { return ErrorCode == EIoErrorCode::Cancelled; }
+	bool					IsCached()		const { return EnumHasAnyFlags(Flags, EIoHttpResponseFlags::Cached); } 
+
+private:
+	FIoBuffer				Body;
+	EIoErrorCode			ErrorCode;
+	uint32					StatusCode = 0;
+	EIoHttpResponseFlags	Flags = EIoHttpResponseFlags::None;
+};
+
+FIoHttpResponse::FIoHttpResponse(const FIoBuffer& InBody, EIoErrorCode InErrorCode, uint32 InStatusCode, EIoHttpResponseFlags InFlags)
+	: Body(InBody)
+	, ErrorCode(InErrorCode)
+	, StatusCode(InStatusCode)
+	, Flags(InFlags)
+{
+}
+
+/** HTTP completion callback. */
+using FIoHttpRequestCompleted = TUniqueFunction<void(FIoHttpResponse&&)>;
+
+/** Issue one or more HTTP request in a batch. */
+class FIoHttpBatch
+{
+public:
+							FIoHttpBatch(const FIoHttpRequest&) = delete;
+							FIoHttpBatch(FIoHttpBatch&& Other) = delete;
+	UE_API					~FIoHttpBatch();
+
+	UE_API FIoHttpRequest	Get(const FName& HostGroup,
+								const FIoRelativeUrl& RelativeUrl,
+								FIoHttpHeaders&& Headers,
+								const FIoHttpOptions& Options,
+								const FIoHash& ChunkHash,
+								FIoHttpRequestCompleted&& OnCompleted);
+	UE_API FIoHttpRequest	Get(const FName& HostGroup,
+								const FIoRelativeUrl& RelativeUrl,
+								FIoHttpHeaders&& Headers,
+								const FIoHttpOptions& Options,
+								FIoHttpRequestCompleted&& OnCompleted);
+	UE_API FIoHttpRequest	Get(const FName& HostGroup,
+								const FIoRelativeUrl& RelativeUrl,
+								FIoHttpHeaders&& Headers,
+								FIoHttpRequestCompleted&& OnCompleted);
+	UE_API FIoHttpRequest	Get(const FName& HostGroup,
+								const FIoRelativeUrl& RelativeUrl,
+								FIoHttpRequestCompleted&& OnCompleted);
+	UE_API void				Issue();
+
+	FIoHttpBatch&			operator=(const FIoHttpBatch&) = delete;
+	FIoHttpBatch&			operator=(FIoHttpBatch&& Other) = delete;
+
+private:
+							friend class FHttpIoDispatcher;
+							FIoHttpBatch() = default;
+	FIoHttpRequestHandle	First = 0;
+	FIoHttpRequestHandle 	Last = 0;
+};
+
+/** HTTP I/O dispatcher . */
+class FHttpIoDispatcher
+{
+public:
+	UE_API static bool				IsInitialized();
+	UE_API static FIoStatus			Initialize(TSharedPtr<class IHttpIoDispatcher> Dispatcher);
+	UE_API static FIoStatus			Shutdown();
+	UE_API static FIoStatus			RegisterHostGroup(const FName& HostGroup, TConstArrayView<FAnsiString> HostNames, FAnsiStringView TestUrl);
+	UE_API static FIoStatus			RegisterHostGroup(const FName& HostGroup, FAnsiStringView HostName, FAnsiStringView TestUrl);
+	UE_API static bool				IsHostGroupRegistered(const FName& HostGroup);
+	UE_API static bool				IsHostGroupOk(const FName& HostGroup);
+	UE_API static FIoHttpBatch		NewBatch();
+	UE_API static FIoHttpRequest	Get(const FName& HostGroup,
+										const FIoRelativeUrl& RelativeUrl,
+										FIoHttpHeaders&& Headers,
+										const FIoHttpOptions& Options,
+										FIoHttpRequestCompleted&& OnCompleted);
+	UE_API static FIoHttpRequest	Get(const FName& HostGroup,
+										const FIoRelativeUrl& RelativeUrl,
+										FIoHttpHeaders&& Headers,
+										FIoHttpRequestCompleted&& OnCompleted);
+	UE_API static FIoHttpRequest	Get(const FName& HostGroup,
+										const FIoRelativeUrl& RelativeUrl,
+										FIoHttpRequestCompleted&& OnCompleted);
+
+	DECLARE_MULTICAST_DELEGATE_OneParam(FHostGroupRegistered, const FName&);
+	UE_API static FHostGroupRegistered& OnHostGroupRegistered();
+};
+
+/** HTTP I/O dispatcher interface. */
+class IHttpIoDispatcher
+{
+public:
+	using FHostGroupRegistered		= FHttpIoDispatcher::FHostGroupRegistered;
+
+	virtual							~IHttpIoDispatcher() = default;
+	virtual void					Shutdown() = 0;				
+	virtual FIoStatus				RegisterHostGroup(const FName& HostGroup, TConstArrayView<FAnsiString> HostNames, FAnsiStringView TestUrl) = 0;
+	virtual bool					IsHostGroupRegistered(const FName& HostGroup) = 0;
+	virtual bool					IsHostGroupOk(const FName& HostGroup) = 0;
+	virtual FHostGroupRegistered&	OnHostGroupRegistered() = 0;
+	virtual FIoHttpRequestHandle	CreateRequest(
+										FIoHttpRequestHandle& First,
+										FIoHttpRequestHandle& Last,
+										const FName& HostGroup,
+										const FIoRelativeUrl& RelativeUrl,
+										const FIoHttpOptions& Options,
+										FIoHttpHeaders&& Headers,
+										FIoHttpRequestCompleted&& OnCompleted,
+										const FIoHash* ChunkHash = nullptr) = 0;
+	virtual void					IssueRequest(FIoHttpRequestHandle RequestHandle) = 0;
+	virtual void					CancelRequest(FIoHttpRequestHandle Handle) = 0;
+	virtual void					UpdateRequestPriority(FIoHttpRequestHandle Handle, int32 NewPriority) = 0;
+	virtual EIoErrorCode			GetRequestStatus(FIoHttpRequestHandle Handle) = 0;
+	virtual void					ReleaseRequest(FIoHttpRequestHandle Handle) = 0;
+};
+
+} // namespace UE
+
+#undef UE_API
