Enable HTTP I/O dispatcher by default.

#rnx

diff --git a/Engine/Source/Runtime/Experimental/IoStore/OnDemand/Private/OnDemandIoStore.cpp b/Engine/Source/Runtime/Experimental/IoStore/OnDemand/Private/OnDemandIoStore.cpp
--- a/Engine/Source/Runtime/Experimental/IoStore/OnDemand/Private/OnDemandIoStore.cpp
+++ b/Engine/Source/Runtime/Experimental/IoStore/OnDemand/Private/OnDemandIoStore.cpp
@@ -50,7 +50,7 @@
 #endif
 
 ///////////////////////////////////////////////////////////////////////////////
-bool GHttpIoDispatcherEnabled = false;
+bool GHttpIoDispatcherEnabled = true;
 static FAutoConsoleVariableRef CVar_HttpIoDispatcherEnabled(
 	TEXT("iax.HttpIoDispatcherEnabled"),
 	GHttpIoDispatcherEnabled,
