Fix crash related to CSV net stats being reported with different types (float and integer) due to new change that allows for replication system to tick with no connections.
#rnx

diff --git a/Engine/Source/Runtime/Experimental/Iris/Core/Private/Iris/Stats/NetStats.cpp b/Engine/Source/Runtime/Experimental/Iris/Core/Private/Iris/Stats/NetStats.cpp
--- a/Engine/Source/Runtime/Experimental/Iris/Core/Private/Iris/Stats/NetStats.cpp
+++ b/Engine/Source/Runtime/Experimental/Iris/Core/Private/Iris/Stats/NetStats.cpp
@@ -86,12 +86,12 @@ void FNetSendStats::ReportCsvStats()
 	}
 	else
 	{
-		CSV_CUSTOM_STAT(Iris, AvgScheduledForReplicationRootObjectCount, 0, ECsvCustomStatOp::Set);
-		CSV_CUSTOM_STAT(Iris, AvgReplicatedRootObjectCount, 0, ECsvCustomStatOp::Set);
-		CSV_CUSTOM_STAT(Iris, AvgReplicatedObjectCount, 0, ECsvCustomStatOp::Set);
-		CSV_CUSTOM_STAT(Iris, AvgReplicatedDestructionInfoCount, 0, ECsvCustomStatOp::Set);
-		CSV_CUSTOM_STAT(Iris, AvgReplicatedHugeObjectCount, 0, ECsvCustomStatOp::Set);
-		CSV_CUSTOM_STAT(Iris, AvgDeltaCompressedObjectCount, 0, ECsvCustomStatOp::Set);
+		CSV_CUSTOM_STAT(Iris, AvgScheduledForReplicationRootObjectCount, 0.0f, ECsvCustomStatOp::Set);
+		CSV_CUSTOM_STAT(Iris, AvgReplicatedRootObjectCount, 0.0f, ECsvCustomStatOp::Set);
+		CSV_CUSTOM_STAT(Iris, AvgReplicatedObjectCount, 0.0f, ECsvCustomStatOp::Set);
+		CSV_CUSTOM_STAT(Iris, AvgReplicatedDestructionInfoCount, 0.0f, ECsvCustomStatOp::Set);
+		CSV_CUSTOM_STAT(Iris, AvgReplicatedHugeObjectCount, 0.0f, ECsvCustomStatOp::Set);
+		CSV_CUSTOM_STAT(Iris, AvgDeltaCompressedObjectCount, 0.0f, ECsvCustomStatOp::Set);
 	}
 
 	// Huge object counts
