Chaos - fix per-particle Iteration Counts in Rewind/Resim

#rnx

diff --git a/Engine/Source/Runtime/Experimental/Chaos/Public/Chaos/Evolution/IterationSettings.h b/Engine/Source/Runtime/Experimental/Chaos/Public/Chaos/Evolution/IterationSettings.h
--- a/Engine/Source/Runtime/Experimental/Chaos/Public/Chaos/Evolution/IterationSettings.h
+++ b/Engine/Source/Runtime/Experimental/Chaos/Public/Chaos/Evolution/IterationSettings.h
@@ -64,6 +64,16 @@ namespace Chaos
 			void SetNumPositionIterations(const int32 InNum) { NumPositionIterations = FIndexType(FMath::Clamp(InNum, InvalidIterations, MaxIterations)); }
 			void SetNumVelocityIterations(const int32 InNum) { NumVelocityIterations = FIndexType(FMath::Clamp(InNum, InvalidIterations, MaxIterations)); }
 			void SetNumProjectionIterations(const int32 InNum) { NumProjectionIterations = FIndexType(FMath::Clamp(InNum, InvalidIterations, MaxIterations)); }
+
+			friend bool operator==(const TIterationSettings& L, const TIterationSettings& R)
+			{
+				return (L.NumPositionIterations == R.NumPositionIterations) && (L.NumVelocityIterations == R.NumVelocityIterations) && (L.NumProjectionIterations == R.NumProjectionIterations);
+			}
+
+			friend bool operator!=(const TIterationSettings& L, const TIterationSettings& R)
+			{
+				return !(L == R);
+			}
 		};
 
 		// NOTE: only use signed types (or modify how we handle TIterationSettings::InvalidIterations above)
diff --git a/Engine/Source/Runtime/Experimental/Chaos/Public/Chaos/ParticleDirtyFlags.h b/Engine/Source/Runtime/Experimental/Chaos/Public/Chaos/ParticleDirtyFlags.h
--- a/Engine/Source/Runtime/Experimental/Chaos/Public/Chaos/ParticleDirtyFlags.h
+++ b/Engine/Source/Runtime/Experimental/Chaos/Public/Chaos/ParticleDirtyFlags.h
@@ -460,10 +460,12 @@ public:
 		// @todo(chaos): add this
 		//Ar << MInitialOverlapDepenetrationVelocity;
 		//Ar << MSleepThresholdMultiplier;
+		//Ar << MIterationSettings
 		if (Ar.IsLoading())
 		{
 			MInitialOverlapDepenetrationVelocity = 0;
 			MSleepThresholdMultiplier = 1.0f;
+			MIterationSettings.Empty();
 		}
 
 		if (bAddControlFlags)
@@ -486,6 +488,7 @@ public:
 		SetSleepType(Other.SleepType());
 		SetCollisionConstraintFlags(Other.CollisionConstraintFlags());
 		SetControlFlags(Other.ControlFlags());
+		SetIterationSettings(Other.IterationSettings());
 		SetDisabled(Other.Disabled());
 	}
 
@@ -503,6 +506,7 @@ public:
 			&& SleepType() == Other.SleepType()
 			&& CollisionConstraintFlags() == Other.CollisionConstraintFlags()
 			&& ControlFlags() == Other.ControlFlags()
+			&& IterationSettings() == Other.IterationSettings()
 			&& Disabled() == Other.Disabled();
 	}
 
