Chaos Modular Vehicles:

Use static FName for type in SimModuleBase to improve perf for comparisons/access (FName creation is expensive).

diff --git a/Engine/Source/Runtime/Experimental/ChaosVehicles/ChaosVehiclesCore/Public/SimModule/SimulationModuleBase.h b/Engine/Source/Runtime/Experimental/ChaosVehicles/ChaosVehiclesCore/Public/SimModule/SimulationModuleBase.h
--- a/Engine/Source/Runtime/Experimental/ChaosVehicles/ChaosVehiclesCore/Public/SimModule/SimulationModuleBase.h
+++ b/Engine/Source/Runtime/Experimental/ChaosVehicles/ChaosVehiclesCore/Public/SimModule/SimulationModuleBase.h
@@ -275,7 +275,7 @@ namespace Chaos
 	
 
 #define DEFINE_CHAOSSIMTYPENAME(cls) \
-	 static FName _CHAOSTYPENAMERETRIVAL_() { return FName(#cls); }
+	 static FName _CHAOSTYPENAMERETRIVAL_() { static FName Name(#cls); return Name; }
 	
 	//FName TSimulationModuleTypeName<class cls>::SimModuleTypeName = FName(#cls); 
 
