Fixup API macros

#rnx

diff --git a/Engine/Source/Runtime/Experimental/IoStore/OnDemandCore/Public/IO/HttpIoDispatcher.h b/Engine/Source/Runtime/Experimental/IoStore/OnDemandCore/Public/IO/HttpIoDispatcher.h
--- a/Engine/Source/Runtime/Experimental/IoStore/OnDemandCore/Public/IO/HttpIoDispatcher.h
+++ b/Engine/Source/Runtime/Experimental/IoStore/OnDemandCore/Public/IO/HttpIoDispatcher.h
@@ -42,7 +42,7 @@ ENUM_CLASS_FLAGS(EIoHttpFlags);
 /** Represents a range within a resource. */
 struct FIoHttpRange
 {
-	UE_API						FIoHttpRange() = default;
+						FIoHttpRange() = default;
 	explicit inline				FIoHttpRange(uint32 InMin, uint32 InMax);
 
 	bool						IsValid()	const	{ return Min <= Max; }
@@ -111,7 +111,7 @@ public:
 	inline						FIoHttpOptions(int32 InPriority, int32 InRetryCount, EIoHttpFlags InFlags);
 	inline						FIoHttpOptions(int32 InPriority, int32 InRetryCount);
 	inline explicit				FIoHttpOptions(int32 InPriority);
-	UE_API						FIoHttpOptions() = default;
+						FIoHttpOptions() = default;
 
 	const FIoHttpRange&			GetRange()		const { return Range; }
 	int32						GetPriority()	const { return Priority; }
@@ -164,9 +164,9 @@ FIoHttpOptions::FIoHttpOptions(int32 InPriority)
 class FIoHttpHeaders
 {
 public:
-	UE_API							FIoHttpHeaders() = default;
-	UE_API							FIoHttpHeaders(const FIoHttpHeaders&) = default;
-	UE_API							FIoHttpHeaders(FIoHttpHeaders&&) = default;
+							FIoHttpHeaders() = default;
+							FIoHttpHeaders(const FIoHttpHeaders&) = default;
+							FIoHttpHeaders(FIoHttpHeaders&&) = default;
 	UE_API FIoHttpHeaders&			Add(FAnsiString&& HeaderName, FAnsiString&& HeaderValue);
 	UE_API FIoHttpHeaders&			Add(FAnsiStringView HeaderName, FAnsiStringView HeaderValue);
 	UE_API TArray<FAnsiString>		ToArray() &&;
@@ -187,7 +187,7 @@ private:
 class FIoRelativeUrl
 {
 public:
-	UE_API							FIoRelativeUrl() = default;
+							FIoRelativeUrl() = default;
 	bool							IsEmpty() const		{ return Url.IsEmpty(); }
 	int32							Len() const			{ return Url.Len(); }
 	const ANSICHAR*					ToString() const	{ return *Url; }
@@ -210,7 +210,7 @@ private:
 class FIoHttpRequest
 {
 public:
-	UE_API					FIoHttpRequest() = default;
+					FIoHttpRequest() = default;
 	explicit 				FIoHttpRequest(FIoHttpRequestHandle InHandle)
 								: Handle(InHandle) { }
 							FIoHttpRequest(FIoHttpRequest&& Other)
@@ -245,7 +245,7 @@ ENUM_CLASS_FLAGS(EIoHttpResponseFlags);
 class FIoHttpResponse
 {
 public:
-	UE_API					FIoHttpResponse() = default;
+					FIoHttpResponse() = default;
 	inline					FIoHttpResponse(const FIoBuffer& InBody, EIoErrorCode InErrorCode, uint32 InStatusCode, EIoHttpResponseFlags InFlags);
 							FIoHttpResponse(EIoErrorCode InErrorCode, uint32 InStatusCode)
 								: ErrorCode(InErrorCode)
