Don't treat cancelled IAS requests as failed.

#rnx

diff --git a/Engine/Source/Runtime/Experimental/IoStore/OnDemand/Private/OnDemandHttpIoDispatcherBackend.cpp b/Engine/Source/Runtime/Experimental/IoStore/OnDemand/Private/OnDemandHttpIoDispatcherBackend.cpp
--- a/Engine/Source/Runtime/Experimental/IoStore/OnDemand/Private/OnDemandHttpIoDispatcherBackend.cpp
+++ b/Engine/Source/Runtime/Experimental/IoStore/OnDemand/Private/OnDemandHttpIoDispatcherBackend.cpp
@@ -84,7 +84,7 @@ public:
 
 private:
 	void									CompleteRequest(FResolvedRequest& ResolvedRequst, FIoHttpResponse&& HttpResponse);
-	static void								LogRequest(FResolvedRequest& ResolvedRequest, FIoRequestImpl& DispacherRequest, bool bDecoded, bool bCached, uint64 DurationMs);
+	static void								LogRequest(FResolvedRequest& ResolvedRequest, FIoRequestImpl& DispacherRequest, bool bDecoded, bool bCached, bool bCancelled, uint64 DurationMs);
 
 	FBackendStatus			BackendStatus;
 	FOnDemandIoBackendStats	Stats;
@@ -251,8 +251,10 @@ void FHttpIoDispatcherBackend::CompleteRequest(FResolvedRequest& ResolvedRequest
 	FOnDemandChunkInfo& ChunkInfo		= ResolvedRequest.ChunkInfo;
 	FIoRequestImpl& DispatcherRequest	= ResolvedRequest.DispatcherRequest;
 
-	bool bDecoded = false;
-	if (HttpResponse.IsOk() && Chunk.GetSize() > 0)
+	const bool bCancelled				= DispatcherRequest.IsCancelled();
+	bool bDecoded						= false;
+
+	if (!bCancelled && HttpResponse.IsOk() && Chunk.GetSize() > 0)
 	{
 		FIoChunkDecodingParams Params;
 		Params.EncryptionKey				= ChunkInfo.EncryptionKey();
@@ -279,6 +281,10 @@ void FHttpIoDispatcherBackend::CompleteRequest(FResolvedRequest& ResolvedRequest
 	{
 		Stats.OnIoRequestComplete(DispatcherRequest.GetBuffer().GetSize(), DurationMs);
 	}
+	else if (bCancelled)
+	{
+		Stats.OnIoRequestCancel();
+	}
 	else
 	{
 		DispatcherRequest.SetFailed();
@@ -293,7 +299,7 @@ void FHttpIoDispatcherBackend::CompleteRequest(FResolvedRequest& ResolvedRequest
 		}
 	}
 	
-	LogRequest(ResolvedRequest, DispatcherRequest, bDecoded, HttpResponse.IsCached(), DurationMs);
+	LogRequest(ResolvedRequest, DispatcherRequest, bDecoded, HttpResponse.IsCached(), bCancelled, DurationMs);
 
 	{
 		TUniqueLock Lock(Mutex);
@@ -303,8 +309,13 @@ void FHttpIoDispatcherBackend::CompleteRequest(FResolvedRequest& ResolvedRequest
 	BackendContext->WakeUpDispatcherThreadDelegate.Execute();
 }
 
-void FHttpIoDispatcherBackend::LogRequest(FResolvedRequest& ResolvedRequest, FIoRequestImpl& DispacherRequest, bool bDecoded, bool bCached, uint64 DurationMs)
+void FHttpIoDispatcherBackend::LogRequest(FResolvedRequest& ResolvedRequest, FIoRequestImpl& DispacherRequest, bool bDecoded, bool bCached, bool bCancelled, uint64 DurationMs)
 {
+	if (bCancelled)
+	{
+		return;
+	}
+
 	const TCHAR* Prefix = [bDecoded, bCached]() -> const TCHAR*
 	{
 		if (bDecoded)
