Iris
- Fixed issue with NetObjects being asked to stop replication from client code not being removed from list when they actually did receive end replication from server
- This lead to issues if the same objeckt was re-instantiated in a later received packet processed during the same frame.
- Fixed issue where NetRefHandles that failed to be instantiated was kept in a broken list even after a receiving a endreplication request from server.
-    NetRefHandles are now properly removed from the broken list when we receive endreplication.

diff --git a/Engine/Source/Runtime/Experimental/Iris/Core/Private/Iris/ReplicationSystem/ReplicationBridge.cpp b/Engine/Source/Runtime/Experimental/Iris/Core/Private/Iris/ReplicationSystem/ReplicationBridge.cpp
--- a/Engine/Source/Runtime/Experimental/Iris/Core/Private/Iris/ReplicationSystem/ReplicationBridge.cpp
+++ b/Engine/Source/Runtime/Experimental/Iris/Core/Private/Iris/ReplicationSystem/ReplicationBridge.cpp
@@ -514,6 +514,9 @@ void UReplicationBridge::InternalDestroyNetObject(FNetRefHandle Handle)
 		FWorldLocations& WorldLocations = ReplicationSystemInternal->GetWorldLocations();
 		WorldLocations.RemoveObjectInfoCache(ObjectInternalIndex);
 
+		// Remove from handles to stop replicating as we might respawn the same object again in a later packet
+		HandlesToStopReplicating.Remove(Handle);
+
 		FNetRefHandleManager::FReplicatedObjectData& ObjectData = NetRefHandleManager->GetReplicatedObjectDataNoCheck(ObjectInternalIndex);
 
 		// At this point we can no longer instantiate objects and must explicitly clear out objects pending create from ReplicationWriter
diff --git a/Engine/Source/Runtime/Experimental/Iris/Core/Private/Iris/ReplicationSystem/ReplicationReader.cpp b/Engine/Source/Runtime/Experimental/Iris/Core/Private/Iris/ReplicationSystem/ReplicationReader.cpp
--- a/Engine/Source/Runtime/Experimental/Iris/Core/Private/Iris/ReplicationSystem/ReplicationReader.cpp
+++ b/Engine/Source/Runtime/Experimental/Iris/Core/Private/Iris/ReplicationSystem/ReplicationReader.cpp
@@ -607,6 +607,11 @@ uint32 FReplicationReader::ReadRootObjectsPendingDestroy(FNetSerializationContex
 			{
 				// If we did not find the object or associated bridge, the packet that would have created the object may have been lost.
 				UE_LOG(LogIris, Verbose, TEXT("FReplicationReader::Read Tried to destroy object %s. This can occur if the server sends destroy for an object that has not yet been confirmed as created."), *IncompleteHandle.ToString());
+				// Remove from broken list, so that we can try again if the server asks us to instantiate to object again.
+				BrokenObjects.SetNum(Algo::RemoveIf(BrokenObjects, [&IncompleteHandle](const FNetRefHandle& Handle)
+				{
+					return Handle.GetId() == IncompleteHandle.GetId();
+				}));
 			}
 		}
 	}
