Fixed typo in IAS bandwidth stat name and change to use double instead of integer.

diff --git a/Engine/Source/Runtime/Experimental/IoStore/OnDemand/Private/Statistics.cpp b/Engine/Source/Runtime/Experimental/IoStore/OnDemand/Private/Statistics.cpp
--- a/Engine/Source/Runtime/Experimental/IoStore/OnDemand/Private/Statistics.cpp
+++ b/Engine/Source/Runtime/Experimental/IoStore/OnDemand/Private/Statistics.cpp
@@ -250,6 +250,7 @@ private:
 
 #if COUNTERSTRACE_ENABLED
 	using FCounterInt		= FCountersTrace::FCounterInt;
+	using FCounterFloat		= FCountersTrace::FCounterFloat;
 	using FCounterAtomicInt = FCountersTrace::FCounterAtomicInt;
 #else
 	template <typename Type>
@@ -263,7 +264,19 @@ private:
 		int64 Get() const { return V;}
 		Type V = 0;
 	};
+	template <typename Type>
+	struct TCounterFloat
+	{
+		TCounterFloat(...) {}
+		void Set(Type i) { V = i; }
+		void Add(Type d) { V += d; }
+		void Increment() { V += 1; }
+		void Decrement() { V -= 1; }
+		Type Get() const { return V;}
+		Type V = 0;
+	};
 	using FCounterInt		= TCounterInt<int64>;
+	using FCounterFloat		= TCounterFloat<double>;
 	using FCounterAtomicInt = TCounterInt<std::atomic<int64>>;
 #endif
 
@@ -305,7 +318,7 @@ UE_IAX_COUNTER(HttpPendingCount, FCounterAtomicInt);
 UE_IAX_COUNTER(HttpInflightCount, FCounterInt);
 UE_IAX_MEMORY_COUNTER(HttpDownloadedBytes, FCounterInt);
 UE_IAX_COUNTER(HttpDurationMs, FCounterInt);
-UE_IAX_COUNTER(HttpBandwidthMbps, FCounterInt);
+UE_IAX_COUNTER(HttpBandwidthMbps, FCounterFloat);
 double GHttpDurationMsAvg[(uint8)EHttpRequestType::NUM_SOURCES] = { 0.0 };
 int32 GHttpDurationMsMax[(uint8)EHttpRequestType::NUM_SOURCES] = { 0 };
 int64 GHttpDurationMsSum[(uint8)EHttpRequestType::NUM_SOURCES] = { 0 };
@@ -340,9 +353,9 @@ struct FHttpRecentHistoryStatistics
 		Index = (Index + 1) % HistoryCount;
 	}
 
-	int64 GetBandwidthMbps() const
+	double GetBandwidthMbps() const
 	{
-		return (TotalBytes * 8) / (TotalDuration + 1) / 1000;
+		return double(TotalBytes * 8) / double(TotalDuration + 1) / 1000.0;
 	}
 
 	double GetAverage() const
@@ -399,7 +412,7 @@ UE_IAX_CSV_DEFINE_STAT(HttpCancelCount);
 UE_IAX_CSV_DEFINE_STAT(HttpErrorCount);
 UE_IAX_CSV_DEFINE_STAT(HttpPendingCount);
 UE_IAX_CSV_DEFINE_STAT(HttpDownloadedMB);
-UE_IAX_CSV_DEFINE_STAT(HttpBandwidthMpbs);
+UE_IAX_CSV_DEFINE_STAT(HttpBandwidthMbps);
 UE_IAX_CSV_DEFINE_STAT(HttpDurationMsAvg);
 UE_IAX_CSV_DEFINE_STAT(HttpDurationMsMax);
 
@@ -562,7 +575,7 @@ void FOnDemandIoBackendStats::UpdateCSVValues()
 			UE_IAX_CSV_CUSTOM_STAT_DEFINED(Type, HttpErrorCount, (int32)GHttpErrorCount[(uint8)Type].Get(), ECsvCustomStatOp::Set);
 			UE_IAX_CSV_CUSTOM_STAT_DEFINED(Type, HttpPendingCount, (int32)GHttpPendingCount[(uint8)Type].Get(), ECsvCustomStatOp::Set);
 			UE_IAX_CSV_CUSTOM_STAT_DEFINED(Type, HttpDownloadedMB, (float)GHttpDownloadedBytes[(uint8)Type].Get() * OneOver1024 * OneOver1024, ECsvCustomStatOp::Set);
-			UE_IAX_CSV_CUSTOM_STAT_DEFINED(Type, HttpBandwidthMpbs, (int32)GHttpBandwidthMbps[(uint8)Type].Get(), ECsvCustomStatOp::Set);
+			UE_IAX_CSV_CUSTOM_STAT_DEFINED(Type, HttpBandwidthMbps, GHttpBandwidthMbps[(uint8)Type].Get(), ECsvCustomStatOp::Set);
 			UE_IAX_CSV_CUSTOM_STAT_DEFINED(Type, HttpDurationMsAvg, (int32)GHttpDurationMsAvg[(uint8)Type], ECsvCustomStatOp::Set);
 			UE_IAX_CSV_CUSTOM_STAT_DEFINED(Type, HttpDurationMsMax, GHttpDurationMsMax[(uint8)Type], ECsvCustomStatOp::Set);
 		};
