Constraints: make sure the tick function is safe to be registered / unregistered

#rnx

diff --git a/Engine/Source/Runtime/Experimental/Animation/Constraints/Private/Transform/TransformConstraint.cpp b/Engine/Source/Runtime/Experimental/Animation/Constraints/Private/Transform/TransformConstraint.cpp
--- a/Engine/Source/Runtime/Experimental/Animation/Constraints/Private/Transform/TransformConstraint.cpp
+++ b/Engine/Source/Runtime/Experimental/Animation/Constraints/Private/Transform/TransformConstraint.cpp
@@ -623,21 +623,29 @@ void UTickableTransformConstraint::PostEvaluate() const
 
 void UTickableTransformConstraint::InitConstraint(UWorld *InWorld)
 {
-	if (!ensure(InWorld))
+	ULevel* Level = InWorld ? InWorld->GetCurrentLevel() : nullptr;
+	if (!ensure(Level))
 	{
 		return;
 	}
 	
-	FConstraintTickFunction& ConstraintTick = ConstraintTicks.FindOrAdd(InWorld->GetCurrentLevel());
+	FConstraintTickFunction& ConstraintTick = ConstraintTicks.FindOrAdd(Level);
 
 	if (ConstraintTick.ConstraintFunctions.IsEmpty())
 	{
 		ConstraintTick.RegisterFunction(GetFunction());
 	}
-	ConstraintTick.RegisterTickFunction(InWorld->GetCurrentLevel());
+
+	if (!ConstraintTick.IsTickFunctionRegistered())
+	{
+		ConstraintTick.RegisterTickFunction(Level);
+	}
+	
 	ConstraintTick.Constraint = this;
+
 	SetupDependencies(InWorld);
 	RegisterDelegates();
+
 	bValid = true;
 }
 
@@ -647,22 +655,25 @@ void UTickableTransformConstraint::TeardownConstraint(UWorld* InWorld)
 	{
 		return;
 	}
-	
-	FConstraintTickFunction& ConstraintTick = ConstraintTicks.FindOrAdd(InWorld->GetCurrentLevel());
-	ConstraintTick.SetTickFunctionEnable(false);
-	ConstraintTick.UnRegisterTickFunction();
 
-	if (FTickFunction* ChildTickFunction = GetChildHandleTickFunction())
+	const ULevel* Level = InWorld ? InWorld->GetCurrentLevel() : nullptr;
+	if (FConstraintTickFunction* ConstraintTick = Level ? ConstraintTicks.Find(Level) : nullptr)
 	{
-		ChildTickFunction->RemovePrerequisite(this, ConstraintTick);
-	}
+		ConstraintTick->UnRegisterTickFunction();
+		ConstraintTick->SetTickFunctionEnable(false);
+
+		if (FTickFunction* ChildTickFunction = GetChildHandleTickFunction())
+		{
+			ChildTickFunction->RemovePrerequisite(this, *ConstraintTick);
+		}
 	
-	if (FTickFunction* ParentTickFunction = GetParentHandleTickFunction())
-	{
-		ConstraintTick.RemovePrerequisite(ParentTRSHandle->GetPrerequisiteObject(), *ParentTickFunction);	
+		if (FTickFunction* ParentTickFunction = GetParentHandleTickFunction())
+		{
+			ConstraintTick->RemovePrerequisite(ParentTRSHandle->GetPrerequisiteObject(), *ParentTickFunction);	
+		}
+
+		ConstraintTicks.Remove(Level);
 	}
-	
-	ConstraintTicks.Remove(InWorld->GetCurrentLevel());
 
 	// unregister delegates. should handles delegates be unregistered as well?
 	UnregisterDelegates();
