Headless Chaos CIS fix: svd decomposition

diff --git a/Engine/Source/Runtime/Experimental/ChaosCore/Public/Chaos/ImplicitQRSVD.h b/Engine/Source/Runtime/Experimental/ChaosCore/Public/Chaos/ImplicitQRSVD.h
--- a/Engine/Source/Runtime/Experimental/ChaosCore/Public/Chaos/ImplicitQRSVD.h
+++ b/Engine/Source/Runtime/Experimental/ChaosCore/Public/Chaos/ImplicitQRSVD.h
@@ -375,7 +375,7 @@ namespace Chaos {
 		T x = S_Sym.M[0*2+0];
 		T y = S_Sym.M[1*2+0];
 		T z = S_Sym.M[3];
-		if (y < tol) {
+		if (FGenericPlatformMath::Abs(y) < tol) {
 			// S is already diagonal
 			cosine = (T)1.;
 			sine = (T)0.;
