Chaos flesh: post 5.6 fix, solver segfault failsafe

diff --git a/Engine/Plugins/Experimental/ChaosFlesh/Source/ChaosFleshNodes/Private/Dataflow/ChaosFleshComputeFiberFieldNode.cpp b/Engine/Plugins/Experimental/ChaosFlesh/Source/ChaosFleshNodes/Private/Dataflow/ChaosFleshComputeFiberFieldNode.cpp
--- a/Engine/Plugins/Experimental/ChaosFlesh/Source/ChaosFleshNodes/Private/Dataflow/ChaosFleshComputeFiberFieldNode.cpp
+++ b/Engine/Plugins/Experimental/ChaosFlesh/Source/ChaosFleshNodes/Private/Dataflow/ChaosFleshComputeFiberFieldNode.cpp
@@ -17,11 +17,14 @@ void FComputeFiberFieldNode::Evaluate(UE::Dataflow::FContext& Context, const FDa
 		//
 		// Gather inputs
 		//
-		if (!IsConnected(&Collection) || !IsConnected(&OriginIndices) || !IsConnected(&InsertionIndices))
+		FManagedArrayCollection InCollection = GetValue<FManagedArrayCollection>(Context, &Collection);
+		if (!IsConnected(&OriginIndices) || !IsConnected(&InsertionIndices))
 		{
+			Context.Warning(FString::Printf(
+				TEXT("OriginIndices/InsertionIndices is not connected.")), this);
+			Out->SetValue(MoveTemp(InCollection), Context);
 			return;
 		}
-		FManagedArrayCollection InCollection = GetValue<FManagedArrayCollection>(Context, &Collection);
 		TArray<int32> InOriginIndices = GetValue<TArray<int32>>(Context, &OriginIndices);
 		TArray<int32> InInsertionIndices = GetValue<TArray<int32>>(Context, &InsertionIndices);
 
diff --git a/Engine/Source/Runtime/Experimental/Chaos/Public/Chaos/Deformable/MuscleActivationConstraints.h b/Engine/Source/Runtime/Experimental/Chaos/Public/Chaos/Deformable/MuscleActivationConstraints.h
--- a/Engine/Source/Runtime/Experimental/Chaos/Public/Chaos/Deformable/MuscleActivationConstraints.h
+++ b/Engine/Source/Runtime/Experimental/Chaos/Public/Chaos/Deformable/MuscleActivationConstraints.h
@@ -43,6 +43,15 @@ namespace Chaos::Softs
 						ensureMsgf(MuscleActivationData.InflationVolumeScale > 0,
 							TEXT("FiberLengthRatioAtMaxActivation %f of muscle indexed %d is <= 0, please check your setup."), MuscleActivationData.InflationVolumeScale, MuscleIdx))
 					{
+						MuscleActivationData.OriginInsertionPair[0] += VertexOffset;
+						MuscleActivationData.OriginInsertionPair[1] += VertexOffset;
+						if (!RestParticles.XArray().IsValidIndex(MuscleActivationData.OriginInsertionPair[0]) ||
+							!RestParticles.XArray().IsValidIndex(MuscleActivationData.OriginInsertionPair[1]))
+						{
+							UE_LOG(LogMuscleActivationConstraint, Error, TEXT("Muscle Idx[%d] has invalid origin[%d] or insertion[%d]."),
+								MuscleIdx, MuscleActivationData.OriginInsertionPair[0], MuscleActivationData.OriginInsertionPair[1]);
+							continue;
+						}
 						const int32 OldSize = MuscleActivationElement.AddDefaulted(1);
 						FiberDirectionMatrix.AddDefaulted(1);
 						ContractionVolumeScale.AddDefaulted(1);
@@ -58,8 +67,6 @@ namespace Chaos::Softs
 						FiberLengthRatioAtMaxActivation.Add(MuscleActivationData.FiberLengthRatioAtMaxActivation);
 						MuscleLengthRatioThresholdForMaxActivation.Add(MuscleActivationData.MuscleLengthRatioThresholdForMaxActivation);
 						InflationVolumeScale.Add(MuscleActivationData.InflationVolumeScale);
-						MuscleActivationData.OriginInsertionPair[0] += VertexOffset;
-						MuscleActivationData.OriginInsertionPair[1] += VertexOffset;
 						OriginInsertionPair.Add(MuscleActivationData.OriginInsertionPair);
 						MuscleRestLength.Add((RestParticles.GetX(MuscleActivationData.OriginInsertionPair[0]) - RestParticles.GetX(MuscleActivationData.OriginInsertionPair[1])).Size());
 						MuscleActivation.Add(0.f);
