Add support for converting a geometry collection to a skeletal mesh / skeleton
- Add a collection to skeletal mesh Dataflow node
- Add an option in the right click menu of the asset to export to a skeletal mesh

diff --git a/Engine/Plugins/Experimental/GeometryCollectionPlugin/Source/GeometryCollectionEditor/GeometryCollectionEditor.Build.cs b/Engine/Plugins/Experimental/GeometryCollectionPlugin/Source/GeometryCollectionEditor/GeometryCollectionEditor.Build.cs
--- a/Engine/Plugins/Experimental/GeometryCollectionPlugin/Source/GeometryCollectionEditor/GeometryCollectionEditor.Build.cs
+++ b/Engine/Plugins/Experimental/GeometryCollectionPlugin/Source/GeometryCollectionEditor/GeometryCollectionEditor.Build.cs
@@ -38,6 +38,13 @@ namespace UnrealBuildTool.Rules
 					"DataflowEditor",
 				}
 			);
+			
+			PrivateDependencyModuleNames.AddRange(
+				new string[] {
+					"ContentBrowser",
+					"EditorScriptingUtilities",
+				}
+			);
 
 			PrivateDefinitions.Add("CHAOS_INCLUDE_LEVEL_1=1");
 		}
diff --git a/Engine/Plugins/Experimental/GeometryCollectionPlugin/Source/GeometryCollectionEditor/Private/GeometryCollection/AssetDefinition_GeometryCollection.cpp b/Engine/Plugins/Experimental/GeometryCollectionPlugin/Source/GeometryCollectionEditor/Private/GeometryCollection/AssetDefinition_GeometryCollection.cpp
--- a/Engine/Plugins/Experimental/GeometryCollectionPlugin/Source/GeometryCollectionEditor/Private/GeometryCollection/AssetDefinition_GeometryCollection.cpp
+++ b/Engine/Plugins/Experimental/GeometryCollectionPlugin/Source/GeometryCollectionEditor/Private/GeometryCollection/AssetDefinition_GeometryCollection.cpp
@@ -2,13 +2,22 @@
 
 #include "GeometryCollection/AssetDefinition_GeometryCollection.h"
 #include "GeometryCollection/GeometryCollectionObject.h"
+#include "GeometryCollection/GeometryCollectionEngineConversion.h"
 #include "Dataflow/DataflowEditor.h"
 #include "Dataflow/DataflowEditorModule.h"
 #include "Dataflow/DataflowEditorToolkit.h"
 #include "Math/Color.h"
 #include "ThumbnailRendering/SceneThumbnailInfo.h"
 #include "Toolkits/SimpleAssetEditor.h"
+#include "ContentBrowserMenuContexts.h"
+#include "ToolMenu.h"
+#include "ToolMenus.h"
+#include "AssetToolsModule.h"
+#include "EditorAssetLibrary.h"
+#include "AssetRegistry/AssetRegistryModule.h"
 
+#include "Engine/SkeletalMesh.h"
+#include "Animation/Skeleton.h"
 
 #define LOCTEXT_NAMESPACE "AssetActions_GeometryCollection"
 
@@ -100,4 +109,107 @@ EAssetCommandResult UAssetDefinition_GeometryCollection::OpenAssets(const FAsset
 	return EAssetCommandResult::Unhandled;
 }
 
+// Menu Extensions
+//--------------------------------------------------------------------
+
+namespace MenuExtension_GeometryCollection
+{
+	template <typename T>
+	static T* GetOrCreateAsset(const FString Path, FString Name)
+	{
+		const FString AssetPath = FPaths::Combine(Path, Name);
+		if (UEditorAssetLibrary::DoesAssetExist(AssetPath))
+		{
+			if (T* ExistingObject = Cast<T>(UEditorAssetLibrary::LoadAsset(AssetPath)))
+			{
+				return ExistingObject;
+			}
+		}
+		// let's create new one 
+		FString UniqueAssetName;
+		FString UniquePackageName;
+		FAssetToolsModule& AssetToolsModule = FModuleManager::LoadModuleChecked<FAssetToolsModule>("AssetTools");
+		AssetToolsModule.Get().CreateUniqueAssetName(AssetPath, TEXT(""), UniquePackageName, UniqueAssetName);
+
+		// Create a package for the asset
+		UObject* OuterForAsset = CreatePackage(*UniquePackageName);
+
+		// Create a frame in the package
+		T* NewAsset = NewObject<T>(OuterForAsset, T::StaticClass(), *UniqueAssetName, RF_Public | RF_Standalone | RF_Transactional);
+		FAssetRegistryModule::AssetCreated(NewAsset);
+
+		NewAsset->Modify();
+		return NewAsset;
+	}
+
+	static void ExecuteExportToSkeletalMesh(const FToolMenuContext& InContext)
+	{
+		const FAssetToolsModule& AssetToolsModule = FModuleManager::Get().LoadModuleChecked<FAssetToolsModule>("AssetTools");
+
+		if (const UContentBrowserAssetContextMenuContext* Context = UContentBrowserAssetContextMenuContext::FindContextWithAssets(InContext))
+		{
+			TArray<UGeometryCollection*> GeometryCollections = Context->LoadSelectedObjects<UGeometryCollection>();
+			if (!GeometryCollections.IsEmpty())
+			{
+				for (const UGeometryCollection* GeometryCollection : GeometryCollections)
+				{
+					if (GeometryCollection)
+					{
+						// find path to reference asset
+						UPackage* AssetOuterPackage = CastChecked<UPackage>(GeometryCollection->GetOuter());
+						FString AssetPackageName = AssetOuterPackage->GetName();
+						FString AssetFolderPath = FPackageName::GetLongPackagePath(AssetPackageName);
+						FString AssetBaseName = FPaths::GetBaseFilename(AssetPackageName);
+						AssetBaseName.RemoveFromStart(TEXT("GC_"));
+
+						FString SkeletalMeshAssetName = FString::Printf(TEXT("SK_%s"), *AssetBaseName);
+						FString SkeletonAssetName = FString::Printf(TEXT("SKEL_%s"), *AssetBaseName);
+
+						USkeletalMesh* SkeletalMesh = GetOrCreateAsset<USkeletalMesh>(AssetFolderPath, SkeletalMeshAssetName);
+						USkeleton* Skeleton = GetOrCreateAsset<USkeleton>(AssetFolderPath, SkeletonAssetName);
+
+						if (SkeletalMesh && Skeleton)
+						{
+
+							if (GeometryCollection->GetGeometryCollection())
+							{
+								FGeometryCollectionEngineConversion::ConvertCollectionToSkeletalMesh(*GeometryCollection->GetGeometryCollection(), GeometryCollection->Materials, *SkeletalMesh, *Skeleton);
+							}
+
+							SkeletalMesh->PostEditChange();
+							Skeleton->PostEditChange();
+
+							SkeletalMesh->MarkPackageDirty();
+							Skeleton->MarkPackageDirty();
+
+							FAssetRegistryModule::AssetCreated(SkeletalMesh);
+							FAssetRegistryModule::AssetCreated(Skeleton);
+						}
+					}
+				}
+			}
+		}
+
+	}
+
+	static FDelayedAutoRegisterHelper DelayedAutoRegister(EDelayedRegisterRunPhase::EndOfEngineInit, [] {
+		UToolMenus::RegisterStartupCallback(
+			FSimpleMulticastDelegate::FDelegate::CreateLambda([]()
+				{
+					FToolMenuOwnerScoped OwnerScoped(UE_MODULE_NAME);
+					UToolMenu* Menu = UE::ContentBrowser::ExtendToolMenu_AssetContextMenu(UGeometryCollection::StaticClass());
+					FToolMenuSection& Section = Menu->FindOrAddSection("GetAssetActions");
+					{
+						const TAttribute<FText> Label = LOCTEXT("ExportToSkeletalMesh_Label", "Export to SkeletalMesh");
+						const TAttribute<FText> ToolTip = LOCTEXT("ExportToSkeletalMesh_Tooltip", "Export the geometry collection to a skeletal mesh asset with the same name");
+						const TAttribute<FSlateIcon> Icon = FSlateIcon(FAppStyle::GetAppStyleSetName(), "Icons.SkeletalMesh");
+
+						FToolUIAction UIAction;
+						UIAction.ExecuteAction = FToolMenuExecuteAction::CreateStatic(&ExecuteExportToSkeletalMesh);
+						Section.AddMenuEntry("ControlRigPose_SelectControls", Label, ToolTip, Icon, UIAction);
+					}
+				}));
+		});
+}
+
 #undef LOCTEXT_NAMESPACE
diff --git a/Engine/Plugins/Experimental/GeometryCollectionPlugin/Source/GeometryCollectionNodes/Private/Dataflow/GeometryCollectionNodesPlugin.cpp b/Engine/Plugins/Experimental/GeometryCollectionPlugin/Source/GeometryCollectionNodes/Private/Dataflow/GeometryCollectionNodesPlugin.cpp
--- a/Engine/Plugins/Experimental/GeometryCollectionPlugin/Source/GeometryCollectionNodes/Private/Dataflow/GeometryCollectionNodesPlugin.cpp
+++ b/Engine/Plugins/Experimental/GeometryCollectionPlugin/Source/GeometryCollectionNodes/Private/Dataflow/GeometryCollectionNodesPlugin.cpp
@@ -76,6 +76,7 @@ void IGeometryCollectionNodesPlugin::StartupModule()
 	DATAFLOW_NODE_REGISTER_CREATION_FACTORY(FSetVertexColorFromVertexIndicesDataflowNode);
 	DATAFLOW_NODE_REGISTER_CREATION_FACTORY(FSetVertexColorFromVertexSelectionDataflowNode);
 	DATAFLOW_NODE_REGISTER_CREATION_FACTORY(FSkeletalMeshToCollectionDataflowNode);
+	DATAFLOW_NODE_REGISTER_CREATION_FACTORY(FCollectionToSkeletalMeshDataflowNode);
 	DATAFLOW_NODE_REGISTER_CREATION_FACTORY(FGeometryCollectionTransferVertexAttributeNode);
 	DATAFLOW_NODE_REGISTER_CREATION_FACTORY(FGeometryCollectionTransferVertexSkinWeightsNode);
 	DATAFLOW_NODE_REGISTER_CREATION_FACTORY(FGeometryCollectionSetKinematicVertexSelectionNode);
diff --git a/Engine/Plugins/Experimental/GeometryCollectionPlugin/Source/GeometryCollectionNodes/Private/Dataflow/GeometryCollectionSkeletalMeshToCollectionNode.cpp b/Engine/Plugins/Experimental/GeometryCollectionPlugin/Source/GeometryCollectionNodes/Private/Dataflow/GeometryCollectionSkeletalMeshToCollectionNode.cpp
--- a/Engine/Plugins/Experimental/GeometryCollectionPlugin/Source/GeometryCollectionNodes/Private/Dataflow/GeometryCollectionSkeletalMeshToCollectionNode.cpp
+++ b/Engine/Plugins/Experimental/GeometryCollectionPlugin/Source/GeometryCollectionNodes/Private/Dataflow/GeometryCollectionSkeletalMeshToCollectionNode.cpp
@@ -2,20 +2,18 @@
 
 #include "Dataflow/GeometryCollectionSkeletalMeshToCollectionNode.h"
 
+#include "Animation/Skeleton.h"
+#include "Dataflow/DataflowDebugDrawInterface.h"
+#include "Dataflow/DataflowElement.h"
+#include "Dataflow/DataflowSkeletalMeshNodes.h"
 #include "Engine/SkeletalMesh.h"
-#include "GeometryCollection/GeometryCollectionObject.h"
-#include "GeometryCollection/ManagedArrayCollection.h"
-#include "GeometryCollection/GeometryCollection.h"
+#include "GeometryCollection/Facades/CollectionHierarchyFacade.h"
 #include "GeometryCollection/GeometryCollectionEngineConversion.h"
+#include "Materials/MaterialInterface.h"
 
-#include "DynamicMesh/DynamicMesh3.h"
-#include "MeshDescriptionToDynamicMesh.h"
-
-#include "MeshDescription.h"
-#include "Rendering/SkeletalMeshLODImporterData.h"
-#include "Rendering/SkeletalMeshModel.h"
-#include "IndexTypes.h"
-#include "GeometryCollection/GeometryCollectionAlgo.h"
+#if WITH_EDITOR
+#include "Dataflow/DataflowRenderingViewMode.h"
+#endif
 
 void FSkeletalMeshToCollectionDataflowNode::Evaluate(UE::Dataflow::FContext& Context, const FDataflowOutput* Out) const
 {
@@ -31,3 +29,66 @@ void FSkeletalMeshToCollectionDataflowNode::Evaluate(UE::Dataflow::FContext& Con
 	}
 }
 
+////////////////////////////////////////////////////////////////////////////
+
+FCollectionToSkeletalMeshDataflowNode::FCollectionToSkeletalMeshDataflowNode(const UE::Dataflow::FNodeParameters& InParam, FGuid InGuid)
+	: FDataflowNode(InParam, InGuid)
+{
+	RegisterInputConnection(&Collection);
+	RegisterInputConnection(&Materials);
+
+	RegisterOutputConnection(&SkeletalMesh);
+	RegisterOutputConnection(&Skeleton);
+}
+
+void FCollectionToSkeletalMeshDataflowNode::Evaluate(UE::Dataflow::FContext& Context, const FDataflowOutput* Out) const
+{
+if (Out->IsA(&SkeletalMesh) || Out->IsA(&Skeleton))
+	{
+		USkeletalMesh* OutSkeletalMesh = NewObject<USkeletalMesh>();
+		USkeleton* OutSkeleton = NewObject<USkeleton>();
+		
+		if (OutSkeletalMesh || OutSkeleton)
+		{
+			// make a copy so that we can ensure the level array is created
+			FManagedArrayCollection InCollection = GetValue(Context, &Collection);
+			const TArray<TObjectPtr<UMaterialInterface>>& InMaterials = GetValue(Context, &Materials);
+
+			Chaos::Facades::FCollectionHierarchyFacade HierarchyFacade(InCollection);
+			if (HierarchyFacade.IsValid())
+			{
+				HierarchyFacade.GenerateLevelAttribute();
+			}
+
+			FGeometryCollectionEngineConversion::ConvertCollectionToSkeletalMesh(InCollection, InMaterials, *OutSkeletalMesh, *OutSkeleton);
+		}
+		
+		SetValue(Context, TObjectPtr<const USkeletalMesh>(OutSkeletalMesh), &SkeletalMesh);
+		SetValue(Context, TObjectPtr<const USkeleton>(OutSkeleton), &Skeleton);
+	}
+}
+
+#if WITH_EDITOR
+void FCollectionToSkeletalMeshDataflowNode::DebugDraw(UE::Dataflow::FContext& Context,
+	IDataflowDebugDrawInterface& DataflowRenderingInterface,
+	const FDebugDrawParameters& DebugDrawParameters) const
+{
+	if (const FDataflowOutput* Output = FindOutput(&SkeletalMesh))
+	{
+		if (TObjectPtr<const USkeletalMesh> InSkeletalMesh = Output->ReadValue(Context, SkeletalMesh))
+		{
+			TRefCountPtr<IDataflowDebugDrawObject> SkeletonObject(
+				MakeDebugDrawObject<FDataflowDebugDrawSkeletonObject>(
+					DataflowRenderingInterface.ModifyDataflowElements(), InSkeletalMesh->GetRefSkeleton()
+				));
+
+			DataflowRenderingInterface.DrawObject(SkeletonObject);
+		}
+	}
+}
+
+bool FCollectionToSkeletalMeshDataflowNode::CanDebugDrawViewMode(const FName& ViewModeName) const
+{
+	return ViewModeName == UE::Dataflow::FDataflowConstruction3DViewMode::Name;
+}
+#endif
diff --git a/Engine/Plugins/Experimental/GeometryCollectionPlugin/Source/GeometryCollectionNodes/Public/Dataflow/GeometryCollectionSkeletalMeshToCollectionNode.h b/Engine/Plugins/Experimental/GeometryCollectionPlugin/Source/GeometryCollectionNodes/Public/Dataflow/GeometryCollectionSkeletalMeshToCollectionNode.h
--- a/Engine/Plugins/Experimental/GeometryCollectionPlugin/Source/GeometryCollectionNodes/Public/Dataflow/GeometryCollectionSkeletalMeshToCollectionNode.h
+++ b/Engine/Plugins/Experimental/GeometryCollectionPlugin/Source/GeometryCollectionNodes/Public/Dataflow/GeometryCollectionSkeletalMeshToCollectionNode.h
@@ -15,6 +15,7 @@ namespace UE_DEPRECATED(5.5, "Use UE::Dataflow instead.") Dataflow {}
 #endif
 
 class USkeletalMesh;
+class UMaterialInterface;
 
 
 USTRUCT(meta = (DataflowGeometryCollection))
@@ -46,3 +47,35 @@ public:
 	virtual void Evaluate(UE::Dataflow::FContext& Context, const FDataflowOutput* Out) const override;
 };
 
+
+USTRUCT(meta = (DataflowGeometryCollection))
+struct FCollectionToSkeletalMeshDataflowNode : public FDataflowNode
+{
+	GENERATED_USTRUCT_BODY()
+	DATAFLOW_NODE_DEFINE_INTERNAL(FCollectionToSkeletalMeshDataflowNode, "CollectionToSkeletalMesh", "GeometryCollection", "")
+	DATAFLOW_NODE_RENDER_TYPE("SurfaceRender", FGeometryCollection::StaticType(), "Collection")
+
+private:
+	UPROPERTY(meta = (DataflowInput))
+	FManagedArrayCollection Collection;
+
+	UPROPERTY(meta = (DataflowInput))
+	TArray<TObjectPtr<UMaterialInterface>> Materials;
+
+	UPROPERTY(meta = (DataflowOutput))
+	TObjectPtr<const USkeletalMesh> SkeletalMesh = nullptr;
+
+	UPROPERTY(meta = (DataflowOutput))
+	TObjectPtr<const USkeleton> Skeleton = nullptr;
+
+	virtual void Evaluate(UE::Dataflow::FContext& Context, const FDataflowOutput* Out) const override;
+
+#if WITH_EDITOR
+	virtual void DebugDraw(UE::Dataflow::FContext& Context, IDataflowDebugDrawInterface& DataflowRenderingInterface, const FDebugDrawParameters& DebugDrawParameters) const override;
+	virtual bool CanDebugDrawViewMode(const FName& ViewModeName) const override;
+#endif 
+
+public:
+	FCollectionToSkeletalMeshDataflowNode(const UE::Dataflow::FNodeParameters& InParam, FGuid InGuid = FGuid::NewGuid());
+};
+
diff --git a/Engine/Source/Runtime/Experimental/GeometryCollectionEngine/GeometryCollectionEngine.Build.cs b/Engine/Source/Runtime/Experimental/GeometryCollectionEngine/GeometryCollectionEngine.Build.cs
--- a/Engine/Source/Runtime/Experimental/GeometryCollectionEngine/GeometryCollectionEngine.Build.cs
+++ b/Engine/Source/Runtime/Experimental/GeometryCollectionEngine/GeometryCollectionEngine.Build.cs
@@ -36,6 +36,7 @@ namespace UnrealBuildTool.Rules
 					"MeshConversion",
 					"GeometryCore",
 					"PhysicsCore",
+					"SkeletalMeshUtilitiesCommon",
 				}
 				);
 
diff --git a/Engine/Source/Runtime/Experimental/GeometryCollectionEngine/Private/GeometryCollection/GeometryCollectionEngineConversion.cpp b/Engine/Source/Runtime/Experimental/GeometryCollectionEngine/Private/GeometryCollection/GeometryCollectionEngineConversion.cpp
--- a/Engine/Source/Runtime/Experimental/GeometryCollectionEngine/Private/GeometryCollection/GeometryCollectionEngineConversion.cpp
+++ b/Engine/Source/Runtime/Experimental/GeometryCollectionEngine/Private/GeometryCollection/GeometryCollectionEngineConversion.cpp
@@ -15,6 +15,10 @@
 #include "GeometryCollection/Facades/CollectionInstancedMeshFacade.h"
 #include "GeometryCollection/Facades/CollectionTransformFacade.h"
 #include "GeometryCollection/Facades/CollectionTransformSourceFacade.h"
+#include "GeometryCollection/Facades/CollectionMeshFacade.h"
+#include "GeometryCollection/Facades/CollectionTransformSelectionFacade.h"
+#include "GeometryCollection/Facades/CollectionHierarchyFacade.h"
+#include "GeometryCollection/Facades/CollectionBoundsFacade.h"
 #include "GeometryCollection/GeometryCollection.h"
 #include "GeometryCollection/GeometryCollectionActor.h"
 #include "GeometryCollection/GeometryCollectionAlgo.h"
@@ -33,6 +37,7 @@
 #include "Physics/Experimental/ChaosInterfaceUtils.h"
 #include "PhysicsEngine/BodySetup.h"
 #include "Rendering/SkeletalMeshRenderData.h"
+#include "Rendering/SkeletalMeshModel.h"
 #include "ReferenceSkeleton.h"
 #include "SkeletalMeshAttributes.h"
 #include "StaticMeshAttributes.h"
@@ -41,6 +46,10 @@
 #include "Util/ColorConstants.h"
 #include "GeometryCollection/Facades/CollectionVertexBoneWeightsFacade.h"
 
+#if WITH_EDITOR
+#include "StaticToSkeletalMeshConverter.h"
+#endif
+
 DEFINE_LOG_CATEGORY_STATIC(UGeometryCollectionConversionLogging, Log, All);
 
 #define LOCTEXT_NAMESPACE "GeometryCollectionConversion"
@@ -1782,21 +1791,83 @@ void FGeometryCollectionEngineConversion::ConvertActorToGeometryCollection(const
 }
 
 
-void FGeometryCollectionEngineConversion::ConvertCollectionToSkeleton(const FManagedArrayCollection& InCollection, USkeleton* OutSkeleton, TArray<int32>& OutIndexRemap)
+void FGeometryCollectionEngineConversion::ConvertCollectionToSkeleton(const FManagedArrayCollection& InCollection, USkeleton* OutSkeleton, TArray<int32>& OutIndexRemap, bool bCenterBonesOnGeometry)
 {
+	//clear the existing skeleton data 
+	const TArray<FMeshBoneInfo> BoneInfos = OutSkeleton->GetReferenceSkeleton().GetRefBoneInfo();
+	if (BoneInfos.Num())
+	{
+		int32 RootIndex = 0;
+		while (BoneInfos[RootIndex].ParentIndex != INDEX_NONE)
+		{
+			RootIndex = BoneInfos[RootIndex].ParentIndex;
+		}
+		if (RootIndex != INDEX_NONE)
+		{
+			FReferenceSkeletonModifier Edit(OutSkeleton);
+			Edit.Remove(BoneInfos[RootIndex].Name, /* bRemoveChildren */ true);
+		}
+	}
+
 	FManagedArrayCollection LocalCollection = InCollection;
-	GeometryCollection::Facades::FCollectionTransformFacade Transforms(LocalCollection);
-	if (Transforms.IsValid() && Transforms.HasBoneNameAttribute() && OutSkeleton)
+	GeometryCollection::Facades::FCollectionTransformFacade TransformFacade(LocalCollection);
+	
+	const GeometryCollection::Facades::FBoundsFacade BoundsFacade(InCollection);
+	if (!BoundsFacade.IsValid())
 	{
-		Transforms.EnforceSingleRoot("root");
+		UE_LOG(UGeometryCollectionConversionLogging, Error, TEXT("FGeometryCollectionEngineConversion::ConvertCollectionToSkeleton() - Bounds Facade not valid - check the schema of the collection"));
+		return;
+	}
 
-		OutIndexRemap.Init(INDEX_NONE, Transforms.Num());
-		auto AddMapping = [&OutIndexRemap](int32 A, int32 B)
+	const TManagedArrayAccessor<int32> SimulationTypeAttribute(InCollection, "SimulationType", FTransformCollection::TransformGroup);
+	const TManagedArray<int32>& TransformToGeometryIndexAttribute = BoundsFacade.GetTransformToGeometryIndex();
+
+	auto GetGeometryCenter = [&TransformToGeometryIndexAttribute, &SimulationTypeAttribute, &BoundsFacade, &bCenterBonesOnGeometry](int32 TransformIndex)
 		{
-			OutIndexRemap[A] = B;
+			FVector3f GeometryCenter = FVector3f::ZeroVector;
+			if (bCenterBonesOnGeometry)
+			{
+				const bool bIsLeaf = (SimulationTypeAttribute.IsValid() && SimulationTypeAttribute.IsValidIndex(TransformIndex)) ? SimulationTypeAttribute[TransformIndex] == FGeometryCollection::ESimulationTypes::FST_Rigid : true;
+				if (bIsLeaf && TransformToGeometryIndexAttribute.IsValidIndex(TransformIndex))
+				{
+					const int32 GeometryIndex = TransformToGeometryIndexAttribute[TransformIndex];
+					if (GeometryIndex != INDEX_NONE)
+					{
+						GeometryCenter = FVector3f(BoundsFacade.GetBoundingBoxes()[GeometryIndex].GetCenter());
+					}
+				}
+			}
+			return GeometryCenter;
 		};
 
-		auto AddChildren = [&Transforms, &OutSkeleton, &AddMapping](const TArray<int32>& CollectionChildren)
+	if (TransformFacade.IsValid() && TransformFacade.HasBoneNameAttribute() && OutSkeleton)
+	{
+		TransformFacade.EnforceSingleRoot("root");
+
+		const TManagedArray<FTransform3f>* Transforms = TransformFacade.FindTransforms();
+		const TManagedArray<FString>* BoneNames = TransformFacade.FindBoneNames();
+		ensure(Transforms && BoneNames);
+
+		OutIndexRemap.Init(INDEX_NONE, TransformFacade.Num());
+		
+		auto AddBoneToSkeleton = [&BoneNames, &Transforms, &OutSkeleton, &OutIndexRemap, GetGeometryCenter](int32 TransformIndex, int32 ParentBoneIndex)
+			{
+				const FName BoneName((*BoneNames)[TransformIndex]);
+				const FTransform3f GeomtryCenterOffsetTransform(GetGeometryCenter(TransformIndex));
+				const FTransform Transform(GeomtryCenterOffsetTransform * (*Transforms)[TransformIndex]);
+				FMeshBoneInfo Info(BoneName, BoneName.ToString(), ParentBoneIndex);
+				{
+					FReferenceSkeletonModifier Edit(OutSkeleton);
+					if (Edit.FindBoneIndex(BoneName) == INDEX_NONE) //Bone does not exist
+					{
+						Edit.Add(Info, Transform, true /*bAllowMultipleRoots*/);
+					}
+				}
+				OutIndexRemap[TransformIndex] = OutSkeleton->GetReferenceSkeleton().GetNum();
+			};
+
+
+		auto AddChildren = [&TransformFacade, &OutSkeleton, &AddBoneToSkeleton](const TArray<int32>& CollectionChildren)
 		{
 			TQueue<int32> Children;
 			auto Enqueue = [&Children](const TArray<int32>& List)
@@ -1811,51 +1882,249 @@ void FGeometryCollectionEngineConversion::ConvertCollectionToSkeleton(const FMan
 			{
 				Children.Dequeue(CurrentIndex);
 
-				int32 CollectionParentIndex = (*Transforms.GetParents())[CurrentIndex];
-				int32 SkeletionParentIndex = OutSkeleton->GetReferenceSkeleton().FindBoneIndex(FName((*Transforms.FindBoneNames())[CollectionParentIndex]));
+				int32 CollectionParentIndex = (*TransformFacade.GetParents())[CurrentIndex];
+				int32 SkeletionParentIndex = OutSkeleton->GetReferenceSkeleton().FindBoneIndex(FName((*TransformFacade.FindBoneNames())[CollectionParentIndex]));
 
-				FName BoneName = FName((*Transforms.FindBoneNames())[CurrentIndex]);
-				FTransform Transform = FTransform((*Transforms.FindTransforms())[CurrentIndex]);
-				FMeshBoneInfo Info(BoneName, BoneName.ToString(), SkeletionParentIndex);
-				{
-					FReferenceSkeletonModifier Edit(OutSkeleton);
-					if (Edit.FindBoneIndex(BoneName) == INDEX_NONE) //Bone does not exist
-					{
-						Edit.Add(Info, Transform, true /*bAllowMultipleRoots*/);
-					}
-				}
-				AddMapping(CurrentIndex, OutSkeleton->GetReferenceSkeleton().GetRawBoneNum());
+				AddBoneToSkeleton(CurrentIndex, SkeletionParentIndex);
 
-				Enqueue((*Transforms.FindChildren())[CurrentIndex].Array());
+				Enqueue((*TransformFacade.FindChildren())[CurrentIndex].Array());
 			}
 		};
 
 
 		// must insert in decending order from parent to child. 
-		TArray<FString> BoneNameStrings = Transforms.FindBoneNames()->GetConstArray();
-		for (int i = 0; i < Transforms.GetParents()->Num(); i++)
+		TArray<FString> BoneNameStrings = TransformFacade.FindBoneNames()->GetConstArray();
+		for (int i = 0; i < TransformFacade.GetParents()->Num(); i++)
 		{
-			if ((*Transforms.GetParents())[i] == INDEX_NONE) // No parent
+			if ((*TransformFacade.GetParents())[i] == INDEX_NONE) // No parent
 			{
-				FName BoneName = FName((*Transforms.FindBoneNames())[i]);
-				FTransform Transform = FTransform((*Transforms.FindTransforms())[i]);
-				FMeshBoneInfo Info(BoneName, BoneName.ToString(), INDEX_NONE);
+				AddBoneToSkeleton(i, INDEX_NONE);
+
+				if ((*TransformFacade.FindChildren())[i].Num())
 				{
-					FReferenceSkeletonModifier Edit(OutSkeleton);
-					Edit.Add(Info, Transform, true /*bAllowMultipleRoots*/);
+					AddChildren((*TransformFacade.FindChildren())[i].Array());
 				}
-				AddMapping(i,OutSkeleton->GetReferenceSkeleton().GetNum());
+			}
+		}
+
+		TransformFacade.EnforceSingleRoot("root");
+		UE_LOG(UGeometryCollectionConversionLogging, Log, TEXT("FGeometryCollectionEngineConversion::ConvertCollectionToSkeleton(NumTransforms:%d)"), OutSkeleton->GetReferenceSkeleton().GetRawBoneNum());
+	}
+}
+
+///////////////////////////////////////////////////////////////////////////////////////////////////////////////////////////////////////////////////
+void FGeometryCollectionEngineConversion::ConvertCollectionToSkeletalMesh(const FManagedArrayCollection& InCollection, const TArray<TObjectPtr<UMaterialInterface>> InMaterials, USkeletalMesh& OutSkeletalMesh, USkeleton& OutSkeleton)
+{
+#if WITH_EDITOR
+	// make sure we have all the necessary facade supported 
+	const GeometryCollection::Facades::FCollectionTransformFacade TransformFacade(InCollection);
+	if (!TransformFacade.IsValid())
+	{
+		UE_LOG(UGeometryCollectionConversionLogging, Error, TEXT("FGeometryCollectionEngineConversion::ConvertCollectionToSkeletalMesh() - Transform Facade not valid - check the schema of the collection"));
+		return;
+	}
+
+	const GeometryCollection::Facades::FCollectionMeshFacade MeshFacade(InCollection);
+	if (!MeshFacade.IsValid())
+	{
+		UE_LOG(UGeometryCollectionConversionLogging, Error, TEXT("FGeometryCollectionEngineConversion::ConvertCollectionToSkeletalMesh() - Mesh Facade not valid - check the schema of the collection"));
+		return;
+	}
+
+	const GeometryCollection::Facades::FBoundsFacade BoundsFacade(InCollection);
+	if (!BoundsFacade.IsValid())
+	{
+		UE_LOG(UGeometryCollectionConversionLogging, Error, TEXT("FGeometryCollectionEngineConversion::ConvertCollectionToSkeletalMesh() - Bounds Facade not valid - check the schema of the collection"));
+		return;
+	}
+
+	const TManagedArrayAccessor<int32> SimulationTypeAttribute(InCollection, "SimulationType", FTransformCollection::TransformGroup);
+	const TManagedArray<FString>* Names = TransformFacade.FindBoneNames();
+	const TManagedArray<int32>* Parents = TransformFacade.GetParents();
+	const TManagedArray<FTransform3f>* Transforms = TransformFacade.FindTransforms();
+
+	// Reset the skeletal mesh
+	OutSkeletalMesh.ClearAllMeshDescriptions();
+	OutSkeletalMesh.GetImportedModel()->LODModels.Reset();
+	OutSkeletalMesh.ResetLODInfo();
+	OutSkeletalMesh.ReleaseResources();
+
+	// first extract the skeleton 
+	TArray<int32> TransformToBoneIndexMapping;
+	FGeometryCollectionEngineConversion::ConvertCollectionToSkeleton(InCollection, &OutSkeleton, TransformToBoneIndexMapping, /* bCenterBonesOnGeometry */true);
+
+	// convert the collection to a mesh description including skeletal mesh attributes
+	FMeshDescription MeshDescription;
+
+	FSkeletalMeshAttributes Attributes(MeshDescription);
+	Attributes.Register();
+
+	const int32 NumUVChannels = MeshFacade.GetNumUVLayers();
+
+	MeshDescription.SetNumUVChannels(MeshFacade.GetNumUVLayers());
+	MeshDescription.VertexInstanceAttributes().SetAttributeChannelCount(MeshAttribute::VertexInstance::TextureCoordinate, NumUVChannels);
+
+	TVertexAttributesRef<FVector3f> VertexPositions = Attributes.GetVertexPositions();
+	FSkinWeightsVertexAttributesRef VertexSkinWeights = Attributes.GetVertexSkinWeights();
+
+	TVertexInstanceAttributesRef<FVector3f> VertexInstanceNormals = Attributes.GetVertexInstanceNormals();
+	TVertexInstanceAttributesRef<FVector3f> VertexInstanceTangents = Attributes.GetVertexInstanceTangents();
+	TVertexInstanceAttributesRef<float> VertexInstanceBinormalSigns = Attributes.GetVertexInstanceBinormalSigns();
+	TVertexInstanceAttributesRef<FVector4f> VertexInstanceColors = Attributes.GetVertexInstanceColors();
+	TVertexInstanceAttributesRef<FVector2f> VertexInstanceUVs = Attributes.GetVertexInstanceUVs();
+
+	FSkeletalMeshAttributes::FBoneNameAttributesRef BoneNames = Attributes.GetBoneNames();
+	FSkeletalMeshAttributes::FBoneParentIndexAttributesRef BoneParentIndices = Attributes.GetBoneParentIndices();
+	FSkeletalMeshAttributes::FBonePoseAttributesRef BonePoses = Attributes.GetBonePoses();
 
-				if ((*Transforms.FindChildren())[i].Num())
+	// make sure we create a polygroup per material and that they have properly defined slot names
+	int32 MaxMaterialID = 0;
+	for (int32 FaceIndex = 0; FaceIndex < MeshFacade.MaterialIDAttribute.Num(); ++FaceIndex)
+	{
+		MaxMaterialID = FMath::Max(MaxMaterialID, MeshFacade.MaterialIDAttribute[FaceIndex]);
+	}
+	const int32 MaterialCount = (MaxMaterialID + 1);
+
+	TPolygonGroupAttributesRef<FName> PolygonGroupMaterialSlotNames = Attributes.GetPolygonGroupMaterialSlotNames();
+
+	TArray<FSkeletalMaterial> SkeletalMaterials;
+	TArray<FPolygonGroupID> MatIDToPolyGroupID;
+	MatIDToPolyGroupID.SetNum(MaterialCount);
+	for (int32 MatIndex = 0; MatIndex < MaterialCount; ++MatIndex)
+	{ 
+		const FPolygonGroupID PolygonGroupID = MeshDescription.CreatePolygonGroup();
+		MatIDToPolyGroupID[MatIndex] = PolygonGroupID;
+		const FName SlotName = FName("Material", MatIndex);
+		PolygonGroupMaterialSlotNames.Set(PolygonGroupID, SlotName);
+
+		FSkeletalMaterial NewMaterial;
+		NewMaterial.MaterialInterface = InMaterials.IsValidIndex(MatIndex) ? InMaterials[MatIndex] : nullptr;
+		NewMaterial.MaterialSlotName = SlotName;
+		NewMaterial.ImportedMaterialSlotName = SlotName;
+		SkeletalMaterials.Add(NewMaterial);
+	}
+
+	int32 NewVtxId = 0;
+
+	// convert collection geometry to mesh decsription
+	const int32 NumTransforms = InCollection.NumElements(FTransformCollection::TransformGroup);
+	for (int32 TransformIndex = 0; TransformIndex < NumTransforms; ++TransformIndex)
+	{
+		const int32 GeometryIndex = MeshFacade.TransformToGeometryIndexAttribute[TransformIndex];
+		const bool bIsLeaf = (SimulationTypeAttribute.IsValid() && SimulationTypeAttribute.IsValidIndex(TransformIndex)) ? SimulationTypeAttribute[TransformIndex] == FGeometryCollection::ESimulationTypes::FST_Rigid : true;
+		if (bIsLeaf && GeometryIndex != INDEX_NONE)
+		{
+			// make sure we operate in the right space
+			const FTransform3f Transform = FTransform3f(TransformFacade.ComputeCollectionSpaceTransform(TransformIndex));
+
+			const int32 VertexStart = MeshFacade.VertexStartAttribute[GeometryIndex];
+			const int32 VertexCount = MeshFacade.VertexCountAttribute[GeometryIndex];
+
+			MeshDescription.ReserveNewVertices(VertexCount);
+			MeshDescription.ReserveNewVertexInstances(VertexCount);
+			for (int32 UVChannelIndex = 0; UVChannelIndex < NumUVChannels; ++UVChannelIndex)
+			{
+				MeshDescription.ReserveNewUVs(VertexCount, UVChannelIndex);
+			}
+
+			const int32 VtxIndexToVtxId = (NewVtxId - VertexStart);
+
+			for (int32 VtxIndex = VertexStart; VtxIndex < (VertexStart + VertexCount); ++VtxIndex)
+			{
+				const FVector3f VtxPos = Transform.TransformPosition(MeshFacade.VertexAttribute[VtxIndex]);
+				const FVector3f VtxNormal = Transform.TransformVector(MeshFacade.NormalAttribute[VtxIndex]);
+				const FVector3f VtxTangentU = Transform.TransformVector(MeshFacade.TangentUAttribute[VtxIndex]);
+				const FVector3f VtxTangentV = Transform.TransformVector(MeshFacade.TangentVAttribute[VtxIndex]);
+				const FLinearColor& VtxColor = MeshFacade.ColorAttribute[VtxIndex];
+
+				const FVertexID VertexID(NewVtxId);
+				const FVertexInstanceID VertexInstanceID(NewVtxId);
+
+				MeshDescription.CreateVertexWithID(VertexID);
+				MeshDescription.CreateVertexInstanceWithID(VertexInstanceID, VertexID);
+
+				VertexPositions.Set(VertexID, VtxPos);
+				VertexInstanceNormals.Set(VertexInstanceID, VtxNormal);
+				VertexInstanceTangents.Set(VertexInstanceID, VtxTangentU);
+				const float BinormalSign = GetBasisDeterminantSign(FVector(VtxTangentU), FVector(VtxTangentV), FVector(VtxNormal));
+				VertexInstanceBinormalSigns.Set(VertexInstanceID, BinormalSign);
+				VertexInstanceColors.Set(VertexInstanceID, FVector4f(VtxColor));
+				for (int32 UVChannelIndex = 0; UVChannelIndex < NumUVChannels; ++UVChannelIndex)
 				{
-					AddChildren((*Transforms.FindChildren())[i].Array());
+					const FVector2f& UV = MeshFacade.GetUV(VtxIndex, UVChannelIndex);
+					VertexInstanceUVs.Set(VertexInstanceID, UVChannelIndex, UV);
 				}
+
+				// geometry collection are rigidly skinned to bones
+				FBoneIndexType InfluenceBones[MAX_TOTAL_INFLUENCES];
+				uint16 InfluenceWeights[MAX_TOTAL_INFLUENCES] = { 0 };
+
+				const int32 BoneIndex = (TransformToBoneIndexMapping.IsValidIndex(TransformIndex)) ? TransformToBoneIndexMapping[TransformIndex] - 1 : 0;
+				InfluenceBones[0] = BoneIndex;
+				InfluenceWeights[0] = std::numeric_limits<uint16>::max();
+
+				VertexSkinWeights.Set(VertexID, UE::AnimationCore::FBoneWeights::Create(InfluenceBones, InfluenceWeights));
+
+				++NewVtxId;
+			}
+
+			const int32 FaceStart = MeshFacade.FaceStartAttribute[GeometryIndex];
+			const int32 FaceCount = MeshFacade.FaceCountAttribute[GeometryIndex];
+
+			MeshDescription.ReserveNewTriangles(FaceCount);
+
+			for (int32 TriIndex = FaceStart; TriIndex < FaceStart + FaceCount; ++TriIndex)
+			{
+				const FIntVector& VtxIndices = MeshFacade.IndicesAttribute[TriIndex];
+				const int32 Vtx0 = VtxIndices.X + VtxIndexToVtxId;
+				const int32 Vtx1 = VtxIndices.Y + VtxIndexToVtxId;
+				const int32 Vtx2 = VtxIndices.Z + VtxIndexToVtxId;
+
+				const FVertexInstanceID TriVertInstanceIDs[] = { FVertexInstanceID(Vtx0), FVertexInstanceID(Vtx1), FVertexInstanceID(Vtx2) };
+				const int32 MaterialID = MeshFacade.MaterialIDAttribute[TriIndex];
+				MeshDescription.CreateTriangle(MatIDToPolyGroupID[MaterialID], MakeConstArrayView(TriVertInstanceIDs, 3));
 			}
 		}
+	
+		// find corresponding bone from skeleton 
+		const FName BoneName = FName((*Names)[TransformIndex]);
+		const int32 BoneIndex = OutSkeleton.GetReferenceSkeleton().FindBoneIndex(BoneName);
+		const int32 ParentBoneIndex = (BoneIndex != INDEX_NONE)? OutSkeleton.GetReferenceSkeleton().GetParentIndex(BoneIndex): INDEX_NONE;
 
-		Transforms.EnforceSingleRoot("root");
-		UE_LOG(UGeometryCollectionConversionLogging, Log, TEXT("FGeometryCollectionEngineConversion::ConvertCollectionToSkeleton(NumTransforms:%d)"), OutSkeleton->GetReferenceSkeleton().GetRawBoneNum());
+		const TArray<FTransform> BoneTransforms = OutSkeleton.GetReferenceSkeleton().GetRefBonePose();
+		const FTransform BoneTransform = (BoneTransforms.IsValidIndex(BoneIndex)) ? BoneTransforms[BoneIndex] : FTransform::Identity;
+
+		const FBoneID BoneID(BoneIndex);
+
+		Attributes.CreateBone(BoneIndex);
+		{
+			BoneNames.Set(BoneID, BoneName);
+			BoneParentIndices.Set(BoneID, ParentBoneIndex);
+			BonePoses.Set(BoneID, BoneTransform);
+		}
 	}
+
+	TArray<const FMeshDescription*> MeshDescriptions = { &MeshDescription };
+
+	// Update the skeletal mesh from the description 
+	constexpr bool bRecomputeNormals = false;
+	constexpr bool bRecomputeTangents = false;
+	FStaticToSkeletalMeshConverter::InitializeSkeletalMeshFromMeshDescriptions(
+		&OutSkeletalMesh,
+		MeshDescriptions,
+		SkeletalMaterials,
+		OutSkeleton.GetReferenceSkeleton(),
+		bRecomputeNormals,
+		bRecomputeTangents);
+
+	// finalize skeletal mesh binding to skeleton 
+	OutSkeletalMesh.SetSkeleton(&OutSkeleton);
+	OutSkeleton.SetPreviewMesh(&OutSkeletalMesh);
+	OutSkeleton.RecreateBoneTree(&OutSkeletalMesh);
+
+#else
+	UE_LOG(UGeometryCollectionConversionLogging, Error, TEXT("FGeometryCollectionEngineConversion::ConvertCollectionToSkeletalMesh() - conversion only supported in editor."));
+#endif // WITH_EDITOR
 }
 
 void FGeometryCollectionEngineConversion::GetMaterialsFromInstances(const TArray<TObjectPtr<UMaterialInterface>>& MaterialInstances, TArray<TObjectPtr<UMaterial>>& OutMaterials)
diff --git a/Engine/Source/Runtime/Experimental/GeometryCollectionEngine/Public/GeometryCollection/GeometryCollectionEngineConversion.h b/Engine/Source/Runtime/Experimental/GeometryCollectionEngine/Public/GeometryCollection/GeometryCollectionEngineConversion.h
--- a/Engine/Source/Runtime/Experimental/GeometryCollectionEngine/Public/GeometryCollection/GeometryCollectionEngineConversion.h
+++ b/Engine/Source/Runtime/Experimental/GeometryCollectionEngine/Public/GeometryCollection/GeometryCollectionEngineConversion.h
@@ -221,10 +221,19 @@ public:
 	*  @param OutSkeleton : USkeleton to populate
 	*  @param OutIndexRemap : Remapping from Collection transforms to skeletal transforms.
 	*/
-	static GEOMETRYCOLLECTIONENGINE_API void ConvertCollectionToSkeleton(const FManagedArrayCollection& InCollection, USkeleton* OutSkeleton, TArray<int32>& OutIndexRemap);
+	static GEOMETRYCOLLECTIONENGINE_API void ConvertCollectionToSkeleton(const FManagedArrayCollection& InCollection, USkeleton* OutSkeleton, TArray<int32>& OutIndexRemap, bool bCenterBonesOnGeometry  = false);
+
+	/**
+	*  Converts a Geometry collection to a SkeletalMesh
+	*  @param InCollection: Geometry collection to convert to a skeletal mesh
+	*  @param OutSkeletalMesh : skeletal mesh asset to export to
+	*  @param OutSkeleton : Skeleton asset to export the bone hierarchy to
+	*/
+	static GEOMETRYCOLLECTIONENGINE_API void ConvertCollectionToSkeletalMesh(const FManagedArrayCollection& InCollection, const TArray<TObjectPtr<UMaterialInterface>> InMaterials, USkeletalMesh& OutSkeletalMesh, USkeleton& OutSkeleton);
 
 	/** utility function to convert an array of material instances to an array of material */
 	static GEOMETRYCOLLECTIONENGINE_API void GetMaterialsFromInstances(const TArray<TObjectPtr<UMaterialInterface>>& MaterialInstances, TArray<TObjectPtr<UMaterial>>& OutMaterials);
+
 private:
 	/**
 	*  Appends a GeometryCollectionSource to a GeometryCollection ignoring the materials
