Fix compilation issues with CL 43675006
- Add the right plugin dependency in the geometry collection one
- make sure the module is added as editor only

diff --git a/Engine/Plugins/Experimental/GeometryCollectionPlugin/GeometryCollectionPlugin.uplugin b/Engine/Plugins/Experimental/GeometryCollectionPlugin/GeometryCollectionPlugin.uplugin
--- a/Engine/Plugins/Experimental/GeometryCollectionPlugin/GeometryCollectionPlugin.uplugin
+++ b/Engine/Plugins/Experimental/GeometryCollectionPlugin/GeometryCollectionPlugin.uplugin
@@ -69,6 +69,10 @@
 		{
 			"Name": "Fracture",
 			"Enabled": true
+		},
+		{
+			"Name": "EditorScriptingUtilities",
+			"Enabled": true
 		}
 	]
 }
diff --git a/Engine/Source/Runtime/Experimental/GeometryCollectionEngine/GeometryCollectionEngine.Build.cs b/Engine/Source/Runtime/Experimental/GeometryCollectionEngine/GeometryCollectionEngine.Build.cs
--- a/Engine/Source/Runtime/Experimental/GeometryCollectionEngine/GeometryCollectionEngine.Build.cs
+++ b/Engine/Source/Runtime/Experimental/GeometryCollectionEngine/GeometryCollectionEngine.Build.cs
@@ -1,5 +1,7 @@
 // Copyright Epic Games, Inc. All Rights Reserved.
 
+using System.Collections.Generic;
+
 namespace UnrealBuildTool.Rules
 {
 	public class GeometryCollectionEngine : ModuleRules
@@ -36,7 +38,6 @@ namespace UnrealBuildTool.Rules
 					"MeshConversion",
 					"GeometryCore",
 					"PhysicsCore",
-					"SkeletalMeshUtilitiesCommon",
 				}
 				);
 
@@ -49,6 +50,8 @@ namespace UnrealBuildTool.Rules
 
 				PublicDependencyModuleNames.Add("EditorFramework");
                 PublicDependencyModuleNames.Add("UnrealEd");
+
+				PrivateDependencyModuleNames.Add("SkeletalMeshUtilitiesCommon");
 			}
 
 			PrivateDefinitions.Add("CHAOS_INCLUDE_LEVEL_1=1");
