Add a safer 'initialize from collection' method to Dataflow Selection types, to automatically pick the correct group size

diff --git a/Engine/Source/Runtime/Experimental/Dataflow/Core/Private/Dataflow/DataflowSelection.cpp b/Engine/Source/Runtime/Experimental/Dataflow/Core/Private/Dataflow/DataflowSelection.cpp
--- a/Engine/Source/Runtime/Experimental/Dataflow/Core/Private/Dataflow/DataflowSelection.cpp
+++ b/Engine/Source/Runtime/Experimental/Dataflow/Core/Private/Dataflow/DataflowSelection.cpp
@@ -16,6 +16,11 @@ void FDataflowSelection::Initialize(const FDataflowSelection& Other)
 	SelectionArray = Other.SelectionArray;
 }
 
+void FDataflowSelection::InitializeFromCollection(const FManagedArrayCollection& InCollection, bool Value)
+{
+	SelectionArray.Init(Value, InCollection.NumElements(GroupName));
+}
+
 void FDataflowSelection::Clear()
 {
 	SelectionArray.Init(false, SelectionArray.Num());
@@ -197,7 +202,7 @@ bool FDataflowSelection::IsValidForCollection(const FManagedArrayCollection& InC
 
 bool FDataflowSelection::InitFromArray(const FManagedArrayCollection& InCollection, const TArray<int32>& InSelectionArr)
 {
-	Initialize(InCollection.NumElements(GroupName), false);
+	InitializeFromCollection(InCollection, false);
 	return SetFromArray(InSelectionArr);
 }
 
diff --git a/Engine/Source/Runtime/Experimental/Dataflow/Core/Public/Dataflow/DataflowSelection.h b/Engine/Source/Runtime/Experimental/Dataflow/Core/Public/Dataflow/DataflowSelection.h
--- a/Engine/Source/Runtime/Experimental/Dataflow/Core/Public/Dataflow/DataflowSelection.h
+++ b/Engine/Source/Runtime/Experimental/Dataflow/Core/Public/Dataflow/DataflowSelection.h
@@ -22,10 +22,15 @@ struct FDataflowSelection
 
 	DATAFLOWCORE_API void Initialize(int32 NumBits, bool Value);
 	DATAFLOWCORE_API void Initialize(const FDataflowSelection& Other);
+	DATAFLOWCORE_API void InitializeFromCollection(const FManagedArrayCollection& InCollection, bool Value);
 	DATAFLOWCORE_API void Clear();
 	int32 Num() const { return SelectionArray.Num(); }
 	DATAFLOWCORE_API int32 NumSelected() const;
 	DATAFLOWCORE_API bool AnySelected() const;
+	inline bool IsValidIndex(int32 Idx) const
+	{
+		return SelectionArray.IsValidIndex(Idx);
+	}
 	bool IsSelected(int32 Idx) const { return SelectionArray[Idx]; }
 	void SetSelected(int32 Idx) { SelectionArray[Idx] = true; }
 	DATAFLOWCORE_API void SetSelected(TArray<int32> Indices);
