Dataflow : Better support for arrays and general anytype inputs and outputs
- Added a new FContextValue type to interact with cached value of unknown types
- Added support for converting a general value to an array

diff --git a/Engine/Plugins/Experimental/GeometryCollectionPlugin/Source/GeometryCollectionNodes/Private/Dataflow/GeometryCollectionArrayNodes.cpp b/Engine/Plugins/Experimental/GeometryCollectionPlugin/Source/GeometryCollectionNodes/Private/Dataflow/GeometryCollectionArrayNodes.cpp
--- a/Engine/Plugins/Experimental/GeometryCollectionPlugin/Source/GeometryCollectionNodes/Private/Dataflow/GeometryCollectionArrayNodes.cpp
+++ b/Engine/Plugins/Experimental/GeometryCollectionPlugin/Source/GeometryCollectionNodes/Private/Dataflow/GeometryCollectionArrayNodes.cpp
@@ -2,12 +2,27 @@
 
 #include "Dataflow/GeometryCollectionArrayNodes.h"
 #include "Dataflow/DataflowCore.h"
-
+#include "Dataflow/DataflowAnyTypeRegistry.h"
+#include "Dataflow/DataflowTypePolicy.h"
 
 //#include UE_INLINE_GENERATED_CPP_BY_NAME(GeometryCollectionArrayNodes)
 
 namespace UE::Dataflow
 {
+	namespace Private
+	{
+		static bool AutoConvertArrayToSingleFilter(FName FromType, FName ToType)
+		{
+			const FName ElementType = FDataflowArrayTypePolicy::GetElementType(FromType);
+			return (ToType == ElementType);
+		}
+
+		static bool AutoConvertSingleToArrayFilter(FName FromType, FName ToType)
+		{
+			const FName ElementType = FDataflowArrayTypePolicy::GetElementType(ToType);
+			return (FromType == ElementType);
+		}
+	}
 
 	void GeometryCollectionArrayNodes()
 	{
@@ -22,6 +37,7 @@ namespace UE::Dataflow
 		DATAFLOW_NODE_REGISTER_CREATION_FACTORY(FRandomizeFloatArrayDataflowNode);
 		DATAFLOW_NODE_REGISTER_CREATION_FACTORY(FDataflowGetArraySizeNode);
 		DATAFLOW_NODE_REGISTER_CREATION_FACTORY(FDataflowGetArrayElementNode);
+		DATAFLOW_NODE_REGISTER_CREATION_FACTORY(FDataflowConvertToArrayNode);
 		DATAFLOW_NODE_REGISTER_CREATION_FACTORY(FDataflowMakeManagedArrayCollectionArrayNode);
 		
 
@@ -29,6 +45,9 @@ namespace UE::Dataflow
 		DATAFLOW_NODE_REGISTER_CREATION_FACTORY(FGetFloatArrayElementDataflowNode);
 		DATAFLOW_NODE_REGISTER_CREATION_FACTORY(FGetArrayElementDataflowNode);
 		DATAFLOW_NODE_REGISTER_CREATION_FACTORY(FGetNumArrayElementsDataflowNode);
+
+		UE_DATAFLOW_REGISTER_AUTOCONVERT_WITH_FILTER(FDataflowArrayTypes, FDataflowAnyType, FDataflowGetArrayElementNode, &Private::AutoConvertArrayToSingleFilter);
+		UE_DATAFLOW_REGISTER_AUTOCONVERT_WITH_FILTER(FDataflowAnyType, FDataflowArrayTypes, FDataflowConvertToArrayNode, &Private::AutoConvertSingleToArrayFilter);
 	}
 }
 
@@ -630,7 +649,8 @@ void FDataflowGetArraySizeNode::Evaluate(UE::Dataflow::FContext& Context, const
 {
 	if (Out->IsA(&Size))
 	{
-		SetArraySizeFromInput(Context, &Array, &Size);
+		const UE::Dataflow::FContextValue InArray = GetValue(Context, &Array);
+		SetValue(Context, InArray.Num(), &Size);
 	}
 }
 
@@ -672,8 +692,47 @@ void FDataflowGetArrayElementNode::Evaluate(UE::Dataflow::FContext& Context, con
 	}
 	else if (Out->IsA(&Element))
 	{
+		const UE::Dataflow::FContextValue InArray = GetValue(Context, &Array);
 		const int32 InIndex = GetValue(Context, &Index);
-		SetArrayElementFromInput(Context, &Array, InIndex, &Element);
+		UE::Dataflow::FContextValue OutElement = InArray.GetAt(InIndex);
+		SetValue(Context, MoveTemp(OutElement), &Element);
+	}
+}
+
+///////////////////////////////////////////////////////////////////////////////////////////////
+FDataflowConvertToArrayNode::FDataflowConvertToArrayNode(const UE::Dataflow::FNodeParameters& InParam, FGuid InGuid)
+	: FDataflowNode(InParam, InGuid)
+{
+	static FName ArrayElementDependencyGroup("ArrayElement");
+	
+	RegisterInputConnection(&Element)
+		.SetTypeDependencyGroup(ArrayElementDependencyGroup);
+
+	RegisterOutputConnection(&Array)
+		.SetTypeDependencyGroup(ArrayElementDependencyGroup);
+}
+
+FName FDataflowConvertToArrayNode::GetDependentConnectionType(const FDataflowConnection& SourceConnection, const FDataflowConnection& DependentConnection) const
+{
+	if (SourceConnection.IsA(&Array) && DependentConnection.IsA(&Element))
+	{
+		return FDataflowArrayTypePolicy::GetElementType(SourceConnection.GetType());
+	}
+	else if (SourceConnection.IsA(&Element) && DependentConnection.IsA(&Array))
+	{
+		return FDataflowArrayTypePolicy::GetArrayType(SourceConnection.GetType());
+	}
+	// default for all the rest 
+	return FDataflowNode::GetDependentConnectionType(SourceConnection, DependentConnection);
+}
+
+void FDataflowConvertToArrayNode::Evaluate(UE::Dataflow::FContext& Context, const FDataflowOutput* Out) const
+{
+	if (Out->IsA(&Array))
+	{
+		const UE::Dataflow::FContextValue InElement = GetValue(Context, &Element);
+		UE::Dataflow::FContextValue OutArray = InElement.ToArray();
+		SetValue(Context, MoveTemp(OutArray), &Array);
 	}
 }
 
diff --git a/Engine/Plugins/Experimental/GeometryCollectionPlugin/Source/GeometryCollectionNodes/Public/Dataflow/GeometryCollectionArrayNodes.h b/Engine/Plugins/Experimental/GeometryCollectionPlugin/Source/GeometryCollectionNodes/Public/Dataflow/GeometryCollectionArrayNodes.h
--- a/Engine/Plugins/Experimental/GeometryCollectionPlugin/Source/GeometryCollectionNodes/Public/Dataflow/GeometryCollectionArrayNodes.h
+++ b/Engine/Plugins/Experimental/GeometryCollectionPlugin/Source/GeometryCollectionNodes/Public/Dataflow/GeometryCollectionArrayNodes.h
@@ -543,7 +543,7 @@ USTRUCT()
 struct FDataflowGetArraySizeNode : public FDataflowNode
 {
 	GENERATED_USTRUCT_BODY()
-	DATAFLOW_NODE_DEFINE_INTERNAL(FDataflowGetArraySizeNode, "GetArraySize", "Utilities|Array", "")
+	DATAFLOW_NODE_DEFINE_INTERNAL(FDataflowGetArraySizeNode, "GetArraySize", "Utilities|Array", "Count Length Num")
 
 private:
 	/** Array to get size from */
@@ -567,14 +567,14 @@ USTRUCT()
 struct FDataflowGetArrayElementNode : public FDataflowNode
 {
 	GENERATED_USTRUCT_BODY()
-	DATAFLOW_NODE_DEFINE_INTERNAL(FDataflowGetArrayElementNode, "GetArrayElement", "Utilities|Array", "")
+	DATAFLOW_NODE_DEFINE_INTERNAL(FDataflowGetArrayElementNode, "GetArrayElement", "Utilities|Array", "Index At")
 
 private:
-	/** Array to get size from */
+	/** Array to get the element from */
 	UPROPERTY(meta = (DataflowInput, DataflowOutput, DataflowPassthrough = "Array", DataflowIntrinsic))
 	FDataflowArrayTypes Array;
 
-	/** Array to get size from */
+	/** Index to get the element at */
 	UPROPERTY(EditAnywhere, Category = "Array", meta = (DataflowInput))
 	int32 Index = 0;
 
@@ -591,6 +591,31 @@ public:
 };
 
 
+/**
+ * convert a single element to an array 
+ */
+USTRUCT()
+struct FDataflowConvertToArrayNode : public FDataflowNode
+{
+	GENERATED_USTRUCT_BODY()
+	DATAFLOW_NODE_DEFINE_INTERNAL(FDataflowConvertToArrayNode, "ConvertToArray", "Utilities|Array", "Element")
+
+public:
+	FDataflowConvertToArrayNode(const UE::Dataflow::FNodeParameters& InParam, FGuid InGuid = FGuid::NewGuid());
+
+private:
+	/** Array to Convert to */
+	UPROPERTY(meta = (DataflowOutput))
+	FDataflowArrayTypes Array;
+
+	/** Element to convert to an array */
+	UPROPERTY(meta = (DataflowInput))
+	FDataflowAllTypes Element;
+
+	virtual void Evaluate(UE::Dataflow::FContext& Context, const FDataflowOutput* Out) const override;
+
+	virtual FName GetDependentConnectionType(const FDataflowConnection& SourceConnection, const FDataflowConnection& DependentConnection) const override;
+};
 
 /**
  * Append an element to an array of ManagedArrayCollections.
diff --git a/Engine/Source/Runtime/Experimental/Dataflow/Core/Private/Dataflow/DataflowNode.cpp b/Engine/Source/Runtime/Experimental/Dataflow/Core/Private/Dataflow/DataflowNode.cpp
--- a/Engine/Source/Runtime/Experimental/Dataflow/Core/Private/Dataflow/DataflowNode.cpp
+++ b/Engine/Source/Runtime/Experimental/Dataflow/Core/Private/Dataflow/DataflowNode.cpp
@@ -1443,6 +1443,7 @@ void FDataflowNode::ForwardInput(UE::Dataflow::FContext& Context, const UE::Data
 	}
 }
 
+// DEPRECATED 5.7
 void FDataflowNode::SetArraySizeFromInput(UE::Dataflow::FContext& Context, const UE::Dataflow::FConnectionReference& InputReference, const int32* OutputReference) const
 {
 	int32 OutSize = 0;
@@ -1453,12 +1454,15 @@ void FDataflowNode::SetArraySizeFromInput(UE::Dataflow::FContext& Context, const
 			// Pull the value to populate the cache
 			ArrayInput->PullValue(Context);
 
+			PRAGMA_DISABLE_DEPRECATION_WARNINGS
 			OutSize = Context.GetArraySizeFromData(ArrayInput->GetConnection()->CacheKey());
+			PRAGMA_ENABLE_DEPRECATION_WARNINGS
 		}
 	}
 	SetValue(Context, OutSize, OutputReference);
 }
 
+// DEPRECATED 5.7
 void FDataflowNode::SetArrayElementFromInput(UE::Dataflow::FContext& Context, const UE::Dataflow::FConnectionReference& InputReference, int32 Index, const void* OutputReference) const
 {
 	if (const FDataflowOutput* Output = FindOutput(OutputReference))
@@ -1473,7 +1477,9 @@ void FDataflowNode::SetArrayElementFromInput(UE::Dataflow::FContext& Context, co
 				const UE::Dataflow::FContextCacheKey ArrayKey = ArrayInput->GetConnection()->CacheKey();
 				const UE::Dataflow::FContextCacheKey ElementKey = Output->CacheKey();
 
+				PRAGMA_DISABLE_DEPRECATION_WARNINGS
 				Context.SetArrayElementFromData(ArrayKey, Index, ElementKey, Output->GetProperty(), Output->GetOwningNodeGuid(), Output->GetOwningNodeValueHash(), Output->GetOwningNodeTimestamp());
+				PRAGMA_ENABLE_DEPRECATION_WARNINGS
 				return;
 			}
 		}
diff --git a/Engine/Source/Runtime/Experimental/Dataflow/Core/Private/Dataflow/DataflowNodeParameters.cpp b/Engine/Source/Runtime/Experimental/Dataflow/Core/Private/Dataflow/DataflowNodeParameters.cpp
--- a/Engine/Source/Runtime/Experimental/Dataflow/Core/Private/Dataflow/DataflowNodeParameters.cpp
+++ b/Engine/Source/Runtime/Experimental/Dataflow/Core/Private/Dataflow/DataflowNodeParameters.cpp
@@ -393,6 +393,7 @@ namespace UE::Dataflow
 #endif // EDITOR_ONLY
 	}
 
+	// DEPRECATED 5.7
 	int32 FContext::GetArraySizeFromData(const FContextCacheKey InKey) const
 	{
 		// Query the size without knowning the type ( the cache elements do know it )
@@ -404,6 +405,7 @@ namespace UE::Dataflow
 		return 0;
 	}
 
+	// DEPRECATED 5.7
 	void FContext::SetArrayElementFromData(const FContextCacheKey InArrayKey, int32 Index, const FContextCacheKey InElementKey, const FProperty* InProperty, const FGuid& InNodeGuid, uint32 InNodeHash, const FTimestamp& InTimestamp)
 	{
 		// get the element from the cache 
diff --git a/Engine/Source/Runtime/Experimental/Dataflow/Core/Public/Dataflow/DataflowContextCache.h b/Engine/Source/Runtime/Experimental/Dataflow/Core/Public/Dataflow/DataflowContextCache.h
--- a/Engine/Source/Runtime/Experimental/Dataflow/Core/Public/Dataflow/DataflowContextCache.h
+++ b/Engine/Source/Runtime/Experimental/Dataflow/Core/Public/Dataflow/DataflowContextCache.h
@@ -8,6 +8,7 @@
 #include "Templates/Decay.h"
 #include "Templates/UniquePtr.h"
 #include "UObject/UnrealType.h"
+#include "Dataflow/DataflowTypePolicy.h"
 
 namespace UE::Dataflow
 {
@@ -15,6 +16,9 @@ namespace UE::Dataflow
 	struct FContextCacheElementBase;
 	struct FContextCacheElementNull;
 
+	template<class T>
+	struct TContextCacheElementUObjectArray;
+
 	typedef uint32 FContextCacheKey;
 
 	/** Trait used to select the UObject* or TObjectPtr cache element path code. */
@@ -89,8 +93,10 @@ namespace UE::Dataflow
 
 		virtual const void* GetUntypedData(const IContextCacheStore& Context, const FProperty* PropertyIn) const = 0;
 
+		virtual bool IsArray(const IContextCacheStore& Context) const = 0;
 		virtual int32 GetNumArrayElements(const IContextCacheStore& Context) const = 0;
 		virtual TUniquePtr<FContextCacheElementBase> CreateFromArrayElement(const IContextCacheStore& Context, int32 Index, const FProperty* InProperty, const FGuid& InNodeGuid, uint32 InNodeHash, const FTimestamp& InTimestamp) const = 0;
+		virtual TUniquePtr<FContextCacheElementBase> CreateArrayFromElement(const IContextCacheStore& Context, const FProperty* InProperty, const FGuid& InNodeGuid, uint32 InNodeHash, const FTimestamp& InTimestamp) const = 0;
 
 		EType GetType() const {	return Type; }
 
@@ -138,7 +144,7 @@ namespace UE::Dataflow
 			: FContextCacheElementBase(EType::CacheElementTyped, Other.GetNodeGuid(), Other.GetProperty(), Other.GetNodeHash(), Other.GetTimestamp())
 			, Data(Other.Data)
 		{}
-		
+
 		inline const T& GetData(const IContextCacheStore& Context, const FProperty* PropertyIn, const T& Default) const;
 
 		virtual const void* GetUntypedData(const IContextCacheStore& Context, const FProperty* PropertyIn) const
@@ -146,6 +152,11 @@ namespace UE::Dataflow
 			return &Data;
 		}
 
+		virtual bool IsArray(const IContextCacheStore& Context) const override
+		{
+			return (TIsTArray<FDataType>::Value);
+		}
+
 		virtual int32 GetNumArrayElements(const IContextCacheStore& Context) const override
 		{
 			if constexpr (TIsTArray<FDataType>::Value)
@@ -174,6 +185,19 @@ namespace UE::Dataflow
 			}
 		}
 
+		virtual TUniquePtr<FContextCacheElementBase> CreateArrayFromElement(const IContextCacheStore& Context, const FProperty* InProperty, const FGuid& InNodeGuid, uint32 InNodeHash, const FTimestamp& InTimestamp) const override
+		{
+			if constexpr (TIsTArray<FDataType>::Value)
+			{
+				return {}; // already an Array
+			}
+			else 
+			{
+				TArray<FDataType> Array{ Data };
+				return MakeUnique<TContextCacheElement<TArray<FDataType>>>(InNodeGuid, InProperty, MoveTemp(Array), InNodeHash, InTimestamp);
+			}
+		}
+
 		const T& GetDataDirect() const { return Data; }
 
 		inline virtual TUniquePtr<FContextCacheElementBase> Clone(const IContextCacheStore& Context) const override;
@@ -193,15 +217,20 @@ namespace UE::Dataflow
 			, DataKey(InDataKey)
 		{}
 
+
 		template<class T>
 		inline const T& GetData(const IContextCacheStore& Context, const FProperty* PropertyIn, const T& Default) const;
 
 		inline virtual const void* GetUntypedData(const IContextCacheStore& Context, const FProperty* PropertyIn) const override;
 
+		inline virtual bool IsArray(const IContextCacheStore& Context) const override;
+
 		inline virtual int32 GetNumArrayElements(const IContextCacheStore& Context) const override;
 
 		inline virtual TUniquePtr<FContextCacheElementBase> CreateFromArrayElement(const IContextCacheStore& Context, int32 Index, const FProperty* InProperty, const FGuid& InNodeGuid, uint32 InNodeHash, const FTimestamp& InTimestamp) const override;
 
+		inline virtual TUniquePtr<FContextCacheElementBase> CreateArrayFromElement(const IContextCacheStore& Context, const FProperty* InProperty, const FGuid& InNodeGuid, uint32 InNodeHash, const FTimestamp& InTimestamp) const override;
+
 		inline virtual TUniquePtr<FContextCacheElementBase> Clone(const IContextCacheStore& Context) const override;
 
 	private:
@@ -231,6 +260,11 @@ namespace UE::Dataflow
 			return nullptr;
 		}
 
+		virtual bool IsArray(const IContextCacheStore& Context) const override
+		{
+			return false; // we have no way to tell at this point , null entry means that we don't have a value and we let the requester of teh value to get a default value 
+		}
+
 		virtual int32 GetNumArrayElements(const IContextCacheStore& Context) const override
 		{
 			return 0;
@@ -241,6 +275,11 @@ namespace UE::Dataflow
 			return {};
 		}
 
+		virtual TUniquePtr<FContextCacheElementBase> CreateArrayFromElement(const IContextCacheStore& Context, const FProperty* InProperty, const FGuid& InNodeGuid, uint32 InNodeHash, const FTimestamp& InTimestamp) const override
+		{
+			return {};
+		}
+
 		inline virtual TUniquePtr<FContextCacheElementBase> Clone(const IContextCacheStore& Context) const override;
 	};
 
@@ -270,6 +309,11 @@ namespace UE::Dataflow
 			return (void*)ObjectPtr;
 		}
 
+		virtual bool IsArray(const IContextCacheStore& Context) const override
+		{
+			return false;
+		}
+
 		virtual int32 GetNumArrayElements(const IContextCacheStore& Context) const override
 		{
 			return 0;
@@ -280,6 +324,12 @@ namespace UE::Dataflow
 			return {};
 		}
 
+		virtual TUniquePtr<FContextCacheElementBase> CreateArrayFromElement(const IContextCacheStore& Context, const FProperty* InProperty, const FGuid& InNodeGuid, uint32 InNodeHash, const FTimestamp& InTimestamp) const override
+		{
+			TArray<FDataType> Array{ ObjectPtr };
+			return MakeUnique<TContextCacheElementUObjectArray<FDataType>>(InNodeGuid, InProperty, MoveTemp(Array), InNodeHash, InTimestamp);
+		}
+
 		inline virtual TUniquePtr<FContextCacheElementBase> Clone(const IContextCacheStore& Context) const override;
 
 		//~ Begin FGCObject interface
@@ -310,6 +360,11 @@ namespace UE::Dataflow
 			, ObjectPtrArray(Other.ObjectPtrArray)
 		{}
 
+		virtual bool IsArray(const IContextCacheStore& Context) const override
+		{
+			return true;
+		}
+
 		const TArray<T>& GetData(const IContextCacheStore& /*Context*/, const FProperty* /*PropertyIn*/, const TArray<T>& /*Default*/) const
 		{
 			return ObjectPtrArray;
@@ -335,6 +390,11 @@ namespace UE::Dataflow
 			return {};
 		}
 
+		virtual TUniquePtr<FContextCacheElementBase> CreateArrayFromElement(const IContextCacheStore& Context, const FProperty* InProperty, const FGuid& InNodeGuid, uint32 InNodeHash, const FTimestamp& InTimestamp) const override
+		{
+			return {}; // no array of array 
+		}
+
 		inline virtual TUniquePtr<FContextCacheElementBase> Clone(const IContextCacheStore& Context) const override;
 
 		//~ Begin FGCObject interface
@@ -380,6 +440,11 @@ namespace UE::Dataflow
 			return InstancedStruct.GetMemory();
 		}
 
+		virtual bool IsArray(const IContextCacheStore& Context) const override
+		{
+			return false;
+		}
+
 		virtual int32 GetNumArrayElements(const IContextCacheStore& Context) const override
 		{
 			return 0;
@@ -390,6 +455,11 @@ namespace UE::Dataflow
 			return {};
 		}
 
+		virtual TUniquePtr<FContextCacheElementBase> CreateArrayFromElement(const IContextCacheStore& Context, const FProperty* InProperty, const FGuid& InNodeGuid, uint32 InNodeHash, const FTimestamp& InTimestamp) const override
+		{
+			return {}; // can't make arrays of arrays
+		}
+
 		inline virtual TUniquePtr<FContextCacheElementBase> Clone(const IContextCacheStore& Context) const override
 		{
 			return MakeUnique<FContextCacheElementUStruct>(*this);
@@ -434,6 +504,11 @@ namespace UE::Dataflow
 			return InstancedStructArray.GetMemory();
 		}
 
+		virtual bool IsArray(const IContextCacheStore& Context) const override
+		{
+			return true;
+		}
+
 		virtual int32 GetNumArrayElements(const IContextCacheStore& Context) const override
 		{
 			return InstancedStructArray.Num();
@@ -450,6 +525,11 @@ namespace UE::Dataflow
 			return {};
 		}
 
+		virtual TUniquePtr<FContextCacheElementBase> CreateArrayFromElement(const IContextCacheStore& Context, const FProperty* InProperty, const FGuid& InNodeGuid, uint32 InNodeHash, const FTimestamp& InTimestamp) const override
+		{
+			return {}; // can't make arrays of arrays
+		}
+
 		inline virtual TUniquePtr<FContextCacheElementBase> Clone(const IContextCacheStore& Context) const override
 		{
 			return MakeUnique<FContextCacheElementUStructArray>(*this);
@@ -618,6 +698,15 @@ namespace UE::Dataflow
 		return nullptr;
 	}
 
+	bool FContextCacheElementReference::IsArray(const IContextCacheStore& Context) const
+	{
+		if (const TUniquePtr<FContextCacheElementBase>* Cache = Context.FindCacheElement(DataKey))
+		{
+			return (*Cache)->IsArray(Context);
+		}
+		return false;
+	}
+
 	int32 FContextCacheElementReference::GetNumArrayElements(const IContextCacheStore& Context) const
 	{
 		if (const TUniquePtr<FContextCacheElementBase>* Cache = Context.FindCacheElement(DataKey))
@@ -633,7 +722,16 @@ namespace UE::Dataflow
 		{
 			return (*Cache)->CreateFromArrayElement(Context, Index, InProperty, InNodeGuid, InNodeHash, InTimestamp);
 		}
-		return 0;
+		return {};
+	}
+
+	TUniquePtr<FContextCacheElementBase> FContextCacheElementReference::CreateArrayFromElement(const IContextCacheStore& Context, const FProperty* InProperty, const FGuid& InNodeGuid, uint32 InNodeHash, const FTimestamp& InTimestamp) const
+	{
+		if (const TUniquePtr<FContextCacheElementBase>* Cache = Context.FindCacheElement(DataKey))
+		{
+			return (*Cache)->CreateArrayFromElement(Context, InProperty, InNodeGuid, InNodeHash, InTimestamp);
+		}
+		return {};
 	}
 
 	TUniquePtr<FContextCacheElementBase> FContextCacheElementReference::Clone(const IContextCacheStore& Context) const
@@ -664,6 +762,82 @@ namespace UE::Dataflow
 	{
 		return MakeUnique<TContextCacheElementUObjectArray<T>>(*this);
 	}
+
+	////////////////////////////////////////////////////
+
+	struct FContextValue
+	{
+	public:
+		FContextValue(IContextCacheStore& InContext, FContextCacheKey InCacheKey)
+			: Context(InContext)
+			, CacheKey(InCacheKey)
+		{}
+
+		FContextValue(IContextCacheStore& InContext, TUniquePtr<FContextCacheElementBase>&& InCacheElement)
+			: Context(InContext)
+			, CacheKey(0)
+			, CacheElement(MoveTemp(InCacheElement))
+		{}
+
+		int32 Num() const
+		{
+			if (const FContextCacheElementBase* CacheEntry = GetCacheEntry())
+			{
+				return CacheEntry->GetNumArrayElements(Context);
+			}
+			return 0;
+		}
+
+		FContextValue GetAt(int32 Index) const
+		{
+			if (const FContextCacheElementBase* CacheEntry = GetCacheEntry())
+			{
+				TUniquePtr<FContextCacheElementBase> NewCacheElement = CacheEntry->CreateFromArrayElement(Context, Index, CacheEntry->GetProperty(), CacheEntry->GetNodeGuid(), CacheEntry->GetNodeHash(), CacheEntry->GetTimestamp());
+				return FContextValue(Context, MoveTemp(NewCacheElement));
+			}
+			return FContextValue(Context, 0);
+		}
+
+		FContextValue ToArray() const
+		{
+			if (const FContextCacheElementBase* CacheEntry = GetCacheEntry())
+			{
+				// arrays return themselves
+				if (CacheEntry->IsArray(Context))
+				{
+					if (CacheElement)
+					{
+						return FContextValue(Context, CacheEntry->Clone(Context));
+					}
+					// simply return the reference to the original cached array 
+					return FContextValue(Context, CacheKey);
+				}
+
+				// actually convert to an array 
+				TUniquePtr<FContextCacheElementBase> NewCacheElement = CacheEntry->CreateArrayFromElement(Context, CacheEntry->GetProperty(), CacheEntry->GetNodeGuid(), CacheEntry->GetNodeHash(), CacheEntry->GetTimestamp());
+				return FContextValue(Context, MoveTemp(NewCacheElement));
+			}
+			return FContextValue(Context, 0);
+		}
+
+		const FContextCacheElementBase* GetCacheEntry() const
+		{
+			if (CacheElement)
+			{
+				return CacheElement.Get();
+			}
+			if (const TUniquePtr<FContextCacheElementBase>* CacheEntry = Context.FindCacheElement(CacheKey))
+			{
+				return CacheEntry->Get();
+			}
+			return nullptr;
+		}
+
+	private:
+		IContextCacheStore& Context;
+		FContextCacheKey CacheKey; // only set when CacheElement is null
+		TUniquePtr<FContextCacheElementBase> CacheElement; // when CacheKey is define this should be null
+	};
 };
 
 inline FArchive& operator<<(FArchive& Ar, UE::Dataflow::FTimestamp& ValueIn)
@@ -680,3 +854,5 @@ inline FArchive& operator<<(FArchive& Ar, UE::Dataflow::FContextCache& ValueIn)
 }
 
 
+
+
diff --git a/Engine/Source/Runtime/Experimental/Dataflow/Core/Public/Dataflow/DataflowInputOutput.h b/Engine/Source/Runtime/Experimental/Dataflow/Core/Public/Dataflow/DataflowInputOutput.h
--- a/Engine/Source/Runtime/Experimental/Dataflow/Core/Public/Dataflow/DataflowInputOutput.h
+++ b/Engine/Source/Runtime/Experimental/Dataflow/Core/Public/Dataflow/DataflowInputOutput.h
@@ -82,9 +82,12 @@ public:
 	template<class T>
 	const T& GetValue(UE::Dataflow::FContext& Context, const T& Default) const;
 
-	template<typename TAnyType>
+	template<typename TAnyType UE_REQUIRES(!std::is_same_v<typename TAnyType::FStorageType, void>)>
 	typename TAnyType::FStorageType GetValueFromAnyType(UE::Dataflow::FContext& Context, const typename TAnyType::FStorageType& Default) const;
 
+	template<typename TAnyType UE_REQUIRES(std::is_same_v<typename TAnyType::FStorageType, void>)>
+	typename UE::Dataflow::FContextValue GetValueFromAnyType(UE::Dataflow::FContext& Context) const;
+
 	/**
 	* pull the value from the upstream connections
 	* the upstream graph is evaluated if necessary and values are cached along the way 
@@ -420,7 +423,7 @@ const T& FDataflowInput::GetValue(UE::Dataflow::FContext& Context, const T& Defa
 	return Default;
 }
 
-template<typename TAnyType>
+template<typename TAnyType UE_REQUIRES(!std::is_same_v<typename TAnyType::FStorageType, void>)>
 typename TAnyType::FStorageType FDataflowInput::GetValueFromAnyType(UE::Dataflow::FContext& Context, const typename TAnyType::FStorageType& Default) const
 {
 	typename TAnyType::FStorageType ReturnValue = Default;
@@ -451,6 +454,25 @@ typename TAnyType::FStorageType FDataflowInput::GetValueFromAnyType(UE::Dataflow
 	return ReturnValue;
 }
 
+template<typename TAnyType UE_REQUIRES(std::is_same_v<typename TAnyType::FStorageType, void>)>
+typename UE::Dataflow::FContextValue FDataflowInput::GetValueFromAnyType(UE::Dataflow::FContext& Context) const
+{
+	UE::Dataflow::FContextCacheKey CacheKey = 0;
+	if (const FDataflowOutput* ConnectionOut = GetConnection())
+	{
+		if (ConnectionOut->HasFrozenValue())
+		{
+			//return ConnectionOut->GetFrozenValue(Default);
+			ensure(false); // unsupported because the frozen values are actually not in the cache, they should probably be 
+		}
+		else if (ConnectionOut->Evaluate(Context))
+		{
+			CacheKey = ConnectionOut->CacheKey();
+		}
+	}
+	return UE::Dataflow::FContextValue(Context, CacheKey);
+}
+
 template<class T>
 TFuture<const T&> FDataflowInput::GetValueParallel(UE::Dataflow::FContext& Context, const T& Default) const
 {
diff --git a/Engine/Source/Runtime/Experimental/Dataflow/Core/Public/Dataflow/DataflowNode.h b/Engine/Source/Runtime/Experimental/Dataflow/Core/Public/Dataflow/DataflowNode.h
--- a/Engine/Source/Runtime/Experimental/Dataflow/Core/Public/Dataflow/DataflowNode.h
+++ b/Engine/Source/Runtime/Experimental/Dataflow/Core/Public/Dataflow/DataflowNode.h
@@ -465,31 +465,44 @@ PRAGMA_ENABLE_DEPRECATION_WARNINGS
 	*   @param Reference : Pointer to a member of this node that corresponds with the output to set.
 	*						*Reference will be used as the default if the input is not connected. 
 	*/
-	template<class T, typename = std::enable_if_t<!std::is_base_of_v<FDataflowAnyType, T>>>
+	template<class T UE_REQUIRES(!std::is_base_of_v<FDataflowAnyType, T>)>
 	const T& GetValue(UE::Dataflow::FContext& Context, const UE::Dataflow::TConnectionReference<T>& Reference) const
 	{
 		checkSlow(FindInput(Reference));
 		return FindInput(Reference)->template GetValue<T>(Context, *static_cast<const T*>(Reference.Reference));
 	}
 
-	template<typename TAnyType, typename = std::enable_if_t<std::is_base_of_v<FDataflowAnyType, TAnyType>>>
+	template<typename TAnyType UE_REQUIRES(std::is_base_of_v<FDataflowAnyType, TAnyType> && !std::is_same_v<typename TAnyType::FStorageType, void>)>
 	typename TAnyType::FStorageType GetValue(UE::Dataflow::FContext& Context, const UE::Dataflow::TConnectionReference<TAnyType>& Reference) const
 	{
 		checkSlow(Reference.Reference && FindInput(Reference));
 		return FindInput(Reference)->template GetValueFromAnyType<TAnyType>(Context, static_cast<const TAnyType*>(Reference.Reference)->Value);
 	}
 
-	template<class T, typename = std::enable_if_t<!std::is_base_of_v<FDataflowAnyType, T>>>
+	template<typename TAnyType UE_REQUIRES(std::is_base_of_v<FDataflowAnyType, TAnyType> && std::is_same_v<typename TAnyType::FStorageType, void>)>
+	typename UE::Dataflow::FContextValue GetValue(UE::Dataflow::FContext& Context, const UE::Dataflow::TConnectionReference<TAnyType>& Reference) const
+	{
+		checkSlow(Reference.Reference && FindInput(Reference));
+		return FindInput(Reference)->template GetValueFromAnyType<TAnyType>(Context);
+	}
+
+	template<class T UE_REQUIRES(!std::is_base_of_v<FDataflowAnyType, T>)>
 	const T& GetValue(UE::Dataflow::FContext& Context, const T* Reference) const
 	{
 		return GetValue(Context, UE::Dataflow::TConnectionReference<T>(Reference));
 	}
-
-	template<typename TAnyType, typename = std::enable_if_t<std::is_base_of_v<FDataflowAnyType, TAnyType>>>
+	
+	template<typename TAnyType UE_REQUIRES(std::is_base_of_v<FDataflowAnyType, TAnyType> && !std::is_same_v<typename TAnyType::FStorageType, void>)>
 	typename TAnyType::FStorageType GetValue(UE::Dataflow::FContext& Context, const TAnyType* Reference) const
 	{
 		return GetValue(Context, UE::Dataflow::TConnectionReference<TAnyType>(Reference));
 	}
+
+	template<typename TAnyType UE_REQUIRES(std::is_base_of_v<FDataflowAnyType, TAnyType> && std::is_same_v<typename TAnyType::FStorageType, void>)>
+	typename UE::Dataflow::FContextValue GetValue(UE::Dataflow::FContext& Context, const TAnyType* Reference) const
+	{
+		return GetValue(Context, UE::Dataflow::TConnectionReference<TAnyType>(Reference));
+	}
 	
 	template<class T>
 	TFuture<const T&> GetValueParallel(UE::Dataflow::FContext& Context, const UE::Dataflow::TConnectionReference<T>& Reference) const
@@ -549,7 +562,7 @@ PRAGMA_ENABLE_DEPRECATION_WARNINGS
 	*   @param Value : The value to store in the contexts data store. 
 	*   @param Reference : Pointer to a member of this node that corresponds with the output to set. 
 	*/
-	template<class T, class U, typename = std::enable_if_t<!std::is_base_of_v<FDataflowAnyType, U>>>
+	template<class T, class U UE_REQUIRES(!std::is_base_of_v<FDataflowAnyType, U>)>
 	void SetValue(UE::Dataflow::FContext& Context, T&& Value, const U* Reference) const
 	{
 		using FInValueType = typename TDecay<decltype(Value)>::Type;
@@ -561,11 +574,11 @@ PRAGMA_ENABLE_DEPRECATION_WARNINGS
 		}
 		else
 		{
-			checkfSlow(false, TEXT("This output could not be found within this node, check this has been properly registered in the node constructor"));
+			ensureMsgf(false, TEXT("This output could not be found within this node, check this has been properly registered in the node constructor"));
 		}
 	}
 
-	template<typename TAnyType, typename = std::enable_if_t<std::is_base_of_v<FDataflowAnyType, TAnyType>>>
+	template<typename TAnyType UE_REQUIRES(std::is_base_of_v<FDataflowAnyType, TAnyType> && !std::is_same_v<typename TAnyType::FStorageType, void>)>
 	void SetValue(UE::Dataflow::FContext& Context, const typename TAnyType::FStorageType& Value, const TAnyType* Reference) const
 	{
 		if (const FDataflowOutput* Output = FindOutput(Reference))
@@ -574,7 +587,20 @@ PRAGMA_ENABLE_DEPRECATION_WARNINGS
 		}
 		else
 		{
-			checkfSlow(false, TEXT("This output could not be found within this node, check this has been properly registered in the node constructor"));
+			ensureMsgf(false, TEXT("This output could not be found within this node, check this has been properly registered in the node constructor"));
+		}
+	}
+
+	template<typename TAnyType UE_REQUIRES(std::is_base_of_v<FDataflowAnyType, TAnyType> && std::is_same_v<typename TAnyType::FStorageType, void>)>
+	void SetValue(UE::Dataflow::FContext& Context, UE::Dataflow::FContextValue&& Value, const TAnyType* Reference) const
+	{
+		if (const FDataflowOutput* Output = FindOutput(Reference))
+		{
+			Output->SetValue(Forward<UE::Dataflow::FContextValue>(Value), Context);
+		}
+		else
+		{
+			ensureMsgf(false, TEXT("This output could not be found within this node, check this has been properly registered in the node constructor"));
 		}
 	}
 
@@ -631,6 +657,7 @@ PRAGMA_ENABLE_DEPRECATION_WARNINGS
 	* When called, the input will be pulled as a result.
 	* If the input is not an array, the value set on the output will be 0
 	*/
+	UE_DEPRECATED(5.7, "Use UE::Dataflow::FContextValue from calling GetValue on the node to access the same functionality")
 	DATAFLOWCORE_API void SetArraySizeFromInput(UE::Dataflow::FContext& Context, const UE::Dataflow::FConnectionReference& InputReference, const int32* OutputReference)  const;
 
 	/**
@@ -638,6 +665,7 @@ PRAGMA_ENABLE_DEPRECATION_WARNINGS
 	* When called, the input will be pulled as a result.
 	* If the input is not an array or if the index is out of bounds, a null value will be set on the output (which will result in a default value when read from) 
 	*/
+	UE_DEPRECATED(5.7, "Use UE::Dataflow::FContextValue from calling GetValue on the node to access the same functionality")
 	DATAFLOWCORE_API void SetArrayElementFromInput(UE::Dataflow::FContext& Context, const UE::Dataflow::FConnectionReference& InputReference, int32 Index, const void* OutputReference) const;
 
 	/*
diff --git a/Engine/Source/Runtime/Experimental/Dataflow/Core/Public/Dataflow/DataflowNodeParameters.h b/Engine/Source/Runtime/Experimental/Dataflow/Core/Public/Dataflow/DataflowNodeParameters.h
--- a/Engine/Source/Runtime/Experimental/Dataflow/Core/Public/Dataflow/DataflowNodeParameters.h
+++ b/Engine/Source/Runtime/Experimental/Dataflow/Core/Public/Dataflow/DataflowNodeParameters.h
@@ -196,37 +196,59 @@ namespace UE::Dataflow
 		void SetData(FContextCacheKey InKey, const FProperty* InProperty, T&& InValue, const FGuid& InNodeGuid, uint32 InNodeHash, const FTimestamp& InTimestamp)
 		{
 			TUniquePtr<FContextCacheElementBase> DataStoreEntry;
-			bool bMadeEntry = false;
-			if constexpr (TIsTArray<typename TDecay<T>::Type>::Value)
+
+			if constexpr (std::is_same_v<T, FContextValue>)
 			{
-				typedef typename TDecay<T>::Type ArrayType;
-				if constexpr (TIsUObjectPtrElement<typename ArrayType::ElementType>::Value)
-				{
-					bMadeEntry = true;
-					DataStoreEntry = MakeUnique<TContextCacheElementUObjectArray<typename ArrayType::ElementType>>(InNodeGuid, InProperty, Forward<T>(InValue), InNodeHash, InTimestamp);
-				}
-				else if constexpr (TIsReflectedStruct<ArrayType>::Value)
+				if (const FContextCacheElementBase* CacheEntryToClone = InValue.GetCacheEntry())
 				{
-					DataStoreEntry = MakeUnique<FContextCacheElementUStructArray>(InNodeGuid, InProperty, Forward<T>(InValue), InNodeHash, InTimestamp);
+					DataStoreEntry = CacheEntryToClone->Clone(*this);
+					if (DataStoreEntry)
+					{
+						DataStoreEntry->UpdatePropertyAndNodeData(InProperty, InNodeGuid, InNodeHash, InTimestamp);
+					}
 				}
 			}
-			if (!bMadeEntry)
+			else
 			{
-				if constexpr (TIsUObjectPtrElement<T>::Value)
-				{
-					DataStoreEntry = MakeUnique<TContextCacheElementUObject<T>>(InNodeGuid, InProperty, Forward<T>(InValue), InNodeHash, InTimestamp);
-				}
-				else if constexpr (TIsReflectedStruct<T>::Value)
+				bool bMadeEntry = false;
+				if constexpr (TIsTArray<typename TDecay<T>::Type>::Value)
 				{
-					DataStoreEntry = MakeUnique<FContextCacheElementUStruct>(InNodeGuid, InProperty, Forward<T>(InValue), InNodeHash, InTimestamp);
+					typedef typename TDecay<T>::Type ArrayType;
+					if constexpr (TIsUObjectPtrElement<typename ArrayType::ElementType>::Value)
+					{
+						bMadeEntry = true;
+						DataStoreEntry = MakeUnique<TContextCacheElementUObjectArray<typename ArrayType::ElementType>>(InNodeGuid, InProperty, Forward<T>(InValue), InNodeHash, InTimestamp);
+					}
+					else if constexpr (TIsReflectedStruct<ArrayType>::Value)
+					{
+						DataStoreEntry = MakeUnique<FContextCacheElementUStructArray>(InNodeGuid, InProperty, Forward<T>(InValue), InNodeHash, InTimestamp);
+					}
 				}
-				else
+				if (!bMadeEntry)
 				{
-					DataStoreEntry = MakeUnique<TContextCacheElement<T>>(InNodeGuid, InProperty, Forward<T>(InValue), InNodeHash, InTimestamp);
+					if constexpr (TIsUObjectPtrElement<T>::Value)
+					{
+						DataStoreEntry = MakeUnique<TContextCacheElementUObject<T>>(InNodeGuid, InProperty, Forward<T>(InValue), InNodeHash, InTimestamp);
+					}
+					else if constexpr (TIsReflectedStruct<T>::Value)
+					{
+						DataStoreEntry = MakeUnique<FContextCacheElementUStruct>(InNodeGuid, InProperty, Forward<T>(InValue), InNodeHash, InTimestamp);
+					}
+					else
+					{
+						DataStoreEntry = MakeUnique<TContextCacheElement<T>>(InNodeGuid, InProperty, Forward<T>(InValue), InNodeHash, InTimestamp);
+					}
 				}
 			}
 
-			SetDataImpl(InKey, MoveTemp(DataStoreEntry));
+			if (DataStoreEntry)
+			{
+				SetDataImpl(InKey, MoveTemp(DataStoreEntry));
+			}
+			else
+			{
+				SetNullData(InKey, InProperty, InNodeGuid, InNodeHash, InTimestamp);
+			}
 		}
 
 		void SetDataFromStructView(FContextCacheKey InKey, const FProperty* InProperty, const FConstStructView& StructView, const FGuid& InNodeGuid, uint32 InNodeHash, const FTimestamp& InTimestamp)
@@ -266,8 +288,10 @@ namespace UE::Dataflow
 			SetDataImpl(InKey, MoveTemp(CacheNullElement));
 		}
 
+		UE_DEPRECATED(5.7, "Use UE::Dataflow::FContextValue from calling GetValue on the node to access the same functionality")
 		DATAFLOWCORE_API int32 GetArraySizeFromData(const FContextCacheKey InKey) const;
 
+		UE_DEPRECATED(5.7, "Use UE::Dataflow::FContextValue from calling GetValue on the node to access the same functionality")
 		DATAFLOWCORE_API void SetArrayElementFromData(const FContextCacheKey InArrayKey, int32 Index, const FContextCacheKey InElementKey, const FProperty* InProperty, const FGuid& InNodeGuid, uint32 InNodeHash, const FTimestamp& InTimestamp);
 
 		DATAFLOWCORE_API bool CopyDataToAnotherContext(const FContextCacheKey InSourceKey, FContext& TargetContext, const FContextCacheKey InTargetKey, const FProperty* InProperty, const FGuid& InNodeGuid, uint32 InNodeHash, const FTimestamp& InTimestamp) const;
