Dataflow : add support for generic array function for structures and array fo structures
- Moved a lot of the context and context cache methods implementation in a cpp file
- Added structure to array of struct and array of  struct to struct  support

diff --git a/Engine/Source/Runtime/Experimental/Dataflow/Core/Private/Dataflow/DataflowContextCache.cpp b/Engine/Source/Runtime/Experimental/Dataflow/Core/Private/Dataflow/DataflowContextCache.cpp
new file mode 100644
--- /dev/null
+++ b/Engine/Source/Runtime/Experimental/Dataflow/Core/Private/Dataflow/DataflowContextCache.cpp
@@ -0,0 +1,328 @@
+// Copyright Epic Games, Inc. All Rights Reserved.
+
+#include "Dataflow/DataflowContextCache.h"
+
+namespace UE::Dataflow
+{
+	//////////////////////////////////////////////////////////////////////////////////////////////////////
+
+	const void* FContextCacheElementReference::GetUntypedData(const IContextCacheStore& Context, const FProperty* InProperty) const
+	{
+		if (const TUniquePtr<FContextCacheElementBase>* Cache = Context.FindCacheElement(DataKey))
+		{
+			return (*Cache)->GetUntypedData(Context, InProperty);
+		}
+		return nullptr;
+	}
+
+	bool FContextCacheElementReference::IsArray(const IContextCacheStore& Context) const
+	{
+		if (const TUniquePtr<FContextCacheElementBase>* Cache = Context.FindCacheElement(DataKey))
+		{
+			return (*Cache)->IsArray(Context);
+		}
+		return false;
+	}
+
+	int32 FContextCacheElementReference::GetNumArrayElements(const IContextCacheStore& Context) const
+	{
+		if (const TUniquePtr<FContextCacheElementBase>* Cache = Context.FindCacheElement(DataKey))
+		{
+			return (*Cache)->GetNumArrayElements(Context);
+		}
+		return 0;
+	}
+
+	TUniquePtr<FContextCacheElementBase> FContextCacheElementReference::CreateFromArrayElement(const IContextCacheStore& Context, int32 Index, const FProperty* InProperty, const FGuid& InNodeGuid, uint32 InNodeHash, const FTimestamp& InTimestamp) const
+	{
+		if (const TUniquePtr<FContextCacheElementBase>* Cache = Context.FindCacheElement(DataKey))
+		{
+			return (*Cache)->CreateFromArrayElement(Context, Index, InProperty, InNodeGuid, InNodeHash, InTimestamp);
+		}
+		return {};
+	}
+
+	TUniquePtr<FContextCacheElementBase> FContextCacheElementReference::CreateArrayFromElement(const IContextCacheStore& Context, const FProperty* InProperty, const FGuid& InNodeGuid, uint32 InNodeHash, const FTimestamp& InTimestamp) const
+	{
+		if (const TUniquePtr<FContextCacheElementBase>* Cache = Context.FindCacheElement(DataKey))
+		{
+			return (*Cache)->CreateArrayFromElement(Context, InProperty, InNodeGuid, InNodeHash, InTimestamp);
+		}
+		return {};
+	}
+
+	TUniquePtr<FContextCacheElementBase> FContextCacheElementReference::Clone(const IContextCacheStore& Context) const
+	{
+		if (const TUniquePtr<FContextCacheElementBase>* ReferencedCacheEntry = Context.FindCacheElement(DataKey))
+		{
+			if (ReferencedCacheEntry && ReferencedCacheEntry->IsValid())
+			{
+				return (*ReferencedCacheEntry)->Clone(Context);
+			}
+		}
+		return MakeUnique<FContextCacheElementNull>(GetNodeGuid(), GetProperty(), GetNodeHash(), GetTimestamp());
+	}
+
+	//////////////////////////////////////////////////////////////////////////////////////////////////////
+
+	const void* FContextCacheElementNull::GetUntypedData(const IContextCacheStore& Context, const FProperty* PropertyIn) const
+	{
+		return nullptr;
+	}
+
+	bool FContextCacheElementNull::IsArray(const IContextCacheStore& Context) const
+	{
+		return false; // we have no way to tell at this point , null entry means that we don't have a value and we let the requester of teh value to get a default value 
+	}
+
+	int32 FContextCacheElementNull::GetNumArrayElements(const IContextCacheStore& Context) const
+	{
+		return 0;
+	}
+
+	TUniquePtr<FContextCacheElementBase> FContextCacheElementNull::CreateFromArrayElement(const IContextCacheStore& Context, int32 Index, const FProperty* InProperty, const FGuid& InNodeGuid, uint32 InNodeHash, const FTimestamp& InTimestamp) const
+	{
+		return {};
+	}
+
+	TUniquePtr<FContextCacheElementBase> FContextCacheElementNull::CreateArrayFromElement(const IContextCacheStore& Context, const FProperty* InProperty, const FGuid& InNodeGuid, uint32 InNodeHash, const FTimestamp& InTimestamp) const
+	{
+		return {};
+	}
+
+	TUniquePtr<FContextCacheElementBase> FContextCacheElementNull::Clone(const IContextCacheStore& Context) const
+	{
+		return MakeUnique<FContextCacheElementNull>(GetNodeGuid(), GetProperty(), GetNodeHash(), GetTimestamp());
+	}
+
+	//////////////////////////////////////////////////////////////////////////////////////////////////////
+
+	FContextCacheElementUStruct::FContextCacheElementUStruct(FGuid InNodeGuid, const FProperty* InProperty, const FConstStructView& StructView, uint32 InNodeHash, FTimestamp Timestamp)
+		: FContextCacheElementBase(EType::CacheElementUStruct, InNodeGuid, InProperty, InNodeHash, Timestamp)
+		, InstancedStruct(StructView)
+	{}
+
+	FContextCacheElementUStruct::FContextCacheElementUStruct(const FContextCacheElementUStruct& Other)
+		: FContextCacheElementBase(EType::CacheElementUStruct, Other.GetNodeGuid(), Other.GetProperty(), Other.GetNodeHash(), Other.GetTimestamp())
+		, InstancedStruct(Other.InstancedStruct)
+	{}
+
+	const void* FContextCacheElementUStruct::GetUntypedData(const IContextCacheStore& /*Context*/, const FProperty* /*PropertyIn*/) const
+	{
+		return InstancedStruct.GetMemory();
+	}
+
+	bool FContextCacheElementUStruct::IsArray(const IContextCacheStore& Context) const
+	{
+		return false;
+	}
+
+	int32 FContextCacheElementUStruct::GetNumArrayElements(const IContextCacheStore& Context) const
+	{
+		return 0;
+	}
+
+	TUniquePtr<FContextCacheElementBase> FContextCacheElementUStruct::CreateFromArrayElement(const IContextCacheStore& Context, int32 Index, const FProperty* InProperty, const FGuid& InNodeGuid, uint32 InNodeHash, const FTimestamp& InTimestamp) const
+	{
+		return {};
+	}
+
+	TUniquePtr<FContextCacheElementBase> FContextCacheElementUStruct::CreateArrayFromElement(const IContextCacheStore& Context, const FProperty* InProperty, const FGuid& InNodeGuid, uint32 InNodeHash, const FTimestamp& InTimestamp) const
+	{
+		return MakeUnique<FContextCacheElementUStructArray>(InNodeGuid, InProperty, FConstStructView(InstancedStruct), InNodeHash, InTimestamp);
+	}
+
+	TUniquePtr<FContextCacheElementBase> FContextCacheElementUStruct::Clone(const IContextCacheStore& Context) const
+	{
+		return MakeUnique<FContextCacheElementUStruct>(*this);
+	}
+
+	//////////////////////////////////////////////////////////////////////////////////////////////////////
+
+	const void* FContextCacheElementUStructArray::GetUntypedData(const IContextCacheStore& /*Context*/, const FProperty* /*PropertyIn*/) const
+	{
+		return InstancedStructArray.GetMemory();
+	}
+
+	bool FContextCacheElementUStructArray::IsArray(const IContextCacheStore& Context) const
+	{
+		return true;
+	}
+
+	int32 FContextCacheElementUStructArray::GetNumArrayElements(const IContextCacheStore& Context) const
+	{
+		return InstancedStructArray.Num();
+	}
+
+	TUniquePtr<FContextCacheElementBase> FContextCacheElementUStructArray::CreateFromArrayElement(const IContextCacheStore& Context, int32 Index, const FProperty* InProperty, const FGuid& InNodeGuid, uint32 InNodeHash, const FTimestamp& InTimestamp) const
+	{
+		if (InstancedStructArray.IsValidIndex(Index))
+		{
+			const FConstStructView ElementStructView = InstancedStructArray.GetScriptViewAt(Index);
+			return MakeUnique<FContextCacheElementUStruct>(GetNodeGuid(), GetProperty(), ElementStructView, GetNodeHash(), GetTimestamp());
+		}
+		return {};
+	}
+
+	TUniquePtr<FContextCacheElementBase> FContextCacheElementUStructArray::CreateArrayFromElement(const IContextCacheStore& Context, const FProperty* InProperty, const FGuid& InNodeGuid, uint32 InNodeHash, const FTimestamp& InTimestamp) const
+	{
+		return {}; // can't make arrays of arrays
+	}
+
+	TUniquePtr<FContextCacheElementBase> FContextCacheElementUStructArray::Clone(const IContextCacheStore& Context) const
+	{
+		return MakeUnique<FContextCacheElementUStructArray>(*this);
+	}
+
+	FConstStructArrayView FContextCacheElementUStructArray::GetStructArrayView() const
+	{
+		return InstancedStructArray.GetScriptStruct() ?
+			FConstStructArrayView(*InstancedStructArray.GetScriptStruct(), InstancedStructArray.GetData(), InstancedStructArray.Num()) :
+			FConstStructArrayView();
+	}
+
+	FContextCacheElementUStructArray::FInstancedStructArray::FInstancedStructArray(const UScriptStruct* const InScriptStruct)
+		: ScriptStruct(InScriptStruct)
+	{}
+
+	FContextCacheElementUStructArray::FInstancedStructArray::FInstancedStructArray(const FConstStructView& StructView)
+		: ScriptStruct(StructView.GetScriptStruct())
+	{
+		InitFromRawData(StructView.GetMemory(), 1);
+	}
+
+	FContextCacheElementUStructArray::FInstancedStructArray::FInstancedStructArray(const FConstStructArrayView& StructArrayView)
+		: ScriptStruct(StructArrayView.GetScriptStruct())
+	{
+		InitFromRawData(StructArrayView.GetData(), StructArrayView.Num());
+	}
+	
+	FContextCacheElementUStructArray::FInstancedStructArray::~FInstancedStructArray()
+	{
+		if (ScriptStruct)
+		{
+			ScriptStruct->DestroyStruct(GetData(), Num());
+			GetAllocatorInstance().ResizeAllocation(ArrayMax, 0, GetStructureSize(), ScriptStruct->GetMinAlignment());
+		}
+	}
+
+	const int32 FContextCacheElementUStructArray::FInstancedStructArray::GetStructureSize() const
+	{
+		return FMath::Max(1, ScriptStruct? ScriptStruct->GetStructureSize(): 0);
+	}
+
+	const UScriptStruct* FContextCacheElementUStructArray::FInstancedStructArray::GetScriptStruct() const
+	{
+		return ScriptStruct;
+	}
+
+	FConstStructView FContextCacheElementUStructArray::FInstancedStructArray::GetScriptViewAt(int32 Index) const
+	{
+		if (IsValidIndex(Index))
+		{
+			return FConstStructView(ScriptStruct, GetData() + (GetStructureSize() * Index));
+		}
+		return FConstStructView();
+	}
+
+	const void* FContextCacheElementUStructArray::FInstancedStructArray::GetMemory() const
+	{
+		return GetData();
+	}
+
+	void FContextCacheElementUStructArray::FInstancedStructArray::InitFromRawData(const void* Data, const int32 Num)
+	{
+		if (ScriptStruct)  // Null if the array has 0 element
+		{
+			SetNumUnsafeInternal(Num);
+			ArrayMax = Num;
+			GetAllocatorInstance().ResizeAllocation(0, Num, GetStructureSize(), ScriptStruct->GetMinAlignment());
+			ScriptStruct->InitializeStruct(GetData(), Num);
+			ScriptStruct->CopyScriptStruct(GetData(), Data, Num);
+		}
+	}
+
+	//////////////////////////////////////////////////////////////////////////////////////////////////////
+
+	FContextValue::FContextValue(IContextCacheStore& InContext, FContextCacheKey InCacheKey)
+		: Context(InContext)
+		, CacheKey(InCacheKey)
+	{}
+
+	FContextValue::FContextValue(IContextCacheStore& InContext, TUniquePtr<FContextCacheElementBase>&& InCacheElement)
+		: Context(InContext)
+		, CacheKey(0)
+		, CacheElement(MoveTemp(InCacheElement))
+	{}
+
+	int32 FContextValue::Num() const
+	{
+		if (const FContextCacheElementBase* CacheEntry = GetCacheEntry())
+		{
+			return CacheEntry->GetNumArrayElements(Context);
+		}
+		return 0;
+	}
+
+	FContextValue FContextValue::GetAt(int32 Index) const
+	{
+		if (const FContextCacheElementBase* CacheEntry = GetCacheEntry())
+		{
+			TUniquePtr<FContextCacheElementBase> NewCacheElement = CacheEntry->CreateFromArrayElement(Context, Index, CacheEntry->GetProperty(), CacheEntry->GetNodeGuid(), CacheEntry->GetNodeHash(), CacheEntry->GetTimestamp());
+			return FContextValue(Context, MoveTemp(NewCacheElement));
+		}
+		return FContextValue(Context, 0);
+	}
+
+	FContextValue FContextValue::ToArray() const
+	{
+		if (const FContextCacheElementBase* CacheEntry = GetCacheEntry())
+		{
+			// arrays return themselves
+			if (CacheEntry->IsArray(Context))
+			{
+				if (CacheElement)
+				{
+					return FContextValue(Context, CacheEntry->Clone(Context));
+				}
+				// simply return the reference to the original cached array 
+				return FContextValue(Context, CacheKey);
+			}
+
+			// actually convert to an array 
+			TUniquePtr<FContextCacheElementBase> NewCacheElement = CacheEntry->CreateArrayFromElement(Context, CacheEntry->GetProperty(), CacheEntry->GetNodeGuid(), CacheEntry->GetNodeHash(), CacheEntry->GetTimestamp());
+			return FContextValue(Context, MoveTemp(NewCacheElement));
+		}
+		return FContextValue(Context, 0);
+	}
+
+	const FContextCacheElementBase* FContextValue::GetCacheEntry() const
+	{
+		if (CacheElement)
+		{
+			return CacheElement.Get();
+		}
+		if (const TUniquePtr<FContextCacheElementBase>* CacheEntry = Context.FindCacheElement(CacheKey))
+		{
+			return CacheEntry->Get();
+		}
+		return nullptr;
+	}
+};
+
+FArchive& operator<<(FArchive& Ar, UE::Dataflow::FTimestamp& ValueIn)
+{
+	Ar << ValueIn.Value;
+	Ar << ValueIn.Invalid;
+	return Ar;
+}
+
+FArchive& operator<<(FArchive& Ar, UE::Dataflow::FContextCache& ValueIn)
+{
+	ValueIn.Serialize(Ar);
+	return Ar;
+}
+
+
+
+
diff --git a/Engine/Source/Runtime/Experimental/Dataflow/Core/Private/Dataflow/DataflowNodeParameters.cpp b/Engine/Source/Runtime/Experimental/Dataflow/Core/Private/Dataflow/DataflowNodeParameters.cpp
--- a/Engine/Source/Runtime/Experimental/Dataflow/Core/Private/Dataflow/DataflowNodeParameters.cpp
+++ b/Engine/Source/Runtime/Experimental/Dataflow/Core/Private/Dataflow/DataflowNodeParameters.cpp
@@ -91,6 +91,80 @@ namespace UE::Dataflow
 
 	//////////////////////////////////////////////////////////////////////
 
+	void FContext::SetDataFromStructView(FContextCacheKey InKey, const FProperty* InProperty, const FConstStructView& StructView, const FGuid& InNodeGuid, uint32 InNodeHash, const FTimestamp& InTimestamp)
+	{
+		TUniquePtr<FContextCacheElementBase> DataStoreEntry = MakeUnique<FContextCacheElementUStruct>(InNodeGuid, InProperty, StructView, InNodeHash, InTimestamp);
+		SetDataImpl(InKey, MoveTemp(DataStoreEntry));
+	}
+
+	void FContext::SetDataFromStructArrayView(FContextCacheKey InKey, const FProperty* InProperty, const FConstStructArrayView& StructArrayView, const FGuid& InNodeGuid, uint32 InNodeHash, const FTimestamp& InTimestamp)
+	{
+		TUniquePtr<FContextCacheElementBase> DataStoreEntry = MakeUnique<FContextCacheElementUStructArray>(InNodeGuid, InProperty, StructArrayView, InNodeHash, InTimestamp);
+		SetDataImpl(InKey, MoveTemp(DataStoreEntry));
+	}
+
+	void FContext::SetDataReference(FContextCacheKey Key, const FProperty* Property, FContextCacheKey ReferenceKey, const FTimestamp& InTimestamp)
+	{
+		// find the reference key to get 
+		if (const TUniquePtr<FContextCacheElementBase>* CacheElement = GetDataImpl(ReferenceKey))
+		{
+			TUniquePtr<FContextCacheElementBase> CacheReferenceElement = (*CacheElement)->CreateReference(ReferenceKey);
+			CacheReferenceElement->SetTimestamp(InTimestamp);
+
+			SetDataImpl(Key, MoveTemp(CacheReferenceElement));
+		}
+		else
+		{
+			ensure(false); // could not find the original cache element 
+		}
+	}
+
+	// this is useful when there's a need to have to have cache entry but  the type is not known and there no connected output
+	// ( like reroute nodes with unconnected input for example ) 
+	// in that case posting an invalid reference, will allow the evaluatino to go through and the node reading it will get a default value instead
+	void FContext::SetNullData(FContextCacheKey InKey, const FProperty* InProperty, const FGuid& InNodeGuid, uint32 InNodeHash, const FTimestamp& InTimestamp)
+	{
+		TUniquePtr<FContextCacheElementNull> CacheNullElement = MakeUnique<FContextCacheElementNull>(InNodeGuid, InProperty, InNodeHash, InTimestamp);
+		SetDataImpl(InKey, MoveTemp(CacheNullElement));
+	}
+
+	const void* FContext::GetUntypedData(FContextCacheKey Key, const FProperty* InProperty) const
+	{
+		if (const TUniquePtr<FContextCacheElementBase>* Cache = GetDataImpl(Key))
+		{
+			return (*Cache)->GetUntypedData(*this, InProperty);
+		}
+		return nullptr;
+	}
+
+	bool FContext::HasData(FContextCacheKey Key, FTimestamp InTimestamp) const
+	{
+		FContextCacheKey IntKey = (FContextCacheKey)Key;
+		return HasDataImpl(Key, InTimestamp);
+	}
+
+	bool FContext::IsEmpty() const
+	{
+		return IsEmptyImpl();
+	}
+
+	void FContext::Serialize(FArchive& Ar)
+	{
+		FTimestamp Timestamp = FTimestamp::Invalid;
+		Ar << Timestamp;
+		Ar << DataStore;
+	}
+
+	const TUniquePtr<FContextCacheElementBase>* FContext::FindCacheElement(FContextCacheKey InKey) const
+	{
+		return GetDataImpl(InKey);
+	}
+
+	bool FContext::HasCacheElement(FContextCacheKey InKey, FTimestamp InTimestamp) const
+	{
+		return HasDataImpl(InKey, InTimestamp);
+	}
+
 	void FContext::SetThreaded(bool bValue)
 	{
 		if (bValue != IsThreaded())
diff --git a/Engine/Source/Runtime/Experimental/Dataflow/Core/Public/Dataflow/DataflowContextCache.h b/Engine/Source/Runtime/Experimental/Dataflow/Core/Public/Dataflow/DataflowContextCache.h
--- a/Engine/Source/Runtime/Experimental/Dataflow/Core/Public/Dataflow/DataflowContextCache.h
+++ b/Engine/Source/Runtime/Experimental/Dataflow/Core/Public/Dataflow/DataflowContextCache.h
@@ -221,17 +221,17 @@ namespace UE::Dataflow
 		template<class T>
 		inline const T& GetData(const IContextCacheStore& Context, const FProperty* PropertyIn, const T& Default) const;
 
-		inline virtual const void* GetUntypedData(const IContextCacheStore& Context, const FProperty* PropertyIn) const override;
+		virtual const void* GetUntypedData(const IContextCacheStore& Context, const FProperty* PropertyIn) const override;
 
-		inline virtual bool IsArray(const IContextCacheStore& Context) const override;
+		virtual bool IsArray(const IContextCacheStore& Context) const override;
 
-		inline virtual int32 GetNumArrayElements(const IContextCacheStore& Context) const override;
+		virtual int32 GetNumArrayElements(const IContextCacheStore& Context) const override;
 
-		inline virtual TUniquePtr<FContextCacheElementBase> CreateFromArrayElement(const IContextCacheStore& Context, int32 Index, const FProperty* InProperty, const FGuid& InNodeGuid, uint32 InNodeHash, const FTimestamp& InTimestamp) const override;
+		virtual TUniquePtr<FContextCacheElementBase> CreateFromArrayElement(const IContextCacheStore& Context, int32 Index, const FProperty* InProperty, const FGuid& InNodeGuid, uint32 InNodeHash, const FTimestamp& InTimestamp) const override;
 
-		inline virtual TUniquePtr<FContextCacheElementBase> CreateArrayFromElement(const IContextCacheStore& Context, const FProperty* InProperty, const FGuid& InNodeGuid, uint32 InNodeHash, const FTimestamp& InTimestamp) const override;
+		virtual TUniquePtr<FContextCacheElementBase> CreateArrayFromElement(const IContextCacheStore& Context, const FProperty* InProperty, const FGuid& InNodeGuid, uint32 InNodeHash, const FTimestamp& InTimestamp) const override;
 
-		inline virtual TUniquePtr<FContextCacheElementBase> Clone(const IContextCacheStore& Context) const override;
+		virtual TUniquePtr<FContextCacheElementBase> Clone(const IContextCacheStore& Context) const override;
 
 	private:
 		const FContextCacheKey DataKey; // this is a key to another cache element
@@ -255,32 +255,14 @@ namespace UE::Dataflow
 			: FContextCacheElementBase(EType::CacheElementNull, InNodeGuid, InProperty, InNodeHash, FTimestamp((Timestamp.Value == 0) ? 0 : (Timestamp.Value - 1)))
 		{}
 
-		virtual const void* GetUntypedData(const IContextCacheStore& Context, const FProperty* PropertyIn) const override
-		{
-			return nullptr;
-		}
-
-		virtual bool IsArray(const IContextCacheStore& Context) const override
-		{
-			return false; // we have no way to tell at this point , null entry means that we don't have a value and we let the requester of teh value to get a default value 
-		}
-
-		virtual int32 GetNumArrayElements(const IContextCacheStore& Context) const override
-		{
-			return 0;
-		}
-
-		virtual TUniquePtr<FContextCacheElementBase> CreateFromArrayElement(const IContextCacheStore& Context, int32 Index, const FProperty* InProperty, const FGuid& InNodeGuid, uint32 InNodeHash, const FTimestamp& InTimestamp) const override
-		{
-			return {};
-		}
+		virtual const void* GetUntypedData(const IContextCacheStore& Context, const FProperty* PropertyIn) const override;
 
-		virtual TUniquePtr<FContextCacheElementBase> CreateArrayFromElement(const IContextCacheStore& Context, const FProperty* InProperty, const FGuid& InNodeGuid, uint32 InNodeHash, const FTimestamp& InTimestamp) const override
-		{
-			return {};
-		}
+		virtual bool IsArray(const IContextCacheStore& Context) const override;
+		virtual int32 GetNumArrayElements(const IContextCacheStore& Context) const override;
+		virtual TUniquePtr<FContextCacheElementBase> CreateFromArrayElement(const IContextCacheStore& Context, int32 Index, const FProperty* InProperty, const FGuid& InNodeGuid, uint32 InNodeHash, const FTimestamp& InTimestamp) const override;
+		virtual TUniquePtr<FContextCacheElementBase> CreateArrayFromElement(const IContextCacheStore& Context, const FProperty* InProperty, const FGuid& InNodeGuid, uint32 InNodeHash, const FTimestamp& InTimestamp) const override;
 
-		inline virtual TUniquePtr<FContextCacheElementBase> Clone(const IContextCacheStore& Context) const override;
+		virtual TUniquePtr<FContextCacheElementBase> Clone(const IContextCacheStore& Context) const override;
 	};
 
 	//--------------------------------------------------------------------
@@ -412,58 +394,31 @@ namespace UE::Dataflow
 	//--------------------------------------------------------------------
 	struct FContextCacheElementUStruct : public FContextCacheElementBase
 	{
-		FContextCacheElementUStruct(FGuid InNodeGuid, const FProperty* InProperty, const FConstStructView& StructView, uint32 InNodeHash, FTimestamp Timestamp)
-			: FContextCacheElementBase(EType::CacheElementUStruct, InNodeGuid, InProperty, InNodeHash, Timestamp)
-			, InstancedStruct(StructView)
-		{}
+		explicit FContextCacheElementUStruct(FGuid InNodeGuid, const FProperty* InProperty, const FConstStructView& StructView, uint32 InNodeHash, FTimestamp Timestamp);
 
 		template<typename T>
-		FContextCacheElementUStruct(FGuid InNodeGuid, const FProperty* InProperty, T&& InStruct, uint32 InNodeHash, FTimestamp Timestamp)
+		explicit FContextCacheElementUStruct(FGuid InNodeGuid, const FProperty* InProperty, T&& InStruct, uint32 InNodeHash, FTimestamp Timestamp)
 			: FContextCacheElementBase(EType::CacheElementUStruct, InNodeGuid, InProperty, InNodeHash, Timestamp)
 		{
 			InstancedStruct.InitializeAs<typename TDecay<T>::Type>(Forward<T>(InStruct));
 		}
 
-		FContextCacheElementUStruct(const FContextCacheElementUStruct& Other)
-			: FContextCacheElementBase(EType::CacheElementUStruct, Other.GetNodeGuid(), Other.GetProperty(), Other.GetNodeHash(), Other.GetTimestamp())
-			, InstancedStruct(Other.InstancedStruct)
-		{}
+		explicit FContextCacheElementUStruct(const FContextCacheElementUStruct& Other);
 
 		template<typename T>
-		const T& GetData(const IContextCacheStore& /*Context*/, const FProperty* /*PropertyIn*/, const T& /*Default*/) const
+		inline const T& GetData(const IContextCacheStore& /*Context*/, const FProperty* /*PropertyIn*/, const T& /*Default*/) const
 		{
 			return InstancedStruct.Get<T>();
 		}
 
-		virtual const void* GetUntypedData(const IContextCacheStore& /*Context*/, const FProperty* /*PropertyIn*/) const override
-		{
-			return InstancedStruct.GetMemory();
-		}
+		DATAFLOWCORE_API virtual const void* GetUntypedData(const IContextCacheStore& /*Context*/, const FProperty* /*PropertyIn*/) const override;
 
-		virtual bool IsArray(const IContextCacheStore& Context) const override
-		{
-			return false;
-		}
+		DATAFLOWCORE_API virtual bool IsArray(const IContextCacheStore& Context) const override;
+		DATAFLOWCORE_API virtual int32 GetNumArrayElements(const IContextCacheStore& Context) const override;
+		DATAFLOWCORE_API virtual TUniquePtr<FContextCacheElementBase> CreateFromArrayElement(const IContextCacheStore& Context, int32 Index, const FProperty* InProperty, const FGuid& InNodeGuid, uint32 InNodeHash, const FTimestamp& InTimestamp) const override;
+		DATAFLOWCORE_API virtual TUniquePtr<FContextCacheElementBase> CreateArrayFromElement(const IContextCacheStore& Context, const FProperty* InProperty, const FGuid& InNodeGuid, uint32 InNodeHash, const FTimestamp& InTimestamp) const override;
 
-		virtual int32 GetNumArrayElements(const IContextCacheStore& Context) const override
-		{
-			return 0;
-		}
-
-		virtual TUniquePtr<FContextCacheElementBase> CreateFromArrayElement(const IContextCacheStore& Context, int32 Index, const FProperty* InProperty, const FGuid& InNodeGuid, uint32 InNodeHash, const FTimestamp& InTimestamp) const override
-		{
-			return {};
-		}
-
-		virtual TUniquePtr<FContextCacheElementBase> CreateArrayFromElement(const IContextCacheStore& Context, const FProperty* InProperty, const FGuid& InNodeGuid, uint32 InNodeHash, const FTimestamp& InTimestamp) const override
-		{
-			return {}; // can't make arrays of arrays
-		}
-
-		inline virtual TUniquePtr<FContextCacheElementBase> Clone(const IContextCacheStore& Context) const override
-		{
-			return MakeUnique<FContextCacheElementUStruct>(*this);
-		}
+		DATAFLOWCORE_API virtual TUniquePtr<FContextCacheElementBase> Clone(const IContextCacheStore& Context) const override;
 
 	private:
 		FInstancedStruct InstancedStruct;
@@ -474,23 +429,29 @@ namespace UE::Dataflow
 	//--------------------------------------------------------------------
 	struct FContextCacheElementUStructArray : public FContextCacheElementBase
 	{
-		FContextCacheElementUStructArray(FGuid InNodeGuid, const FProperty* InProperty, const FConstStructArrayView& InStructArrayView, uint32 InNodeHash, FTimestamp Timestamp)
+		explicit FContextCacheElementUStructArray(FGuid InNodeGuid, const FProperty* InProperty, const FConstStructArrayView& InStructArrayView, uint32 InNodeHash, FTimestamp Timestamp)
 			: FContextCacheElementBase(EType::CacheElementUStructArray, InNodeGuid, InProperty, InNodeHash, Timestamp)
 			, InstancedStructArray(InStructArrayView)
 		{}
 
+		// construct from a single struct view into an array 
+		explicit FContextCacheElementUStructArray(FGuid InNodeGuid, const FProperty* InProperty, const FConstStructView& InStructView, uint32 InNodeHash, FTimestamp Timestamp)
+			: FContextCacheElementBase(EType::CacheElementUStructArray, InNodeGuid, InProperty, InNodeHash, Timestamp)
+			, InstancedStructArray(InStructView)
+		{}
+
 		template<typename T>
-		FContextCacheElementUStructArray(FGuid InNodeGuid, const FProperty* InProperty, TArray<T>&& InStructArray, uint32 InNodeHash, FTimestamp Timestamp)
+		explicit FContextCacheElementUStructArray(FGuid InNodeGuid, const FProperty* InProperty, TArray<T>&& InStructArray, uint32 InNodeHash, FTimestamp Timestamp)
 			: FContextCacheElementBase(EType::CacheElementUStructArray, InNodeGuid, InProperty, InNodeHash, Timestamp)
 			, InstancedStructArray(TBaseStructure<T>::Get())
 		{
 			InstancedStructArray.Get<T>() = Forward<TArray<T>>(InStructArray);
 		}
 
-		FContextCacheElementUStructArray(const FContextCacheElementUStructArray& Other)
+		explicit FContextCacheElementUStructArray(const FContextCacheElementUStructArray& Other)
 			: FContextCacheElementUStructArray(Other.GetNodeGuid(), Other.GetProperty(), Other.GetStructArrayView(), Other.GetNodeHash(), Other.GetTimestamp())
 		{}
-	
+
 		virtual ~FContextCacheElementUStructArray() override = default;
 
 		template<typename T>
@@ -499,90 +460,41 @@ namespace UE::Dataflow
 			return InstancedStructArray.Get<T>();
 		}
 
-		virtual const void* GetUntypedData(const IContextCacheStore& /*Context*/, const FProperty* /*PropertyIn*/) const override
-		{
-			return InstancedStructArray.GetMemory();
-		}
-
-		virtual bool IsArray(const IContextCacheStore& Context) const override
-		{
-			return true;
-		}
-
-		virtual int32 GetNumArrayElements(const IContextCacheStore& Context) const override
-		{
-			return InstancedStructArray.Num();
-		}
-
-		virtual TUniquePtr<FContextCacheElementBase> CreateFromArrayElement(const IContextCacheStore& Context, int32 Index, const FProperty* InProperty, const FGuid& InNodeGuid, uint32 InNodeHash, const FTimestamp& InTimestamp) const override
-		{
-			if (InstancedStructArray.IsValidIndex(Index))
-			{
-				//return MakeUnique<FContextCacheElementUStruct>(GetNodeGuid(), GetProperty(), Forward<T>(ObjectPtrArray[Index]), GetNodeHash(), GetTimestamp());
-				// todo : support extracting structure from array of structures
-				return {};
-			}
-			return {};
-		}
+		virtual const void* GetUntypedData(const IContextCacheStore& /*Context*/, const FProperty* /*PropertyIn*/) const override;
 
-		virtual TUniquePtr<FContextCacheElementBase> CreateArrayFromElement(const IContextCacheStore& Context, const FProperty* InProperty, const FGuid& InNodeGuid, uint32 InNodeHash, const FTimestamp& InTimestamp) const override
-		{
-			return {}; // can't make arrays of arrays
-		}
+		virtual bool IsArray(const IContextCacheStore& Context) const override;
+		virtual int32 GetNumArrayElements(const IContextCacheStore& Context) const override;
+		virtual TUniquePtr<FContextCacheElementBase> CreateFromArrayElement(const IContextCacheStore& Context, int32 Index, const FProperty* InProperty, const FGuid& InNodeGuid, uint32 InNodeHash, const FTimestamp& InTimestamp) const override;
+		virtual TUniquePtr<FContextCacheElementBase> CreateArrayFromElement(const IContextCacheStore& Context, const FProperty* InProperty, const FGuid& InNodeGuid, uint32 InNodeHash, const FTimestamp& InTimestamp) const override;
 
-		inline virtual TUniquePtr<FContextCacheElementBase> Clone(const IContextCacheStore& Context) const override
-		{
-			return MakeUnique<FContextCacheElementUStructArray>(*this);
-		}
+		virtual TUniquePtr<FContextCacheElementBase> Clone(const IContextCacheStore& Context) const override;
 
 	private:
-		FConstStructArrayView GetStructArrayView() const
-		{
-			return InstancedStructArray.GetScriptStruct() ?
-				FConstStructArrayView(*InstancedStructArray.GetScriptStruct(), InstancedStructArray.GetData(), InstancedStructArray.Num()) :
-				FConstStructArrayView();
-		}
+		FConstStructArrayView GetStructArrayView() const;
 
 		// Implements a FInstancedStruct for arrays (FInstancedStructContainer cannot be cast to a TArray)
 		struct FInstancedStructArray final : public TArray<uint8>
 		{
-			explicit FInstancedStructArray(const UScriptStruct* const InScriptStruct)
-				: ScriptStruct(InScriptStruct)
-			{}
-			explicit FInstancedStructArray(const FConstStructArrayView& StructArrayView)
-				: ScriptStruct(StructArrayView.GetScriptStruct())
-			{
-				if (ScriptStruct)  // Null if the array has 0 element
-				{
-					SetNumUnsafeInternal(StructArrayView.Num());
-					ArrayMax = Num();
-					GetAllocatorInstance().ResizeAllocation(0, Num(), FMath::Max(1, ScriptStruct->GetStructureSize()), ScriptStruct->GetMinAlignment());
-					ScriptStruct->InitializeStruct(GetData(), Num());
-					ScriptStruct->CopyScriptStruct(GetData(), StructArrayView.GetData(), Num());
-				}
-			}
-			~FInstancedStructArray()
-			{
-				if (ScriptStruct)
-				{
-					ScriptStruct->DestroyStruct(GetData(), Num());
-					GetAllocatorInstance().ResizeAllocation(ArrayMax, 0, FMath::Max(1, ScriptStruct->GetStructureSize()), ScriptStruct->GetMinAlignment());
-				}
-			}
-			const UScriptStruct* GetScriptStruct() const
-			{
-				return ScriptStruct;
-			}
-			const void* GetMemory() const
-			{
-				return this;
-			}
+			explicit FInstancedStructArray(const UScriptStruct* const InScriptStruct);
+			explicit FInstancedStructArray(const FConstStructView& StructView);
+			explicit FInstancedStructArray(const FConstStructArrayView& StructArrayView);
+			~FInstancedStructArray();
+
+			const int32 GetStructureSize() const;
+
+			const UScriptStruct* GetScriptStruct() const;
+
+			FConstStructView GetScriptViewAt(int32 Index) const;
+
+			const void* GetMemory() const;
+
 			template<typename T>
 			const TArray<T>& Get() const
 			{
 				check(!ScriptStruct || TBaseStructure<T>::Get() == ScriptStruct);
 				return reinterpret_cast<const TArray<T>&>(*this);
 			}
+
 			template<typename T>
 			TArray<T>& Get()
 			{
@@ -590,6 +502,8 @@ namespace UE::Dataflow
 				return reinterpret_cast<TArray<T>&>(*this);
 			}
 		private:
+			void InitFromRawData(const void* Data, const int32 Num);
+
 			const UScriptStruct* const ScriptStruct;
 		};
 
@@ -689,68 +603,6 @@ namespace UE::Dataflow
 		return Default;
 	}
 
-	const void* FContextCacheElementReference::GetUntypedData(const IContextCacheStore& Context, const FProperty* InProperty) const
-	{
-		if (const TUniquePtr<FContextCacheElementBase>* Cache = Context.FindCacheElement(DataKey))
-		{
-			return (*Cache)->GetUntypedData(Context, InProperty);
-		}
-		return nullptr;
-	}
-
-	bool FContextCacheElementReference::IsArray(const IContextCacheStore& Context) const
-	{
-		if (const TUniquePtr<FContextCacheElementBase>* Cache = Context.FindCacheElement(DataKey))
-		{
-			return (*Cache)->IsArray(Context);
-		}
-		return false;
-	}
-
-	int32 FContextCacheElementReference::GetNumArrayElements(const IContextCacheStore& Context) const
-	{
-		if (const TUniquePtr<FContextCacheElementBase>* Cache = Context.FindCacheElement(DataKey))
-		{
-			return (*Cache)->GetNumArrayElements(Context);
-		}
-		return 0;
-	}
-
-	TUniquePtr<FContextCacheElementBase> FContextCacheElementReference::CreateFromArrayElement(const IContextCacheStore& Context, int32 Index, const FProperty* InProperty, const FGuid& InNodeGuid, uint32 InNodeHash, const FTimestamp& InTimestamp) const
-	{
-		if (const TUniquePtr<FContextCacheElementBase>* Cache = Context.FindCacheElement(DataKey))
-		{
-			return (*Cache)->CreateFromArrayElement(Context, Index, InProperty, InNodeGuid, InNodeHash, InTimestamp);
-		}
-		return {};
-	}
-
-	TUniquePtr<FContextCacheElementBase> FContextCacheElementReference::CreateArrayFromElement(const IContextCacheStore& Context, const FProperty* InProperty, const FGuid& InNodeGuid, uint32 InNodeHash, const FTimestamp& InTimestamp) const
-	{
-		if (const TUniquePtr<FContextCacheElementBase>* Cache = Context.FindCacheElement(DataKey))
-		{
-			return (*Cache)->CreateArrayFromElement(Context, InProperty, InNodeGuid, InNodeHash, InTimestamp);
-		}
-		return {};
-	}
-
-	TUniquePtr<FContextCacheElementBase> FContextCacheElementReference::Clone(const IContextCacheStore& Context) const
-	{
-		if (const TUniquePtr<FContextCacheElementBase>* ReferencedCacheEntry = Context.FindCacheElement(DataKey))
-		{
-			if (ReferencedCacheEntry && ReferencedCacheEntry->IsValid())
-			{
-				return (*ReferencedCacheEntry)->Clone(Context);
-			}
-		}
-		return MakeUnique<FContextCacheElementNull>(GetNodeGuid(), GetProperty(), GetNodeHash(), GetTimestamp());
-	}
-
-	TUniquePtr<FContextCacheElementBase> FContextCacheElementNull::Clone(const IContextCacheStore& Context) const
-	{
-		return MakeUnique<FContextCacheElementNull>(GetNodeGuid(), GetProperty(), GetNodeHash(), GetTimestamp());
-	}
-
 	template<class T>
 	TUniquePtr<FContextCacheElementBase> TContextCacheElementUObject<T>::Clone(const IContextCacheStore& Context) const
 	{
@@ -763,75 +615,19 @@ namespace UE::Dataflow
 		return MakeUnique<TContextCacheElementUObjectArray<T>>(*this);
 	}
 
-	////////////////////////////////////////////////////
+	////////////////////////////////////////////////////////////////////////////////////////////////////////
 
 	struct FContextValue
 	{
 	public:
-		FContextValue(IContextCacheStore& InContext, FContextCacheKey InCacheKey)
-			: Context(InContext)
-			, CacheKey(InCacheKey)
-		{}
-
-		FContextValue(IContextCacheStore& InContext, TUniquePtr<FContextCacheElementBase>&& InCacheElement)
-			: Context(InContext)
-			, CacheKey(0)
-			, CacheElement(MoveTemp(InCacheElement))
-		{}
-
-		int32 Num() const
-		{
-			if (const FContextCacheElementBase* CacheEntry = GetCacheEntry())
-			{
-				return CacheEntry->GetNumArrayElements(Context);
-			}
-			return 0;
-		}
+		DATAFLOWCORE_API FContextValue(IContextCacheStore& InContext, FContextCacheKey InCacheKey);
+		DATAFLOWCORE_API FContextValue(IContextCacheStore& InContext, TUniquePtr<FContextCacheElementBase>&& InCacheElement);
 
-		FContextValue GetAt(int32 Index) const
-		{
-			if (const FContextCacheElementBase* CacheEntry = GetCacheEntry())
-			{
-				TUniquePtr<FContextCacheElementBase> NewCacheElement = CacheEntry->CreateFromArrayElement(Context, Index, CacheEntry->GetProperty(), CacheEntry->GetNodeGuid(), CacheEntry->GetNodeHash(), CacheEntry->GetTimestamp());
-				return FContextValue(Context, MoveTemp(NewCacheElement));
-			}
-			return FContextValue(Context, 0);
-		}
+		DATAFLOWCORE_API int32 Num() const;
+		DATAFLOWCORE_API FContextValue GetAt(int32 Index) const;
+		DATAFLOWCORE_API FContextValue ToArray() const;
 
-		FContextValue ToArray() const
-		{
-			if (const FContextCacheElementBase* CacheEntry = GetCacheEntry())
-			{
-				// arrays return themselves
-				if (CacheEntry->IsArray(Context))
-				{
-					if (CacheElement)
-					{
-						return FContextValue(Context, CacheEntry->Clone(Context));
-					}
-					// simply return the reference to the original cached array 
-					return FContextValue(Context, CacheKey);
-				}
-
-				// actually convert to an array 
-				TUniquePtr<FContextCacheElementBase> NewCacheElement = CacheEntry->CreateArrayFromElement(Context, CacheEntry->GetProperty(), CacheEntry->GetNodeGuid(), CacheEntry->GetNodeHash(), CacheEntry->GetTimestamp());
-				return FContextValue(Context, MoveTemp(NewCacheElement));
-			}
-			return FContextValue(Context, 0);
-		}
-
-		const FContextCacheElementBase* GetCacheEntry() const
-		{
-			if (CacheElement)
-			{
-				return CacheElement.Get();
-			}
-			if (const TUniquePtr<FContextCacheElementBase>* CacheEntry = Context.FindCacheElement(CacheKey))
-			{
-				return CacheEntry->Get();
-			}
-			return nullptr;
-		}
+		DATAFLOWCORE_API const FContextCacheElementBase* GetCacheEntry() const;
 
 	private:
 		IContextCacheStore& Context;
@@ -840,18 +636,8 @@ namespace UE::Dataflow
 	};
 };
 
-inline FArchive& operator<<(FArchive& Ar, UE::Dataflow::FTimestamp& ValueIn)
-{
-	Ar << ValueIn.Value;
-	Ar << ValueIn.Invalid;
-	return Ar;
-}
-
-inline FArchive& operator<<(FArchive& Ar, UE::Dataflow::FContextCache& ValueIn)
-{
-	ValueIn.Serialize(Ar);
-	return Ar;
-}
+DATAFLOWCORE_API FArchive& operator<<(FArchive& Ar, UE::Dataflow::FTimestamp& ValueIn);
+DATAFLOWCORE_API FArchive& operator<<(FArchive& Ar, UE::Dataflow::FContextCache& ValueIn);
 
 
 
diff --git a/Engine/Source/Runtime/Experimental/Dataflow/Core/Public/Dataflow/DataflowNodeParameters.h b/Engine/Source/Runtime/Experimental/Dataflow/Core/Public/Dataflow/DataflowNodeParameters.h
--- a/Engine/Source/Runtime/Experimental/Dataflow/Core/Public/Dataflow/DataflowNodeParameters.h
+++ b/Engine/Source/Runtime/Experimental/Dataflow/Core/Public/Dataflow/DataflowNodeParameters.h
@@ -251,42 +251,16 @@ namespace UE::Dataflow
 			}
 		}
 
-		void SetDataFromStructView(FContextCacheKey InKey, const FProperty* InProperty, const FConstStructView& StructView, const FGuid& InNodeGuid, uint32 InNodeHash, const FTimestamp& InTimestamp)
-		{
-			TUniquePtr<FContextCacheElementBase> DataStoreEntry = MakeUnique<FContextCacheElementUStruct>(InNodeGuid, InProperty, StructView, InNodeHash, InTimestamp);
-			SetDataImpl(InKey, MoveTemp(DataStoreEntry));
-		}
+		DATAFLOWCORE_API void SetDataFromStructView(FContextCacheKey InKey, const FProperty* InProperty, const FConstStructView& StructView, const FGuid& InNodeGuid, uint32 InNodeHash, const FTimestamp& InTimestamp);
 
-		void SetDataFromStructArrayView(FContextCacheKey InKey, const FProperty* InProperty, const FConstStructArrayView& StructArrayView, const FGuid& InNodeGuid, uint32 InNodeHash, const FTimestamp& InTimestamp)
-		{
-			TUniquePtr<FContextCacheElementBase> DataStoreEntry = MakeUnique<FContextCacheElementUStructArray>(InNodeGuid, InProperty, StructArrayView, InNodeHash, InTimestamp);
-			SetDataImpl(InKey, MoveTemp(DataStoreEntry));
-		}
+		DATAFLOWCORE_API void SetDataFromStructArrayView(FContextCacheKey InKey, const FProperty* InProperty, const FConstStructArrayView& StructArrayView, const FGuid& InNodeGuid, uint32 InNodeHash, const FTimestamp& InTimestamp);
 
-		void SetDataReference(FContextCacheKey Key, const FProperty* Property, FContextCacheKey ReferenceKey, const FTimestamp& InTimestamp)
-		{
-			// find the reference key to get 
-			if (const TUniquePtr<FContextCacheElementBase>* CacheElement = GetDataImpl(ReferenceKey))
-			{
-				TUniquePtr<FContextCacheElementBase> CacheReferenceElement = (*CacheElement)->CreateReference(ReferenceKey);
-				CacheReferenceElement->SetTimestamp(InTimestamp);
-
-				SetDataImpl(Key, MoveTemp(CacheReferenceElement));
-			}
-			else
-			{
-				ensure(false); // could not find the original cache element 
-			}
-		}
+		DATAFLOWCORE_API void SetDataReference(FContextCacheKey Key, const FProperty* Property, FContextCacheKey ReferenceKey, const FTimestamp& InTimestamp);
 
 		// this is useful when there's a need to have to have cache entry but  the type is not known and there no connected output
 		// ( like reroute nodes with unconnected input for example ) 
 		// in that case posting an invalid reference, will allow the evaluatino to go through and the node reading it will get a default value instead
-		void SetNullData(FContextCacheKey InKey, const FProperty* InProperty, const FGuid& InNodeGuid, uint32 InNodeHash, const FTimestamp& InTimestamp)
-		{
-			TUniquePtr<FContextCacheElementNull> CacheNullElement = MakeUnique<FContextCacheElementNull>(InNodeGuid, InProperty, InNodeHash, InTimestamp);
-			SetDataImpl(InKey, MoveTemp(CacheNullElement));
-		}
+		DATAFLOWCORE_API void SetNullData(FContextCacheKey InKey, const FProperty* InProperty, const FGuid& InNodeGuid, uint32 InNodeHash, const FTimestamp& InTimestamp);
 
 		UE_DEPRECATED(5.7, "Use UE::Dataflow::FContextValue from calling GetValue on the node to access the same functionality")
 		DATAFLOWCORE_API int32 GetArraySizeFromData(const FContextCacheKey InKey) const;
@@ -306,42 +280,14 @@ namespace UE::Dataflow
 			return Default;
 		}
 
-		const void* GetUntypedData(FContextCacheKey Key, const FProperty* InProperty) const
-		{
-			if (const TUniquePtr<FContextCacheElementBase>* Cache = GetDataImpl(Key))
-			{
-				return (*Cache)->GetUntypedData(*this, InProperty);
-			}
-			return nullptr;
-		}
-
-		bool HasData(FContextCacheKey Key, FTimestamp InTimestamp = FTimestamp::Invalid) const
-		{
-			FContextCacheKey IntKey = (FContextCacheKey)Key;
-			return HasDataImpl(Key, InTimestamp);
-		}
-
-		bool IsEmpty() const
-		{
-			return IsEmptyImpl();
-		}
-
-		virtual void Serialize(FArchive& Ar)
-		{
-			FTimestamp Timestamp = FTimestamp::Invalid;
-			Ar << Timestamp;
-			Ar << DataStore;
-		}
+		DATAFLOWCORE_API const void* GetUntypedData(FContextCacheKey Key, const FProperty* InProperty) const;
+		DATAFLOWCORE_API bool HasData(FContextCacheKey Key, FTimestamp InTimestamp = FTimestamp::Invalid) const;
+		DATAFLOWCORE_API bool IsEmpty() const;
+		DATAFLOWCORE_API virtual void Serialize(FArchive& Ar);
 
 		// Begin - IContextCacheStore
-		virtual const TUniquePtr<FContextCacheElementBase>* FindCacheElement(FContextCacheKey InKey) const override
-		{
-			return GetDataImpl(InKey);
-		}
-		virtual bool HasCacheElement(FContextCacheKey InKey, FTimestamp InTimestamp = FTimestamp::Invalid) const override
-		{
-			return HasDataImpl(InKey, InTimestamp);
-		}
+		DATAFLOWCORE_API virtual const TUniquePtr<FContextCacheElementBase>* FindCacheElement(FContextCacheKey InKey) const override;
+		DATAFLOWCORE_API virtual bool HasCacheElement(FContextCacheKey InKey, FTimestamp InTimestamp = FTimestamp::Invalid) const override;
 		// End - IContextCacheStore
 
 		// Begin - IContextAssetStoreInterface
