Enabled chunked transfer when resolving the  distributed endpoint for IAX.

#rnx

- Based this change on the work to allow chunked transfer for the http client in CL 43025307
-- Exposed ::JoinIoBuffers from the client code so that it can be re-used.

diff --git a/Engine/Source/Runtime/Experimental/IoStore/OnDemand/Private/DistributionEndpoints.cpp b/Engine/Source/Runtime/Experimental/IoStore/OnDemand/Private/DistributionEndpoints.cpp
--- a/Engine/Source/Runtime/Experimental/IoStore/OnDemand/Private/DistributionEndpoints.cpp
+++ b/Engine/Source/Runtime/Experimental/IoStore/OnDemand/Private/DistributionEndpoints.cpp
@@ -13,6 +13,7 @@
 #include "Logging/StructuredLog.h"
 #include "Misc/CommandLine.h"
 #include "Misc/ConfigCacheIni.h"
+#include "OnDemandHttpClient.h"
 #include "Serialization/JsonReader.h"
 #include "Serialization/JsonSerializer.h"
 #include "Serialization/MemoryReader.h"
@@ -52,7 +53,9 @@ FDistributionEndpoints::EResult FDistributionEndpoints::ResolveEndpoints(const F
 		&OutServiceUrls,
 		&Event,
 		&Result,
-		Dest=FIoBuffer()
+		bIsChunkedTransfer = false,
+		Dest=FIoBuffer(),
+		Chunk = FIoBuffer()
 		] (const FTicketStatus& Status) mutable
 	{
 		if (Status.GetId() == FTicketStatus::EId::Error)
@@ -69,9 +72,16 @@ FDistributionEndpoints::EResult FDistributionEndpoints::ResolveEndpoints(const F
 
 		if (Status.GetId() == FTicketStatus::EId::Content)
 		{
+			if (bIsChunkedTransfer && Chunk.DataSize() != 0)
+			{
+				Dest = JoinIoBuffers(Dest, Chunk);
+				return;
+			}
+
 			bHasResponse = true;
 			Result = ParseResponse(Dest, OutServiceUrls);
 			Event.Trigger();
+
 			return;
 		}
 
@@ -82,7 +92,8 @@ FDistributionEndpoints::EResult FDistributionEndpoints::ResolveEndpoints(const F
 			return;
 		}
 
-		Response.SetDestination(&Dest);
+		bIsChunkedTransfer = Response.GetContentLength() == -1;
+		Response.SetDestination(bIsChunkedTransfer == false ? &Dest : &Chunk);
 	};
 
 	auto AnsiUrl = StringCast<ANSICHAR>(*DistributionUrl, DistributionUrl.Len());
@@ -93,7 +104,7 @@ FDistributionEndpoints::EResult FDistributionEndpoints::ResolveEndpoints(const F
 		Loop.SetFailTimeout(GDistributedEndpointTimeout * 1000);
 	}
 
-	FEventLoop::FRequestParams RequestParams = { .bAllowChunked = false };
+	FEventLoop::FRequestParams RequestParams = { .bAllowChunked = true };
 	FRequest Request = Loop.Get(AnsiUrl, &RequestParams);
 	Request.Header("Accept", "application/json");
 	Loop.Send(MoveTemp(Request), OnRequestStatus);
diff --git a/Engine/Source/Runtime/Experimental/IoStore/OnDemand/Private/OnDemandHttpClient.cpp b/Engine/Source/Runtime/Experimental/IoStore/OnDemand/Private/OnDemandHttpClient.cpp
--- a/Engine/Source/Runtime/Experimental/IoStore/OnDemand/Private/OnDemandHttpClient.cpp
+++ b/Engine/Source/Runtime/Experimental/IoStore/OnDemand/Private/OnDemandHttpClient.cpp
@@ -28,8 +28,7 @@ static FAutoConsoleVariableRef CVar_IaxWarnOnFailure(
 	TEXT("When enabled FMultiEndpointHttpClient will log failed requests as warnings instead of the config verbosity")
 );
 
-/** Utility function to join two FIoBuffers together and return the combined result as a new FIoBuffer */
-static FIoBuffer JoinIoBuffers(const FIoBuffer& LHS, const FIoBuffer& RHS)
+FIoBuffer JoinIoBuffers(const FIoBuffer& LHS, const FIoBuffer& RHS)
 {
 	const uint64 TotalSize = LHS.GetSize() + RHS.GetSize();
 	FIoBuffer Output(TotalSize);
diff --git a/Engine/Source/Runtime/Experimental/IoStore/OnDemand/Private/OnDemandHttpClient.h b/Engine/Source/Runtime/Experimental/IoStore/OnDemand/Private/OnDemandHttpClient.h
--- a/Engine/Source/Runtime/Experimental/IoStore/OnDemand/Private/OnDemandHttpClient.h
+++ b/Engine/Source/Runtime/Experimental/IoStore/OnDemand/Private/OnDemandHttpClient.h
@@ -27,6 +27,9 @@ namespace UE::IoStore
 
 namespace HTTP { class FTicketStatus; }
 
+/** Utility function to join two FIoBuffers together and return the combined result as a new FIoBuffer */
+FIoBuffer JoinIoBuffers(const FIoBuffer& LHS, const FIoBuffer& RHS);
+
 FIoStatus LoadDefaultHttpCertificates(bool& bWasLoaded);
 
 /** Returns if a Http status code is okay (2XX) or not */
