Optimize Integration step
Use task graph instead of parallel for
Run DirtyParticles and the integration in parallel
Up to 35% speed up in some scenarios

diff --git a/Engine/Source/Runtime/Experimental/Chaos/Private/Chaos/PBDRigidsEvolutionGBF.cpp b/Engine/Source/Runtime/Experimental/Chaos/Private/Chaos/PBDRigidsEvolutionGBF.cpp
--- a/Engine/Source/Runtime/Experimental/Chaos/Private/Chaos/PBDRigidsEvolutionGBF.cpp
+++ b/Engine/Source/Runtime/Experimental/Chaos/Private/Chaos/PBDRigidsEvolutionGBF.cpp
@@ -673,13 +673,12 @@ void FPBDRigidsEvolutionGBF::AdvanceOneTimeStepImpl(const FReal Dt, const FSubSt
 	}
 
 	// Assign all islands to a set of groups. Each group is solved in parallel with the others.
-	int32 NumGroups = 0;
 	{
 		SCOPE_CYCLE_COUNTER(STAT_Evolution_BuildGroups);
 		CSV_SCOPED_TIMING_STAT(PhysicsVerbose, StepSolver_BuildGroups);
 
 		// If we are resimulating, only build island groups for islands that require to be simulated based of particles sync state and resim type
-		NumGroups = IslandGroupManager.BuildGroups(bIsResim);
+		IslandGroupManager.BuildGroups(bIsResim);
 	}
 
 
@@ -806,7 +805,7 @@ void FPBDRigidsEvolutionGBF::AdvanceOneTimeStepImpl(const FReal Dt, const FSubSt
 
 void FPBDRigidsEvolutionGBF::Integrate(FReal Dt)
 {
-	TParticleView<FPBDRigidParticles>& ParticlesView = Particles.GetActiveParticlesView();
+	const TArray<FPBDRigidParticleHandle*>& ParticlesArray = Particles.GetActiveParticlesArray();
 
 	//SCOPE_CYCLE_COUNTER(STAT_Integrate);
 	CHAOS_SCOPED_TIMER(Integrate);
@@ -821,149 +820,174 @@ void FPBDRigidsEvolutionGBF::Integrate(FReal Dt)
 	const bool bAllowMACD = CVars::bChaosUseMACD;
 	const bool bForceMACD = CVars::bChaosForceMACD;
 
-	ParticlesView.ParallelFor([&](auto& GeomParticle, int32 Index)
-		{
-			//question: can we enforce this at the API layer? Right now islands contain non dynamic which makes this hard
-			auto PBDParticle = GeomParticle.CastToRigidParticle();
-			if (PBDParticle && PBDParticle->ObjectState() == EObjectStateType::Dynamic)
-			{
-				auto& Particle = *PBDParticle;
+	const int32 NumParticles = ParticlesArray.Num();
+	int32 NumTasks = !!GSingleThreadedPhysics ? 1 : FMath::Min(FTaskGraphInterface::Get().GetNumWorkerThreads(), Chaos::MaxNumWorkers);
+	NumTasks = FMath::Max(FMath::Min(NumTasks, NumParticles), 1);
 
-				FVec3 V = Particle.GetV();
-				FVec3 W = Particle.GetW();
 
-				//save off previous velocities
-				Particle.SetPreV(V);
-				Particle.SetPreW(W);
+	constexpr int32 MinParticlePerTask = 40;
+	const int32 ParticleByTask = FMath::Max(FMath::DivideAndRoundUp(NumParticles, NumTasks), MinParticlePerTask);
+	int32 NumBatches = FMath::DivideAndRoundUp(NumParticles, ParticleByTask);
+
+	TArray<UE::Tasks::FTask> PendingTasks;
+	PendingTasks.Reserve(NumBatches);
+	for (int32 BatchIndex = 0; BatchIndex < NumBatches; BatchIndex++)
+	{
+		const int32 StartIndex = BatchIndex * ParticleByTask;
+		int32 EndIndex = (BatchIndex + 1) * ParticleByTask;
+		EndIndex = FMath::Min(NumParticles, EndIndex);
 
-				for (FForceRule ForceRule : ForceRules)
+		UE::Tasks::FTask PendingTask = UE::Tasks::Launch(UE_SOURCE_LOCATION, [this, StartIndex, EndIndex, &ParticlesArray,
+		Dt, HackMaxAngularSpeedSq, HackMaxLinearSpeedSq, bAllowMACD, bForceMACD, VelocityBoundsMultiplierMACD, MaxVelocityBoundsExpansionMACD,
+		VelocityBoundsMultiplier, MaxVelocityBoundsExpansion, BoundsThickness]()
+		{
+			for (int32 Index = StartIndex; Index < EndIndex; Index++)
+			{
+				//question: can we enforce this at the API layer? Right now islands contain non dynamic which makes this hard
+				FPBDRigidParticleHandle* PBDParticle = ParticlesArray[Index]->CastToRigidParticle();
+				if (PBDParticle && PBDParticle->ObjectState() == EObjectStateType::Dynamic)
 				{
-					ForceRule(Particle, Dt);
-				}
+					FPBDRigidParticleHandle& Particle = *PBDParticle;
 
-				//EulerStepVelocityRule.Apply(Particle, Dt);
-				V += Particle.Acceleration() * Dt;
-				W += Particle.AngularAcceleration() * Dt;
+					FVec3 V = Particle.GetV();
+					FVec3 W = Particle.GetW();
 
-				//AddImpulsesRule.Apply(Particle, Dt);
-				V += Particle.LinearImpulseVelocity();
-				W += Particle.AngularImpulseVelocity();
-				Particle.LinearImpulseVelocity() = FVec3(0);
-				Particle.AngularImpulseVelocity() = FVec3(0);
+					//save off previous velocities
+					Particle.SetPreV(V);
+					Particle.SetPreW(W);
 
-				//EtherDragRule.Apply(Particle, Dt);
+					for (FForceRule ForceRule : ForceRules)
+					{
+						ForceRule(Particle, Dt);
+					}
 
-				const FReal LinearDrag = LinearEtherDragOverride >= 0 ? LinearEtherDragOverride : Particle.LinearEtherDrag() * Dt;
-				const FReal LinearMultiplier = FMath::Max(FReal(0), FReal(1) - LinearDrag);
-				V *= LinearMultiplier;
+					//EulerStepVelocityRule.Apply(Particle, Dt);
+					V += Particle.Acceleration() * Dt;
+					W += Particle.AngularAcceleration() * Dt;
 
-				const FReal AngularDrag = AngularEtherDragOverride >= 0 ? AngularEtherDragOverride : Particle.AngularEtherDrag() * Dt;
-				const FReal AngularMultiplier = FMath::Max(FReal(0), FReal(1) - AngularDrag);
-				W *= AngularMultiplier;
+					//AddImpulsesRule.Apply(Particle, Dt);
+					V += Particle.LinearImpulseVelocity();
+					W += Particle.AngularImpulseVelocity();
+					Particle.LinearImpulseVelocity() = FVec3(0);
+					Particle.AngularImpulseVelocity() = FVec3(0);
 
-				// Apply gyroscopic as a velocity change. This could be applied before the clamps, or after,
-				// but the user might be surprised to see an angular velocity that exceeds the clamp
-				// so do it before.
-				if (Particle.GyroscopicTorqueEnabled())
-				{
-					W = Utilities::GetAngularVelocityAdjustedForGyroscopicTorques(Particle.RCom(), Particle.I(), W, Dt);
-				}
+					//EtherDragRule.Apply(Particle, Dt);
 
-				FReal LinearSpeedSq = V.SizeSquared();
-				if (LinearSpeedSq > Particle.MaxLinearSpeedSq())
-				{
-					V *= FMath::Sqrt(Particle.MaxLinearSpeedSq() / LinearSpeedSq);
-				}
+					const FReal LinearDrag = LinearEtherDragOverride >= 0 ? LinearEtherDragOverride : Particle.LinearEtherDrag() * Dt;
+					const FReal LinearMultiplier = FMath::Max(FReal(0), FReal(1) - LinearDrag);
+					V *= LinearMultiplier;
 
-				FReal AngularSpeedSq = W.SizeSquared();
-				if (AngularSpeedSq > Particle.MaxAngularSpeedSq())
-				{
-					W *= FMath::Sqrt(Particle.MaxAngularSpeedSq() / AngularSpeedSq);
-				}
+					const FReal AngularDrag = AngularEtherDragOverride >= 0 ? AngularEtherDragOverride : Particle.AngularEtherDrag() * Dt;
+					const FReal AngularMultiplier = FMath::Max(FReal(0), FReal(1) - AngularDrag);
+					W *= AngularMultiplier;
 
-				if (CVars::HackMaxAngularVelocity >= 0.f)
-				{
-					AngularSpeedSq = W.SizeSquared();
-					if (AngularSpeedSq > HackMaxAngularSpeedSq)
+					// Apply gyroscopic as a velocity change. This could be applied before the clamps, or after,
+					// but the user might be surprised to see an angular velocity that exceeds the clamp
+					// so do it before.
+					if (Particle.GyroscopicTorqueEnabled())
 					{
-						W = W * (CVars::HackMaxAngularVelocity / FMath::Sqrt(AngularSpeedSq));
+						W = Utilities::GetAngularVelocityAdjustedForGyroscopicTorques(Particle.RCom(), Particle.I(), W, Dt);
 					}
-				}
 
-				if (CVars::HackMaxVelocity >= 0.f)
-				{
-					LinearSpeedSq = V.SizeSquared();
-					if (LinearSpeedSq > HackMaxLinearSpeedSq)
+					FReal LinearSpeedSq = V.SizeSquared();
+					if (LinearSpeedSq > Particle.MaxLinearSpeedSq())
 					{
-						V = V * (CVars::HackMaxVelocity / FMath::Sqrt(LinearSpeedSq));
+						V *= FMath::Sqrt(Particle.MaxLinearSpeedSq() / LinearSpeedSq);
 					}
-				}
 
-				FVec3 PCoM = Particle.XCom();
-				FRotation3 QCoM = Particle.RCom();
-
-				PCoM = PCoM + V * Dt;
-				QCoM = FRotation3::IntegrateRotationWithAngularVelocity(QCoM, W, Dt);
-
-				Particle.SetTransformPQCom(PCoM, QCoM);
-				Particle.SetV(V);
-				Particle.SetW(W);
-
-				// We need to expand the bounds back along velocity otherwise we can miss collisions when a box
-				// lands on another box, imparting velocity to the lower box and causing the boxes to be
-				// separated by more than Cull Distance at collision detection time.
-				// NOTE: We use a different (larger) max bounds expansion when MACD is enabled.
-				FVec3 VelocityBoundsDelta = FVec3(0);
-				const bool bIsMACD = bForceMACD || (bAllowMACD && Particle.MACDEnabled());
-				const FReal ParticleVelocityBoundsMultiplier = bIsMACD ? VelocityBoundsMultiplierMACD : VelocityBoundsMultiplier;
-				const FReal ParticleMaxVelocityBoundsExpansion = bIsMACD ? MaxVelocityBoundsExpansionMACD : MaxVelocityBoundsExpansion;
-				if ((ParticleVelocityBoundsMultiplier > 0) && (ParticleMaxVelocityBoundsExpansion > 0))
-				{
-					VelocityBoundsDelta = (-ParticleVelocityBoundsMultiplier * Dt) * V;
-					VelocityBoundsDelta = VelocityBoundsDelta.BoundToCube(ParticleMaxVelocityBoundsExpansion);
-				}
-
-				if (!Particle.CCDEnabled())
-				{
-					// Expand bounds about P/Q by a small amount. This can still result in missed collisions, especially
-					// when we have joints that pull the body back to X/R, if P-X is greater than the BoundsThickness
-					Particle.UpdateWorldSpaceStateSwept(FRigidTransform3(Particle.GetP(), Particle.GetQ()), FVec3(BoundsThickness), VelocityBoundsDelta);
-				}
-				else
-				{
+					FReal AngularSpeedSq = W.SizeSquared();
+					if (AngularSpeedSq > Particle.MaxAngularSpeedSq())
+					{
+						W *= FMath::Sqrt(Particle.MaxAngularSpeedSq() / AngularSpeedSq);
+					}
 
-#if CHAOS_DEBUG_DRAW
-					if (CVars::ChaosSolverDrawCCDThresholds)
+					if (CVars::HackMaxAngularVelocity >= 0.f)
 					{
-						DebugDraw::DrawCCDAxisThreshold(Particle.GetX(), Particle.CCDAxisThreshold(), Particle.GetP() - Particle.GetX(), Particle.GetQ());
+						AngularSpeedSq = W.SizeSquared();
+						if (AngularSpeedSq > HackMaxAngularSpeedSq)
+						{
+							W = W * (CVars::HackMaxAngularVelocity / FMath::Sqrt(AngularSpeedSq));
+						}
 					}
-#endif
 
-					if (FCCDHelpers::DeltaExceedsThreshold(Particle.CCDAxisThreshold(), Particle.GetP() - Particle.GetX(), Particle.GetQ()))
+					if (CVars::HackMaxVelocity >= 0.f)
 					{
-						// We sweep the bounds from P back along the velocity and expand by a small amount.
-						// If not using tight bounds we also expand the bounds in all directions by Velocity. This is necessary only for secondary CCD collisions
-						// @todo(chaos): expanding the bounds by velocity is very expensive - revisit this
-						const FVec3 VDt = V * Dt;
-						FReal CCDBoundsExpansion = BoundsThickness;
-						if (!CVars::bChaosCollisionCCDUseTightBoundingBox && (CVars::ChaosCollisionCCDConstraintMaxProcessCount > 1))
+						LinearSpeedSq = V.SizeSquared();
+						if (LinearSpeedSq > HackMaxLinearSpeedSq)
 						{
-							CCDBoundsExpansion += VDt.GetAbsMax();
+							V = V * (CVars::HackMaxVelocity / FMath::Sqrt(LinearSpeedSq));
 						}
-						Particle.UpdateWorldSpaceStateSwept(FRigidTransform3(Particle.GetP(), Particle.GetQ()), FVec3(CCDBoundsExpansion), -VDt);
 					}
-					else
+
+					FVec3 PCoM = Particle.XCom();
+					FRotation3 QCoM = Particle.RCom();
+
+					PCoM = PCoM + V * Dt;
+					QCoM = FRotation3::IntegrateRotationWithAngularVelocity(QCoM, W, Dt);
+
+					Particle.SetTransformPQCom(PCoM, QCoM);
+					Particle.SetV(V);
+					Particle.SetW(W);
+
+					// We need to expand the bounds back along velocity otherwise we can miss collisions when a box
+					// lands on another box, imparting velocity to the lower box and causing the boxes to be
+					// separated by more than Cull Distance at collision detection time.
+					// NOTE: We use a different (larger) max bounds expansion when MACD is enabled.
+					FVec3 VelocityBoundsDelta = FVec3(0);
+					const bool bIsMACD = bForceMACD || (bAllowMACD && Particle.MACDEnabled());
+					const FReal ParticleVelocityBoundsMultiplier = bIsMACD ? VelocityBoundsMultiplierMACD : VelocityBoundsMultiplier;
+					const FReal ParticleMaxVelocityBoundsExpansion = bIsMACD ? MaxVelocityBoundsExpansionMACD : MaxVelocityBoundsExpansion;
+					if ((ParticleVelocityBoundsMultiplier > 0) && (ParticleMaxVelocityBoundsExpansion > 0))
 					{
+						VelocityBoundsDelta = (-ParticleVelocityBoundsMultiplier * Dt) * V;
+						VelocityBoundsDelta = VelocityBoundsDelta.BoundToCube(ParticleMaxVelocityBoundsExpansion);
+					}
+
+					if (!Particle.CCDEnabled())
+					{
+						// Expand bounds about P/Q by a small amount. This can still result in missed collisions, especially
+						// when we have joints that pull the body back to X/R, if P-X is greater than the BoundsThickness
 						Particle.UpdateWorldSpaceStateSwept(FRigidTransform3(Particle.GetP(), Particle.GetQ()), FVec3(BoundsThickness), VelocityBoundsDelta);
 					}
+					else
+					{
+
+#if CHAOS_DEBUG_DRAW
+						if (CVars::ChaosSolverDrawCCDThresholds)
+						{
+							DebugDraw::DrawCCDAxisThreshold(Particle.GetX(), Particle.CCDAxisThreshold(), Particle.GetP() - Particle.GetX(), Particle.GetQ());
+						}
+#endif
+
+						if (FCCDHelpers::DeltaExceedsThreshold(Particle.CCDAxisThreshold(), Particle.GetP() - Particle.GetX(), Particle.GetQ()))
+						{
+							// We sweep the bounds from P back along the velocity and expand by a small amount.
+							// If not using tight bounds we also expand the bounds in all directions by Velocity. This is necessary only for secondary CCD collisions
+							// @todo(chaos): expanding the bounds by velocity is very expensive - revisit this
+							const FVec3 VDt = V * Dt;
+							FReal CCDBoundsExpansion = BoundsThickness;
+							if (!CVars::bChaosCollisionCCDUseTightBoundingBox && (CVars::ChaosCollisionCCDConstraintMaxProcessCount > 1))
+							{
+								CCDBoundsExpansion += VDt.GetAbsMax();
+							}
+							Particle.UpdateWorldSpaceStateSwept(FRigidTransform3(Particle.GetP(), Particle.GetQ()), FVec3(CCDBoundsExpansion), -VDt);
+						}
+						else
+						{
+							Particle.UpdateWorldSpaceStateSwept(FRigidTransform3(Particle.GetP(), Particle.GetQ()), FVec3(BoundsThickness), VelocityBoundsDelta);
+						}
+					}
 				}
 			}
 		});
+		PendingTasks.Add(PendingTask);
+	}
 
-	for (auto& Particle : ParticlesView)
+	for (FPBDRigidParticleHandle* Particle : ParticlesArray)
 	{
-		Base::DirtyParticle(Particle);
+		Base::DirtyParticle(*Particle);
 	}
+	UE::Tasks::Wait(PendingTasks);
 }
 
 void FPBDRigidsEvolutionGBF::ApplyKinematicTargets(const FReal Dt, const FReal StepFraction)
