Iris
- Changed assert to ensure to not crash server
- Added ensure and early out if trying to replicate non-actor template objects which is not supported

diff --git a/Engine/Source/Runtime/Experimental/Iris/Core/Private/Iris/ReplicationSystem/ObjectReplicationBridge.cpp b/Engine/Source/Runtime/Experimental/Iris/Core/Private/Iris/ReplicationSystem/ObjectReplicationBridge.cpp
--- a/Engine/Source/Runtime/Experimental/Iris/Core/Private/Iris/ReplicationSystem/ObjectReplicationBridge.cpp
+++ b/Engine/Source/Runtime/Experimental/Iris/Core/Private/Iris/ReplicationSystem/ObjectReplicationBridge.cpp
@@ -503,7 +503,13 @@ UE::Net::FNetRefHandle UObjectReplicationBridge::StartReplicatingRootObject(UObj
 		ensureMsgf(false, TEXT("StartReplicatingRootObject is not allowed during this operation. %s will not be replicated"), *GetNameSafe(Instance));
 		return FNetRefHandle::GetInvalid();
 	}
-	
+
+	if (!ensureMsgf(!Instance->IsTemplate(), TEXT("Object %s cannot be replicated because it's an %s"), 
+		*GetNameSafe(Instance), Instance->HasAnyFlags(RF_ArchetypeObject)?TEXT("Archetype"):TEXT("DefaultObject")))
+	{
+		return FNetRefHandle::GetInvalid();
+	}
+
 	EReplicationFragmentTraits Traits = EReplicationFragmentTraits::CanReplicate;
 	Traits |= Params.bNeedsPreUpdate ? EReplicationFragmentTraits::NeedsPreSendUpdate : EReplicationFragmentTraits::None;
 
@@ -613,8 +619,11 @@ UE::Net::FNetRefHandle UObjectReplicationBridge::StartReplicatingSubObject(UObje
 
 	const FInternalNetRefIndex RootObjectIndex = LocalNetRefHandleManager.GetInternalIndex(Params.RootObjectHandle);
 
-	checkf(!Instance->IsTemplate(), TEXT("Iris cannot replicate subobject %s owned by %s because it's an %s"), 
-		*GetNameSafe(Instance), *GetNameSafe(LocalNetRefHandleManager.GetReplicatedObjectInstance(RootObjectIndex)), Instance->HasAnyFlags(RF_ArchetypeObject)?TEXT("Archetype"):TEXT("DefaultObject"));
+	if (!ensureMsgf(!Instance->IsTemplate(), TEXT("Iris cannot replicate subobject %s owned by %s because it's an %s"), 
+		*GetNameSafe(Instance), *GetNameSafe(LocalNetRefHandleManager.GetReplicatedObjectInstance(RootObjectIndex)), Instance->HasAnyFlags(RF_ArchetypeObject)?TEXT("Archetype"):TEXT("DefaultObject")))
+	{
+		return FNetRefHandle::GetInvalid();
+	}
 
 	if (!ensureMsgf(IsReplicatedHandle(Params.RootObjectHandle), TEXT("Owner %s (%s) must be replicated for subobject %s to replicate."), 
 		*GetNameSafe(LocalNetRefHandleManager.GetReplicatedObjectInstance(RootObjectIndex)), *Params.RootObjectHandle.ToString(), *GetNameSafe(Instance)))
