Allow the shared IAX http client to be hot fixed.

#rnx

- The creation of the shared client/thread occurs very early in engine start up now so hot fixes will not be applied to cvars at the point that they are created, so we need to react to the cvars changing instead.
- Added a new delegate ::OnRecreateHttpClient that is broadcast when any of the http client cvars are changed. Currently this is not publically exposed and only FOnDemandHttpThread is registering to it but in the future we may have more systems.
- Once we notify FOnDemandHttpThread that it needs to recreate it's client we stop dequeuing new requests (if any) so that all inflight requests in the client are completed. Next time FOnDemandHttpThread::TickRequests is called we destroy the old client and create the new one before we resume servicing requests.
- Added a mutex RecreateMutex to make sure if many cvars are being updated while the client is being recreated that they will all be applied in a second pass.
- Added TRACE_CPUPROFILER_EVENT_SCOPE to track areas of code where locks are being added to make long locks easier to find.

diff --git a/Engine/Source/Runtime/Experimental/IoStore/OnDemand/Private/OnDemandHttpThread.cpp b/Engine/Source/Runtime/Experimental/IoStore/OnDemand/Private/OnDemandHttpThread.cpp
--- a/Engine/Source/Runtime/Experimental/IoStore/OnDemand/Private/OnDemandHttpThread.cpp
+++ b/Engine/Source/Runtime/Experimental/IoStore/OnDemand/Private/OnDemandHttpThread.cpp
@@ -52,6 +52,22 @@ namespace UE::IoStore
 {
 
 ///////////////////////////////////////////////////////////////////////////////
+
+DECLARE_MULTICAST_DELEGATE(FOnRecreateHttpClient);
+FOnRecreateHttpClient OnRecreateHttpClient;
+
+void OnHttpClientCVarChanged(IConsoleVariable* CVar)
+{
+	if (ensure(CVar != nullptr))
+	{
+		const FString CVarName = IConsoleManager::Get().FindConsoleObjectName(CVar);
+		UE_LOGFMT(LogIas, Log, "Existing http client config invalidated as cvar '{CvarName}' has changed", CVarName);
+
+		OnRecreateHttpClient.Broadcast();
+	}
+}
+///////////////////////////////////////////////////////////////////////////////
+
 int32 GIasHttpHealthCheckWaitTime = 3000;
 static FAutoConsoleVariableRef CVar_IasHttpHealthCheckWaitTime(
 	TEXT("ias.HttpHealthCheckWaitTime"),
@@ -71,35 +87,40 @@ int32 GIasHttpConnectionCount = 4;
 static FAutoConsoleVariableRef CVar_IasHttpConnectionCount(
 	TEXT("ias.HttpConnectionCount"),
 	GIasHttpConnectionCount,
-	TEXT("Number of open HTTP connections to the on demand endpoint(s).")
+	TEXT("Number of open HTTP connections to the on demand endpoint(s)."),
+	FConsoleVariableDelegate::CreateStatic(&OnHttpClientCVarChanged)
 );
 
 int32 GIasHttpRecvBufKiB = -1;
 static FAutoConsoleVariableRef CVar_GIasHttpRecvBufKiB(
 	TEXT("ias.HttpRecvBufKiB"),
 	GIasHttpRecvBufKiB,
-	TEXT("Recv buffer size")
+	TEXT("Recv buffer size"),
+	FConsoleVariableDelegate::CreateStatic(&OnHttpClientCVarChanged)
 );
 
 int32 GIasHttpRetryCount = -1;
 static FAutoConsoleVariableRef CVar_IasHttpRetryCount(
 	TEXT("ias.HttpRetryCount"),
 	GIasHttpRetryCount,
-	TEXT("Number of times that a request should be retried before being considered failed. A negative value will use the default behaviour, which is one retry per host url provided.")
+	TEXT("Number of times that a request should be retried before being considered failed. A negative value will use the default behaviour, which is one retry per host url provided."),
+	FConsoleVariableDelegate::CreateStatic(&OnHttpClientCVarChanged)
 );
 
 int32 GIasHttpFailTimeOutMs = 4 * 1000;
 static FAutoConsoleVariableRef CVar_IasHttpFailTimeOutMs(
 	TEXT("ias.HttpFailTimeOutMs"),
 	GIasHttpFailTimeOutMs,
-	TEXT("Fail infinite network waits that take longer than this (in ms, a value of zero will use the default)")
+	TEXT("Fail infinite network waits that take longer than this (in ms, a value of zero will use the default)"),
+	FConsoleVariableDelegate::CreateStatic(&OnHttpClientCVarChanged)
 );
 
 bool GIasHttpAllowChunkedXfer = true;
 static FAutoConsoleVariableRef CVar_IasHttpAllowChunkedXfer(
 	TEXT("ias.HttpAllowChunkedXfer"),
 	GIasHttpAllowChunkedXfer,
-	TEXT("Enable/disable IAS' support for chunked transfer encoding")
+	TEXT("Enable/disable IAS' support for chunked transfer encoding"),
+	FConsoleVariableDelegate::CreateStatic(&OnHttpClientCVarChanged)
 );
 
 int32 GIasHttpConcurrentRequests = 8;
@@ -203,10 +224,25 @@ FOnDemandHttpThread::FOnDemandHttpThread()
 
 	Thread.Reset(FForkProcessHelper::CreateForkableThread(this, TEXT("IoStoreOnDemand.Http"), StackSize, ThreadPriority, ThreadAffinityMask, CreateFlags, bAllowPreFork));
 #endif
+
+	OnRecreateHttpClientHandle = OnRecreateHttpClient.AddLambda([this]()
+	{
+		TRACE_CPUPROFILER_EVENT_SCOPE(FOnDemandHttpThread::OnRecreateHttpClient);
+
+		UE::TUniqueLock _(RecreateMutex);
+		bRecreateHttpClient = true;
+		TickThreadEvent->Trigger();
+	});
 }
 
 FOnDemandHttpThread::~FOnDemandHttpThread()
 {
+	if (OnRecreateHttpClientHandle.IsValid())
+	{
+		OnRecreateHttpClient.Remove(OnRecreateHttpClientHandle);
+		OnRecreateHttpClientHandle.Reset();
+	}
+
 	Thread.Reset();
 
 	// TODO: Not 100% sure this is still needed?
@@ -315,6 +351,29 @@ void FOnDemandHttpThread::Tick()
 
 void FOnDemandHttpThread::TickRequests()
 {
+	TRACE_CPUPROFILER_EVENT_SCOPE(FOnDemandHttpThread::TickRequests);
+
+	if (bRecreateHttpClient)
+	{
+		TRACE_CPUPROFILER_EVENT_SCOPE(FOnDemandHttpThread::RecreateHttpClient);
+		UE::TUniqueLock _(RecreateMutex);
+
+		// Keep the old client at first in case we fail to create the new one so that we can restore it
+		TUniquePtr<FMultiEndpointHttpClient> OldHttpClient = MoveTemp(HttpClient);
+		if (TryCreateHttpClient())
+		{
+			OldHttpClient.Reset();
+			UE_LOGFMT(LogIas, Display, "FOnDemandHttpThread: Successfully created a new http client for use");
+		}
+		else
+		{
+			HttpClient = MoveTemp(OldHttpClient);
+			UE_LOGFMT(LogIas, Warning, "FOnDemandHttpThread: Failed to create a new http client, the existing client will continue to be used");
+		}
+
+		bRecreateHttpClient = false;
+	}
+
 	// We can't limit cvars so add a hard coded cap to prevent concurrent requests from getting out of hand.
 	const int32 MaxConcurrentRequests = FMath::Min(GIasHttpConcurrentRequests, 32);
 
@@ -419,7 +478,7 @@ void FOnDemandHttpThread::TickRequests()
 				}
 			}
 
-			if (!NextHttpRequest)
+			if (!NextHttpRequest && !bRecreateHttpClient)
 			{
 				NextHttpRequest = HttpRequests.Dequeue(MaxConcurrentRequests - NumConcurrentRequests);
 			}
@@ -430,7 +489,7 @@ void FOnDemandHttpThread::TickRequests()
 			TRACE_CPUPROFILER_EVENT_SCOPE(FOnDemandHttpThread::TickHttp);
 			while (HttpClient->Tick(GIasHttpPollTimeoutMs, GIasHttpRateLimitKiBPerSecond))
 			{
-				if (!NextHttpRequest)
+				if (!NextHttpRequest && !bRecreateHttpClient)
 				{
 					NextHttpRequest = HttpRequests.Dequeue(MaxConcurrentRequests - NumConcurrentRequests);
 				}
diff --git a/Engine/Source/Runtime/Experimental/IoStore/OnDemand/Private/OnDemandHttpThread.h b/Engine/Source/Runtime/Experimental/IoStore/OnDemand/Private/OnDemandHttpThread.h
--- a/Engine/Source/Runtime/Experimental/IoStore/OnDemand/Private/OnDemandHttpThread.h
+++ b/Engine/Source/Runtime/Experimental/IoStore/OnDemand/Private/OnDemandHttpThread.h
@@ -122,11 +122,15 @@ private:
 	FOnTickIdle OnTickIdleDelegate;
 	TUniquePtr<FMultiEndpointHttpClient> HttpClient;
 
+	FDelegateHandle OnRecreateHttpClientHandle;
+	UE::FMutex RecreateMutex;
+
 	TThreadSafeIntrusiveQueue<FHttpRequest> HttpRequests;
 
 	EThreadPriority ThreadPriority = EThreadPriority::TPri_Normal;
 
 	std::atomic_bool bStopRequested = false;
+	std::atomic_bool bRecreateHttpClient = false;
 	std::atomic_bool bHttpEnabled = true;
 
 	UE::FMutex AllocatorMutex;
