Allow the size of the send buffers to be tweaked for IAX.

#rnx

- We don't really expect this value to be used but if we're exposing the size of the receive buffers we should also expose the size of the send buffers.

diff --git a/Engine/Source/Runtime/Experimental/IoStore/OnDemand/Private/OnDemandHttpClient.cpp b/Engine/Source/Runtime/Experimental/IoStore/OnDemand/Private/OnDemandHttpClient.cpp
--- a/Engine/Source/Runtime/Experimental/IoStore/OnDemand/Private/OnDemandHttpClient.cpp
+++ b/Engine/Source/Runtime/Experimental/IoStore/OnDemand/Private/OnDemandHttpClient.cpp
@@ -447,10 +447,17 @@ TUniquePtr<HTTP::FConnectionPool> FMultiEndpointHttpClient::CreateConnection(FAn
 {
 	HTTP::FConnectionPool::FParams Params;
 	ensure(Params.SetHostFromUrl(HostUrl) >= 0);
+	
 	if (Config.ReceiveBufferSize >= 0)
 	{
 		Params.RecvBufSize = Config.ReceiveBufferSize;
 	}
+
+	if (Config.SendBufferSize >= 0)
+	{
+		Params.SendBufSize = Config.SendBufferSize;
+	}
+
 	Params.ConnectionCount = uint16(Config.MaxConnectionCount);
 
 	return MakeUnique<HTTP::FConnectionPool>(Params);
diff --git a/Engine/Source/Runtime/Experimental/IoStore/OnDemand/Private/OnDemandHttpClient.h b/Engine/Source/Runtime/Experimental/IoStore/OnDemand/Private/OnDemandHttpClient.h
--- a/Engine/Source/Runtime/Experimental/IoStore/OnDemand/Private/OnDemandHttpClient.h
+++ b/Engine/Source/Runtime/Experimental/IoStore/OnDemand/Private/OnDemandHttpClient.h
@@ -56,6 +56,7 @@ struct FMultiEndpointHttpClientConfig
 {
 	int32						MaxConnectionCount = 4;
 	int32						ReceiveBufferSize = -1;
+	int32						SendBufferSize = -1;
 	int32						MaxRetryCount = -1;	// Positive: The number of times to retry a failed request
 													// Zero: Failed requests will not be retried
 													// Negative: A failed request will retry once per provided host url
diff --git a/Engine/Source/Runtime/Experimental/IoStore/OnDemand/Private/OnDemandHttpThread.cpp b/Engine/Source/Runtime/Experimental/IoStore/OnDemand/Private/OnDemandHttpThread.cpp
--- a/Engine/Source/Runtime/Experimental/IoStore/OnDemand/Private/OnDemandHttpThread.cpp
+++ b/Engine/Source/Runtime/Experimental/IoStore/OnDemand/Private/OnDemandHttpThread.cpp
@@ -99,6 +99,14 @@ static FAutoConsoleVariableRef CVar_GIasHttpRecvBufKiB(
 	FConsoleVariableDelegate::CreateStatic(&OnHttpClientCVarChanged)
 );
 
+int32 GIasHttpSendBufKiB = -1;
+static FAutoConsoleVariableRef CVar_GIasHttpSendBufKiB(
+	TEXT("ias.HttpSendBufKiB"),
+	GIasHttpSendBufKiB,
+	TEXT("Send buffer size"),
+	FConsoleVariableDelegate::CreateStatic(&OnHttpClientCVarChanged)
+);
+
 int32 GIasHttpRetryCount = -1;
 static FAutoConsoleVariableRef CVar_IasHttpRetryCount(
 	TEXT("ias.HttpRetryCount"),
@@ -524,6 +532,7 @@ bool FOnDemandHttpThread::TryCreateHttpClient()
 		{
 			.MaxConnectionCount = GIasHttpConnectionCount,
 			.ReceiveBufferSize = GIasHttpRecvBufKiB >= 0 ? GIasHttpRecvBufKiB << 10 : -1,
+			.SendBufferSize = GIasHttpSendBufKiB >= 0 ? GIasHttpSendBufKiB << 10 : -1,
 			.MaxRetryCount = GIasHttpRetryCount,
 			.TimeoutMs = GIasHttpFailTimeOutMs,
 			.Redirects = EHttpRedirects::Disabled,
