Splitting ensure statement to allow better determination of failure cases when NaNs appear in collision events in EventDefaults

diff --git a/Engine/Source/Runtime/Experimental/Chaos/Private/EventDefaults.cpp b/Engine/Source/Runtime/Experimental/Chaos/Private/EventDefaults.cpp
--- a/Engine/Source/Runtime/Experimental/Chaos/Private/EventDefaults.cpp
+++ b/Engine/Source/Runtime/Experimental/Chaos/Private/EventDefaults.cpp
@@ -137,44 +137,49 @@ namespace Chaos
 							}
 
 							const FPBDCollisionConstraint& Constraint = ContactHandle->GetContact();
-							if (ensure(!Constraint.AccumulatedImpulse.ContainsNaN() && FMath::IsFinite(Constraint.GetPhi())))
+							
+							if (ensureMsgf(!Constraint.AccumulatedImpulse.ContainsNaN(), TEXT("Constraint contains a NaN within the accumulated impulse.")))
 							{
-	
-								const FPerShapeData* Shape0 = Constraint.GetShape0();
-								const FPerShapeData* Shape1 = Constraint.GetShape1();
 
-								// If we don't have a filter - allow the notify, otherwise obey the filter flag
-								const bool bFilter0Notify = Shape0 ? Shape0->GetSimData().HasFlag(EFilterFlags::ContactNotify) : true;
-								const bool bFilter1Notify = Shape1 ? Shape1->GetSimData().HasFlag(EFilterFlags::ContactNotify) : true;
-
-								if(!bFilter0Notify && !bFilter1Notify)
+								if (ensureMsgf(FMath::IsFinite(Constraint.GetPhi()), TEXT("Constraint Phi is not finite.")))
 								{
-									// No need to notify - engine didn't request notifications for either shape.
-									continue;
-								}
 
-								const FGeometryParticleHandle* Particle0 = Constraint.GetParticle0();
-								const FGeometryParticleHandle* Particle1 = Constraint.GetParticle1();
-								const FKinematicGeometryParticleHandle* Body0 = Particle0->CastToKinematicParticle();
+									const FPerShapeData* Shape0 = Constraint.GetShape0();
+									const FPerShapeData* Shape1 = Constraint.GetShape1();
 
-								// presently when a rigidbody or kinematic hits static geometry then Body1 is null
-								const FKinematicGeometryParticleHandle* Body1 = Particle1->CastToKinematicParticle();
+									// If we don't have a filter - allow the notify, otherwise obey the filter flag
+									const bool bFilter0Notify = Shape0 ? Shape0->GetSimData().HasFlag(EFilterFlags::ContactNotify) : true;
+									const bool bFilter1Notify = Shape1 ? Shape1->GetSimData().HasFlag(EFilterFlags::ContactNotify) : true;
 
-								const FKinematicGeometryParticleHandle* Primary = Body0 ? Body0 : Body1;
-								const FKinematicGeometryParticleHandle* Secondary = Body0 ? Body1 : Body0;
+									if (!bFilter0Notify && !bFilter1Notify)
+									{
+										// No need to notify - engine didn't request notifications for either shape.
+										continue;
+									}
 
-								//
-								const int32 NumManifoldPoints = Constraint.GetManifoldPoints().Num();
+									const FGeometryParticleHandle* Particle0 = Constraint.GetParticle0();
+									const FGeometryParticleHandle* Particle1 = Constraint.GetParticle1();
+									const FKinematicGeometryParticleHandle* Body0 = Particle0->CastToKinematicParticle();
 
-								if (((Constraint.IsProbe() && NumManifoldPoints > 0) || !Constraint.AccumulatedImpulse.IsZero()) && Primary)
-								{
-									if (ensure(!Constraint.CalculateWorldContactLocation().ContainsNaN() &&
-										!Constraint.CalculateWorldContactNormal().ContainsNaN()) &&
-										!Primary->GetV().ContainsNaN() &&
-										!Primary->GetW().ContainsNaN() &&
-										(Secondary == nullptr || ((!Secondary->GetV().ContainsNaN()) && !Secondary->GetW().ContainsNaN())))
+									// presently when a rigidbody or kinematic hits static geometry then Body1 is null
+									const FKinematicGeometryParticleHandle* Body1 = Particle1->CastToKinematicParticle();
+
+									const FKinematicGeometryParticleHandle* Primary = Body0 ? Body0 : Body1;
+									const FKinematicGeometryParticleHandle* Secondary = Body0 ? Body1 : Body0;
+
+									//
+									const int32 NumManifoldPoints = Constraint.GetManifoldPoints().Num();
+
+									if (((Constraint.IsProbe() && NumManifoldPoints > 0) || !Constraint.AccumulatedImpulse.IsZero()) && Primary)
 									{
-										ValidArray[Index] = true;
+										if (ensure(!Constraint.CalculateWorldContactLocation().ContainsNaN() &&
+											!Constraint.CalculateWorldContactNormal().ContainsNaN()) &&
+											!Primary->GetV().ContainsNaN() &&
+											!Primary->GetW().ContainsNaN() &&
+											(Secondary == nullptr || ((!Secondary->GetV().ContainsNaN()) && !Secondary->GetW().ContainsNaN())))
+										{
+											ValidArray[Index] = true;
+										}
 									}
 								}
 							}
