Lock on-demand install cache lookup tables when flushing pending cache block.

#rnx

diff --git a/Engine/Source/Runtime/Experimental/IoStore/OnDemand/Private/OnDemandInstallCache.cpp b/Engine/Source/Runtime/Experimental/IoStore/OnDemand/Private/OnDemandInstallCache.cpp
--- a/Engine/Source/Runtime/Experimental/IoStore/OnDemand/Private/OnDemandInstallCache.cpp
+++ b/Engine/Source/Runtime/Experimental/IoStore/OnDemand/Private/OnDemandInstallCache.cpp
@@ -167,6 +167,9 @@ struct FCas
 	using FLastAccess			= TMap<FCasBlockId, int64>;
 	using FBlockIdHandleCounts	= TMap<FCasBlockId, int32>;
 
+	void						Lock()		{ Mutex.Lock(); }
+	void						Unlock()	{ Mutex.Unlock(); }
+
 	FIoStatus					Initialize(FStringView Directory, bool bDeleteExisting = false);
 	FCasLocation				FindChunk(const FIoHash& Hash) const;
 	FCasBlockId					CreateBlock();
@@ -3095,15 +3098,18 @@ FIoStatus FOnDemandInstallCache::FlushPendingChunksImpl(FPendingChunks& Chunks,
 
 			check(ChunkHashes.Num() == Offsets.Num());
 			check(CurrentBlockId.IsValid());
-			for (int32 Idx = 0, Count = Offsets.Num(); Idx < Count; ++Idx)
 			{
-				const FCasAddr	CasAddr = FCasAddr::From(ChunkHashes[Idx]);
-				const uint32	ChunkOffset = IntCastChecked<uint32>(Offsets[Idx]);
+				TUniqueLock Lock(Cas);
+				for (int32 Idx = 0, Count = Offsets.Num(); Idx < Count; ++Idx)
+				{
+					const FCasAddr	CasAddr = FCasAddr::From(ChunkHashes[Idx]);
+					const uint32	ChunkOffset = IntCastChecked<uint32>(Offsets[Idx]);
 
-				FCasLocation& Loc = Cas.Lookup.FindOrAdd(CasAddr);
-				Loc.BlockId	= CurrentBlockId;
-				Loc.BlockOffset	= ChunkOffset;
-				Transaction.ChunkLocation(Loc, CasAddr);
+					FCasLocation& Loc = Cas.Lookup.FindOrAdd(CasAddr);
+					Loc.BlockId	= CurrentBlockId;
+					Loc.BlockOffset	= ChunkOffset;
+					Transaction.ChunkLocation(Loc, CasAddr);
+				}
 			}
 		}
 
