Chaos Visual Debugger:: Fix wrong visualization of per particle position/velocity/projection iteration counts. Move iteration counts to FChaosVDParticleDynamicsMisc and add serialization version guards accordingly.

diff --git a/Engine/Source/Runtime/Core/Public/UObject/FortniteMainBranchObjectVersions.inl b/Engine/Source/Runtime/Core/Public/UObject/FortniteMainBranchObjectVersions.inl
--- a/Engine/Source/Runtime/Core/Public/UObject/FortniteMainBranchObjectVersions.inl
+++ b/Engine/Source/Runtime/Core/Public/UObject/FortniteMainBranchObjectVersions.inl
@@ -655,5 +655,7 @@ FORTNITE_MAIN_VERSION(LensComponentDefaultToDistortionSVE, 213)
 // Animation default blend option changed from Linear to HermiteCubic (aka SmoothStep, ease in / ease out)
 FORTNITE_MAIN_VERSION(ChangeDefaultAlphaBlendType, 214)
 
+// Moved Position/Velocity/Projection Iteration Counts from FChaosVDFRigidParticleControlFlags to FChaosVDParticleDynamicMisc
+FORTNITE_MAIN_VERSION(PerParticleIterationCountMovedToDynamicMisc, 215)
 // -----<new versions can be added above this line>-------------------------------------------------
 #undef FORTNITE_MAIN_VERSION
\ No newline at end of file
diff --git a/Engine/Source/Runtime/Experimental/ChaosVisualDebugger/Private/DataWrappers/ChaosVDParticleDataWrapper.cpp b/Engine/Source/Runtime/Experimental/ChaosVisualDebugger/Private/DataWrappers/ChaosVDParticleDataWrapper.cpp
--- a/Engine/Source/Runtime/Experimental/ChaosVisualDebugger/Private/DataWrappers/ChaosVDParticleDataWrapper.cpp
+++ b/Engine/Source/Runtime/Experimental/ChaosVisualDebugger/Private/DataWrappers/ChaosVDParticleDataWrapper.cpp
@@ -31,11 +31,14 @@ bool FChaosVDFRigidParticleControlFlags::Serialize(FArchive& Ar)
 
 	Ar.UsingCustomVersion(FFortniteMainBranchObjectVersion::GUID);
 
-	if (Ar.CustomVer(FFortniteMainBranchObjectVersion::GUID) >= FFortniteMainBranchObjectVersion::SolverIterationsDataSupportInChaosVisualDebugger)
+	if (Ar.CustomVer(FFortniteMainBranchObjectVersion::GUID) >= FFortniteMainBranchObjectVersion::SolverIterationsDataSupportInChaosVisualDebugger
+		&& Ar.CustomVer(FFortniteMainBranchObjectVersion::GUID) < FFortniteMainBranchObjectVersion::PerParticleIterationCountMovedToDynamicMisc)
 	{
-		Ar << PositionSolverIterationCount;
-		Ar << VelocitySolverIterationCount;
-		Ar << ProjectionSolverIterationCount;
+		// Three uint8 (PositionSolverIterations, VelocitySolverIterations, ProjectionSolverIterations) were moved to FChaosVDParticleDynamicMisc.
+		uint8 DummyInt;
+		Ar << DummyInt;
+		Ar << DummyInt;
+		Ar << DummyInt;
 	}
 
 	if (Ar.CustomVer(FFortniteMainBranchObjectVersion::GUID) >= FFortniteMainBranchObjectVersion::GyroscopicTorquesSupportInChaosVisualDebugger)
@@ -147,6 +150,13 @@ bool FChaosVDParticleDynamicMisc::Serialize(FArchive& Ar)
 
 	MControlFlags.Serialize(Ar);
 
+	if (Ar.CustomVer(FFortniteMainBranchObjectVersion::GUID) >= FFortniteMainBranchObjectVersion::PerParticleIterationCountMovedToDynamicMisc)
+	{
+		Ar << PositionSolverIterationCount;
+		Ar << VelocitySolverIterationCount;
+		Ar << ProjectionSolverIterationCount;
+	}
+
 	return true;
 }
 
diff --git a/Engine/Source/Runtime/Experimental/ChaosVisualDebugger/Public/DataWrappers/ChaosVDParticleDataWrapper.h b/Engine/Source/Runtime/Experimental/ChaosVisualDebugger/Public/DataWrappers/ChaosVDParticleDataWrapper.h
--- a/Engine/Source/Runtime/Experimental/ChaosVisualDebugger/Public/DataWrappers/ChaosVDParticleDataWrapper.h
+++ b/Engine/Source/Runtime/Experimental/ChaosVisualDebugger/Public/DataWrappers/ChaosVDParticleDataWrapper.h
@@ -87,10 +87,7 @@ struct FChaosVDFRigidParticleControlFlags : public FChaosVDWrapperDataBase
 		  GravityGroupIndex(0),
 		  bMACDEnabled(false), 
 		  bPartialIslandSleepAllowed(true),
-		  bGyroscopicTorqueEnabled(false),
-		  PositionSolverIterationCount(8),
-		  VelocitySolverIterationCount(1),
-		  ProjectionSolverIterationCount(1)
+		  bGyroscopicTorqueEnabled(false)
 	{
 	}
 
@@ -122,6 +119,7 @@ struct FChaosVDFRigidParticleControlFlags : public FChaosVDWrapperDataBase
 		Other.SetMACDEnabled(bMACDEnabled);
 		Other.SetPartialIslandSleepAllowed(bPartialIslandSleepAllowed);
 		Other.SetGyroscopicTorqueEnabled(bGyroscopicTorqueEnabled);
+		
 	}
 
 	UPROPERTY(VisibleAnywhere, Category= "Particle Control Flags")
@@ -140,12 +138,6 @@ struct FChaosVDFRigidParticleControlFlags : public FChaosVDWrapperDataBase
 	bool bPartialIslandSleepAllowed;
 	UPROPERTY(VisibleAnywhere, Category = "Particle Control Flags")
 	bool bGyroscopicTorqueEnabled;
-	UPROPERTY(VisibleAnywhere, Category = "Particle Control Flags")
-	uint8 PositionSolverIterationCount;
-	UPROPERTY(VisibleAnywhere, Category = "Particle Control Flags")
-	uint8 VelocitySolverIterationCount;
-	UPROPERTY(VisibleAnywhere, Category = "Particle Control Flags")
-	uint8 ProjectionSolverIterationCount;
 };
 
 CVD_IMPLEMENT_SERIALIZER(FChaosVDFRigidParticleControlFlags)
@@ -410,7 +402,7 @@ struct FChaosVDParticleDynamicMisc : public FChaosVDWrapperDataBase
 	GENERATED_BODY()
 
 	FChaosVDParticleDynamicMisc(): MLinearEtherDrag(0), MAngularEtherDrag(0), MMaxLinearSpeedSq(0), MMaxAngularSpeedSq(0), MInitialOverlapDepenetrationVelocity(0), MSleepThresholdMultiplier(1),
-	                               MCollisionGroup(0), MObjectState(), MSleepType(), bDisabled(false)
+	                               MCollisionGroup(0), MObjectState(), MSleepType(), bDisabled(false), PositionSolverIterationCount(INDEX_NONE), VelocitySolverIterationCount(INDEX_NONE), ProjectionSolverIterationCount(INDEX_NONE)
 	{
 	}
 
@@ -435,6 +427,10 @@ struct FChaosVDParticleDynamicMisc : public FChaosVDWrapperDataBase
 		bDisabled = Other.Disabled();
 
 		bHasValidData = true;
+
+		PositionSolverIterationCount = static_cast<int8>(Other.IterationSettings().GetNumPositionIterations());
+		VelocitySolverIterationCount = static_cast<int8>(Other.IterationSettings().GetNumVelocityIterations());
+		ProjectionSolverIterationCount = static_cast<int8>(Other.IterationSettings().GetNumProjectionIterations());
 	}
 	
 	template <typename OtherType, typename ControlFlagsType, typename SleepStateType>
@@ -456,6 +452,10 @@ struct FChaosVDParticleDynamicMisc : public FChaosVDWrapperDataBase
 		Other.SetControlFlags(ControlFlags);
 		
 		Other.SetDisabled(bDisabled);
+
+		Other.SetPositionSolverIterations((uint32)PositionSolverIterationCount);
+		Other.SetVelocitySolverIterations((uint32)VelocitySolverIterationCount);
+		Other.SetProjectionSolverIterations((uint32)ProjectionSolverIterationCount);
 	}
 
 	UPROPERTY(VisibleAnywhere, Category= "Particle Dynamic Misc")
@@ -493,6 +493,15 @@ struct FChaosVDParticleDynamicMisc : public FChaosVDWrapperDataBase
 
 	UPROPERTY(VisibleAnywhere, Category= "Particle Dynamic Misc")
 	bool bDisabled;
+
+	UPROPERTY(VisibleAnywhere, Category = "Particle Control Flags")
+	int8 PositionSolverIterationCount;
+
+	UPROPERTY(VisibleAnywhere, Category = "Particle Control Flags")
+	int8 VelocitySolverIterationCount;
+
+	UPROPERTY(VisibleAnywhere, Category = "Particle Control Flags")
+	int8 ProjectionSolverIterationCount;
 };
 
 CVD_IMPLEMENT_SERIALIZER(FChaosVDParticleDynamicMisc)
