Dataflow : fix issues where unregistering and registering inputs and outputs can mess with their order
- for now c all CompactStable on the maps containing the inputs/outputs
- in the future we may want to make the internal structure of the node more "order aware"

diff --git a/Engine/Source/Runtime/Experimental/Dataflow/Core/Private/Dataflow/DataflowNode.cpp b/Engine/Source/Runtime/Experimental/Dataflow/Core/Private/Dataflow/DataflowNode.cpp
--- a/Engine/Source/Runtime/Experimental/Dataflow/Core/Private/Dataflow/DataflowNode.cpp
+++ b/Engine/Source/Runtime/Experimental/Dataflow/Core/Private/Dataflow/DataflowNode.cpp
@@ -1048,6 +1048,8 @@ void FDataflowNode::UnregisterInputConnection(const UE::Dataflow::FConnectionKey
 {
 	if (ExpandedInputs.Remove(Key))
 	{
+		// make sure we compact so that the next added element will be in proper order 
+		ExpandedOutputs.CompactStable();
 		// Invalidate graph as this input might have had connections
 		Invalidate();
 	}
@@ -1057,6 +1059,8 @@ void FDataflowNode::UnregisterOutputConnection(const UE::Dataflow::FConnectionKe
 {
 	if (ExpandedOutputs.Remove(Key))
 	{
+		// make sure we compact so that the next added element will be in proper order 
+		ExpandedOutputs.CompactStable();
 		// Invalidate graph as this output might have had connections
 		Invalidate();
 	}
