Dataflow node : Fix using the wrong map

diff --git a/Engine/Source/Runtime/Experimental/Dataflow/Core/Private/Dataflow/DataflowNode.cpp b/Engine/Source/Runtime/Experimental/Dataflow/Core/Private/Dataflow/DataflowNode.cpp
--- a/Engine/Source/Runtime/Experimental/Dataflow/Core/Private/Dataflow/DataflowNode.cpp
+++ b/Engine/Source/Runtime/Experimental/Dataflow/Core/Private/Dataflow/DataflowNode.cpp
@@ -1049,7 +1049,7 @@ void FDataflowNode::UnregisterInputConnection(const UE::Dataflow::FConnectionKey
 	if (ExpandedInputs.Remove(Key))
 	{
 		// make sure we compact so that the next added element will be in proper order 
-		ExpandedOutputs.CompactStable();
+		ExpandedInputs.CompactStable();
 		// Invalidate graph as this input might have had connections
 		Invalidate();
 	}
