[NetPhysics] CVar to limit the number of overlaping frames when triggering a rewind to already resimulated frames.
- p.Resim.AllowRewindToResimulatedFrames.OverlapPercent set to 1.0 (100%) by default.
Example: A rewind goes back from frame 20 to frame 10 and resimulates. At a 100% overlap we block resims from frame 10 so the next resim can happen from frame 11. At 50% resim we would block resims from frame 15 allowing new resimulations to happen from frame 16.

diff --git a/Engine/Source/Runtime/Experimental/Chaos/Private/RewindData.cpp b/Engine/Source/Runtime/Experimental/Chaos/Private/RewindData.cpp
--- a/Engine/Source/Runtime/Experimental/Chaos/Private/RewindData.cpp
+++ b/Engine/Source/Runtime/Experimental/Chaos/Private/RewindData.cpp
@@ -411,6 +411,9 @@ bool FRewindData::CheckQuaternionThreshold(FQuat A, FQuat B, float ThresholdDegr
 CHAOS_API bool bResimAllowRewindToResimulatedFrames = true;
 FAutoConsoleVariableRef CVarResimAllowRewindToResimulatedFrames(TEXT("p.Resim.AllowRewindToResimulatedFrames"), bResimAllowRewindToResimulatedFrames, TEXT("Allow rewinding back to a frame that was previously part of a resimulation. If a resimulation is performed between frame 100-110, allow a new resim from 105-115 if needed, else next resim will be able to start from frame 111."));
 
+CHAOS_API float ResimAllowRewindToResimulatedFramesOverlapPercent = 1.0f;
+FAutoConsoleVariableRef CVarResimAllowRewindToResimulatedFramesOverlapPercent(TEXT("p.Resim.AllowRewindToResimulatedFrames.OverlapPercent"), ResimAllowRewindToResimulatedFramesOverlapPercent, TEXT("Value in percent as a multiplier, 0.3 = 30%. When p.Resim.AllowRewindToResimulatedFrames is true, this sets how many frames of overlap we allow when rewinding back past a previous resimulation. If a resimulation is performed between frame 100-110 (10 frames), 30% overlap would allow 3 frames to get resimulated again, i.e. we would be allowed to rewind back to frame 107 again if we need to."));
+
 bool FRewindData::RewindToFrame(int32 Frame)
 {
 	QUICK_SCOPE_CYCLE_COUNTER(RewindToFrame);
@@ -441,7 +444,16 @@ bool FRewindData::RewindToFrame(int32 Frame)
 	FFrameAndPhase RewindFrameAndPhase{ Frame, FFrameAndPhase::PrePushData };
 	FFrameAndPhase CurFrameAndPhase{ CurFrame, FFrameAndPhase::PrePushData };
 
-	BlockResimFrame = bResimAllowRewindToResimulatedFrames ? Frame : CurFrame;
+	if (bResimAllowRewindToResimulatedFrames)
+	{
+		float NumFrames = static_cast<float>(CurFrame - Frame);
+		NumFrames *= ResimAllowRewindToResimulatedFramesOverlapPercent;
+		BlockResimFrame = CurFrame - FMath::CeilToInt(NumFrames);
+	}
+	else
+	{
+		BlockResimFrame = CurFrame;
+	}
 
 	ResimFrame = Frame;
 	CurFrame = Frame;
