Changed log verbosity to Verbose when releasing on-demand content handles.

#rnx

diff --git a/Engine/Source/Runtime/Experimental/IoStore/OnDemand/Private/OnDemandIoStore.cpp b/Engine/Source/Runtime/Experimental/IoStore/OnDemand/Private/OnDemandIoStore.cpp
--- a/Engine/Source/Runtime/Experimental/IoStore/OnDemand/Private/OnDemandIoStore.cpp
+++ b/Engine/Source/Runtime/Experimental/IoStore/OnDemand/Private/OnDemandIoStore.cpp
@@ -1866,7 +1866,7 @@ FIoStatus FOnDemandIoStore::GetContainersAndPackagesForInstall(
 void FOnDemandIoStore::ReleaseContent(FOnDemandInternalContentHandle& ContentHandle)
 {
 	LLM_SCOPE_BYNAME(TEXT("FileSystem/OnDemandIoStore"));
-	UE_LOG(LogIoStoreOnDemand, Log, TEXT("Releasing content handle, %s"), *LexToString(ContentHandle));
+	UE_LOG(LogIoStoreOnDemand, Verbose, TEXT("Releasing content handle, %s"), *LexToString(ContentHandle));
 
 	UE::TUniqueLock Lock(ContainerMutex);
 
