Fixed sim callback objects requiring PreSimulate set in their options in order to receive any callbacks throughout integration and solve.

diff --git a/Engine/Source/Runtime/Experimental/Chaos/Private/PBDRigidsSolver.cpp b/Engine/Source/Runtime/Experimental/Chaos/Private/PBDRigidsSolver.cpp
--- a/Engine/Source/Runtime/Experimental/Chaos/Private/PBDRigidsSolver.cpp
+++ b/Engine/Source/Runtime/Experimental/Chaos/Private/PBDRigidsSolver.cpp
@@ -1850,6 +1850,15 @@ namespace Chaos
 		return true;
 	}
 
+	bool HasSimulationCallback(const ISimCallbackObject* Callback)
+	{
+		return Callback->HasOption(ESimCallbackOptions::Presimulate) ||
+			Callback->HasOption(ESimCallbackOptions::PreSolve) ||
+			Callback->HasOption(ESimCallbackOptions::PostSolve) ||
+			Callback->HasOption(ESimCallbackOptions::PreIntegrate) ||
+			Callback->HasOption(ESimCallbackOptions::PostIntegrate);
+	}
+
 	void FPBDRigidsSolver::ProcessPushedData_Internal(FPushPhysicsData& PushData)
 	{
 		QUICK_SCOPE_CYCLE_COUNTER(ChaosPushData);
@@ -1862,7 +1871,7 @@ namespace Chaos
 		SimCallbackObjects.Reserve(SimCallbackObjects.Num() + PushData.SimCallbackObjectsToAdd.Num());
 		for(ISimCallbackObject* SimCallbackObject : PushData.SimCallbackObjectsToAdd)
 		{
-			if (SimCallbackObject->HasOption(ESimCallbackOptions::Presimulate))
+			if (HasSimulationCallback(SimCallbackObject))
 			{
 				SimCallbackObjects.Add(SimCallbackObject);
 			}
