Fix for metrics related server crash

#tests Tested against Cooked Data Dedicated Server build
#rnx

diff --git a/Engine/Source/Runtime/Experimental/Iris/Core/Private/Iris/ReplicationSystem/Polling/ObjectPoller.cpp b/Engine/Source/Runtime/Experimental/Iris/Core/Private/Iris/ReplicationSystem/Polling/ObjectPoller.cpp
--- a/Engine/Source/Runtime/Experimental/Iris/Core/Private/Iris/ReplicationSystem/Polling/ObjectPoller.cpp
+++ b/Engine/Source/Runtime/Experimental/Iris/Core/Private/Iris/ReplicationSystem/Polling/ObjectPoller.cpp
@@ -314,7 +314,7 @@ void FObjectPoller::ForcePollObject(FInternalNetRefIndex ObjectIndex, FNetStatsC
 	{
 		IRIS_PROFILER_PROTOCOL_NAME_CONDITIONAL(ObjectData.Protocol->DebugName->Name, CVars::bEnableVerbosePollProfiling);
 		UE_NET_IRIS_STATS_TIMER(Timer, InNetStatsContext);
-		UE_NET_TRACE_POLL_OBJECT_SCOPE(ObjectData.RefHandle, Timer);
+		//UE_NET_TRACE_POLL_OBJECT_SCOPE(ObjectData.RefHandle, Timer); //Temporarily disable to fix a crash
 
 		const bool bIsGCAffectedObject = GarbageCollectionAffectedObjects.GetBit(ObjectIndex);
 		GarbageCollectionAffectedObjects.ClearBit(ObjectIndex);
@@ -336,7 +336,7 @@ void FObjectPoller::ForcePollObject(FInternalNetRefIndex ObjectIndex, FNetStatsC
 		else
 		{
 			UE_NET_IRIS_STATS_ADD_TIME_AND_COUNT_FOR_OBJECT_AS_WASTE(Timer, Poll, ObjectIndex);
-			UE_NET_TRACE_POLL_OBJECT_IS_WASTE();
+			//UE_NET_TRACE_POLL_OBJECT_IS_WASTE(); //Temporarily disable to fix a crash
 		}
 		++InPollStats.PolledObjectCount;
 	}
@@ -388,7 +388,7 @@ void FObjectPoller::PushModelPollObject(FInternalNetRefIndex ObjectIndex, FNetSt
 		{
 			IRIS_PROFILER_PROTOCOL_NAME_CONDITIONAL(ObjectData.Protocol->DebugName->Name, CVars::bEnableVerbosePollProfiling);
 			UE_NET_IRIS_STATS_TIMER(Timer, NetStatsContext);
-			UE_NET_TRACE_POLL_OBJECT_SCOPE(ObjectData.RefHandle, Timer);
+			//UE_NET_TRACE_POLL_OBJECT_SCOPE(ObjectData.RefHandle, Timer); //Temporarily disable to fix a crash
 
 			// We need to do a poll if object is marked as dirty
 			EReplicationFragmentPollFlags PollOptions = bUsePerPropertyDirtyTracking && !bWantsFullPoll ? EReplicationFragmentPollFlags::PollDirtyState : EReplicationFragmentPollFlags::PollAllState;
@@ -403,7 +403,7 @@ void FObjectPoller::PushModelPollObject(FInternalNetRefIndex ObjectIndex, FNetSt
 		{
 			IRIS_PROFILER_PROTOCOL_NAME_CONDITIONAL(ObjectData.Protocol->DebugName->Name, CVars::bEnableVerbosePollProfiling);
 			UE_NET_IRIS_STATS_TIMER(Timer, NetStatsContext);
-			UE_NET_TRACE_POLL_OBJECT_SCOPE(ObjectData.RefHandle, Timer);
+			//UE_NET_TRACE_POLL_OBJECT_SCOPE(ObjectData.RefHandle, Timer); //Temporarily disable to fix a crash
 
 			// If this object might have been affected by GC, only refresh cached references
 			bPollFoundDirty = FReplicationInstanceOperations::PollAndCopyObjectReferences(InstanceProtocol, EReplicationFragmentTraits::None);
@@ -420,7 +420,7 @@ void FObjectPoller::PushModelPollObject(FInternalNetRefIndex ObjectIndex, FNetSt
 	{
 		IRIS_PROFILER_PROTOCOL_NAME_CONDITIONAL(ObjectData.Protocol->DebugName->Name, CVars::bEnableVerbosePollProfiling);
 		UE_NET_IRIS_STATS_TIMER(Timer, NetStatsContext);
-		UE_NET_TRACE_POLL_OBJECT_SCOPE(ObjectData.RefHandle, Timer);
+		//UE_NET_TRACE_POLL_OBJECT_SCOPE(ObjectData.RefHandle, Timer); //Temporarily disable to fix a crash
 
 		// If the object has fragments with pushed based properties, and is not marked dirty and object is affected by GC we need to make sure that we refresh cached references for all fragments with push based properties
 		const bool bIsFullPushBasedObject = EnumHasAnyFlags(InstanceTraits, EReplicationInstanceProtocolTraits::HasFullPushBasedDirtiness);
@@ -457,7 +457,7 @@ void FObjectPoller::PushModelPollObject(FInternalNetRefIndex ObjectIndex, FNetSt
 		else
 		{
 			UE_NET_IRIS_STATS_ADD_TIME_AND_COUNT_FOR_OBJECT_AS_WASTE(Timer, Poll, ObjectIndex);
-			UE_NET_TRACE_POLL_OBJECT_IS_WASTE();
+			//UE_NET_TRACE_POLL_OBJECT_IS_WASTE(); //Temporarily disable to fix a crash
 		}
 	}
 }
