CVD : fix crash when trying to get the bounds of a null particle

diff --git a/Engine/Source/Runtime/Experimental/Chaos/Private/ChaosVisualDebugger/ChaosVDDataWrapperUtils.cpp b/Engine/Source/Runtime/Experimental/Chaos/Private/ChaosVisualDebugger/ChaosVDDataWrapperUtils.cpp
--- a/Engine/Source/Runtime/Experimental/Chaos/Private/ChaosVisualDebugger/ChaosVDDataWrapperUtils.cpp
+++ b/Engine/Source/Runtime/Experimental/Chaos/Private/ChaosVisualDebugger/ChaosVDDataWrapperUtils.cpp
@@ -27,15 +27,16 @@ namespace Chaos::VisualDebugger::Utils
 	FBox GetGeometricGTParticleBounds(const FGeometryParticle* GeometryParticle)
 	{
 		using namespace Chaos;
-		const FShapesArray& Shapes = GeometryParticle->ShapesArray();
 
 		FBox Bounds(ForceInitToZero);
-
-		for (const TUniquePtr<FPerShapeData>& Shape : Shapes)
+		if (GeometryParticle)
 		{
-			Bounds += FBox(Shape->GetWorldSpaceShapeBounds().Min(), Shape->GetWorldSpaceShapeBounds().Max());
+			const FShapesArray& Shapes = GeometryParticle->ShapesArray();
+			for (const TUniquePtr<FPerShapeData>& Shape : Shapes)
+			{
+				Bounds += FBox(Shape->GetWorldSpaceShapeBounds().Min(), Shape->GetWorldSpaceShapeBounds().Max());
+			}
 		}
-
 		return Bounds;
 	}
 
