Iris
- Use accumulated set of dirty objects to ensure that objects marked dirty when out of relevancy are properly polled

diff --git a/Engine/Source/Runtime/Experimental/Iris/Core/Private/Iris/ReplicationSystem/ObjectReplicationBridge.cpp b/Engine/Source/Runtime/Experimental/Iris/Core/Private/Iris/ReplicationSystem/ObjectReplicationBridge.cpp
--- a/Engine/Source/Runtime/Experimental/Iris/Core/Private/Iris/ReplicationSystem/ObjectReplicationBridge.cpp
+++ b/Engine/Source/Runtime/Experimental/Iris/Core/Private/Iris/ReplicationSystem/ObjectReplicationBridge.cpp
@@ -1424,10 +1424,9 @@ void UObjectReplicationBridge::BuildPollList(UE::Net::FNetBitArrayView ObjectsCo
 			FNetBitArrayView DirtyAndRelevantObjectsView = MakeNetBitArrayView(DirtyAndRelevantObjects, FNetBitArrayView::NoResetNoValidate);
 
 			{
-				FDirtyObjectsAccessor DirtyObjectsAccessor(ReplicationSystemInternal->GetDirtyNetObjectTracker());
-				const FNetBitArrayView DirtyObjectsThisFrame = DirtyObjectsAccessor.GetDirtyNetObjects();
-
-				DirtyAndRelevantObjectsView.Set(DirtyObjectsThisFrame, FNetBitArray::OrOp, ForceNetUpdateObjects);
+				// Use the accumulated set of dirty objects to ensure correctness for objects dirtied when not relevant to any connection.
+				const FNetBitArrayView AccumulatedDirtyNetObject = DirtyNetObjectTracker.GetAccumulatedDirtyNetObjects();
+				DirtyAndRelevantObjectsView.Set(AccumulatedDirtyNetObject, FNetBitArray::OrOp, ForceNetUpdateObjects);
 				DirtyAndRelevantObjectsView.Combine(RelevantObjects, FNetBitArray::AndOp);
 			}
 
