Optimize update kinematic targets
Replace parallel for in Kinematic Update by [REDACTED_PERSON]
Link the Intergration task with the Kinematics update
Perf Test with kinematics particles can get up to 18% improvement

diff --git a/Engine/Source/Runtime/Experimental/Chaos/Private/Chaos/PBDRigidsEvolution.cpp b/Engine/Source/Runtime/Experimental/Chaos/Private/Chaos/PBDRigidsEvolution.cpp
--- a/Engine/Source/Runtime/Experimental/Chaos/Private/Chaos/PBDRigidsEvolution.cpp
+++ b/Engine/Source/Runtime/Experimental/Chaos/Private/Chaos/PBDRigidsEvolution.cpp
@@ -658,7 +658,7 @@ namespace Chaos
 							NewPendingData.Add(PendingData);
 						}
 					}
-				});
+				}, LowLevelTasks::ETaskPriority::High);
 				PendingTasks.Add(PendingTask);
 			}
 
diff --git a/Engine/Source/Runtime/Experimental/Chaos/Private/Chaos/PBDRigidsEvolutionGBF.cpp b/Engine/Source/Runtime/Experimental/Chaos/Private/Chaos/PBDRigidsEvolutionGBF.cpp
--- a/Engine/Source/Runtime/Experimental/Chaos/Private/Chaos/PBDRigidsEvolutionGBF.cpp
+++ b/Engine/Source/Runtime/Experimental/Chaos/Private/Chaos/PBDRigidsEvolutionGBF.cpp
@@ -488,6 +488,7 @@ void FPBDRigidsEvolutionGBF::UpdateCollisionSolverType()
 
 void FPBDRigidsEvolutionGBF::AdvanceOneTimeStepImpl(const FReal Dt, const FSubStepInfo& SubStepInfo)
 {
+	CSV_SCOPED_TIMING_STAT(PhysicsVerbose, StepSolver_AdvanceOneTimeStepImpl);
 	SCOPE_CYCLE_COUNTER(STAT_Evolution_AdvanceOneTimeStep);
 
 	//for now we never allow solver to schedule more than two tasks back to back
@@ -524,49 +525,54 @@ void FPBDRigidsEvolutionGBF::AdvanceOneTimeStepImpl(const FReal Dt, const FSubSt
 	}
 #endif
 
-	if (PreIntegrateCallback != nullptr)
-	{
-		SCOPE_CYCLE_COUNTER(STAT_Evolution_PreIntegrateCallback);
-		PreIntegrateCallback(Dt);
-	}
-
+	
 	{
-		SCOPE_CYCLE_COUNTER(STAT_Evolution_Integrate);
-		CSV_SCOPED_TIMING_STAT(PhysicsVerbose, StepSolver_Integrate);
-		Integrate(Dt);
-	}
+		QUICK_SCOPE_CYCLE_COUNTER(STAT_Evolution_StepSolver_IntegrateKinematicSpatial);
+		CSV_SCOPED_TIMING_STAT(PhysicsVerbose, StepSolver_IntegrateKinematicSpatial);
+		if (PreIntegrateCallback != nullptr)
+		{
+			SCOPE_CYCLE_COUNTER(STAT_Evolution_PreIntegrateCallback);
+			PreIntegrateCallback(Dt);
+		}
 
-	{
-		if (CaptureRewindKinematicTarget)
 		{
-			CaptureRewindKinematicTarget(Particles.GetActiveMovingKinematicParticlesView());
+			SCOPE_CYCLE_COUNTER(STAT_Evolution_Integrate);
+			CSV_SCOPED_TIMING_STAT(PhysicsVerbose, StepSolver_Integrate);
+			Integrate(Dt);
 		}
 
-		SCOPE_CYCLE_COUNTER(STAT_Evolution_KinematicTargets);
-		CSV_SCOPED_TIMING_STAT(PhysicsVerbose, StepSolver_KinematicTargets);
-		ApplyKinematicTargets(Dt, SubStepInfo.PseudoFraction);
-	}
+		{
+			if (CaptureRewindKinematicTarget)
+			{
+				CaptureRewindKinematicTarget(Particles.GetActiveMovingKinematicParticlesView());
+			}
 
-	{
-		CVD_SCOPE_TRACE_SOLVER_STEP(CVDDC_PostIntegrate, TEXT("Post Integrate"));
-		CVD_TRACE_PARTICLES_SOA(Particles, &Clustering);
-	}
+			SCOPE_CYCLE_COUNTER(STAT_Evolution_KinematicTargets);
+			CSV_SCOPED_TIMING_STAT(PhysicsVerbose, StepSolver_KinematicTargets);
+			ApplyKinematicTargets(Dt, SubStepInfo.PseudoFraction);
+		}
 
-	if (PostIntegrateCallback != nullptr)
-	{
-		SCOPE_CYCLE_COUNTER(STAT_Evolution_PostIntegrateCallback);
-		PostIntegrateCallback(Dt);
-	}
+		{
+			CVD_SCOPE_TRACE_SOLVER_STEP(CVDDC_PostIntegrate, TEXT("Post Integrate"));
+			CVD_TRACE_PARTICLES_SOA(Particles, &Clustering);
+		}
 
-	{
-		SCOPE_CYCLE_COUNTER(STAT_Evolution_UpdateConstraintPositionBasedState);
-		CSV_SCOPED_TIMING_STAT(PhysicsVerbose, StepSolver_UpdateConstraintPositionBasedState);
-		UpdateConstraintPositionBasedState(Dt);
-	}
-	{
-		SCOPE_CYCLE_COUNTER(STAT_Evolution_ComputeIntermediateSpatialAcceleration);
-		CSV_SCOPED_TIMING_STAT(PhysicsVerbose, StepSolver_ComputeIntermediateSpatialAcceleration);
-		Base::ComputeIntermediateSpatialAcceleration();
+		if (PostIntegrateCallback != nullptr)
+		{
+			SCOPE_CYCLE_COUNTER(STAT_Evolution_PostIntegrateCallback);
+			PostIntegrateCallback(Dt);
+		}
+
+		{
+			SCOPE_CYCLE_COUNTER(STAT_Evolution_UpdateConstraintPositionBasedState);
+			CSV_SCOPED_TIMING_STAT(PhysicsVerbose, StepSolver_UpdateConstraintPositionBasedState);
+			UpdateConstraintPositionBasedState(Dt);
+		}
+		{
+			SCOPE_CYCLE_COUNTER(STAT_Evolution_ComputeIntermediateSpatialAcceleration);
+			CSV_SCOPED_TIMING_STAT(PhysicsVerbose, StepSolver_ComputeIntermediateSpatialAcceleration);
+			Base::ComputeIntermediateSpatialAcceleration();
+		}
 	}
 
 	// Collision detection is sensitive to duplication bugs in the particle views
@@ -805,9 +811,6 @@ void FPBDRigidsEvolutionGBF::AdvanceOneTimeStepImpl(const FReal Dt, const FSubSt
 
 void FPBDRigidsEvolutionGBF::Integrate(FReal Dt)
 {
-	const TArray<FPBDRigidParticleHandle*>& ParticlesArray = Particles.GetActiveParticlesArray();
-
-	//SCOPE_CYCLE_COUNTER(STAT_Integrate);
 	CHAOS_SCOPED_TIMER(Integrate);
 
 	const FReal BoundsThickness = GetCollisionConstraints().GetDetectorSettings().BoundsExpansion;
@@ -820,17 +823,15 @@ void FPBDRigidsEvolutionGBF::Integrate(FReal Dt)
 	const bool bAllowMACD = CVars::bChaosUseMACD;
 	const bool bForceMACD = CVars::bChaosForceMACD;
 
+	const TArray<FPBDRigidParticleHandle*>& ParticlesArray = Particles.GetActiveParticlesArray();
 	const int32 NumParticles = ParticlesArray.Num();
 	int32 NumTasks = !!GSingleThreadedPhysics ? 1 : FMath::Min(FTaskGraphInterface::Get().GetNumWorkerThreads(), Chaos::MaxNumWorkers);
 	NumTasks = FMath::Max(FMath::Min(NumTasks, NumParticles), 1);
-
-
 	constexpr int32 MinParticlePerTask = 40;
 	const int32 ParticleByTask = FMath::Max(FMath::DivideAndRoundUp(NumParticles, NumTasks), MinParticlePerTask);
 	int32 NumBatches = FMath::DivideAndRoundUp(NumParticles, ParticleByTask);
 
-	TArray<UE::Tasks::FTask> PendingTasks;
-	PendingTasks.Reserve(NumBatches);
+	IntegrationPendingTasks.Reset(NumBatches);
 	for (int32 BatchIndex = 0; BatchIndex < NumBatches; BatchIndex++)
 	{
 		const int32 StartIndex = BatchIndex * ParticleByTask;
@@ -841,6 +842,7 @@ void FPBDRigidsEvolutionGBF::Integrate(FReal Dt)
 		Dt, HackMaxAngularSpeedSq, HackMaxLinearSpeedSq, bAllowMACD, bForceMACD, VelocityBoundsMultiplierMACD, MaxVelocityBoundsExpansionMACD,
 		VelocityBoundsMultiplier, MaxVelocityBoundsExpansion, BoundsThickness]()
 		{
+			QUICK_SCOPE_CYCLE_COUNTER(STAT_IntegrateTask);
 			for (int32 Index = StartIndex; Index < EndIndex; Index++)
 			{
 				//question: can we enforce this at the API layer? Right now islands contain non dynamic which makes this hard
@@ -979,15 +981,18 @@ void FPBDRigidsEvolutionGBF::Integrate(FReal Dt)
 					}
 				}
 			}
-		});
-		PendingTasks.Add(PendingTask);
+		}, LowLevelTasks::ETaskPriority::High);
+		IntegrationPendingTasks.Add(PendingTask);
 	}
 
-	for (FPBDRigidParticleHandle* Particle : ParticlesArray)
+	DirtyParticleTask = UE::Tasks::Launch(UE_SOURCE_LOCATION, [this, &ParticlesArray]()
 	{
-		Base::DirtyParticle(*Particle);
-	}
-	UE::Tasks::Wait(PendingTasks);
+		QUICK_SCOPE_CYCLE_COUNTER(STAT_DirtyDynamicParticleTask);
+		for (FPBDRigidParticleHandle* Particle : ParticlesArray)
+		{
+			Base::DirtyParticle(*Particle);
+		}
+	}, LowLevelTasks::ETaskPriority::High);
 }
 
 void FPBDRigidsEvolutionGBF::ApplyKinematicTargets(const FReal Dt, const FReal StepFraction)
@@ -997,112 +1002,43 @@ void FPBDRigidsEvolutionGBF::ApplyKinematicTargets(const FReal Dt, const FReal S
 
 	const bool IsLastStep = (FMath::IsNearlyEqual(StepFraction, (FReal)1, (FReal)UE_KINDA_SMALL_NUMBER));
 
-	// NOTE: ApplyKinematicTargetForParticle is run in a parallel-for. We only write to particle state
-	const auto& ApplyParticleKinematicTarget =
-	[Dt, StepFraction, IsLastStep](FTransientPBDRigidParticleHandle& Particle, const int32 ParticleIndex) -> void
-	{
-		FKinematicTarget& KinematicTarget = Particle.KinematicTarget();
-		const FVec3 CurrentX = Particle.GetX();
-		const FRotation3f CurrentR = Particle.GetRf();
-		constexpr FReal MinDt = 1e-6f;
+	TParticleView<FPBDRigidParticles>& KienamticParticles = Particles.GetActiveMovingKinematicParticlesView();
+	using TSOA = typename TParticleView<FPBDRigidParticles>::TSOA;
+	using THandle = typename TSOA::THandleType;
+	using THandleBase = typename THandle::THandleBase;
+	using TTransientHandle = typename THandle::TTransientHandle;
 
-		bool bMoved = false;
-		switch (KinematicTarget.GetMode())
-		{
-		case EKinematicTargetMode::None:
-			// Nothing to do
-			break;
-
-		case EKinematicTargetMode::Reset:
+	for (int32 ViewIndex = 0; ViewIndex < KienamticParticles.SOAViews.Num(); ++ViewIndex)
+	{
+		const TSOAView<TSOA>& SOAView = KienamticParticles.SOAViews[ViewIndex];
+		const int32 NumParticles = SOAView.Size();
+		if (NumParticles == 0)
 		{
-			// Reset velocity and then switch to do-nothing mode
-			Particle.SetVf(FVec3f(0.0f, 0.0f, 0.0f));
-			Particle.SetWf(FVec3f(0.0f, 0.0f, 0.0f));
-			Particle.ClearIsMovingKinematic();
-			KinematicTarget.SetMode(EKinematicTargetMode::None);
-			break;
+			continue;
 		}
-
-		case EKinematicTargetMode::Position:
+		// Dispatch tasks according if it is an array of handles or a SOA
+		if (const TArray<THandle*>* CurHandlesArray = SOAView.HandlesArray)
 		{
-			// Move to kinematic target and update velocities to match
-			// Target positions only need to be processed once, and we reset the velocity next frame (if no new target is set)
-			FVec3 NewX;
-			FRotation3f NewR;
-			if (IsLastStep)
-			{
-				NewX = KinematicTarget.GetPosition();
-				NewR = KinematicTarget.GetRotation();
-				KinematicTarget.SetMode(EKinematicTargetMode::Reset);
-			}
-			else
-			{
-				// as a reminder, stepfraction is the remaing fraction of the step from the remaining steps
-				// for total of 4 steps and current step of 2, this will be 1/3 ( 1 step passed, 3 steps remains )
-				NewX = FVec3::Lerp(CurrentX, KinematicTarget.GetPosition(), StepFraction);
-				NewR = FRotation3f::Slerp(CurrentR, KinematicTarget.GetRotation(), decltype(FQuat4f::X)(StepFraction));
-			}
-
-			const bool bPositionChanged = !FVec3::IsNearlyEqual(NewX, CurrentX, UE_SMALL_NUMBER);
-			const bool bRotationChanged = !FRotation3::IsNearlyEqual(NewR, CurrentR, UE_KINDA_SMALL_NUMBER);
-			bMoved = bPositionChanged || bRotationChanged;
-			FVec3 NewV = FVec3(0);
-			FVec3f NewW = FVec3f(0);
-			if (Dt > MinDt)
-			{
-				if (bPositionChanged)
-				{
-					NewV = FVec3::CalculateVelocity(CurrentX, NewX, Dt);
-				}
-				if (bRotationChanged)
+			auto ReturnHandle = [CurHandlesArray](int32 Index) -> TParticleView<FPBDRigidParticles>::THandle::THandleBase
 				{
-					NewW = FRotation3f::CalculateAngularVelocity(CurrentR, NewR, static_cast<FRealSingle>(Dt));
-				}
-			}
-			Particle.SetX(NewX);
-			Particle.SetRf(NewR);
-			Particle.SetV(NewV);
-			Particle.SetWf(NewW);
-			Particle.SetIsMovingKinematic();
-
-			break;
+					THandle* HandlePtr = (*CurHandlesArray)[Index];
+					return THandleBase(HandlePtr->GeometryParticles, HandlePtr->ParticleIdx);
+				};
+			using LambdaType = decltype(ReturnHandle);
+			const int32 NumHandles = CurHandlesArray->Num();
+			DispatchKinematicTargetsTasks<LambdaType>(NumHandles, Dt, StepFraction, IsLastStep, ReturnHandle);
 		}
-
-		case EKinematicTargetMode::Velocity:
+		else if (TSOA* Soa = SOAView.SOA)
 		{
-			// Move based on velocity
-			bMoved = true;
-			Particle.SetX(Particle.GetX() + Particle.GetV() * Dt);
-			Particle.SetRf(FRotation3f::IntegrateRotationWithAngularVelocity(Particle.GetRf(), Particle.GetWf(), FRealSingle(Dt)));
-			Particle.SetIsMovingKinematic();
-
-			break;
-		}
-		}
-
-		// Set positions and previous velocities if we can
-		// Note: At present kinematics are in fact rigid bodies
-		Particle.SetP(Particle.GetX());
-		Particle.SetQf(Particle.GetRf());
-		Particle.SetPreVf(Particle.GetVf());
-		Particle.SetPreWf(Particle.GetWf());
-
-		if (bMoved)
-		{
-			if (!Particle.CCDEnabled())
-			{
-				Particle.UpdateWorldSpaceState(FRigidTransform3(Particle.GetP(), Particle.GetQ()), FVec3(0));
-			}
-			else
-			{
-				Particle.UpdateWorldSpaceStateSwept(FRigidTransform3(Particle.GetP(), Particle.GetQ()), FVec3(0), -Particle.GetV() * Dt);
-			}
+			auto ReturnHandle = [Soa](int32 Index) { return THandleBase(Soa, Index); };
+			using LambdaType = decltype(ReturnHandle);
+			DispatchKinematicTargetsTasks<LambdaType>(NumParticles, Dt, StepFraction, IsLastStep, ReturnHandle);
 		}
-	};
-
-	// Apply kinematic targets in parallel
-	Particles.GetActiveMovingKinematicParticlesView().ParallelFor(ApplyParticleKinematicTarget);
+	}
 
+	// To start dirtying Kinematics particles, the Dynamics particles must be finished
+	UE::Tasks::Wait(DirtyParticleTask);
+	DirtyParticleTask = UE::Tasks::FTask();
 	for(FTransientPBDRigidParticleHandle& Particle : Particles.GetActiveMovingKinematicParticlesView())
 	{
 		// Moving kinematics need to have SQ updates queued to apply substepped movement to the
@@ -1111,16 +1047,167 @@ void FPBDRigidsEvolutionGBF::ApplyKinematicTargets(const FReal Dt, const FReal S
 		DirtyParticle(Particle);
 	}
 
+	UE::Tasks::Wait(KinematicPendingTasks);
+	KinematicPendingTasks.Reset();
 	// done with update, let's clear the tracking structures
 	if (IsLastStep)
 	{
 		Particles.UpdateAllMovingKinematic();
 	}
 
+	// Updating the views can be done only when all kinematics and dynamics update are finished.
+	UE::Tasks::Wait(IntegrationPendingTasks);
+	IntegrationPendingTasks.Reset();
 	// If we changed any particle state, the views need to be refreshed
 	Particles.UpdateDirtyViews();
 }
 
+template <typename LambdaType>
+void FPBDRigidsEvolutionGBF::DispatchKinematicTargetsTasks(int32 NumParticles, FReal Dt, FReal StepFraction, bool IsLastStep, LambdaType ReturnHandle)
+{
+	using TSOA = typename TParticleView<FPBDRigidParticles>::TSOA;
+	using THandle = typename TSOA::THandleType;
+	using THandleBase = typename THandle::THandleBase;
+
+	auto KinematicTargetsLambda = [this, ReturnHandle, Dt, StepFraction, IsLastStep](int32 StartIndex, int32 EndIndex)
+	{
+		for (int32 Index = StartIndex; Index < EndIndex; Index++)
+		{
+			THandleBase Handle = ReturnHandle(Index);
+			FTransientPBDRigidParticleHandle& Particle = static_cast<FTransientPBDRigidParticleHandle&>(Handle);
+			if (!Particle.LightWeightDisabled())
+			{
+
+				FKinematicTarget& KinematicTarget = Particle.KinematicTarget();
+				const FVec3 CurrentX = Particle.GetX();
+				const FRotation3f CurrentR = Particle.GetRf();
+				constexpr FReal MinDt = 1e-6f;
+
+				bool bMoved = false;
+				switch (KinematicTarget.GetMode())
+				{
+				case EKinematicTargetMode::None:
+					// Nothing to do
+					break;
+
+				case EKinematicTargetMode::Reset:
+				{
+					// Reset velocity and then switch to do-nothing mode
+					Particle.SetVf(FVec3f(0.0f, 0.0f, 0.0f));
+					Particle.SetWf(FVec3f(0.0f, 0.0f, 0.0f));
+					Particle.ClearIsMovingKinematic();
+					KinematicTarget.SetMode(EKinematicTargetMode::None);
+					break;
+				}
+
+				case EKinematicTargetMode::Position:
+				{
+					// Move to kinematic target and update velocities to match
+					// Target positions only need to be processed once, and we reset the velocity next frame (if no new target is set)
+					FVec3 NewX;
+					FRotation3f NewR;
+					if (IsLastStep)
+					{
+						NewX = KinematicTarget.GetPosition();
+						NewR = KinematicTarget.GetRotation();
+						KinematicTarget.SetMode(EKinematicTargetMode::Reset);
+					}
+					else
+					{
+						// as a reminder, stepfraction is the remaing fraction of the step from the remaining steps
+						// for total of 4 steps and current step of 2, this will be 1/3 ( 1 step passed, 3 steps remains )
+						NewX = FVec3::Lerp(CurrentX, KinematicTarget.GetPosition(), StepFraction);
+						NewR = FRotation3f::Slerp(CurrentR, KinematicTarget.GetRotation(), decltype(FQuat4f::X)(StepFraction));
+					}
+
+					const bool bPositionChanged = !FVec3::IsNearlyEqual(NewX, CurrentX, UE_SMALL_NUMBER);
+					const bool bRotationChanged = !FRotation3::IsNearlyEqual(NewR, CurrentR, UE_KINDA_SMALL_NUMBER);
+					bMoved = bPositionChanged || bRotationChanged;
+					FVec3 NewV = FVec3(0);
+					FVec3f NewW = FVec3f(0);
+					if (Dt > MinDt)
+					{
+						if (bPositionChanged)
+						{
+							NewV = FVec3::CalculateVelocity(CurrentX, NewX, Dt);
+						}
+						if (bRotationChanged)
+						{
+							NewW = FRotation3f::CalculateAngularVelocity(CurrentR, NewR, static_cast<FRealSingle>(Dt));
+						}
+					}
+					Particle.SetX(NewX);
+					Particle.SetRf(NewR);
+					Particle.SetV(NewV);
+					Particle.SetWf(NewW);
+					Particle.SetIsMovingKinematic();
+
+					break;
+				}
+
+				case EKinematicTargetMode::Velocity:
+				{
+					// Move based on velocity
+					bMoved = true;
+					Particle.SetX(Particle.GetX() + Particle.GetV() * Dt);
+					Particle.SetRf(FRotation3f::IntegrateRotationWithAngularVelocity(Particle.GetRf(), Particle.GetWf(), FRealSingle(Dt)));
+					Particle.SetIsMovingKinematic();
+
+					break;
+				}
+				}
+
+				// Set positions and previous velocities if we can
+				// Note: At present kinematics are in fact rigid bodies
+				Particle.SetP(Particle.GetX());
+				Particle.SetQf(Particle.GetRf());
+				Particle.SetPreVf(Particle.GetVf());
+				Particle.SetPreWf(Particle.GetWf());
+
+				if (bMoved)
+				{
+					if (!Particle.CCDEnabled())
+					{
+						Particle.UpdateWorldSpaceState(FRigidTransform3(Particle.GetP(), Particle.GetQ()), FVec3(0));
+					}
+					else
+					{
+						Particle.UpdateWorldSpaceStateSwept(FRigidTransform3(Particle.GetP(), Particle.GetQ()), FVec3(0), -Particle.GetV() * Dt);
+					}
+				}
+			}
+		}
+	};
+
+	int32 NumTasks = !!GSingleThreadedPhysics ? 1 : FMath::Min(FTaskGraphInterface::Get().GetNumWorkerThreads(), Chaos::MaxNumWorkers);
+	NumTasks = FMath::Max(FMath::Min(NumTasks, NumParticles), 1);
+	constexpr int32 MinParticlePerTask = 40;
+	const int32 ParticleByTask = FMath::Max(FMath::DivideAndRoundUp(NumParticles, NumTasks), MinParticlePerTask);
+	int32 NumBatches = FMath::DivideAndRoundUp(NumParticles, ParticleByTask);
+
+	if (NumBatches > 1)
+	{
+		// We only write to particle state
+		for (int32 BatchIndex = 0; BatchIndex < NumBatches; BatchIndex++)
+		{
+			const int32 StartIndex = BatchIndex * ParticleByTask;
+			int32 EndIndex = (BatchIndex + 1) * ParticleByTask;
+			EndIndex = FMath::Min(NumParticles, EndIndex);
+
+			UE::Tasks::FTask PendingTask = UE::Tasks::Launch(UE_SOURCE_LOCATION, [KinematicTargetsLambda, StartIndex, EndIndex]()
+			{
+				QUICK_SCOPE_CYCLE_COUNTER(STAT_KinematicTargetsTasks);
+				KinematicTargetsLambda(StartIndex, EndIndex);
+			}, LowLevelTasks::ETaskPriority::High);
+			KinematicPendingTasks.Add(PendingTask);
+		}
+	}
+	else
+	{
+		KinematicTargetsLambda(0, NumParticles);
+	}
+}
+
 void FPBDRigidsEvolutionGBF::ApplyMidPhaseModifier(FReal Dt)
 {
 	if (MidPhaseModifiers != nullptr)
diff --git a/Engine/Source/Runtime/Experimental/Chaos/Public/Chaos/PBDRigidsEvolutionGBF.h b/Engine/Source/Runtime/Experimental/Chaos/Public/Chaos/PBDRigidsEvolutionGBF.h
--- a/Engine/Source/Runtime/Experimental/Chaos/Public/Chaos/PBDRigidsEvolutionGBF.h
+++ b/Engine/Source/Runtime/Experimental/Chaos/Public/Chaos/PBDRigidsEvolutionGBF.h
@@ -295,6 +295,9 @@ namespace Chaos
 
 		void UpdateCollisionSolverType();
 
+		template <typename LambdaType>
+		void DispatchKinematicTargetsTasks(int32 NumParticles, FReal Dt, FReal StepFraction, bool IsLastStep, LambdaType ReturnHandle);
+
 		FRigidClustering Clustering;
 
 		FPBDJointCombinedConstraints JointConstraints;
@@ -321,6 +324,10 @@ namespace Chaos
 
 		FCCDManager CCDManager;
 
+		TArray<UE::Tasks::FTask> IntegrationPendingTasks;
+		TArray<UE::Tasks::FTask> KinematicPendingTasks;
+		UE::Tasks::FTask DirtyParticleTask;
+
 		bool bIsDeterministic;
 
 #if CHAOS_EVOLUTION_COLLISION_TESTMODE
