Various fixes for determinism of archetype creation, this allows us to identify when an archetype is changed rather than seeing its identity change every time the engine boots

#rnx

diff --git a/Engine/Plugins/Runtime/StateTree/Source/StateTreeEditorModule/Private/StateTreeEditorData.cpp b/Engine/Plugins/Runtime/StateTree/Source/StateTreeEditorModule/Private/StateTreeEditorData.cpp
--- a/Engine/Plugins/Runtime/StateTree/Source/StateTreeEditorModule/Private/StateTreeEditorData.cpp
+++ b/Engine/Plugins/Runtime/StateTree/Source/StateTreeEditorModule/Private/StateTreeEditorData.cpp
@@ -84,7 +84,10 @@ void UStateTreeEditorData::PostInitProperties()
 {
 	Super::PostInitProperties();
 
-	RootParametersGuid = FGuid::NewGuid();
+	if(!IsTemplate())
+	{
+		RootParametersGuid = FGuid::NewGuid();
+	}
 
 #if WITH_EDITOR
 	if (!HasAnyFlags(RF_ClassDefaultObject))
diff --git a/Engine/Plugins/Runtime/StateTree/Source/StateTreeEditorModule/Private/StateTreeState.cpp b/Engine/Plugins/Runtime/StateTree/Source/StateTreeEditorModule/Private/StateTreeState.cpp
--- a/Engine/Plugins/Runtime/StateTree/Source/StateTreeEditorModule/Private/StateTreeState.cpp
+++ b/Engine/Plugins/Runtime/StateTree/Source/StateTreeEditorModule/Private/StateTreeState.cpp
@@ -68,9 +68,12 @@ PRAGMA_ENABLE_DEPRECATION_WARNINGS
 
 UStateTreeState::UStateTreeState(const FObjectInitializer& ObjectInitializer)
 	: Super(ObjectInitializer)
-	, ID(FGuid::NewGuid())
 {
-	Parameters.ID = FGuid::NewGuid();
+	if(!IsTemplate())
+	{
+		ID = FGuid::NewGuid();
+		Parameters.ID = FGuid::NewGuid();
+	}
 }
 
 UStateTreeState::~UStateTreeState()
diff --git a/Engine/Source/Runtime/Engine/Private/Components/SkyAtmosphereComponent.cpp b/Engine/Source/Runtime/Engine/Private/Components/SkyAtmosphereComponent.cpp
--- a/Engine/Source/Runtime/Engine/Private/Components/SkyAtmosphereComponent.cpp
+++ b/Engine/Source/Runtime/Engine/Private/Components/SkyAtmosphereComponent.cpp
@@ -138,7 +138,10 @@ USkyAtmosphereComponent::USkyAtmosphereComponent(const FObjectInitializer& Objec
 
 	memset(OverrideAtmosphericLight, 0, sizeof(OverrideAtmosphericLight));
 
-	ValidateStaticLightingGUIDs();
+	if(!IsTemplate())
+	{
+		ValidateStaticLightingGUIDs();
+	}
 }
 
 USkyAtmosphereComponent::~USkyAtmosphereComponent()
diff --git a/Engine/Source/Runtime/Engine/Private/DialogueVoice.cpp b/Engine/Source/Runtime/Engine/Private/DialogueVoice.cpp
--- a/Engine/Source/Runtime/Engine/Private/DialogueVoice.cpp
+++ b/Engine/Source/Runtime/Engine/Private/DialogueVoice.cpp
@@ -9,7 +9,7 @@
 
 UDialogueVoice::UDialogueVoice(const FObjectInitializer& ObjectInitializer)
 	: Super(ObjectInitializer)
-	, LocalizationGUID( FGuid::NewGuid() )
+	, LocalizationGUID( IsTemplate() ? FGuid() : FGuid::NewGuid() )
 {
 }
 
diff --git a/Engine/Source/Runtime/Engine/Private/DialogueWave.cpp b/Engine/Source/Runtime/Engine/Private/DialogueWave.cpp
--- a/Engine/Source/Runtime/Engine/Private/DialogueWave.cpp
+++ b/Engine/Source/Runtime/Engine/Private/DialogueWave.cpp
@@ -538,7 +538,7 @@ UCurveTable* UDialogueSoundWaveProxy::GetCurveData() const
 
 UDialogueWave::UDialogueWave(const FObjectInitializer& ObjectInitializer)
 	: Super(ObjectInitializer)
-	, LocalizationGUID(FGuid::NewGuid())
+	, LocalizationGUID(IsTemplate() ? FGuid() : FGuid::NewGuid())
 {
 #if WITH_EDITORONLY_DATA
 	{ static const FAutoRegisterLocalizationDataGatheringCallback AutomaticRegistrationOfLocalizationGatherer(UDialogueWave::StaticClass(), &GatherDialogueWaveForLocalization); }
diff --git a/Engine/Source/Runtime/Engine/Private/LightWeightInstanceStaticMeshManager.cpp b/Engine/Source/Runtime/Engine/Private/LightWeightInstanceStaticMeshManager.cpp
--- a/Engine/Source/Runtime/Engine/Private/LightWeightInstanceStaticMeshManager.cpp
+++ b/Engine/Source/Runtime/Engine/Private/LightWeightInstanceStaticMeshManager.cpp
@@ -296,9 +296,12 @@ void ALightWeightInstanceStaticMeshManager::SetInstancedStaticMeshParams()
 	// Allows per-instance selection in the editor
 	ISMComponent->bHasPerInstanceHitProxies = true;
 
-	while(ISMComponent->InstancingRandomSeed == 0)
+	if(!IsTemplate())
 	{
-		ISMComponent->InstancingRandomSeed = FMath::Rand();
+		while(ISMComponent->InstancingRandomSeed == 0)
+		{
+			ISMComponent->InstancingRandomSeed = FMath::Rand();
+		}
 	}
 } 
 
diff --git a/Engine/Source/Runtime/Engine/Private/Materials/MaterialExpressions.cpp b/Engine/Source/Runtime/Engine/Private/Materials/MaterialExpressions.cpp
--- a/Engine/Source/Runtime/Engine/Private/Materials/MaterialExpressions.cpp
+++ b/Engine/Source/Runtime/Engine/Private/Materials/MaterialExpressions.cpp
@@ -937,9 +937,12 @@ void UMaterialExpression::PostInitProperties()
 {
 	Super::PostInitProperties();
 
-	UpdateParameterGuid(false, false);
+	if (!IsTemplate())
+	{
+		UpdateParameterGuid(false, false);
 	
-	UpdateMaterialExpressionGuid(false, false);
+		UpdateMaterialExpressionGuid(false, false);
+	}
 }
 
 void UMaterialExpression::PostLoad()
@@ -13429,7 +13432,10 @@ void UMaterialFunctionInterface::PostInitProperties()
 	Super::PostInitProperties();
 
 	// Initialize StateId to something unique, in case this is a new function
-	StateId = FGuid::NewGuid();
+	if (!IsTemplate())
+	{
+		StateId = FGuid::NewGuid();
+	}
 }
 
 void UMaterialFunctionInterface::PostDuplicate(bool bDuplicateForPIE)
diff --git a/Engine/Source/Runtime/Experimental/GeometryCollectionEngine/Private/GeometryCollection/GeometryCollectionObject.cpp b/Engine/Source/Runtime/Experimental/GeometryCollectionEngine/Private/GeometryCollection/GeometryCollectionObject.cpp
--- a/Engine/Source/Runtime/Experimental/GeometryCollectionEngine/Private/GeometryCollection/GeometryCollectionObject.cpp
+++ b/Engine/Source/Runtime/Experimental/GeometryCollectionEngine/Private/GeometryCollection/GeometryCollectionObject.cpp
@@ -142,7 +142,10 @@ UGeometryCollection::UGeometryCollection(const FObjectInitializer& ObjectInitial
 	, DataflowInstance(this)
 	, GeometryCollection(new FGeometryCollection())
 {
-	PersistentGuid = FGuid::NewGuid();
+	if(!IsTemplate())
+	{
+		PersistentGuid = FGuid::NewGuid();
+	}
 	InvalidateCollection();
 #if WITH_EDITOR
 	SimulationDataGuid = StateGuid;
@@ -1706,7 +1709,10 @@ void UGeometryCollection::ReleaseResources()
 
 void UGeometryCollection::InvalidateCollection()
 {
-	StateGuid = FGuid::NewGuid();
+	if(!IsTemplate())
+	{
+		StateGuid = FGuid::NewGuid();
+	}
 	UpdateRootIndex();
 	CacheBreadthFirstTransformIndices();
 }
diff --git a/Engine/Source/Runtime/Foliage/Private/InstancedFoliage.cpp b/Engine/Source/Runtime/Foliage/Private/InstancedFoliage.cpp
--- a/Engine/Source/Runtime/Foliage/Private/InstancedFoliage.cpp
+++ b/Engine/Source/Runtime/Foliage/Private/InstancedFoliage.cpp
@@ -657,7 +657,10 @@ UFoliageType::UFoliageType(const FObjectInitializer& ObjectInitializer)
 	Curve->AddKey(0.f, 0.f);
 	Curve->AddKey(1.f, 1.f);
 
-	UpdateGuid = FGuid::NewGuid();
+	if(!IsTemplate())
+	{
+		UpdateGuid = FGuid::NewGuid();
+	}
 #if WITH_EDITORONLY_DATA
 	HiddenEditorViews = 0;
 #endif
diff --git a/Engine/Source/Runtime/Foliage/Private/ProceduralFoliageComponent.cpp b/Engine/Source/Runtime/Foliage/Private/ProceduralFoliageComponent.cpp
--- a/Engine/Source/Runtime/Foliage/Private/ProceduralFoliageComponent.cpp
+++ b/Engine/Source/Runtime/Foliage/Private/ProceduralFoliageComponent.cpp
@@ -27,7 +27,10 @@ UProceduralFoliageComponent::UProceduralFoliageComponent(const FObjectInitialize
 : Super(ObjectInitializer)
 {
 	TileOverlap = 0.f;
-	ProceduralGuid = FGuid::NewGuid();
+	if(!IsTemplate())
+	{
+		ProceduralGuid = FGuid::NewGuid();
+	}
 #if WITH_EDITORONLY_DATA
 	bAllowLandscape = true;
 	bAllowBSP = true;
diff --git a/Engine/Source/Runtime/MediaAssets/Private/Assets/MediaPlayer.cpp b/Engine/Source/Runtime/MediaAssets/Private/Assets/MediaPlayer.cpp
--- a/Engine/Source/Runtime/MediaAssets/Private/Assets/MediaPlayer.cpp
+++ b/Engine/Source/Runtime/MediaAssets/Private/Assets/MediaPlayer.cpp
@@ -57,7 +57,7 @@ UMediaPlayer::UMediaPlayer(const FObjectInitializer& ObjectInitializer)
 	, HorizontalFieldOfView(90.0f)
 	, VerticalFieldOfView(60.0f)
 	, ViewRotation(FRotator::ZeroRotator)
-	, PlayerGuid(FGuid::NewGuid())
+	, PlayerGuid(IsTemplate() ? FGuid() : FGuid::NewGuid())
 	, PlayOnNext(false)
 	, RegisteredWithMediaModule(false)
 #if WITH_EDITORONLY_DATA
