Iris: Fix UpdateInstancesWorldLocation() ensure failure on destroyed actor

[RN] Iris: Fixed ensure() failure in UpdateInstancesWorldLocation() handling a destroyed object

diff --git a/Engine/Plugins/Runtime/ReplicationSystemTestPlugin/Source/Private/Tests/ReplicationSystem/ReplicatedTestObjectFactory.cpp b/Engine/Plugins/Runtime/ReplicationSystemTestPlugin/Source/Private/Tests/ReplicationSystem/ReplicatedTestObjectFactory.cpp
--- a/Engine/Plugins/Runtime/ReplicationSystemTestPlugin/Source/Private/Tests/ReplicationSystem/ReplicatedTestObjectFactory.cpp
+++ b/Engine/Plugins/Runtime/ReplicationSystemTestPlugin/Source/Private/Tests/ReplicationSystem/ReplicatedTestObjectFactory.cpp
@@ -217,12 +217,16 @@ void UReplicatedTestObjectFactory::DestroyReplicatedObject(const FDestroyedConte
 	}
 }
 
-void UReplicatedTestObjectFactory::GetWorldInfo(const FWorldInfoContext& Context, FWorldInfoData& OutData)
+TOptional<UNetObjectFactory::FWorldInfoData> UReplicatedTestObjectFactory::GetWorldInfo(const FWorldInfoContext& Context) const
 {
 	if (GetInstanceWorldObjectInfoFunction)
 	{
+		FWorldInfoData OutData;
 		GetInstanceWorldObjectInfoFunction(Context.Handle, Context.Instance, OutData.WorldLocation, OutData.CullDistance);
+		return OutData;
 	}
+
+	return NullOpt;
 }
 
 void UReplicatedTestObjectFactory::SubObjectCreatedFromReplication(UE::Net::FNetRefHandle RootObject, UE::Net::FNetRefHandle SubObjectCreated)
diff --git a/Engine/Plugins/Runtime/ReplicationSystemTestPlugin/Source/Private/Tests/ReplicationSystem/ReplicatedTestObjectFactory.h b/Engine/Plugins/Runtime/ReplicationSystemTestPlugin/Source/Private/Tests/ReplicationSystem/ReplicatedTestObjectFactory.h
--- a/Engine/Plugins/Runtime/ReplicationSystemTestPlugin/Source/Private/Tests/ReplicationSystem/ReplicatedTestObjectFactory.h
+++ b/Engine/Plugins/Runtime/ReplicationSystemTestPlugin/Source/Private/Tests/ReplicationSystem/ReplicatedTestObjectFactory.h
@@ -42,7 +42,7 @@ public:
 	// UNetObjectFactory interface
 	virtual void PostInit(const FPostInitContext& Context) override;
 	virtual void DestroyReplicatedObject(const FDestroyedContext& Context) override;
-	virtual void GetWorldInfo(const FWorldInfoContext& Context, FWorldInfoData& OutData) override;
+	virtual TOptional<FWorldInfoData> GetWorldInfo(const FWorldInfoContext& Context) const override;
 	virtual void SubObjectCreatedFromReplication(UE::Net::FNetRefHandle RootObject, UE::Net::FNetRefHandle SubObjectCreated) override;
 	virtual void SubObjectDestroyedFromReplication(const FDestroyedContext& Context) override;
 
diff --git a/Engine/Source/Runtime/Engine/Private/Net/Iris/ReplicationSystem/NetActorFactory.cpp b/Engine/Source/Runtime/Engine/Private/Net/Iris/ReplicationSystem/NetActorFactory.cpp
--- a/Engine/Source/Runtime/Engine/Private/Net/Iris/ReplicationSystem/NetActorFactory.cpp
+++ b/Engine/Source/Runtime/Engine/Private/Net/Iris/ReplicationSystem/NetActorFactory.cpp
@@ -537,16 +537,24 @@ void UNetActorFactory::DestroyReplicatedObject(const FDestroyedContext& Context)
 	}	
 }
 
-void UNetActorFactory::GetWorldInfo(const FWorldInfoContext& Context, FWorldInfoData& OutData)
+TOptional<UNetActorFactory::FWorldInfoData> UNetActorFactory::GetWorldInfo(const FWorldInfoContext& Context) const
 {
 	using namespace UE::Net;
 
 	const AActor* Actor = Cast<AActor>(Context.Instance);
 	if (!Actor)
 	{
-		ensureMsgf(Actor, TEXT("UNetActorFactory::GetWorldInfo received invalid replicated instance: %s | NetRefHandle: %s"), *GetNameSafe(Context.Instance), *Context.Handle.ToString());
-		return;
+		// This is typical when an actor has been destroyed locally but not replicated to all clients yet.
+		UE_LOG(
+			LogIris,
+			VeryVerbose,
+			TEXT("UNetActorFactory::GetWorldInfo requested for non-existent actor (likely destroyed): %s | NetRefHandle: %s"),
+			*GetNameSafe(Context.Instance),
+			*Context.Handle.ToString());
+		return NullOpt;
 	}
+
+	FWorldInfoData OutData;
 	
 	if (EnumHasAnyFlags(Context.InfoRequested, EWorldInfoRequested::Location))
 	{
@@ -556,6 +564,8 @@ void UNetActorFactory::GetWorldInfo(const FWorldInfoContext& Context, FWorldInfo
 	{
 		OutData.CullDistance = Actor->GetNetCullDistanceSquared() > 0.0f ? FMath::Sqrt(Actor->GetNetCullDistanceSquared()) : 0.0f;
 	}
+
+	return OutData;
 }
 
 float UNetActorFactory::GetPollFrequency(UE::Net::FNetRefHandle RootObjectHandle, UObject* RootObjectInstance)
diff --git a/Engine/Source/Runtime/Engine/Private/Net/Iris/ReplicationSystem/NetSubObjectFactory.cpp b/Engine/Source/Runtime/Engine/Private/Net/Iris/ReplicationSystem/NetSubObjectFactory.cpp
--- a/Engine/Source/Runtime/Engine/Private/Net/Iris/ReplicationSystem/NetSubObjectFactory.cpp
+++ b/Engine/Source/Runtime/Engine/Private/Net/Iris/ReplicationSystem/NetSubObjectFactory.cpp
@@ -276,9 +276,10 @@ void UNetSubObjectFactory::DestroyReplicatedObject(const FDestroyedContext& Cont
 	Context.DestroyedInstance->MarkAsGarbage();
 }
 
-void UNetSubObjectFactory::GetWorldInfo(const FWorldInfoContext& Context, FWorldInfoData& OutData)
+TOptional<UNetObjectFactory::FWorldInfoData> UNetSubObjectFactory::GetWorldInfo(const FWorldInfoContext& Context) const
 {
 	ensureMsgf(false, TEXT("UNetSubObjectFactory::GetWorldInfo called but subobjects should never support this. Instance: %s, NetRefHandle: %s"), *GetNameSafe(Context.Instance), *Bridge->PrintObjectFromNetRefHandle(Context.Handle));
+	return NullOpt;
 }
 
 float UNetSubObjectFactory::GetPollFrequency(UE::Net::FNetRefHandle RootObjectHandle, UObject* RootObjectInstance)
diff --git a/Engine/Source/Runtime/Engine/Public/Net/Iris/ReplicationSystem/NetActorFactory.h b/Engine/Source/Runtime/Engine/Public/Net/Iris/ReplicationSystem/NetActorFactory.h
--- a/Engine/Source/Runtime/Engine/Public/Net/Iris/ReplicationSystem/NetActorFactory.h
+++ b/Engine/Source/Runtime/Engine/Public/Net/Iris/ReplicationSystem/NetActorFactory.h
@@ -128,7 +128,7 @@ public:
 
 	virtual void DestroyReplicatedObject(const FDestroyedContext& Context) override;
 
-	virtual void GetWorldInfo(const FWorldInfoContext& Context, FWorldInfoData& OutData) override;
+	virtual TOptional<FWorldInfoData> GetWorldInfo(const FWorldInfoContext& Context) const override;
 
 	virtual float GetPollFrequency(UE::Net::FNetRefHandle RootObjectHandle, UObject* RootObjectInstance) override;
 
diff --git a/Engine/Source/Runtime/Engine/Public/Net/Iris/ReplicationSystem/NetSubObjectFactory.h b/Engine/Source/Runtime/Engine/Public/Net/Iris/ReplicationSystem/NetSubObjectFactory.h
--- a/Engine/Source/Runtime/Engine/Public/Net/Iris/ReplicationSystem/NetSubObjectFactory.h
+++ b/Engine/Source/Runtime/Engine/Public/Net/Iris/ReplicationSystem/NetSubObjectFactory.h
@@ -33,7 +33,7 @@ public:
 
 	virtual void DestroyReplicatedObject(const FDestroyedContext& Context) override;
 
-	virtual void GetWorldInfo(const FWorldInfoContext& Context, FWorldInfoData& OutData) override;
+	virtual TOptional<FWorldInfoData> GetWorldInfo(const FWorldInfoContext& Context) const override;
 
 	virtual float GetPollFrequency(UE::Net::FNetRefHandle RootObjectHandle, UObject* RootObjectInstance) override;
 
diff --git a/Engine/Source/Runtime/Experimental/Iris/Core/Private/Iris/ReplicationSystem/Filtering/NetObjectGridFilter.cpp b/Engine/Source/Runtime/Experimental/Iris/Core/Private/Iris/ReplicationSystem/Filtering/NetObjectGridFilter.cpp
--- a/Engine/Source/Runtime/Experimental/Iris/Core/Private/Iris/ReplicationSystem/Filtering/NetObjectGridFilter.cpp
+++ b/Engine/Source/Runtime/Experimental/Iris/Core/Private/Iris/ReplicationSystem/Filtering/NetObjectGridFilter.cpp
@@ -602,10 +602,13 @@ void UNetObjectGridWorldLocFilter::UpdateObjectInfo(UNetObjectGridFilter::FPerOb
 {
 	check(ObjectLocationInfo.IsUsingWorldLocations());
 
-	const UE::Net::FWorldLocations::FWorldInfo& Info = WorldLocations->GetWorldInfo(PerObjectInfo.ObjectIndex);
+	const TOptional<UE::Net::FWorldLocations::FWorldInfo> Info = WorldLocations->GetWorldInfo(PerObjectInfo.ObjectIndex);
 	
-	PerObjectInfo.Position = Info.WorldLocation;
-	PerObjectInfo.SetCullDistance(Info.CullDistance);
+	if (Info)
+	{
+		PerObjectInfo.Position = Info->WorldLocation;
+		PerObjectInfo.SetCullDistance(Info->CullDistance);
+	}
 }
 
 bool UNetObjectGridWorldLocFilter::BuildObjectInfo(uint32 ObjectIndex, FNetObjectFilterAddObjectParams& Params)
diff --git a/Engine/Source/Runtime/Experimental/Iris/Core/Private/Iris/ReplicationSystem/ObjectReplicationBridge.cpp b/Engine/Source/Runtime/Experimental/Iris/Core/Private/Iris/ReplicationSystem/ObjectReplicationBridge.cpp
--- a/Engine/Source/Runtime/Experimental/Iris/Core/Private/Iris/ReplicationSystem/ObjectReplicationBridge.cpp
+++ b/Engine/Source/Runtime/Experimental/Iris/Core/Private/Iris/ReplicationSystem/ObjectReplicationBridge.cpp
@@ -1618,6 +1618,7 @@ void UObjectReplicationBridge::ForceUpdateWorldLocation(FNetRefHandle NetRefHand
 void UObjectReplicationBridge::UpdateRootObjectWorldInfo(const FUpdateWorldInfoContext& UpdateContext)
 {
 	using namespace UE::Net::Private;
+	using FWorldInfoData = UNetObjectFactory::FWorldInfoData;
 
 	const FNetRefHandleManager::FReplicatedObjectData& ObjectData = NetRefHandleManager->GetReplicatedObjectDataNoCheck(UpdateContext.ObjectIndex);
 
@@ -1631,8 +1632,14 @@ void UObjectReplicationBridge::UpdateRootObjectWorldInfo(const FUpdateWorldInfoC
 		.InfoRequested = UNetObjectFactory::EWorldInfoRequested::All,
 	};
 
-	UNetObjectFactory::FWorldInfoData WorldInfo;
-	NetFactory->GetWorldInfo(WorldInfoContext, WorldInfo);
+	TOptional<FWorldInfoData> OptWorldInfo = NetFactory->GetWorldInfo(WorldInfoContext);
+	if (!OptWorldInfo)
+	{
+		// OK: Locally-destroyed-but-not-replicated objects no longer have location updates.
+		return;
+	}
+
+	const FWorldInfoData& WorldInfo = OptWorldInfo.GetValue();
 
 	// Check if we need to warn about invalid world locations.
 	// Note that FWorldLocations will store the invalid location but only return a valid clamped position from it's API
diff --git a/Engine/Source/Runtime/Experimental/Iris/Core/Private/Iris/ReplicationSystem/ObjectReplicationBridgeDebugging.cpp b/Engine/Source/Runtime/Experimental/Iris/Core/Private/Iris/ReplicationSystem/ObjectReplicationBridgeDebugging.cpp
--- a/Engine/Source/Runtime/Experimental/Iris/Core/Private/Iris/ReplicationSystem/ObjectReplicationBridgeDebugging.cpp
+++ b/Engine/Source/Runtime/Experimental/Iris/Core/Private/Iris/ReplicationSystem/ObjectReplicationBridgeDebugging.cpp
@@ -991,11 +991,12 @@ void UObjectReplicationBridge::PrintNetCullDistances(const TArray<FString>& Args
 					.InfoRequested = UNetObjectFactory::EWorldInfoRequested::CullDistance
 				};
 
-				UNetObjectFactory::FWorldInfoData WorldData;
-				NetFactory->GetWorldInfo(Context, WorldData);
-
 				Info.Class = RepObjClass;
-				Info.CDOCullDistance = WorldData.CullDistance;
+
+				if (TOptional<UNetObjectFactory::FWorldInfoData> OptWorldData = NetFactory->GetWorldInfo(Context))
+				{
+					Info.CDOCullDistance = OptWorldData->CullDistance;
+				}
 			}
 
 			Info.NumTotal++;
@@ -1277,10 +1278,9 @@ void UObjectReplicationBridge::PrintDebugInfoForNetRefHandle(const FNetRefHandle
 	if (GetReplicationSystem()->IsServer())
 	{
 		// World location
-		if (WorldLocations.HasInfoForObject(ObjectIndex))
+		if (const TOptional<FWorldLocations::FWorldInfo> WorldInfo = WorldLocations.GetWorldInfo(ObjectIndex))
 		{
-			const FWorldLocations::FWorldInfo WorldInfo = WorldLocations.GetWorldInfo(ObjectIndex);
-			UE_LOG(LogIris, Display, TEXT("\tWorldLocation: %s | CullDistance: %f"), *WorldInfo.WorldLocation.ToCompactString(), WorldInfo.CullDistance);
+			UE_LOG(LogIris, Display, TEXT("\tWorldLocation: %s | CullDistance: %f"), *WorldInfo->WorldLocation.ToCompactString(), WorldInfo->CullDistance);
 		}
 		else
 		{
diff --git a/Engine/Source/Runtime/Experimental/Iris/Core/Public/Iris/ReplicationSystem/NetObjectFactory.h b/Engine/Source/Runtime/Experimental/Iris/Core/Public/Iris/ReplicationSystem/NetObjectFactory.h
--- a/Engine/Source/Runtime/Experimental/Iris/Core/Public/Iris/ReplicationSystem/NetObjectFactory.h
+++ b/Engine/Source/Runtime/Experimental/Iris/Core/Public/Iris/ReplicationSystem/NetObjectFactory.h
@@ -9,6 +9,7 @@
 #include "Iris/ReplicationSystem/NetRefHandle.h"
 
 #include "Misc/EnumClassFlags.h"
+#include "Misc/Optional.h"
 
 #include "NetObjectFactory.generated.h"
 
@@ -166,9 +167,9 @@ public:
 	 * Fetch world information about a replicated object so it can be updated in the network engine.
 	 * Only called for root objects.
 	 * @param Context Gives access to the replicated object and which specific world information needs to be updated.
-	 * @param OutData Set in this the requested information.
+	 * @return The object's world info, or NullOpt if there is none.
 	 */
-	virtual void GetWorldInfo(const FWorldInfoContext& Context, FWorldInfoData& OutData) PURE_VIRTUAL(UNetObjectFactory::GetWorldInfo,);
+	virtual TOptional<FWorldInfoData> GetWorldInfo(const FWorldInfoContext& Context) const PURE_VIRTUAL(UNetObjectFactory::GetWorldInfo, return NullOpt;);
 
 	/** Return the poll frequency of a root object managed by this factory. */
 	virtual float GetPollFrequency(UE::Net::FNetRefHandle RootObjectHandle, UObject* RootObjectInstance) PURE_VIRTUAL(UNetObjectFactory::GetPollFrequency, return 100.f;);
diff --git a/Engine/Source/Runtime/Experimental/Iris/Core/Public/Iris/ReplicationSystem/WorldLocations.h b/Engine/Source/Runtime/Experimental/Iris/Core/Public/Iris/ReplicationSystem/WorldLocations.h
--- a/Engine/Source/Runtime/Experimental/Iris/Core/Public/Iris/ReplicationSystem/WorldLocations.h
+++ b/Engine/Source/Runtime/Experimental/Iris/Core/Public/Iris/ReplicationSystem/WorldLocations.h
@@ -7,6 +7,7 @@
 #include "Iris/Core/NetChunkedArray.h"
 #include "Net/Core/NetBitArray.h"
 #include "Math/Vector.h"
+#include "Misc/Optional.h"
 #include "WorldLocations.generated.h"
 
 class UReplicationSystem;
@@ -83,8 +84,8 @@ public:
 	/** Get the object's current cull distance. */
 	inline float GetCullDistance(UE::Net::Private::FInternalNetRefIndex ObjectIndex) const;
 
-	/** Return the current stored world information of the given object. Returns a default struct if the object did not register in the world location cache. */
-	inline FWorldLocations::FWorldInfo GetWorldInfo(UE::Net::Private::FInternalNetRefIndex ObjectIndex) const;
+	/** Return the current stored world information of the given object, or NullOpt if the object did not register in the world location cache. */
+	inline TOptional<FWorldLocations::FWorldInfo> GetWorldInfo(UE::Net::Private::FInternalNetRefIndex ObjectIndex) const;
 
 	/** Set the mandatory info of a replicated root object*/
 	void SetObjectInfo(UE::Net::Private::FInternalNetRefIndex ObjectIndex, const FVector& Location, float NetCullDistance);
@@ -246,7 +247,7 @@ inline float FWorldLocations::GetCullDistance(UE::Net::Private::FInternalNetRefI
 	return 0.0f;
 }
 
-inline FWorldLocations::FWorldInfo FWorldLocations::GetWorldInfo(UE::Net::Private::FInternalNetRefIndex ObjectIndex) const
+inline TOptional<FWorldLocations::FWorldInfo> FWorldLocations::GetWorldInfo(UE::Net::Private::FInternalNetRefIndex ObjectIndex) const
 {
 	if (ValidInfoIndexes.IsBitSet(ObjectIndex))
 	{
@@ -258,7 +259,7 @@ inline FWorldLocations::FWorldInfo FWorldLocations::GetWorldInfo(UE::Net::Privat
 		};
 	}
 
-	return FWorldInfo();
+	return NullOpt;
 }
 
 } // end namespace UE::Net
