Fix failed checkSlow() in dataflow's TArray-of-UStruct init + move the TArray inheritance to private so we can only selectively expose a subset of the interface directly

diff --git a/Engine/Source/Runtime/Experimental/Dataflow/Core/Private/Dataflow/DataflowContextCache.cpp b/Engine/Source/Runtime/Experimental/Dataflow/Core/Private/Dataflow/DataflowContextCache.cpp
--- a/Engine/Source/Runtime/Experimental/Dataflow/Core/Private/Dataflow/DataflowContextCache.cpp
+++ b/Engine/Source/Runtime/Experimental/Dataflow/Core/Private/Dataflow/DataflowContextCache.cpp
@@ -234,6 +234,8 @@ namespace UE::Dataflow
 	{
 		if (ScriptStruct)  // Null if the array has 0 element
 		{
+			ArrayNum = Num;
+			// Note: SetNumUnsafeInternal cannot expand the array, but we use it here to make sure e.g. slack tracking is called (if enabled)
 			SetNumUnsafeInternal(Num);
 			ArrayMax = Num;
 			GetAllocatorInstance().ResizeAllocation(0, Num, GetStructureSize(), ScriptStruct->GetMinAlignment());
diff --git a/Engine/Source/Runtime/Experimental/Dataflow/Core/Public/Dataflow/DataflowContextCache.h b/Engine/Source/Runtime/Experimental/Dataflow/Core/Public/Dataflow/DataflowContextCache.h
--- a/Engine/Source/Runtime/Experimental/Dataflow/Core/Public/Dataflow/DataflowContextCache.h
+++ b/Engine/Source/Runtime/Experimental/Dataflow/Core/Public/Dataflow/DataflowContextCache.h
@@ -473,13 +473,19 @@ namespace UE::Dataflow
 		FConstStructArrayView GetStructArrayView() const;
 
 		// Implements a FInstancedStruct for arrays (FInstancedStructContainer cannot be cast to a TArray)
-		struct FInstancedStructArray final : public TArray<uint8>
+		struct FInstancedStructArray final : private TArray<uint8>
 		{
 			explicit FInstancedStructArray(const UScriptStruct* const InScriptStruct);
 			explicit FInstancedStructArray(const FConstStructView& StructView);
 			explicit FInstancedStructArray(const FConstStructArrayView& StructArrayView);
 			~FInstancedStructArray();
 
+			// expose parts of the TArray interface that are safe to use directly
+			// Note: otherwise use Get()
+			using TArray<uint8>::IsValidIndex;
+			using TArray<uint8>::Num;
+			using TArray<uint8>::GetData;
+
 			const int32 GetStructureSize() const;
 
 			const UScriptStruct* GetScriptStruct() const;
