Dataflow:
- Added DataflowColorRamp type
- Added DataflowColorRamp detail customization

diff --git a/Engine/Plugins/Experimental/Dataflow/Source/DataflowEditor/Private/Dataflow/DataflowColorRampCustomization.cpp b/Engine/Plugins/Experimental/Dataflow/Source/DataflowEditor/Private/Dataflow/DataflowColorRampCustomization.cpp
new file mode 100644
--- /dev/null
+++ b/Engine/Plugins/Experimental/Dataflow/Source/DataflowEditor/Private/Dataflow/DataflowColorRampCustomization.cpp
@@ -0,0 +1,48 @@
+// Copyright Epic Games, Inc. All Rights Reserved.
+
+#include "Dataflow/DataflowColorRampCustomization.h"
+#include "DetailLayoutBuilder.h"
+#include "DetailCategoryBuilder.h"
+#include "Dataflow/DataflowColorRamp.h"
+#include "SColorGradientEditor.h"
+#include "Dataflow/DataflowGraphEditor.h"
+#include "DetailWidgetRow.h"
+
+namespace UE::Dataflow
+{
+	TSharedRef<IPropertyTypeCustomization> FDataflowColorRampCustomization::MakeInstance()
+	{
+		return MakeShareable(new FDataflowColorRampCustomization);
+	}
+
+	void FDataflowColorRampCustomization::CustomizeHeader(TSharedRef<IPropertyHandle> PropertyHandle, FDetailWidgetRow& HeaderRow, IPropertyTypeCustomizationUtils& CustomizationUtils)
+	{
+		if (const FStructProperty* const StructProperty = CastField<FStructProperty>(PropertyHandle->GetProperty()))
+		{
+			if (StructProperty->Struct)
+			{
+				void* Data = nullptr;
+
+				if (PropertyHandle->GetValueData(Data) == FPropertyAccess::Success)
+				{
+					FDataflowColorRamp* ColorRamp = reinterpret_cast<FDataflowColorRamp*>(Data);
+
+					TSharedPtr<SColorGradientEditor> GradientEditor;
+
+					HeaderRow
+						[
+							SAssignNew(GradientEditor, SColorGradientEditor)
+								.ViewMinInput(0.0f)
+								.ViewMaxInput(1.0f)
+								.ClampStopsToViewRange(true)
+						];
+
+					GradientEditor->SetCurveOwner(&ColorRamp->ColorRamp);
+				}
+			}
+		}
+	}
+}  
+
+
+
diff --git a/Engine/Plugins/Experimental/Dataflow/Source/DataflowEditor/Private/Dataflow/DataflowEditorModule.cpp b/Engine/Plugins/Experimental/Dataflow/Source/DataflowEditor/Private/Dataflow/DataflowEditorModule.cpp
--- a/Engine/Plugins/Experimental/Dataflow/Source/DataflowEditor/Private/Dataflow/DataflowEditorModule.cpp
+++ b/Engine/Plugins/Experimental/Dataflow/Source/DataflowEditor/Private/Dataflow/DataflowEditorModule.cpp
@@ -27,6 +27,8 @@
 #include "Dataflow/DataflowCollectionEditSkinWeightsNode.h"
 #include "DataflowEditorTools/DataflowEditorSkinWeightsPaintTool.h"
 #include "DataflowEditorTools/DataflowEditorCorrectSkinWeightsNode.h"
+#include "Dataflow/DataflowColorRamp.h"
+#include "Dataflow/DataflowColorRampCustomization.h"
 
 #define LOCTEXT_NAMESPACE "DataflowEditor"
 
@@ -36,6 +38,7 @@ namespace UE::Dataflow::Private
 	static const FName DataflowFunctionPropertyName = TEXT("DataflowFunctionProperty");
 	static const FName DataflowVariableOverridesName = TEXT("DataflowVariableOverrides");
 	static const FName DataflowFreezeActionsName = TEXT("DataflowFreezeActions");
+	static const FName DataflowColorRampName = TEXT("DataflowColorRamp");
 
 	class FDataflowEditorWeightMapPaintToolActionCommands : public TInteractiveToolCommands<FDataflowEditorWeightMapPaintToolActionCommands>
 	{
@@ -120,6 +123,9 @@ void FDataflowEditorModule::StartupModule()
 		PropertyModule->RegisterCustomPropertyTypeLayout(DataflowFunctionPropertyName, FOnGetPropertyTypeCustomizationInstance::CreateStatic(&UE::Dataflow::FFunctionPropertyCustomization::MakeInstance));
 		PropertyModule->RegisterCustomPropertyTypeLayout(DataflowVariableOverridesName, FOnGetPropertyTypeCustomizationInstance::CreateStatic(&FDataflowVariableOverridesDetails::MakeInstance));
 		PropertyModule->RegisterCustomPropertyTypeLayout(DataflowFreezeActionsName, FOnGetPropertyTypeCustomizationInstance::CreateStatic(&UE::Dataflow::FFreezeActionsCustomization::MakeInstance));
+
+		PropertyModule->RegisterCustomPropertyTypeLayout(DataflowColorRampName, FOnGetPropertyTypeCustomizationInstance::CreateStatic(&UE::Dataflow::FDataflowColorRampCustomization::MakeInstance));
+
 		UE::Dataflow::RegisterAnyTypeCustomizations(*PropertyModule);
 	}
 
diff --git a/Engine/Plugins/Experimental/Dataflow/Source/DataflowEditor/Public/Dataflow/DataflowColorRampCustomization.h b/Engine/Plugins/Experimental/Dataflow/Source/DataflowEditor/Public/Dataflow/DataflowColorRampCustomization.h
new file mode 100644
--- /dev/null
+++ b/Engine/Plugins/Experimental/Dataflow/Source/DataflowEditor/Public/Dataflow/DataflowColorRampCustomization.h
@@ -0,0 +1,20 @@
+// Copyright Epic Games, Inc. All Rights Reserved.
+
+#pragma once
+
+#include "IPropertyTypeCustomization.h"
+
+namespace UE::Dataflow
+{
+	class FDataflowColorRampCustomization : public IPropertyTypeCustomization
+	{
+	public:
+
+		/** Makes a new instance of this detail layout class for a specific detail view requesting it */
+		DATAFLOWEDITOR_API static TSharedRef<IPropertyTypeCustomization> MakeInstance();
+
+		/** IDetailCustomization interface */
+		DATAFLOWEDITOR_API virtual void CustomizeHeader(TSharedRef<IPropertyHandle> PropertyHandle, FDetailWidgetRow& HeaderRow, IPropertyTypeCustomizationUtils& CustomizationUtils) override;
+		virtual void CustomizeChildren(TSharedRef<IPropertyHandle> /*StructPropertyHandle*/, IDetailChildrenBuilder& /*ChildBuilder*/, IPropertyTypeCustomizationUtils& /*CustomizationUtils*/) override {}
+	};
+}
diff --git a/Engine/Source/Runtime/Experimental/Dataflow/Engine/Private/Dataflow/DataflowColorRamp.cpp b/Engine/Source/Runtime/Experimental/Dataflow/Engine/Private/Dataflow/DataflowColorRamp.cpp
new file mode 100644
--- /dev/null
+++ b/Engine/Source/Runtime/Experimental/Dataflow/Engine/Private/Dataflow/DataflowColorRamp.cpp
@@ -0,0 +1,67 @@
+// Copyright Epic Games, Inc. All Rights Reserved.
+
+#include "Dataflow/DataflowColorRamp.h"
+
+TArray<FRichCurveEditInfoConst> FDataflowColorCurveOwner::GetCurves() const
+{
+	return ConstCurves;
+}
+
+TArray<FRichCurveEditInfo> FDataflowColorCurveOwner::GetCurves()
+{
+	return Curves;
+}
+
+void FDataflowColorCurveOwner::ModifyOwner()
+{}
+
+TArray<const UObject*> FDataflowColorCurveOwner::GetOwners() const
+{
+	return Owners;
+}
+
+void FDataflowColorCurveOwner::MakeTransactional()
+{}
+
+void FDataflowColorCurveOwner::OnCurveChanged(const TArray<FRichCurveEditInfo>& ChangedCurveEditInfos)
+{
+	TArray<FRichCurve*> ChangedCurves;
+	for (const FRichCurveEditInfo& ChangedCurveEditInfo : ChangedCurveEditInfos)
+	{
+		ChangedCurves.Add((FRichCurve*)ChangedCurveEditInfo.CurveToEdit);
+	}
+	if (ChangedCurves.Num() > 0)
+	{
+		OnColorCurveChangedDelegate.Broadcast(ChangedCurves);
+	}
+}
+
+bool FDataflowColorCurveOwner::IsLinearColorCurve() const
+{
+	return true;
+}
+
+FLinearColor FDataflowColorCurveOwner::GetLinearColorValue(float InTime) const
+{
+	return FLinearColor(
+		Curves[0].CurveToEdit->Eval(InTime),
+		Curves[1].CurveToEdit->Eval(InTime),
+		Curves[2].CurveToEdit->Eval(InTime),
+		Curves[3].CurveToEdit->Eval(InTime));
+}
+
+bool FDataflowColorCurveOwner::HasAnyAlphaKeys() const
+{
+	return Curves[3].CurveToEdit->GetNumKeys() != 0;
+}
+
+bool FDataflowColorCurveOwner::IsValidCurve(FRichCurveEditInfo CurveInfo)
+{
+	return Curves.Contains(CurveInfo);
+}
+
+FLinearColor FDataflowColorCurveOwner::GetCurveColor(FRichCurveEditInfo CurveInfo) const
+{
+	return FLinearColor::White;
+}
+
diff --git a/Engine/Source/Runtime/Experimental/Dataflow/Engine/Public/Dataflow/DataflowColorRamp.h b/Engine/Source/Runtime/Experimental/Dataflow/Engine/Public/Dataflow/DataflowColorRamp.h
new file mode 100644
--- /dev/null
+++ b/Engine/Source/Runtime/Experimental/Dataflow/Engine/Public/Dataflow/DataflowColorRamp.h
@@ -0,0 +1,59 @@
+// Copyright Epic Games, Inc. All Rights Reserved.
+
+#pragma once 
+
+#include "Curves/RichCurve.h"
+#include "Curves/CurveOwnerInterface.h"
+
+#include "DataflowColorRamp.generated.h"
+
+struct FDataflowColorCurveOwner : public FCurveOwnerInterface
+{
+public:
+	DECLARE_MULTICAST_DELEGATE_OneParam(FOnColorCurveChanged, TArray<FRichCurve*> /* ChangedCurves */);
+
+public:
+	FDataflowColorCurveOwner() 
+	{ 
+		ConstCurves.SetNum(4);
+		Curves.SetNum(4);
+
+		RichCurves.SetNum(4);
+
+		for (int32 Idx = 0; Idx < 4; ++Idx)
+		{
+			Curves[Idx].CurveToEdit = &RichCurves[Idx];
+			ConstCurves[Idx].CurveToEdit = &RichCurves[Idx];
+		}
+	}
+
+	/** implemenmt all the necessary virtual function from FCurveOwnerInterface */
+	DATAFLOWENGINE_API virtual TArray<FRichCurveEditInfoConst> GetCurves() const override;
+	DATAFLOWENGINE_API virtual TArray<FRichCurveEditInfo> GetCurves() override;
+	DATAFLOWENGINE_API virtual void ModifyOwner() override;
+	DATAFLOWENGINE_API virtual TArray<const UObject*> GetOwners() const override;
+	DATAFLOWENGINE_API virtual void MakeTransactional() override;
+	DATAFLOWENGINE_API virtual void OnCurveChanged(const TArray<FRichCurveEditInfo>& ChangedCurveEditInfos) override;
+	DATAFLOWENGINE_API virtual bool IsLinearColorCurve() const override;
+	DATAFLOWENGINE_API virtual FLinearColor GetLinearColorValue(float InTime) const override;
+	DATAFLOWENGINE_API virtual bool HasAnyAlphaKeys() const override;
+	DATAFLOWENGINE_API virtual bool IsValidCurve(FRichCurveEditInfo CurveInfo) override;
+	DATAFLOWENGINE_API virtual FLinearColor GetCurveColor(FRichCurveEditInfo CurveInfo) const override;
+
+	FOnColorCurveChanged OnColorCurveChangedDelegate;
+
+private:
+	TArray<FRichCurveEditInfoConst> ConstCurves;
+	TArray<FRichCurveEditInfo> Curves;
+	TArray<const UObject*> Owners;
+
+	TArray<FRichCurve> RichCurves;
+};
+
+USTRUCT()
+struct FDataflowColorRamp
+{
+	GENERATED_USTRUCT_BODY()
+
+	FDataflowColorCurveOwner ColorRamp;
+};
\ No newline at end of file
