[ChaosVD] Fix an issue where CVD's sessions manager broadcasts a session discovery ping every second, even if CVD is not open in the editor.

diff --git a/Engine/Plugins/ChaosVD/Source/ChaosVD/Private/Widgets/SChaosVDMainTab.cpp b/Engine/Plugins/ChaosVD/Source/ChaosVD/Private/Widgets/SChaosVDMainTab.cpp
--- a/Engine/Plugins/ChaosVD/Source/ChaosVD/Private/Widgets/SChaosVDMainTab.cpp
+++ b/Engine/Plugins/ChaosVD/Source/ChaosVD/Private/Widgets/SChaosVDMainTab.cpp
@@ -35,6 +35,7 @@
 #include "StatusBarSubsystem.h"
 #include "ToolMenu.h"
 #include "ToolMenus.h"
+#include "Chaos/ChaosVDEngineEditorBridge.h"
 #include "Components/ChaosVDInstancedStaticMeshComponent.h"
 #include "Components/ChaosVDParticleDataComponent.h"
 #include "Components/ChaosVDSolverCharacterGroundConstraintDataComponent.h"
@@ -212,10 +213,20 @@ void SChaosVDMainTab::Construct(const FArguments& InArgs, TSharedPtr<FChaosVDEng
 	TabManager->TryInvokeTab(FChaosVDTabID::DetailsPanel);
 
 	SetUpDisableCPUThrottlingDelegate();
+
+	if (TSharedPtr<FChaosVDRemoteSessionsManager> RemoteSessionManager = FChaosVDEngineEditorBridge::Get().GetRemoteSessionsManager())
+	{
+		RemoteSessionManager->StartSessionDiscovery();
+	}
 }
 
 SChaosVDMainTab::~SChaosVDMainTab()
 {
+	if (TSharedPtr<FChaosVDRemoteSessionsManager> RemoteSessionManager = FChaosVDEngineEditorBridge::Get().GetRemoteSessionsManager())
+	{
+		RemoteSessionManager->StopSessionDiscovery();
+	}
+
 	FChaosVDExtensionsManager::Get().OnExtensionRegistered().RemoveAll(this);
 	CleanUpDisableCPUThrottlingDelegate();
 }
diff --git a/Engine/Source/Runtime/Experimental/ChaosSolverEngine/Private/Chaos/ChaosVDRemoteSessionsManager.cpp b/Engine/Source/Runtime/Experimental/ChaosSolverEngine/Private/Chaos/ChaosVDRemoteSessionsManager.cpp
--- a/Engine/Source/Runtime/Experimental/ChaosSolverEngine/Private/Chaos/ChaosVDRemoteSessionsManager.cpp
+++ b/Engine/Source/Runtime/Experimental/ChaosSolverEngine/Private/Chaos/ChaosVDRemoteSessionsManager.cpp
@@ -8,6 +8,7 @@
 #include "MessageEndpoint.h"
 #include "MessageEndpointBuilder.h"
 #include "ChaosVisualDebugger/ChaosVDOptionalDataChannel.h"
+#include "ChaosVisualDebugger/ChaosVisualDebuggerTrace.h"
 #include "Misc/App.h"
 
 const FGuid FChaosVDRemoteSessionsManager::AllRemoteSessionsWrapperGUID = FGuid::NewGuid();
@@ -177,9 +178,6 @@ void FChaosVDRemoteSessionsManager::Initialize()
 	ActiveSessionsByInstanceId.Add(AllRemoteClientsWrapperGUID, CreatedWrapperSessionInfo(AllRemoteClientsWrapperGUID, AllRemoteClientsTargetName));
 	ActiveSessionsByInstanceId.Add(AllSessionsWrapperGUID, CreatedWrapperSessionInfo(AllSessionsWrapperGUID, AllSessionsTargetName));
 	ActiveSessionsByInstanceId.Add(CustomSessionsWrapperGUID, CreatedWrapperSessionInfo(CustomSessionsWrapperGUID, CustomSessionsTargetName));
-
-	constexpr float TickInterval = 1.0f;
-	TickHandle = FTSTicker::GetCoreTicker().AddTicker(FTickerDelegate::CreateRaw(this, &FChaosVDRemoteSessionsManager::Tick), TickInterval);
 }
 
 void FChaosVDRemoteSessionsManager::Shutdown()
@@ -196,8 +194,28 @@ void FChaosVDRemoteSessionsManager::Shutdown()
 		MessageEndpoint->Unsubscribe<FChaosVDFullSessionInfoResponseMessage>();
 		MessageEndpoint->Unsubscribe<FChaosVDChannelStateChangeResponseMessage>();
 	}
+}
+
+void FChaosVDRemoteSessionsManager::StartSessionDiscovery()
+{
+	if (TickHandle.IsValid())
+	{
+		UE_LOG(LogChaosVDRemoteSession, Warning, TEXT("[%hs] Session discovery already started"), __func__);
+		return;
+	}
+	
+	constexpr float TickInterval = 1.0f;
+	TickHandle = FTSTicker::GetCoreTicker().AddTicker(FTickerDelegate::CreateRaw(this, &FChaosVDRemoteSessionsManager::Tick), TickInterval);
+}
 
-	FTSTicker::GetCoreTicker().RemoveTicker(TickHandle);
+void FChaosVDRemoteSessionsManager::StopSessionDiscovery()
+{
+	if (TickHandle.IsValid())
+	{
+		RemoveExpiredSessions(ERemoveSessionOptions::ForceRemoveAll);
+		FTSTicker::GetCoreTicker().RemoveTicker(TickHandle);
+		TickHandle = FTSTicker::FDelegateHandle();
+	}
 }
 
 void FChaosVDRemoteSessionsManager::PublishRecordingStatusUpdate(const FChaosVDRecordingStatusMessage& UpdateMessage)
@@ -537,7 +555,7 @@ void FChaosVDRemoteSessionsManager::HandleFullSessionStateResponseMessage(const
 	}
 }
 
-void FChaosVDRemoteSessionsManager::RemoveExpiredSessions()
+void FChaosVDRemoteSessionsManager::RemoveExpiredSessions(ERemoveSessionOptions Options)
 {
 	bool bAnySessionRemoved = false;
 	FDateTime CurrentTime = FDateTime::UtcNow();
@@ -561,7 +579,7 @@ void FChaosVDRemoteSessionsManager::RemoveExpiredSessions()
 		// A session goes into busy state if we are attempting to issue a command that might stall the target, currently that only happens on recording start commands of complex maps
 		// In these cases, we need to allow more time between pings. If a recording command failed, it is expected the state to be changed to Ready again
 		const float MaxAllowedTimeBetweenPings = SessionInfoPtr->ReadyState == EChaosVDRemoteSessionReadyState::Busy ? 3.0f : 60.0f;
-		if (ElapsedTime > FTimespan::FromSeconds(MaxAllowedTimeBetweenPings))
+		if (EnumHasAnyFlags(Options, ERemoveSessionOptions::ForceRemoveAll) || ElapsedTime > FTimespan::FromSeconds(MaxAllowedTimeBetweenPings))
 		{
 			DeRegisterSessionInMultiSessionWrapper(SessionInfoPtr.ToSharedRef());
 			RemoveIterator.RemoveCurrent();
diff --git a/Engine/Source/Runtime/Experimental/ChaosSolverEngine/Public/Chaos/ChaosVDRemoteSessionsManager.h b/Engine/Source/Runtime/Experimental/ChaosSolverEngine/Public/Chaos/ChaosVDRemoteSessionsManager.h
--- a/Engine/Source/Runtime/Experimental/ChaosSolverEngine/Public/Chaos/ChaosVDRemoteSessionsManager.h
+++ b/Engine/Source/Runtime/Experimental/ChaosSolverEngine/Public/Chaos/ChaosVDRemoteSessionsManager.h
@@ -259,6 +259,16 @@ public:
 	}
 
 
+	/**
+	 * Starts sending ping requests that potential targets (client, servers or other editors) can answer to and report themselves as available
+	 */
+	CHAOSSOLVERENGINE_API void StartSessionDiscovery();
+	
+	/**
+     * Stop sending ping requests to potential targets (client, servers or other editors).
+     */
+	CHAOSSOLVERENGINE_API void StopSessionDiscovery();
+
 	/**
 	 * Issues a command to the provided address that will start a CVD recording
 	 * @param InDestinationAddress Message bus address that will execute the command
@@ -411,7 +421,14 @@ private:
 	void HandleFullSessionStateRequestMessage(const FChaosVDFullSessionInfoRequestMessage& InMessage, const TSharedRef<IMessageContext>& InContext);
 	void HandleFullSessionStateResponseMessage(const FChaosVDFullSessionInfoResponseMessage& InMessage, const TSharedRef<IMessageContext>& InContext);
 
-	void RemoveExpiredSessions();
+	enum class ERemoveSessionOptions : uint8
+	{
+		None = 0,
+		ForceRemoveAll = 1 << 0
+	};
+	FRIEND_ENUM_CLASS_FLAGS(ERemoveSessionOptions)
+
+	void RemoveExpiredSessions(ERemoveSessionOptions Options = ERemoveSessionOptions::None);
 
 	/** Holds the time at which the last ping was sent. */
     FDateTime LastPingTime;
