ManagedArrayCollection : fixed copy attribute to not copy the attribute twice and making sure new attribute get a chance to not be copied if the conditions are right

diff --git a/Engine/Source/Runtime/Experimental/Chaos/Private/GeometryCollection/ManagedArrayCollection.cpp b/Engine/Source/Runtime/Experimental/Chaos/Private/GeometryCollection/ManagedArrayCollection.cpp
--- a/Engine/Source/Runtime/Experimental/Chaos/Private/GeometryCollection/ManagedArrayCollection.cpp
+++ b/Engine/Source/Runtime/Experimental/Chaos/Private/GeometryCollection/ManagedArrayCollection.cpp
@@ -624,14 +624,16 @@ void FManagedArrayCollection::CopyAttribute(const FManagedArrayCollection& InCol
 	FKeyType SrcKey = FManagedArrayCollection::MakeMapKey(SrcName, Group);
 	FKeyType DestKey = FManagedArrayCollection::MakeMapKey(DestName, Group);
 
+	const FValueType& SrcMapEntry = InCollection.Map[SrcKey];
+
 	if (!HasAttribute(DestName, Group))
 	{
-		FValueType NewAttribute(InCollection.Map[SrcKey]);
-		Map.Add(DestKey, MoveTemp(NewAttribute));
+		Map.Emplace(DestKey, FValueType(SrcMapEntry));
+	}
+	else
+	{
+		Map[DestKey].InitFrom(SrcMapEntry);
 	}
-
-	const FValueType& OriginalValue = InCollection.Map[SrcKey];
-	Map[DestKey].InitFrom(OriginalValue);
 }
 
 bool FManagedArrayCollection::IsConnected(FName StartingNode, FName TargetNode)
