Dataflow : Tentative fix for crash in Subgraph nodes

Making sure we unregister the handlers even if the dataflow asset is no longer valid

diff --git a/Engine/Source/Runtime/Experimental/Dataflow/Engine/Private/Dataflow/DataflowSubGraphNodes.cpp b/Engine/Source/Runtime/Experimental/Dataflow/Engine/Private/Dataflow/DataflowSubGraphNodes.cpp
--- a/Engine/Source/Runtime/Experimental/Dataflow/Engine/Private/Dataflow/DataflowSubGraphNodes.cpp
+++ b/Engine/Source/Runtime/Experimental/Dataflow/Engine/Private/Dataflow/DataflowSubGraphNodes.cpp
@@ -447,10 +447,12 @@ FDataflowCallSubGraphNode::FDataflowCallSubGraphNode(const UE::Dataflow::FNodePa
 
 FDataflowCallSubGraphNode::~FDataflowCallSubGraphNode()
 {
+	UDataflowSubGraph* SubGraph = nullptr;
 	if (TStrongObjectPtr<UDataflow> DataflowAsset = DataflowAssetWeakPtr.Pin())
 	{
-		UnregisterHandlers(DataflowAsset->FindSubGraphByGuid(SubGraphGuid));
+		SubGraph = DataflowAsset->FindSubGraphByGuid(SubGraphGuid);
 	}
+	UnregisterHandlers(SubGraph);
 }
 
 void FDataflowCallSubGraphNode::RegisterHandlers(UDataflowSubGraph* SubGraph)
