Iris - temporary last-resort serializers for FRemoteServerId and FRemoteObjectId in the engine.

diff --git a/Engine/Source/Runtime/Experimental/Iris/Core/Private/Iris/Serialization/RemoteObjectIdNetSerializer.cpp b/Engine/Source/Runtime/Experimental/Iris/Core/Private/Iris/Serialization/RemoteObjectIdNetSerializer.cpp
new file mode 100644
--- /dev/null
+++ b/Engine/Source/Runtime/Experimental/Iris/Core/Private/Iris/Serialization/RemoteObjectIdNetSerializer.cpp
@@ -0,0 +1,6 @@
+// Copyright Epic Games, Inc. All Rights Reserved.
+
+#include "Iris/ReplicationState/PropertyNetSerializerInfoRegistry.h"
+#include "UObject/RemoteObjectTypes.h"
+
+UE_NET_IMPLEMENT_NAMED_STRUCT_LASTRESORT_NETSERIALIZER_AND_REGISTRY_DELEGATES(RemoteObjectId);
diff --git a/Engine/Source/Runtime/Experimental/Iris/Core/Private/Iris/Serialization/RemoteServerIdNetSerializer.cpp b/Engine/Source/Runtime/Experimental/Iris/Core/Private/Iris/Serialization/RemoteServerIdNetSerializer.cpp
new file mode 100644
--- /dev/null
+++ b/Engine/Source/Runtime/Experimental/Iris/Core/Private/Iris/Serialization/RemoteServerIdNetSerializer.cpp
@@ -0,0 +1,6 @@
+// Copyright Epic Games, Inc. All Rights Reserved.
+
+#include "Iris/ReplicationState/PropertyNetSerializerInfoRegistry.h"
+#include "UObject/RemoteObjectTypes.h"
+
+UE_NET_IMPLEMENT_NAMED_STRUCT_LASTRESORT_NETSERIALIZER_AND_REGISTRY_DELEGATES(RemoteServerId);
