Kinematic Collider and Skinned Triangle Mesh bugfixes. These are regressions from implementing Skinned Triangle Meshes in Chaos Cloth.

* Kinematic Collider was using previous frame's positions for its constraint Apply causing it to lag behind and not fully resolve collisions.
* Kinematic Collider debug draw was broken due to an indexing bug.
* Added ClothCollisionThickness parameter on the SimulationCollisionConfigNode was not setting the correct string, so it wasn't being used.  This fix for this will not affect existing assets (ClothCollisionThickness will be set back to 0 to retain existing behavior).

diff --git a/Engine/Plugins/ChaosCloth/Source/ChaosCloth/Private/ChaosCloth/ChaosClothVisualization.cpp b/Engine/Plugins/ChaosCloth/Source/ChaosCloth/Private/ChaosCloth/ChaosClothVisualization.cpp
--- a/Engine/Plugins/ChaosCloth/Source/ChaosCloth/Private/ChaosCloth/ChaosClothVisualization.cpp
+++ b/Engine/Plugins/ChaosCloth/Source/ChaosCloth/Private/ChaosCloth/ChaosClothVisualization.cpp
@@ -2638,33 +2638,36 @@ namespace Chaos
 
 				const TArray<int32>& KinematicCollidingParticles = SelfCollisionConstraints->GetKinematicCollidingParticles();
 				const TArray<TMap<int32, Softs::FSolverReal>>& KinematicColliderTimers = SelfCollisionConstraints->GetKinematicColliderTimers();
-				const FTriangleMesh& TriangleMesh = SelfCollisionConstraints->GetTriangleMesh();
-				for (int32 Index1 : KinematicCollidingParticles)
+				if (const FTriangleMesh* TriangleMesh = SelfCollisionConstraints->GetKinematicColliderTriangleMesh())
 				{
-					const FVector P = GetWorldPosition(*Solver, Positions[Index1]);
 
-					static const FLinearColor Orange(0.3f, 0.15f, 0.f);
-					DrawPoint(PDI, P, Orange, nullptr, 2.f);
-
-					const TMap<int32, Softs::FSolverReal>& Timers = KinematicColliderTimers[Index1 - Offset];
-					for (const TPair<int32, Softs::FSolverReal>& ElemAndTimer : Timers)
+					for (int32 Index1 : KinematicCollidingParticles)
 					{
-						const int32 Index2 = TriangleMesh.GetElements()[ElemAndTimer.Get<0>()][0];
-						const int32 Index3 = TriangleMesh.GetElements()[ElemAndTimer.Get<0>()][1];
-						const int32 Index4 = TriangleMesh.GetElements()[ElemAndTimer.Get<0>()][2];
-
-						const Softs::FSolverVec3& P1 = Positions[Index1];
-						const Softs::FSolverVec3& P2 = Positions[Index2];
-						const Softs::FSolverVec3& P3 = Positions[Index3];
-						const Softs::FSolverVec3& P4 = Positions[Index4];
-						Softs::FSolverVec3 Bary;
-						const FVector Pos1 = GetWorldPosition(*Solver, FindClosestPointAndBaryOnTriangle(P2, P3, P4, P1, Bary));
-
-						static const FLinearColor LtRed(0.6f, 0.f, 0.f);
-						static const FLinearColor DkRed(0.3f, 0.f, 0.f);
-						const FLinearColor& Color = ElemAndTimer.Get<1>() > 0.f ? LtRed : DkRed;
-						DrawPoint(PDI, Pos1, Color, nullptr, 2.f);
-						DrawLine(PDI, Pos1, P, Color);
+						const FVector P = GetWorldPosition(*Solver, Positions[Index1]);
+
+						static const FLinearColor Orange(0.3f, 0.15f, 0.f);
+						DrawPoint(PDI, P, Orange, nullptr, 2.f);
+
+						const TMap<int32, Softs::FSolverReal>& Timers = KinematicColliderTimers[Index1 - Offset];
+						for (const TPair<int32, Softs::FSolverReal>& ElemAndTimer : Timers)
+						{
+							const int32 Index2 = TriangleMesh->GetElements()[ElemAndTimer.Get<0>()][0];
+							const int32 Index3 = TriangleMesh->GetElements()[ElemAndTimer.Get<0>()][1];
+							const int32 Index4 = TriangleMesh->GetElements()[ElemAndTimer.Get<0>()][2];
+
+							const Softs::FSolverVec3& P1 = Positions[Index1];
+							const Softs::FSolverVec3& P2 = Positions[Index2];
+							const Softs::FSolverVec3& P3 = Positions[Index3];
+							const Softs::FSolverVec3& P4 = Positions[Index4];
+							Softs::FSolverVec3 Bary;
+							const FVector Pos1 = GetWorldPosition(*Solver, FindClosestPointAndBaryOnTriangle(P2, P3, P4, P1, Bary));
+
+							static const FLinearColor LtRed(0.6f, 0.f, 0.f);
+							static const FLinearColor DkRed(0.3f, 0.f, 0.f);
+							const FLinearColor& Color = ElemAndTimer.Get<1>() > 0.f ? LtRed : DkRed;
+							DrawPoint(PDI, Pos1, Color, nullptr, 2.f);
+							DrawLine(PDI, Pos1, P, Color);
+						}
 					}
 				}
 			}
diff --git a/Engine/Plugins/ChaosClothAssetEditor/Source/ChaosClothAssetDataflowNodes/Private/ChaosClothAsset/SimulationCollisionConfigNode.cpp b/Engine/Plugins/ChaosClothAssetEditor/Source/ChaosClothAssetDataflowNodes/Private/ChaosClothAsset/SimulationCollisionConfigNode.cpp
--- a/Engine/Plugins/ChaosClothAssetEditor/Source/ChaosClothAssetDataflowNodes/Private/ChaosClothAsset/SimulationCollisionConfigNode.cpp
+++ b/Engine/Plugins/ChaosClothAssetEditor/Source/ChaosClothAssetDataflowNodes/Private/ChaosClothAsset/SimulationCollisionConfigNode.cpp
@@ -13,7 +13,7 @@ FChaosClothAssetSimulationCollisionConfigNode::FChaosClothAssetSimulationCollisi
 	RegisterInputConnection(&FrictionCoefficientWeighted.WeightMap)
 		.SetCanHidePin(true)
 		.SetPinIsHidden(true);
-	RegisterInputConnection(&ClothCollisionThickness.WeightMap)
+	RegisterInputConnection(&ClothCollisionThicknessImported.WeightMap)
 		.SetCanHidePin(true)
 		.SetPinIsHidden(true);
 }
@@ -36,7 +36,7 @@ void FChaosClothAssetSimulationCollisionConfigNode::AddProperties(FPropertyHelpe
 			return FabricFacade.GetFriction();
 		}, {});
 
-	PropertyHelper.SetFabricPropertyWeighted(FName(TEXT("SoftBodyCollisionThickness")), ClothCollisionThickness,
+	PropertyHelper.SetFabricPropertyWeighted(FName(TEXT("ClothCollisionThickness")), ClothCollisionThicknessImported,
 		[](const UE::Chaos::ClothAsset::FCollectionClothFabricFacade& FabricFacade)-> float
 		{
 			return FabricFacade.GetCollisionThickness();
@@ -77,6 +77,8 @@ void FChaosClothAssetSimulationCollisionConfigNode::Serialize(FArchive& Ar)
 			FrictionCoefficientWeighted.bImportFabricBounds = true;
 			FrictionCoefficientImported.bUseImportedValue = false;
 		}
+		
+		// Intentionally not upgrading ClothCollisionThickness. This property was using the incorrect string and was unused.
 		PRAGMA_ENABLE_DEPRECATION_WARNINGS
 #endif
 	}
diff --git a/Engine/Plugins/ChaosClothAssetEditor/Source/ChaosClothAssetDataflowNodes/Public/ChaosClothAsset/SimulationCollisionConfigNode.h b/Engine/Plugins/ChaosClothAssetEditor/Source/ChaosClothAssetDataflowNodes/Public/ChaosClothAsset/SimulationCollisionConfigNode.h
--- a/Engine/Plugins/ChaosClothAssetEditor/Source/ChaosClothAssetDataflowNodes/Public/ChaosClothAsset/SimulationCollisionConfigNode.h
+++ b/Engine/Plugins/ChaosClothAssetEditor/Source/ChaosClothAssetDataflowNodes/Public/ChaosClothAsset/SimulationCollisionConfigNode.h
@@ -49,8 +49,8 @@ public:
 	bool bUseSelfCollisionSubstepsForSkinnedTriangleMeshes = true;
 
 	/** Thickness added to the cloth when colliding against collision shapes. Currently only Skinned Triangle Meshes use the weighted value. All other collisions only use the Low value. */
-	UPROPERTY(EditAnywhere, Category = "Collision Properties", Meta = (InteractorName = "SoftBodyCollisionThickness"))
-	FChaosClothAssetWeightedValue ClothCollisionThickness = { true, UE::Chaos::ClothAsset::FDefaultFabric::ClothCollisionThickness,
+	UPROPERTY(EditAnywhere, Category = "Collision Properties", Meta = (InteractorName = "ClothCollisionThickness"))
+	FChaosClothAssetWeightedValue ClothCollisionThicknessImported = { true, UE::Chaos::ClothAsset::FDefaultFabric::ClothCollisionThickness,
 		UE::Chaos::ClothAsset::FDefaultFabric::ClothCollisionThickness, TEXT("ClothCollisionThickness"), true };
 
 	/** Stiffness for proximity repulsion forces (Force-based solver only). Units = kg cm/ s^2 (same as XPBD springs)*/
@@ -69,6 +69,10 @@ public:
 	UPROPERTY()
 	FChaosClothAssetImportedFloatValue FrictionCoefficientImported = { UE::Chaos::ClothAsset::FDefaultFabric::Friction };
 
+	UE_DEPRECATED(5.7, "Use ClothCollisionThicknessImported instead.")
+	UPROPERTY()
+	float ClothCollisionThickness = 0.;
+
 	FChaosClothAssetSimulationCollisionConfigNode(const UE::Dataflow::FNodeParameters& InParam, FGuid InGuid = FGuid::NewGuid());
 
 	virtual void Serialize(FArchive& Ar) override;
diff --git a/Engine/Source/Runtime/Experimental/Chaos/Private/Chaos/PBDCollisionSpringConstraintsBase.cpp b/Engine/Source/Runtime/Experimental/Chaos/Private/Chaos/PBDCollisionSpringConstraintsBase.cpp
--- a/Engine/Source/Runtime/Experimental/Chaos/Private/Chaos/PBDCollisionSpringConstraintsBase.cpp
+++ b/Engine/Source/Runtime/Experimental/Chaos/Private/Chaos/PBDCollisionSpringConstraintsBase.cpp
@@ -125,7 +125,7 @@ void FPBDCollisionSpringConstraintsBase::Init(const SolverParticlesOrRange& Part
 
 	if constexpr (std::is_same<SpatialAccelerator, THierarchicalSpatialHash<int32, FSolverReal>>::value && std::is_same<FSolverParticlesRange, SolverParticlesOrRange>::value)
 	{
-		KinematicCollisions.SetGeometry(CollidableSubMesh.GetKinematicColliderSubMesh(), Particles.XArray(), Particles.GetV(), KinematicColliderSpatial);
+		KinematicCollisions.SetGeometry(CollidableSubMesh.GetKinematicColliderSubMesh(), Particles.XArray(), Particles.GetPAndInvM(), Particles.GetV(), KinematicColliderSpatial);
 		KinematicCollisions.Init(Particles, Dt);
 	}
 	else
diff --git a/Engine/Source/Runtime/Experimental/Chaos/Private/Chaos/PBDKinematicTriangleMeshCollisions.cpp b/Engine/Source/Runtime/Experimental/Chaos/Private/Chaos/PBDKinematicTriangleMeshCollisions.cpp
--- a/Engine/Source/Runtime/Experimental/Chaos/Private/Chaos/PBDKinematicTriangleMeshCollisions.cpp
+++ b/Engine/Source/Runtime/Experimental/Chaos/Private/Chaos/PBDKinematicTriangleMeshCollisions.cpp
@@ -140,7 +140,6 @@ namespace Chaos::Softs {
 		CollidingElements.SetNum(KinematicConstraintNum, EAllowShrinking::No);
 	}
 
-
 	void FPBDKinematicTriangleMeshCollisions::Apply(FSolverParticlesRange& Particles, const FSolverReal Dt) const
 	{
 		if (!CollidingParticles.Num())
@@ -160,43 +159,89 @@ namespace Chaos::Softs {
 			{
 				if (bWithFriction)
 				{
-					ispc::ApplyKinematicTriangleCollisionsWithFrictionAndMaps(
-						(ispc::FVector4f*)Particles.GetPAndInvM().GetData(),
-						(const ispc::FVector3f*)Particles.XArray().GetData(),
-						(const ispc::FVector3f*)Positions.GetData(),
-						(const ispc::FVector3f*)Velocities.GetData(),
-						CollidingParticles.GetData(),
-						(const ispc::FIntVector*)CollidingElements.GetData(),
-						(const ispc::FIntVector*)TriangleMesh->GetElements().GetData(),
-						reinterpret_cast<const ispc::FVector2f&>(Thickness.GetOffsetRange()),
-						Thickness.GetMapValues().GetData(),
-						Dt,
-						ColliderThickness,
-						KinematicColliderFalloffMultiplier,
-						KinematicColliderMaxDepthMultiplier,
-						Stiffness,
-						FrictionCoefficient.HasWeightMap(),
-						reinterpret_cast<const ispc::FVector2f&>(FrictionCoefficient.GetOffsetRange()),
-						FrictionCoefficient.GetMapValues().GetData(),
-						CollidingParticles.Num()
-					);
+					if (PAndInvM.IsEmpty())
+					{
+						ispc::ApplyKinematicTriangleCollisionsWithFrictionAndMaps(
+							(ispc::FVector4f*)Particles.GetPAndInvM().GetData(),
+							(const ispc::FVector3f*)Particles.XArray().GetData(),
+							(const ispc::FVector3f*)Positions.GetData(),
+							(const ispc::FVector3f*)Velocities.GetData(),
+							CollidingParticles.GetData(),
+							(const ispc::FIntVector*)CollidingElements.GetData(),
+							(const ispc::FIntVector*)TriangleMesh->GetElements().GetData(),
+							reinterpret_cast<const ispc::FVector2f&>(Thickness.GetOffsetRange()),
+							Thickness.GetMapValues().GetData(),
+							Dt,
+							ColliderThickness,
+							KinematicColliderFalloffMultiplier,
+							KinematicColliderMaxDepthMultiplier,
+							Stiffness,
+							FrictionCoefficient.HasWeightMap(),
+							reinterpret_cast<const ispc::FVector2f&>(FrictionCoefficient.GetOffsetRange()),
+							FrictionCoefficient.GetMapValues().GetData(),
+							CollidingParticles.Num()
+						);
+					}
+					else
+					{
+						ispc::ApplyKinematicTriangleCollisionsWithFrictionAndMapsPAndInvM(
+							(ispc::FVector4f*)Particles.GetPAndInvM().GetData(),
+							(const ispc::FVector3f*)Particles.XArray().GetData(),
+							(const ispc::FVector4f*)PAndInvM.GetData(),
+							(const ispc::FVector3f*)Velocities.GetData(),
+							CollidingParticles.GetData(),
+							(const ispc::FIntVector*)CollidingElements.GetData(),
+							(const ispc::FIntVector*)TriangleMesh->GetElements().GetData(),
+							reinterpret_cast<const ispc::FVector2f&>(Thickness.GetOffsetRange()),
+							Thickness.GetMapValues().GetData(),
+							Dt,
+							ColliderThickness,
+							KinematicColliderFalloffMultiplier,
+							KinematicColliderMaxDepthMultiplier,
+							Stiffness,
+							FrictionCoefficient.HasWeightMap(),
+							reinterpret_cast<const ispc::FVector2f&>(FrictionCoefficient.GetOffsetRange()),
+							FrictionCoefficient.GetMapValues().GetData(),
+							CollidingParticles.Num()
+						);
+					}
 				}
 				else
 				{
-					ispc::ApplyKinematicTriangleCollisionsWithMaps(
-						(ispc::FVector4f*)Particles.GetPAndInvM().GetData(),
-						(const ispc::FVector3f*)Positions.GetData(),
-						CollidingParticles.GetData(),
-						(const ispc::FIntVector*)CollidingElements.GetData(),
-						(const ispc::FIntVector*)TriangleMesh->GetElements().GetData(),
-						reinterpret_cast<const ispc::FVector2f&>(Thickness.GetOffsetRange()),
-						Thickness.GetMapValues().GetData(),
-						ColliderThickness,
-						KinematicColliderFalloffMultiplier,
-						KinematicColliderMaxDepthMultiplier,
-						Stiffness,
-						CollidingParticles.Num()
-					);
+					if (PAndInvM.IsEmpty())
+					{
+						ispc::ApplyKinematicTriangleCollisionsWithMaps(
+							(ispc::FVector4f*)Particles.GetPAndInvM().GetData(),
+							(const ispc::FVector3f*)Positions.GetData(),
+							CollidingParticles.GetData(),
+							(const ispc::FIntVector*)CollidingElements.GetData(),
+							(const ispc::FIntVector*)TriangleMesh->GetElements().GetData(),
+							reinterpret_cast<const ispc::FVector2f&>(Thickness.GetOffsetRange()),
+							Thickness.GetMapValues().GetData(),
+							ColliderThickness,
+							KinematicColliderFalloffMultiplier,
+							KinematicColliderMaxDepthMultiplier,
+							Stiffness,
+							CollidingParticles.Num()
+						);
+					}
+					else
+					{
+						ispc::ApplyKinematicTriangleCollisionsWithMapsPAndInvM(
+							(ispc::FVector4f*)Particles.GetPAndInvM().GetData(),
+							(const ispc::FVector4f*)PAndInvM.GetData(),
+							CollidingParticles.GetData(),
+							(const ispc::FIntVector*)CollidingElements.GetData(),
+							(const ispc::FIntVector*)TriangleMesh->GetElements().GetData(),
+							reinterpret_cast<const ispc::FVector2f&>(Thickness.GetOffsetRange()),
+							Thickness.GetMapValues().GetData(),
+							ColliderThickness,
+							KinematicColliderFalloffMultiplier,
+							KinematicColliderMaxDepthMultiplier,
+							Stiffness,
+							CollidingParticles.Num()
+						);
+					}
 				}
 			}
 			else
@@ -206,58 +251,122 @@ namespace Chaos::Softs {
 				const FSolverReal MaxDepth = -Height * KinematicColliderMaxDepthMultiplier;
 				if (FrictionCoefficient.HasWeightMap())
 				{
-					ispc::ApplyKinematicTriangleCollisionsWithFrictionMap(
-						(ispc::FVector4f*)Particles.GetPAndInvM().GetData(),
-						(const ispc::FVector3f*)Particles.XArray().GetData(),
-						(const ispc::FVector3f*)Positions.GetData(),
-						(const ispc::FVector3f*)Velocities.GetData(),
-						CollidingParticles.GetData(),
-						(const ispc::FIntVector*)CollidingElements.GetData(),
-						(const ispc::FIntVector*)TriangleMesh->GetElements().GetData(),
-						Dt,
-						Height,
-						OneOverTangentialFalloffDist,
-						MaxDepth,
-						Stiffness,
-						reinterpret_cast<const ispc::FVector2f&>(FrictionCoefficient.GetOffsetRange()),
-						FrictionCoefficient.GetMapValues().GetData(),
-						CollidingParticles.Num()
-					);
+					if (PAndInvM.IsEmpty())
+					{
+						ispc::ApplyKinematicTriangleCollisionsWithFrictionMap(
+							(ispc::FVector4f*)Particles.GetPAndInvM().GetData(),
+							(const ispc::FVector3f*)Particles.XArray().GetData(),
+							(const ispc::FVector3f*)Positions.GetData(),
+							(const ispc::FVector3f*)Velocities.GetData(),
+							CollidingParticles.GetData(),
+							(const ispc::FIntVector*)CollidingElements.GetData(),
+							(const ispc::FIntVector*)TriangleMesh->GetElements().GetData(),
+							Dt,
+							Height,
+							OneOverTangentialFalloffDist,
+							MaxDepth,
+							Stiffness,
+							reinterpret_cast<const ispc::FVector2f&>(FrictionCoefficient.GetOffsetRange()),
+							FrictionCoefficient.GetMapValues().GetData(),
+							CollidingParticles.Num()
+						);
+					}
+					else
+					{
+						ispc::ApplyKinematicTriangleCollisionsWithFrictionMapPAndInvM(
+							(ispc::FVector4f*)Particles.GetPAndInvM().GetData(),
+							(const ispc::FVector3f*)Particles.XArray().GetData(),
+							(const ispc::FVector4f*)PAndInvM.GetData(),
+							(const ispc::FVector3f*)Velocities.GetData(),
+							CollidingParticles.GetData(),
+							(const ispc::FIntVector*)CollidingElements.GetData(),
+							(const ispc::FIntVector*)TriangleMesh->GetElements().GetData(),
+							Dt,
+							Height,
+							OneOverTangentialFalloffDist,
+							MaxDepth,
+							Stiffness,
+							reinterpret_cast<const ispc::FVector2f&>(FrictionCoefficient.GetOffsetRange()),
+							FrictionCoefficient.GetMapValues().GetData(),
+							CollidingParticles.Num()
+						);
+					}
 				}
 				else if (bWithFriction)
 				{
-					ispc::ApplyKinematicTriangleCollisionsWithFriction(
-						(ispc::FVector4f*)Particles.GetPAndInvM().GetData(),
-						(const ispc::FVector3f*)Particles.XArray().GetData(),
-						(const ispc::FVector3f*)Positions.GetData(),
-						(const ispc::FVector3f*)Velocities.GetData(),
-						CollidingParticles.GetData(),
-						(const ispc::FIntVector*)CollidingElements.GetData(),
-						(const ispc::FIntVector*)TriangleMesh->GetElements().GetData(),
-						Dt,
-						Height,
-						OneOverTangentialFalloffDist,
-						MaxDepth,
-						Stiffness,
-						(FSolverReal)FrictionCoefficient,
-						CollidingParticles.Num()
-					);
+					if (PAndInvM.IsEmpty())
+					{
+						ispc::ApplyKinematicTriangleCollisionsWithFriction(
+							(ispc::FVector4f*)Particles.GetPAndInvM().GetData(),
+							(const ispc::FVector3f*)Particles.XArray().GetData(),
+							(const ispc::FVector3f*)Positions.GetData(),
+							(const ispc::FVector3f*)Velocities.GetData(),
+							CollidingParticles.GetData(),
+							(const ispc::FIntVector*)CollidingElements.GetData(),
+							(const ispc::FIntVector*)TriangleMesh->GetElements().GetData(),
+							Dt,
+							Height,
+							OneOverTangentialFalloffDist,
+							MaxDepth,
+							Stiffness,
+							(FSolverReal)FrictionCoefficient,
+							CollidingParticles.Num()
+						);
+					}
+					else
+					{
+						ispc::ApplyKinematicTriangleCollisionsWithFrictionPAndInvM(
+							(ispc::FVector4f*)Particles.GetPAndInvM().GetData(),
+							(const ispc::FVector3f*)Particles.XArray().GetData(),
+							(const ispc::FVector4f*)PAndInvM.GetData(),
+							(const ispc::FVector3f*)Velocities.GetData(),
+							CollidingParticles.GetData(),
+							(const ispc::FIntVector*)CollidingElements.GetData(),
+							(const ispc::FIntVector*)TriangleMesh->GetElements().GetData(),
+							Dt,
+							Height,
+							OneOverTangentialFalloffDist,
+							MaxDepth,
+							Stiffness,
+							(FSolverReal)FrictionCoefficient,
+							CollidingParticles.Num()
+						);
+					}
 				}
 				else
 				{
-					// No friction, no maps
-					ispc::ApplyKinematicTriangleCollisions(
-						(ispc::FVector4f*)Particles.GetPAndInvM().GetData(),
-						(const ispc::FVector3f*)Positions.GetData(),
-						CollidingParticles.GetData(),
-						(const ispc::FIntVector*)CollidingElements.GetData(),
-						(const ispc::FIntVector*)TriangleMesh->GetElements().GetData(),
-						Height,
-						OneOverTangentialFalloffDist,
-						MaxDepth,
-						Stiffness,
-						CollidingParticles.Num()
-					);
+					if (PAndInvM.IsEmpty())
+					{
+						// No friction, no maps
+						ispc::ApplyKinematicTriangleCollisions(
+							(ispc::FVector4f*)Particles.GetPAndInvM().GetData(),
+							(const ispc::FVector3f*)Positions.GetData(),
+							CollidingParticles.GetData(),
+							(const ispc::FIntVector*)CollidingElements.GetData(),
+							(const ispc::FIntVector*)TriangleMesh->GetElements().GetData(),
+							Height,
+							OneOverTangentialFalloffDist,
+							MaxDepth,
+							Stiffness,
+							CollidingParticles.Num()
+						);
+					}
+					else
+					{
+						// No friction, no maps
+						ispc::ApplyKinematicTriangleCollisionsPAndInvM(
+							(ispc::FVector4f*)Particles.GetPAndInvM().GetData(),
+							(const ispc::FVector4f*)PAndInvM.GetData(),
+							CollidingParticles.GetData(),
+							(const ispc::FIntVector*)CollidingElements.GetData(),
+							(const ispc::FIntVector*)TriangleMesh->GetElements().GetData(),
+							Height,
+							OneOverTangentialFalloffDist,
+							MaxDepth,
+							Stiffness,
+							CollidingParticles.Num()
+						);
+					}
 				}
 			}
 		}
diff --git a/Engine/Source/Runtime/Experimental/Chaos/Private/Chaos/PBDKinematicTriangleMeshCollisions.ispc b/Engine/Source/Runtime/Experimental/Chaos/Private/Chaos/PBDKinematicTriangleMeshCollisions.ispc
--- a/Engine/Source/Runtime/Experimental/Chaos/Private/Chaos/PBDKinematicTriangleMeshCollisions.ispc
+++ b/Engine/Source/Runtime/Experimental/Chaos/Private/Chaos/PBDKinematicTriangleMeshCollisions.ispc
@@ -619,3 +619,478 @@ export void ApplyKinematicTriangleCollisionsWithFrictionAndMaps(
 		VectorScatter(&PandInvM[Index1], SetVector4(P1, IM));
 	}
 }
+
+
+export void ApplyKinematicTriangleCollisionsPAndInvM(
+			uniform FVector4f PandInvM[],
+			const uniform FVector4f KinematicPositions[],
+			const uniform int32 KinematicCollidingParticles[],
+			const uniform FIntVector KinematicColliderElements[],
+			const uniform FIntVector TriangleMeshElements[],
+			const uniform float Height,
+			const uniform float OneOverTangentialFalloffDist,
+			const uniform float MaxDepth,
+			const uniform float KinematicColliderStiffness,
+			const uniform int32 NumConstraints)
+{
+	foreach(Index = 0...NumConstraints)
+	{
+		const int32 Index1 = KinematicCollidingParticles[Index];
+
+		const FVector4f PandInvM1 = VectorGather(&PandInvM[Index1]);
+		varying FVector3f P1;
+		varying float IM;
+		UnzipPandInvM(PandInvM1, P1, IM);
+
+		const varying FIntVector Elements = VectorLoad(&KinematicColliderElements[extract(Index,0)]);
+
+		for(uniform int32 EIndex = 0; EIndex < 3; ++EIndex)
+		{
+			const varying int32 ElementIndex = Elements.V[EIndex];
+			if(ElementIndex == -1)
+			{
+				break;
+			}
+
+			const FIntVector Element = VectorGather(&TriangleMeshElements[ElementIndex]);
+
+			const varying int32 Index2 = Element.V[0];
+			const varying int32 Index3 = Element.V[1];
+			const varying int32 Index4 = Element.V[2];
+		
+			const varying FVector4f PandInvM2 = VectorGather(&KinematicPositions[Index2]);
+			const varying FVector4f PandInvM3 = VectorGather(&KinematicPositions[Index3]);
+			const varying FVector4f PandInvM4 = VectorGather(&KinematicPositions[Index4]);
+
+			varying FVector3f P2, P3, P4;
+			varying float IM2, IM3, IM4;
+			UnzipPandInvM(PandInvM2, P2, IM2);
+			UnzipPandInvM(PandInvM3, P3, IM3);
+			UnzipPandInvM(PandInvM4, P4, IM4);
+
+			const varying FVector3f Normal = GetTriangleNormalFlipped(P2, P3, P4);
+
+			varying FVector3f Bary;
+			const varying FVector3f P = FindClosestPointAndBaryOnTriangle(P2, P3, P4, P1, Bary);
+			const varying FVector3f Difference = P1 - P;
+			const varying float NormalDifference = VectorDot(Difference, Normal);
+			if(NormalDifference >= Height || NormalDifference < MaxDepth)
+			{
+				continue;
+			}
+
+			const varying float TangentialDifference = VectorSize(Difference - NormalDifference * Normal);
+			const uniform float TangentialFalloffMultiplier = 1.5f;
+			const varying float TangentialFalloff = 1.f - TangentialDifference * OneOverTangentialFalloffDist;
+			if(TangentialFalloff <= 0.f)
+			{
+				continue;
+			}
+
+			const float NormalDelta = Height - NormalDifference;
+			const FVector3f RepulsionDelta = KinematicColliderStiffness * TangentialFalloff * NormalDelta *  Normal;
+
+			P1 = P1 + RepulsionDelta;
+		}
+
+		VectorScatter(&PandInvM[Index1], SetVector4(P1, IM));
+	}
+}
+
+export void ApplyKinematicTriangleCollisionsWithFrictionPAndInvM(
+			uniform FVector4f PandInvM[],
+			const uniform FVector3f ParticlesX[],
+			const uniform FVector4f KinematicPositions[],
+			const uniform FVector3f KinematicVelocities[],
+			const uniform int32 KinematicCollidingParticles[],
+			const uniform FIntVector KinematicColliderElements[],
+			const uniform FIntVector TriangleMeshElements[],
+			const uniform float Dt,
+			const uniform float Height,
+			const uniform float OneOverTangentialFalloffDist,
+			const uniform float MaxDepth,
+			const uniform float KinematicColliderStiffness, 
+			const uniform float KinematicColliderFrictionCoefficient, 
+			const uniform int32 NumConstraints)
+{
+	foreach(Index = 0...NumConstraints)
+	{
+		const int32 Index1 = KinematicCollidingParticles[Index];
+
+		const FVector4f PandInvM1 = VectorGather(&PandInvM[Index1]);
+		varying FVector3f P1;
+		varying float IM;
+		UnzipPandInvM(PandInvM1, P1, IM);
+
+		const varying FIntVector Elements = VectorLoad(&KinematicColliderElements[extract(Index,0)]);
+
+		for(uniform int32 EIndex = 0; EIndex < 3; ++EIndex)
+		{
+			const varying int32 ElementIndex = Elements.V[EIndex];
+			if(ElementIndex == -1)
+			{
+				break;
+			}
+
+			const FIntVector Element = VectorGather(&TriangleMeshElements[ElementIndex]);
+
+			const varying int32 Index2 = Element.V[0];
+			const varying int32 Index3 = Element.V[1];
+			const varying int32 Index4 = Element.V[2];
+		
+			const varying FVector4f PandInvM2 = VectorGather(&KinematicPositions[Index2]);
+			const varying FVector4f PandInvM3 = VectorGather(&KinematicPositions[Index3]);
+			const varying FVector4f PandInvM4 = VectorGather(&KinematicPositions[Index4]);
+
+			varying FVector3f P2, P3, P4;
+			varying float IM2, IM3, IM4;
+			UnzipPandInvM(PandInvM2, P2, IM2);
+			UnzipPandInvM(PandInvM3, P3, IM3);
+			UnzipPandInvM(PandInvM4, P4, IM4);
+
+			const varying FVector3f Normal = GetTriangleNormalFlipped(P2, P3, P4);
+
+			varying FVector3f Bary;
+			const varying FVector3f P = FindClosestPointAndBaryOnTriangle(P2, P3, P4, P1, Bary);
+			const varying FVector3f Difference = P1 - P;
+			const varying float NormalDifference = VectorDot(Difference, Normal);
+			if(NormalDifference >= Height || NormalDifference < MaxDepth)
+			{
+				continue;
+			}
+
+			const varying float TangentialDifference = VectorSize(Difference - NormalDifference * Normal);
+			const uniform float TangentialFalloffMultiplier = 1.5f;
+			const varying float TangentialFalloff = 1.f - TangentialDifference * OneOverTangentialFalloffDist;
+			if(TangentialFalloff <= 0.f)
+			{
+				continue;
+			}
+
+			const float NormalDelta = Height - NormalDifference;
+			const FVector3f RepulsionDelta = KinematicColliderStiffness * TangentialFalloff * NormalDelta *  Normal;
+
+			P1 = P1 + RepulsionDelta;
+
+			const FVector3f X1 = VectorGather(&ParticlesX[Index1]);
+			const FVector3f V2 = VectorGather(&KinematicVelocities[Index2]);
+			const FVector3f V3 = VectorGather(&KinematicVelocities[Index3]);
+			const FVector3f V4 = VectorGather(&KinematicVelocities[Index4]);
+			const FVector3f V = Bary.V[0] * V2 + Bary.V[1] * V3 + Bary.V[2] * V4;
+			const FVector3f RelativeDisplacement = (P1 - X1) - V * Dt;
+			const FVector3f RelativeDisplacementTangent = RelativeDisplacement - VectorDot(RelativeDisplacement,Normal) * Normal;
+			const float RelativeDisplacementTangentLength = VectorSize(RelativeDisplacementTangent);
+			const float PositionCorrection = min(NormalDelta * KinematicColliderFrictionCoefficient, RelativeDisplacementTangentLength);
+			const float CorrectionRatio = RelativeDisplacementTangentLength < FLOAT_SMALL_NUMBER ? 0.f : PositionCorrection / RelativeDisplacementTangentLength;
+			const FVector3f FrictionDelta = -CorrectionRatio * RelativeDisplacementTangent;
+
+			P1 = P1 + FrictionDelta;
+		}
+
+		VectorScatter(&PandInvM[Index1], SetVector4(P1, IM));
+	}
+}
+
+export void ApplyKinematicTriangleCollisionsWithFrictionMapPAndInvM(
+			uniform FVector4f PandInvM[],
+			const uniform FVector3f ParticlesX[],
+			const uniform FVector4f KinematicPositions[],
+			const uniform FVector3f KinematicVelocities[],
+			const uniform int32 KinematicCollidingParticles[],
+			const uniform FIntVector KinematicColliderElements[],
+			const uniform FIntVector TriangleMeshElements[],
+			const uniform float Dt,
+			const uniform float Height,
+			const uniform float OneOverTangentialFalloffDist,
+			const uniform float MaxDepth,
+			const uniform float KinematicColliderStiffness, 
+			const uniform FVector2f& KinematicColliderFrictionCoefficientOffsetRange, 
+			const uniform float KinematicColliderFrictionCoefficientMapValues[],
+			const uniform int32 NumConstraints)
+{
+	foreach(Index = 0...NumConstraints)
+	{
+		const int32 Index1 = KinematicCollidingParticles[Index];
+
+		const FVector4f PandInvM1 = VectorGather(&PandInvM[Index1]);
+		varying FVector3f P1;
+		varying float IM;
+		UnzipPandInvM(PandInvM1, P1, IM);
+
+		const varying FIntVector Elements = VectorLoad(&KinematicColliderElements[extract(Index,0)]);
+		
+		#pragma ignore warning(perf)
+		const varying float FrictionMapValue = KinematicColliderFrictionCoefficientMapValues[Index1];
+		const varying float FrictionCoefficient = KinematicColliderFrictionCoefficientOffsetRange.V[0] + FrictionMapValue * KinematicColliderFrictionCoefficientOffsetRange.V[1];
+
+		for(uniform int32 EIndex = 0; EIndex < 3; ++EIndex)
+		{
+			const varying int32 ElementIndex = Elements.V[EIndex];
+			if(ElementIndex == -1)
+			{
+				break;
+			}
+
+			const FIntVector Element = VectorGather(&TriangleMeshElements[ElementIndex]);
+
+			const varying int32 Index2 = Element.V[0];
+			const varying int32 Index3 = Element.V[1];
+			const varying int32 Index4 = Element.V[2];
+		
+			const varying FVector4f PandInvM2 = VectorGather(&KinematicPositions[Index2]);
+			const varying FVector4f PandInvM3 = VectorGather(&KinematicPositions[Index3]);
+			const varying FVector4f PandInvM4 = VectorGather(&KinematicPositions[Index4]);
+
+			varying FVector3f P2, P3, P4;
+			varying float IM2, IM3, IM4;
+			UnzipPandInvM(PandInvM2, P2, IM2);
+			UnzipPandInvM(PandInvM3, P3, IM3);
+			UnzipPandInvM(PandInvM4, P4, IM4);
+
+			const varying FVector3f Normal = GetTriangleNormalFlipped(P2, P3, P4);
+
+			varying FVector3f Bary;
+			const varying FVector3f P = FindClosestPointAndBaryOnTriangle(P2, P3, P4, P1, Bary);
+			const varying FVector3f Difference = P1 - P;
+			const varying float NormalDifference = VectorDot(Difference, Normal);
+			if(NormalDifference >= Height || NormalDifference < MaxDepth)
+			{
+				continue;
+			}
+
+			const varying float TangentialDifference = VectorSize(Difference - NormalDifference * Normal);
+			const uniform float TangentialFalloffMultiplier = 1.5f;
+			const varying float TangentialFalloff = 1.f - TangentialDifference * OneOverTangentialFalloffDist;
+			if(TangentialFalloff <= 0.f)
+			{
+				continue;
+			}
+
+			const float NormalDelta = Height - NormalDifference;
+			const FVector3f RepulsionDelta = KinematicColliderStiffness * TangentialFalloff * NormalDelta *  Normal;
+
+			P1 = P1 + RepulsionDelta;
+
+			if (FrictionCoefficient > 0)
+			{
+				const FVector3f X1 = VectorGather(&ParticlesX[Index1]);
+				const FVector3f V2 = VectorGather(&KinematicVelocities[Index2]);
+				const FVector3f V3 = VectorGather(&KinematicVelocities[Index3]);
+				const FVector3f V4 = VectorGather(&KinematicVelocities[Index4]);
+				const FVector3f V = Bary.V[0] * V2 + Bary.V[1] * V3 + Bary.V[2] * V4;
+				const FVector3f RelativeDisplacement = (P1 - X1) - V * Dt;
+				const FVector3f RelativeDisplacementTangent = RelativeDisplacement - VectorDot(RelativeDisplacement,Normal) * Normal;
+				const float RelativeDisplacementTangentLength = VectorSize(RelativeDisplacementTangent);
+				const float PositionCorrection = min(NormalDelta * FrictionCoefficient, RelativeDisplacementTangentLength);
+				const float CorrectionRatio = RelativeDisplacementTangentLength < FLOAT_SMALL_NUMBER ? 0.f : PositionCorrection / RelativeDisplacementTangentLength;
+				const FVector3f FrictionDelta = -CorrectionRatio * RelativeDisplacementTangent;
+
+				P1 = P1 + FrictionDelta;
+			}
+		}
+
+		VectorScatter(&PandInvM[Index1], SetVector4(P1, IM));
+	}
+}
+
+export void ApplyKinematicTriangleCollisionsWithMapsPAndInvM(
+			uniform FVector4f PandInvM[],
+			const uniform FVector4f KinematicPositions[],
+			const uniform int32 KinematicCollidingParticles[],
+			const uniform FIntVector KinematicColliderElements[],
+			const uniform FIntVector TriangleMeshElements[],
+			const uniform FVector2f& ParticleThicknessOffsetRange,
+			const uniform float ParticleThicknessMapValues[],
+			const uniform float KinematicColliderThickness,
+			const uniform float KinematicColliderFalloffMultiplier,
+			const uniform float KinematicColliderMaxDepthMultiplier,
+			const uniform float KinematicColliderStiffness,
+			const uniform int32 NumConstraints)
+{
+	foreach(Index = 0...NumConstraints)
+	{
+		const int32 Index1 = KinematicCollidingParticles[Index];
+
+		const FVector4f PandInvM1 = VectorGather(&PandInvM[Index1]);
+		varying FVector3f P1;
+		varying float IM;
+		UnzipPandInvM(PandInvM1, P1, IM);
+
+		const varying FIntVector Elements = VectorLoad(&KinematicColliderElements[extract(Index,0)]);
+		
+		#pragma ignore warning(perf)
+		const varying float ThicknessMapValue = ParticleThicknessMapValues[Index1];
+		const varying float ParticleThickness = ParticleThicknessOffsetRange.V[0] + ThicknessMapValue * ParticleThicknessOffsetRange.V[1];
+		const float Height = ParticleThickness + KinematicColliderThickness;
+		const float OneOverTangentialFalloffDist = 1.f / max(KinematicColliderFalloffMultiplier * Height, FLOAT_KINDA_SMALL_NUMBER);
+		const float MaxDepth = -Height * KinematicColliderMaxDepthMultiplier;
+
+		for(uniform int32 EIndex = 0; EIndex < 3; ++EIndex)
+		{
+			const varying int32 ElementIndex = Elements.V[EIndex];
+			if(ElementIndex == -1)
+			{
+				break;
+			}
+
+			const FIntVector Element = VectorGather(&TriangleMeshElements[ElementIndex]);
+
+			const varying int32 Index2 = Element.V[0];
+			const varying int32 Index3 = Element.V[1];
+			const varying int32 Index4 = Element.V[2];
+		
+			const varying FVector4f PandInvM2 = VectorGather(&KinematicPositions[Index2]);
+			const varying FVector4f PandInvM3 = VectorGather(&KinematicPositions[Index3]);
+			const varying FVector4f PandInvM4 = VectorGather(&KinematicPositions[Index4]);
+
+			varying FVector3f P2, P3, P4;
+			varying float IM2, IM3, IM4;
+			UnzipPandInvM(PandInvM2, P2, IM2);
+			UnzipPandInvM(PandInvM3, P3, IM3);
+			UnzipPandInvM(PandInvM4, P4, IM4);
+
+			const varying FVector3f Normal = GetTriangleNormalFlipped(P2, P3, P4);
+
+			varying FVector3f Bary;
+			const varying FVector3f P = FindClosestPointAndBaryOnTriangle(P2, P3, P4, P1, Bary);
+			const varying FVector3f Difference = P1 - P;
+			const varying float NormalDifference = VectorDot(Difference, Normal);
+			if(NormalDifference >= Height || NormalDifference < MaxDepth)
+			{
+				continue;
+			}
+
+			const varying float TangentialDifference = VectorSize(Difference - NormalDifference * Normal);
+			const uniform float TangentialFalloffMultiplier = 1.5f;
+			const varying float TangentialFalloff = 1.f - TangentialDifference * OneOverTangentialFalloffDist;
+			if(TangentialFalloff <= 0.f)
+			{
+				continue;
+			}
+
+			const float NormalDelta = Height - NormalDifference;
+			const FVector3f RepulsionDelta = KinematicColliderStiffness * TangentialFalloff * NormalDelta *  Normal;
+
+			P1 = P1 + RepulsionDelta;
+		}
+
+		VectorScatter(&PandInvM[Index1], SetVector4(P1, IM));
+	}
+}
+
+export void ApplyKinematicTriangleCollisionsWithFrictionAndMapsPAndInvM(
+			uniform FVector4f PandInvM[],
+			const uniform FVector3f ParticlesX[],
+			const uniform FVector4f KinematicPositions[],
+			const uniform FVector3f KinematicVelocities[],
+			const uniform int32 KinematicCollidingParticles[],
+			const uniform FIntVector KinematicColliderElements[],
+			const uniform FIntVector TriangleMeshElements[],
+			const uniform FVector2f& ParticleThicknessOffsetRange,
+			const uniform float ParticleThicknessMapValues[],
+			const uniform float Dt,
+			const uniform float KinematicColliderThickness,
+			const uniform float KinematicColliderFalloffMultiplier,
+			const uniform float KinematicColliderMaxDepthMultiplier,
+			const uniform float KinematicColliderStiffness, 
+			const uniform bool bKinematicColliderFrictionHasMap,
+			const uniform FVector2f& KinematicColliderFrictionCoefficientOffsetRange, 
+			const uniform float KinematicColliderFrictionCoefficientMapValues[],
+			const uniform int32 NumConstraints)
+{
+	foreach(Index = 0...NumConstraints)
+	{
+		const int32 Index1 = KinematicCollidingParticles[Index];
+
+		const FVector4f PandInvM1 = VectorGather(&PandInvM[Index1]);
+		varying FVector3f P1;
+		varying float IM;
+		UnzipPandInvM(PandInvM1, P1, IM);
+
+		const varying FIntVector Elements = VectorLoad(&KinematicColliderElements[extract(Index,0)]);
+		
+		#pragma ignore warning(perf)
+		const varying float ThicknessMapValue = ParticleThicknessMapValues[Index1];
+		const varying float ParticleThickness = ParticleThicknessOffsetRange.V[0] + ThicknessMapValue * ParticleThicknessOffsetRange.V[1];
+		const float Height = ParticleThickness + KinematicColliderThickness;
+		const float OneOverTangentialFalloffDist = 1.f / max(KinematicColliderFalloffMultiplier * Height, FLOAT_KINDA_SMALL_NUMBER);
+		const float MaxDepth = -Height * KinematicColliderMaxDepthMultiplier;
+
+		varying float FrictionCoefficient = KinematicColliderFrictionCoefficientOffsetRange.V[0];
+		if(bKinematicColliderFrictionHasMap)
+		{
+			#pragma ignore warning(perf)
+			const varying float FrictionMapValue = KinematicColliderFrictionCoefficientMapValues[Index1];
+			FrictionCoefficient += FrictionMapValue * KinematicColliderFrictionCoefficientOffsetRange.V[1];
+		}
+
+		for(uniform int32 EIndex = 0; EIndex < 3; ++EIndex)
+		{
+			const varying int32 ElementIndex = Elements.V[EIndex];
+			if(ElementIndex == -1)
+			{
+				break;
+			}
+
+			const FIntVector Element = VectorGather(&TriangleMeshElements[ElementIndex]);
+
+			const varying int32 Index2 = Element.V[0];
+			const varying int32 Index3 = Element.V[1];
+			const varying int32 Index4 = Element.V[2];
+		
+			const varying FVector4f PandInvM2 = VectorGather(&KinematicPositions[Index2]);
+			const varying FVector4f PandInvM3 = VectorGather(&KinematicPositions[Index3]);
+			const varying FVector4f PandInvM4 = VectorGather(&KinematicPositions[Index4]);
+
+			varying FVector3f P2, P3, P4;
+			varying float IM2, IM3, IM4;
+			UnzipPandInvM(PandInvM2, P2, IM2);
+			UnzipPandInvM(PandInvM3, P3, IM3);
+			UnzipPandInvM(PandInvM4, P4, IM4);
+
+			const varying FVector3f Normal = GetTriangleNormalFlipped(P2, P3, P4);
+
+			varying FVector3f Bary;
+			const varying FVector3f P = FindClosestPointAndBaryOnTriangle(P2, P3, P4, P1, Bary);
+			const varying FVector3f Difference = P1 - P;
+			const varying float NormalDifference = VectorDot(Difference, Normal);
+			if(NormalDifference >= Height || NormalDifference < MaxDepth)
+			{
+				continue;
+			}
+
+			const varying float TangentialDifference = VectorSize(Difference - NormalDifference * Normal);
+			const uniform float TangentialFalloffMultiplier = 1.5f;
+			const varying float TangentialFalloff = 1.f - TangentialDifference * OneOverTangentialFalloffDist;
+			if(TangentialFalloff <= 0.f)
+			{
+				continue;
+			}
+
+			const float NormalDelta = Height - NormalDifference;
+			const FVector3f RepulsionDelta = KinematicColliderStiffness * TangentialFalloff * NormalDelta *  Normal;
+
+			P1 = P1 + RepulsionDelta;
+
+
+			if (FrictionCoefficient > 0)
+			{
+				const FVector3f X1 = VectorGather(&ParticlesX[Index1]);
+				const FVector3f V2 = VectorGather(&KinematicVelocities[Index2]);
+				const FVector3f V3 = VectorGather(&KinematicVelocities[Index3]);
+				const FVector3f V4 = VectorGather(&KinematicVelocities[Index4]);
+				const FVector3f V = Bary.V[0] * V2 + Bary.V[1] * V3 + Bary.V[2] * V4;
+				const FVector3f RelativeDisplacement = (P1 - X1) - V * Dt;
+				const FVector3f RelativeDisplacementTangent = RelativeDisplacement - VectorDot(RelativeDisplacement,Normal) * Normal;
+				const float RelativeDisplacementTangentLength = VectorSize(RelativeDisplacementTangent);
+				const float PositionCorrection = min(NormalDelta * FrictionCoefficient, RelativeDisplacementTangentLength);
+				const float CorrectionRatio = RelativeDisplacementTangentLength < FLOAT_SMALL_NUMBER ? 0.f : PositionCorrection / RelativeDisplacementTangentLength;
+				const FVector3f FrictionDelta = -CorrectionRatio * RelativeDisplacementTangent;
+
+				P1 = P1 + FrictionDelta;
+			}
+		}
+
+		VectorScatter(&PandInvM[Index1], SetVector4(P1, IM));
+	}
+}
diff --git a/Engine/Source/Runtime/Experimental/Chaos/Public/Chaos/PBDCollisionSpringConstraintsBase.h b/Engine/Source/Runtime/Experimental/Chaos/Public/Chaos/PBDCollisionSpringConstraintsBase.h
--- a/Engine/Source/Runtime/Experimental/Chaos/Public/Chaos/PBDCollisionSpringConstraintsBase.h
+++ b/Engine/Source/Runtime/Experimental/Chaos/Public/Chaos/PBDCollisionSpringConstraintsBase.h
@@ -91,6 +91,7 @@ public:
 
 	const TArray<int32>& GetKinematicCollidingParticles() const { return KinematicCollisions.GetCollidingParticles(); }
 	const TArray<TMap<int32, FSolverReal>>& GetKinematicColliderTimers() const { return KinematicCollisions.GetTimers(); }
+	const FTriangleMesh* GetKinematicColliderTriangleMesh() const { return KinematicCollisions.GetTriangleMesh(); }
 	const FTriangleMesh& GetTriangleMesh() const { return TriangleMesh; }
 
 	void SetThicknessWeighted(const FSolverVec2 InThickness) { ThicknessWeighted.SetWeightedValue(FSolverVec2::Max(InThickness, FSolverVec2(0.f))); }
diff --git a/Engine/Source/Runtime/Experimental/Chaos/Public/Chaos/PBDKinematicTriangleMeshCollisions.h b/Engine/Source/Runtime/Experimental/Chaos/Public/Chaos/PBDKinematicTriangleMeshCollisions.h
--- a/Engine/Source/Runtime/Experimental/Chaos/Public/Chaos/PBDKinematicTriangleMeshCollisions.h
+++ b/Engine/Source/Runtime/Experimental/Chaos/Public/Chaos/PBDKinematicTriangleMeshCollisions.h
@@ -43,6 +43,16 @@ namespace Chaos::Softs
 		{
 			TriangleMesh = &InTriangleMesh;
 			Positions = InPositions;
+			PAndInvM = TConstArrayView<FPAndInvM>();
+			Velocities = InVelocities;
+			SpatialHash = &InSpatialHash;
+		}
+
+		void SetGeometry(const FTriangleMesh& InTriangleMesh, const TConstArrayView<FSolverVec3>& InPositions, const TConstArrayView<FPAndInvM>& InPAndInvM, const TConstArrayView<FSolverVec3>& InVelocities, const FTriangleMesh::TSpatialHashType<FSolverReal>& InSpatialHash)
+		{
+			TriangleMesh = &InTriangleMesh;
+			Positions = InPositions;
+			PAndInvM = InPAndInvM;
 			Velocities = InVelocities;
 			SpatialHash = &InSpatialHash;
 		}
@@ -64,7 +74,7 @@ namespace Chaos::Softs
 		const TArray<int32>& GetCollidingParticles() const { return CollidingParticles; }
 		const TArray<TVector<int32, MaxKinematicConnectionsPerPoint>>& GetCollidingElements() const { return CollidingElements; }
 		const TArray<TMap<int32, FSolverReal>>& GetTimers() const { return Timers; }
-
+		const FTriangleMesh* GetTriangleMesh() const { return TriangleMesh; }
 	protected:
 		int32 NumParticles;
 		const FPBDFlatWeightMap& Thickness;
@@ -74,7 +84,8 @@ namespace Chaos::Softs
 		FSolverReal ColliderThickness;
 
 		const FTriangleMesh* TriangleMesh = nullptr;
-		TConstArrayView<FSolverVec3> Positions;
+		TConstArrayView<FSolverVec3> Positions; 
+		TConstArrayView<FPAndInvM> PAndInvM;
 		TConstArrayView<FSolverVec3> Velocities;
 		const FTriangleMesh::TSpatialHashType<FSolverReal>* SpatialHash = nullptr;
 
