Fix CL 43938016  not compiling in shipping editor config
- fix preprocessor conditions around some functions

diff --git a/Engine/Source/Runtime/Experimental/Dataflow/Engine/Private/Dataflow/DataflowEdNode.cpp b/Engine/Source/Runtime/Experimental/Dataflow/Engine/Private/Dataflow/DataflowEdNode.cpp
--- a/Engine/Source/Runtime/Experimental/Dataflow/Engine/Private/Dataflow/DataflowEdNode.cpp
+++ b/Engine/Source/Runtime/Experimental/Dataflow/Engine/Private/Dataflow/DataflowEdNode.cpp
@@ -368,7 +368,7 @@ void UDataflowEdNode::UpdatePinsConnectionsFromDataflowNode()
 #endif // WITH_EDITOR && !UE_BUILD_SHIPPING
 }
 
-#if WITH_EDITOR && !(UE_BUILD_SHIPPING || UE_BUILD_TEST)
+#if WITH_EDITOR
 bool UDataflowEdNode::SupportsEditablePinType(const UEdGraphPin& Pin)
 {
 	if (const FDataflowConnection* Connection = GetConnectionFromPin(&Pin))
@@ -390,7 +390,7 @@ bool UDataflowEdNode::SupportsEditablePinType(const UEdGraphPin& Pin)
 	}
 	return false;
 }
-#endif // WITH_EDITOR && !UE_BUILD_SHIPPING
+#endif // WITH_EDITOR
 
 void UDataflowEdNode::UpdatePinsDefaultValuesFromNode()
 {
diff --git a/Engine/Source/Runtime/Experimental/Dataflow/Engine/Public/Dataflow/DataflowEdNode.h b/Engine/Source/Runtime/Experimental/Dataflow/Engine/Public/Dataflow/DataflowEdNode.h
--- a/Engine/Source/Runtime/Experimental/Dataflow/Engine/Public/Dataflow/DataflowEdNode.h
+++ b/Engine/Source/Runtime/Experimental/Dataflow/Engine/Public/Dataflow/DataflowEdNode.h
@@ -119,7 +119,7 @@ public:
 	DATAFLOWENGINE_API static TSharedPtr<FDataflowNode> GetDataflowNodeFromEdNode(UEdGraphNode* EdNode);
 	DATAFLOWENGINE_API static TSharedPtr<const FDataflowNode> GetDataflowNodeFromEdNode(const UEdGraphNode* EdNode);
 
-#if WITH_EDITOR && !(UE_BUILD_SHIPPING || UE_BUILD_TEST)
+#if WITH_EDITOR
 	DATAFLOWENGINE_API static bool SupportsEditablePinType(const UEdGraphPin& Pin);
 #endif
 
