Fix block last access time when defragging

diff --git a/Engine/Source/Runtime/Experimental/IoStore/OnDemand/Private/OnDemandInstallCache.cpp b/Engine/Source/Runtime/Experimental/IoStore/OnDemand/Private/OnDemandInstallCache.cpp
--- a/Engine/Source/Runtime/Experimental/IoStore/OnDemand/Private/OnDemandInstallCache.cpp
+++ b/Engine/Source/Runtime/Experimental/IoStore/OnDemand/Private/OnDemandInstallCache.cpp
@@ -172,7 +172,8 @@ struct FCas
 
 	FIoStatus					Initialize(FStringView Directory, bool bDeleteExisting = false);
 	FCasLocation				FindChunk(const FIoHash& Hash) const;
-	FCasBlockId					CreateBlock();
+	FCasBlockId					CreateBlock(int64 UtcTicks);
+	FCasBlockId					CreateBlock() { return CreateBlock(FDateTime::UtcNow().GetTicks()); }
 	FIoStatus					DeleteBlock(FCasBlockId BlockId, TArray<FCasAddr>& OutAddrs);
 	FString						GetBlockFilename(FCasBlockId BlockId) const;
 	FSharedFileOpenResult		OpenRead(FCasBlockId BlockId);
@@ -256,7 +257,7 @@ FCasLocation FCas::FindChunk(const FIoHash& Hash) const
 	return FCasLocation{};
 }
 
-FCasBlockId FCas::CreateBlock()
+FCasBlockId FCas::CreateBlock(int64 UtcTicks)
 {
 	IPlatformFile&	Ipf = FPlatformFileManager::Get().GetPlatformFile();
 	FCasBlockId		Out = FCasBlockId::Invalid;
@@ -279,7 +280,7 @@ FCasBlockId FCas::CreateBlock()
 		}
 
 		BlockIds.Add(BlockId, 0);
-		LastAccess.FindOrAdd(BlockId, FDateTime::UtcNow().GetTicks());
+		LastAccess.FindOrAdd(BlockId, UtcTicks);
 		Out = BlockId;
 	}
 
@@ -2881,7 +2882,7 @@ FIoStatus FOnDemandInstallCache::Defrag(
 							*Status.ToString());
 					}
 
-					if (FIoStatus Status = FlushPendingChunks(DefragPendingChunks); Status.IsOk() == false)
+					if (FIoStatus Status = FlushPendingChunks(DefragPendingChunks, DefragBlock.LastAccess); Status.IsOk() == false)
 					{
 						FOnDemandInstallCacheStats::OnDefrag(Status.GetErrorCode(), 0);
 						return Status;
@@ -2906,7 +2907,7 @@ FIoStatus FOnDemandInstallCache::Defrag(
 
 			FileHandle.Reset();
 
-			if (FIoStatus Status = FlushPendingChunks(DefragPendingChunks); Status.IsOk() == false)
+			if (FIoStatus Status = FlushPendingChunks(DefragPendingChunks, DefragBlock.LastAccess); Status.IsOk() == false)
 			{
 				FOnDemandInstallCacheStats::OnDefrag(Status.GetErrorCode(), 0);
 				return Status;
@@ -3034,7 +3035,14 @@ FIoStatus FOnDemandInstallCache::FlushPendingChunksImpl(FPendingChunks& Chunks,
 		
 		if (CurrentBlockId.IsValid() == false)
 		{
-			CurrentBlockId = Cas.CreateBlock();
+			if (UtcAccessTicks)
+			{
+				CurrentBlockId = Cas.CreateBlock(UtcAccessTicks);
+			}
+			else
+			{
+				CurrentBlockId = Cas.CreateBlock();
+			}
 			ensure(CurrentBlockId.IsValid());
 			CurrentBlock = CurrentBlockId;
 			Transaction.BlockCreated(CurrentBlockId);
