Collision Filter Refactor Part 1:
Deprecated the old Query and Sim Filters fields.
Added new get/set methods for the filters.
Added the new shape filter and instance data structs.
Did minimal cleanup for deprecation.

diff --git a/Engine/Source/Programs/HeadlessChaos/Private/HeadlessChaosTestPhysicsScene.cpp b/Engine/Source/Programs/HeadlessChaos/Private/HeadlessChaosTestPhysicsScene.cpp
--- a/Engine/Source/Programs/HeadlessChaos/Private/HeadlessChaosTestPhysicsScene.cpp
+++ b/Engine/Source/Programs/HeadlessChaos/Private/HeadlessChaosTestPhysicsScene.cpp
@@ -488,10 +488,13 @@ namespace ChaosTest {
 		for (const TUniquePtr<Chaos::FPerShapeData>& Shape : Particle.ShapesArray())
 		{
 			const FCollisionData& CollisionData = Shape->GetCollisionData();
-			EXPECT_EQ(CollisionData.QueryData.Word0, 0); // ensure query filter is defaulted to 0 (no query allowed at all)
+			PRAGMA_DISABLE_INTERNAL_WARNINGS
+			const FCollisionFilterData QueryData = CollisionData.GetQueryData();
+			PRAGMA_ENABLE_INTERNAL_WARNINGS
+			EXPECT_EQ(QueryData.Word0, 0); // ensure query filter is defaulted to 0 (no query allowed at all)
 
 			// Verify query is filtered out with default collision data on shape
-			bool bFiltered = PrePreQueryFilterImp(QueryFilter, CollisionData.QueryData);
+			bool bFiltered = PrePreQueryFilterImp(QueryFilter, QueryData);
 			EXPECT_EQ(bFiltered, true);
 		}
 
@@ -509,16 +512,18 @@ namespace ChaosTest {
 		NewParticleQueryFilter.Word2 = TNumericLimits<int32>::Max();
 		for (const TUniquePtr<Chaos::FPerShapeData>& Shape : Particle.ShapesArray())
 		{
+			PRAGMA_DISABLE_INTERNAL_WARNINGS
 			const FCollisionData& CollisionData = Shape->GetCollisionData();
 
 			// Update filter
 			FCollisionData NewCollisionData = CollisionData;
-			NewCollisionData.QueryData = NewParticleQueryFilter;
+			NewCollisionData.SetQueryData(NewParticleQueryFilter);
 			Shape->SetCollisionData(NewCollisionData);
 
 			// Filter with new data, ensuring we pass and are not filtered out.
-			bool bFiltered = PrePreQueryFilterImp(QueryFilter, NewCollisionData.QueryData);
+			bool bFiltered = PrePreQueryFilterImp(QueryFilter, NewCollisionData.GetQueryData());
 			EXPECT_EQ(bFiltered, false);
+			PRAGMA_ENABLE_INTERNAL_WARNINGS
 		}
 
 		// Update particle in GT accel structure so cached PrePreFilter updates
diff --git a/Engine/Source/Runtime/Engine/Private/PhysicsEngine/ImmediatePhysics/ImmediatePhysicsChaos/ImmediatePhysicsActorHandle_Chaos.cpp b/Engine/Source/Runtime/Engine/Private/PhysicsEngine/ImmediatePhysics/ImmediatePhysicsChaos/ImmediatePhysicsActorHandle_Chaos.cpp
--- a/Engine/Source/Runtime/Engine/Private/PhysicsEngine/ImmediatePhysics/ImmediatePhysicsChaos/ImmediatePhysicsActorHandle_Chaos.cpp
+++ b/Engine/Source/Runtime/Engine/Private/PhysicsEngine/ImmediatePhysics/ImmediatePhysicsChaos/ImmediatePhysicsActorHandle_Chaos.cpp
@@ -88,9 +88,10 @@ namespace ImmediatePhysics_Chaos
 			// Set the collision filter data for the shapes to collide with everything.
 			// Even though we already tried to do this when we created the original shapes array, 
 			// that gets thrown away and we need to do it here. This is not a good API
+			Chaos::Filter::FShapeFilterData ShapeFilter = Chaos::Filter::FShapeFilterBuilder::BuildLegacyBlockAllSimShapeFilter();
+
 			FCollisionData CollisionData;
-			CollisionData.SimData.Word1 = 0xFFFFF;
-			CollisionData.SimData.Word3 = 0xFFFFF;
+			CollisionData.SetShapeFilterData(ShapeFilter);
 			CollisionData.bSimCollision = 1;
 			for (const TUniquePtr<FPerShapeData>& Shape : ParticleHandle->ShapesArray())
 			{
diff --git a/Engine/Source/Runtime/Experimental/Chaos/Private/Chaos/CollisionFilterData.cpp b/Engine/Source/Runtime/Experimental/Chaos/Private/Chaos/CollisionFilterData.cpp
new file mode 100644
--- /dev/null
+++ b/Engine/Source/Runtime/Experimental/Chaos/Private/Chaos/CollisionFilterData.cpp
@@ -0,0 +1,380 @@
+// Copyright Epic Games, Inc. All Rights Reserved.
+
+#include "Chaos/CollisionFilterData.h"
+#include "Chaos/ParticleHandle.h"
+#include "Chaos/Collision/CollisionFilter.h"
+
+namespace Chaos::Filter
+{
+	constexpr uint8 NumFlagsBits = 21;
+	constexpr uint8 NumCollisionChannelBits = 5;
+	constexpr uint8 NumMaskFilterBits = 6;
+	static_assert(NumFlagsBits + NumMaskFilterBits + NumCollisionChannelBits == 32);
+
+	constexpr uint32 CollisionChannelBitsOffset = NumFlagsBits;
+	constexpr uint32 MaskFilterBitsOffset = CollisionChannelBitsOffset + NumCollisionChannelBits;
+
+	constexpr uint32 FilterFlagsMask = (1 << NumFlagsBits) - 1;
+	constexpr uint32 CollisionChannelMask = ((1 << NumCollisionChannelBits) - 1) << CollisionChannelBitsOffset;
+	constexpr uint32 MaskFilterMask = ((1 << NumMaskFilterBits) - 1) << MaskFilterBitsOffset;
+
+	uint32 GetFlags(uint32 Word3)
+	{
+		const uint32 FilterFlags = (Word3 & FilterFlagsMask);
+		return FilterFlags;
+	}
+
+	void SetFlags(uint32& Word3, uint32 Flags)
+	{
+		const uint32 Word3Cleared = Word3 & ~FilterFlagsMask;
+		Word3 = Word3Cleared | (Flags | FilterFlagsMask);
+	}
+
+	bool HasFlag(uint32 Word3, EFilterFlags InFlag)
+	{
+		const uint32 FilterFlags = GetFlags(Word3);
+		return FilterFlags & static_cast<uint32>(InFlag);
+	}
+
+	uint8 GetCollisionChannelIndex(uint32 Word3)
+	{
+		const uint32 ChannelIndex = (Word3 & CollisionChannelMask) >> CollisionChannelBitsOffset;
+		return (uint8)ChannelIndex;
+	}
+
+	uint64 GetCollisionChannelMask64(uint32 Word3)
+	{
+		return (uint64)1 << (uint64)GetCollisionChannelIndex(Word3);
+	}
+
+	uint8 GetMaskFilter(uint32 Word3)
+	{
+		return Word3 >> (32u - NumMaskFilterBits);
+	}
+
+	void SetMaskFilter(uint32& Word3, uint8 MaskFilter)
+	{
+		static_assert(NumMaskFilterBits <= 8, "Only up to 8 extra filter bits are supported.");
+		// Clear the old mask filter bits
+		Word3 &= ~MaskFilterMask;
+		// Set the new mask filter bits
+		Word3 |= uint32(MaskFilter) << MaskFilterBitsOffset;
+	}
+
+	FInstanceData::FInstanceData(const uint32 ActorId, const uint32 ComponentId)
+		: ActorId(ActorId)
+		, ComponentId(ComponentId)
+	{
+	}
+
+	bool FInstanceData::IsValid() const
+	{
+		return ActorId != 0 || ComponentId != 0 || BodyIndex != 0;
+	}
+
+	uint32 FInstanceData::GetActorId() const
+	{
+		return ActorId;
+	}
+
+	void FInstanceData::SetActorId(const uint32 InActorId)
+	{
+		ActorId = InActorId;
+	}
+
+	uint32 FInstanceData::GetComponentId() const
+	{
+		return ComponentId;
+	}
+
+	void FInstanceData::SetComponentId(const uint32 InComponentId)
+	{
+		ComponentId = InComponentId;
+	}
+
+	bool FShapeFilterData::IsValid() const
+	{
+		return QueryBlockChannels != 0 || QueryOverlapChannels != 0 || SimBlockChannels != 0 || Word3 != 0 || SimWord3 != 0;
+	}
+
+	EFilterFlags FShapeFilterData::GetFlags() const
+	{
+		return (EFilterFlags)Chaos::Filter::GetFlags(Word3);
+	}
+
+	void FShapeFilterData::SetFlags(EFilterFlags InFlags)
+	{
+		Chaos::Filter::SetFlags(Word3, (uint32)InFlags);
+	}
+
+	bool FShapeFilterData::HasFlag(EFilterFlags InFlag) const
+	{
+		return Chaos::Filter::HasFlag(Word3, InFlag);
+	}
+
+	uint8 FShapeFilterData::GetMaskFilter() const
+	{
+		return Chaos::Filter::GetMaskFilter(Word3);
+	}
+
+	void FShapeFilterData::SetMaskFilter(uint8 MaskFilter)
+	{
+		Chaos::Filter::SetMaskFilter(Word3, MaskFilter);
+	}
+
+	uint8 FShapeFilterData::GetCollisionChannelIndex() const
+	{
+		return Chaos::Filter::GetCollisionChannelIndex(Word3);
+	}
+
+	uint64 FShapeFilterData::GetCollisionChannelMask() const
+	{
+		return Chaos::Filter::GetCollisionChannelMask64(Word3);
+	}
+
+	uint64 FShapeFilterData::GetQueryBlockChannels() const
+	{
+		return QueryBlockChannels;
+	}
+
+	uint64 FShapeFilterData::GetQueryOverlapChannels() const
+	{
+		return QueryOverlapChannels;
+	}
+
+	uint64 FShapeFilterData::GetSimBlockChannels() const
+	{
+		return SimBlockChannels;
+	}
+
+	bool FQueryFilterData::IsValid() const
+	{
+		return Word0 != 0 || Word1 != 0 || Word2 != 0 || Word3 != 0;
+	}
+
+	FQueryFilterData::EQueryType FQueryFilterData::GetQueryType() const
+	{
+		return (Word0 & 0x1) != 0 ? EQueryType::Channel : EQueryType::ObjectType;
+	}
+
+	EFilterFlags FQueryFilterData::GetFlags() const
+	{
+		return (EFilterFlags)Chaos::Filter::GetFlags(Word3);
+	}
+
+	void FQueryFilterData::SetFlags(EFilterFlags InFlags)
+	{
+		Chaos::Filter::SetFlags(Word3, (uint32)InFlags);
+	}
+
+	bool FQueryFilterData::HasFlag(EFilterFlags InFlag) const
+	{
+		return Chaos::Filter::HasFlag(Word3, InFlag);
+	}
+
+	uint8 FQueryFilterData::GetIgnoreMask() const
+	{
+		return GetMaskFilter(Word3);
+	}
+
+	uint8 FQueryFilterData::GetCollisionChannelIndex() const
+	{
+		return Chaos::Filter::GetCollisionChannelIndex(Word3);
+	}
+
+	uint64 FQueryFilterData::GetCollisionChannelMask() const
+	{
+		return Chaos::Filter::GetCollisionChannelMask64(Word3);
+	}
+
+	uint64 FQueryFilterData::GetBlockChannels() const
+	{
+		return Word1;
+	}
+
+	uint64 FQueryFilterData::GetOverlapChannels() const
+	{
+		return Word2;
+	}
+
+	uint64 FQueryFilterData::GetObjectTypesToQueryMask() const
+	{
+		return Word1;
+	}
+
+	bool FQueryFilterData::IsMultiQuery() const
+	{
+		const uint32 MultiTrace = GetCollisionChannelIndex();
+		return MultiTrace != 0;
+	}
+
+	void FQueryFilterData::SetMaskChannelAndFlags(const uint8 IgnoreMask, const uint8 ChannelIndex, const EFilterFlags FilterFlags)
+	{
+		Word3 = ((uint32)IgnoreMask << MaskFilterBitsOffset) | ((uint32)ChannelIndex << CollisionChannelBitsOffset) | (uint32)FilterFlags;
+	}
+
+	FCombinedShapeFilterData::FCombinedShapeFilterData(const FShapeFilterData& InShapeFilter, const FInstanceData& InInstanceData)
+		: ShapeFilterData(InShapeFilter)
+		, InstanceData(InInstanceData)
+	{
+	}
+
+	const FInstanceData& FCombinedShapeFilterData::GetInstanceData() const
+	{
+		return InstanceData;
+	}
+
+	void FCombinedShapeFilterData::SetInstanceData(const FInstanceData& InData)
+	{
+		InstanceData = InData;
+	}
+
+	const FShapeFilterData& FCombinedShapeFilterData::GetShapeFilterData() const
+	{
+		return ShapeFilterData;
+	}
+
+	void FCombinedShapeFilterData::SetShapeFilterData(const FShapeFilterData& InData)
+	{
+		ShapeFilterData = InData;
+	}
+
+	bool FCombinedShapeFilterData::IsValid() const
+	{
+		return ShapeFilterData.IsValid() || InstanceData.IsValid();
+	}
+
+	FShapeFilterData FShapeFilterBuilder::BuildLegacyShapeFilter(const uint8 InChannelIndex, const uint32 InQueryOverlapChannels, const uint32 InQueryBlockChannels, const uint32 InSimBlockChannels, const uint8 InMaskFilter, const EFilterFlags InFilterFlags)
+	{
+		// (Mask Filter 6 bits) | (Channel Index 5 bits) | (Flags 21 bits)
+		FShapeFilterData Result;
+		Result.QueryBlockChannels = InQueryBlockChannels;
+		Result.QueryOverlapChannels = InQueryOverlapChannels;
+		Result.SimBlockChannels = InSimBlockChannels;
+		Result.Word3 = FilterFlagsMask & (uint32)InFilterFlags;
+		Result.Word3 |= CollisionChannelMask & ((uint32)InChannelIndex << CollisionChannelBitsOffset);
+		Result.Word3 |= MaskFilterMask & ((uint32)InMaskFilter << MaskFilterBitsOffset);
+		Result.SimWord3 = Result.Word3;
+		return Result;
+	}
+
+	FShapeFilterData FShapeFilterBuilder::BuildLegacyBlockAllSimShapeFilter(const uint8 InChannelIndex, const uint8 InMaskFilter, const EFilterFlags InFilterFlags)
+	{
+		return BuildLegacyShapeFilter(InChannelIndex, 0, 0, (uint32)-1, 0, InFilterFlags);
+	}
+
+	FShapeFilterData FShapeFilterBuilder::BuildFromLegacyQueryFilter(const FCollisionFilterData& QueryFilter)
+	{
+		FShapeFilterData Result;
+		Result.QueryBlockChannels = QueryFilter.Word1;
+		Result.QueryOverlapChannels = QueryFilter.Word2;
+		Result.Word3 = QueryFilter.Word3;
+		Result.SimWord3 = QueryFilter.Word3;
+		return Result;
+	}
+
+	FShapeFilterData FShapeFilterBuilder::BuildFromLegacySimFilter(const FCollisionFilterData& SimFilter)
+	{
+		FShapeFilterData Result;
+		Result.SimBlockChannels = SimFilter.Word1;
+		Result.Word3 = SimFilter.Word3;
+		Result.SimWord3 = SimFilter.Word3;
+		return Result;
+	}
+
+	FCombinedShapeFilterData FShapeFilterBuilder::BuildFromLegacyShapeFilter(const FCollisionFilterData& QueryFilterData, const FCollisionFilterData& SimFilterData)
+	{
+		FCombinedShapeFilterData CombinedShapeFilterData;
+		CombinedShapeFilterData.InstanceData.ActorId = QueryFilterData.Word0;
+		CombinedShapeFilterData.InstanceData.ComponentId = SimFilterData.Word2;
+		CombinedShapeFilterData.InstanceData.BodyIndex = SimFilterData.Word0;
+
+		CombinedShapeFilterData.ShapeFilterData.QueryBlockChannels = QueryFilterData.Word1;
+		CombinedShapeFilterData.ShapeFilterData.QueryOverlapChannels = QueryFilterData.Word2;
+		CombinedShapeFilterData.ShapeFilterData.Word3 = QueryFilterData.Word3;
+		CombinedShapeFilterData.ShapeFilterData.SimBlockChannels = SimFilterData.Word1;
+		CombinedShapeFilterData.ShapeFilterData.SimWord3 = SimFilterData.Word3;
+
+		return CombinedShapeFilterData;
+	}
+
+	void FShapeFilterBuilder::SetLegacyShapeQueryFilter(FCombinedShapeFilterData& CombinedShapeFilterData, const FCollisionFilterData& QueryFilterData)
+	{
+		const FCollisionFilterData SimFilterData = GetLegacyShapeSimFilter(CombinedShapeFilterData);
+		CombinedShapeFilterData = BuildFromLegacyShapeFilter(QueryFilterData, SimFilterData);
+	}
+
+	void FShapeFilterBuilder::SetLegacyShapeSimFilter(FCombinedShapeFilterData& CombinedShapeFilterData, const FCollisionFilterData& SimFilterData)
+	{
+		const FCollisionFilterData QueryFilterData = GetLegacyShapeQueryFilter(CombinedShapeFilterData);
+		CombinedShapeFilterData = BuildFromLegacyShapeFilter(QueryFilterData, SimFilterData);
+	}
+
+	void FShapeFilterBuilder::GetLegacyShapeFilter(const FCombinedShapeFilterData& CombinedShapeFilterData, FCollisionFilterData& OutQueryFilterData, FCollisionFilterData& OutSimFilterData)
+	{
+		OutQueryFilterData.Word0 = CombinedShapeFilterData.InstanceData.ActorId;
+		OutSimFilterData.Word2 = CombinedShapeFilterData.InstanceData.ComponentId;
+		OutSimFilterData.Word0 = CombinedShapeFilterData.InstanceData.BodyIndex;
+
+		OutQueryFilterData.Word1 = CombinedShapeFilterData.ShapeFilterData.QueryBlockChannels;
+		OutQueryFilterData.Word2 = CombinedShapeFilterData.ShapeFilterData.QueryOverlapChannels;
+		OutQueryFilterData.Word3 = CombinedShapeFilterData.ShapeFilterData.Word3;
+		OutSimFilterData.Word1 = CombinedShapeFilterData.ShapeFilterData.SimBlockChannels;
+		OutSimFilterData.Word3 = CombinedShapeFilterData.ShapeFilterData.SimWord3;
+	}
+
+	FCollisionFilterData FShapeFilterBuilder::GetLegacyShapeQueryFilter(const FCombinedShapeFilterData& CombinedShapeFilterData)
+	{
+		FCollisionFilterData QueryFilter, SimFilter;
+		GetLegacyShapeFilter(CombinedShapeFilterData, QueryFilter, SimFilter);
+		return QueryFilter;
+	}
+
+	FCollisionFilterData FShapeFilterBuilder::GetLegacyShapeSimFilter(const FCombinedShapeFilterData& CombinedShapeFilterData)
+	{
+		FCollisionFilterData QueryFilter, SimFilter;
+		GetLegacyShapeFilter(CombinedShapeFilterData, QueryFilter, SimFilter);
+		return SimFilter;
+	}
+
+	FQueryFilterData FQueryFilterBuilder::CreateLegacyObjectTypeFilter(const uint32 ObjectTypesToQuery, const bool bMultiTrace, const uint8 IgnoreMask, const EFilterFlags FilterFlags)
+	{
+		const uint8 MultiTraceValue = bMultiTrace ? 0x1 : 0x0;
+		FQueryFilterData Result;
+		Result.Word0 = 0x0;
+		Result.Word1 = ObjectTypesToQuery;
+		Result.Word2 = 0;
+		Result.SetMaskChannelAndFlags(IgnoreMask, MultiTraceValue, FilterFlags);
+		return Result;
+	}
+
+	FQueryFilterData FQueryFilterBuilder::CreateLegacyTraceFilter(const uint8 ChannelIndex, const uint32 BlockChannelMask, const uint32 OverlapChannelMask, const uint8 IgnoreMask, const EFilterFlags FilterFlags)
+	{
+		FQueryFilterData Result;
+		Result.Word0 = 0x1;
+		Result.Word1 = BlockChannelMask;
+		Result.Word2 = OverlapChannelMask;
+		Result.SetMaskChannelAndFlags(IgnoreMask, ChannelIndex, FilterFlags);
+		return Result;
+	}
+
+	FQueryFilterData FQueryFilterBuilder::BuildFromLegacyQueryFilter(const FCollisionFilterData& QueryFilterData)
+	{
+		FQueryFilterData Result;
+		Result.Word0 = QueryFilterData.Word0;
+		Result.Word1 = QueryFilterData.Word1;
+		Result.Word2 = QueryFilterData.Word2;
+		Result.Word3 = QueryFilterData.Word3;
+		return Result;
+	}
+
+	FCollisionFilterData FQueryFilterBuilder::GetLegacyQueryFilter(const FQueryFilterData& QueryFilterData)
+	{
+		FCollisionFilterData Result;
+		Result.Word0 = QueryFilterData.Word0;
+		Result.Word1 = QueryFilterData.Word1;
+		Result.Word2 = QueryFilterData.Word2;
+		Result.Word3 = QueryFilterData.Word3;
+		return Result;
+	}
+} // namespace Chaos::Filter
diff --git a/Engine/Source/Runtime/Experimental/Chaos/Private/Chaos/ParticleDirtyFlags.cpp b/Engine/Source/Runtime/Experimental/Chaos/Private/Chaos/ParticleDirtyFlags.cpp
new file mode 100644
--- /dev/null
+++ b/Engine/Source/Runtime/Experimental/Chaos/Private/Chaos/ParticleDirtyFlags.cpp
@@ -0,0 +1,70 @@
+// Copyright Epic Games, Inc. All Rights Reserved.
+
+#include "Chaos/ParticleDirtyFlags.h"
+
+namespace Chaos
+{
+	const FCollisionFilterData& FCollisionData::GetQueryData() const
+	{
+		PRAGMA_DISABLE_DEPRECATION_WARNINGS
+		return QueryData;
+		PRAGMA_ENABLE_DEPRECATION_WARNINGS
+	}
+
+	void FCollisionData::SetQueryData(const FCollisionFilterData& InQueryData)
+	{
+		PRAGMA_DISABLE_DEPRECATION_WARNINGS
+		QueryData = InQueryData;
+		PRAGMA_ENABLE_DEPRECATION_WARNINGS
+	}
+
+	const FCollisionFilterData& FCollisionData::GetSimData() const
+	{
+		PRAGMA_DISABLE_DEPRECATION_WARNINGS
+		return SimData;
+		PRAGMA_ENABLE_DEPRECATION_WARNINGS
+	}
+
+	void FCollisionData::SetSimData(const FCollisionFilterData& InSimData)
+	{
+		PRAGMA_DISABLE_DEPRECATION_WARNINGS
+		SimData = InSimData;
+		PRAGMA_ENABLE_DEPRECATION_WARNINGS
+	}
+
+	Chaos::Filter::FShapeFilterData FCollisionData::GetShapeFilterData() const
+	{
+		return GetCombinedShapeFilterData().GetShapeFilterData();
+	}
+
+	void FCollisionData::SetShapeFilterData(const Chaos::Filter::FShapeFilterData& ShapeFilter)
+	{
+		Chaos::Filter::FCombinedShapeFilterData CombinedShapeFilter = GetCombinedShapeFilterData();
+		CombinedShapeFilter.SetShapeFilterData(ShapeFilter);
+		SetCombinedShapeFilterData(CombinedShapeFilter);
+	}
+
+	Chaos::Filter::FInstanceData FCollisionData::GetFilterInstanceData() const
+	{
+		return GetCombinedShapeFilterData().GetInstanceData();
+	}
+
+	void FCollisionData::SetFilterInstanceData(const Chaos::Filter::FInstanceData& InstanceData)
+	{
+		Chaos::Filter::FCombinedShapeFilterData CombinedShapeFilter = GetCombinedShapeFilterData();
+		CombinedShapeFilter.SetInstanceData(InstanceData);
+		SetCombinedShapeFilterData(CombinedShapeFilter);
+	}
+
+	Chaos::Filter::FCombinedShapeFilterData FCollisionData::GetCombinedShapeFilterData() const
+	{
+		return Chaos::Filter::FShapeFilterBuilder::BuildFromLegacyShapeFilter(GetQueryData(), GetSimData());
+	}
+
+	void FCollisionData::SetCombinedShapeFilterData(const Chaos::Filter::FCombinedShapeFilterData& CombinedShapeFilter)
+	{
+		PRAGMA_DISABLE_DEPRECATION_WARNINGS
+		Chaos::Filter::FShapeFilterBuilder::GetLegacyShapeFilter(CombinedShapeFilter, QueryData, SimData);
+		PRAGMA_ENABLE_DEPRECATION_WARNINGS
+	}
+} // namespace Chaos
diff --git a/Engine/Source/Runtime/Experimental/Chaos/Private/ChaosVisualDebugger/ChaosVDDataWrapperUtils.cpp b/Engine/Source/Runtime/Experimental/Chaos/Private/ChaosVisualDebugger/ChaosVDDataWrapperUtils.cpp
--- a/Engine/Source/Runtime/Experimental/Chaos/Private/ChaosVisualDebugger/ChaosVDDataWrapperUtils.cpp
+++ b/Engine/Source/Runtime/Experimental/Chaos/Private/ChaosVisualDebugger/ChaosVDDataWrapperUtils.cpp
@@ -680,15 +680,20 @@ void FChaosVDDataWrapperUtils::CopyShapeDataToWrapper(const Chaos::FShapeInstanc
 	OutCopyTo.bSimCollision = CollisionData.bSimCollision;
 	OutCopyTo.CollisionTraceType = static_cast<EChaosVDCollisionTraceFlag>(CollisionData.CollisionTraceType);
 
-	OutCopyTo.SimData.Word0 = CollisionData.SimData.Word0;
-	OutCopyTo.SimData.Word1 = CollisionData.SimData.Word1;
-	OutCopyTo.SimData.Word2 = CollisionData.SimData.Word2;
-	OutCopyTo.SimData.Word3 = CollisionData.SimData.Word3;
-	
-	OutCopyTo.QueryData.Word0 = CollisionData.QueryData.Word0;
-	OutCopyTo.QueryData.Word1 = CollisionData.QueryData.Word1;
-	OutCopyTo.QueryData.Word2 = CollisionData.QueryData.Word2;
-	OutCopyTo.QueryData.Word3 = CollisionData.QueryData.Word3;
+	// TODO @ JoshD: Update
+	const Chaos::Filter::FCombinedShapeFilterData CombinedShapeFilter = CollisionData.GetCombinedShapeFilterData();
+	FCollisionFilterData QueryData, SimData;
+	Chaos::Filter::FShapeFilterBuilder::GetLegacyShapeFilter(CombinedShapeFilter, QueryData, SimData);
+
+	OutCopyTo.SimData.Word0 = SimData.Word0;
+	OutCopyTo.SimData.Word1 = SimData.Word1;
+	OutCopyTo.SimData.Word2 = SimData.Word2;
+	OutCopyTo.SimData.Word3 = SimData.Word3;
+
+	OutCopyTo.QueryData.Word0 = QueryData.Word0;
+	OutCopyTo.QueryData.Word1 = QueryData.Word1;
+	OutCopyTo.QueryData.Word2 = QueryData.Word2;
+	OutCopyTo.QueryData.Word3 = QueryData.Word3;
 }
 
 void FChaosVDDataWrapperUtils::BuildDataWrapperFromAABBStructure(const Chaos::ISpatialAccelerationCollection<Chaos::FAccelerationStructureHandle, Chaos::FReal, 3>* SceneAccelerationStructures, int32 OwnerSolverID, TArray<FChaosVDAABBTreeDataWrapper>& OutAABBTrees)
diff --git a/Engine/Source/Runtime/Experimental/Chaos/Public/Chaos/CollisionFilterData.h b/Engine/Source/Runtime/Experimental/Chaos/Public/Chaos/CollisionFilterData.h
--- a/Engine/Source/Runtime/Experimental/Chaos/Public/Chaos/CollisionFilterData.h
+++ b/Engine/Source/Runtime/Experimental/Chaos/Public/Chaos/CollisionFilterData.h
@@ -12,6 +12,7 @@ namespace Chaos
 {
 	enum class EFilterFlags : uint8
 	{
+		None					= 0b00000000,
 		SimpleCollision			= 0b00000001,	// The shape is used for simple collision
 		ComplexCollision		= 0b00000010,	// The shape is used for complex (trimesh) collision
 		CCD						= 0b00000100,	// Unused - present for compatibility. CCD handled per-particle in Chaos
@@ -19,12 +20,15 @@ namespace Chaos
 		StaticShape				= 0b00010000,	// Unused - present for compatibility
 		ModifyContacts			= 0b00100000,	// Unused - present for compatibility, whether to allow contact modification, handled in Chaos callbacks now
 		KinematicKinematicPairs	= 0b01000000,	// Unused - present for compatibility, whether to generate KK pairs, Chaos never generates KK pairs
+		All						= 0xFF,
 	};
+	ENUM_CLASS_FLAGS(EFilterFlags);
 
 	inline const TCHAR* LexToString(EFilterFlags FilterFlag)
 	{
 		switch (FilterFlag)
 		{
+			case EFilterFlags::None: return TEXT("None");
 			case EFilterFlags::SimpleCollision: return TEXT("SimpleCollision");
 			case EFilterFlags::ComplexCollision: return TEXT("ComplexCollision");
 			case EFilterFlags::CCD: return TEXT("CCD");
@@ -32,6 +36,7 @@ namespace Chaos
 			case EFilterFlags::StaticShape: return TEXT("StaticShape");
 			case EFilterFlags::ModifyContacts: return TEXT("ModifyContacts");
 			case EFilterFlags::KinematicKinematicPairs: return TEXT("KinematicKinematicPairs");
+			case EFilterFlags::All: return TEXT("All");
 			default: return TEXT("Invalid");
 		}
 	}
@@ -66,3 +71,167 @@ inline Chaos::FChaosArchive& operator<<(Chaos::FChaosArchive& Ar, FCollisionFilt
 	Ar << Filter.Word0 << Filter.Word1 << Filter.Word2 << Filter.Word3;
 	return Ar;
 }
+
+namespace Chaos::Filter
+{
+	struct FCombinedShapeFilterData;
+	struct FQueryFilterData;
+
+	struct FInstanceData
+	{
+		FInstanceData() = default;
+		CHAOS_API FInstanceData(const uint32 ActorId, const uint32 ComponentId);
+
+		bool operator==(const FInstanceData&) const = default;
+		bool operator!=(const FInstanceData&) const = default;
+
+		CHAOS_API bool IsValid() const;
+
+		CHAOS_API uint32 GetActorId() const;
+		CHAOS_API void SetActorId(const uint32 InActorId);
+
+		UE_INTERNAL CHAOS_API uint32 GetComponentId() const;
+		UE_INTERNAL CHAOS_API void SetComponentId(const uint32 InComponentId);
+
+	private:
+		friend class FShapeFilterBuilder;
+		friend FCombinedShapeFilterData;
+
+		uint32 ActorId = 0;
+		uint32 ComponentId = 0;
+		// TODO @ JoshD: Remove if possible
+		uint32 BodyIndex = 0;
+	};
+
+	struct FShapeFilterData
+	{
+		FShapeFilterData() = default;
+
+		bool operator==(const FShapeFilterData&) const = default;
+		bool operator!=(const FShapeFilterData&) const = default;
+
+		CHAOS_API bool IsValid() const;
+
+		CHAOS_API EFilterFlags GetFlags() const;
+		CHAOS_API void SetFlags(EFilterFlags InFlags);
+		CHAOS_API bool HasFlag(EFilterFlags InFlag) const;
+
+		CHAOS_API uint8 GetMaskFilter() const;
+		CHAOS_API void SetMaskFilter(uint8 MaskFilter);
+
+		CHAOS_API uint8 GetCollisionChannelIndex() const;
+		CHAOS_API uint64 GetCollisionChannelMask() const;
+
+		CHAOS_API uint64 GetQueryBlockChannels() const;
+		CHAOS_API uint64 GetQueryOverlapChannels() const;
+		CHAOS_API uint64 GetSimBlockChannels() const;
+
+	private:
+		friend class FShapeFilterBuilder;
+		friend FQueryFilterData;
+		friend FCombinedShapeFilterData;
+
+		uint32 QueryBlockChannels = 0;
+		uint32 QueryOverlapChannels = 0;
+		uint32 SimBlockChannels = 0;
+		uint32 Word3 = 0;
+		// TODO @ JoshD: Remove this eventually. This is kept here for legacy conversions until usage is cleaned up.
+		uint32 SimWord3 = 0;
+	};
+
+	struct FQueryFilterData
+	{
+		CHAOS_API enum EQueryType
+		{
+			Channel = 0,
+			ObjectType,
+		};
+
+		FQueryFilterData() = default;
+
+		bool operator==(const FQueryFilterData&) const = default;
+		bool operator!=(const FQueryFilterData&) const = default;
+
+		CHAOS_API bool IsValid() const;
+
+		CHAOS_API EQueryType GetQueryType() const;
+
+		CHAOS_API EFilterFlags GetFlags() const;
+		CHAOS_API void SetFlags(EFilterFlags InFlags);
+		CHAOS_API bool HasFlag(EFilterFlags InFlag) const;
+
+		CHAOS_API uint8 GetIgnoreMask() const;
+
+		// Channel API
+		CHAOS_API uint8 GetCollisionChannelIndex() const;
+		CHAOS_API uint64 GetCollisionChannelMask() const;
+		CHAOS_API uint64 GetBlockChannels() const;
+		CHAOS_API uint64 GetOverlapChannels() const;
+
+		// ObjectType API
+		CHAOS_API uint64 GetObjectTypesToQueryMask() const;
+		CHAOS_API bool IsMultiQuery() const;
+
+	private:
+		friend class FQueryFilterBuilder;
+
+		void SetMaskChannelAndFlags(const uint8 IgnoreMask, const uint8 ChannelIndex, const EFilterFlags FilterFlags);
+
+		uint32 Word0 = 0;
+		uint32 Word1 = 0;
+		uint32 Word2 = 0;
+		uint32 Word3 = 0;
+	};
+
+	struct FCombinedShapeFilterData
+	{
+		FCombinedShapeFilterData() = default;
+		CHAOS_API FCombinedShapeFilterData(const FShapeFilterData& InShapeFilter, const FInstanceData& InInstanceData);
+
+		bool operator==(const FCombinedShapeFilterData&) const = default;
+		bool operator!=(const FCombinedShapeFilterData&) const = default;
+
+		CHAOS_API const FInstanceData& GetInstanceData() const;
+		CHAOS_API void SetInstanceData(const FInstanceData& InData);
+
+		CHAOS_API const FShapeFilterData& GetShapeFilterData() const;
+		CHAOS_API void SetShapeFilterData(const FShapeFilterData& InData);
+
+		CHAOS_API bool IsValid() const;
+
+	private:
+		friend class FShapeFilterBuilder;
+		FShapeFilterData ShapeFilterData;
+		FInstanceData InstanceData;
+	};
+
+	class UE_INTERNAL FShapeFilterBuilder
+	{
+	public:
+		UE_INTERNAL CHAOS_API static FShapeFilterData BuildLegacyShapeFilter(const uint8 InChannelIndex, const uint32 InQueryOverlapChannels, const uint32 InQueryBlockChannels, const uint32 InSimBlockChannels, const uint8 InMaskFilter = 0, const EFilterFlags InFilterFlags = EFilterFlags::None);
+		UE_INTERNAL CHAOS_API static FShapeFilterData BuildLegacyBlockAllSimShapeFilter(const uint8 InChannelIndex = 0, const uint8 InMaskFilter = 0, const EFilterFlags InFilterFlags = EFilterFlags::All);
+
+		UE_INTERNAL CHAOS_API static FShapeFilterData BuildFromLegacyQueryFilter(const FCollisionFilterData& QueryFilter);
+		UE_INTERNAL CHAOS_API static FShapeFilterData BuildFromLegacySimFilter(const FCollisionFilterData& SimFilter);
+		UE_INTERNAL CHAOS_API static FCombinedShapeFilterData BuildFromLegacyShapeFilter(const FCollisionFilterData& QueryFilterData, const FCollisionFilterData& SimFilterData);
+
+		UE_INTERNAL CHAOS_API static void SetLegacyShapeQueryFilter(FCombinedShapeFilterData& CombinedShapeFilterData, const FCollisionFilterData& QueryFilterData);
+		UE_INTERNAL CHAOS_API static void SetLegacyShapeSimFilter(FCombinedShapeFilterData& CombinedShapeFilterData, const FCollisionFilterData& SimFilterData);
+
+		UE_INTERNAL CHAOS_API static void GetLegacyShapeFilter(const FCombinedShapeFilterData& CombinedShapeFilterData, FCollisionFilterData& OutQueryFilterData, FCollisionFilterData& OutSimFilterData);
+		UE_INTERNAL CHAOS_API static FCollisionFilterData GetLegacyShapeQueryFilter(const FCombinedShapeFilterData& CombinedShapeFilterData);
+		UE_INTERNAL CHAOS_API static FCollisionFilterData GetLegacyShapeSimFilter(const FCombinedShapeFilterData& CombinedShapeFilterData);
+	private:
+	};
+
+	class UE_INTERNAL FQueryFilterBuilder
+	{
+	public:
+		UE_INTERNAL CHAOS_API static FQueryFilterData CreateLegacyObjectTypeFilter(const uint32 ObjectTypesToQuery, const bool bMultiTrace, const uint8 IgnoreMask = 0, const EFilterFlags FilterFlags = EFilterFlags::None);
+		UE_INTERNAL CHAOS_API static FQueryFilterData CreateLegacyTraceFilter(const uint8 ChannelIndex, const uint32 BlockChannelMask, const uint32 OverlapChannelMask, const uint8 IgnoreMask = 0, const EFilterFlags FilterFlags = EFilterFlags::None);
+
+		UE_INTERNAL CHAOS_API static FQueryFilterData BuildFromLegacyQueryFilter(const FCollisionFilterData& QueryFilterData);
+		UE_INTERNAL CHAOS_API static FCollisionFilterData GetLegacyQueryFilter(const FQueryFilterData& QueryFilterData);
+	private:
+	};
+} // namespace Chaos::Filter
diff --git a/Engine/Source/Runtime/Experimental/Chaos/Public/Chaos/ParticleDirtyFlags.h b/Engine/Source/Runtime/Experimental/Chaos/Public/Chaos/ParticleDirtyFlags.h
--- a/Engine/Source/Runtime/Experimental/Chaos/Public/Chaos/ParticleDirtyFlags.h
+++ b/Engine/Source/Runtime/Experimental/Chaos/Public/Chaos/ParticleDirtyFlags.h
@@ -720,8 +720,11 @@ inline FChaosArchive& operator<<(FChaosArchive& Ar,FParticleNonFrequentData& Dat
 
 struct FCollisionData
 {
+	UE_DEPRECATED(5.7, "Use Get/SetCombinedShapeFilterData instead")
 	FCollisionFilterData QueryData;
+	UE_DEPRECATED(5.7, "Use Get/SetCombinedShapeFilterData instead")
 	FCollisionFilterData SimData;
+
 	void* UserData;
 	EChaosCollisionTraceFlag CollisionTraceType;
 	uint8 bSimCollision : 1;
@@ -737,6 +740,28 @@ struct FCollisionData
 	{
 	}
 
+	PRAGMA_DISABLE_DEPRECATION_WARNINGS
+	FCollisionData(const FCollisionData& Rhs) = default;
+	FCollisionData(FCollisionData& Rhs) = default;
+
+	FCollisionData& operator=(const FCollisionData& Rhs) = default;
+	FCollisionData& operator=(FCollisionData&& Rhs) = default;
+	PRAGMA_ENABLE_DEPRECATION_WARNINGS
+
+	UE_INTERNAL CHAOS_API const FCollisionFilterData& GetQueryData() const;
+	UE_INTERNAL CHAOS_API void SetQueryData(const FCollisionFilterData& InQueryData);
+	UE_INTERNAL CHAOS_API const FCollisionFilterData& GetSimData() const;
+	UE_INTERNAL CHAOS_API void SetSimData(const FCollisionFilterData& InSimData);
+
+	CHAOS_API Chaos::Filter::FShapeFilterData GetShapeFilterData() const;
+	CHAOS_API void SetShapeFilterData(const Chaos::Filter::FShapeFilterData& ShapeFilter);
+
+	CHAOS_API Chaos::Filter::FInstanceData GetFilterInstanceData() const;
+	CHAOS_API void SetFilterInstanceData(const Chaos::Filter::FInstanceData& InstanceData);
+
+	CHAOS_API Chaos::Filter::FCombinedShapeFilterData GetCombinedShapeFilterData() const;
+	CHAOS_API void SetCombinedShapeFilterData(const Chaos::Filter::FCombinedShapeFilterData& CombinedShapeFilter);
+
 	bool HasCollisionData() const { return bSimCollision || bQueryCollision; }
 	bool HasQueryOnlyData() const { return !bSimCollision && bQueryCollision; }
 
@@ -746,8 +771,10 @@ struct FCollisionData
 		Ar.UsingCustomVersion(FExternalPhysicsMaterialCustomObjectVersion::GUID);
 		Ar.UsingCustomVersion(FFortniteSeasonBranchObjectVersion::GUID);
 
+		PRAGMA_DISABLE_DEPRECATION_WARNINGS
 		Ar << QueryData;
 		Ar << SimData;
+		PRAGMA_ENABLE_DEPRECATION_WARNINGS
 
 		if (Ar.CustomVer(FExternalPhysicsCustomObjectVersion::GUID) >= FExternalPhysicsCustomObjectVersion::AddShapeSimAndQueryCollisionEnabled)
 		{
diff --git a/Engine/Source/Runtime/Experimental/Chaos/Public/Chaos/ShapeInstance.h b/Engine/Source/Runtime/Experimental/Chaos/Public/Chaos/ShapeInstance.h
--- a/Engine/Source/Runtime/Experimental/Chaos/Public/Chaos/ShapeInstance.h
+++ b/Engine/Source/Runtime/Experimental/Chaos/Public/Chaos/ShapeInstance.h
@@ -279,17 +279,19 @@ namespace Chaos
 			CollisionData.Modify(true, DirtyFlags, Proxy, ShapeIdx, [InUserData](FCollisionData& Data) { Data.UserData = InUserData; });
 		}
 
-		const FCollisionFilterData& GetQueryData() const { return CollisionData.Read().QueryData; }
+		PRAGMA_DISABLE_INTERNAL_WARNINGS
+		const FCollisionFilterData& GetQueryData() const { return CollisionData.Read().GetQueryData(); }
 		void SetQueryData(const FCollisionFilterData& InQueryData)
 		{
-			CollisionData.Modify(true, DirtyFlags, Proxy, ShapeIdx, [InQueryData](FCollisionData& Data) { Data.QueryData = InQueryData; });
+			CollisionData.Modify(true, DirtyFlags, Proxy, ShapeIdx, [InQueryData](FCollisionData& Data) { Data.SetQueryData(InQueryData); });
 		}
 
-		const FCollisionFilterData& GetSimData() const { return CollisionData.Read().SimData; }
+		const FCollisionFilterData& GetSimData() const { return CollisionData.Read().GetSimData(); }
 		void SetSimData(const FCollisionFilterData& InSimData)
 		{
-			CollisionData.Modify(true, DirtyFlags, Proxy, ShapeIdx, [InSimData](FCollisionData& Data) { Data.SimData = InSimData; });
+			CollisionData.Modify(true, DirtyFlags, Proxy, ShapeIdx, [InSimData](FCollisionData& Data) { Data.SetSimData(InSimData); });
 		}
+		PRAGMA_ENABLE_INTERNAL_WARNINGS
 
 		CHAOS_API void UpdateWorldSpaceState(const FRigidTransform3& WorldTransform);
 
@@ -428,7 +430,9 @@ namespace Chaos
 		template <typename Lambda>
 		void ModifySimData(const Lambda& LambdaFunc)
 		{
+			PRAGMA_DISABLE_DEPRECATION_WARNINGS
 			CollisionData.Modify(true, DirtyFlags, Proxy, ShapeIdx, [&LambdaFunc](FCollisionData& Data) { LambdaFunc(Data.SimData); });
+			PRAGMA_ENABLE_DEPRECATION_WARNINGS
 		}
 
 		template <typename Lambda>
@@ -554,11 +558,13 @@ namespace Chaos
 		void* GetUserData() const { return CollisionData.UserData; }
 		void SetUserData(void* InUserData) { CollisionData.UserData = InUserData; }
 
-		const FCollisionFilterData& GetQueryData() const { return CollisionData.QueryData; }
-		void SetQueryData(const FCollisionFilterData& InQueryData) { CollisionData.QueryData = InQueryData; }
+		PRAGMA_DISABLE_INTERNAL_WARNINGS
+		const FCollisionFilterData& GetQueryData() const { return CollisionData.GetQueryData(); }
+		void SetQueryData(const FCollisionFilterData& InQueryData) { CollisionData.SetQueryData(InQueryData); }
 
-		const FCollisionFilterData& GetSimData() const { return CollisionData.SimData; }
-		void SetSimData(const FCollisionFilterData& InSimData) { CollisionData.SimData = InSimData; }
+		const FCollisionFilterData& GetSimData() const { return CollisionData.GetSimData(); }
+		void SetSimData(const FCollisionFilterData& InSimData) { CollisionData.SetSimData(InSimData); }
+		PRAGMA_ENABLE_INTERNAL_WARNINGS
 
 		CHAOS_API void UpdateWorldSpaceState(const FRigidTransform3& WorldTransform);
 
@@ -662,7 +668,9 @@ namespace Chaos
 		void SyncRemoteData(FDirtyPropertiesManager& Manager, int32 ShapeDataIdx, FShapeDirtyData& RemoteData) { check(false); }
 		void SetProxy(IPhysicsProxyBase* InProxy) { check(false); }
 
+		PRAGMA_DISABLE_DEPRECATION_WARNINGS
 		template <typename Lambda> void ModifySimData(const Lambda& LambdaFunc) { LambdaFunc(CollisionData.SimData); }
+		PRAGMA_ENABLE_DEPRECATION_WARNINGS
 		template <typename Lambda> void ModifyMaterials(const Lambda& LambdaFunc) { LambdaFunc(GetMaterialDataImpl().Materials); }
 		template <typename Lambda> void ModifyMaterialMasks(const Lambda& LambdaFunc) { LambdaFunc(GetMaterialDataImpl().MaterialMasks); }
 		template <typename Lambda> void ModifyMaterialMaskMaps(const Lambda& LambdaFunc) { LambdaFunc(GetMaterialDataImpl().MaterialMaskMaps); }
diff --git a/Engine/Source/Runtime/Experimental/Chaos/Public/RewindData.h b/Engine/Source/Runtime/Experimental/Chaos/Public/RewindData.h
--- a/Engine/Source/Runtime/Experimental/Chaos/Public/RewindData.h
+++ b/Engine/Source/Runtime/Experimental/Chaos/Public/RewindData.h
@@ -891,7 +891,12 @@ struct FPerShapeDataStateBase
 
 	//helper functions for shape API
 	template <typename TParticle>
-	static const FCollisionFilterData& GetQueryData(const FPerShapeDataStateBase* State, const TParticle& Particle, int32 ShapeIdx) { return State && State->CollisionData.IsSet() ? State->CollisionData.Read().QueryData : Particle.ShapesArray()[ShapeIdx]->GetQueryData(); }
+	static const FCollisionFilterData& GetQueryData(const FPerShapeDataStateBase* State, const TParticle& Particle, int32 ShapeIdx)
+	{
+		PRAGMA_DISABLE_INTERNAL_WARNINGS
+		return State && State->CollisionData.IsSet() ? State->CollisionData.Read().GetQueryData() : Particle.ShapesArray()[ShapeIdx]->GetQueryData();
+		PRAGMA_ENABLE_INTERNAL_WARNINGS
+	}
 };
 
 class FPerShapeDataState
