Fixing compilation error caused by invalid attribute location on an enum

diff --git a/Engine/Source/Runtime/Experimental/Chaos/Public/Chaos/CollisionFilterData.h b/Engine/Source/Runtime/Experimental/Chaos/Public/Chaos/CollisionFilterData.h
--- a/Engine/Source/Runtime/Experimental/Chaos/Public/Chaos/CollisionFilterData.h
+++ b/Engine/Source/Runtime/Experimental/Chaos/Public/Chaos/CollisionFilterData.h
@@ -141,7 +141,7 @@ namespace Chaos::Filter
 
 	struct FQueryFilterData
 	{
-		CHAOS_API enum EQueryType
+		enum CHAOS_API EQueryType
 		{
 			Channel = 0,
 			ObjectType,
