Chaos - fix FParticleFlags serialization of MIterationSettings

diff --git a/Engine/Source/Runtime/Experimental/Chaos/Public/Chaos/Evolution/IterationSettings.h b/Engine/Source/Runtime/Experimental/Chaos/Public/Chaos/Evolution/IterationSettings.h
--- a/Engine/Source/Runtime/Experimental/Chaos/Public/Chaos/Evolution/IterationSettings.h
+++ b/Engine/Source/Runtime/Experimental/Chaos/Public/Chaos/Evolution/IterationSettings.h
@@ -27,7 +27,7 @@ namespace Chaos
 			static const int32 InvalidIterations = INDEX_NONE;
 			static const int32 MaxIterations = std::numeric_limits<FIndexType>::max();
 
-			static TIterationSettings<FIndexType> Empty()
+			static TIterationSettings<FIndexType> MakeEmpty()
 			{
 				return TIterationSettings<FIndexType>(InvalidIterations, InvalidIterations, InvalidIterations);
 			}
diff --git a/Engine/Source/Runtime/Experimental/Chaos/Public/Chaos/ParticleDirtyFlags.h b/Engine/Source/Runtime/Experimental/Chaos/Public/Chaos/ParticleDirtyFlags.h
--- a/Engine/Source/Runtime/Experimental/Chaos/Public/Chaos/ParticleDirtyFlags.h
+++ b/Engine/Source/Runtime/Experimental/Chaos/Public/Chaos/ParticleDirtyFlags.h
@@ -465,7 +465,7 @@ public:
 		{
 			MInitialOverlapDepenetrationVelocity = 0;
 			MSleepThresholdMultiplier = 1.0f;
-			MIterationSettings.Empty();
+			MIterationSettings = MIterationSettings.MakeEmpty();
 		}
 
 		if (bAddControlFlags)
