Chaos - fix waking of collision constraints (left as "disabled", so they were missing from collision events)

Pre-existing Collisions are set to Disabled at the start of each tick and enabled if they are still valid. But a collision in a sleeping island was not setting the flag, so it was being skipped in the collision solve and event dispatch.

The Disabled flag can also be changed by the user in a callback. We need to separarate the user and internal flag to prevent re-enabling of user-disabled collisions.

#rnx

diff --git a/Engine/Source/Runtime/Experimental/Chaos/Private/Chaos/Collision/ParticlePairMidPhase.cpp b/Engine/Source/Runtime/Experimental/Chaos/Private/Chaos/Collision/ParticlePairMidPhase.cpp
--- a/Engine/Source/Runtime/Experimental/Chaos/Private/Chaos/Collision/ParticlePairMidPhase.cpp
+++ b/Engine/Source/Runtime/Experimental/Chaos/Private/Chaos/Collision/ParticlePairMidPhase.cpp
@@ -282,6 +282,23 @@ namespace Chaos
 		};
 	};
 
+	void ReactivateCollisionConstraint(FPBDCollisionConstraint* Constraint)
+	{
+		// Update the epoch and add to the active list
+		Constraint->ConcreteContainer()->GetConstraintAllocator().AddActiveConstraint(Constraint);
+
+		// Re-enable the constraint
+		// NOTE: The Enabled/Disabled flag can also be controlled by the user in a
+		// contact modifier. There is no danger of re-enabling a constraint that the 
+		// user disabled because a disabled constraint is removed from the graph and 
+		// therefore will not get reactivated if an island wakes.
+		// @todo(chaos): Despite the note, we should separate Internal and Extenal Enabled flags.
+		// Internal-Enabled would be controlled by collision detection and waking. 
+		// External-Enabled would be controlled by the user via a contact callback.
+		// We would only be enabled when both are set. Probably call them Activated and Enabled.
+		Constraint->SetEnabled(true);
+	}
+
 
 	////////////////////////////////////////////////////////////////////////////////////////
 	////////////////////////////////////////////////////////////////////////////////////////
@@ -701,8 +718,7 @@ namespace Chaos
 		{
 			if (bChaosMidPhaseActivateWakingConstraints)
 			{
-				// Update the epoch and add to the active list
-				Constraint->ConcreteContainer()->GetConstraintAllocator().AddActiveConstraint(Constraint.Get());
+				ReactivateCollisionConstraint(Constraint.Get());
 				LastUsedEpoch = Constraint->GetContainerCookie().LastUsedEpoch;
 			}
 			else
@@ -2107,8 +2123,7 @@ namespace Chaos
 			{
 				if (bChaosMidPhaseActivateWakingConstraints)
 				{
-					// Update the epoch and add to the active list
-					Constraint->ConcreteContainer()->GetConstraintAllocator().AddActiveConstraint(Constraint.Get());
+					ReactivateCollisionConstraint(Constraint.Get());
 				}
 				else
 				{
@@ -2313,8 +2328,7 @@ namespace Chaos
 		{
 			if (bChaosMidPhaseActivateWakingConstraints)
 			{
-				// Update the epoch and add to the active list
-				Constraint->ConcreteContainer()->GetConstraintAllocator().AddActiveConstraint(Constraint.Get());
+				ReactivateCollisionConstraint(Constraint.Get());
 			}
 			else
 			{
