[ChaosVD] Adding a way to only trace CVD data that is within a specific area instead of the entire world.

Also, added a way to start a recording with a specific set of default data channels. This is used to set the same channels enabled/disabled state you have in CVD's UI when you start a recording in a packaged build.

diff --git a/Engine/Plugins/ChaosVD/Source/ChaosVD/Private/Widgets/SChaosVDRecordingControls.cpp b/Engine/Plugins/ChaosVD/Source/ChaosVD/Private/Widgets/SChaosVDRecordingControls.cpp
--- a/Engine/Plugins/ChaosVD/Source/ChaosVD/Private/Widgets/SChaosVDRecordingControls.cpp
+++ b/Engine/Plugins/ChaosVD/Source/ChaosVD/Private/Widgets/SChaosVDRecordingControls.cpp
@@ -765,6 +765,15 @@ void SChaosVDRecordingControls::SetSessionRecordingState(bool bIsRecording, ECha
 
 			constexpr bool bAppendPort = false;
 			RecordingParams.Target = LocalIP->ToString(bAppendPort);
+	
+			Chaos::VisualDebugger::FChaosVDDataChannelsManager::Get().EnumerateChannels([&RecordingParams](const TSharedRef<Chaos::VisualDebugger::FChaosVDOptionalDataChannel>& Channel)
+			{
+				if (Channel->IsChannelEnabled())
+				{
+					RecordingParams.DataChannelsEnabledOverrideList.Add(Channel->GetId().ToString());
+				}
+				return true;
+			});
 
 			RemoteSessionManager->SendStartRecordingCommand(SessionInfoRef->Address, RecordingParams);
 
@@ -796,6 +805,16 @@ void SChaosVDRecordingControls::SetSessionRecordingState(bool bIsRecording, ECha
 		{
 			FChaosVDStartRecordingCommandMessage RecordingParams;
 			RecordingParams.RecordingMode = EChaosVDRecordingMode::File;
+
+			Chaos::VisualDebugger::FChaosVDDataChannelsManager::Get().EnumerateChannels([&RecordingParams](const TSharedRef<Chaos::VisualDebugger::FChaosVDOptionalDataChannel>& Channel)
+			{
+				if (Channel->IsChannelEnabled())
+				{
+					RecordingParams.DataChannelsEnabledOverrideList.Add(Channel->GetId().ToString());
+				}
+				return true;
+			});
+			
 			RemoteSessionManager->SendStartRecordingCommand(SessionInfoRef->Address, RecordingParams);
 
 			// Once the start recording command is issued, we need to check if the recording started, which might take a few frames.
diff --git a/Engine/Plugins/ChaosVD/Source/ChaosVDBlueprint/Private/ChaosVDRuntimeBlueprintLibrary.cpp b/Engine/Plugins/ChaosVD/Source/ChaosVDBlueprint/Private/ChaosVDRuntimeBlueprintLibrary.cpp
--- a/Engine/Plugins/ChaosVD/Source/ChaosVDBlueprint/Private/ChaosVDRuntimeBlueprintLibrary.cpp
+++ b/Engine/Plugins/ChaosVD/Source/ChaosVDBlueprint/Private/ChaosVDRuntimeBlueprintLibrary.cpp
@@ -30,3 +30,8 @@ void UChaosVDRuntimeBlueprintLibrary::RecordDebugDrawSphere(const UObject* World
 {
 	CVD_TRACE_DEBUG_DRAW_SPHERE(InCenter, Radius, Tag, Color.ToFColorSRGB(), WorldContext ? CVD_TRACE_GET_SOLVER_ID_FROM_WORLD(WorldContext->GetWorld()) : INDEX_NONE);
 }
+
+void UChaosVDRuntimeBlueprintLibrary::SetTraceRelevancyVolume(const UObject* WorldContext, const FBox& RelevancyVolume)
+{
+	CVD_SET_RELEVANCY_VOLUME(RelevancyVolume);
+}
diff --git a/Engine/Plugins/ChaosVD/Source/ChaosVDBlueprint/Public/ChaosVDRuntimeBlueprintLibrary.h b/Engine/Plugins/ChaosVD/Source/ChaosVDBlueprint/Public/ChaosVDRuntimeBlueprintLibrary.h
--- a/Engine/Plugins/ChaosVD/Source/ChaosVDBlueprint/Public/ChaosVDRuntimeBlueprintLibrary.h
+++ b/Engine/Plugins/ChaosVD/Source/ChaosVDBlueprint/Public/ChaosVDRuntimeBlueprintLibrary.h
@@ -24,4 +24,7 @@ class UChaosVDRuntimeBlueprintLibrary : public UBlueprintFunctionLibrary
 
 	UFUNCTION(BlueprintCallable, Category = "Chaos Visual Debugger", DisplayName="CVD Record Debug Draw Sphere", meta = (WorldContext = "WorldContext"))
 	static void RecordDebugDrawSphere(const UObject* WorldContext, const FVector& InCenter, float Radius, FName Tag = NAME_None, FLinearColor Color = FLinearColor::Blue);
+
+	UFUNCTION(BlueprintCallable, Category = "Chaos Visual Debugger", DisplayName="CVD Set Trace Relevancy Volume", meta = (WorldContext = "WorldContext"))
+	static void SetTraceRelevancyVolume(const UObject* WorldContext, const FBox& RelevancyVolume);
 };
diff --git a/Engine/Source/Runtime/Experimental/Chaos/Private/ChaosVisualDebugger/ChaosVisualDebuggerTrace.cpp b/Engine/Source/Runtime/Experimental/Chaos/Private/ChaosVisualDebugger/ChaosVisualDebuggerTrace.cpp
--- a/Engine/Source/Runtime/Experimental/Chaos/Private/ChaosVisualDebugger/ChaosVisualDebuggerTrace.cpp
+++ b/Engine/Source/Runtime/Experimental/Chaos/Private/ChaosVisualDebugger/ChaosVisualDebuggerTrace.cpp
@@ -49,6 +49,7 @@ UE_TRACE_EVENT_DEFINE(ChaosVDLogger, ChaosVDNonSolverTransform)
 UE_TRACE_EVENT_DEFINE(ChaosVDLogger, ChaosVDNetworkTickOffset)
 UE_TRACE_EVENT_DEFINE(ChaosVDLogger, ChaosVDRolledBackDataID)
 UE_TRACE_EVENT_DEFINE(ChaosVDLogger, ChaosVDUsesAutoRTFM)
+UE_TRACE_EVENT_DEFINE(ChaosVDLogger, ChaosVDTraceRelevancyVolume)
 
 namespace Chaos::VisualDebugger::Cvars
 {
@@ -123,6 +124,7 @@ TSet<int32> FChaosVisualDebuggerTrace::SolverIDsForDeltaRecording = TSet<int32>(
 TSet<int32> FChaosVisualDebuggerTrace::RequestedFullCaptureSolverIDs = TSet<int32>();
 TSharedRef<FChaosVDSerializableNameTable> FChaosVisualDebuggerTrace::CVDNameTable = MakeShared<FChaosVDSerializableNameTable>();
 std::atomic<bool> FChaosVisualDebuggerTrace::bIsTracing = false;
+FBox FChaosVisualDebuggerTrace::TraceRelevancyVolume = FBox(EForceInit::ForceInitToZero);
 
 void FChaosVisualDebuggerTrace::TraceParticle(const Chaos::FGeometryParticleHandle* ParticleHandle)
 {
@@ -154,9 +156,16 @@ void FChaosVisualDebuggerTrace::TraceParticle(Chaos::FGeometryParticleHandle* Pa
 		return;
 	}
 
-	const uint32 GeometryHash = GeometryTracerObject.GetGeometryHashForImplicit(ParticleHandle->GetGeometry());
+	if (!IsInRelevancyVolume(ParticleHandle))
+	{
+		return;
+	}
+
+	const Chaos::FImplicitObjectRef ParticleGeometry = ParticleHandle->GetGeometry();
+
+	const uint32 GeometryHash = GeometryTracerObject.GetGeometryHashForImplicit(ParticleGeometry);
 	
-	TraceImplicitObject({ GeometryHash, ParticleHandle->GetGeometry() });
+	TraceImplicitObject({ GeometryHash, ParticleGeometry });
 
 	{
 		FChaosVDParticleDataWrapper ParticleDataWrapper = FChaosVDDataWrapperUtils::BuildParticleDataWrapperFromParticle(ParticleHandle);
@@ -237,6 +246,38 @@ void FChaosVisualDebuggerTrace::TraceParticleDestroyed(const Chaos::FGeometryPar
 	};
 }
 
+bool FChaosVisualDebuggerTrace::IsInRelevancyVolume(const Chaos::FGeometryParticleHandle* InParticleHandle)
+{
+	if (!InParticleHandle)
+	{
+		return false;
+	}
+
+	if (!TraceRelevancyVolume.IsValid)
+	{
+		return true;
+	}
+
+	const Chaos::FImplicitObjectRef ParticleGeometry = InParticleHandle->GetGeometry();
+	return ParticleGeometry && ParticleGeometry->HasBoundingBox() ? IsInRelevancyVolume(ParticleGeometry->CalculateTransformedBounds(InParticleHandle->GetTransformXR()))
+										: IsInRelevancyVolume(InParticleHandle->GetX());
+}
+
+bool FChaosVisualDebuggerTrace::IsInRelevancyVolume(const Chaos::FImplicitObject* Geometry, const Chaos::FRigidTransform3& InTransform)
+{
+	if (!Geometry)
+	{
+		return false;
+	}
+
+	if (!TraceRelevancyVolume.IsValid)
+	{
+		return true;
+	}
+
+	return Geometry->HasBoundingBox() ? IsInRelevancyVolume(Geometry->CalculateTransformedBounds(InTransform)) : IsInRelevancyVolume(InTransform.GetLocation());
+}
+
 void FChaosVisualDebuggerTrace::TraceParticleClusterChildData(const Chaos::TParticleView<Chaos::TPBDRigidParticles<Chaos::FReal, 3>>& ParticlesView, Chaos::FRigidClustering* ClusteringData, const FChaosVDContext& CVDContextData)
 {
 	if (!IsTracing())
@@ -347,6 +388,13 @@ void FChaosVisualDebuggerTrace::TraceMidPhase(const Chaos::FParticlePairMidPhase
 		return;
 	}
 
+	using namespace Chaos;
+	const bool bIsInRelevancyVolume = IsInRelevancyVolume(MidPhase->GetParticle0()) || IsInRelevancyVolume(MidPhase->GetParticle1());
+	if (!bIsInRelevancyVolume)
+	{
+		return;
+	}
+
 	FChaosVDParticlePairMidPhase CVDMidPhase = FChaosVDDataWrapperUtils::BuildMidPhaseDataWrapperFromMidPhase(*MidPhase);
 	CVDMidPhase.SolverID = CVDContextData->Id;
 
@@ -395,20 +443,31 @@ void FChaosVisualDebuggerTrace::TraceJointsConstraints(Chaos::FPBDJointConstrain
 		return;
 	}
 
-	const Chaos::FPBDJointConstraints::FHandles& JointHandles = InJointConstraints.GetConstConstraintHandles();
+	using namespace Chaos;
+	const FPBDJointConstraints::FHandles& JointHandles = InJointConstraints.GetConstConstraintHandles();
 
 	ParallelFor(JointHandles.Num(), [&JointHandles, CopyContext = *CVDContextData](int32 ConstraintIndex)
 	{
 		CVD_SCOPE_CONTEXT(CopyContext);
-		
-		FChaosVDJointConstraint WrappedJointConstraintData = FChaosVDDataWrapperUtils::BuildJointDataWrapper(JointHandles[ConstraintIndex]);
-
-		WrappedJointConstraintData.SolverID = CopyContext.Id;
 
-		FChaosVDScopedTLSBufferAccessor TLSDataBuffer;
-		Chaos::VisualDebugger::WriteDataToBuffer(TLSDataBuffer.BufferRef, WrappedJointConstraintData);
+		if (FPBDJointConstraints::FConstraintContainerHandle* ConstraintHandle = JointHandles[ConstraintIndex])
+		{
+			FParticlePair ParticlePair = ConstraintHandle->GetConstrainedParticles();
+			bool bIsInRelevancyVolume = IsInRelevancyVolume(ParticlePair[0]) || IsInRelevancyVolume(ParticlePair[1]);
+			if (!bIsInRelevancyVolume)
+			{
+				return;
+			}
 
-		TraceBinaryData(TLSDataBuffer.BufferRef, FChaosVDJointConstraint::WrapperTypeName);
+			FChaosVDJointConstraint WrappedJointConstraintData = FChaosVDDataWrapperUtils::BuildJointDataWrapper(ConstraintHandle);
+			
+			WrappedJointConstraintData.SolverID = CopyContext.Id;
+			
+			FChaosVDScopedTLSBufferAccessor TLSDataBuffer;
+			VisualDebugger::WriteDataToBuffer(TLSDataBuffer.BufferRef, WrappedJointConstraintData);
+			
+			TraceBinaryData(TLSDataBuffer.BufferRef, FChaosVDJointConstraint::WrapperTypeName);
+		}
 	});
 }
 
@@ -428,20 +487,31 @@ void FChaosVisualDebuggerTrace::TraceCharacterGroundConstraints(Chaos::FCharacte
 		return;
 	}
 
-	const Chaos::FCharacterGroundConstraintContainer::FConstConstraints& ConstraintHandles = InConstraints.GetConstConstraints();
+	using namespace Chaos;
+	const FCharacterGroundConstraintContainer::FConstConstraints& ConstraintHandles = InConstraints.GetConstConstraints();
 
 	ParallelFor(ConstraintHandles.Num(), [&ConstraintHandles, CopyContext = *CVDContextData](int32 ConstraintIndex)
 	{
 		CVD_SCOPE_CONTEXT(CopyContext);
 
-		FChaosVDCharacterGroundConstraint WrappedConstraintData = FChaosVDDataWrapperUtils::BuildCharacterGroundConstraintDataWrapper(ConstraintHandles[ConstraintIndex]);
+		if (const FCharacterGroundConstraintHandle* const GroundConstraint = ConstraintHandles[ConstraintIndex])
+		{
+			FParticlePair ParticlePair = GroundConstraint->GetConstrainedParticles();
+			const bool bIsInRelevancyVolume = IsInRelevancyVolume(ParticlePair[0]) || IsInRelevancyVolume(ParticlePair[1]);
+			if (!bIsInRelevancyVolume)
+			{
+				return;
+			}
 
-		WrappedConstraintData.SolverID = CopyContext.Id;
+			FChaosVDCharacterGroundConstraint WrappedConstraintData = FChaosVDDataWrapperUtils::BuildCharacterGroundConstraintDataWrapper(GroundConstraint);
 
-		FChaosVDScopedTLSBufferAccessor TLSDataBuffer;
-		Chaos::VisualDebugger::WriteDataToBuffer(TLSDataBuffer.BufferRef, WrappedConstraintData);
+			WrappedConstraintData.SolverID = CopyContext.Id;
 
-		TraceBinaryData(TLSDataBuffer.BufferRef, FChaosVDCharacterGroundConstraint::WrapperTypeName);
+			FChaosVDScopedTLSBufferAccessor TLSDataBuffer;
+			VisualDebugger::WriteDataToBuffer(TLSDataBuffer.BufferRef, WrappedConstraintData);
+
+			TraceBinaryData(TLSDataBuffer.BufferRef, FChaosVDCharacterGroundConstraint::WrapperTypeName);
+		}
 	});
 }
 
@@ -460,13 +530,25 @@ void FChaosVisualDebuggerTrace::TraceCollisionConstraint(const Chaos::FPBDCollis
 		return;
 	}
 
-	FChaosVDConstraint CVDConstraint = FChaosVDDataWrapperUtils::BuildConstraintDataWrapperFromConstraint(*CollisionConstraint);
-	CVDConstraint.SolverID = CVDContextData->Id;
+	using namespace Chaos;
 
-	FChaosVDScopedTLSBufferAccessor TLSDataBuffer;
-	Chaos::VisualDebugger::WriteDataToBuffer(TLSDataBuffer.BufferRef, CVDConstraint);
+	if (CollisionConstraint)
+	{
+		FParticlePair ParticlePair = CollisionConstraint->GetConstrainedParticles();
+		const bool bIsInRelevancyVolume = IsInRelevancyVolume(ParticlePair[0]) || IsInRelevancyVolume(ParticlePair[1]);
+		if (!bIsInRelevancyVolume)
+		{
+			return;
+		}
+
+		FChaosVDConstraint CVDConstraint = FChaosVDDataWrapperUtils::BuildConstraintDataWrapperFromConstraint(*CollisionConstraint);
+		CVDConstraint.SolverID = CVDContextData->Id;
+
+		FChaosVDScopedTLSBufferAccessor TLSDataBuffer;
+		VisualDebugger::WriteDataToBuffer(TLSDataBuffer.BufferRef, CVDConstraint);
 
-	TraceBinaryData(TLSDataBuffer.BufferRef, FChaosVDConstraint::WrapperTypeName);
+		TraceBinaryData(TLSDataBuffer.BufferRef, FChaosVDConstraint::WrapperTypeName);
+	}
 }
 
 void FChaosVisualDebuggerTrace::TraceCollisionConstraintView(TArrayView<Chaos::FPBDCollisionConstraint* const> CollisionConstraintView)
@@ -837,6 +919,24 @@ void FChaosVisualDebuggerTrace::TraceSceneQueryStart(const Chaos::FImplicitObjec
 		return;
 	}
 
+	auto IsInRelevancyVolumeAtPosition = [InputGeometry](const FVector& InPosition, const FQuat& InOrientation)
+	{
+		return InputGeometry ? IsInRelevancyVolume(InputGeometry, Chaos::FRigidTransform3(InPosition, InOrientation)) : IsInRelevancyVolume(InPosition);
+	};
+
+	// In overlap queries, only the start position is valid
+	const bool bEvaluateBothPositions = QueryType != EChaosVDSceneQueryType::Overlap;
+	bool bIsRelevant = bEvaluateBothPositions ? (IsInRelevancyVolumeAtPosition(Start, GeometryOrientation) | IsInRelevancyVolumeAtPosition(End, GeometryOrientation))
+												: IsInRelevancyVolumeAtPosition(Start, GeometryOrientation);
+
+	if (!bIsRelevant)
+	{
+		// TODO: We need to find a better way to do this. Scene queries and shape visit data are traced independently, and the stitched together during load, using the following context ID.
+		// therefore invalidating this id is not supposed to change is the less intrusive way to avoid tracing the shape visit data for now
+		const_cast<FChaosVDContext*>(CVDContextData)->Id  = INDEX_NONE;
+		return;
+	}
+
 	FChaosVDQueryDataWrapper WrappedQueryData;
 
 	if (InputGeometry)
@@ -885,6 +985,12 @@ void FChaosVisualDebuggerTrace::TraceSceneQueryVisit(FChaosVDQueryVisitStep&& In
 	{
 		return;
 	}
+
+	if (CVDContextData->Id == INDEX_NONE)
+	{
+		return;
+	}
+
 	const bool bIsQueryContext = CVDContextData->Type == static_cast<int32>(EChaosVDContextType::Query) ||  CVDContextData->Type == static_cast<int32>(EChaosVDContextType::SubTraceQuery);
 
 	if (!ensure(bIsQueryContext))
@@ -978,6 +1084,11 @@ void FChaosVisualDebuggerTrace::TraceDebugDrawBox(const FBox& InBox, FName Tag,
 		return;
 	}
 
+	if (!IsInRelevancyVolume(InBox))
+	{
+		return;
+	}
+
 	FChaosVDDebugDrawBoxDataWrapper DataWrapper;
 	DataWrapper.SolverID = SolverID;
 	DataWrapper.Tag = Tag;
@@ -1006,6 +1117,11 @@ void FChaosVisualDebuggerTrace::TraceDebugDrawLine(const FVector& InStartLocatio
 		return;
 	}
 
+	if (!IsInRelevancyVolume(InStartLocation, InEndLocation))
+	{
+		return;
+	}
+
 	FChaosVDDebugDrawLineDataWrapper DataWrapper;
 	DataWrapper.SolverID = SolverID;
 	DataWrapper.Tag = Tag;
@@ -1035,12 +1151,19 @@ void FChaosVisualDebuggerTrace::TraceDebugDrawVector(const FVector& InStartLocat
 		return;
 	}
 
+	FVector EndLocation = InStartLocation + InVector;
+
+	if (!IsInRelevancyVolume(InStartLocation, EndLocation))
+	{
+		return;
+	}
+
 	FChaosVDDebugDrawLineDataWrapper DataWrapper;
 	DataWrapper.SolverID = SolverID;
 	DataWrapper.Tag = Tag;
 	DataWrapper.Color = Color;
 	DataWrapper.StartLocation = InStartLocation;
-	DataWrapper.EndLocation = InStartLocation + InVector;
+	DataWrapper.EndLocation = EndLocation;
 	DataWrapper.bIsArrow = true;
 	DataWrapper.ThreadContext = IsInGameThread() ? EChaosVDParticleContext::GameThread : EChaosVDParticleContext::PhysicsThread;
 
@@ -1065,6 +1188,11 @@ void FChaosVisualDebuggerTrace::TraceDebugDrawSphere(const FVector& Center, floa
 		return;
 	}
 
+	if (!IsInRelevancyVolume(Center, Radius))
+	{
+		return;
+	}
+
 	FChaosVDDebugDrawSphereDataWrapper DataWrapper;
 	DataWrapper.SolverID = SolverID;
 	DataWrapper.Tag = Tag;
@@ -1087,6 +1215,11 @@ void FChaosVisualDebuggerTrace::TraceDebugDrawImplicitObject(const Chaos::FImpli
 	{
 		return;
 	}
+
+	if (!Implicit)
+	{
+		return;
+	}
 	
 	// Generic Debug Draw might not have a context if they are recorded from the game thread in game code, in that case we allow the trace anyway
 	if (!CanTraceDebugDrawShape(SolverID))
@@ -1094,6 +1227,14 @@ void FChaosVisualDebuggerTrace::TraceDebugDrawImplicitObject(const Chaos::FImpli
 		return;
 	}
 
+	if (TraceRelevancyVolume.IsValid)
+	{
+		if (Implicit->HasBoundingBox() && !IsInRelevancyVolume(Implicit->CalculateTransformedBounds(InParentTransform)))
+		{
+			return;
+		}
+	}
+
 	FChaosVDDebugDrawImplicitObjectDataWrapper DataWrapper;
 	DataWrapper.SolverID = SolverID;
 	DataWrapper.Tag = Tag;
@@ -1114,6 +1255,15 @@ void FChaosVisualDebuggerTrace::TraceDebugDrawImplicitObject(const Chaos::FImpli
 	TraceBinaryData(TLSDataBuffer.BufferRef, FChaosVDDebugDrawImplicitObjectDataWrapper::WrapperTypeName);
 }
 
+void FChaosVisualDebuggerTrace::SetTraceRelevancyVolume(const FBox& InTraceRelevancyVolume)
+{
+	TraceRelevancyVolume = InTraceRelevancyVolume;
+	
+	UE_TRACE_LOG(ChaosVDLogger, ChaosVDTraceRelevancyVolume, ChaosVDChannel)
+		<< CVD_TRACE_VECTOR_ON_EVENT(ChaosVDTraceRelevancyVolume, BoxMin, InTraceRelevancyVolume.Min)
+		<< CVD_TRACE_VECTOR_ON_EVENT(ChaosVDTraceRelevancyVolume, BoxMax, InTraceRelevancyVolume.Max);
+}
+
 bool FChaosVisualDebuggerTrace::IsTracing()
 {
 	return bIsTracing;
@@ -1172,6 +1322,8 @@ void FChaosVisualDebuggerTrace::Reset()
 {
 	CVDNameTable->ResetTable();
 
+	TraceRelevancyVolume = FBox(EForceInit::ForceInitToZero);
+
 	{
 		UE::TWriteScopeLock WriteLock(DeltaRecordingStatesLock);
 		RequestedFullCaptureSolverIDs.Reset();
@@ -1209,6 +1361,26 @@ void FChaosVisualDebuggerTrace::TraceArchiveHeader()
 	TraceBinaryData(HeaderDataBuffer, FChaosVDArchiveHeader::WrapperTypeName, EChaosVDTraceBinaryDataOptions::ForceTrace);
 }
 
+void FChaosVisualDebuggerTrace::OverrideDefaultEnabledDataChannels(TConstArrayView<FString> EnabledDataChannelsOverrideList)
+{
+	UE_LOG(LogChaos, Log, TEXT("[%s] Channel list override provided - Enabling [%d] Requested channels..."), ANSI_TO_TCHAR(__FUNCTION__), EnabledDataChannelsOverrideList.Num());
+
+	using namespace Chaos::VisualDebugger;
+	FChaosVDDataChannelsManager::Get().EnumerateChannels([&EnabledDataChannelsOverrideList](const TSharedRef<FChaosVDOptionalDataChannel>& Channel)
+	{
+		if (Channel->CanChangeEnabledState())
+		{
+			// This is far from efficient, but this will be called once when the recording start command is executed, and we only have a handful of channels
+			const FString ChannelIdAsString = Channel->GetId().ToString();
+			const bool bChannelShouldBeEnabled = EnabledDataChannelsOverrideList.Contains(ChannelIdAsString);
+			Channel->SetChannelEnabled(bChannelShouldBeEnabled);
+
+			UE_LOG(LogChaos, Log, TEXT("[%s] Setting enabled state for channel [%s] to [%s]..."), ANSI_TO_TCHAR(__FUNCTION__), *ChannelIdAsString, bChannelShouldBeEnabled ? TEXT("True") : TEXT("False"));
+		}
+		return true;
+	});
+}
+
 void FChaosVisualDebuggerTrace::HandleRecordingStart()
 {
 	Reset();
@@ -1219,23 +1391,7 @@ void FChaosVisualDebuggerTrace::HandleRecordingStart()
 	{
 		TArray<FString> ParsedChannels;
 		Chaos::VisualDebugger::ParseChannelListFromCommandArgument(ParsedChannels, CommandlineEnabledCVDChannels);
-
-		UE_LOG(LogChaos, Log, TEXT("[%s] Channel list override provided via commandline - Enabling [%d] Requested channels..."), ANSI_TO_TCHAR(__FUNCTION__), ParsedChannels.Num());
-
-		using namespace Chaos::VisualDebugger;
-		FChaosVDDataChannelsManager::Get().EnumerateChannels([&ParsedChannels](const TSharedRef<FChaosVDOptionalDataChannel>& Channel)
-		{
-			if (Channel->CanChangeEnabledState())
-			{
-				// This is far from efficient, but this will be called once when the recording start command is executed, and we only have a handful of channels
-				const FString ChannelIdAsString = Channel->GetId().ToString();
-				const bool bChannelShouldBeEnabled = ParsedChannels.Contains(ChannelIdAsString);
-				Channel->SetChannelEnabled(bChannelShouldBeEnabled);
-
-				UE_LOG(LogChaos, Log, TEXT("[%s] Setting enabled state for channel [%s] to [%s]..."), ANSI_TO_TCHAR(__FUNCTION__), *ChannelIdAsString, bChannelShouldBeEnabled ? TEXT("True") : TEXT("False"));
-			}
-			return true;
-		});
+		OverrideDefaultEnabledDataChannels(ParsedChannels);
 	}
 
 	TraceArchiveHeader();
diff --git a/Engine/Source/Runtime/Experimental/Chaos/Public/Chaos/Collision/ParticlePairMidPhase.h b/Engine/Source/Runtime/Experimental/Chaos/Public/Chaos/Collision/ParticlePairMidPhase.h
--- a/Engine/Source/Runtime/Experimental/Chaos/Public/Chaos/Collision/ParticlePairMidPhase.h
+++ b/Engine/Source/Runtime/Experimental/Chaos/Public/Chaos/Collision/ParticlePairMidPhase.h
@@ -216,9 +216,9 @@ namespace Chaos
 			const Private::FCollisionParticlePairKey& InKey,
 			const FCollisionContext& Context);
 
-		inline FGeometryParticleHandle* GetParticle0() { return Particle0; }
+		inline FGeometryParticleHandle* GetParticle0() const { return Particle0; }
 
-		inline FGeometryParticleHandle* GetParticle1() { return Particle1; }
+		inline FGeometryParticleHandle* GetParticle1() const { return Particle1; }
 
 		inline const Private::FCollisionParticlePairKey& GetKey() const { return ParticlePairKey; }
 
diff --git a/Engine/Source/Runtime/Experimental/Chaos/Public/ChaosVisualDebugger/ChaosVDTraceMacros.h b/Engine/Source/Runtime/Experimental/Chaos/Public/ChaosVisualDebugger/ChaosVDTraceMacros.h
--- a/Engine/Source/Runtime/Experimental/Chaos/Public/ChaosVisualDebugger/ChaosVDTraceMacros.h
+++ b/Engine/Source/Runtime/Experimental/Chaos/Public/ChaosVisualDebugger/ChaosVDTraceMacros.h
@@ -214,6 +214,11 @@
 		FChaosVisualDebuggerTrace::TraceDebugDrawImplicitObject(ImplicitObject, ParentTransform,  ##__VA_ARGS__)
 	#endif
 
+	#ifndef CVD_SET_RELEVANCY_VOLUME
+		#define CVD_SET_RELEVANCY_VOLUME(VolumeBox) \
+		FChaosVisualDebuggerTrace::SetTraceRelevancyVolume(VolumeBox)
+	#endif
+
 #else // WITH_CHAOS_VISUAL_DEBUGGER
 
 	#ifndef CVD_TRACE_PARTICLE
@@ -342,4 +347,8 @@
 		#define CVD_TRACE_DEBUG_DRAW_IMPLICIT_OBJECT(ImplicitObject, ParentTransform, ...)
 	#endif
 
+	#ifndef CVD_SET_RELEVANCY_VOLUME
+		#define CVD_SET_RELEVANCY_VOLUME(VolumeBox)
+	#endif
+
 #endif // WITH_CHAOS_VISUAL_DEBUGGER
\ No newline at end of file
diff --git a/Engine/Source/Runtime/Experimental/Chaos/Public/ChaosVisualDebugger/ChaosVisualDebuggerTrace.h b/Engine/Source/Runtime/Experimental/Chaos/Public/ChaosVisualDebugger/ChaosVisualDebuggerTrace.h
--- a/Engine/Source/Runtime/Experimental/Chaos/Public/ChaosVisualDebugger/ChaosVisualDebuggerTrace.h
+++ b/Engine/Source/Runtime/Experimental/Chaos/Public/ChaosVisualDebugger/ChaosVisualDebuggerTrace.h
@@ -130,6 +130,11 @@ UE_TRACE_EVENT_BEGIN_EXTERN(ChaosVDLogger, ChaosVDUsesAutoRTFM)
 	UE_TRACE_EVENT_FIELD(bool, bUsingAutoRTFM)
 UE_TRACE_EVENT_END()
 
+UE_TRACE_EVENT_BEGIN_EXTERN(ChaosVDLogger, ChaosVDTraceRelevancyVolume)
+	CVD_DEFINE_TRACE_VECTOR(Chaos::FReal, BoxMin)
+	CVD_DEFINE_TRACE_VECTOR(Chaos::FReal, BoxMax)
+UE_TRACE_EVENT_END()
+
 UE_TRACE_EVENT_BEGIN_EXTERN(ChaosVDLogger, ChaosVDDummyEvent)
 	UE_TRACE_EVENT_FIELD(int32, SolverID)
 UE_TRACE_EVENT_END()
@@ -345,6 +350,106 @@ public:
 	 */
 	static CHAOS_API void TraceDebugDrawImplicitObject(const Chaos::FImplicitObject* Implicit, const FTransform& InParentTransform, FName Tag = NAME_None, FColor Color = FColor::Blue, int32 SolverID = INDEX_NONE);
 
+	/**
+	 * Sets a bounding box to be used to filter out particles and other data outside of it.
+	 * @param InTraceRelevancyVolume Volume within which object's data should be traced
+	 */
+	static CHAOS_API void SetTraceRelevancyVolume(const FBox& InTraceRelevancyVolume);
+
+	/*
+	 * Returns the currently bounding box used to determine if we should trace data or not, for objects that can be applied to.
+	 */
+	static FBox GetTraceRelevancyVolume()
+	{
+		return TraceRelevancyVolume;
+	}
+
+
+	/**
+	 * Evaluates the provided bounds to see if it is withing the current relevancy volume, if any.
+	 * @param InChaosBounds Bounds to evaluate
+	 * @return Returns true if the bounds are relevant or if there is no relevancy volume configured
+	 */
+	static bool IsInRelevancyVolume(const Chaos::FAABB3& InChaosBounds)
+	{
+		if (!TraceRelevancyVolume.IsValid)
+		{
+			// An invalid volume means we want to trace everything
+			return true;
+		}
+
+		return TraceRelevancyVolume.Intersect(FBox(InChaosBounds.Min(), InChaosBounds.Max()));
+	}
+
+	/**
+	 * Evaluates the provided bounds to see if it is withing the current relevancy volume, if any.
+	 * @param InBounds Bounds to evaluate
+	 * @return Returns true if the bounds are relevant or if there is no relevancy volume configured
+	 */
+	static bool IsInRelevancyVolume(const FBox& InBounds)
+	{
+		if (!TraceRelevancyVolume.IsValid)
+		{
+			// An invalid volume means we want to trace everything
+			return true;
+		}
+
+		return TraceRelevancyVolume.Intersect(InBounds);
+	}
+
+	/**
+	 * Evaluates the provided sphere to see if it is withing the current relevancy volume, if any.
+	 * @param Center Center of the sphere
+	 * @param Radius Radius of the Sphere
+	 * @return Returns true if the sphere is relevant or if there is no relevancy volume configured
+	 */
+	static bool IsInRelevancyVolume(const FVector& Center, double Radius)
+	{
+		if (!TraceRelevancyVolume.IsValid)
+		{
+			// An invalid volume means we want to trace everything
+			return true;
+		}
+	
+		return FMath::SphereAABBIntersection(Center, Radius * Radius, TraceRelevancyVolume);
+	}
+
+	/**
+	 * Evaluates the provided point to see if it is withing the current relevancy volume, if any.
+	 * @param InPosition point to evaluate
+	 * @return Returns true if the point is relevant or if there is no relevancy volume configured
+	 */
+	static bool IsInRelevancyVolume(const Chaos::FVec3& InPosition)
+	{
+		if (!TraceRelevancyVolume.IsValid)
+		{
+			// An invalid volume means we want to trace everything
+			return true;
+		}
+
+		return FMath::PointBoxIntersection(InPosition, TraceRelevancyVolume);
+	}
+
+	/**
+	 * Evaluates the provided line to see if it is withing the current relevancy volume, if any.
+	 * @param InStartPoint Start point of the line to evaluate
+	 * @param InEndPoint End point of the line to evaluate
+	 * @return Returns true if the line is relevant or if there is no relevancy volume configured
+	 */
+	static bool IsInRelevancyVolume(const Chaos::FVec3& InStartPoint, const Chaos::FVec3& InEndPoint)
+	{
+		if (!TraceRelevancyVolume.IsValid)
+		{
+			// An invalid volume means we want to trace everything
+			return true;
+		}
+
+		return TraceRelevancyVolume.Intersect(FBox(FVector::Min(InStartPoint, InEndPoint), FVector::Max(InStartPoint, InEndPoint)));
+	}
+
+	static bool IsInRelevancyVolume(const Chaos::FGeometryParticleHandle* InParticleHandle);
+	static bool IsInRelevancyVolume(const Chaos::FImplicitObject* Geometry, const Chaos::FRigidTransform3& InTransform);
+
 	/** Returns the ID of the main solver of the provided world
 	 * @param World Ptr to the world we want to get the solverID from
 	 */
@@ -385,6 +490,12 @@ public:
 
 	static CHAOS_API TSharedRef<FChaosVDSerializableNameTable>& GetNameTableInstance();
 
+	/**
+	 * If the provided array is not empty, all data channels that are not in the list and can be disabled, will be disabled
+	 * @param EnabledDataChannelsOverrideList List of the only channels that should be enabled by default
+	 */
+	static CHAOS_API void OverrideDefaultEnabledDataChannels(TConstArrayView<FString> EnabledDataChannelsOverrideList);
+
 private:
 
 	/**
@@ -447,6 +558,8 @@ private:
 	static TSet<int32> RequestedFullCaptureSolverIDs;
 	static TSharedRef<FChaosVDSerializableNameTable> CVDNameTable;
 
+	static FBox TraceRelevancyVolume;
+
 	static std::atomic<bool> bIsTracing;
 
 	static FRWLock DeltaRecordingStatesLock;
diff --git a/Engine/Source/Runtime/Experimental/ChaosSolverEngine/Private/Chaos/ChaosVDRemoteSessionsManager.cpp b/Engine/Source/Runtime/Experimental/ChaosSolverEngine/Private/Chaos/ChaosVDRemoteSessionsManager.cpp
--- a/Engine/Source/Runtime/Experimental/ChaosSolverEngine/Private/Chaos/ChaosVDRemoteSessionsManager.cpp
+++ b/Engine/Source/Runtime/Experimental/ChaosSolverEngine/Private/Chaos/ChaosVDRemoteSessionsManager.cpp
@@ -463,6 +463,7 @@ void FChaosVDRemoteSessionsManager::HandleRecordingStartCommandMessage(const FCh
 			RecordingParams.Add(TEXT("File"));
 		}
 
+		FChaosVisualDebuggerTrace::OverrideDefaultEnabledDataChannels(InMessage.DataChannelsEnabledOverrideList);
 		FChaosVDRuntimeModule::Get().StartRecording(RecordingParams);
 	};
 
diff --git a/Engine/Source/Runtime/Experimental/ChaosSolverEngine/Public/Chaos/ChaosVDRemoteSessionsManager.h b/Engine/Source/Runtime/Experimental/ChaosSolverEngine/Public/Chaos/ChaosVDRemoteSessionsManager.h
--- a/Engine/Source/Runtime/Experimental/ChaosSolverEngine/Public/Chaos/ChaosVDRemoteSessionsManager.h
+++ b/Engine/Source/Runtime/Experimental/ChaosSolverEngine/Public/Chaos/ChaosVDRemoteSessionsManager.h
@@ -53,6 +53,9 @@ struct FChaosVDStartRecordingCommandMessage
 	
 	UPROPERTY()
 	FString Target;
+
+	UPROPERTY()
+	TArray<FString> DataChannelsEnabledOverrideList;
 };
 
 USTRUCT()
