Add an optional EarlyStartup attribute to indexed cache storage
The syntax is now
[IndexedCacheStorage]
+Storage=(Name="StorageName",Index=StorageIndex,Mount="StorageMount", EarlyStartup=IniName:[IniSection]:IniKey)

It will try to fetch the value in IniName set as
[IniSection]
IniKey=256MB

During IndexedCacheStorage module initialization, we gather all entries with an EarlyStartup attribute, and aggregate them into a single call to a dummy EarlyStartupCache entry. This entry is declared the same way as the other entries, so for example:
[IndexedCacheStorage]
...
+Storage=(Name="EarlyStartupCache",Index=1,Mount="es-cache")

It will be always created/destroyed whenever a cache with EarlyStartupSize is not up-to-date. The default cache (GamePersistentDownloadDir) is providing its EarlyStartupSize with GetPersistentDownloadDirEarlyStartupSize

This allows the user to know they need to free up 2.2GB of disk space rather that being requested to free 1.2GB and then 1GB.

The size needs to be exactly the same as the one later used when calling CreateCacheStorage, otherwise we may end up creating/destroying caches all the time.

diff --git a/Engine/Source/Runtime/Experimental/IndexedCacheStorage/Private/IndexedCacheStorage.cpp b/Engine/Source/Runtime/Experimental/IndexedCacheStorage/Private/IndexedCacheStorage.cpp
--- a/Engine/Source/Runtime/Experimental/IndexedCacheStorage/Private/IndexedCacheStorage.cpp
+++ b/Engine/Source/Runtime/Experimental/IndexedCacheStorage/Private/IndexedCacheStorage.cpp
@@ -34,7 +34,7 @@ public:
 	{
 	}
 	
-	virtual void GetCacheStorageInfo(int32 InCacheStorageIndex, uint64& OutCacheStorageSize) override
+	virtual void GetCacheStorageInfo(int32 InCacheStorageIndex, uint64& OutCacheStorageDataSize, uint64& OutCacheStorageJournalSize) override
 	{
 	}
 	
@@ -51,6 +51,16 @@ public:
 	{
 	}
 
+	virtual uint64 GetPersistentDownloadDirEarlyStartupSize() override
+	{
+		return 0;
+	}
+
+	virtual uint64 GetCacheStorageJournalSize(uint64 CacheDataSize, uint32 CacheIndex) override
+	{
+		return 0;
+	}
+
 };
 
 #if !defined(HAS_CUSTOM_INDEXED_CACHE_STORAGE)
diff --git a/Engine/Source/Runtime/Experimental/IndexedCacheStorage/Private/IndexedCacheStorage.h b/Engine/Source/Runtime/Experimental/IndexedCacheStorage/Private/IndexedCacheStorage.h
--- a/Engine/Source/Runtime/Experimental/IndexedCacheStorage/Private/IndexedCacheStorage.h
+++ b/Engine/Source/Runtime/Experimental/IndexedCacheStorage/Private/IndexedCacheStorage.h
@@ -66,7 +66,7 @@ public:
  	  * @param InCacheStorageIndex - the index of the cache
  	  * @param OutCacheStorageSize - How large the cache storage is
 	  */
-	virtual void GetCacheStorageInfo(int32 InCacheStorageIndex, uint64& OutCacheStorageSize) = 0;
+	virtual void GetCacheStorageInfo(int32 InCacheStorageIndex, uint64& OutCacheStorageDataSize, uint64& OutCacheStorageJournalSize) = 0;
 	
 	/**
 	  * Gets the platform mount path for a cached storage based on its generic Name
@@ -90,6 +90,11 @@ public:
  	  * @param InMountName - the generic mount name
 	  */
 	virtual void UnmountCacheStorage(const FString& InMountName) = 0;
+
+	virtual uint64 GetPersistentDownloadDirEarlyStartupSize() = 0;
+
+	virtual uint64 GetCacheStorageJournalSize(uint64 CacheDataSize, uint32 CacheIndex) = 0;
+
 };
 
 }
diff --git a/Engine/Source/Runtime/Experimental/IndexedCacheStorage/Private/IndexedCacheStorageManager.cpp b/Engine/Source/Runtime/Experimental/IndexedCacheStorage/Private/IndexedCacheStorageManager.cpp
--- a/Engine/Source/Runtime/Experimental/IndexedCacheStorage/Private/IndexedCacheStorageManager.cpp
+++ b/Engine/Source/Runtime/Experimental/IndexedCacheStorage/Private/IndexedCacheStorageManager.cpp
@@ -11,9 +11,12 @@
 #include "HAL/PlatformStackWalk.h"
 #include "Misc/ScopeExit.h"
 #include "Misc/ScopeRWLock.h"
+#include "String/LexFromString.h"
+#include "String/ParseTokens.h"
 
 DECLARE_LOG_CATEGORY_EXTERN(LogIndexedCacheStorageManager, Display, All);
 DEFINE_LOG_CATEGORY(LogIndexedCacheStorageManager);
+#define LOG_ICS(Verbosity, STRING, ...) UE_LOG(LogIndexedCacheStorageManager, Verbosity, TEXT("%s: ") STRING, ANSI_TO_TCHAR(__FUNCTION__), ##__VA_ARGS__ )
 
 namespace Experimental
 {
@@ -42,9 +45,10 @@ public:
 		return Strings[int32(Val)];
 	}
 	
-	FCacheStorageMetaData(int32 InCacheIndex, const FString& InMountName, bool bCacheExists)
+	FCacheStorageMetaData(int32 InCacheIndex, const FString& InMountName, bool bCacheExists, uint64 InEarlyStartupSize)
 	: MountName(InMountName)
 	, CacheIndex(InCacheIndex)
+	, EarlyStartupSize(InEarlyStartupSize)
 	{
 		CacheStorageState = bCacheExists ? ECacheStorageState::Created : ECacheStorageState::Deleted;
 	}
@@ -59,6 +63,11 @@ public:
 		return MountName;
 	}
 
+	uint64 GetEarlyStartupSize() const
+	{
+		return EarlyStartupSize;
+	}
+
 	FString GetMountPath() const
 	{
 		FString MountedPath;
@@ -83,6 +92,8 @@ public:
 	{
 		UE::TUniqueLock Lock(Mutex);
 		 
+		check(EarlyStartupSize == 0 || RequestNumberOfBytes == EarlyStartupSize);
+
 		UE_LOG(LogIndexedCacheStorageManager, Display, TEXT("CreateCacheStorage request for %ull bytes at cache %d. Current state %s"), RequestNumberOfBytes, CacheIndex, LexToString(CacheStorageState));
 
 		WaitForPredicateWhileLocked([this]()->bool 
@@ -160,21 +171,22 @@ public:
 
 		UE_LOG(LogIndexedCacheStorageManager, Display, TEXT("GetCacheStorageInfo request at cache %d. Current state %s"), CacheIndex, LexToString(CacheStorageState));
 
-		uint64 ExistingSpace = 0;
+		uint64 CurrentDataSize = 0;
+		uint64 CurrentJournalSize = 0;
 
 		if (CacheStorageState == ECacheStorageState::Deleted)
 		{
-			UE_LOG(LogIndexedCacheStorageManager, Display, TEXT("GetCacheStorageInfo cache %d is empty"), CacheIndex, ExistingSpace);
-			return ExistingSpace;
+			UE_LOG(LogIndexedCacheStorageManager, Display, TEXT("GetCacheStorageInfo cache %d is empty"), CacheIndex, CurrentDataSize);
+			return CurrentDataSize;
 		}
 		
 		UE_LOG(LogIndexedCacheStorageManager, Display, TEXT("GetCacheStorageInfo cache %d"), CacheIndex);
-		IIndexedCacheStorage::Get().GetCacheStorageInfo(CacheIndex, ExistingSpace);
-		UE_LOG(LogIndexedCacheStorageManager, Display, TEXT("GetCacheStorageInfo cache %d is %llu bytes"), CacheIndex, ExistingSpace);
+		IIndexedCacheStorage::Get().GetCacheStorageInfo(CacheIndex, CurrentDataSize, CurrentJournalSize);
+		UE_LOG(LogIndexedCacheStorageManager, Display, TEXT("GetCacheStorageInfo cache %d = data %llu bytes / journal %llu bytes"), CacheIndex, CurrentDataSize, CurrentJournalSize);
 		
 		NotifyWaitersWhileLocked();
 
-		return ExistingSpace;
+		return CurrentDataSize;
 	}
 
 	FString MountCacheStorage()
@@ -261,6 +273,7 @@ public:
 private:
 	const FString MountName;
 	const int32 CacheIndex;
+	const uint64 EarlyStartupSize;
 	int32 MountRefCount = 0;
 	ECacheStorageState CacheStorageState = ECacheStorageState::Deleted;
 	UE::FMutex Mutex;
@@ -295,9 +308,186 @@ FIndexedCacheStorageManager& FIndexedCacheStorageManager::Get()
 	return Instance;
 }
 
+void FIndexedCacheStorageManager::EnsureEarlyStartupCache(const TArray<bool>& InCacheStorageAlreadyExists)
+{
+	int32 EarlyStartupCacheIndex = GetStorageIndex(TEXT("EarlyStartupCache"));
+	if (EarlyStartupCacheIndex == 0)
+	{
+		LOG_ICS(Display, TEXT("No Storage named 'EarlyStartupCache' has been set. Skipping"));
+		return;
+	}
+
+	TArray<bool> CacheStorageAlreadyExists(InCacheStorageAlreadyExists);
+
+	if (CacheStorageAlreadyExists[EarlyStartupCacheIndex])
+	{
+		LOG_ICS(Error, TEXT("'EarlyStartupCache' entry already exists. Destroying it"));
+		DestroyCacheStorage(EarlyStartupCacheIndex);
+		CacheStorageAlreadyExists[EarlyStartupCacheIndex] = false;
+	}
+
+	TMap<int32, uint64> EarlyStartupCaches;
+	{
+		uint64 PersistentDownloadDirEarlyStartupSize = IIndexedCacheStorage::Get().GetPersistentDownloadDirEarlyStartupSize();
+		if (PersistentDownloadDirEarlyStartupSize > 0)
+		{
+			EarlyStartupCaches.Add(0, PersistentDownloadDirEarlyStartupSize);
+		}
+		
+		for (int32 CacheIndex = 0; CacheIndex < IndexedCacheStorageMetaData->IndexedCacheEntriesMetaData.Num(); ++CacheIndex)
+		{
+			FCacheStorageMetaData* CacheStorageMetaData = IndexedCacheStorageMetaData->IndexedCacheEntriesMetaData[CacheIndex].Get();
+			if (CacheStorageMetaData == nullptr)
+			{
+				continue;
+			}
+
+			uint64 EarlyStartupSize = CacheStorageMetaData->GetEarlyStartupSize();
+
+			if (EarlyStartupSize > 0)
+			{
+				EarlyStartupCaches.Add(CacheIndex, EarlyStartupSize);
+			}
+		}
+	}
+
+	uint64 EarlyStartupCacheSize = 0;
+	uint64 NumCacheStorageInGroup = 0;
+
+	for (const TPair<int32, uint64>& EarlyStartupCache : EarlyStartupCaches)
+	{
+		uint64 CacheStorageSizeForCommonAlloc = 0;
+		const int32 CacheStorageIndex = EarlyStartupCache.Key;
+		const uint64 CacheStorageNewDataSize = EarlyStartupCache.Value;
+		const uint64 CacheStorageNewJournalSize = IIndexedCacheStorage::Get().GetCacheStorageJournalSize(CacheStorageNewDataSize, CacheStorageIndex);
+		if (CacheStorageAlreadyExists[CacheStorageIndex])
+		{
+			uint64 CacheStorageCurrentDataSize = 0;
+			uint64 CacheStorageCurrentJournalSize = 0;
+			IIndexedCacheStorage::Get().GetCacheStorageInfo(CacheStorageIndex, CacheStorageCurrentDataSize, CacheStorageCurrentJournalSize);
+			if (CacheStorageCurrentDataSize != CacheStorageNewDataSize || CacheStorageCurrentJournalSize != CacheStorageNewJournalSize)
+			{
+				LOG_ICS(Display, TEXT("Cache %d size mismatch: data: %llu vs %llu. journal: %llu vs %llu. Destroying cache"),  
+					   CacheStorageIndex, CacheStorageCurrentDataSize, CacheStorageNewDataSize, CacheStorageCurrentJournalSize, CacheStorageNewJournalSize);
+
+				DestroyCacheStorage(CacheStorageIndex);
+				CacheStorageSizeForCommonAlloc = CacheStorageNewDataSize + CacheStorageNewJournalSize;
+			}
+			else
+			{
+				LOG_ICS(Display, TEXT("Cache %d up to date with %llu bytes. Skipping from grouped alloc"), CacheStorageIndex, CacheStorageCurrentDataSize);
+			}
+		}
+		else if (CacheStorageNewDataSize > 0)
+		{
+			LOG_ICS(Display, TEXT("Cache %d does not exist, adding %llu bytes"), CacheStorageIndex, CacheStorageNewDataSize);
+			CacheStorageSizeForCommonAlloc = CacheStorageNewDataSize + CacheStorageNewJournalSize;
+		}
+
+		if (CacheStorageSizeForCommonAlloc > 0)
+		{
+			EarlyStartupCacheSize += CacheStorageSizeForCommonAlloc;
+			++NumCacheStorageInGroup;
+		}
+	}
+
+	if (EarlyStartupCacheSize == 0)
+	{
+		LOG_ICS(Display, TEXT("All caches are up-to-date, skipping reservation"));
+		return;
+	}
+
+	if (NumCacheStorageInGroup <= 1)
+	{
+		LOG_ICS(Display, TEXT("Only one cache requested, skipping reservation"));
+		return;
+	}
+
+	{
+		// Make sure we alloc at least 2MB
+		const uint64 MinGroupAllocSize = 2 << 20;
+		EarlyStartupCacheSize = FMath::Max(MinGroupAllocSize, EarlyStartupCacheSize);
+
+		// Remove alloc meta data because we want to allocate exactly EarlyStartupCacheSize
+		uint64 EarlyStartupCacheSizeMetaData = IIndexedCacheStorage::Get().GetCacheStorageJournalSize(EarlyStartupCacheSize, EarlyStartupCacheIndex);
+		// make sure we don't underflow
+		EarlyStartupCacheSize -= FMath::Min(EarlyStartupCacheSizeMetaData, EarlyStartupCacheSize);
+
+		// Make sure we alloc at least the metadata size
+		EarlyStartupCacheSize = FMath::Max(EarlyStartupCacheSizeMetaData, EarlyStartupCacheSize);
+	}
+
+	LOG_ICS(Display, TEXT("Ensure %llu bytes through cache %d"), EarlyStartupCacheSize, EarlyStartupCacheIndex);
+	bool bCacheCreated = true;
+	while (true)
+	{
+		if (CreateCacheStorage(EarlyStartupCacheSize, EarlyStartupCacheIndex))
+		{
+			DestroyCacheStorage(EarlyStartupCacheIndex);
+			break;
+		}
+	}
+}
+
+////////////////////////////////////////////////////////////////////////////////
+uint64 ParseSizeParam(FStringView Value)
+{
+	Value = Value.TrimStartAndEnd();
+
+	uint64 Size = 0;
+	LexFromString(Size, Value);
+	if (Size > 0)
+	{
+		if (Value.EndsWith(TEXT("GB"))) return Size << 30;
+		if (Value.EndsWith(TEXT("MB"))) return Size << 20;
+		if (Value.EndsWith(TEXT("KB"))) return Size << 10;
+	}
+	return Size;
+}
+
+// Parse IniString expected to be in the form of OutIniName:[OutIniSection]:OutIniKey
+bool GetIniParameters(FString& OutIniName, FString& OutIniSection, FString& OutIniKey, const FString& IniString)
+{
+	TArray<FStringView> ConfigStringArrayNameSize;
+	UE::String::ParseTokens(IniString, TEXTVIEW(":"), [&](FStringView Element)
+	{
+		ConfigStringArrayNameSize.Add(Element);
+	});
+
+	if (ConfigStringArrayNameSize.Num() != 3 || 
+		ConfigStringArrayNameSize[0].Len() == 0 ||
+		ConfigStringArrayNameSize[1].Len() == 0 ||
+		ConfigStringArrayNameSize[2].Len() == 0 )
+	{
+		UE_LOG(LogIndexedCacheStorageManager, Error, TEXT("Error while parsing %s: expecting OutIniName:[OutIniSection]:OutIniKey"), *IniString);
+		return false;
+	}
+
+	FStringView IniSectionView = ConfigStringArrayNameSize[1];
+	if (IniSectionView[0] != TEXT('[') || IniSectionView[IniSectionView.Len()-1] != TEXT(']'))
+	{
+		UE_LOG(LogIndexedCacheStorageManager, Error, TEXT("Error while parsing ini section %.*s"), IniSectionView.Len(), IniSectionView.GetData());
+		return false;
+	}
+
+	IniSectionView.RemovePrefix(1);
+	IniSectionView.RemoveSuffix(1);
+
+	OutIniName = ConfigStringArrayNameSize[0];
+	OutIniSection = IniSectionView;
+	OutIniKey = ConfigStringArrayNameSize[2];
+
+	return true;
+}
+
 void FIndexedCacheStorageManager::Initialize()
 {
 	uint32 NumIndexedCacheEntries = IIndexedCacheStorage::Get().GetNumIndexedCacheStorages();
+	if (NumIndexedCacheEntries == 0)
+	{
+		return;
+	}
+
 	FCacheStorageMetaDataArray IndexedCacheEntriesMetaData;
 	IndexedCacheEntriesMetaData.SetNum(NumIndexedCacheEntries);
 
@@ -336,6 +526,7 @@ void FIndexedCacheStorageManager::Initialize()
 		const TCHAR* PropertyName = TEXT("Name=");
 		const TCHAR* PropertyIndex = TEXT("Index=");
 		const TCHAR* PropertyMount = TEXT("Mount=");
+		const TCHAR* PropertyEarlyStartup = TEXT("EarlyStartup=");
 		FString StorageName;
 		int32 StorageIndex;
 		FString StorageMount;
@@ -352,7 +543,7 @@ void FIndexedCacheStorageManager::Initialize()
 
 			if (StorageIndex <= 0 || StorageIndex >= (int32)NumIndexedCacheEntries)
 			{
-				UE_LOG(LogIndexedCacheStorageManager, Error, TEXT("Found invalid Index in [IndexedCacheStorage]:Storage: %d, max allowed = %d"), StorageIndex, NumIndexedCacheEntries-1);
+				UE_LOG(LogIndexedCacheStorageManager, Error, TEXT("Found invalid Index in [IndexedCacheStorage]:Storage: %d, max allowed = %d"), StorageIndex, NumIndexedCacheEntries - 1);
 				continue;
 			}
 
@@ -373,9 +564,30 @@ void FIndexedCacheStorageManager::Initialize()
 			}
 			ExistingStorageIndex = StorageIndex;
 
+			uint64 EarlyStartupSize = 0;
+			{
+				FString EarlyStartupIniPath;
+				if (FParse::Value(*TrimmedCategory, PropertyEarlyStartup, EarlyStartupIniPath))
+				{
+					FString IniName;
+					FString IniSection;
+					FString IniKey;
+					if (GetIniParameters(IniName, IniSection, IniKey, EarlyStartupIniPath))
+					{
+						FString EarlyStartupSizeStr = GConfig->GetStr(*IniSection, *IniKey, *IniName);
+						EarlyStartupSize = ParseSizeParam(EarlyStartupSizeStr);
+						if (EarlyStartupSize == 0)
+						{
+							UE_LOG(LogIndexedCacheStorageManager, Error, 
+								   TEXT("%s%s specified, but no %s:%s exists in the %s ini file"), PropertyEarlyStartup, *EarlyStartupIniPath, *IniSection, *IniKey, *IniName);
+						}
+					}
+				}			
+			}
+
 			if (IndexedCacheEntriesMetaData[StorageIndex] == nullptr)
 			{
-				IndexedCacheEntriesMetaData[StorageIndex] = MakeUnique<FCacheStorageMetaData>(StorageIndex, StorageMount, CacheStorageAlreadyExists[StorageIndex]);
+				IndexedCacheEntriesMetaData[StorageIndex] = MakeUnique<FCacheStorageMetaData>(StorageIndex, StorageMount, CacheStorageAlreadyExists[StorageIndex], EarlyStartupSize);
 				++RegisteredStorages;
 			}
 			else
@@ -390,6 +602,7 @@ void FIndexedCacheStorageManager::Initialize()
 	if (bHasRegisteredEntries)
 	{
 		IndexedCacheStorageMetaData = MakeUnique<FIndexedCacheStorageMetaData>(MoveTemp(IndexedCacheEntriesMetaData));
+		EnsureEarlyStartupCache(CacheStorageAlreadyExists);
 	}
 }
 
@@ -404,6 +617,18 @@ int32 FIndexedCacheStorageManager::GetStorageIndex(const FString& StorageName)
 	return StorageIndex ? *StorageIndex : 0;
 }
 
+uint64 FIndexedCacheStorageManager::GetCacheEarlyStartupSize(int32 CacheIndex)
+{
+	FCacheStorageMetaData* CacheStorageMetaData = GetCacheStorageMetaData(CacheIndex);
+	if (CacheStorageMetaData == nullptr)
+	{
+		UE_LOG(LogIndexedCacheStorageManager, Error, TEXT("GetCacheEarlyStartupSize skip invalid cache %d"), CacheIndex);
+		return 0;
+	}
+
+	return CacheStorageMetaData->GetEarlyStartupSize();
+}
+
 void FIndexedCacheStorageManager::EnumerateCacheStorages(TArray<FString>& OutCacheStorageNames)
 {
 	OutCacheStorageNames.Reserve(StorageNameToStorageIndex.Num());
diff --git a/Engine/Source/Runtime/Experimental/IndexedCacheStorage/Public/IndexedCacheStorageManager.h b/Engine/Source/Runtime/Experimental/IndexedCacheStorage/Public/IndexedCacheStorageManager.h
--- a/Engine/Source/Runtime/Experimental/IndexedCacheStorage/Public/IndexedCacheStorageManager.h
+++ b/Engine/Source/Runtime/Experimental/IndexedCacheStorage/Public/IndexedCacheStorageManager.h
@@ -36,9 +36,15 @@ public:
 
 	// Returns the cache index registered to that storage name. Registration happens in the config section:
 	// [IndexedCacheStorage]
-	// (Name="StorageName",Index=StorageIndex,Mount="StorageMount")
+	// +Storage=(Name="StorageName",Index=StorageIndex,Mount="StorageMount",EarlyStartup=IniName:[IniSection]:IniKey)
+	// The EarlyStartup attribute is optional. It will try to fetch the value from the specified init path
+	// The value in the ini file can be written as 512, 512KB, 512MB or 512GB
+	// Requires that ones of the entries in the Storage array is called EarlyStartupCache
 	UE_API int32 GetStorageIndex(const FString& CacheStorageName);
 
+	// Returns 0 if no early startup size configured
+	UE_API uint64 GetCacheEarlyStartupSize(int32 CacheIndex);
+
 	// returns all registered storage names. The assoaciated cache storages might not have been created
 	UE_API void EnumerateCacheStorages(TArray<FString>& OutCacheStorageNames);
 
@@ -67,6 +73,7 @@ private:
 	friend class FIndexedCacheStorageModule;
 
 	void Initialize();
+	void EnsureEarlyStartupCache(const TArray<bool>& CacheStorageAlreadyExists);
 	FCacheStorageMetaData* GetCacheStorageMetaData(int32 CacheIndex);
 
 	TUniquePtr<struct FIndexedCacheStorageMetaData> IndexedCacheStorageMetaData;
