Fix for one in a million crash seen in LIVE. Iris parallel polling logic needs to be protected against possible null Context. Added ensure as this should never happen and we should investigate further if it reoccurs. Also protected against a possible race condition if adding a new Context happens whilst another is being released.

#tests Tested on Cooked Dedicated Server
#rnx

diff --git a/Engine/Source/Runtime/Experimental/Iris/Core/Private/Iris/Stats/NetStats.cpp b/Engine/Source/Runtime/Experimental/Iris/Core/Private/Iris/Stats/NetStats.cpp
--- a/Engine/Source/Runtime/Experimental/Iris/Core/Private/Iris/Stats/NetStats.cpp
+++ b/Engine/Source/Runtime/Experimental/Iris/Core/Private/Iris/Stats/NetStats.cpp
@@ -228,7 +228,9 @@ FNetStatsContext* FNetTypeStats::AcquireChildNetStatsContext()
 	{
 		FNetStatsContext* NewContext = CreateChildNetStatsContext();
 		check(NewContext);
+		ChildStatsContextCS.Lock();
 		NewContext->Acquire();
+		ChildStatsContextCS.Unlock();
 		return NewContext;
 	}
 
@@ -237,9 +239,12 @@ FNetStatsContext* FNetTypeStats::AcquireChildNetStatsContext()
 
 void FNetTypeStats::ReleaseChildNetStatsContext(FNetStatsContext* StatsContext)
 {
-	FScopeLock Lock(&ChildStatsContextCS);
-	check(!StatsContext->CanAcquire());
-	StatsContext->Release();
+	if(ensureMsgf(StatsContext, TEXT("ReleaseChildNetStatsContext was passed a null StatsContext. Possibility of net stats data being lost.")))
+	{
+		FScopeLock Lock(&ChildStatsContextCS);
+		check(!StatsContext->CanAcquire());
+		StatsContext->Release();
+	}
 }
 
 void FNetTypeStats::CleanupChildNetStatsContexts()
