Fix V781: The value of the variable is checked after it was used.

#rnx

diff --git a/Engine/Plugins/Media/ElectraUtil/Source/ElectraBase/Private/StringHelpers.cpp b/Engine/Plugins/Media/ElectraUtil/Source/ElectraBase/Private/StringHelpers.cpp
--- a/Engine/Plugins/Media/ElectraUtil/Source/ElectraBase/Private/StringHelpers.cpp
+++ b/Engine/Plugins/Media/ElectraUtil/Source/ElectraBase/Private/StringHelpers.cpp
@@ -15,7 +15,7 @@ namespace Electra
 			}
 			else if (InNumCharsToConvert < 0)
 			{
-				for(InNumCharsToConvert=0; InStringToConvert[InNumCharsToConvert] && InNumCharsToConvert < 16384; ++InNumCharsToConvert)
+				for(InNumCharsToConvert=0; InNumCharsToConvert < 16384 && InStringToConvert[InNumCharsToConvert]; ++InNumCharsToConvert)
 				{}
 			}
 			TArray<uint8> ConvBuf;
diff --git a/Engine/Source/Runtime/Experimental/IoStore/HttpClient/Private/Misc.inl b/Engine/Source/Runtime/Experimental/IoStore/HttpClient/Private/Misc.inl
--- a/Engine/Source/Runtime/Experimental/IoStore/HttpClient/Private/Misc.inl
+++ b/Engine/Source/Runtime/Experimental/IoStore/HttpClient/Private/Misc.inl
@@ -254,18 +254,18 @@ static int32 ParseMessage(FAnsiStringView Message, FMessageOffsets& Out)
 	int32 i = Protocol.Len();
 
 	// Trim left and tightly reject anything adventurous
-	for (int n = 32; Cursor[i] == ' ' && i < n; ++i);
+	for (int n = 32; i < n && Cursor[i] == ' '; ++i);
 	Out.StatusCode = uint8(i);
 
 	// At least one status line digit. (Note to self; expect exactly three)
-	for (int n = 32; uint32(Cursor[i] - 0x30) <= 9 && i < n; ++i);
+	for (int n = 32; i < n && uint32(Cursor[i] - 0x30) <= 9; ++i);
 	if (uint32(i - Out.StatusCode - 1) > 32)
 	{
 		return -1;
 	}
 
 	// Trim left
-	for (int n = 32; Cursor[i] == ' ' && i < n; ++i);
+	for (int n = 32; i < n && Cursor[i] == ' '; ++i);
 	Out.Message = uint8(i);
 
 	// Extra conservative length allowance
diff --git a/Engine/Source/Runtime/Renderer/Private/SceneRendering.cpp b/Engine/Source/Runtime/Renderer/Private/SceneRendering.cpp
--- a/Engine/Source/Runtime/Renderer/Private/SceneRendering.cpp
+++ b/Engine/Source/Runtime/Renderer/Private/SceneRendering.cpp
@@ -4740,7 +4740,7 @@ void FSceneRenderer::SetupMeshPass(FViewInfo& View, FExclusiveDepthStencil::Type
 			TArray<int32, TInlineAllocator<2> > ViewIds;
 			ViewIds.Add(View.SceneRendererPrimaryViewId);
 			// Only apply instancing for ISR to main view passes
-			const bool bIsMainViewPass = PassType != EMeshPass::Num && (FPassProcessorManager::GetPassFlags(ShadingPath, PassType) & EMeshPassFlags::MainView) != EMeshPassFlags::None;
+			const bool bIsMainViewPass = (FPassProcessorManager::GetPassFlags(ShadingPath, PassType) & EMeshPassFlags::MainView) != EMeshPassFlags::None;
 
 			EInstanceCullingMode InstanceCullingMode = bIsMainViewPass && View.IsInstancedStereoPass() ? EInstanceCullingMode::Stereo : EInstanceCullingMode::Normal;
 			if (InstanceCullingMode == EInstanceCullingMode::Stereo)
