Fix against potential race condition where a GC can be removed while iterating on the sections array.

diff --git a/Engine/Source/Runtime/Experimental/GeometryCollectionEngine/Private/GeometryCollection/GeometryCollectionComponent.cpp b/Engine/Source/Runtime/Experimental/GeometryCollectionEngine/Private/GeometryCollection/GeometryCollectionComponent.cpp
--- a/Engine/Source/Runtime/Experimental/GeometryCollectionEngine/Private/GeometryCollection/GeometryCollectionComponent.cpp
+++ b/Engine/Source/Runtime/Experimental/GeometryCollectionEngine/Private/GeometryCollection/GeometryCollectionComponent.cpp
@@ -7238,7 +7238,8 @@ void UGeometryCollectionComponent::CollectPSOPrecacheData(const FPSOPrecachePara
 		VFDataList.Add(FPSOPrecacheVertexFactoryData(&FGeometryCollectionVertexFactory::StaticType));
 	}
 
-	const TManagedArray<FGeometryCollectionSection>& SectionsArray = RestCollection->GetGeometryCollection()->Sections;
+	TSharedPtr<FGeometryCollection, ESPMode::ThreadSafe> RestGeomCollection = RestCollection->GetGeometryCollection();
+	const TManagedArray<FGeometryCollectionSection>& SectionsArray = RestGeomCollection->Sections;
 	for (int32 SectionIndex = 0; SectionIndex < SectionsArray.Num(); ++SectionIndex)
 	{
 		const FGeometryCollectionSection& MeshSection = SectionsArray[SectionIndex];
