Add immediate analytics for IoStoreOnDemand.  These send individual events when needed.
- Add an event for IadCacheDefrag that returns error codes
- Add an event for IadCacheBlockDeleted that returns block age

Fix FCas::DeleteBlock not removing block from LastAccess  map.

Fix FOnDemandInstallCache::Defrag not returning an error if we fail to delete a block.

diff --git a/Engine/Source/Runtime/Experimental/IoStore/OnDemand/Private/OnDemandInstallCache.cpp b/Engine/Source/Runtime/Experimental/IoStore/OnDemand/Private/OnDemandInstallCache.cpp
--- a/Engine/Source/Runtime/Experimental/IoStore/OnDemand/Private/OnDemandInstallCache.cpp
+++ b/Engine/Source/Runtime/Experimental/IoStore/OnDemand/Private/OnDemandInstallCache.cpp
@@ -335,6 +335,7 @@ FIoStatus FCas::DeleteBlock(FCasBlockId BlockId, TArray<FCasAddr>& OutAddrs)
 
 	BlockIds.Remove(BlockId);
 	ReadHandles.Remove(BlockId);
+	LastAccess.Remove(BlockId);
 	for (auto It = Lookup.CreateIterator(); It; ++It)
 	{
 		if (It->Value.BlockId == BlockId)
@@ -2624,6 +2625,8 @@ FIoStatus FOnDemandInstallCache::Purge(FCasBlockInfoMap& BlockInfo, const uint64
 			return Status;
 		}
 
+		FOnDemandInstallCacheStats::OnBlockDeleted(Info.LastAccess, false /*FromDefrag*/);
+
 		// This should be the only thread writing to CurrentBlock
 		FCasBlockId MaybeCurrentBlock = BlockId;
 		CurrentBlock.compare_exchange_strong(MaybeCurrentBlock, FCasBlockId::Invalid);
@@ -2920,7 +2923,13 @@ FIoStatus FOnDemandInstallCache::Defrag(
 		// Flushing should overwrite the lookup info for the cas addr to point at the new block.
 		// Can now remove the old block
 		TArray<FCasAddr> DeletedChunkAddresses;
-		Cas.DeleteBlock(DefragBlock.BlockId, DeletedChunkAddresses);
+		if (FIoStatus Status = Cas.DeleteBlock(DefragBlock.BlockId, DeletedChunkAddresses); !Status.IsOk())
+		{
+			FOnDemandInstallCacheStats::OnDefrag(Status.GetErrorCode(), 0);
+			return Status;
+		}
+
+		FOnDemandInstallCacheStats::OnBlockDeleted(DefragBlock.LastAccess, true /*FromDefrag*/);
 
 		for (const FCasAddr& Addr : DeletedChunkAddresses)
 		{
diff --git a/Engine/Source/Runtime/Experimental/IoStore/OnDemand/Private/OnDemandIoStore.cpp b/Engine/Source/Runtime/Experimental/IoStore/OnDemand/Private/OnDemandIoStore.cpp
--- a/Engine/Source/Runtime/Experimental/IoStore/OnDemand/Private/OnDemandIoStore.cpp
+++ b/Engine/Source/Runtime/Experimental/IoStore/OnDemand/Private/OnDemandIoStore.cpp
@@ -44,6 +44,7 @@
 #include "OnDemandPackageStoreBackend.h"
 #include "ProfilingDebugging/AssetMetadataTrace.h"
 #include "Serialization/MemoryReader.h"
+#include "Statistics.h"
 
 #if !(UE_BUILD_SHIPPING|UE_BUILD_TEST)
 #include "String/LexFromString.h"
@@ -2100,6 +2101,11 @@ TUniquePtr<IAnalyticsRecording> FOnDemandIoStore::StartAnalyticsRecording() cons
 	return TUniquePtr<IAnalyticsRecording>();
 }
 
+void FOnDemandIoStore::OnImmediateAnalytic(FOnDemandImmediateAnalyticHandler EventHandler)
+{
+	OnDemandSetImmediateAnalyticHandler(MoveTemp(EventHandler));
+}
+
 } // namespace UE::IoStore
 
 ////////////////////////////////////////////////////////////////////////////////
diff --git a/Engine/Source/Runtime/Experimental/IoStore/OnDemand/Private/OnDemandIoStore.h b/Engine/Source/Runtime/Experimental/IoStore/OnDemand/Private/OnDemandIoStore.h
--- a/Engine/Source/Runtime/Experimental/IoStore/OnDemand/Private/OnDemandIoStore.h
+++ b/Engine/Source/Runtime/Experimental/IoStore/OnDemand/Private/OnDemandIoStore.h
@@ -360,6 +360,7 @@ public:
 	virtual void						SetStreamingOptions(EOnDemandStreamingOptions Options) override;
 	virtual void						ReportAnalytics(TArray<FAnalyticsEventAttribute>& OutAnalyticsArray) const override;
 	virtual TUniquePtr<IAnalyticsRecording> StartAnalyticsRecording() const override;
+	virtual void						OnImmediateAnalytic(FOnDemandImmediateAnalyticHandler EventHandler) override;
 	virtual void						CancelInstallRequest(FSharedInternalInstallRequest InstallRequest) override;
 	virtual void						UpdateInstallRequestPriority(FSharedInternalInstallRequest InstallRequest, int32 NewPriority) override;
 	virtual void						ReleaseContent(FOnDemandInternalContentHandle& ContentHandle) override;
diff --git a/Engine/Source/Runtime/Experimental/IoStore/OnDemand/Private/Statistics.cpp b/Engine/Source/Runtime/Experimental/IoStore/OnDemand/Private/Statistics.cpp
--- a/Engine/Source/Runtime/Experimental/IoStore/OnDemand/Private/Statistics.cpp
+++ b/Engine/Source/Runtime/Experimental/IoStore/OnDemand/Private/Statistics.cpp
@@ -51,6 +51,12 @@ static FAutoConsoleVariableRef CVar_ReportIadAnalyticsEnabled(
 	GIadReportAnalyticsEnabled,
 	TEXT("Enables reporting analytics for individual asset downloads."));
 
+float GIaxImmediateAnalyticsChance = 1.0f;
+static FAutoConsoleVariableRef CVar_ImmediateAnalyticsChance(
+	TEXT("iax.ImmediateAnalyticsChance"),
+	GIaxImmediateAnalyticsChance,
+	TEXT("Chance of sending and immediate analytics event."));
+
 #if UE_ENABLE_ONSCREEN_STATISTICS
 
 bool GIasDisplayOnScreenStatistics = false;
@@ -63,6 +69,24 @@ static FAutoConsoleVariableRef CVar_DisplayOnScreenStatistics(
 
 extern int32 GIasHttpRateLimitKiBPerSecond;
 
+////////////////////////////////////////////////////////////////////////////////
+static FOnDemandImmediateAnalyticHandler OnDemandImmediateAnalyticEventHandler;
+void OnDemandSetImmediateAnalyticHandler(FOnDemandImmediateAnalyticHandler&& EventHandler)
+{
+	OnDemandImmediateAnalyticEventHandler = MoveTemp(EventHandler);
+}
+
+static void SendImmediateAnalytic(FOnDemandImmediateAnalytic&& Event)
+{
+	if (GIaxImmediateAnalyticsChance > 0.0f && FMath::FRand() <= GIaxImmediateAnalyticsChance)
+	{
+		if (OnDemandImmediateAnalyticEventHandler.EventHandler)
+		{
+			OnDemandImmediateAnalyticEventHandler.EventHandler(MoveTemp(Event));
+		}
+	}
+}
+
 ////////////////////////////////////////////////////////////////////////////////
 static int32 BytesToApproxMB(uint64 Bytes) { return int32(Bytes >> 20); }
 static int32 BytesToApproxKB(uint64 Bytes) { return int32(Bytes >> 10); }
@@ -1414,6 +1438,37 @@ void FOnDemandInstallCacheStats::OnDefrag(EIoErrorCode ErrorCode, uint64 Fragmen
 	}
 	InstallCacheAnalytics.DefragCount++;
 	InstallCacheAnalytics.FragmentedBytes += FragmentedBytes;
+
+	SendImmediateAnalytic(FOnDemandImmediateAnalytic
+	{
+		TEXT("IadCacheDefrag"),
+		MakeAnalyticsEventAttributeArray(
+			TEXT("Success"), ErrorCode == EIoErrorCode::Ok,
+			TEXT("ErrorCode"), GetIoErrorText(ErrorCode),
+			TEXT("FragmentedBytes"), FragmentedBytes
+		)
+	});
+}
+
+void FOnDemandInstallCacheStats::OnBlockDeleted(int64 LastAccessTicks, bool bFromDefrag)
+{
+	if (!OnDemandImmediateAnalyticEventHandler.EventHandler)
+	{
+		return;
+	}
+
+	const FDateTime Now = FDateTime::UtcNow();
+	const FDateTime BlockAccess(LastAccessTicks);
+	const FTimespan BlockAge = ensure(Now >= BlockAccess) ? (Now - BlockAccess) : FTimespan::MaxValue();
+
+	SendImmediateAnalytic(FOnDemandImmediateAnalytic
+	{
+		TEXT("IadCacheBlockDeleted"),
+		MakeAnalyticsEventAttributeArray(
+			TEXT("BlockAgeMinutes"), BlockAge.GetTotalMinutes(),
+			TEXT("FromDefrag"), bFromDefrag
+		)
+	});
 }
 
 void FOnDemandInstallCacheStats::OnCacheUsage(
diff --git a/Engine/Source/Runtime/Experimental/IoStore/OnDemand/Private/Statistics.h b/Engine/Source/Runtime/Experimental/IoStore/OnDemand/Private/Statistics.h
--- a/Engine/Source/Runtime/Experimental/IoStore/OnDemand/Private/Statistics.h
+++ b/Engine/Source/Runtime/Experimental/IoStore/OnDemand/Private/Statistics.h
@@ -25,6 +25,7 @@ namespace UE::IoStore
 class IAnalyticsRecording;
 enum class EHttpRequestType : uint8;
 struct FAvailableEndPoints;
+struct FOnDemandImmediateAnalyticHandler;
 
 #if IAS_WITH_STATISTICS
 #	define IAS_STATISTICS_IMPL(...) ;
@@ -32,6 +33,9 @@ struct FAvailableEndPoints;
 #	define IAS_STATISTICS_IMPL(...) { return __VA_ARGS__; }
 #endif
 
+////////////////////////////////////////////////////////////////////////////////
+extern void OnDemandSetImmediateAnalyticHandler(FOnDemandImmediateAnalyticHandler&& EventHandler);
+
 ////////////////////////////////////////////////////////////////////////////////
 class FOnDemandIoBackendStats
 {
@@ -135,6 +139,7 @@ public:
 		uint64 BytesToPurge,
 		uint64 PurgedBytes) IAS_STATISTICS_IMPL()
 	static void OnDefrag(EIoErrorCode ErrorCode, uint64 FragmentedBytes) IAS_STATISTICS_IMPL()
+	static void OnBlockDeleted(int64 LastAccessTicks, bool bFromDefrag) IAS_STATISTICS_IMPL()
 	static void OnCacheUsage(
 		uint64 MaxCacheSize,
 		uint64 CacheSize,
diff --git a/Engine/Source/Runtime/Experimental/IoStore/OnDemandCore/Public/IO/IoStoreOnDemand.h b/Engine/Source/Runtime/Experimental/IoStore/OnDemandCore/Public/IO/IoStoreOnDemand.h
--- a/Engine/Source/Runtime/Experimental/IoStore/OnDemandCore/Public/IO/IoStoreOnDemand.h
+++ b/Engine/Source/Runtime/Experimental/IoStore/OnDemandCore/Public/IO/IoStoreOnDemand.h
@@ -495,6 +495,17 @@ public:
 	virtual void StopRecording() = 0;
 };
 
+struct FOnDemandImmediateAnalytic
+{
+	FString EventName;
+	TArray<FAnalyticsEventAttribute> AnalyticsArray;
+};
+
+struct FOnDemandImmediateAnalyticHandler
+{
+	TUniqueFunction<void(FOnDemandImmediateAnalytic)> EventHandler;
+};
+
 /** Streaming options (IAS). */
 enum class EOnDemandStreamingOptions
 {
@@ -579,6 +590,11 @@ public:
 	 * This can return a nullptr if the analytics system is not enabled.
 	 */
 	virtual TUniquePtr<IAnalyticsRecording> StartAnalyticsRecording() const = 0;
+	/**
+	 * Report immediate analytics events that are not recorded over a timespan.
+	 * Useful for analytics that are better reported as single events.
+	 */
+	virtual void OnImmediateAnalytic(FOnDemandImmediateAnalyticHandler EventHandler) = 0;
 
 private:
 	/** Initialize the I/O store. Called after the I/O store factory module feature has been registered. */
