Fix IASTool

diff --git a/Engine/Source/Runtime/Experimental/IoStore/OnDemand/Private/Statistics.h b/Engine/Source/Runtime/Experimental/IoStore/OnDemand/Private/Statistics.h
--- a/Engine/Source/Runtime/Experimental/IoStore/OnDemand/Private/Statistics.h
+++ b/Engine/Source/Runtime/Experimental/IoStore/OnDemand/Private/Statistics.h
@@ -34,7 +34,11 @@ struct FOnDemandImmediateAnalyticHandler;
 #endif
 
 ////////////////////////////////////////////////////////////////////////////////
+#if IAS_WITH_STATISTICS
 extern void OnDemandSetImmediateAnalyticHandler(FOnDemandImmediateAnalyticHandler&& EventHandler);
+#else
+inline void OnDemandSetImmediateAnalyticHandler(FOnDemandImmediateAnalyticHandler&& EventHandler) {}
+#endif
 
 ////////////////////////////////////////////////////////////////////////////////
 class FOnDemandIoBackendStats
