[Chaos] Add CVar controlling maximum number of steps allowed to take during a single gametick.

diff --git a/Engine/Source/Runtime/Experimental/Chaos/Private/Chaos/Framework/PhysicsSolverBase.cpp b/Engine/Source/Runtime/Experimental/Chaos/Private/Chaos/Framework/PhysicsSolverBase.cpp
--- a/Engine/Source/Runtime/Experimental/Chaos/Private/Chaos/Framework/PhysicsSolverBase.cpp
+++ b/Engine/Source/Runtime/Experimental/Chaos/Private/Chaos/Framework/PhysicsSolverBase.cpp
@@ -426,6 +426,9 @@ namespace Chaos
 		}
 	}
 
+	int32 MaxPhysicsStepsPerGameTick = 3;
+	FAutoConsoleVariableRef CVarMaxPhysicsStepsPerGameTick(TEXT("p.MaxPhysicsStepsPerGameTick"), MaxPhysicsStepsPerGameTick, TEXT("The maximum number of physics steps per gametick frame."));
+
 	FGraphEventRef FPhysicsSolverBase::AdvanceAndDispatch_External(FReal InDt)
 	{
 		LLM_SCOPE(ELLMTag::ChaosScene);
@@ -447,7 +450,6 @@ namespace Chaos
 			}
 			else
 			{
-
 				InternalDt = AsyncDt;
 				NumSteps = FMath::FloorToInt32(AccumulatedTime / InternalDt);
 				AccumulatedTime -= InternalDt * static_cast<FReal>(NumSteps);
@@ -478,7 +480,7 @@ namespace Chaos
 		if (AsyncBlockMode == EAsyncBlockMode::DoNoBlock)
 		{
 			// Make sure not to accumulate too many physics solver tasks.
-			constexpr int32 MaxPhysicsStepToKeep = 3;
+			const int32 MaxPhysicsStepToKeep = MaxPhysicsStepsPerGameTick;
 			const int32 MaxNumSteps = MaxPhysicsStepToKeep - NumPendingSolverAdvanceTasks;
 			if (NumSteps > MaxNumSteps)
 			{
