Add some missing basic functionality in Primitive Component when objects have physics objects and not body instances.

diff --git a/Engine/Source/Runtime/Engine/Private/PrimitiveComponentPhysics.cpp b/Engine/Source/Runtime/Engine/Private/PrimitiveComponentPhysics.cpp
--- a/Engine/Source/Runtime/Engine/Private/PrimitiveComponentPhysics.cpp
+++ b/Engine/Source/Runtime/Engine/Private/PrimitiveComponentPhysics.cpp
@@ -387,6 +387,23 @@ void UPrimitiveComponent::SetPhysicsLinearVelocity(FVector NewVel, bool bAddToCu
 		WarnInvalidPhysicsOperations(LOCTEXT("SetPhysicsLinearVelocity", "SetPhysicsLinearVelocity"), nullptr, BoneName);
 		BI->SetLinearVelocity(NewVel, bAddToCurrent);
 	}
+	else if (BoneName == NAME_None)
+	{
+		TArray<Chaos::FPhysicsObjectHandle> PhysicsObjects = GetAllPhysicsObjects();
+		FLockedWritePhysicsObjectExternalInterface Interface = FPhysicsObjectExternalInterface::LockWrite(PhysicsObjects);
+
+		PhysicsObjects = PhysicsObjects.FilterByPredicate(
+			[&Interface](Chaos::FPhysicsObject* Object) {
+				return !Interface->AreAllDisabled({ &Object, 1 });
+			}
+		);
+		Interface->SetLinearImpulseVelocity(PhysicsObjects, NewVel, bAddToCurrent);
+	}
+	else if (Chaos::FPhysicsObject* Object = GetPhysicsObjectByName(BoneName))
+	{
+		FLockedWritePhysicsObjectExternalInterface Interface = FPhysicsObjectExternalInterface::LockWrite({ &Object, 1 });
+		Interface->SetLinearImpulseVelocity({ &Object, 1 }, NewVel, bAddToCurrent);
+	}
 }
 
 FVector UPrimitiveComponent::GetPhysicsLinearVelocity(FName BoneName)
@@ -429,6 +446,23 @@ void UPrimitiveComponent::SetPhysicsAngularVelocityInRadians(FVector NewAngVel,
 		WarnInvalidPhysicsOperations(LOCTEXT("SetPhysicsAngularVelocity", "SetPhysicsAngularVelocity"), nullptr, BoneName);
 		BI->SetAngularVelocityInRadians(NewAngVel, bAddToCurrent);
 	}
+	else if (BoneName == NAME_None)
+	{
+		TArray<Chaos::FPhysicsObjectHandle> PhysicsObjects = GetAllPhysicsObjects();
+		FLockedWritePhysicsObjectExternalInterface Interface = FPhysicsObjectExternalInterface::LockWrite(PhysicsObjects);
+
+		PhysicsObjects = PhysicsObjects.FilterByPredicate(
+			[&Interface](Chaos::FPhysicsObject* Object) {
+				return !Interface->AreAllDisabled({ &Object, 1 });
+			}
+		);
+		Interface->SetAngularVelocityInRadians(PhysicsObjects, NewAngVel, bAddToCurrent);
+	}
+	else if (Chaos::FPhysicsObject* Object = GetPhysicsObjectByName(BoneName))
+	{
+		FLockedWritePhysicsObjectExternalInterface Interface = FPhysicsObjectExternalInterface::LockWrite({ &Object, 1 });
+		Interface->SetAngularVelocityInRadians({ &Object, 1 }, NewAngVel, bAddToCurrent);
+	}
 }
 
 void UPrimitiveComponent::SetPhysicsMaxAngularVelocityInRadians(float NewMaxAngVel, bool bAddToCurrent, FName BoneName)
diff --git a/Engine/Source/Runtime/Experimental/Chaos/Private/Chaos/PhysicsObjectInterface.cpp b/Engine/Source/Runtime/Experimental/Chaos/Private/Chaos/PhysicsObjectInterface.cpp
--- a/Engine/Source/Runtime/Experimental/Chaos/Private/Chaos/PhysicsObjectInterface.cpp
+++ b/Engine/Source/Runtime/Experimental/Chaos/Private/Chaos/PhysicsObjectInterface.cpp
@@ -1028,6 +1028,96 @@ namespace Chaos
 		}
 	}
 
+	template<EThreadContext Id>
+	void FWritePhysicsObjectInterface<Id>::SetLinearVelocity(TArrayView<const FPhysicsObjectHandle> InObjects, const FVector& InVelocity, bool bAddToCurrent)
+	{
+		if (InObjects.IsEmpty())
+		{
+			return;
+		}
+
+		for (const FPhysicsObjectHandle Object : InObjects)
+		{
+			if (!Object)
+			{
+				continue;
+			}
+
+			if (TThreadParticle<Id>* Particle = Object->GetParticle<Id>())
+			{
+				if (Chaos::TThreadRigidParticle<Id>* Rigid = Particle->CastToRigidParticle())
+				{
+					FVector Final = InVelocity;
+
+					if (bAddToCurrent)
+					{
+						Final += Rigid->LinearImpulseVelocity();
+					}
+
+					Rigid->SetLinearImpulseVelocity(Final, false);
+				}
+			}
+		}
+
+		if constexpr (Id == EThreadContext::External)
+		{
+			if (Chaos::FPhysicsSolverBase* Solver = GetSolverForObjects(InObjects))
+			{
+				Solver->EnqueueCommandImmediate(
+					[AllObjects = TArray<FPhysicsObjectHandle>{ InObjects }, InVelocity, bAddToCurrent]() {
+						Chaos::FWritePhysicsObjectInterface_Internal Interface = Chaos::FPhysicsObjectInternalInterface::GetWrite();
+						Interface.SetLinearVelocity(AllObjects, InVelocity, bAddToCurrent);
+					}
+				);
+			}
+		}
+	}
+
+	template<EThreadContext Id>
+	void FWritePhysicsObjectInterface<Id>::SetAngularVelocityInRadians(TArrayView<const FPhysicsObjectHandle> InObjects, const FVector& InAngularVelocity, bool bAddToCurrent)
+	{
+		if (InObjects.IsEmpty())
+		{
+			return;
+		}
+
+		for (const FPhysicsObjectHandle Object : InObjects)
+		{
+			if (!Object)
+			{
+				continue;
+			}
+
+			if (TThreadParticle<Id>* Particle = Object->GetParticle<Id>())
+			{
+				if (Chaos::TThreadRigidParticle<Id>* Rigid = Particle->CastToRigidParticle())
+				{
+					FVector Final = InAngularVelocity;
+
+					if (bAddToCurrent)
+					{
+						Final += Rigid->AngularImpulseVelocity();
+					}
+
+					Rigid->SetAngularImpulseVelocity(Final, false);
+				}
+			}
+		}
+
+		if constexpr (Id == EThreadContext::External)
+		{
+			if (Chaos::FPhysicsSolverBase* Solver = GetSolverForObjects(InObjects))
+			{
+				Solver->EnqueueCommandImmediate(
+					[AllObjects = TArray<FPhysicsObjectHandle>{ InObjects }, InAngularVelocity, bAddToCurrent]() {
+						Chaos::FWritePhysicsObjectInterface_Internal Interface = Chaos::FPhysicsObjectInternalInterface::GetWrite();
+						Interface.SetAngularVelocityInRadians(AllObjects, InAngularVelocity, bAddToCurrent);
+					}
+				);
+			}
+		}
+	}
+
 	template<EThreadContext Id>
 	void FWritePhysicsObjectInterface<Id>::SetLinearImpulseVelocity(TArrayView<const FPhysicsObjectHandle> InObjects, const FVector& Impulse, bool bVelChange)
 	{
diff --git a/Engine/Source/Runtime/Experimental/Chaos/Public/Chaos/PhysicsObjectInterface.h b/Engine/Source/Runtime/Experimental/Chaos/Public/Chaos/PhysicsObjectInterface.h
--- a/Engine/Source/Runtime/Experimental/Chaos/Public/Chaos/PhysicsObjectInterface.h
+++ b/Engine/Source/Runtime/Experimental/Chaos/Public/Chaos/PhysicsObjectInterface.h
@@ -122,6 +122,8 @@ namespace Chaos
 		CHAOS_API void AddTorque(TArrayView<const FPhysicsObjectHandle> InObjects, const FVector& Torque, bool bInvalidate);
 		CHAOS_API void AddTorque(TArrayView<const FPhysicsObjectHandle> InObjects, const FVector& Torque, bool bInvalidate, bool bAccelChange = false);
 		CHAOS_API void SetLinearImpulseVelocity(TArrayView<const FPhysicsObjectHandle> InObjects, const FVector& Impulse, bool bVelChange);
+		CHAOS_API void SetLinearVelocity(TArrayView<const FPhysicsObjectHandle> InObjects, const FVector& InVelocity, bool bAddToCurrent = false);
+		CHAOS_API void SetAngularVelocityInRadians(TArrayView<const FPhysicsObjectHandle> InObjects, const FVector& InAngularVelocity, bool bAddToCurrent = false);
 
 		UE_DEPRECATED(5.4, "This version of AddRadialImpulse has been deprecated. Please use the version where the strain value is passed explicitly")
 		CHAOS_API void AddRadialImpulse(TArrayView<const FPhysicsObjectHandle> InObjects, FVector Origin, float Radius, float Strength, enum ERadialImpulseFalloff Falloff, bool bApplyStrain, bool bInvalidate, bool bVelChange = false);
