Change IO store such that its disk quota may come from Indexed cache storage if specified

diff --git a/Engine/Source/Runtime/Experimental/IoStore/OnDemand/Private/OnDemandIoStore.cpp b/Engine/Source/Runtime/Experimental/IoStore/OnDemand/Private/OnDemandIoStore.cpp
--- a/Engine/Source/Runtime/Experimental/IoStore/OnDemand/Private/OnDemandIoStore.cpp
+++ b/Engine/Source/Runtime/Experimental/IoStore/OnDemand/Private/OnDemandIoStore.cpp
@@ -942,7 +942,9 @@ FIoStatus FOnDemandIoStore::InitializeInstallCache()
 	int32 SelectedCacheIndex = 0;
 	{
 		// OnDemandInstallCacheConfig.DiskQuota is the quota to store the data, but we also need some room for metadata, see GetJournalFilename/GetSnapshotFilename
-		uint64 OnDemandInstallRequestedSize = OnDemandInstallCacheConfig.DiskQuota + 2 * OnDemandInstallCacheConfig.JournalMaxSize + 2 * 32;
+		uint64 MetaDataSize = 2 * OnDemandInstallCacheConfig.JournalMaxSize + 2 * 32;
+		uint64 OnDemandInstallRequestedSize = OnDemandInstallCacheConfig.DiskQuota + MetaDataSize;
+
 		if (OnDemandInstallRequestedSize > 0)
 		{
 			FModuleManager::Get().LoadModule(TEXT("IndexedCacheStorage"));
@@ -952,6 +954,18 @@ FIoStatus FOnDemandIoStore::InitializeInstallCache()
 
 			if (SelectedCacheIndex > 0)
 			{
+				uint64 EarlyStartupSize = Experimental::FIndexedCacheStorageManager::Get().GetCacheEarlyStartupSize(SelectedCacheIndex);
+				if (EarlyStartupSize > 0)
+				{
+					OnDemandInstallRequestedSize = EarlyStartupSize;
+
+					uint64 MinimumDiskQuota = 4 << 20;
+					check(EarlyStartupSize > MetaDataSize + MinimumDiskQuota);
+
+					// Update DiskQuota so that EarlyStartupSize can contain the whole request;
+					OnDemandInstallCacheConfig.DiskQuota = EarlyStartupSize - MetaDataSize;
+				}
+
 				// When IAD is assigned to an indexed cache, we *have* to get enough disk space.
 				while (true)
 				{
