Add ics.CopyFiles <DestDir> / ics.ListFiles commands to allow file visualization / retrieval of cache storages

diff --git a/Engine/Source/Runtime/Experimental/IndexedCacheStorage/Private/IndexedCacheStorageManager.cpp b/Engine/Source/Runtime/Experimental/IndexedCacheStorage/Private/IndexedCacheStorageManager.cpp
--- a/Engine/Source/Runtime/Experimental/IndexedCacheStorage/Private/IndexedCacheStorageManager.cpp
+++ b/Engine/Source/Runtime/Experimental/IndexedCacheStorage/Private/IndexedCacheStorageManager.cpp
@@ -3,9 +3,12 @@
 #include "IndexedCacheStorageManager.h"
 #include "IndexedCacheStorage.h"
 
+#include "Async/Async.h"
 #include "Async/Mutex.h"
 #include "Async/UniqueLock.h"
 #include "Experimental/Async/ConditionVariable.h"
+#include "HAL/FileManager.h"
+#include "HAL/IConsoleManager.h"
 #include "HAL/PlatformMisc.h"
 #include "Misc/ConfigCacheIni.h"
 #include "HAL/PlatformStackWalk.h"
@@ -740,4 +743,180 @@ FCacheStorageMetaData* FIndexedCacheStorageManager::GetCacheStorageMetaData(int3
 	return nullptr;
 }
 
+#if !UE_BUILD_SHIPPING
+namespace IndexedCacheStorageManager_Debug
+{
+static void CopyFiles(const FString& TargetDirectory);
+static void ListFiles();
+
+static FAutoConsoleCommand GCopyFiles(
+	TEXT("ics.CopyFiles"),
+	TEXT("Copies all cache files to the supplied folder"),
+	FConsoleCommandWithArgsDelegate::CreateLambda( [](const TArray<FString>& Args)
+	{
+		if (Args.IsEmpty())
+		{
+			UE_LOG(LogIndexedCacheStorageManager, Error, TEXT("usage: ics.CopyFiles <DestDir>"));
+			return;
+		}
+
+		CopyFiles(Args[0]);
+	}));
+
+static FAutoConsoleCommand GListFiles(
+	TEXT("ics.ListFiles"),
+	TEXT("Copies all cache files to the supplied folder"),
+	FConsoleCommandWithArgsDelegate::CreateLambda( [](const TArray<FString>&)
+	{
+		ListFiles();
+	}));
+
+static void CollectCacheIndices(TArray<int32>& CacheIndices)
+{
+	uint32 NumIndexedCacheEntries = IIndexedCacheStorage::Get().GetNumIndexedCacheStorages();
+	CacheIndices.Reserve(NumIndexedCacheEntries+1);
+
+	if (FPaths::HasProjectPersistentDownloadDir())
+	{
+		CacheIndices.Add(0);
+	}
+
+	for (uint32 CacheIndex = 1; CacheIndex < NumIndexedCacheEntries; ++CacheIndex)
+	{
+		uint64 CacheCapacity = FIndexedCacheStorageManager::Get().GetCacheStorageCapacity(CacheIndex);
+		if (CacheCapacity == 0)
+		{
+			continue;
+		}
+
+		CacheIndices.Add(CacheIndex);
+	}
+
+}
+
+static void ListFiles()
+{
+	TArray<int32> CacheIndices;
+	CollectCacheIndices(CacheIndices);
+
+	UE_LOG(LogConsoleResponse, Display, TEXT("Listing cache files ..."));
+
+	AsyncTask(
+		ENamedThreads::AnyBackgroundThreadNormalTask, [CacheIndices]() mutable
+		{
+			for (int32 CacheIndex : CacheIndices)
+			{
+				FIndexedCacheStorageScopeMount IndexedCacheStorageScopeMount(CacheIndex);
+
+				FString SourcePath = (CacheIndex == 0) ? FPaths::ProjectPersistentDownloadDir() : IndexedCacheStorageScopeMount.MountPath;
+
+				if (SourcePath.IsEmpty())
+				{
+					return;
+				}
+
+				IFileManager::Get().IterateDirectoryStatRecursively(
+					*SourcePath,
+					[](const TCHAR* FileOrDir, const FFileStatData& StatData)
+					{
+						if (StatData.bIsValid)
+						{
+							int64 FileSize = !StatData.bIsDirectory ? StatData.FileSize : 0;
+							UE_LOG(LogConsoleResponse, Display, TEXT("%10llu %s%c"), FileSize, FileOrDir, (StatData.bIsDirectory ? TEXT('/') : TEXT('\0')));
+						}
+						return true;
+					}
+				);
+			}
+
+			UE_LOG(LogConsoleResponse, Display, TEXT("Listing cache files done"));
+
+		}
+	);
+}
+
+static void CopyFiles(const FString& InTargetDirectory)
+{
+	if (!IFileManager::Get().DirectoryExists(*InTargetDirectory))
+	{
+		IFileManager::Get().MakeDirectory(*InTargetDirectory);
+	}
+
+	TArray<int32> CacheIndices;
+	CollectCacheIndices(CacheIndices);
+
+	FString TargetDirectory(InTargetDirectory);
+
+	UE_LOG(LogConsoleResponse, Display, TEXT("Copy files to %s ..."), *TargetDirectory);
+
+	AsyncTask(
+		ENamedThreads::AnyBackgroundThreadNormalTask, [CacheIndices, TargetDirectory]() mutable
+		{
+			for (int32 CacheIndex : CacheIndices)
+			{
+				FIndexedCacheStorageScopeMount IndexedCacheStorageScopeMount(CacheIndex);
+
+				FString SourcePath = (CacheIndex == 0) ? FPaths::ProjectPersistentDownloadDir() : IndexedCacheStorageScopeMount.MountPath;
+
+				if (SourcePath.IsEmpty())
+				{
+					return;
+				}
+
+				if (!SourcePath.EndsWith(TEXT("/")))
+				{
+					SourcePath += TEXT("/");
+				}
+
+				TArray<FString> SourceFiles;
+				TArray<FString> DestFiles;
+
+				FString TargetDirectoryCache = FString::Printf(TEXT("%s/cache_%02d"), *TargetDirectory, CacheIndex);
+
+				IFileManager::Get().IterateDirectoryStatRecursively(
+					*SourcePath,
+					[SourcePath, TargetDirectoryCache, &SourceFiles, &DestFiles](const TCHAR* FileOrDir, const FFileStatData& StatData)
+					{
+						if (StatData.bIsValid)
+						{
+							FString RelativeSourceFile(FileOrDir);
+
+							if (StatData.bIsDirectory && !RelativeSourceFile.EndsWith(TEXT("/")))
+							{
+								RelativeSourceFile += TEXT("/");
+							}
+
+							FPaths::MakePathRelativeTo(RelativeSourceFile, *SourcePath);
+							FString DestFile = TargetDirectoryCache / RelativeSourceFile;
+							FString DestDir = StatData.bIsDirectory ? DestFile : FPaths::GetPath(DestFile);
+
+							if (StatData.bIsDirectory)
+							{
+								IFileManager::Get().MakeDirectory(*DestDir, true);
+							}
+							else
+							{
+								SourceFiles.Add(FileOrDir);
+								DestFiles.Add(DestFile);
+							}
+						}
+						return true;
+					}
+				);
+
+				for	(int32 SourceFileIndex = 0; SourceFileIndex < SourceFiles.Num(); ++SourceFileIndex)
+				{
+					UE_LOG(LogConsoleResponse, Display, TEXT("[%d/%d] %s -> %s"), SourceFileIndex+1, SourceFiles.Num(), *SourceFiles[SourceFileIndex], *DestFiles[SourceFileIndex]);
+					IFileManager::Get().Copy(*DestFiles[SourceFileIndex], *SourceFiles[SourceFileIndex]);
+				}
+			}
+
+			UE_LOG(LogConsoleResponse, Display, TEXT("Copy files to %s done"), *TargetDirectory);
+		}
+	);
+}
+
+} // namespace IndexedCacheStorageManager_Debug
+#endif // !UE_BUILD_SHIPPING
+
 }
