Adding helper function to FNetRefHandle that returns an abbreviated (less verbose) string for debugging or logging

#rnx
#tests Used the new ToCompactString() in logging net ref handles and ensured the logging is abbreviated, properly formatted, but still contains the most critical info.

diff --git a/Engine/Source/Runtime/Experimental/Iris/Core/Private/Iris/ReplicationSystem/NetRefHandle.cpp b/Engine/Source/Runtime/Experimental/Iris/Core/Private/Iris/ReplicationSystem/NetRefHandle.cpp
--- a/Engine/Source/Runtime/Experimental/Iris/Core/Private/Iris/ReplicationSystem/NetRefHandle.cpp
+++ b/Engine/Source/Runtime/Experimental/Iris/Core/Private/Iris/ReplicationSystem/NetRefHandle.cpp
@@ -21,6 +21,19 @@ FString FNetRefHandle::ToString() const
 	}
 }
 
+FString FNetRefHandle::ToCompactString() const
+{
+	if (IsCompleteHandle())
+	{
+		const uint32 ReplicationSystemIdToDisplay = GetReplicationSystemId();
+		return FString::Printf(TEXT("%" UINT64_FMT ":%u"), GetId(), ReplicationSystemIdToDisplay);
+	}
+	else
+	{
+		return FString::Printf(TEXT("%" UINT64_FMT ":?"), GetId());
+	}
+}
+
 }
 
 FStringBuilderBase& operator<<(FStringBuilderBase& Builder, const UE::Net::FNetRefHandle& NetRefHandle) 
diff --git a/Engine/Source/Runtime/Experimental/Iris/Core/Public/Iris/ReplicationSystem/NetRefHandle.h b/Engine/Source/Runtime/Experimental/Iris/Core/Public/Iris/ReplicationSystem/NetRefHandle.h
--- a/Engine/Source/Runtime/Experimental/Iris/Core/Public/Iris/ReplicationSystem/NetRefHandle.h
+++ b/Engine/Source/Runtime/Experimental/Iris/Core/Public/Iris/ReplicationSystem/NetRefHandle.h
@@ -57,6 +57,7 @@ public:
 	bool operator!=(const FNetRefHandle& Other)const { return Id != Other.Id; }
 
 	IRISCORE_API FString ToString() const;
+	IRISCORE_API FString ToCompactString() const;
 
 	static bool FullCompare(FNetRefHandle A, FNetRefHandle B) { return A.Value == B.Value; }
 	
