IAD analytics improvements for purge and defrag.
- Add EIoErrorCode::DeleteError for deleting cache blocks
- Make sure OnDemandSetImmediateAnalyticHandler is thread safe
- Add IadCachePurge immediate event
- IadCachePurge and IadCacheDefrag now send full IoStatus strings instead of just error codes. Error codes alone aren't unqiue enough to pinpoint where an error occurs.
- Removed any dynamic state from relevant IoStatus errors. For analytics purposes, we need all errors to report the same string from the same code location.
- Fixed some cases where errors weren't correctly propogated up the callstack.

diff --git a/Engine/Source/Runtime/Core/Public/IO/IoStatus.h b/Engine/Source/Runtime/Core/Public/IO/IoStatus.h
--- a/Engine/Source/Runtime/Core/Public/IO/IoStatus.h
+++ b/Engine/Source/Runtime/Core/Public/IO/IoStatus.h
@@ -36,6 +36,7 @@ enum class EIoErrorCode
 	NotInstalled,
 	PendingHostGroup,
 	Timeout,
+	DeleteError,
 	Last
 };
 
diff --git a/Engine/Source/Runtime/Experimental/IoStore/OnDemand/Private/OnDemandInstallCache.cpp b/Engine/Source/Runtime/Experimental/IoStore/OnDemand/Private/OnDemandInstallCache.cpp
--- a/Engine/Source/Runtime/Experimental/IoStore/OnDemand/Private/OnDemandInstallCache.cpp
+++ b/Engine/Source/Runtime/Experimental/IoStore/OnDemand/Private/OnDemandInstallCache.cpp
@@ -310,9 +310,7 @@ FIoStatus FCas::DeleteBlock(FCasBlockId BlockId, TArray<FCasAddr>& OutAddrs)
 			if (FPlatformTime::ToMilliseconds(FPlatformTime::Cycles() - StartTimeCycles) > DeleteBlockMaxWaitTimeMs)
 			{
 				return FIoStatusBuilder(EIoErrorCode::Timeout)
-					<< TEXT("Timed out waiting for pending read(s) while deleting CAS block '")
-					<< Filename
-					<< TEXT("'"); 
+					<< TEXT("Timed out waiting for pending read(s) while deleting CAS block"); 
 			}
 
 			BlockReadsDoneEvent->Wait(WaitTimeMs);
@@ -327,10 +325,7 @@ FIoStatus FCas::DeleteBlock(FCasBlockId BlockId, TArray<FCasAddr>& OutAddrs)
 	UE_LOG(LogIoStoreOnDemand, Log, TEXT("Deleting CAS block '%s'"), *Filename);
 	if (Ipf.DeleteFile(*Filename) == false)
 	{
-		return FIoStatusBuilder(EIoErrorCode::WriteError)
-			<< TEXT("Failed to delete CAS block '")
-			<< Filename
-			<< TEXT("'");
+		return FIoStatusBuilder(EIoErrorCode::DeleteError) << TEXT("Failed to delete CAS block");
 	}
 
 	BlockIds.Remove(BlockId);
@@ -910,9 +905,7 @@ FIoStatus FCasJournal::Commit(FTransaction&& Transaction)
 		{
 			FOnDemandInstallCacheStats::OnJournalCommit(EIoErrorCode::ReadError, 0);
 			return FIoStatusBuilder(EIoErrorCode::FileOpenFailed)
-				<< TEXT("Failed to validate CAS journal file '")
-				<< Transaction.JournalFile
-				<< TEXT("'");
+				<< TEXT("Failed to validate CAS journal file");
 		}
 
 		FHeader Header;
@@ -920,9 +913,7 @@ FIoStatus FCasJournal::Commit(FTransaction&& Transaction)
 		{
 			FOnDemandInstallCacheStats::OnJournalCommit(EIoErrorCode::SignatureError, 0);
 			return FIoStatusBuilder(EIoErrorCode::ReadError)
-				<< TEXT("Failed to validate CAS journal header '")
-				<< Transaction.JournalFile
-				<< TEXT("'");
+				<< TEXT("Failed to validate CAS journal header");
 		}
 
 		const int64 FooterPos = FileSize - sizeof(FFooter);
@@ -930,9 +921,7 @@ FIoStatus FCasJournal::Commit(FTransaction&& Transaction)
 		{
 			FOnDemandInstallCacheStats::OnJournalCommit(EIoErrorCode::SignatureError, 0);
 			return FIoStatusBuilder(EIoErrorCode::ReadError)
-				<< TEXT("Failed to validate CAS journal footer '")
-				<< Transaction.JournalFile
-				<< TEXT("'");
+				<< TEXT("Failed to validate CAS journal footer");
 		}
 
 		FFooter Footer;
@@ -940,9 +929,7 @@ FIoStatus FCasJournal::Commit(FTransaction&& Transaction)
 		{
 			FOnDemandInstallCacheStats::OnJournalCommit(EIoErrorCode::SignatureError, 0);
 			return FIoStatusBuilder(EIoErrorCode::ReadError)
-				<< TEXT("Failed to validate CAS journal footer '")
-				<< Transaction.JournalFile
-				<< TEXT("'");
+				<< TEXT("Failed to validate CAS journal footer");
 		}
 	}
 
@@ -957,9 +944,7 @@ FIoStatus FCasJournal::Commit(FTransaction&& Transaction)
 		{
 			FOnDemandInstallCacheStats::OnJournalCommit(EIoErrorCode::WriteError, 0);
 			return FIoStatusBuilder(EIoErrorCode::FileOpenFailed)
-				<< TEXT("Failed to open CAS journal '")
-				<< Transaction.JournalFile
-				<< TEXT("'");
+				<< TEXT("Failed to open CAS journal");
 		}
 
 		const int64 TotalEntrySize = Transaction.Entries.Num() * sizeof(FEntry);
@@ -969,9 +954,7 @@ FIoStatus FCasJournal::Commit(FTransaction&& Transaction)
 		{
 			FOnDemandInstallCacheStats::OnJournalCommit(EIoErrorCode::WriteError, 0);
 			return FIoStatusBuilder(EIoErrorCode::WriteError)
-				<< TEXT("Failed to write CAS journal entries to '")
-				<< Transaction.JournalFile
-				<< TEXT("'");
+				<< TEXT("Failed to write CAS journal entries");
 		}
 
 		FFooter Footer;
@@ -980,15 +963,14 @@ FIoStatus FCasJournal::Commit(FTransaction&& Transaction)
 		{
 			FOnDemandInstallCacheStats::OnJournalCommit(EIoErrorCode::WriteError, 0);
 			return FIoStatusBuilder(EIoErrorCode::WriteError)
-				<< TEXT("Failed to write CAS journal footer to '")
-				<< Transaction.JournalFile
-				<< TEXT("'");
+				<< TEXT("Failed to write CAS journal footer");
 		}
 
 		if (FileHandle->Flush() == false)
 		{
 			FOnDemandInstallCacheStats::OnJournalCommit(EIoErrorCode::WriteError, 0);
-			return EIoErrorCode::WriteError;
+			return FIoStatusBuilder(EIoErrorCode::WriteError)
+				<< TEXT("Failed to flush CAS journal");
 		}
 
 		UE_LOG(LogIoStoreOnDemand, Log, TEXT("Committed %d CAS journal entries of total %.2lf KiB to '%s'"),
@@ -2109,9 +2091,7 @@ FIoStatus FOnDemandInstallCache::Purge(TSet<FIoHash>&& ChunksToInstall)
 		LogIoStoreOnDemand, Warning, TEXT("Max install cache size exceeded by %.2lf MiB (%.2lf%%)"),
 			ToMiB(NewCachedBytes - MaxCacheSize), 100.0 * (double(NewCachedBytes - MaxCacheSize) / double(MaxCacheSize)));
 
-	FOnDemandInstallCacheStats::OnPurge(Status.GetErrorCode(), MaxCacheSize, NewCachedBytes, TotalBytesToPurge, TotalPurgedBytes);
-
-	if (TotalPurgedBytes < TotalBytesToPurge)
+	if (Status.IsOk() && TotalPurgedBytes < TotalBytesToPurge)
 	{
 		if (UE::IoStore::CVars::GIoStoreOnDemandEnableDefrag)
 		{
@@ -2120,16 +2100,18 @@ FIoStatus FOnDemandInstallCache::Purge(TSet<FIoHash>&& ChunksToInstall)
 			FIoStatus DefragStatus = Defrag(Containers, ChunkEntryIndices, BlockInfo, &DefragBytesToPurge);
 			if (DefragStatus.IsOk() == false)
 			{
-				return FIoStatusBuilder(EIoErrorCode::WriteError) 
-					<< FString::Printf(TEXT("Failed to purge %" UINT64_FMT " from install cache after defrag (%s)"), TotalBytesToPurge, *DefragStatus.ToString());
+				Status = FIoStatusBuilder(EIoErrorCode::WriteError)
+					<< FString::Printf(TEXT("Defrag failed (%s)"), *DefragStatus.ToString());
 			}
 		}
 		else
 		{
-			return FIoStatusBuilder(EIoErrorCode::WriteError) << FString::Printf(TEXT("Failed to purge %" UINT64_FMT " from install cache"), TotalBytesToPurge);
+			Status = FIoStatusBuilder(EIoErrorCode::WriteError) << TEXT("Failed to purge required size from install cache");
 		}
 	}
 
+	FOnDemandInstallCacheStats::OnPurge(Status, MaxCacheSize, NewCachedBytes, TotalBytesToPurge, TotalPurgedBytes);
+
 	return Status;
 }
 
@@ -2180,13 +2162,13 @@ FIoStatus FOnDemandInstallCache::PurgeAllUnreferenced(bool bDefrag, const uint64
 			FIoStatus DefragStatus = Defrag(Containers, ChunkEntryIndices, BlockInfo, &DefragBytesToPurge);
 			if (DefragStatus.IsOk() == false)
 			{
-				return FIoStatusBuilder(EIoErrorCode::WriteError)
-					<< FString::Printf(TEXT("Failed to purge %" UINT64_FMT " from install cache after defrag (%s)"), TotalBytesToPurge, *DefragStatus.ToString());
+				Status = FIoStatusBuilder(EIoErrorCode::WriteError)
+					<< FString::Printf(TEXT("Defrag failed (%s)"), *DefragStatus.ToString());
 			}
 		}
 		else
 		{
-			return FIoStatusBuilder(EIoErrorCode::WriteError) << FString::Printf(TEXT("Failed to purge %" UINT64_FMT " from install cache"), TotalBytesToPurge);
+			Status = FIoStatusBuilder(EIoErrorCode::WriteError) << TEXT("Failed to purge required size from install cache");
 		}
 	}
 	else if (bDefrag)
@@ -2195,10 +2177,13 @@ FIoStatus FOnDemandInstallCache::PurgeAllUnreferenced(bool bDefrag, const uint64
 		FIoStatus DefragStatus = Defrag(Containers, ChunkEntryIndices, BlockInfo);
 		if (DefragStatus.IsOk() == false)
 		{
-			return DefragStatus;
+			Status = FIoStatusBuilder(EIoErrorCode::WriteError)
+				<< FString::Printf(TEXT("Defrag failed (%s)"), *DefragStatus.ToString());
 		}
 	}
 
+	FOnDemandInstallCacheStats::OnPurge(Status, MaxCacheSize, NewCachedBytes, TotalBytesToPurge, TotalPurgedBytes);
+
 	return Status;
 }
 
@@ -2673,21 +2658,23 @@ FIoStatus FOnDemandInstallCache::Defrag(
 	if (TotalCachedBytes > MaxCacheSize)
 	{
 		// Ruh-Roh! There's not enough of the disk quota left to run a defrag!
-		const FString ErrorMsg = FString::Printf(TEXT("Cache size is greater than disk quota - Cannot Defragment!, MaxCacheSize=%.2lf MiB, TotalCachedBytes=%.2lf MiB"), 
-			ToMiB(MaxCacheSize), ToMiB(TotalCachedBytes));
+		const FString ErrorMsg = FString::Printf(TEXT("Cache size is greater than disk quota - Cannot Defragment!, MaxCacheSize=%.2lf MiB"), 
+			ToMiB(MaxCacheSize));
 		UE_LOG(LogIoStoreOnDemand, Error, TEXT("%s"), *ErrorMsg);
-		FOnDemandInstallCacheStats::OnDefrag(EIoErrorCode::WriteError, 0);
+
+		FIoStatus Status = FIoStatusBuilder(EIoErrorCode::WriteError) << ErrorMsg;
+		FOnDemandInstallCacheStats::OnDefrag(Status, 0);
 
 		// Append a critical error entry to clear the cache at next startup
 		FCasJournal::FTransaction Transaction = FCasJournal::Begin(GetJournalFilename());
 		Transaction.CriticalError(FCasJournal::EErrorCode::DefragOutOfDiskSpace);
-		if (FIoStatus Status = FCasJournal::Commit(MoveTemp(Transaction)); Status.IsOk() == false)
+		if (FIoStatus JournalStatus = FCasJournal::Commit(MoveTemp(Transaction)); Status.IsOk() == false)
 		{
 			UE_LOG(LogIoStoreOnDemand, Error, TEXT("Failed to append critical error to journal, reason '%s'"),
-				*Status.ToString());
+				*JournalStatus.ToString());
 		}
 
-		return FIoStatusBuilder(EIoErrorCode::WriteError) << ErrorMsg;
+		return Status;
 	}
 
 	struct FDefragBlockReferencedChunk
@@ -2761,8 +2748,9 @@ FIoStatus FOnDemandInstallCache::Defrag(
 
 		if (!bPossibleToFreeBytes)
 		{
-			FOnDemandInstallCacheStats::OnDefrag(EIoErrorCode::WriteError, 0);
-			return FIoStatusBuilder(EIoErrorCode::WriteError) << FString::Printf(TEXT("Defrag failed - cannot free %" UINT64_FMT), *TotalBytesToFree);
+			FIoStatus Status = FIoStatusBuilder(EIoErrorCode::WriteError) << TEXT("Defrag cannot free requested storage space");
+			FOnDemandInstallCacheStats::OnDefrag(Status, 0);
+			return Status;
 		}
 	}
 	else
@@ -2847,11 +2835,11 @@ FIoStatus FOnDemandInstallCache::Defrag(
 			FSharedFileOpenResult FileOpenResult = Cas.OpenRead(DefragBlock.BlockId);
 			if (!FileOpenResult.IsValid())
 			{
-				const FString Filename = Cas.GetBlockFilename(DefragBlock.BlockId);
 				const FString ErrorMsg = FileOpenResult.GetError().GetMessage();
-				UE_LOG(LogIoStoreOnDemand, Error, TEXT("Failed to open CAS block '%s' for reading, reason '%s'"), *Filename, *ErrorMsg);
-				FOnDemandInstallCacheStats::OnDefrag(EIoErrorCode::FileOpenFailed, 0);
-				return FIoStatusBuilder(EIoErrorCode::FileOpenFailed) << ErrorMsg;
+				UE_LOG(LogIoStoreOnDemand, Error, TEXT("Defrag failed to open CAS block for reading, reason '%s'"), *ErrorMsg);
+				FIoStatus Status = FIoStatusBuilder(EIoErrorCode::FileOpenFailed) << ErrorMsg;
+				FOnDemandInstallCacheStats::OnDefrag(Status, 0);
+				return Status;
 			}
 
 			FSharedFileHandle FileHandle = FileOpenResult.StealValue();
@@ -2887,7 +2875,7 @@ FIoStatus FOnDemandInstallCache::Defrag(
 
 					if (FIoStatus Status = FlushPendingChunks(DefragPendingChunks, DefragBlock.LastAccess); Status.IsOk() == false)
 					{
-						FOnDemandInstallCacheStats::OnDefrag(Status.GetErrorCode(), 0);
+						FOnDemandInstallCacheStats::OnDefrag(Status, 0);
 						return Status;
 					}
 	
@@ -2899,7 +2887,7 @@ FIoStatus FOnDemandInstallCache::Defrag(
 				{
 					if (FIoStatus Status = FlushPendingChunks(DefragPendingChunks, DefragBlock.LastAccess); Status.IsOk() == false)
 					{
-						FOnDemandInstallCacheStats::OnDefrag(Status.GetErrorCode(), 0);
+						FOnDemandInstallCacheStats::OnDefrag(Status, 0);
 						return Status;
 					}
 					check(DefragPendingChunks.IsEmpty());
@@ -2912,7 +2900,7 @@ FIoStatus FOnDemandInstallCache::Defrag(
 
 			if (FIoStatus Status = FlushPendingChunks(DefragPendingChunks, DefragBlock.LastAccess); Status.IsOk() == false)
 			{
-				FOnDemandInstallCacheStats::OnDefrag(Status.GetErrorCode(), 0);
+				FOnDemandInstallCacheStats::OnDefrag(Status, 0);
 				return Status;
 			}
 			check(DefragPendingChunks.IsEmpty());
@@ -2925,7 +2913,7 @@ FIoStatus FOnDemandInstallCache::Defrag(
 		TArray<FCasAddr> DeletedChunkAddresses;
 		if (FIoStatus Status = Cas.DeleteBlock(DefragBlock.BlockId, DeletedChunkAddresses); !Status.IsOk())
 		{
-			FOnDemandInstallCacheStats::OnDefrag(Status.GetErrorCode(), 0);
+			FOnDemandInstallCacheStats::OnDefrag(Status, 0);
 			return Status;
 		}
 
@@ -2939,7 +2927,7 @@ FIoStatus FOnDemandInstallCache::Defrag(
 
 		if (FIoStatus Status = FCasJournal::Commit(MoveTemp(Transaction)); !Status.IsOk())
 		{
-			FOnDemandInstallCacheStats::OnDefrag(Status.GetErrorCode(), 0);
+			FOnDemandInstallCacheStats::OnDefrag(Status, 0);
 			return Status;
 		}
 	}
@@ -3061,10 +3049,7 @@ FIoStatus FOnDemandInstallCache::FlushPendingChunksImpl(FPendingChunks& Chunks,
 		if (CasFileHandle.IsValid() == false)
 		{
 			FOnDemandInstallCacheStats::OnFlush(EIoErrorCode::WriteError, 0);
-			return FIoStatusBuilder(EIoErrorCode::FileOpenFailed)
-				<< TEXT("Failed to open cache block file '")
-				<< Cas.GetBlockFilename(CurrentBlockId)
-				<< TEXT("'");
+			return FIoStatusBuilder(EIoErrorCode::FileOpenFailed) << TEXT("Failed to open cache block file");
 		}
 
 		const int64 CasBlockOffset = CasFileHandle->Tell();
diff --git a/Engine/Source/Runtime/Experimental/IoStore/OnDemand/Private/Statistics.cpp b/Engine/Source/Runtime/Experimental/IoStore/OnDemand/Private/Statistics.cpp
--- a/Engine/Source/Runtime/Experimental/IoStore/OnDemand/Private/Statistics.cpp
+++ b/Engine/Source/Runtime/Experimental/IoStore/OnDemand/Private/Statistics.cpp
@@ -71,18 +71,34 @@ extern int32 GIasHttpRateLimitKiBPerSecond;
 
 ////////////////////////////////////////////////////////////////////////////////
 static FOnDemandImmediateAnalyticHandler OnDemandImmediateAnalyticEventHandler;
+static FMutex OnDemandImmediateAnalyticEventHandlerMutex;
 void OnDemandSetImmediateAnalyticHandler(FOnDemandImmediateAnalyticHandler&& EventHandler)
 {
+	TUniqueLock Lock(OnDemandImmediateAnalyticEventHandlerMutex);
 	OnDemandImmediateAnalyticEventHandler = MoveTemp(EventHandler);
 }
 
 static void SendImmediateAnalytic(FOnDemandImmediateAnalytic&& Event)
 {
+	// FSafeHandler avoids racing with OnDemandSetImmediateAnalyticHandler.
+	// The handler should be set once before any event is sent.
+	struct FSafeHandler
+	{
+		FSafeHandler()
+		{
+			TUniqueLock Lock(OnDemandImmediateAnalyticEventHandlerMutex);
+			EventHandler = MoveTemp(OnDemandImmediateAnalyticEventHandler.EventHandler);
+		}
+
+		TUniqueFunction<void(FOnDemandImmediateAnalytic)> EventHandler;
+	};
+
 	if (GIaxImmediateAnalyticsChance > 0.0f && FMath::FRand() <= GIaxImmediateAnalyticsChance)
 	{
-		if (OnDemandImmediateAnalyticEventHandler.EventHandler)
+		static FSafeHandler SafeHandler;
+		if (SafeHandler.EventHandler)
 		{
-			OnDemandImmediateAnalyticEventHandler.EventHandler(MoveTemp(Event));
+			SafeHandler.EventHandler(MoveTemp(Event));
 		}
 	}
 }
@@ -1404,47 +1420,61 @@ void FOnDemandInstallCacheStats::OnCasVerificationError(int32 RemoveChunks)
 }
 
 void FOnDemandInstallCacheStats::OnPurge(
-	EIoErrorCode ErrorCode,
+	const FIoStatus& Status,
 	uint64 MaxCacheSize,
 	uint64 NewCacheSize,
 	uint64 BytesToPurge,
 	uint64 PurgedBytes)
 {
-	if (ErrorCode == EIoErrorCode::Cancelled)
+	if (Status.GetErrorCode() == EIoErrorCode::Cancelled)
 	{
 		return;
 	}
 
-	UE::TUniqueLock Lock(InstallCacheAnalytics);
-	if (ErrorCode != EIoErrorCode::Ok)
 	{
-		InstallCacheAnalytics.PurgeErrorCount++;
+		UE::TUniqueLock Lock(InstallCacheAnalytics);
+		if (Status.IsOk() == false)
+		{
+			InstallCacheAnalytics.PurgeErrorCount++;
+		}
+		InstallCacheAnalytics.PurgeCount++;
+		InstallCacheAnalytics.PurgedBytes += PurgedBytes;
 	}
-	InstallCacheAnalytics.PurgeCount++;
-	InstallCacheAnalytics.PurgedBytes += PurgedBytes;
+
+	SendImmediateAnalytic(FOnDemandImmediateAnalytic
+	{
+		TEXT("IadCachePurge"),
+		MakeAnalyticsEventAttributeArray(
+			TEXT("Success"), Status.IsOk(),
+			TEXT("ErrorCode"), Status.ToString(),
+			TEXT("PurgedBytes"), PurgedBytes
+		)
+	});
 }
 
-void FOnDemandInstallCacheStats::OnDefrag(EIoErrorCode ErrorCode, uint64 FragmentedBytes)
+void FOnDemandInstallCacheStats::OnDefrag(const FIoStatus& Status, uint64 FragmentedBytes)
 {
-	if (ErrorCode == EIoErrorCode::Cancelled)
+	if (Status.GetErrorCode() == EIoErrorCode::Cancelled)
 	{
 		return;
 	}
 
-	UE::TUniqueLock Lock(InstallCacheAnalytics);
-	if (ErrorCode != EIoErrorCode::Ok)
 	{
-		InstallCacheAnalytics.DefragErrorCount++;
+		UE::TUniqueLock Lock(InstallCacheAnalytics);
+		if (Status.IsOk() == false)
+		{
+			InstallCacheAnalytics.DefragErrorCount++;
+		}
+		InstallCacheAnalytics.DefragCount++;
+		InstallCacheAnalytics.FragmentedBytes += FragmentedBytes;
 	}
-	InstallCacheAnalytics.DefragCount++;
-	InstallCacheAnalytics.FragmentedBytes += FragmentedBytes;
 
 	SendImmediateAnalytic(FOnDemandImmediateAnalytic
 	{
 		TEXT("IadCacheDefrag"),
 		MakeAnalyticsEventAttributeArray(
-			TEXT("Success"), ErrorCode == EIoErrorCode::Ok,
-			TEXT("ErrorCode"), GetIoErrorText(ErrorCode),
+			TEXT("Success"), Status.IsOk(),
+			TEXT("ErrorCode"), Status.ToString(),
 			TEXT("FragmentedBytes"), FragmentedBytes
 		)
 	});
@@ -1452,11 +1482,6 @@ void FOnDemandInstallCacheStats::OnDefrag(EIoErrorCode ErrorCode, uint64 Fragmen
 
 void FOnDemandInstallCacheStats::OnBlockDeleted(int64 LastAccessTicks, bool bFromDefrag)
 {
-	if (!OnDemandImmediateAnalyticEventHandler.EventHandler)
-	{
-		return;
-	}
-
 	const FDateTime Now = FDateTime::UtcNow();
 	const FDateTime BlockAccess(LastAccessTicks);
 	const FTimespan BlockAge = ensure(Now >= BlockAccess) ? (Now - BlockAccess) : FTimespan::MaxValue();
diff --git a/Engine/Source/Runtime/Experimental/IoStore/OnDemand/Private/Statistics.h b/Engine/Source/Runtime/Experimental/IoStore/OnDemand/Private/Statistics.h
--- a/Engine/Source/Runtime/Experimental/IoStore/OnDemand/Private/Statistics.h
+++ b/Engine/Source/Runtime/Experimental/IoStore/OnDemand/Private/Statistics.h
@@ -15,6 +15,7 @@
 #endif
 
 struct FAnalyticsEventAttribute;
+class FIoStatus;
 enum class EIoErrorCode;
 
 LLM_DECLARE_TAG(Ias);
@@ -137,12 +138,12 @@ public:
 	static void OnJournalCommit(EIoErrorCode ErrorCode, int64 ByteCount) IAS_STATISTICS_IMPL()
 	static void OnCasVerificationError(int32 RemoveChunks) IAS_STATISTICS_IMPL()
 	static void OnPurge(
-		EIoErrorCode ErrorCode,
+		const FIoStatus& Status,
 		uint64 MaxCacheSize,
 		uint64 NewCacheSize,
 		uint64 BytesToPurge,
 		uint64 PurgedBytes) IAS_STATISTICS_IMPL()
-	static void OnDefrag(EIoErrorCode ErrorCode, uint64 FragmentedBytes) IAS_STATISTICS_IMPL()
+	static void OnDefrag(const FIoStatus& Status, uint64 FragmentedBytes) IAS_STATISTICS_IMPL()
 	static void OnBlockDeleted(int64 LastAccessTicks, bool bFromDefrag) IAS_STATISTICS_IMPL()
 	static void OnCacheUsage(
 		uint64 MaxCacheSize,
