Dataflow : make dynamic connection logic to be more resilient for reconnecting inputs and outputs after changes occurs to the property bags
- the logic was relying on the fact that GUID may be the same for properties that do not change, but it seems to be possible for those ID to be changed
- In the case the ID has changed, the logic now fallback to use the name as a way to retrieve the saved remote connections
- example of this problem : Geometry script nodes : repro steps : Double click to open the geometry script , press compile ( all the connections would disconnect the first time compiled is pressed )

diff --git a/Engine/Source/Runtime/Experimental/Dataflow/Engine/Private/Dataflow/DataflowDynamicConnections.cpp b/Engine/Source/Runtime/Experimental/Dataflow/Engine/Private/Dataflow/DataflowDynamicConnections.cpp
--- a/Engine/Source/Runtime/Experimental/Dataflow/Engine/Private/Dataflow/DataflowDynamicConnections.cpp
+++ b/Engine/Source/Runtime/Experimental/Dataflow/Engine/Private/Dataflow/DataflowDynamicConnections.cpp
@@ -71,6 +71,8 @@ void FDataflowDynamicConnections::Refresh()
 		}
 	}
 	// clear all the dynamic connections
+	// and first keep the name to id mapping in case we need it for the fallback of not finding the id from the new properties in the bag 
+	TMap<FName, FGuid> PreviousConnectionNameToPropertyId = MoveTemp(ConnectionNameToPropertyId);
 	ClearDynamicConnections();
 
 	OwnerNode->ResumeInvalidations();
@@ -88,8 +90,19 @@ void FDataflowDynamicConnections::Refresh()
 				const FName CppType = GetCppTypeFromPropertyDesc(PropertyDesc);
 				if (FDataflowConnection* NewConnection = AddNewConnectionFromPropertyDesc(PropertyDesc))
 				{
+					FRemoteConnections* RemoteConnections = RemoteConnectionsByPropertyId.Find(PropertyDesc.ID);
+					if (!RemoteConnections)
+					{
+						// Ids may not match anymore ( internal logic of the property bag ) 
+						// so let's fall back to the name
+						if (const FGuid* PropertyId = PreviousConnectionNameToPropertyId.Find(NewConnection->GetName()))
+						{
+							RemoteConnections = RemoteConnectionsByPropertyId.Find(*PropertyId);
+						}
+					}
+
 					// now try to reconnect what was connected before
-					if (FRemoteConnections* RemoteConnections = RemoteConnectionsByPropertyId.Find(PropertyDesc.ID))
+					if (RemoteConnections)
 					{
 						for (FDataflowConnection* RemoteConnection : *RemoteConnections)
 						{
