Iris
* Added CreationDependency list
* Renamed NetRefHandleManager::IsValidHandle to IsNetRefHandleAssigned to more clearly show its a different method from Handle.IsValid
* Added FAsyncLoadingSimulator data to be later used to simulate async loading issues on specific objects in unit tests

#rnx

diff --git a/Engine/Source/Runtime/Experimental/Iris/Core/Private/Iris/ReplicationSystem/NetDependencyData.cpp b/Engine/Source/Runtime/Experimental/Iris/Core/Private/Iris/ReplicationSystem/NetDependencyData.cpp
--- a/Engine/Source/Runtime/Experimental/Iris/Core/Private/Iris/ReplicationSystem/NetDependencyData.cpp
+++ b/Engine/Source/Runtime/Experimental/Iris/Core/Private/Iris/ReplicationSystem/NetDependencyData.cpp
@@ -35,6 +35,11 @@ void FNetDependencyData::FreeStoredDependencyDataForObject(FInternalNetRefIndex
 			DependentObjectInfosStorage.RemoveAt(Entry->DependentObjectsInfoArrayIndex);
 		}
 
+		if (Entry->CreationDependencyArrayIndex != FDependencyInfo::InvalidCacheIndex)
+		{
+			CreationDependencyInfosStorage.RemoveAt(Entry->CreationDependencyArrayIndex);
+		}
+
 		DependencyInfos.Remove(InternalIndex);
 	}
 }
@@ -116,6 +121,38 @@ FNetDependencyData::FInternalNetRefIndexArray& FNetDependencyData::GetOrCreateIn
 	}
 }
 
+
+FNetDependencyData::FCreationDependencyInfoArray& FNetDependencyData::GetOrCreateCreationDependencyInfoArray(FInternalNetRefIndex ChildIndex)
+{
+	FDependencyInfo& Entry = GetOrCreateCacheEntry(ChildIndex);
+	
+	if (Entry.CreationDependencyArrayIndex == FDependencyInfo::InvalidCacheIndex)
+	{
+		FSparseArrayAllocationInfo AllocInfo = CreationDependencyInfosStorage.AddUninitialized();
+		Entry.CreationDependencyArrayIndex = AllocInfo.Index;
+
+		FCreationDependencyInfoArray* InfoArray = new (AllocInfo.Pointer) FCreationDependencyInfoArray();
+		
+		return *InfoArray;
+	}
+	else
+	{
+		return CreationDependencyInfosStorage[Entry.CreationDependencyArrayIndex];
+	}
+}
+
+void FNetDependencyData::FreeCreationDependencyInfoArray(FInternalNetRefIndex ChildIndex)
+{
+	if (FDependencyInfo* Entry = DependencyInfos.Find(ChildIndex))
+	{
+		if (Entry->CreationDependencyArrayIndex != FDependencyInfo::InvalidCacheIndex)
+		{
+			CreationDependencyInfosStorage.RemoveAt(Entry->CreationDependencyArrayIndex);
+			Entry->CreationDependencyArrayIndex = FDependencyInfo::InvalidCacheIndex;
+		}
+	}
+}
+
 FNetDependencyData::FInternalNetRefIndexArray& FNetDependencyData::GetOrCreateInternalIndexArray(FInternalNetRefIndex OwnerIndex, EArrayType ArrayType)
 {
 	FDependencyInfo& Entry = GetOrCreateCacheEntry(OwnerIndex);
diff --git a/Engine/Source/Runtime/Experimental/Iris/Core/Private/Iris/ReplicationSystem/NetDependencyData.h b/Engine/Source/Runtime/Experimental/Iris/Core/Private/Iris/ReplicationSystem/NetDependencyData.h
--- a/Engine/Source/Runtime/Experimental/Iris/Core/Private/Iris/ReplicationSystem/NetDependencyData.h
+++ b/Engine/Source/Runtime/Experimental/Iris/Core/Private/Iris/ReplicationSystem/NetDependencyData.h
@@ -47,12 +47,14 @@ public:
 	typedef TArray<FInternalNetRefIndex, TInlineAllocator<8>> FInternalNetRefIndexArray;
 	typedef TArray<FLifeTimeConditionStorage, TInlineAllocator<8>> FSubObjectConditionalsArray;
 	typedef TArray<FDependentObjectInfo, TInlineAllocator<8>> FDependentObjectInfoArray;
+	typedef TArray<FInternalNetRefIndex, TInlineAllocator<2>> FCreationDependencyInfoArray;
 
 	enum EArrayType
 	{
 		SubObjects = 0U,
 		ChildSubObjects,
 		DependentParentObjects,
+		CreationDependencies,
 		Count
 	};
 
@@ -120,6 +122,34 @@ public:
 		return true;
 	}
 
+	/** Create or return the creation dependency list of an object */
+	FCreationDependencyInfoArray& GetOrCreateCreationDependencyInfoArray(FInternalNetRefIndex ChildIndex);
+	
+	/** Liberate the creation dependency list assigned to an object */
+	void FreeCreationDependencyInfoArray(FInternalNetRefIndex ChildIndex);
+
+	FCreationDependencyInfoArray* GetCreationDependencyInfoArray(FInternalNetRefIndex ChildIndex)
+	{
+		const FDependencyInfo* Entry = DependencyInfos.Find(ChildIndex);
+		const uint32 ArrayIndex = Entry ? Entry->CreationDependencyArrayIndex : FDependencyInfo::InvalidCacheIndex;
+		if (ArrayIndex != FDependencyInfo::InvalidCacheIndex)
+		{
+			return &CreationDependencyInfosStorage[ArrayIndex];
+		}
+		return nullptr;
+	}
+
+	const FCreationDependencyInfoArray* GetCreationDependencyInfoArray(FInternalNetRefIndex ChildIndex) const
+	{
+		const FDependencyInfo* Entry = DependencyInfos.Find(ChildIndex);
+		const uint32 ArrayIndex = Entry ? Entry->CreationDependencyArrayIndex : FDependencyInfo::InvalidCacheIndex;
+		if (ArrayIndex != FDependencyInfo::InvalidCacheIndex)
+		{
+			return &CreationDependencyInfosStorage[ArrayIndex];
+		}
+		return nullptr;
+	}
+
 	template<EArrayType TypeIndex>
 	FInternalNetRefIndexArray* GetInternalIndexArray(FInternalNetRefIndex InternalIndex)
 	{
@@ -154,9 +184,10 @@ private:
 	struct FDependencyInfo
 	{
 		constexpr static uint32 InvalidCacheIndex = ~(0U);
-		uint32 ArrayIndices[EArrayType::Count] = { InvalidCacheIndex, InvalidCacheIndex, InvalidCacheIndex };
+		uint32 ArrayIndices[EArrayType::Count] = { InvalidCacheIndex, InvalidCacheIndex, InvalidCacheIndex, InvalidCacheIndex };
 		uint32 SubObjectConditionalArrayIndex = InvalidCacheIndex;
 		uint32 DependentObjectsInfoArrayIndex = InvalidCacheIndex;
+		uint32 CreationDependencyArrayIndex = InvalidCacheIndex;
 	};
 
 private:
@@ -165,12 +196,19 @@ private:
 	
 	// Map to track the replicated objects with subObjects or dependencies
 	TMap<FInternalNetRefIndex, FDependencyInfo> DependencyInfos;
+
 	// Storage for DependentObjects and SubObjects
 	TSparseArray<FInternalNetRefIndexArray> DependentObjectsStorage;
+
 	// Storage for SubObject conditionals
 	TSparseArray<FSubObjectConditionalsArray> SubObjectConditionalsStorage;
-	// Storage for DependentObjects
+
+	// Storage for DependentObjects traits and info (bound on the Parent)
 	TSparseArray<FDependentObjectInfoArray> DependentObjectInfosStorage;
+
+	// Storage for creation dependency infos (bound on the Child)
+	TSparseArray<FCreationDependencyInfoArray> CreationDependencyInfosStorage;
+
 };
 
 }
diff --git a/Engine/Source/Runtime/Experimental/Iris/Core/Private/Iris/ReplicationSystem/NetRefHandleManager.cpp b/Engine/Source/Runtime/Experimental/Iris/Core/Private/Iris/ReplicationSystem/NetRefHandleManager.cpp
--- a/Engine/Source/Runtime/Experimental/Iris/Core/Private/Iris/ReplicationSystem/NetRefHandleManager.cpp
+++ b/Engine/Source/Runtime/Experimental/Iris/Core/Private/Iris/ReplicationSystem/NetRefHandleManager.cpp
@@ -78,6 +78,7 @@ void FNetRefHandleManager::Init(const FInitParams& InitParams)
 		InitNetBitArray(&SubObjectInternalIndices);
 		InitNetBitArray(&DependentObjectInternalIndices);
 		InitNetBitArray(&DependentObjectsWithParentFilter);
+		InitNetBitArray(&ObjectsWithCreationDependencies);
 		InitNetBitArray(&ObjectsWithDependentObjectsInternalIndices);
 		InitNetBitArray(&DestroyedStartupObjectInternalIndices);
 		InitNetBitArray(&WantToBeDormantInternalIndices);
@@ -891,6 +892,90 @@ void FNetRefHandleManager::RemoveDependentObject(FNetRefHandle DependentHandle)
 	}
 }
 
+bool FNetRefHandleManager::AddCreationDependency(FNetRefHandle Parent, FNetRefHandle Child)
+{
+	if (Parent == Child)
+	{
+		ensureMsgf(false, TEXT("Cannot add a creation dependency on yourself: %s"), *PrintObjectFromNetRefHandle(Child));
+		return false;
+	}
+
+	FInternalNetRefIndex ParentIndex = GetInternalIndex(Parent);
+	FInternalNetRefIndex ChildIndex = GetInternalIndex(Child);
+
+	if (ParentIndex == InvalidInternalIndex || ChildIndex == InvalidInternalIndex)
+	{
+		ensureMsgf(false, TEXT("AddCreationDependency failed due to invalid reference. Parent: %s | Child: %s"), *PrintObjectFromNetRefHandle(Parent), *PrintObjectFromNetRefHandle(Child));
+		return false;
+	}
+
+	if (IsSubObject(ParentIndex) || IsSubObject(ChildIndex))
+	{
+		UE_LOG(LogIris, Error, TEXT("AddCreationDependency only works on RootObjects. Parent: %s | Child: %s"), *PrintObjectFromNetRefHandle(Parent), *PrintObjectFromNetRefHandle(Child));
+		ensureMsgf(false, TEXT("AddCreationDependency not supported on subobjects"));
+		return false;
+	}
+
+	FNetDependencyData::FCreationDependencyInfoArray& CreationDependencyInfo = SubObjects.GetOrCreateCreationDependencyInfoArray(ChildIndex);
+
+	// Make sure the parent isn't already in the list
+	if (CreationDependencyInfo.Find(ParentIndex) != INDEX_NONE)
+	{
+		return true;
+	}
+
+	CreationDependencyInfo.Add(ParentIndex);
+
+	ObjectsWithCreationDependencies.SetBit(ChildIndex);
+
+	return true;
+}
+
+void FNetRefHandleManager::RemoveCreationDependency(FNetRefHandle Parent, FNetRefHandle Child)
+{
+	FInternalNetRefIndex ParentIndex = GetInternalIndex(Parent);
+	FInternalNetRefIndex ChildIndex = GetInternalIndex(Child);
+
+	if (ParentIndex == InvalidInternalIndex || ChildIndex == InvalidInternalIndex)
+	{
+		ensureMsgf(false, TEXT("RemoveCreationDependency failed due to invalid reference. Parent: %s | Child: %s"), *PrintObjectFromNetRefHandle(Parent), *PrintObjectFromNetRefHandle(Child));
+		return;
+	}
+
+	FNetDependencyData::FCreationDependencyInfoArray* CreationDependencyInfo = SubObjects.GetCreationDependencyInfoArray(ChildIndex);
+	if (!CreationDependencyInfo)
+	{
+		// No creation dependencies exist
+		return;
+	}
+
+	CreationDependencyInfo->Remove(ParentIndex);
+
+	if (CreationDependencyInfo->IsEmpty())
+	{
+		ObjectsWithCreationDependencies.ClearBit(ChildIndex);
+		
+		//TODO: Could free the creation info when empty to leave some memory. But if most operations are Remove -> Add then it's faster to keep the array assigned. To meditate on.
+	}
+}
+
+TConstArrayView<const FInternalNetRefIndex> FNetRefHandleManager::GetCreationDependencies(FInternalNetRefIndex ChildInternalIndex) const
+{
+	if (ObjectsWithCreationDependencies.IsBitSet(ChildInternalIndex) == false)
+	{
+		return MakeArrayView<const FInternalNetRefIndex>(nullptr, 0);
+	}
+
+	const FNetDependencyData::FCreationDependencyInfoArray* CreationDependencyArray = SubObjects.GetCreationDependencyInfoArray(ChildInternalIndex);
+	
+	if (!CreationDependencyArray)
+	{
+		return MakeArrayView<const FInternalNetRefIndex>(nullptr, 0);
+	}
+	
+	return MakeArrayView(*CreationDependencyArray);
+}
+
 void FNetRefHandleManager::InternalRemoveDependentObject(FInternalNetRefIndex ParentInternalIndex, FInternalNetRefIndex DependentInternalIndex, ERemoveDependentObjectFlags Flags)
 {
 	if (EnumHasAnyFlags(Flags, ERemoveDependentObjectFlags::RemoveFromDependentParentObjects))
@@ -961,6 +1046,7 @@ void FNetRefHandleManager::InternalRemoveDependentObject(FInternalNetRefIndex De
 	ObjectsWithDependentObjectsInternalIndices.ClearBit(DependentInternalIndex);
 	DependentObjectInternalIndices.ClearBit(DependentInternalIndex);
 	DependentObjectsWithParentFilter.ClearBit(DependentInternalIndex);
+	ObjectsWithCreationDependencies.ClearBit(DependentInternalIndex);
 }
 
 void FNetRefHandleManager::RemoveDependentObject(FNetRefHandle ParentHandle, FNetRefHandle DependentHandle)
diff --git a/Engine/Source/Runtime/Experimental/Iris/Core/Private/Iris/ReplicationSystem/NetRefHandleManager.h b/Engine/Source/Runtime/Experimental/Iris/Core/Private/Iris/ReplicationSystem/NetRefHandleManager.h
--- a/Engine/Source/Runtime/Experimental/Iris/Core/Private/Iris/ReplicationSystem/NetRefHandleManager.h
+++ b/Engine/Source/Runtime/Experimental/Iris/Core/Private/Iris/ReplicationSystem/NetRefHandleManager.h
@@ -205,8 +205,8 @@ public:
 
 	inline EIrisAsyncLoadingPriority GetIrisAsyncLoadingPriority(FInternalNetRefIndex InternalIndex) const;
 
-	/** Verify a handle against internal handle. A handle is valid if it matches internal storage. */
-	inline bool IsValidNetRefHandle(FNetRefHandle Handle) const;
+	/** Returns true if the handle has been assigned an internal index and it's replicated object has been bound to local storage data. */
+	inline bool IsNetRefHandleAssigned(FNetRefHandle Handle) const;
 
 	/** Returns true if the InternaIndex belongs to a replicated object owned by the local peer. */
 	inline bool IsLocal(FInternalNetRefIndex InternalIndex) const;
@@ -279,6 +279,17 @@ public:
 	bool AddDependentObject(FNetRefHandle ParentHandle, FNetRefHandle DependentHandle, EDependentObjectSchedulingHint SchedulingHint, EDependentFilterTrait FilterTrait = EDependentFilterTrait::None);
 	void RemoveDependentObject(FNetRefHandle ParentHandle, FNetRefHandle DependentHandle);
 
+	// Creation dependency API
+	bool AddCreationDependency(FNetRefHandle Parent, FNetRefHandle Child);
+	void RemoveCreationDependency(FNetRefHandle Parent, FNetRefHandle Child);
+
+	const FNetBitArrayView GetObjectsWithCreationDependencies() const
+	{
+		return MakeNetBitArrayView(ObjectsWithCreationDependencies);
+	}
+
+	TConstArrayView<const FInternalNetRefIndex> GetCreationDependencies(FInternalNetRefIndex ChildInternalIndex) const;
+
 	// Remove DependentHandles from all dependent object tracking
 	// $TODO: This might need a pass ObjectReplicationBridge as well since most likely also must restore filter/polling logic
 	void RemoveDependentObject(FNetRefHandle DependentHandle);
@@ -485,6 +496,9 @@ private:
 	// Dependent objects that need their parent to be relevant too
 	FNetBitArray DependentObjectsWithParentFilter;
 
+	// Objects that cannot be created until the parents they depend on also exist
+	FNetBitArray ObjectsWithCreationDependencies;
+
 	// Bitset containing all internal indices that have dependentObjects
 	FNetBitArray ObjectsWithDependentObjectsInternalIndices;
 
@@ -530,6 +544,7 @@ private:
 	uint64 NextStaticHandleIndex = 1;
 	uint64 NextDynamicHandleIndex = 1;
 
+	//TODO: Rename this variable its not exclusive to subobjects
 	FNetDependencyData SubObjects;
 
 	FReplicationProtocolManager& ReplicationProtocolManager;
@@ -592,7 +607,7 @@ EIrisAsyncLoadingPriority FNetRefHandleManager::GetIrisAsyncLoadingPriority(FInt
 	return ObjectData.IrisAsyncLoadingPriority;
 }
 
-bool FNetRefHandleManager::IsValidNetRefHandle(FNetRefHandle Handle) const
+bool FNetRefHandleManager::IsNetRefHandleAssigned(FNetRefHandle Handle) const
 {
 	return RefHandleToInternalIndex.Contains(Handle);
 }
diff --git a/Engine/Source/Runtime/Experimental/Iris/Core/Private/Iris/ReplicationSystem/ReplicationBridge.cpp b/Engine/Source/Runtime/Experimental/Iris/Core/Private/Iris/ReplicationSystem/ReplicationBridge.cpp
--- a/Engine/Source/Runtime/Experimental/Iris/Core/Private/Iris/ReplicationSystem/ReplicationBridge.cpp
+++ b/Engine/Source/Runtime/Experimental/Iris/Core/Private/Iris/ReplicationSystem/ReplicationBridge.cpp
@@ -1264,3 +1264,32 @@ void UReplicationBridge::PostSeamlessTravelGarbageCollect()
 void UReplicationBridge::OnPostSeamlessTravelGarbageCollect()
 {
 }
+
+namespace UE::Net::Private
+{
+
+#if UE_NET_ASYNCLOADING_DEBUG
+void FAsyncLoadingSimulator::AddStalledObject(UReplicationBridge* Bridge, FNetRefHandle ReplicatedObject)
+{
+	check(Bridge);
+	if (ReplicatedObject.IsValid())
+	{
+		Bridge->DebugObjectsForcedToStall.AddUnique(ReplicatedObject);
+	}
+}
+
+void FAsyncLoadingSimulator::RemoveStalledObject(UReplicationBridge* Bridge, FNetRefHandle ReplicatedObject)
+{
+	check(Bridge);
+	Bridge->DebugObjectsForcedToStall.Remove(ReplicatedObject);
+}
+
+void FAsyncLoadingSimulator::RemoveAllStalledObjects(UReplicationBridge* Bridge)
+{
+	check(Bridge);
+	Bridge->DebugObjectsForcedToStall.Empty();
+}
+
+#endif // UE_NET_ASYNCLOADING_DEBUG
+
+} // namespace UE::Net::Private
\ No newline at end of file
diff --git a/Engine/Source/Runtime/Experimental/Iris/Core/Private/Iris/ReplicationSystem/ReplicationSystem.cpp b/Engine/Source/Runtime/Experimental/Iris/Core/Private/Iris/ReplicationSystem/ReplicationSystem.cpp
--- a/Engine/Source/Runtime/Experimental/Iris/Core/Private/Iris/ReplicationSystem/ReplicationSystem.cpp
+++ b/Engine/Source/Runtime/Experimental/Iris/Core/Private/Iris/ReplicationSystem/ReplicationSystem.cpp
@@ -1407,7 +1407,7 @@ UReplicationBridge* UReplicationSystem::GetReplicationBridge() const
 
 bool UReplicationSystem::IsValidHandle(FNetRefHandle Handle) const
 {
-	return Handle.IsValid() && Impl->ReplicationSystemInternal.GetNetRefHandleManager().IsValidNetRefHandle(Handle);
+	return Handle.IsValid() && Impl->ReplicationSystemInternal.GetNetRefHandleManager().IsNetRefHandleAssigned(Handle);
 }
 
 const UE::Net::FReplicationProtocol* UReplicationSystem::GetReplicationProtocol(FNetRefHandle Handle) const
diff --git a/Engine/Source/Runtime/Experimental/Iris/Core/Public/Iris/ReplicationSystem/ReplicationBridge.h b/Engine/Source/Runtime/Experimental/Iris/Core/Public/Iris/ReplicationSystem/ReplicationBridge.h
--- a/Engine/Source/Runtime/Experimental/Iris/Core/Public/Iris/ReplicationSystem/ReplicationBridge.h
+++ b/Engine/Source/Runtime/Experimental/Iris/Core/Public/Iris/ReplicationSystem/ReplicationBridge.h
@@ -25,6 +25,10 @@
 
 #include "ReplicationBridge.generated.h"
 
+#ifndef UE_NET_ASYNCLOADING_DEBUG
+	#define UE_NET_ASYNCLOADING_DEBUG !(UE_BUILD_SHIPPING || UE_BUILD_TEST)
+#endif
+
 class UObjectReplicationBridge;
 class UReplicationSystem;
 class UNetDriver;
@@ -59,6 +63,7 @@ namespace UE::Net
 		
 		struct FChangeMaskCache;
 		struct FCreateNetObjectParams;
+		struct FAsyncLoadingSimulator;
 	}
 
 	typedef TArray<FNetDependencyInfo, TInlineAllocator<32> > FNetDependencyInfoArray;
@@ -393,8 +398,31 @@ private:
 		EPendingEndReplicationImmediate Immediate;
 	};
 	TArray<FPendingEndReplicationInfo> HandlesPendingEndReplication;
+
+#if UE_NET_ASYNCLOADING_DEBUG
+	friend UE::Net::Private::FAsyncLoadingSimulator;
+	/** 
+	 * Debug array that forces objects to be stalled, as if async loading external references.
+	 * Only enabled when net.Iris.AsyncLoading.ForceObjectsToStall is enabled.
+	 */
+	TArray<FNetRefHandle> DebugObjectsForcedToStall;
+#endif
+
 };
 
+namespace UE::Net::Private
+{
+#if UE_NET_ASYNCLOADING_DEBUG
+	struct FAsyncLoadingSimulator
+	{
+		IRISCORE_API static void AddStalledObject(UReplicationBridge* Bridge, FNetRefHandle ReplicatedObject);
+		IRISCORE_API static void RemoveStalledObject(UReplicationBridge* Bridge, FNetRefHandle ReplicatedObject);
+		IRISCORE_API static void RemoveAllStalledObjects(UReplicationBridge* Bridge);
+	};
+#endif
+
+}
+
 inline FReplicationBridgeSerializationContext::FReplicationBridgeSerializationContext(UE::Net::FNetSerializationContext& InSerialiazationContext, uint32 InConnectionId, bool bInIsDestructionInfo)
 : SerializationContext(InSerialiazationContext)
 , ConnectionId(InConnectionId)
