Fixed an issue where IAX would return incorrect data if a ranged read was requested but the target CDN does not support partial content requests.

#rnx

### Bug
- If a CDN does not support partial content reads they tend to just return the entire chunk instead and return (200) rather than (206)
- If we detect that the request was originally asking for a partial range but received a full chunk (200) we can return a partial slice of that chunk to the caller which is what they will be expecting.
-- Note that by passing the original FIoBuffer in as the outer when creating the slice we create a reference count to the original data so it will not be deleted until the slice we create is also deleted. Technically this could take up more memory but in practice the returned buffers are not held onto for very long so it's not worth the cost of allocating a new buffer and copying out the slice instead.
- We should consider testing each CDN for partial content support when first connecting via a HEAD request to a known resource in the future.

## Misc
- Added a new cvar 'iax.HttpAllowPartialContentRequests' that when set to false will prevent us from making the partial content request in the first place. This allows us to easily test this fallback code by effectively forcing all requests to download the full chunk. By default this is set to true and partial content requests are enabled.
- As downloading the full chunk is fairly inefficient I have added a warning to make it clear when this problem is occuring but it will only be logged every 50 requests to prevent log spam.

diff --git a/Engine/Source/Runtime/Experimental/IoStore/OnDemand/Private/OnDemandHttpClient.cpp b/Engine/Source/Runtime/Experimental/IoStore/OnDemand/Private/OnDemandHttpClient.cpp
--- a/Engine/Source/Runtime/Experimental/IoStore/OnDemand/Private/OnDemandHttpClient.cpp
+++ b/Engine/Source/Runtime/Experimental/IoStore/OnDemand/Private/OnDemandHttpClient.cpp
@@ -28,6 +28,20 @@ static FAutoConsoleVariableRef CVar_IaxWarnOnFailure(
 	TEXT("When enabled FMultiEndpointHttpClient will log failed requests as warnings instead of the config verbosity")
 );
 
+/**
+ * When enabled, http requests will make a partial content request by adding Range: bytes=<offset>-<offset+length> to the request if a valid
+ * FIoOffsetAndLength is provided with the request. When disabled we will always request the full chunk and then use the FIoOffsetAndLength to
+ * return a partial slice of the data to the caller.
+ */
+bool GAllowPartialContentRequests = true;
+static FAutoConsoleVariableRef CVar_IaxAllowPartialContentRequests(
+	TEXT("iax.HttpAllowPartialContentRequests"),
+	GAllowPartialContentRequests,
+	TEXT("Enable/disable the use of partial content http requests")
+);
+
+///////////////////////////////////////////////////////////////////////////////
+
 FIoBuffer JoinIoBuffers(const FIoBuffer& LHS, const FIoBuffer& RHS)
 {
 	const uint64 TotalSize = LHS.GetSize() + RHS.GetSize();
@@ -119,6 +133,8 @@ FIoStatus LoadDefaultHttpCertificates(bool& bWasLoaded)
 	return DefaultCerts.Status;
 }
 
+///////////////////////////////////////////////////////////////////////////////
+
 FMultiEndpointHttpClient::FMultiEndpointHttpClient(const FMultiEndpointHttpClientConfig& InConfig)
 	: Config(InConfig)
 {
@@ -320,7 +336,7 @@ void FMultiEndpointHttpClient::IssueRequest(FRequest&& Request)
 		}
 	}
 
-	if (Request.Range.GetOffset() > 0 || Request.Range.GetLength() > 0)
+	if (GAllowPartialContentRequests && Request.Range.IsValid())
 	{
 		HttpRequest.Header(ANSITEXTVIEW("Range"),
 			WriteToAnsiString<64>(ANSITEXTVIEW("bytes="), Request.Range.GetOffset(), ANSITEXTVIEW("-"), Request.Range.GetOffset() + Request.Range.GetLength() - 1));
@@ -388,6 +404,26 @@ void FMultiEndpointHttpClient::IssueRequest(FRequest&& Request)
 					};
 
 					if (IsHttpStatusOk(Request.StatusCode) || Config.bResponseBodyOnError)
+					{
+						if (Request.StatusCode != 206 && Request.Range.IsValid())
+						{
+							// A partial read was requested but the server sent us the full response body so we need to cut down
+							// the data that we return. Note that the partial buffer we return will retain a ref count on the full
+							// sized buffer until it is also destroyed, so we don't need to do any allocations/copies at this point.
+							FMemoryView PartialRange = Request.Body.GetView().Mid(Request.Range.GetOffset(), Request.Range.GetLength());
+							Response.Body = FIoBuffer(PartialRange, Request.Body);
+
+							// Try to limit how frequently we log warnings about this to avoid spam
+							const uint64 LogSpamLimit = 50; 
+							static uint64 LogCounter = 0;
+							UE_CLOG(LogCounter++ % LogSpamLimit == 0, LogIoStoreOnDemand, Warning, TEXT("A partial content request failed and the full data blob was downloaded, check if the CDN supports this"));
+						}
+						else
+						{
+							Response.Body = MoveTemp(Request.Body);
+						}
+					}
+					else if (Config.bResponseBodyOnError)
 					{
 						Response.Body = MoveTemp(Request.Body);
 					}
