Additions to dataflow rendering support
- Add transform support to dataflow rendering geometry groups
- Add capsule rendering support to rendering facade

diff --git a/Engine/Plugins/Experimental/Dataflow/Source/DataflowEditor/Private/Dataflow/DataflowConstructionScene.cpp b/Engine/Plugins/Experimental/Dataflow/Source/DataflowEditor/Private/Dataflow/DataflowConstructionScene.cpp
--- a/Engine/Plugins/Experimental/Dataflow/Source/DataflowEditor/Private/Dataflow/DataflowConstructionScene.cpp
+++ b/Engine/Plugins/Experimental/Dataflow/Source/DataflowEditor/Private/Dataflow/DataflowConstructionScene.cpp
@@ -210,6 +210,7 @@ void FDataflowConstructionScene::UpdateDynamicMeshComponents()
 							if (DynamicMesh.VertexCount())
 							{
 								const FString MeshName = RenderingFacade.GetGeometryName()[MeshIndex];
+								const FTransform& MeshTransform = RenderingFacade.GetGeometryTransform()[MeshIndex];
 
 								const TManagedArray<FString>& MaterialPaths = RenderingFacade.GetMaterialPaths();
 								const int32 MaterialStart = RenderingFacade.GetMaterialStart()[MeshIndex];
@@ -223,7 +224,7 @@ void FDataflowConstructionScene::UpdateDynamicMeshComponents()
 									Materials.Add(Material);
 								}
 
-								AddDynamicMeshComponent({Target, MeshIndex}, MeshName, MoveTemp(DynamicMesh), Materials);
+								AddDynamicMeshComponent({Target, MeshIndex}, MeshName, MoveTemp(DynamicMesh), Materials, MeshTransform);
 
 
 								for (int32 Idx = 0; Idx < VertexCount[MeshIndex]; ++Idx)
@@ -299,8 +300,9 @@ void FDataflowConstructionScene::UpdateDynamicMeshComponents()
 								if (DynamicMesh.VertexCount())
 								{
 									const FString MeshName = RenderingFacade.GetGeometryName()[MeshIndex];
+									const FTransform& MeshTransform = RenderingFacade.GetGeometryTransform()[MeshIndex];
 									const FString UniqueObjectName = MakeUniqueObjectName(RootSceneActor, UDataflowEditorCollectionComponent::StaticClass(), FName(MeshName)).ToString();
-									UDynamicMeshComponent* const NewDynamicMeshComponent = AddDynamicMeshComponent(WireframeDynamicMeshKey, UniqueObjectName, MoveTemp(DynamicMesh), {});
+									UDynamicMeshComponent* const NewDynamicMeshComponent = AddDynamicMeshComponent(WireframeDynamicMeshKey, UniqueObjectName, MoveTemp(DynamicMesh), {}, MeshTransform);
 									NewDynamicMeshComponent->SetVisibility(false);
 									MeshComponentsForWireframeRendering.Add(NewDynamicMeshComponent);
 								}
@@ -403,7 +405,7 @@ void FDataflowConstructionScene::ResetSceneComponents()
 	RemoveSceneObject(RootSceneActor);
 }
 
-TObjectPtr<UDynamicMeshComponent>& FDataflowConstructionScene::AddDynamicMeshComponent(FDataflowRenderKey InKey, const FString& MeshName, UE::Geometry::FDynamicMesh3&& DynamicMesh, const TArray<UMaterialInterface*>& MaterialSet)
+TObjectPtr<UDynamicMeshComponent>& FDataflowConstructionScene::AddDynamicMeshComponent(FDataflowRenderKey InKey, const FString& MeshName, UE::Geometry::FDynamicMesh3&& DynamicMesh, const TArray<UMaterialInterface*>& MaterialSet, const FTransform& InMeshTransform)
 {
 	// Don't use the MakeUniqueObjectName for the component, we need to keep the name aligned with the collection so selection will work in 
 	// other editors. 
@@ -452,7 +454,7 @@ TObjectPtr<UDynamicMeshComponent>& FDataflowConstructionScene::AddDynamicMeshCom
 		}
 	}
 
-	AddComponent(DynamicMeshComponent, DynamicMeshComponent->GetRelativeTransform());	
+	AddComponent(DynamicMeshComponent, InMeshTransform);	
 	DynamicMeshComponents.Emplace(InKey, DynamicMeshComponent);
 	AddSceneObject(DynamicMeshComponent, true);
 	return DynamicMeshComponents[InKey];
@@ -470,7 +472,7 @@ void FDataflowConstructionScene::AddWireframeMeshElementsVisualizer()
 			TObjectPtr<UMeshElementsVisualizer> WireframeDraw = NewObject<UMeshElementsVisualizer>(RootSceneActor);
 			WireframeElements.Add(DynamicMeshComponent, WireframeDraw);
 
-			WireframeDraw->CreateInWorld(GetWorld(), FTransform::Identity);
+			WireframeDraw->CreateInWorld(GetWorld(), DynamicMeshComponent->GetComponentToWorld());
 			checkf(WireframeDraw->Settings, TEXT("Expected UMeshElementsVisualizer::Settings to exist after CreateInWorld"));
 
 			WireframeDraw->Settings->DepthBias = 2.0;
diff --git a/Engine/Plugins/Experimental/Dataflow/Source/DataflowEditor/Public/Dataflow/DataflowConstructionScene.h b/Engine/Plugins/Experimental/Dataflow/Source/DataflowEditor/Public/Dataflow/DataflowConstructionScene.h
--- a/Engine/Plugins/Experimental/Dataflow/Source/DataflowEditor/Public/Dataflow/DataflowConstructionScene.h
+++ b/Engine/Plugins/Experimental/Dataflow/Source/DataflowEditor/Public/Dataflow/DataflowConstructionScene.h
@@ -99,7 +99,7 @@ private:
 	
 	/** Add a dynamic mesh component to the scene */
 	UE_API TObjectPtr<UDynamicMeshComponent>& AddDynamicMeshComponent(
-		FDataflowRenderKey InKey, const FString& MeshName, UE::Geometry::FDynamicMesh3&& DynamicMesh,  const TArray<UMaterialInterface*>& MaterialSet);
+		FDataflowRenderKey InKey, const FString& MeshName, UE::Geometry::FDynamicMesh3&& DynamicMesh,  const TArray<UMaterialInterface*>& MaterialSet, const FTransform& InMeshTransform = FTransform::Identity);
 
 	/** Remove the scene component from selection and destroy it */
 	UE_API void RemoveSceneComponent(USelection* SelectedComponents, UPrimitiveComponent* PrimitiveComponent);
diff --git a/Engine/Plugins/Experimental/Dataflow/Source/DataflowEnginePlugin/Private/Dataflow/DataflowEngineSceneProxy.cpp b/Engine/Plugins/Experimental/Dataflow/Source/DataflowEnginePlugin/Private/Dataflow/DataflowEngineSceneProxy.cpp
--- a/Engine/Plugins/Experimental/Dataflow/Source/DataflowEnginePlugin/Private/Dataflow/DataflowEngineSceneProxy.cpp
+++ b/Engine/Plugins/Experimental/Dataflow/Source/DataflowEnginePlugin/Private/Dataflow/DataflowEngineSceneProxy.cpp
@@ -176,6 +176,7 @@ void FDataflowEngineSceneProxy::CreateMeshRenderThreadResources(FRHICommandListB
 				MeshBatchData.FirstTriangleIndex = Facade.GetIndicesStart()[i] * 3;
 				MeshBatchData.NumTriangles = Facade.GetIndicesCount()[i];
 				MeshBatchData.GeomIndex = i;
+				MeshBatchData.MeshTransform = Facade.GetGeometryTransform()[i];
 			}
 		}
 
@@ -263,8 +264,10 @@ void FDataflowEngineSceneProxy::GetMeshDynamicMeshElements(int32 ViewIndex, FMes
 
 	for (const FDataflowTriangleSetMeshBatchData& MeshBatchData : MeshBatchDatas)
 	{
+		const FMatrix BatchTransform = GetLocalToWorld() * MeshBatchData.MeshTransform.ToMatrixWithScale();
+
 		FDynamicPrimitiveUniformBuffer& DynamicPrimitiveUniformBuffer = Collector.AllocateOneFrameResource<FDynamicPrimitiveUniformBuffer>();
-		DynamicPrimitiveUniformBuffer.Set(Collector.GetRHICommandList(), GetLocalToWorld(), GetLocalToWorld(), GetBounds(), GetLocalBounds(), true, false, AlwaysHasVelocity());
+		DynamicPrimitiveUniformBuffer.Set(Collector.GetRHICommandList(), BatchTransform, BatchTransform, GetBounds(), GetLocalBounds(), true, false, AlwaysHasVelocity());
 
 		FMeshBatch& Mesh = Collector.AllocateMesh();
 		Mesh.bWireframe = false;
diff --git a/Engine/Plugins/Experimental/Dataflow/Source/DataflowEnginePlugin/Private/Dataflow/DataflowEngineSceneProxy.h b/Engine/Plugins/Experimental/Dataflow/Source/DataflowEnginePlugin/Private/Dataflow/DataflowEngineSceneProxy.h
--- a/Engine/Plugins/Experimental/Dataflow/Source/DataflowEnginePlugin/Private/Dataflow/DataflowEngineSceneProxy.h
+++ b/Engine/Plugins/Experimental/Dataflow/Source/DataflowEnginePlugin/Private/Dataflow/DataflowEngineSceneProxy.h
@@ -19,6 +19,7 @@ struct FDataflowTriangleSetMeshBatchData
 	int32 MinVertexIndex = INDEX_NONE;
 	int32 MaxVertexIndex = INDEX_NONE;
 	int32 GeomIndex = INDEX_NONE;
+	FTransform MeshTransform = FTransform::Identity;
 };
 
 struct FDataflowVertexBatchData
diff --git a/Engine/Source/Runtime/Experimental/Chaos/Private/GeometryCollection/Facades/CollectionRenderingFacade.cpp b/Engine/Source/Runtime/Experimental/Chaos/Private/GeometryCollection/Facades/CollectionRenderingFacade.cpp
--- a/Engine/Source/Runtime/Experimental/Chaos/Private/GeometryCollection/Facades/CollectionRenderingFacade.cpp
+++ b/Engine/Source/Runtime/Experimental/Chaos/Private/GeometryCollection/Facades/CollectionRenderingFacade.cpp
@@ -6,6 +6,7 @@
 
 #include "GeometryCollection/Facades/CollectionRenderingFacade.h"
 #include "GeometryCollection/GeometryCollection.h"
+#include "Generators/CapsuleGenerator.h"
 
 namespace GeometryCollection::Facades
 {
@@ -25,6 +26,7 @@ namespace GeometryCollection::Facades
 		, TriangleSectionAttribute(InCollection, "Sections", FGeometryCollection::MaterialGroup)
 		, MaterialPathAttribute(InCollection, "MaterialPath", FGeometryCollection::MaterialGroup)
 		, GeometryNameAttribute(InCollection, "Name", FGeometryCollection::GeometryGroup)
+		, GeometryTransformAttribute(InCollection, "Transform", FGeometryCollection::GeometryGroup)
 		, GeometryHitProxyIndexAttribute(InCollection, "HitIndex", FGeometryCollection::GeometryGroup)
 		, VertexStartAttribute(InCollection, "VertexStart", FGeometryCollection::GeometryGroup)
 		, VertexCountAttribute(InCollection, "VertexCount", FGeometryCollection::GeometryGroup)
@@ -51,6 +53,7 @@ namespace GeometryCollection::Facades
 		, TriangleSectionAttribute(InCollection, "Sections", FGeometryCollection::MaterialGroup)
 		, MaterialPathAttribute(InCollection, "MaterialPath", FGeometryCollection::MaterialGroup)
 		, GeometryNameAttribute(InCollection, "Name", FGeometryCollection::GeometryGroup)
+		, GeometryTransformAttribute(InCollection, "Transform", FGeometryCollection::GeometryGroup)
 		, GeometryHitProxyIndexAttribute(InCollection, "HitIndex", FGeometryCollection::GeometryGroup)
 		, VertexStartAttribute(InCollection, "VertexStart", FGeometryCollection::GeometryGroup, FGeometryCollection::VerticesGroup)
 		, VertexCountAttribute(InCollection, "VertexCount", FGeometryCollection::GeometryGroup)
@@ -80,6 +83,7 @@ namespace GeometryCollection::Facades
 		TriangleSectionAttribute.Add();
 		MaterialPathAttribute.Add();
 		GeometryNameAttribute.Add();
+		GeometryTransformAttribute.Add();
 		GeometryHitProxyIndexAttribute.Add();
 		VertexStartAttribute.Add();
 		VertexCountAttribute.Add();
@@ -503,7 +507,7 @@ namespace GeometryCollection::Facades
 	}
 
 
-	int32 FRenderingFacade::StartGeometryGroup(FString InName)
+	int32 FRenderingFacade::StartGeometryGroup(FString InName, const FTransform& InTm)
 	{
 		check(!IsConst());
 
@@ -512,7 +516,7 @@ namespace GeometryCollection::Facades
 		{
 			GeomIndex = GeometryNameAttribute.AddElements(1);
 			GeometryNameAttribute.Modify()[GeomIndex] = InName;
-
+			GeometryTransformAttribute.Modify()[GeomIndex] = InTm;
 			VertexStartAttribute.Modify()[GeomIndex] = VertexAttribute.Num();
 			VertexCountAttribute.Modify()[GeomIndex] = 0;
 			IndicesStartAttribute.Modify()[GeomIndex] = IndicesAttribute.Num();
@@ -566,6 +570,18 @@ namespace GeometryCollection::Facades
 		}
 	}
 
+	void FRenderingFacade::SetGroupTransform(int32 InGeometryGroupIndex, const FTransform& InTm)
+	{
+		check(!IsConst());
+		if(IsValid())
+		{
+			// Only support changing the currently open group
+			checkf(GeometryNameAttribute.Num() - 1 == InGeometryGroupIndex, TEXT("Attempted to modify a previously closed geometry group transform"));
+
+			GeometryTransformAttribute.Modify()[InGeometryGroupIndex] = InTm;
+		}
+	}
+
 	FRenderingFacade::FStringIntMap FRenderingFacade::GetGeometryNameToIndexMap() const
 	{
 		FStringIntMap Map;
@@ -588,6 +604,33 @@ namespace GeometryCollection::Facades
 		return RetCount;
 	}
 
+	void FRenderingFacade::AddCapsule(const float Length, const float Radius)
+	{
+		UE::Geometry::FCapsuleGenerator Generator;
+		Generator.Radius = Radius;
+		Generator.SegmentLength = Length;
+		Generator.Generate();
+
+		TArray<FVector3f> NarrowVerts;
+		TArray<FIntVector> Triangles;
+		TArray<FLinearColor> Colours;
+		NarrowVerts.Reserve(Generator.Vertices.Num());
+		NarrowVerts.Reserve(Generator.Vertices.Num());
+		Triangles.Reserve(Generator.Triangles.Num());
+
+		for(const FVector3d& WideVertex : Generator.Vertices)
+		{
+			NarrowVerts.Add((FVector3f)WideVertex);
+			Colours.Add(FLinearColor::White);
+		}
+
+		for(const UE::Geometry::FIndex3i& TriIndices : Generator.Triangles)
+		{
+			Triangles.Add({TriIndices.A, TriIndices.B, TriIndices.C});
+		}
+
+		AddSurface(MoveTemp(NarrowVerts), MoveTemp(Triangles), MoveTemp(Generator.Normals), MoveTemp(Colours));
+	}
 
 }; // GeometryCollection::Facades
 
diff --git a/Engine/Source/Runtime/Experimental/Chaos/Public/GeometryCollection/Facades/CollectionRenderingFacade.h b/Engine/Source/Runtime/Experimental/Chaos/Public/GeometryCollection/Facades/CollectionRenderingFacade.h
--- a/Engine/Source/Runtime/Experimental/Chaos/Public/GeometryCollection/Facades/CollectionRenderingFacade.h
+++ b/Engine/Source/Runtime/Experimental/Chaos/Public/GeometryCollection/Facades/CollectionRenderingFacade.h
@@ -72,6 +72,9 @@ namespace GeometryCollection::Facades
 		/**Add spheres to the rendering view.*/
 		CHAOS_API void AddSpheres(const TArray<FSphere>& InSpheres, const FLinearColor& InColor);
 
+		/** Add a capsule to the rendering view */
+		CHAOS_API void AddCapsule(const float Length, const float Radius);
+
 		/** Add a surface to the rendering view.*/
 		CHAOS_API void AddSurface(TArray<FVector3f>&& InVertices, TArray<FIntVector>&& InIndices, TArray<FVector3f>&& InNormals, TArray<FLinearColor>&& InColors);
 
@@ -137,16 +140,22 @@ namespace GeometryCollection::Facades
 		//
 
 		/** Geometry Group Start : */
-		CHAOS_API int32 StartGeometryGroup(FString InName);
+		CHAOS_API int32 StartGeometryGroup(FString InName, const FTransform& InTm = FTransform::Identity);
 
 		/** Geometry Group End : */
 		CHAOS_API void EndGeometryGroup(int32 InGeometryGroupIndex);
 
+		/** Modify the transform of a pre-declared group */
+		CHAOS_API void SetGroupTransform(int32 InGeometryGroupIndex, const FTransform& InTm);
+
 		int32 NumGeometry() const { return GeometryNameAttribute.Num(); }
 
 		/** GeometryNameAttribute */
 		const TManagedArray< FString >& GetGeometryName() const { return GeometryNameAttribute.Get(); }
 
+		/** GeometryTransformAttribute */
+		const TManagedArray<FTransform>& GetGeometryTransform() const { return GeometryTransformAttribute.Get(); }
+
 		/** HitProxyIDAttribute */
 		const TManagedArray< int32 >& GetGeometryHitProxyIndex() const { return GeometryHitProxyIndexAttribute.Get(); }
 		      TManagedArray< int32 >& ModifyGeometryHitProxyIndex() {check(!IsConst());return GeometryHitProxyIndexAttribute.Modify(); }
@@ -204,6 +213,7 @@ namespace GeometryCollection::Facades
 		TManagedArrayAccessor<FString> MaterialPathAttribute;
 
 		TManagedArrayAccessor<FString> GeometryNameAttribute;
+		TManagedArrayAccessor<FTransform> GeometryTransformAttribute;
 		TManagedArrayAccessor<int32> GeometryHitProxyIndexAttribute;
 		TManagedArrayAccessor<int32> VertexStartAttribute;
 		TManagedArrayAccessor<int32> VertexCountAttribute;
