Make OnDemandInstallCache Status codes analytics friendly.
- Add FNamespacedIoStatusBuilderBuilder which can be used to automatically add a namespace, line, and column to unqiuely identify an IoStatus.
- All returned status's in the install cache now use a FNamespacedIoStatusBuilderBuilder
- Add GetAnalyticErrorCode to generate a unique, concise status string that the analytics backend can aggregate and is readable in graphs.
- FIoStatusBuilder now uses a string builder and can be implicitly cast to TIoStatusOr

diff --git a/Engine/Source/Runtime/Core/Private/IO/IoStatus.cpp b/Engine/Source/Runtime/Core/Private/IO/IoStatus.cpp
--- a/Engine/Source/Runtime/Core/Private/IO/IoStatus.cpp
+++ b/Engine/Source/Runtime/Core/Private/IO/IoStatus.cpp
@@ -97,10 +97,12 @@ FIoStatusBuilder::FIoStatusBuilder(EIoErrorCode InStatusCode)
 {
 }
 
-FIoStatusBuilder::FIoStatusBuilder(const FIoStatus& InStatus, FStringView String)
-:	StatusCode(InStatus.ErrorCode)
+FIoStatusBuilder::FIoStatusBuilder(const TCHAR* NamespaceStr, EIoErrorCode StatusCode, const UE::FSourceLocation& SourceLoc)
+	: FIoStatusBuilder(StatusCode)
 {
-	Message.Append(String.GetData(), String.Len());
+	Message
+		<< TEXT('<') << NamespaceStr << TEXT('-')
+		<< SourceLoc.GetLine() << TEXT(',') << SourceLoc.GetColumn() << TEXT("> ");
 }
 
 FIoStatusBuilder::~FIoStatusBuilder()
@@ -114,13 +116,7 @@ FIoStatusBuilder::operator FIoStatus()
 
 FIoStatusBuilder& FIoStatusBuilder::operator<<(FStringView String)
 {
-	Message.Append(String.GetData(), String.Len());
+	Message.Append(String);
 
 	return *this;
 }
-
-FIoStatusBuilder operator<<(const FIoStatus& Status, FStringView String)
-{ 
-	return FIoStatusBuilder(Status, String);
-}
-
diff --git a/Engine/Source/Runtime/Core/Public/IO/IoStatus.h b/Engine/Source/Runtime/Core/Public/IO/IoStatus.h
--- a/Engine/Source/Runtime/Core/Public/IO/IoStatus.h
+++ b/Engine/Source/Runtime/Core/Public/IO/IoStatus.h
@@ -4,6 +4,8 @@
 
 #include "Containers/StringView.h"
 #include "Containers/UnrealString.h"
+#include "Misc/StringBuilder.h"
+#include "Misc/SourceLocation.h"
 
 template <typename CharType> class TStringBuilderBase;
 
@@ -82,27 +84,6 @@ private:
 	friend class FIoStatusBuilder;
 };
 
-///////////////////////////////////////////////////////////////////////////////
-
-/**
- * Helper to make it easier to generate meaningful error messages.
- */
-class FIoStatusBuilder
-{
-	EIoErrorCode		StatusCode;
-	FString				Message;
-public:
-	CORE_API explicit	FIoStatusBuilder(EIoErrorCode StatusCode);
-	CORE_API			FIoStatusBuilder(const FIoStatus& InStatus, FStringView String);
-	CORE_API			~FIoStatusBuilder();
-
-	CORE_API			operator FIoStatus();
-
-	CORE_API FIoStatusBuilder& operator<<(FStringView String);
-};
-
-CORE_API FIoStatusBuilder operator<<(const FIoStatus& Status, FStringView String);
-
 /**
  * Optional I/O result or error status.
  */
@@ -112,7 +93,7 @@ class TIoStatusOr
 	template<typename U> friend class TIoStatusOr;
 
 public:
-	TIoStatusOr() : StatusValue(FIoStatus::Unknown) { }
+	TIoStatusOr() : StatusValue(FIoStatus::Unknown) {}
 	TIoStatusOr(const TIoStatusOr& Other);
 	TIoStatusOr(TIoStatusOr&& Other);
 
@@ -137,10 +118,10 @@ public:
 	template<typename U>
 	TIoStatusOr<T>& operator=(const TIoStatusOr<U>& Other);
 
-	const FIoStatus&	Status() const;
+	const FIoStatus& Status() const;
 	bool				IsOk() const;
 
-	const T&			ValueOrDie();
+	const T& ValueOrDie();
 	T					ConsumeValueOrDie();
 
 	void				Reset();
@@ -150,6 +131,53 @@ private:
 	TTypeCompatibleBytes<T>	Value;
 };
 
+///////////////////////////////////////////////////////////////////////////////
+
+/**
+ * Helper to make it easier to generate meaningful error messages.
+ */
+class FIoStatusBuilder
+{
+protected:
+	EIoErrorCode		StatusCode;
+	TStringBuilder<FIoStatus::MaxErrorMessageLength> Message;
+public:
+	CORE_API explicit	FIoStatusBuilder(EIoErrorCode StatusCode);
+
+	CORE_API			FIoStatusBuilder(const TCHAR* NamespaceStr, EIoErrorCode StatusCode, const UE::FSourceLocation& SourceLoc = UE::FSourceLocation::Current());
+
+	CORE_API			~FIoStatusBuilder();
+
+	CORE_API			operator FIoStatus();
+
+	template<typename T>
+	operator TIoStatusOr<T>()
+	{
+		return FIoStatus(StatusCode, Message);
+	}
+
+	CORE_API FIoStatusBuilder& operator<<(FStringView String);
+};
+
+/**
+ * Creates an FIoStatusBuilder that generates unique messages by prepending a namespace and source location.
+ * Yes, this is in fact a builder builder - deal with it!
+ */
+class FNamespacedIoStatusBuilderBuilder
+{
+private:
+	const TCHAR* NamespaceStr = nullptr;
+public:
+	FNamespacedIoStatusBuilderBuilder(const TCHAR* InNamespaceStr) : NamespaceStr(InNamespaceStr)
+	{
+	}
+
+	FIoStatusBuilder operator()(EIoErrorCode StatusCode, const UE::FSourceLocation& SourceLoc = UE::FSourceLocation::Current()) const
+	{
+		return FIoStatusBuilder(NamespaceStr, StatusCode, SourceLoc);
+	}
+};
+
 CORE_API void StatusOrCrash(const FIoStatus& Status);
 
 template<typename T>
diff --git a/Engine/Source/Runtime/Experimental/IoStore/OnDemand/Private/OnDemandInstallCache.cpp b/Engine/Source/Runtime/Experimental/IoStore/OnDemand/Private/OnDemandInstallCache.cpp
--- a/Engine/Source/Runtime/Experimental/IoStore/OnDemand/Private/OnDemandInstallCache.cpp
+++ b/Engine/Source/Runtime/Experimental/IoStore/OnDemand/Private/OnDemandInstallCache.cpp
@@ -50,6 +50,8 @@
 
 namespace UE::IoStore
 {
+///////////////////////////////////////////////////////////////////////////////
+FNamespacedIoStatusBuilderBuilder InstallCacheStatusBuilder(TEXT("IADCache"));
 
 ///////////////////////////////////////////////////////////////////////////////
 namespace CVars
@@ -219,7 +221,7 @@ FIoStatus FCas::Initialize(FStringView Directory, bool bDeleteExisting)
 
 		if (Ifm.DeleteDirectory(Path.ToString(), bRequireExists, bTree) == false)
 		{
-			return FIoStatusBuilder(EIoErrorCode::WriteError)
+			return InstallCacheStatusBuilder(EIoErrorCode::WriteError)
 				<< TEXT("Failed to delete CAS blocks directory '")
 				<< Path.ToString()
 				<< TEXT("'");
@@ -231,7 +233,7 @@ FIoStatus FCas::Initialize(FStringView Directory, bool bDeleteExisting)
 		const bool bTree = true;
 		if (Ifm.MakeDirectory(Path.ToString(), bTree) == false)
 		{
-			FIoStatus Status = FIoStatusBuilder(EIoErrorCode::WriteError)
+			FIoStatus Status = InstallCacheStatusBuilder(EIoErrorCode::WriteError)
 				<< TEXT("Failed to create directory '")
 				<< Path.ToString()
 				<< TEXT("'");
@@ -239,7 +241,7 @@ FIoStatus FCas::Initialize(FStringView Directory, bool bDeleteExisting)
 		}
 	}
 
-	return EIoErrorCode::Ok;
+	return InstallCacheStatusBuilder(EIoErrorCode::Ok);
 };
 
 FCasLocation FCas::FindChunk(const FIoHash& Hash) const
@@ -309,7 +311,7 @@ FIoStatus FCas::DeleteBlock(FCasBlockId BlockId, TArray<FCasAddr>& OutAddrs)
 
 			if (FPlatformTime::ToMilliseconds(FPlatformTime::Cycles() - StartTimeCycles) > DeleteBlockMaxWaitTimeMs)
 			{
-				return FIoStatusBuilder(EIoErrorCode::Timeout)
+				return InstallCacheStatusBuilder(EIoErrorCode::Timeout)
 					<< TEXT("Timed out waiting for pending read(s) while deleting CAS block"); 
 			}
 
@@ -325,7 +327,7 @@ FIoStatus FCas::DeleteBlock(FCasBlockId BlockId, TArray<FCasAddr>& OutAddrs)
 	UE_LOG(LogIoStoreOnDemand, Log, TEXT("Deleting CAS block '%s'"), *Filename);
 	if (Ipf.DeleteFile(*Filename) == false)
 	{
-		return FIoStatusBuilder(EIoErrorCode::DeleteError) << TEXT("Failed to delete CAS block");
+		return InstallCacheStatusBuilder(EIoErrorCode::DeleteError) << TEXT("Failed to delete CAS block");
 	}
 
 	BlockIds.Remove(BlockId);
@@ -340,7 +342,7 @@ FIoStatus FCas::DeleteBlock(FCasBlockId BlockId, TArray<FCasAddr>& OutAddrs)
 		}
 	}
 
-	return FIoStatus::Ok;
+	return InstallCacheStatusBuilder(EIoErrorCode::Ok);
 }
 
 FString FCas::GetBlockFilename(FCasBlockId BlockId) const
@@ -546,7 +548,7 @@ FIoStatus FCas::Verify(TArray<FCasAddr>& OutAddrs)
 	FCasBlockInfoMap	BlockInfo;
 	const uint64		TotalSize = GetBlockInfo(BlockInfo);
 	uint64				TotalVerifiedBytes = 0;
-	FIoStatus			Status = FIoStatus::Ok;
+	FIoStatus			Status = InstallCacheStatusBuilder(EIoErrorCode::Ok);
 
 	for (auto BlockIt = BlockIds.CreateIterator(); BlockIt; ++BlockIt)
 	{
@@ -562,7 +564,7 @@ FIoStatus FCas::Verify(TArray<FCasAddr>& OutAddrs)
 
 		LastAccess.Remove(BlockId);
 		BlockIt.RemoveCurrent();
-		Status = EIoErrorCode::NotFound;
+		Status = InstallCacheStatusBuilder(EIoErrorCode::NotFound);
 	}
 
 	UE_LOG(LogIoStoreOnDemand, Log, TEXT("Verified %d CAS blocks of total %.2lf MiB"),
@@ -594,7 +596,7 @@ FIoStatus FCas::Verify(TArray<FCasAddr>& OutAddrs)
 			OutAddrs.Add(It->Key);
 			It.RemoveCurrent();
 
-			Status = EIoErrorCode::NotFound;
+			Status = InstallCacheStatusBuilder(EIoErrorCode::NotFound);
 		}
 	}
 
@@ -763,19 +765,19 @@ FIoStatus FCasJournal::Replay(const FString& JournalFile, FEntryHandler&& Handle
 
 	if (Ipf.FileExists(*JournalFile) == false)
 	{
-		return EIoErrorCode::NotFound;
+		return InstallCacheStatusBuilder(EIoErrorCode::NotFound);
 	}
 
 	TUniquePtr<IFileHandle> FileHandle(Ipf.OpenRead(*JournalFile));
 	if (FileHandle.IsValid() == false)
 	{
-		return EIoErrorCode::FileNotOpen;
+		return InstallCacheStatusBuilder(EIoErrorCode::FileNotOpen);
 	}
 
 	FHeader Header;
 	if ((FileHandle->Read(reinterpret_cast<uint8*>(&Header), sizeof(FHeader)) == false) || (Header.IsValid() == false))
 	{
-		return FIoStatusBuilder(EIoErrorCode::ReadError)
+		return InstallCacheStatusBuilder(EIoErrorCode::ReadError)
 			<< TEXT("Failed to validate journal header '")
 			<< JournalFile
 			<< TEXT("'");
@@ -786,31 +788,31 @@ FIoStatus FCasJournal::Replay(const FString& JournalFile, FEntryHandler&& Handle
 
 	if (EntryCount < 0)
 	{
-		return EIoErrorCode::ReadError;
+		return InstallCacheStatusBuilder(EIoErrorCode::ReadError);
 	}
 
 	if (EntryCount == 0)
 	{
-		return EIoErrorCode::Ok;
+		return InstallCacheStatusBuilder(EIoErrorCode::Ok);
 	}
 
 	const int64 FooterPos = FileSize - sizeof(FFooter);
 	if (FooterPos < 0)
 	{
-		return FIoStatusBuilder(EIoErrorCode::CorruptToc)
+		return InstallCacheStatusBuilder(EIoErrorCode::CorruptToc)
 			<< TEXT("Invalid journal footer");
 	}
 
 	const int64 EntriesPos = FileHandle->Tell();
 	if (FileHandle->Seek(FooterPos) == false)
 	{
-		return EIoErrorCode::ReadError;
+		return InstallCacheStatusBuilder(EIoErrorCode::ReadError);
 	}
 
 	FFooter Footer;
 	if ((FileHandle->Read(reinterpret_cast<uint8*>(&Footer), sizeof(FFooter)) == false) || (Footer.IsValid() == false))
 	{
-		return FIoStatusBuilder(EIoErrorCode::ReadError)
+		return InstallCacheStatusBuilder(EIoErrorCode::ReadError)
 			<< TEXT("Failed to validate journal footer '")
 			<< JournalFile
 			<< TEXT("'");
@@ -818,7 +820,7 @@ FIoStatus FCasJournal::Replay(const FString& JournalFile, FEntryHandler&& Handle
 
 	if (FileHandle->Seek(EntriesPos) == false)
 	{
-		return EIoErrorCode::ReadError;
+		return InstallCacheStatusBuilder(EIoErrorCode::ReadError);
 	}
 
 	TArray<FEntry> Entries;
@@ -826,7 +828,7 @@ FIoStatus FCasJournal::Replay(const FString& JournalFile, FEntryHandler&& Handle
 
 	if (FileHandle->Read(reinterpret_cast<uint8*>(Entries.GetData()), sizeof(FEntry) * EntryCount) == false)
 	{
-		return EIoErrorCode::ReadError;
+		return InstallCacheStatusBuilder(EIoErrorCode::ReadError);
 	}
 
 	UE_LOG(LogIoStoreOnDemand, Log, TEXT("Replaying %" INT64_FMT " CAS journal entries of total %.2lf KiB from '%s'"),
@@ -840,13 +842,13 @@ FIoStatus FCasJournal::Replay(const FString& JournalFile, FEntryHandler&& Handle
 			UE_LOG(LogIoStoreOnDemand, Warning, TEXT("Found critical error entry '%s' (%d) in journal '%s'"),
 				GetErrorText(Error.ErrorCode), Error.ErrorCode, *JournalFile);
 
-			return FIoStatus(EIoErrorCode::ReadError, FString(GetErrorText(Error.ErrorCode)));
+			return InstallCacheStatusBuilder(EIoErrorCode::ReadError) << GetErrorText(Error.ErrorCode);
 		}
 
 		Handler(Entry);
 	}
 
-	return EIoErrorCode::Ok;
+	return InstallCacheStatusBuilder(EIoErrorCode::Ok);
 }
 
 FIoStatus FCasJournal::Create(const FString& JournalFile)
@@ -857,7 +859,7 @@ FIoStatus FCasJournal::Create(const FString& JournalFile)
 	TUniquePtr<IFileHandle> FileHandle(Ipf.OpenWrite(*JournalFile));
 	if (FileHandle.IsValid() == false)
 	{
-		return EIoErrorCode::FileNotOpen;
+		return InstallCacheStatusBuilder(EIoErrorCode::FileNotOpen);
 	}
 
 	FHeader Header;
@@ -866,17 +868,17 @@ FIoStatus FCasJournal::Create(const FString& JournalFile)
 
 	if (FileHandle->Write(reinterpret_cast<uint8*>(&Header), sizeof(FHeader)) == false)
 	{
-		return EIoErrorCode::WriteError;
+		return InstallCacheStatusBuilder(EIoErrorCode::WriteError);
 	}
 
 	FFooter Footer;
 	FMemory::Memcpy(&Footer.Magic, &FFooter::MagicSequence, sizeof(FFooter::MagicSequence));
 	if (FileHandle->Write(reinterpret_cast<uint8*>(&Footer), sizeof(FFooter)) == false)
 	{
-		return EIoErrorCode::WriteError;
+		return InstallCacheStatusBuilder(EIoErrorCode::WriteError);
 	}
 
-	return EIoErrorCode::Ok;
+	return InstallCacheStatusBuilder(EIoErrorCode::Ok);
 }
 
 FCasJournal::FTransaction FCasJournal::Begin(FString&& JournalFile)
@@ -891,7 +893,7 @@ FIoStatus FCasJournal::Commit(FTransaction&& Transaction)
 {
 	if (Transaction.Entries.IsEmpty())
 	{
-		return EIoErrorCode::Ok;
+		return InstallCacheStatusBuilder(EIoErrorCode::Ok);
 	}
 
 	IPlatformFile&	Ipf = FPlatformFileManager::Get().GetPlatformFile();
@@ -904,7 +906,7 @@ FIoStatus FCasJournal::Commit(FTransaction&& Transaction)
 		if (FileSize < sizeof(FHeader))
 		{
 			FOnDemandInstallCacheStats::OnJournalCommit(EIoErrorCode::ReadError, 0);
-			return FIoStatusBuilder(EIoErrorCode::FileOpenFailed)
+			return InstallCacheStatusBuilder(EIoErrorCode::FileOpenFailed)
 				<< TEXT("Failed to validate CAS journal file");
 		}
 
@@ -912,7 +914,7 @@ FIoStatus FCasJournal::Commit(FTransaction&& Transaction)
 		if ((FileHandle->Read(reinterpret_cast<uint8*>(&Header), sizeof(FHeader)) == false) || (Header.IsValid() == false))
 		{
 			FOnDemandInstallCacheStats::OnJournalCommit(EIoErrorCode::SignatureError, 0);
-			return FIoStatusBuilder(EIoErrorCode::ReadError)
+			return InstallCacheStatusBuilder(EIoErrorCode::ReadError)
 				<< TEXT("Failed to validate CAS journal header");
 		}
 
@@ -920,7 +922,7 @@ FIoStatus FCasJournal::Commit(FTransaction&& Transaction)
 		if (FileHandle->Seek(FooterPos) == false)
 		{
 			FOnDemandInstallCacheStats::OnJournalCommit(EIoErrorCode::SignatureError, 0);
-			return FIoStatusBuilder(EIoErrorCode::ReadError)
+			return InstallCacheStatusBuilder(EIoErrorCode::ReadError)
 				<< TEXT("Failed to validate CAS journal footer");
 		}
 
@@ -928,7 +930,7 @@ FIoStatus FCasJournal::Commit(FTransaction&& Transaction)
 		if ((FileHandle->Read(reinterpret_cast<uint8*>(&Footer), sizeof(FFooter)) == false) || (Footer.IsValid() == false))
 		{
 			FOnDemandInstallCacheStats::OnJournalCommit(EIoErrorCode::SignatureError, 0);
-			return FIoStatusBuilder(EIoErrorCode::ReadError)
+			return InstallCacheStatusBuilder(EIoErrorCode::ReadError)
 				<< TEXT("Failed to validate CAS journal footer");
 		}
 	}
@@ -943,7 +945,7 @@ FIoStatus FCasJournal::Commit(FTransaction&& Transaction)
 		if ((EntriesPos < 0) || (FileHandle->Seek(EntriesPos) == false))
 		{
 			FOnDemandInstallCacheStats::OnJournalCommit(EIoErrorCode::WriteError, 0);
-			return FIoStatusBuilder(EIoErrorCode::FileOpenFailed)
+			return InstallCacheStatusBuilder(EIoErrorCode::FileOpenFailed)
 				<< TEXT("Failed to open CAS journal");
 		}
 
@@ -953,7 +955,7 @@ FIoStatus FCasJournal::Commit(FTransaction&& Transaction)
 			TotalEntrySize) == false)
 		{
 			FOnDemandInstallCacheStats::OnJournalCommit(EIoErrorCode::WriteError, 0);
-			return FIoStatusBuilder(EIoErrorCode::WriteError)
+			return InstallCacheStatusBuilder(EIoErrorCode::WriteError)
 				<< TEXT("Failed to write CAS journal entries");
 		}
 
@@ -962,14 +964,14 @@ FIoStatus FCasJournal::Commit(FTransaction&& Transaction)
 		if (FileHandle->Write(reinterpret_cast<uint8*>(&Footer), sizeof(FFooter)) == false)
 		{
 			FOnDemandInstallCacheStats::OnJournalCommit(EIoErrorCode::WriteError, 0);
-			return FIoStatusBuilder(EIoErrorCode::WriteError)
+			return InstallCacheStatusBuilder(EIoErrorCode::WriteError)
 				<< TEXT("Failed to write CAS journal footer");
 		}
 
 		if (FileHandle->Flush() == false)
 		{
 			FOnDemandInstallCacheStats::OnJournalCommit(EIoErrorCode::WriteError, 0);
-			return FIoStatusBuilder(EIoErrorCode::WriteError)
+			return InstallCacheStatusBuilder(EIoErrorCode::WriteError)
 				<< TEXT("Failed to flush CAS journal");
 		}
 
@@ -978,7 +980,7 @@ FIoStatus FCasJournal::Commit(FTransaction&& Transaction)
 
 		FOnDemandInstallCacheStats::OnJournalCommit(EIoErrorCode::Ok, TotalEntrySize);
 
-		return EIoErrorCode::Ok;
+		return InstallCacheStatusBuilder(EIoErrorCode::Ok);
 	}
 }
 
@@ -1195,7 +1197,7 @@ TIoStatusOr<int64> FCasSnapshot::Save(const FCasSnapshot& Snapshot, const FStrin
 	TUniquePtr<FArchive> Ar(Ifm.CreateFileWriter(*TmpSnapshotFile));
 	if (Ar.IsValid() == false)
 	{
-		return FIoStatus(EIoErrorCode::FileNotOpen);
+		return InstallCacheStatusBuilder(EIoErrorCode::FileNotOpen);
 	}
 
 	FHeader Header;
@@ -1207,7 +1209,7 @@ TIoStatusOr<int64> FCasSnapshot::Save(const FCasSnapshot& Snapshot, const FStrin
 	{
 		Ar.Reset();
 		Ifm.Delete(*TmpSnapshotFile);
-		return FIoStatus(EIoErrorCode::WriteError);
+		return InstallCacheStatusBuilder(EIoErrorCode::WriteError);
 	}
 
 	FCasSnapshot& NonConst = *const_cast<FCasSnapshot*>(&Snapshot);
@@ -1219,7 +1221,7 @@ TIoStatusOr<int64> FCasSnapshot::Save(const FCasSnapshot& Snapshot, const FStrin
 	{
 		Ar.Reset();
 		Ifm.Delete(*TmpSnapshotFile);
-		return FIoStatus(EIoErrorCode::WriteError);
+		return InstallCacheStatusBuilder(EIoErrorCode::WriteError);
 	}
 
 	FFooter Footer;
@@ -1229,7 +1231,7 @@ TIoStatusOr<int64> FCasSnapshot::Save(const FCasSnapshot& Snapshot, const FStrin
 	{
 		Ar.Reset();
 		Ifm.Delete(*TmpSnapshotFile);
-		return FIoStatus(EIoErrorCode::WriteError);
+		return InstallCacheStatusBuilder(EIoErrorCode::WriteError);
 	}
 
 	const int64 FileSize = Ar->TotalSize();
@@ -1237,13 +1239,13 @@ TIoStatusOr<int64> FCasSnapshot::Save(const FCasSnapshot& Snapshot, const FStrin
 	{
 		Ar.Reset();
 		Ifm.Delete(*TmpSnapshotFile);
-		return FIoStatus(EIoErrorCode::WriteError);
+		return InstallCacheStatusBuilder(EIoErrorCode::WriteError);
 	}
 
 	if (Ifm.Move(*SnapshotFile, *TmpSnapshotFile) == false)
 	{
 		Ifm.Delete(*TmpSnapshotFile);
-		return FIoStatus(EIoErrorCode::WriteError);
+		return InstallCacheStatusBuilder(EIoErrorCode::WriteError);
 	}
 
 	return FileSize; 
@@ -1256,14 +1258,14 @@ TIoStatusOr<FCasSnapshot> FCasSnapshot::Load(const FString& SnapshotFile, int64*
 	TUniquePtr<FArchive> Ar(Ifm.CreateFileReader(*SnapshotFile));
 	if (Ar.IsValid() == false)
 	{
-		return FIoStatus(EIoErrorCode::NotFound);
+		return InstallCacheStatusBuilder(EIoErrorCode::NotFound);
 	}
 
 	FHeader Header;
 	Ar->Serialize(reinterpret_cast<uint8*>(&Header), sizeof(FHeader));
 	if (Ar->IsError() || Header.IsValid() == false)
 	{
-		FIoStatus Status = FIoStatusBuilder(EIoErrorCode::ReadError)
+		FIoStatus Status = InstallCacheStatusBuilder(EIoErrorCode::ReadError)
 			<< TEXT("Failed to validate snapshot header '")
 			<< SnapshotFile 
 			<< TEXT("'");
@@ -1279,7 +1281,7 @@ TIoStatusOr<FCasSnapshot> FCasSnapshot::Load(const FString& SnapshotFile, int64*
 	Ar->Serialize(reinterpret_cast<uint8*>(&Footer), sizeof(FFooter));
 	if (Ar->IsError() || Footer.IsValid() == false)
 	{
-		FIoStatus Status = FIoStatusBuilder(EIoErrorCode::ReadError)
+		FIoStatus Status = InstallCacheStatusBuilder(EIoErrorCode::ReadError)
 			<< TEXT("Failed to validate snapshot footer '")
 			<< SnapshotFile 
 			<< TEXT("'");
@@ -1300,7 +1302,7 @@ TIoStatusOr<int64> FCasSnapshot::TryCreateAndResetJournal(const FString& Snapsho
 	const int64 JournalFileSize = Ifm.FileSize(*JournalFile);
 	if (JournalFileSize < 0)
 	{
-		return FIoStatus(EIoErrorCode::NotFound);
+		return InstallCacheStatusBuilder(EIoErrorCode::NotFound);
 	}
 
 	// Load the snapshot from the journal
@@ -1341,7 +1343,7 @@ TIoStatusOr<int64> FCasSnapshot::TryCreateAndResetJournal(const FString& Snapsho
 			UE_LOG(LogIoStoreOnDemand, Error, TEXT("Failed to delete CAS snapshot '%s'"), *SnapshotFile); 
 		}
 
-		return FIoStatus(EIoErrorCode::WriteError);
+		return InstallCacheStatusBuilder(EIoErrorCode::WriteError);
 	}
 
 	return SnapshotSize;
@@ -1791,12 +1793,12 @@ TIoStatusOr<uint64> FOnDemandInstallCache::GetSizeForChunk(const FIoChunkId& Chu
 		return ChunkInfo.RawSize();
 	}
 
-	return FIoStatus(ErrorCode);
+	return InstallCacheStatusBuilder(ErrorCode);
 }
 
 TIoStatusOr<FIoMappedRegion> FOnDemandInstallCache::OpenMapped(const FIoChunkId& ChunkId, const FIoReadOptions& Options)
 {
-	return FIoStatus(EIoErrorCode::FileOpenFailed);
+	return InstallCacheStatusBuilder(EIoErrorCode::FileOpenFailed);
 }
 
 const TCHAR* FOnDemandInstallCache::GetName() const
@@ -2020,7 +2022,7 @@ FIoStatus FOnDemandInstallCache::PutChunk(FIoBuffer&& Chunk, const FIoHash& Chun
 	}
 
 	PendingChunks->Append(MoveTemp(Chunk), ChunkHash);
-	return FIoStatus::Ok;
+	return InstallCacheStatusBuilder(EIoErrorCode::Ok);
 }
 
 FIoStatus FOnDemandInstallCache::Purge(TSet<FIoHash>&& ChunksToInstall)
@@ -2067,7 +2069,7 @@ FIoStatus FOnDemandInstallCache::Purge(TSet<FIoHash>&& ChunksToInstall)
 	{
 		UE_LOG(LogIoStoreOnDemand, Log, TEXT("Skipping cache purge, MaxCacheSize=%.2lf MiB, CacheSize=%.2lf MiB, ReferencedBlockSize=%.2lf MiB, ReferencedSize=%.2lf MiB, FragmentedBytes=%.2lf MiB, UncachedSize=%.2lf MiB"),
 			ToMiB(MaxCacheSize), ToMiB(TotalCachedBytes), ToMiB(TotalReferencedBlockBytes), ToMiB(ReferencedBytes), ToMiB(FragmentedBytes), ToMiB(TotalUncachedBytes));
-		return FIoStatus::Ok;
+		return InstallCacheStatusBuilder(EIoErrorCode::Ok);
 	}
 
 	//TODO: Compute fragmentation metric and redownload chunks when this number gets too high
@@ -2100,13 +2102,13 @@ FIoStatus FOnDemandInstallCache::Purge(TSet<FIoHash>&& ChunksToInstall)
 			FIoStatus DefragStatus = Defrag(Containers, ChunkEntryIndices, BlockInfo, &DefragBytesToPurge);
 			if (DefragStatus.IsOk() == false)
 			{
-				Status = FIoStatusBuilder(EIoErrorCode::WriteError)
+				Status = InstallCacheStatusBuilder(EIoErrorCode::WriteError)
 					<< FString::Printf(TEXT("Defrag failed (%s)"), *DefragStatus.ToString());
 			}
 		}
 		else
 		{
-			Status = FIoStatusBuilder(EIoErrorCode::WriteError) << TEXT("Failed to purge required size from install cache");
+			Status = InstallCacheStatusBuilder(EIoErrorCode::WriteError) << TEXT("Failed to purge required size from install cache");
 		}
 	}
 
@@ -2162,13 +2164,13 @@ FIoStatus FOnDemandInstallCache::PurgeAllUnreferenced(bool bDefrag, const uint64
 			FIoStatus DefragStatus = Defrag(Containers, ChunkEntryIndices, BlockInfo, &DefragBytesToPurge);
 			if (DefragStatus.IsOk() == false)
 			{
-				Status = FIoStatusBuilder(EIoErrorCode::WriteError)
+				Status = InstallCacheStatusBuilder(EIoErrorCode::WriteError)
 					<< FString::Printf(TEXT("Defrag failed (%s)"), *DefragStatus.ToString());
 			}
 		}
 		else
 		{
-			Status = FIoStatusBuilder(EIoErrorCode::WriteError) << TEXT("Failed to purge required size from install cache");
+			Status = InstallCacheStatusBuilder(EIoErrorCode::WriteError) << TEXT("Failed to purge required size from install cache");
 		}
 	}
 	else if (bDefrag)
@@ -2177,7 +2179,7 @@ FIoStatus FOnDemandInstallCache::PurgeAllUnreferenced(bool bDefrag, const uint64
 		FIoStatus DefragStatus = Defrag(Containers, ChunkEntryIndices, BlockInfo);
 		if (DefragStatus.IsOk() == false)
 		{
-			Status = FIoStatusBuilder(EIoErrorCode::WriteError)
+			Status = InstallCacheStatusBuilder(EIoErrorCode::WriteError)
 				<< FString::Printf(TEXT("Defrag failed (%s)"), *DefragStatus.ToString());
 		}
 	}
@@ -2293,7 +2295,7 @@ FIoStatus FOnDemandInstallCache::Verify()
 	if (ChunkCount == 0)
 	{
 		UE_LOG(LogIoStoreOnDemand, Log, TEXT("Verify skipped, install cache is empty"));
-		return FIoStatus::Ok;
+		return InstallCacheStatusBuilder(EIoErrorCode::Ok);
 	}
 
 	UE_LOG(LogIoStoreOnDemand, Log, TEXT("Verifying %d installed chunks..."), ChunkCount);
@@ -2391,21 +2393,21 @@ FIoStatus FOnDemandInstallCache::Verify()
 
 		if (CorruptChunkCount > 0)
 		{
-			return FIoStatus(EIoErrorCode::SignatureError);
+			return InstallCacheStatusBuilder(EIoErrorCode::SignatureError);
 		}
 
 		if (ReadErrorCount > 0)
 		{
-			return FIoStatus(EIoErrorCode::ReadError);
+			return InstallCacheStatusBuilder(EIoErrorCode::ReadError);
 		}
 
-		return FIoStatus(EIoErrorCode::NotFound);
+		return InstallCacheStatusBuilder(EIoErrorCode::NotFound);
 	}
 
 	UE_LOG(LogIoStoreOnDemand, Log, TEXT("Successfully verified %d chunk(s) of total %.2lf MiB"),
 		ChunkCount, ToMiB(TotalVerifiedBytes));
 
-	return FIoStatus::Ok;
+	return InstallCacheStatusBuilder(EIoErrorCode::Ok);
 }
 
 void FOnDemandInstallCache::RegisterConsoleCommands()
@@ -2441,7 +2443,7 @@ FIoStatus FOnDemandInstallCache::Reset()
 
 	if (Ifm.DeleteDirectory(*CacheDirectory, false, bTree) == false)
 	{
-		return FIoStatusBuilder(EIoErrorCode::WriteError)
+		return InstallCacheStatusBuilder(EIoErrorCode::WriteError)
 			<< TEXT("Failed to delete directory '")
 			<< CacheDirectory
 			<< TEXT("'");
@@ -2449,7 +2451,7 @@ FIoStatus FOnDemandInstallCache::Reset()
 
 	if (Ifm.MakeDirectory(*CacheDirectory, bTree) == false)
 	{
-		return FIoStatusBuilder(EIoErrorCode::WriteError)
+		return InstallCacheStatusBuilder(EIoErrorCode::WriteError)
 			<< TEXT("Failed to create directory '")
 			<< CacheDirectory
 			<< TEXT("'");
@@ -2467,7 +2469,7 @@ FIoStatus FOnDemandInstallCache::Reset()
 	}
 
 	UE_LOG(LogIoStoreOnDemand, Log, TEXT("Created CAS journal '%s'"), *JournalFile);
-	return FIoStatus::Ok;
+	return InstallCacheStatusBuilder(EIoErrorCode::Ok);
 }
 
 FIoStatus FOnDemandInstallCache::InitialVerify()
@@ -2516,12 +2518,12 @@ FIoStatus FOnDemandInstallCache::InitialVerify()
 				UE_LOG(LogIoStoreOnDemand, Error, TEXT("Failed to purge %.2lf MiB from install cache. Actually purged %.2lf MiB from install cache"),
 					ToMiB(TotalBytesToPurge), ToMiB(TotalPurgedBytes));
 
-				return FIoStatusBuilder(EIoErrorCode::WriteError) << FString::Printf(TEXT("Failed to purge overbudget cache (%s)"), *PurgeStatus.ToString());
+				return InstallCacheStatusBuilder(EIoErrorCode::WriteError) << FString::Printf(TEXT("Failed to purge overbudget cache (%s)"), *PurgeStatus.ToString());
 			}
 		}
 	}
 
-	return FIoStatus::Ok;
+	return InstallCacheStatusBuilder(EIoErrorCode::Ok);
 }
 
 uint64 FOnDemandInstallCache::AddReferencesToBlocks(
@@ -2637,7 +2639,7 @@ FIoStatus FOnDemandInstallCache::Purge(FCasBlockInfoMap& BlockInfo, const uint64
 		}
 	}
 
-	return FIoStatus::Ok;
+	return InstallCacheStatusBuilder(EIoErrorCode::Ok);
 }
 
 FIoStatus FOnDemandInstallCache::Defrag(
@@ -2648,7 +2650,7 @@ FIoStatus FOnDemandInstallCache::Defrag(
 {
 	if (TotalBytesToFree && *TotalBytesToFree == 0)
 	{
-		return FIoStatus::Ok;
+		return InstallCacheStatusBuilder(EIoErrorCode::Ok);
 	}
 
 	const uint64 TotalCachedBytes = Algo::TransformAccumulate(BlockInfo,
@@ -2662,7 +2664,7 @@ FIoStatus FOnDemandInstallCache::Defrag(
 			ToMiB(MaxCacheSize));
 		UE_LOG(LogIoStoreOnDemand, Error, TEXT("%s"), *ErrorMsg);
 
-		FIoStatus Status = FIoStatusBuilder(EIoErrorCode::WriteError) << ErrorMsg;
+		FIoStatus Status = InstallCacheStatusBuilder(EIoErrorCode::WriteError) << ErrorMsg;
 		FOnDemandInstallCacheStats::OnDefrag(Status, 0);
 
 		// Append a critical error entry to clear the cache at next startup
@@ -2748,7 +2750,7 @@ FIoStatus FOnDemandInstallCache::Defrag(
 
 		if (!bPossibleToFreeBytes)
 		{
-			FIoStatus Status = FIoStatusBuilder(EIoErrorCode::WriteError) << TEXT("Defrag cannot free requested storage space");
+			FIoStatus Status = InstallCacheStatusBuilder(EIoErrorCode::WriteError) << TEXT("Defrag cannot free requested storage space");
 			FOnDemandInstallCacheStats::OnDefrag(Status, 0);
 			return Status;
 		}
@@ -2787,7 +2789,7 @@ FIoStatus FOnDemandInstallCache::Defrag(
 		{
 			// Already defragged
 			UE_LOG(LogIoStoreOnDemand, Display, TEXT("Cache not fragmented."));
-			return FIoStatus::Ok;
+			return InstallCacheStatusBuilder(EIoErrorCode::Ok);
 		}
 	}
 
@@ -2837,7 +2839,7 @@ FIoStatus FOnDemandInstallCache::Defrag(
 			{
 				const FString ErrorMsg = FileOpenResult.GetError().GetMessage();
 				UE_LOG(LogIoStoreOnDemand, Error, TEXT("Defrag failed to open CAS block for reading, reason '%s'"), *ErrorMsg);
-				FIoStatus Status = FIoStatusBuilder(EIoErrorCode::FileOpenFailed) << ErrorMsg;
+				FIoStatus Status = InstallCacheStatusBuilder(EIoErrorCode::FileOpenFailed) << ErrorMsg;
 				FOnDemandInstallCacheStats::OnDefrag(Status, 0);
 				return Status;
 			}
@@ -2854,8 +2856,9 @@ FIoStatus FOnDemandInstallCache::Defrag(
 				const bool bOk = FileHandle->Read(Buffer.GetData(), Buffer.GetSize());
 				if (!bOk)
 				{
-					FOnDemandInstallCacheStats::OnDefrag(EIoErrorCode::ReadError, 0);
-					return EIoErrorCode::ReadError;
+					FIoStatus Status = InstallCacheStatusBuilder(EIoErrorCode::ReadError);
+					FOnDemandInstallCacheStats::OnDefrag(Status, 0);
+					return Status;
 				}
 
 				const FIoHash ChunkHash = FIoHash::HashBuffer(Buffer.GetView());
@@ -2879,8 +2882,9 @@ FIoStatus FOnDemandInstallCache::Defrag(
 						return Status;
 					}
 	
-					FOnDemandInstallCacheStats::OnDefrag(EIoErrorCode::SignatureError, 0);
-					return EIoErrorCode::SignatureError;
+					FIoStatus Status = InstallCacheStatusBuilder(EIoErrorCode::SignatureError);
+					FOnDemandInstallCacheStats::OnDefrag(Status, 0);
+					return Status;
 				}
 
 				if (DefragPendingChunks.TotalSize > FPendingChunks::MaxPendingBytes)
@@ -2934,9 +2938,9 @@ FIoStatus FOnDemandInstallCache::Defrag(
 
 	UE_LOG(LogIoStoreOnDemand, Display, TEXT("Defrag removed %" UINT64_FMT " fragmented bytes of %" UINT64_FMT " total bytes in %i blocks."), FragmentedBytes, TotalBlockSize, BlocksToDefrag.Num());
 
-	FOnDemandInstallCacheStats::OnDefrag(EIoErrorCode::Ok, FragmentedBytes);
-
-	return FIoStatus::Ok;
+	FIoStatus Status = InstallCacheStatusBuilder(EIoErrorCode::Ok);
+	FOnDemandInstallCacheStats::OnDefrag(Status, FragmentedBytes);
+	return Status;
 }
 
 FIoStatus FOnDemandInstallCache::Flush()
@@ -2948,7 +2952,7 @@ FIoStatus FOnDemandInstallCache::Flush()
 	}
 
 	Cas.Compact();
-	return FIoStatus::Ok;
+	return InstallCacheStatusBuilder(EIoErrorCode::Ok);
 }
 
 FOnDemandInstallCacheUsage FOnDemandInstallCache::GetCacheUsage()
@@ -3001,7 +3005,7 @@ FIoStatus FOnDemandInstallCache::FlushPendingChunks(FPendingChunks& Chunks, int6
 {
 	if (Chunks.IsEmpty())
 	{
-		return FIoStatus::Ok;
+		return InstallCacheStatusBuilder(EIoErrorCode::Ok);
 	}
 
 #if UE_ONDEMANDINSTALLCACHE_EXCLUSIVE_WRITE
@@ -3049,7 +3053,7 @@ FIoStatus FOnDemandInstallCache::FlushPendingChunksImpl(FPendingChunks& Chunks,
 		if (CasFileHandle.IsValid() == false)
 		{
 			FOnDemandInstallCacheStats::OnFlush(EIoErrorCode::WriteError, 0);
-			return FIoStatusBuilder(EIoErrorCode::FileOpenFailed) << TEXT("Failed to open cache block file");
+			return InstallCacheStatusBuilder(EIoErrorCode::FileOpenFailed) << TEXT("Failed to open cache block file");
 		}
 
 		const int64 CasBlockOffset = CasFileHandle->Tell();
@@ -3076,7 +3080,7 @@ FIoStatus FOnDemandInstallCache::FlushPendingChunksImpl(FPendingChunks& Chunks,
 
 			if (CasFileHandle->Write(Ar.GetData(), Ar.Tell()) == false)
 			{
-				return FIoStatusBuilder(EIoErrorCode::WriteError)
+				return InstallCacheStatusBuilder(EIoErrorCode::WriteError)
 					<< TEXT("Failed to serialize chunks to cache block");
 			}
 
@@ -3092,7 +3096,7 @@ FIoStatus FOnDemandInstallCache::FlushPendingChunksImpl(FPendingChunks& Chunks,
 			if (CasFileHandle->Flush() == false)
 			{
 				FOnDemandInstallCacheStats::OnFlush(EIoErrorCode::WriteError, Ar.Tell());
-				return FIoStatusBuilder(EIoErrorCode::WriteError)
+				return InstallCacheStatusBuilder(EIoErrorCode::WriteError)
 					<< TEXT("Failed to flush cache block to disk");
 			}
 
@@ -3126,7 +3130,7 @@ FIoStatus FOnDemandInstallCache::FlushPendingChunksImpl(FPendingChunks& Chunks,
 		}
 	}
 
-	return FIoStatus::Ok;
+	return InstallCacheStatusBuilder(EIoErrorCode::Ok);
 }
 
 void FOnDemandInstallCache::CompleteRequest(FIoRequestImpl* Request, EIoErrorCode Status)
diff --git a/Engine/Source/Runtime/Experimental/IoStore/OnDemand/Private/Statistics.cpp b/Engine/Source/Runtime/Experimental/IoStore/OnDemand/Private/Statistics.cpp
--- a/Engine/Source/Runtime/Experimental/IoStore/OnDemand/Private/Statistics.cpp
+++ b/Engine/Source/Runtime/Experimental/IoStore/OnDemand/Private/Statistics.cpp
@@ -103,6 +103,34 @@ static void SendImmediateAnalytic(FOnDemandImmediateAnalytic&& Event)
 	}
 }
 
+static FString GetAnalyticErrorCode(const FIoStatus& Status)
+{
+	FString RetCode;
+
+	FStringView ErrorId = Status.GetErrorMessage();
+	if (Status.GetErrorCode() != EIoErrorCode::Ok && ErrorId.StartsWith(TEXT('<')))
+	{
+		ErrorId.RightChopInline(1);
+
+		int32 Idx;
+		if (ErrorId.FindChar(TEXT('>'), Idx))
+		{
+			ErrorId.LeftInline(Idx);
+
+			RetCode = FString::Format(TEXT("{0}:{1}"), { GetIoErrorText(Status.GetErrorCode()), ErrorId });
+		}
+	}
+
+	if (RetCode.IsEmpty())
+	{
+		RetCode = GetIoErrorText(Status.GetErrorCode());
+	}
+
+	RetCode.ReplaceCharInline(TEXT(' '), TEXT('_'), ESearchCase::CaseSensitive);
+
+	return RetCode;
+}
+
 ////////////////////////////////////////////////////////////////////////////////
 static int32 BytesToApproxMB(uint64 Bytes) { return int32(Bytes >> 20); }
 static int32 BytesToApproxKB(uint64 Bytes) { return int32(Bytes >> 10); }
@@ -1446,7 +1474,7 @@ void FOnDemandInstallCacheStats::OnPurge(
 		TEXT("IadCachePurge"),
 		MakeAnalyticsEventAttributeArray(
 			TEXT("Success"), Status.IsOk(),
-			TEXT("ErrorCode"), Status.ToString(),
+			TEXT("ErrorCode"), GetAnalyticErrorCode(Status),
 			TEXT("PurgedBytes"), PurgedBytes
 		)
 	});
@@ -1474,7 +1502,7 @@ void FOnDemandInstallCacheStats::OnDefrag(const FIoStatus& Status, uint64 Fragme
 		TEXT("IadCacheDefrag"),
 		MakeAnalyticsEventAttributeArray(
 			TEXT("Success"), Status.IsOk(),
-			TEXT("ErrorCode"), Status.ToString(),
+			TEXT("ErrorCode"), GetAnalyticErrorCode(Status),
 			TEXT("FragmentedBytes"), FragmentedBytes
 		)
 	});
