Iris
* Implemented CreationDependency feature
* Using AddCreationDependencyLink will force the child entity to only be created after the parent exists on the client. The child entity will also only be relevant if the parent also is relevant.
* Clients will now queue incoming data for those objects until the parent is finally instantiated.

#rnx

diff --git a/Engine/Source/Runtime/Experimental/Iris/Core/Private/Iris/ReplicationSystem/Filtering/ReplicationFiltering.cpp b/Engine/Source/Runtime/Experimental/Iris/Core/Private/Iris/ReplicationSystem/Filtering/ReplicationFiltering.cpp
--- a/Engine/Source/Runtime/Experimental/Iris/Core/Private/Iris/ReplicationSystem/Filtering/ReplicationFiltering.cpp
+++ b/Engine/Source/Runtime/Experimental/Iris/Core/Private/Iris/ReplicationSystem/Filtering/ReplicationFiltering.cpp
@@ -1351,38 +1351,38 @@ void FReplicationFiltering::PreUpdateDynamicFiltering()
 	}
 }
 
-void FReplicationFiltering::UpdateParentFilterDependent(uint32 DependentIndex, const FNetBitArrayView DependentsWithParentFilter,  FNetBitArrayView OutConnectionObjectsInScope, bool bRecursive) const
+void FReplicationFiltering::UpdateCreationDependentParent(uint32 ChildINdex, const FNetBitArrayView ObjectsWithCreationDependencies,  FNetBitArrayView OutConnectionObjectsInScope, bool bRecursive) const
 {
-	const bool bStartedRelevant = OutConnectionObjectsInScope.IsBitSet(DependentIndex);
+	const bool bStartedRelevant = OutConnectionObjectsInScope.IsBitSet(ChildINdex);
 	if (!bStartedRelevant)
 	{
 		// No need to check it's parents, the object must be relevant by itself first.
 		return;
 	}
 
-	const TArrayView<const FInternalNetRefIndex> Parents = NetRefHandleManager->GetDependentObjectParents(DependentIndex);
+	const TArrayView<const FInternalNetRefIndex> Parents = NetRefHandleManager->GetDependentObjectParents(ChildINdex);
 
 	for (FInternalNetRefIndex ParentIndex : Parents)
 	{
-		const bool bIsParentDependent = DependentsWithParentFilter.IsBitSet(ParentIndex);
+		const bool bHasCreationDependency = ObjectsWithCreationDependencies.IsBitSet(ParentIndex);
 					
 		// If the parent is a dependent AND relevant check his status first
 		//TODO: Optimize this by flagging which parents were processed already or not
-		if (bIsParentDependent && OutConnectionObjectsInScope.IsBitSet(ParentIndex))
+		if (bHasCreationDependency && OutConnectionObjectsInScope.IsBitSet(ParentIndex))
 		{
 			constexpr bool bIsRecursiveCall = true;
-			UpdateParentFilterDependent(ParentIndex, DependentsWithParentFilter, OutConnectionObjectsInScope, bIsRecursiveCall);
+			UpdateCreationDependentParent(ParentIndex, ObjectsWithCreationDependencies, OutConnectionObjectsInScope, bIsRecursiveCall);
 		}
 
 		const bool bIsParentRelevant = OutConnectionObjectsInScope.IsBitSet(ParentIndex);
-		OutConnectionObjectsInScope.AndBitValue(DependentIndex, bIsParentRelevant);
+		OutConnectionObjectsInScope.AndBitValue(ChildINdex, bIsParentRelevant);
 	}
 
 	// Check if we are still relevant, if not clear our subobjects
-	const bool bIsStillRelevant = OutConnectionObjectsInScope.IsBitSet(DependentIndex);
+	const bool bIsStillRelevant = OutConnectionObjectsInScope.IsBitSet(ChildINdex);
 	if (!bRecursive && !bIsStillRelevant)
 	{
-		for (const FInternalNetRefIndex SubObjectIndex : NetRefHandleManager->GetSubObjects(DependentIndex))
+		for (const FInternalNetRefIndex SubObjectIndex : NetRefHandleManager->GetSubObjects(ChildINdex))
 		{
 			OutConnectionObjectsInScope.ClearBit(SubObjectIndex);
 		}
@@ -1423,8 +1423,8 @@ void FReplicationFiltering::UpdateDynamicFiltering()
 	const FNetBitArrayView DependentObjectsView = NetRefHandleManager->GetDependentObjectInternalIndices();
 	const uint32* DependentObjectsData = DependentObjectsView.GetDataChecked(WordCountForObjectBitArrays);
 
-	const FNetBitArrayView DependentsWithParentFilter = NetRefHandleManager->GetDependentObjectsWithParentFilter();
-	const uint32* DependentsWithParentFilterData = DependentsWithParentFilter.GetDataChecked(WordCountForObjectBitArrays);
+	const FNetBitArrayView ObjectsWithCreationDependencies = NetRefHandleManager->GetObjectsWithCreationDependencies();
+	const uint32* ObjectsWithCreationDependenciesData = ObjectsWithCreationDependencies.GetDataChecked(WordCountForObjectBitArrays);
 	
 	uint32* ConnectionIds = static_cast<uint32*>(FMemory_Alloca(ValidConnections.GetNumBits() * sizeof(uint32)));
 	uint32 ConnectionCount = 0;
@@ -1495,7 +1495,7 @@ void FReplicationFiltering::UpdateDynamicFiltering()
 			{
 				const uint32 SubObjects = SubObjectsData[WordIt];
 				// Set of dependents objects that don't have the parent filter trait
-				const uint32 DependentObjects = (DependentObjectsData[WordIt] & ~DependentsWithParentFilterData[WordIt]);
+				const uint32 DependentObjects = (DependentObjectsData[WordIt] & ~ObjectsWithCreationDependenciesData[WordIt]);
 				const uint32 ObjectsRequiringUpdate = ObjectsRequiringDynamicFilterUpdateData[WordIt];
 				const uint32 PrevFilteredOutObjects = DynamicFilteredOutObjectsData[WordIt];
 				// Mask off group included objects in this loops so they're accounted for when updating hysteresis. A no longer group included object should also be subject to hysteresis.
@@ -1666,12 +1666,12 @@ void FReplicationFiltering::UpdateDynamicFiltering()
 
 		// Update dependent objects that can only be relevant if their parent also is
 		{
-			IRIS_PROFILER_SCOPE(FReplicationFiltering_UpdateDependentsWithParentFilter);
+			IRIS_PROFILER_SCOPE(FReplicationFiltering_UpdateCreationDependencies);
 			FNetBitArrayView ConnectionObjectsInScope = MakeNetBitArrayView(ConnectionInfo.ObjectsInScope);
-			FNetBitArrayView::ForAllSetBits(DependentsWithParentFilter, ConnectionObjectsInScope, FNetBitArrayView::AndOp, [&](uint32 DependentIndex)
+			FNetBitArrayView::ForAllSetBits(ObjectsWithCreationDependencies, ConnectionObjectsInScope, FNetBitArrayView::AndOp, [&](uint32 DependentIndex)
 			{
 				constexpr bool bIsNotRecursive = false;
-				this->UpdateParentFilterDependent(DependentIndex, DependentsWithParentFilter, ConnectionObjectsInScope, bIsNotRecursive);
+				this->UpdateCreationDependentParent(DependentIndex, ObjectsWithCreationDependencies, ConnectionObjectsInScope, bIsNotRecursive);
 			});
 		}
 
diff --git a/Engine/Source/Runtime/Experimental/Iris/Core/Private/Iris/ReplicationSystem/Filtering/ReplicationFiltering.h b/Engine/Source/Runtime/Experimental/Iris/Core/Private/Iris/ReplicationSystem/Filtering/ReplicationFiltering.h
--- a/Engine/Source/Runtime/Experimental/Iris/Core/Private/Iris/ReplicationSystem/Filtering/ReplicationFiltering.h
+++ b/Engine/Source/Runtime/Experimental/Iris/Core/Private/Iris/ReplicationSystem/Filtering/ReplicationFiltering.h
@@ -241,7 +241,7 @@ private:
 
 	bool HasDynamicFilters() const;
 
-	void UpdateParentFilterDependent(uint32 DependentIndex, const FNetBitArrayView DependentsWithParentFilter, FNetBitArrayView OutConnectionObjectsInScope, bool bIsRecursive) const;
+	void UpdateCreationDependentParent(uint32 ChildIndex, const FNetBitArrayView ObjectsWithCreationDependencies, FNetBitArrayView OutConnectionObjectsInScope, bool bIsRecursive) const;
 
 	// Helper to update and reset group exclusion filter effects if objects are removed from a filter or after a filter status change, returns true if the group filter was changed
 	bool ClearGroupExclusionFilterEffectsForObject(uint32 ObjectIndex, uint32 ConnectionId);
diff --git a/Engine/Source/Runtime/Experimental/Iris/Core/Private/Iris/ReplicationSystem/NetDependencyData.h b/Engine/Source/Runtime/Experimental/Iris/Core/Private/Iris/ReplicationSystem/NetDependencyData.h
--- a/Engine/Source/Runtime/Experimental/Iris/Core/Private/Iris/ReplicationSystem/NetDependencyData.h
+++ b/Engine/Source/Runtime/Experimental/Iris/Core/Private/Iris/ReplicationSystem/NetDependencyData.h
@@ -24,19 +24,10 @@ struct FChildSubObjectsInfo
 	uint32 NumSubObjects = 0U;
 };
 
-enum class EDependentFilterTrait : uint8
-{
-	/** Allow the dependent object to be relevant even if the parent isn't*/
-	None,
-	/** The dependent object can only be relevant if the parent also is */
-	ParentMustBeRelevant,
-};
-
 struct FDependentObjectInfo
 {
 	FInternalNetRefIndex NetRefIndex = 0U;
 	EDependentObjectSchedulingHint SchedulingHint = EDependentObjectSchedulingHint::Default;
-	EDependentFilterTrait ParentFilterTrait = EDependentFilterTrait::None;
 };
 
 class FNetDependencyData
diff --git a/Engine/Source/Runtime/Experimental/Iris/Core/Private/Iris/ReplicationSystem/NetRefHandleManager.cpp b/Engine/Source/Runtime/Experimental/Iris/Core/Private/Iris/ReplicationSystem/NetRefHandleManager.cpp
--- a/Engine/Source/Runtime/Experimental/Iris/Core/Private/Iris/ReplicationSystem/NetRefHandleManager.cpp
+++ b/Engine/Source/Runtime/Experimental/Iris/Core/Private/Iris/ReplicationSystem/NetRefHandleManager.cpp
@@ -77,7 +77,6 @@ void FNetRefHandleManager::Init(const FInitParams& InitParams)
 		InitNetBitArray(&AssignedInternalIndices);
 		InitNetBitArray(&SubObjectInternalIndices);
 		InitNetBitArray(&DependentObjectInternalIndices);
-		InitNetBitArray(&DependentObjectsWithParentFilter);
 		InitNetBitArray(&ObjectsWithCreationDependencies);
 		InitNetBitArray(&ObjectsWithDependentObjectsInternalIndices);
 		InitNetBitArray(&DestroyedStartupObjectInternalIndices);
@@ -460,7 +459,7 @@ void FNetRefHandleManager::InternalDestroyNetObject(FInternalNetRefIndex Interna
 	}
 
 	// Remove from all dependent object relationships and clear data
-	InternalRemoveDependentObject(InternalIndex);
+	InternalRemoveAllDependencies(InternalIndex);
 
 	// Free all stored data for object
 	SubObjects.FreeStoredDependencyDataForObject(InternalIndex);
@@ -810,7 +809,7 @@ FInternalNetRefIndex FNetRefHandleManager::GetRootObjectInternalIndexOfSubObject
 	return SubObjectIndex != InvalidInternalIndex ? ReplicatedObjectData[SubObjectIndex].SubObjectRootIndex : InvalidInternalIndex;
 }
 
-bool FNetRefHandleManager::AddDependentObject(FNetRefHandle ParentRefHandle, FNetRefHandle DependentObjectRefHandle, EDependentObjectSchedulingHint SchedulingHint, EDependentFilterTrait FilterTrait)
+bool FNetRefHandleManager::AddDependentObject(FNetRefHandle ParentRefHandle, FNetRefHandle DependentObjectRefHandle, EDependentObjectSchedulingHint SchedulingHint)
 {
 	if (ParentRefHandle == DependentObjectRefHandle)
 	{
@@ -866,7 +865,6 @@ bool FNetRefHandleManager::AddDependentObject(FNetRefHandle ParentRefHandle, FNe
 	FDependentObjectInfo DependentObjectInfo;
 	DependentObjectInfo.NetRefIndex = DependentObjectInternalIndex;
 	DependentObjectInfo.SchedulingHint = SchedulingHint;
-	DependentObjectInfo.ParentFilterTrait = FilterTrait;
 	ParentDependentObjectsArray.Add(DependentObjectInfo);
 
 	// Add parent to dependent's list
@@ -877,7 +875,6 @@ bool FNetRefHandleManager::AddDependentObject(FNetRefHandle ParentRefHandle, FNe
 	ParentObjectData.bHasDependentObjects = true;
 	ObjectsWithDependentObjectsInternalIndices.SetBit(ParentInternalIndex);
 	DependentObjectInternalIndices.SetBit(DependentObjectInternalIndex);
-	DependentObjectsWithParentFilter.SetBitValue(DependentObjectInternalIndex, FilterTrait == EDependentFilterTrait::ParentMustBeRelevant);
 
 	return true;
 }
@@ -888,7 +885,7 @@ void FNetRefHandleManager::RemoveDependentObject(FNetRefHandle DependentHandle)
 
 	if (DependentInternalIndex != InvalidInternalIndex)
 	{
-		InternalRemoveDependentObject(DependentInternalIndex);
+		InternalRemoveAllDependencies(DependentInternalIndex);
 	}
 }
 
@@ -955,7 +952,7 @@ void FNetRefHandleManager::RemoveCreationDependency(FNetRefHandle Parent, FNetRe
 	{
 		ObjectsWithCreationDependencies.ClearBit(ChildIndex);
 		
-		//TODO: Could free the creation info when empty to leave some memory. But if most operations are Remove -> Add then it's faster to keep the array assigned. To meditate on.
+		//We don't free the array because we usually add a new dependency right after or the object will stop replicating soon and the array will be freed there.
 	}
 }
 
@@ -988,7 +985,6 @@ void FNetRefHandleManager::InternalRemoveDependentObject(FInternalNetRefIndex Pa
 				FReplicatedObjectData& DependentObjectData = GetReplicatedObjectDataNoCheck(DependentInternalIndex);
 				DependentObjectData.bIsDependentObject = false;
 				DependentObjectInternalIndices.ClearBit(DependentInternalIndex);
-				DependentObjectsWithParentFilter.ClearBit(DependentInternalIndex);
 			}
 		}
 	}
@@ -1014,7 +1010,7 @@ void FNetRefHandleManager::InternalRemoveDependentObject(FInternalNetRefIndex Pa
 	}
 }
 
-void FNetRefHandleManager::InternalRemoveDependentObject(FInternalNetRefIndex DependentInternalIndex)
+void FNetRefHandleManager::InternalRemoveAllDependencies(FInternalNetRefIndex DependentInternalIndex)
 {
 	// Remove from all parents
 	if (FNetDependencyData::FInternalNetRefIndexArray* ParentObjectArray = SubObjects.GetInternalIndexArray<FNetDependencyData::DependentParentObjects>(DependentInternalIndex))
@@ -1038,6 +1034,14 @@ void FNetRefHandleManager::InternalRemoveDependentObject(FInternalNetRefIndex De
 		DependentObjectArray->Reset();		
 	}
 
+	if (FNetDependencyData::FCreationDependencyInfoArray* CreationDependenciesArray = SubObjects.GetCreationDependencyInfoArray(DependentInternalIndex))
+	{
+		// TODO: This code currently assumes that the parent cannot be destroyed while the children depending on him are still replicating.
+		
+		// Need to add a parent->child link and decide if it's legal for us to remove the dependency automatically or not.
+		CreationDependenciesArray->Reset();
+	}
+
 	FReplicatedObjectData& DependentObjectData = GetReplicatedObjectDataNoCheck(DependentInternalIndex);
 
 	// Clear out flags on this object	
@@ -1045,7 +1049,6 @@ void FNetRefHandleManager::InternalRemoveDependentObject(FInternalNetRefIndex De
 	DependentObjectData.bHasDependentObjects = false;
 	ObjectsWithDependentObjectsInternalIndices.ClearBit(DependentInternalIndex);
 	DependentObjectInternalIndices.ClearBit(DependentInternalIndex);
-	DependentObjectsWithParentFilter.ClearBit(DependentInternalIndex);
 	ObjectsWithCreationDependencies.ClearBit(DependentInternalIndex);
 }
 
diff --git a/Engine/Source/Runtime/Experimental/Iris/Core/Private/Iris/ReplicationSystem/NetRefHandleManager.h b/Engine/Source/Runtime/Experimental/Iris/Core/Private/Iris/ReplicationSystem/NetRefHandleManager.h
--- a/Engine/Source/Runtime/Experimental/Iris/Core/Private/Iris/ReplicationSystem/NetRefHandleManager.h
+++ b/Engine/Source/Runtime/Experimental/Iris/Core/Private/Iris/ReplicationSystem/NetRefHandleManager.h
@@ -271,12 +271,8 @@ public:
 	// DependentObjects
 	const FNetBitArrayView GetDependentObjectInternalIndices() const { return MakeNetBitArrayView(DependentObjectInternalIndices); }
 	const FNetBitArrayView GetObjectsWithDependentObjectsInternalIndices() const { return MakeNetBitArrayView(ObjectsWithDependentObjectsInternalIndices); }
-	const FNetBitArrayView GetDependentObjectsWithParentFilter() const
-	{
-		return MakeNetBitArrayView(DependentObjectsWithParentFilter);
-	}
-
-	bool AddDependentObject(FNetRefHandle ParentHandle, FNetRefHandle DependentHandle, EDependentObjectSchedulingHint SchedulingHint, EDependentFilterTrait FilterTrait = EDependentFilterTrait::None);
+	
+	bool AddDependentObject(FNetRefHandle ParentHandle, FNetRefHandle DependentHandle, EDependentObjectSchedulingHint SchedulingHint);
 	void RemoveDependentObject(FNetRefHandle ParentHandle, FNetRefHandle DependentHandle);
 
 	// Creation dependency API
@@ -419,7 +415,7 @@ private:
 	void InternalRemoveSubObject(FInternalNetRefIndex RootObjectInternalIndex, FInternalNetRefIndex SubObjectInternalIndex, bool bRemoveFromSubObjectArray = true);
 
 	void InternalRemoveDependentObject(FInternalNetRefIndex ParentInternalIndex, FInternalNetRefIndex DependentInternalIndex, ERemoveDependentObjectFlags Flags = ERemoveDependentObjectFlags::All);
-	void InternalRemoveDependentObject(FInternalNetRefIndex DependentInternalIndex);
+	void InternalRemoveAllDependencies(FInternalNetRefIndex DependentInternalIndex);
 
 	/**
 	 * Grow the NetObjectLists to fit more internal indexes.
@@ -494,7 +490,6 @@ private:
 	FNetBitArray DependentObjectInternalIndices;
 
 	// Dependent objects that need their parent to be relevant too
-	FNetBitArray DependentObjectsWithParentFilter;
 
 	// Objects that cannot be created until the parents they depend on also exist
 	FNetBitArray ObjectsWithCreationDependencies;
diff --git a/Engine/Source/Runtime/Experimental/Iris/Core/Private/Iris/ReplicationSystem/ObjectReferenceCache.cpp b/Engine/Source/Runtime/Experimental/Iris/Core/Private/Iris/ReplicationSystem/ObjectReferenceCache.cpp
--- a/Engine/Source/Runtime/Experimental/Iris/Core/Private/Iris/ReplicationSystem/ObjectReferenceCache.cpp
+++ b/Engine/Source/Runtime/Experimental/Iris/Core/Private/Iris/ReplicationSystem/ObjectReferenceCache.cpp
@@ -66,6 +66,11 @@ namespace UE::Net::Private
 static bool bIrisAllowAsyncLoading = true;
 FAutoConsoleVariableRef CVarIrisAllowAsyncLoading(TEXT("net.iris.AllowAsyncLoading"), bIrisAllowAsyncLoading, TEXT("Flag to allow or disallow async loading when using iris replication. Note: net.allowAsyncLoading must also be enabled."), ECVF_Default);
 
+#if UE_NET_ASYNCLOADING_DEBUG
+static bool bForceAsyncLoading = false;
+FAutoConsoleVariableRef CVarIrisForceAsyncLoading(TEXT("net.Iris.AsyncLoading.ForceEnable"), bForceAsyncLoading, TEXT("Force the async loading path even in Editor/PIE."), ECVF_Default);
+#endif
+
 #if UE_NET_OBJECTREFERENCECACHE_ASYNCLOAD_DEBUG
 static FString IrisAsyncLoadingFailPackageName;
 FAutoConsoleVariableRef CVarIrisAsyncLoadingFailPackageName(TEXT("net.Iris.AsyncLoading.FailPackageName"), IrisAsyncLoadingFailPackageName, TEXT("Always fail loading the specified package name. For debugging purposes."), ECVF_Default);
@@ -2190,6 +2195,13 @@ void FObjectReferenceCache::SetAsyncLoadMode(const EAsyncLoadMode NewMode)
 
 bool FObjectReferenceCache::ShouldAsyncLoad() const
 {
+#if UE_NET_ASYNCLOADING_DEBUG
+	if (bForceAsyncLoading)
+	{
+		return true;
+	}
+#endif
+
 	if (!bIrisAllowAsyncLoading)
 	{
 		return false;
diff --git a/Engine/Source/Runtime/Experimental/Iris/Core/Private/Iris/ReplicationSystem/ObjectReferenceTypes.cpp b/Engine/Source/Runtime/Experimental/Iris/Core/Private/Iris/ReplicationSystem/ObjectReferenceTypes.cpp
--- a/Engine/Source/Runtime/Experimental/Iris/Core/Private/Iris/ReplicationSystem/ObjectReferenceTypes.cpp
+++ b/Engine/Source/Runtime/Experimental/Iris/Core/Private/Iris/ReplicationSystem/ObjectReferenceTypes.cpp
@@ -45,6 +45,12 @@ TAsyncLoadPriority ConvertAsyncLoadingPriority(EIrisAsyncLoadingPriority IrisPri
 			return UE::Net::Private::IrisVeryHighAsyncLoadingPriority;
 		} break;
 
+		case EIrisAsyncLoadingPriority::Invalid:
+		{
+			ensureMsgf(false, TEXT("Invalid config used, should not happen."));
+			return UE::Net::Private::IrisDefaultAsyncLoadingPriority;
+		} break;
+
 		default:
 		{
 			return UE::Net::Private::IrisDefaultAsyncLoadingPriority;
diff --git a/Engine/Source/Runtime/Experimental/Iris/Core/Private/Iris/ReplicationSystem/ObjectReplicationBridge.cpp b/Engine/Source/Runtime/Experimental/Iris/Core/Private/Iris/ReplicationSystem/ObjectReplicationBridge.cpp
--- a/Engine/Source/Runtime/Experimental/Iris/Core/Private/Iris/ReplicationSystem/ObjectReplicationBridge.cpp
+++ b/Engine/Source/Runtime/Experimental/Iris/Core/Private/Iris/ReplicationSystem/ObjectReplicationBridge.cpp
@@ -777,7 +777,7 @@ void UObjectReplicationBridge::RemoveDependentObject(FNetRefHandle ParentHandle,
 	}
 }
 
-void UObjectReplicationBridge::AddDependentObjectWithParentRelevancy(FNetRefHandle Parent, FNetRefHandle DependentObject)
+void UObjectReplicationBridge::AddCreationDependencyLink(FNetRefHandle Parent, FNetRefHandle Child)
 {
 	using namespace UE::Net;
 	using namespace UE::Net::Private;
@@ -785,20 +785,23 @@ void UObjectReplicationBridge::AddDependentObjectWithParentRelevancy(FNetRefHand
 	FReplicationSystemInternal* ReplicationSystemInternal = GetReplicationSystem()->GetReplicationSystemInternal();
 	FNetRefHandleManager& LocalNetRefHandleManager = ReplicationSystemInternal->GetNetRefHandleManager();
 
-	if (LocalNetRefHandleManager.AddDependentObject(Parent, DependentObject, EDependentObjectSchedulingHint::Default, EDependentFilterTrait::ParentMustBeRelevant))
+    const bool bSuccess = LocalNetRefHandleManager.AddCreationDependency(Parent, Child);
+    
+    if (bSuccess)
 	{
-		FReplicationFiltering& Filtering = ReplicationSystemInternal->GetFiltering();
-		const FInternalNetRefIndex DependentInternalIndex = LocalNetRefHandleManager.GetInternalIndex(DependentObject);
-		Filtering.NotifyAddedDependentObject(DependentInternalIndex);
-
-		UE_LOG_OBJECTREPLICATIONBRIDGE(Verbose, TEXT("AddDependentObjectWithParentRelevancy Added dependent object %s to parent %s"), *PrintObjectFromNetRefHandle(DependentObject), *PrintObjectFromNetRefHandle(Parent));
+		UE_LOG_OBJECTREPLICATIONBRIDGE(Verbose, TEXT("AddCreationDependencyLink Added creation dependency between child: %s and parent: %s"), *PrintObjectFromNetRefHandle(Child), *PrintObjectFromNetRefHandle(Parent));
 	}
 	else
 	{
-		UE_LOG_OBJECTREPLICATIONBRIDGE(Warning , TEXT("AddDependentObjectWithParentRelevancy Failed to add dependent object %s to parent %s"), *PrintObjectFromNetRefHandle(DependentObject), *PrintObjectFromNetRefHandle(Parent));
+		UE_LOG_OBJECTREPLICATIONBRIDGE(Warning , TEXT("AddCreationDependencyLink Failed to add creation dependent between child: %s and parent: %s"), *PrintObjectFromNetRefHandle(Child), *PrintObjectFromNetRefHandle(Parent));
 	}
 }
 
+void UObjectReplicationBridge::RemoveCreationDependencyLink(FNetRefHandle Parent, FNetRefHandle Child)
+{
+	GetReplicationSystem()->GetReplicationSystemInternal()->GetNetRefHandleManager().RemoveCreationDependency(Parent, Child);
+}
+
 bool UObjectReplicationBridge::WriteNetRefHandleCreationInfo(FReplicationBridgeSerializationContext& Context, FNetRefHandle Handle)
 {
 	using namespace UE::Net;
diff --git a/Engine/Source/Runtime/Experimental/Iris/Core/Private/Iris/ReplicationSystem/PendingBatchData.h b/Engine/Source/Runtime/Experimental/Iris/Core/Private/Iris/ReplicationSystem/PendingBatchData.h
--- a/Engine/Source/Runtime/Experimental/Iris/Core/Private/Iris/ReplicationSystem/PendingBatchData.h
+++ b/Engine/Source/Runtime/Experimental/Iris/Core/Private/Iris/ReplicationSystem/PendingBatchData.h
@@ -41,16 +41,19 @@ struct FPendingBatchData
 	TArray<uint32, TInlineAllocator<32>> DataChunkStorage;		
 	TArray<FQueuedDataChunk, TInlineAllocator<4>> QueuedDataChunks;
 
-	// Must be mapped references pending resolve
+	// The MustBeMapped references that are still not resolved
 	TArray<FNetRefHandle, TInlineAllocator<4>> PendingMustBeMappedReferences;
 
 	// Resolved references for which we have are holding on to references to avoid GC
 	TArray<FNetRefHandle, TInlineAllocator<4>> ResolvedReferences;
 
-	// Batch owner with queued data chunks
-	FNetRefHandle Handle;
+	// Owner of the queued data chunks
+	FNetRefHandle Owner;
 
-	// When we started to accumulate pending batch data for this object
+	// The list of parents that must exist before the owner can be created
+	TArray<FNetRefHandle, TInlineAllocator<4>> CreationDependentParents;
+
+	// Time when we started to accumulate data for this object
 	uint64 PendingBatchStartCycles = 0;
 
 	// At what time should we warn about being blocked too long
@@ -59,14 +62,14 @@ struct FPendingBatchData
 	// Incremented every time we try to process queued batches, reset each time we output warning.
 	int32 PendingBatchTryProcessCount = 0;
 
-	// The async loading priority of the object stored in the queued batch
-	EIrisAsyncLoadingPriority IrisAsyncLoadingPriority = EIrisAsyncLoadingPriority::Default;
+	// The async loading priority of the Owner
+	EIrisAsyncLoadingPriority IrisAsyncLoadingPriority = EIrisAsyncLoadingPriority::Invalid;
 };
 
 struct FPendingBatches
 {
-	FPendingBatchData* Find(FNetRefHandle NetRefHandle) { return PendingBatches.FindByPredicate([&NetRefHandle](const FPendingBatchData& Entry) { return Entry.Handle == NetRefHandle; });}
-	const FPendingBatchData* Find(FNetRefHandle NetRefHandle) const { return PendingBatches.FindByPredicate([&NetRefHandle](const FPendingBatchData& Entry) { return Entry.Handle == NetRefHandle; });}
+	FPendingBatchData* Find(FNetRefHandle NetRefHandle) { return PendingBatches.FindByPredicate([&NetRefHandle](const FPendingBatchData& Entry) { return Entry.Owner == NetRefHandle; });}
+	const FPendingBatchData* Find(FNetRefHandle NetRefHandle) const { return PendingBatches.FindByPredicate([&NetRefHandle](const FPendingBatchData& Entry) { return Entry.Owner == NetRefHandle; });}
 	bool GetHasPendingBatches() const { return !PendingBatches.IsEmpty(); }
 
 	TArray<FPendingBatchData> PendingBatches;
diff --git a/Engine/Source/Runtime/Experimental/Iris/Core/Private/Iris/ReplicationSystem/ReplicationReader.cpp b/Engine/Source/Runtime/Experimental/Iris/Core/Private/Iris/ReplicationSystem/ReplicationReader.cpp
--- a/Engine/Source/Runtime/Experimental/Iris/Core/Private/Iris/ReplicationSystem/ReplicationReader.cpp
+++ b/Engine/Source/Runtime/Experimental/Iris/Core/Private/Iris/ReplicationSystem/ReplicationReader.cpp
@@ -12,7 +12,7 @@
 
 #include "Iris/DataStream/DataStream.h"
 
-#include "Iris/IrisConfigInternal.h"
+#include "Iris/IrisConfig.h"
 
 #include "Iris/ReplicationSystem/DeltaCompression/DeltaCompressionBaselineManager.h"
 #include "Iris/ReplicationSystem/DequantizeAndApplyHelper.h"
@@ -30,6 +30,7 @@
 #include "Iris/Serialization/NetBitStreamUtil.h"
 #include "Iris/Serialization/NetReferenceCollector.h"
 #include "Iris/Serialization/NetSerializer.h"
+#include "Iris/Serialization/ObjectNetSerializer.h"
 
 #include "Misc/CoreMisc.h"
 #include "Misc/ScopeExit.h"
@@ -149,6 +150,14 @@ static FAutoConsoleVariableRef CvarGracefullyHandleReachingEndOfBitstream(
 		bGracefullyHandleReachingEndOfBitstream,
 		TEXT("Gracefully handle reaching end of the ReplicationReader bitstream prior to reading as many object batches as expected. Default is true."));
 
+#if UE_NET_ASYNCLOADING_DEBUG
+static bool bForceObjectsToStall = false;
+static FAutoConsoleVariableRef CVarForceObjectsToStall(
+	TEXT("net.Iris.AsyncLoading.ForceObjectsToStall"),
+	bForceObjectsToStall,
+	TEXT("Force objects in the DebugObjectsForcedToStall list to never resolve their reference and simulate long async loading delays."));
+#endif
+
 static const FName NetError_FailedToFindAttachmentQueue("Failed to find attachment queue");
 
 //------------------------------------------------------------------------
@@ -386,7 +395,7 @@ void FReplicationReader::Deinit()
 {
 	for (FPendingBatchData& PendingBatchData : PendingBatches.PendingBatches)
 	{
-		UE_LOG(LogIris, Warning, TEXT("FReplicationReader::Deinit NetHandle %s has %d unprocessed data batches"), *PendingBatchData.Handle.ToString(), PendingBatchData.QueuedDataChunks.Num());
+		UE_LOG(LogIris, Warning, TEXT("FReplicationReader::Deinit NetHandle %s has %d unprocessed data batches"), *PendingBatchData.Owner.ToString(), PendingBatchData.QueuedDataChunks.Num());
 
 		// Make sure to release all references that we are holding on to
 		if (ObjectReferenceCache)
@@ -748,27 +757,36 @@ void FReplicationReader::DeserializeObjectStateDelta(FNetSerializationContext& C
 	}
 }
 
-FPendingBatchData* FReplicationReader::UpdateUnresolvedMustBeMappedReferences(FNetRefHandle InHandle, TArray<FNetRefHandle>& MustBeMappedReferences, EIrisAsyncLoadingPriority IrisAsyncLoadingPriority)
+FPendingBatchData* FReplicationReader::UpdateUnresolvedMustBeMappedReferences(FNetRefHandle OwnerHandle, TArray<FNetRefHandle>& MustBeMappedReferences, EIrisAsyncLoadingPriority InIrisAsyncLoadingPriority)
 {
-	FPendingBatchData* PendingBatch = PendingBatches.Find(InHandle);
-	// If we already have a pending batch we append any new must be mapped references to it.
+	FPendingBatchData* PendingBatch = PendingBatches.Find(OwnerHandle);
+
+	// Take the priority if it was sent with the MustBeMapped reference list
+	EIrisAsyncLoadingPriority OwnerAsyncLoadingPriority = InIrisAsyncLoadingPriority != EIrisAsyncLoadingPriority::Invalid ? InIrisAsyncLoadingPriority : EIrisAsyncLoadingPriority::Default;
+
+	// If we already have a pending batch, let's append the previous references to this new list.
 	if (PendingBatch)
 	{
-		for (const FNetRefHandle Ref : PendingBatch->PendingMustBeMappedReferences)
+		for (const FNetRefHandle PreviousRef : PendingBatch->PendingMustBeMappedReferences)
 		{
-			MustBeMappedReferences.AddUnique(Ref);
+			MustBeMappedReferences.AddUnique(PreviousRef);
+		}
+
+		if (PendingBatch->IrisAsyncLoadingPriority != EIrisAsyncLoadingPriority::Invalid)
+		{
+			OwnerAsyncLoadingPriority = PendingBatch->IrisAsyncLoadingPriority;
 		}
 	}
 
 	// Resolve
 	TArray<FNetRefHandle, TInlineAllocator<4>> Unresolved;
-	TArray<TPair<FNetRefHandle, UObject*>, TInlineAllocator<4>> QueuedObjectsToTrack;
+	TArray<TPair<FNetRefHandle, UObject*>, TInlineAllocator<4>> ResolvedReferencesToTrack;
 
 	Unresolved.Reserve(MustBeMappedReferences.Num());
-	QueuedObjectsToTrack.SetNum(MustBeMappedReferences.Num());
+	ResolvedReferencesToTrack.Reserve(MustBeMappedReferences.Num());
 
 	// Assign the priority of the current object
-	ResolveContext.AsyncLoadingPriority = ConvertAsyncLoadingPriority(IrisAsyncLoadingPriority);
+	ResolveContext.AsyncLoadingPriority = ConvertAsyncLoadingPriority(OwnerAsyncLoadingPriority);
 	ON_SCOPE_EXIT
 	{
 		// Revert to an invalid priority to detect paths that are not setting the priority
@@ -778,51 +796,52 @@ FPendingBatchData* FReplicationReader::UpdateUnresolvedMustBeMappedReferences(FN
 	for (FNetRefHandle Handle : MustBeMappedReferences)
 	{
 		UObject* ResolvedObject = nullptr;
+
 		// TODO: Report broken status in same call to avoid map lookup
 		ENetObjectReferenceResolveResult ResolveResult = ObjectReferenceCache->ResolveObjectReference(FObjectReferenceCache::MakeNetObjectReference(Handle), ResolveContext, ResolvedObject);
+
+#if UE_NET_ASYNCLOADING_DEBUG
+		if (bForceObjectsToStall && ReplicationBridge->DebugObjectsForcedToStall.Contains(OwnerHandle))
+		{
+			Unresolved.Add(Handle);
+			continue;
+		}
+#endif
+
 		if (EnumHasAnyFlags(ResolveResult, ENetObjectReferenceResolveResult::HasUnresolvedMustBeMappedReferences) && !ObjectReferenceCache->IsNetRefHandleBroken(Handle, true))
 		{
+			// The MustBeMapped reference does not exist yet
 			Unresolved.Add(Handle);
 		}
 		else if (ResolveResult == ENetObjectReferenceResolveResult::None)
 		{
-			QueuedObjectsToTrack.Emplace(Handle, ResolvedObject);
+			// The reference exists
+			ResolvedReferencesToTrack.Emplace(Handle, ResolvedObject);
 		}
 	}
 
+	// We still have references that are not loaded
 	if (Unresolved.Num())
 	{
-		FPendingBatchData* Batch = PendingBatch;
 		// We must create a new batch
-		if (!Batch)
+		if (PendingBatch == nullptr)
 		{
-			Batch = &PendingBatches.PendingBatches.AddDefaulted_GetRef();
-			Batch->Handle = InHandle;
-			Batch->PendingBatchStartCycles = FPlatformTime::Cycles64();
-			Batch->NextWarningTimeout = FPlatformTime::ToSeconds64(Batch->PendingBatchStartCycles) + QueuedBatchTimeoutSeconds;
+			PendingBatch = CreatePendingBatch(OwnerHandle);
 		}
 
-		// Update
-		Batch->PendingMustBeMappedReferences = MoveTemp(Unresolved);
-	
-		// If we resolved more references, add them to tracking list
-		for (TPair<FNetRefHandle, UObject*>& NetRefHandleObjectPair : QueuedObjectsToTrack)
-		{
-			if (!Batch->ResolvedReferences.Contains(NetRefHandleObjectPair.Key))
-			{
-				Batch->ResolvedReferences.Add(NetRefHandleObjectPair.Key);
-				ObjectReferenceCache->AddTrackedQueuedBatchObjectReference(NetRefHandleObjectPair.Key, NetRefHandleObjectPair.Value);
-			}
-		}
-
-		return Batch;
+		// Update the list of references we are waiting on
+		PendingBatch->PendingMustBeMappedReferences = MoveTemp(Unresolved);
 	}
 	else if (PendingBatch)
 	{
+		// Clear all references as they are all resolved now
 		PendingBatch->PendingMustBeMappedReferences.Reset();
+	}
 
-		// If we resolved more references, add them to tracking list
-		for (TPair<FNetRefHandle, UObject*>& NetRefHandleObjectPair : QueuedObjectsToTrack)
+	// If we resolved more references, add them to tracking list
+	if (PendingBatch)
+	{
+		for (TPair<FNetRefHandle, UObject*>& NetRefHandleObjectPair : ResolvedReferencesToTrack)
 		{
 			if (!PendingBatch->ResolvedReferences.Contains(NetRefHandleObjectPair.Key))
 			{
@@ -830,11 +849,37 @@ FPendingBatchData* FReplicationReader::UpdateUnresolvedMustBeMappedReferences(FN
 				ObjectReferenceCache->AddTrackedQueuedBatchObjectReference(NetRefHandleObjectPair.Key, NetRefHandleObjectPair.Value);
 			}
 		}
-		
-		return PendingBatch;
 	}
+
+	return PendingBatch;
+}
+
+FPendingBatchData* FReplicationReader::CreatePendingBatch(FNetRefHandle Owner)
+{
+	FPendingBatchData* NewPendingBatch = &PendingBatches.PendingBatches.AddDefaulted_GetRef();
 	
-	return nullptr;
+	NewPendingBatch->Owner = Owner;
+	NewPendingBatch->PendingBatchStartCycles = FPlatformTime::Cycles64();
+	NewPendingBatch->NextWarningTimeout = FPlatformTime::ToSeconds64(NewPendingBatch->PendingBatchStartCycles) + QueuedBatchTimeoutSeconds;
+	
+	return NewPendingBatch;
+}
+
+FPendingBatchData* FReplicationReader::FindOrCreatePendingBatch(FNetRefHandle Owner)
+{
+	FPendingBatchData* PendingBatch = PendingBatches.Find(Owner);
+	return PendingBatch ? PendingBatch : CreatePendingBatch(Owner);
+}
+
+bool FReplicationReader::DoesParentExist(FNetRefHandle ParentHandle) const
+{
+	if (!ParentHandle.IsValid())
+	{
+		// If the handle we received is somehow invalid, consider it as existing to prevent objects from being blocked infinitely.
+		return true;
+	}
+
+	return NetRefHandleManager->IsNetRefHandleAssigned(ParentHandle);
 }
 
 uint32 FReplicationReader::ReadObjectsInBatchWithoutSizes(FNetSerializationContext& Context, FNetRefHandle IncompleteHandle, bool bHasBatchOwnerData, uint32 BatchEndBitPosition)
@@ -996,7 +1041,8 @@ uint32 FReplicationReader::ReadObjectBatch(FNetSerializationContext& Context, ui
 
 	uint32 ObjectsReadInBatch = 0U;
 	
-	// A batch starts with (RefHandleId | BatchSize | bHasBatchObjectData | bHasExports)
+	// A batch starts with (RefHandleId | BatchSize | bHasBatchObjectData | bHasExports | bHasCreationDependencyHandles | ?CreationParentHandleList)
+
 	// If the batch has exports we must to seek to the end of the batch to read and process exports before reading/processing batch data
 	const FNetRefHandle IncompleteHandle = ReadNetRefHandleId(Context, Reader);
 
@@ -1026,14 +1072,28 @@ uint32 FReplicationReader::ReadObjectBatch(FNetSerializationContext& Context, ui
 
 	// Do we have exports or not?
 	const bool bHasExports = Reader.ReadBool();
-	
+
+	// Read the optional creation dependency handles
+	TArray<FNetRefHandle, TInlineAllocator<4>> CreationDependentParents;
+	{
+		bool bHasCreationDependencyHandles = Reader.ReadBool();
+		while (bHasCreationDependencyHandles)
+		{
+			CreationDependentParents.Emplace(ReadNetRefHandle(Context));
+			bHasCreationDependencyHandles = Reader.ReadBool();
+
+			//UE_LOG(LogIris, Verbose, TEXT("Reading CreationDependency for child: %s. Dependent Parent: %s %s"), *IncompleteHandle.ToString(), *CreationDependentParents.Last().ToString(), DoesParentExist(CreationDependentParents.Last())?TEXT("Exists"):TEXT("NotFound"));
+		}
+	}
+
 	uint32 ReadObjectCount = 0U;
 
 	// First we need to read exports, they are stored at the end of the batch
 	uint32 BatchEndPos = BatchEndOrStartOfExportsPos;
 
 	TempMustBeMappedReferences.Reset();
-	EIrisAsyncLoadingPriority IrisAsyncLoadingPriority = EIrisAsyncLoadingPriority::Default;
+	EIrisAsyncLoadingPriority IrisAsyncLoadingPriority = EIrisAsyncLoadingPriority::Invalid;
+
 	if (bHasExports)
 	{
 		const uint32 ReturnPos = Reader.GetPosBits();
@@ -1057,7 +1117,7 @@ uint32 FReplicationReader::ReadObjectBatch(FNetSerializationContext& Context, ui
 	}
 
 	// Skip over broken objects
-	const bool bIsBroken = BrokenObjects.FindByPredicate([IncompleteHandle](const FNetRefHandle& Entry) { return Entry.GetId() == IncompleteHandle.GetId(); } ) != nullptr;
+	const bool bIsBroken = BrokenObjects.FindByKey(IncompleteHandle) != nullptr;
 	if (bIsBroken)
 	{
 		UE_NET_TRACE_OBJECT_SCOPE(IncompleteHandle, Reader, Context.GetTraceCollector(), ENetTraceVerbosity::Trace);
@@ -1067,15 +1127,63 @@ uint32 FReplicationReader::ReadObjectBatch(FNetSerializationContext& Context, ui
 		return 0U;
 	}
 
-	FPendingBatchData* PendingBatch = PendingBatches.Find(IncompleteHandle);
+	// False if the object has a creation dependency and one of it's parent does not exist yet
+	bool bDoAllParentsExist = true;
+
+	if (CreationDependentParents.Num() > 0)
+	{
+		for (FNetRefHandle ParentHandle : CreationDependentParents)
+		{
+			const bool bParentExists = DoesParentExist(ParentHandle);
+
+			// Check if its a broken handle
+			if (!bParentExists && BrokenObjects.FindByKey(ParentHandle))
+			{
+				// The parent will never resolve so the object can never resolve also. Mark it broken so we don't queue up it's data indefinitely.
+
+				//TODO: Test and support an object removing a broken dependency and becoming resolveable again.
+
+				UE_LOG(LogIris, Error, TEXT("FReplicationReader::ReadObject RootObject: %s has a creation dependency on a broken parent: %s. RootObject will be marked broken too"), *IncompleteHandle.ToString(), *ParentHandle.ToString());
+					
+				BrokenObjects.AddUnique(IncompleteHandle);
+					
+				UE_NET_TRACE_OBJECT_SCOPE(IncompleteHandle, Reader, Context.GetTraceCollector(), ENetTraceVerbosity::Trace);
+				UE_NET_TRACE_SCOPE(SkippedData, Reader, Context.GetTraceCollector(), ENetTraceVerbosity::Trace);
+
+				Reader.Seek(BatchEndPos);
+				return 0U;
+			}
+
+			bDoAllParentsExist = bDoAllParentsExist && bParentExists;
+		}
+	}
+
+	FPendingBatchData* PendingBatchData = nullptr; 
+	
+	if (ObjectReferenceCache->ShouldAsyncLoad())
+	{
+		// Check if this object has pending must be mapped references that must be resolved before we can process the data.
+		PendingBatchData = UpdateUnresolvedMustBeMappedReferences(IncompleteHandle, TempMustBeMappedReferences, IrisAsyncLoadingPriority);
+	}
+	
+	// TODO: Test and support a case where the parents changed and now exist but we still have queued data waiting on the previous parents.
+	if (!bDoAllParentsExist)
+	{
+		if (!PendingBatchData)
+		{
+			PendingBatchData = FindOrCreatePendingBatch(IncompleteHandle);
+		}
+
+		PendingBatchData->CreationDependentParents = CreationDependentParents;
+	}
 
-	// This object has pending must be mapped references that must be resolved before we can process the data.
-	FPendingBatchData* PendingBatchData = ObjectReferenceCache->ShouldAsyncLoad() ? UpdateUnresolvedMustBeMappedReferences(IncompleteHandle, TempMustBeMappedReferences, IrisAsyncLoadingPriority) : nullptr;
+	// If the object's data cannot be processed and must be queued for later
 	if (PendingBatchData)
 	{
-		UE_LOG(LogIris, Verbose, TEXT("FReplicationReader::ReadObjectBatch Handle %s will be defered as it has %d unresolved must be mapped references"), *IncompleteHandle.ToString(), PendingBatchData->PendingMustBeMappedReferences.Num());
+		UE_LOG(LogIris, Verbose, TEXT("FReplicationReader::ReadObjectBatch Handle %s will be defered as it has %d unresolved references and/or %d dependent parents not yet created."), 
+			*IncompleteHandle.ToString(), PendingBatchData->PendingMustBeMappedReferences.Num(), PendingBatchData->CreationDependentParents.Num());
 
-		PendingBatchData->IrisAsyncLoadingPriority = IrisAsyncLoadingPriority;
+		PendingBatchData->IrisAsyncLoadingPriority = IrisAsyncLoadingPriority != EIrisAsyncLoadingPriority::Invalid ? IrisAsyncLoadingPriority : EIrisAsyncLoadingPriority::Default;
 
 		UE_NET_TRACE_OBJECT_SCOPE(IncompleteHandle, Reader, Context.GetTraceCollector(), ENetTraceVerbosity::Trace);
 		UE_NET_TRACE_SCOPE(QueuedBatch, Reader, Context.GetTraceCollector(), ENetTraceVerbosity::Trace);
@@ -1108,6 +1216,7 @@ uint32 FReplicationReader::ReadObjectBatch(FNetSerializationContext& Context, ui
 
 		PendingBatchData->QueuedDataChunks.Add(MoveTemp(DataChunk));
 	}
+	// Everything is good let's process the received data
 	else
 	{
 		ReadObjectCount = ReadObjectsInBatch(Context, IncompleteHandle, bHasBatchOwnerData, BatchEndOrStartOfExportsPos);
@@ -2393,8 +2502,22 @@ void FReplicationReader::ResolveAndDispatchUnresolvedReferences()
 	InternalObjectsToResolve.Append(ObjectsWithAttachmentPendingResolve);
 
 	// Try to resolve objects with updated references
-	for (uint32 InternalIndex : InternalObjectsToResolve)
+	for (FInternalNetRefIndex InternalIndex : InternalObjectsToResolve)
 	{
+
+#if UE_NET_ASYNCLOADING_DEBUG
+		// If we want to fake async loading, test if the owner can resolve objects here
+		if (bForceObjectsToStall && !ReplicationBridge->DebugObjectsForcedToStall.IsEmpty())
+		{
+			const FInternalNetRefIndex OwnerIndex = InternalIndex;
+			const FNetRefHandle OwnerHandle = NetRefHandleManager->GetNetRefHandleFromInternalIndex(OwnerIndex);
+			if (ReplicationBridge->DebugObjectsForcedToStall.Contains(OwnerHandle))
+			{
+				continue;
+			}
+		}
+#endif
+
 		ResolveAndDispatchUnresolvedReferencesForObject(Context, InternalIndex);
 	}
 
@@ -2609,7 +2732,7 @@ void FReplicationReader::ProcessHugeObjectAttachment(FNetSerializationContext& C
 
 bool FReplicationReader::EnqueueEndReplication(FPendingBatchData* PendingBatchData, bool bShouldDestroyInstance, bool bShouldProcessHierarchy, FNetRefHandle NetRefHandleToEndReplication)
 {
-	UE_LOG(LogIris, Verbose, TEXT("FReplicationReader::EnqueueEndReplication for %s since %s has queued batches"), *NetRefHandleToEndReplication.ToString(), *PendingBatchData->Handle.ToString());
+	UE_LOG(LogIris, Verbose, TEXT("FReplicationReader::EnqueueEndReplication for %s since %s has queued batches"), *NetRefHandleToEndReplication.ToString(), *PendingBatchData->Owner.ToString());
 		
 	const uint32 MaxNumDataBits = 66U;
 	const uint32 NumDataWords = (MaxNumDataBits + 31U) / 32U;
@@ -2684,23 +2807,41 @@ void FReplicationReader::ProcessQueuedBatches()
 	Context.SetInternalContext(&InternalContext);
 	Context.SetNetBlobReceiver(&ReplicationSystemInternal->GetNetBlobHandlerManager());
 
+	//TODO: Add a system to retry pending batches waiting on entities resolved this frame.
+
 	for (int BatchIt = 0; BatchIt < PendingBatches.PendingBatches.Num(); )
 	{
 		FPendingBatchData& PendingBatchData = PendingBatches.PendingBatches[BatchIt];
 
 		// Try to resolve remaining must be mapped references
 		TempMustBeMappedReferences.Reset();
-		UpdateUnresolvedMustBeMappedReferences(PendingBatchData.Handle, TempMustBeMappedReferences, PendingBatchData.IrisAsyncLoadingPriority);
+		UpdateUnresolvedMustBeMappedReferences(PendingBatchData.Owner, TempMustBeMappedReferences, PendingBatchData.IrisAsyncLoadingPriority);
+
+		int32 NumUnknownParents = 0;
+		
+		for (const FNetRefHandle CreationDependentParent : PendingBatchData.CreationDependentParents)
+		{
+			if (!DoesParentExist(CreationDependentParent))
+			{
+				++NumUnknownParents;
+			}
+		}
+		
+		const bool bAllParentsExist = NumUnknownParents == 0;
+		const bool bAllReferencesResolved = PendingBatchData.PendingMustBeMappedReferences.IsEmpty();
+
+		//UE_LOG(LogIris, Display, TEXT("ProcessQueuedBatches found %s to have %d unresolved references and %d non-existing parent"), 
+		//	*PendingBatchData.Owner.ToString(), PendingBatchData.PendingMustBeMappedReferences.Num(), NumUnknownParents);
 
-		// If we have no more pending must be references we can apply the received state
-		if (PendingBatchData.PendingMustBeMappedReferences.IsEmpty())		
+		// If we no longer waiting on references or the parent, let's process the data!
+		if (bAllParentsExist && bAllReferencesResolved)
 		{
-			UE_LOG(LogIris, Verbose, TEXT("ProcessQueuedBatches processing %d queued batches for Handle %s "), PendingBatchData.QueuedDataChunks.Num(), *PendingBatchData.Handle.ToString());
+			UE_LOG(LogIris, Verbose, TEXT("ProcessQueuedBatches processing %d queued batches for Handle %s "), PendingBatchData.QueuedDataChunks.Num(), *PendingBatchData.Owner.ToString());
 
 			// Reset for batch
 			Reader = FNetBitStreamReader();
 			Context.ResetErrorContext();
-			Context.SetErrorHandleContext(PendingBatchData.Handle);
+			Context.SetErrorHandleContext(PendingBatchData.Owner);
 			UE_RESET_READ_JOURNAL(Context);
 			UE_ADD_READ_JOURNAL_ENTRY(Context, TEXT("ProcessQueuedBatches"));
 
@@ -2748,12 +2889,12 @@ void FReplicationReader::ProcessQueuedBatches()
 						return Handle.GetId() == NetRefHandleToEndReplication.GetId();
 					}));
 
-					UE_LOG(LogIris, Verbose, TEXT("FReplicationReader::ProcessQueuedBatches EndReplication for %s while processing queued batches for %s"), *NetRefHandleToEndReplication.ToString(), *PendingBatchData.Handle.ToString());
+					UE_LOG(LogIris, Verbose, TEXT("FReplicationReader::ProcessQueuedBatches EndReplication for %s while processing queued batches for %s"), *NetRefHandleToEndReplication.ToString(), *PendingBatchData.Owner.ToString());
 					continue;
 				}
 
 				// Skip over broken objects, we still process remaining chunk if the object has been destroyed.
-				const bool bIsBroken = (BrokenObjects.FindByPredicate([&PendingBatchData](const FNetRefHandle& Entry) { return Entry.GetId() == PendingBatchData.Handle.GetId(); }) != nullptr);
+				const bool bIsBroken = (BrokenObjects.FindByPredicate([&PendingBatchData](const FNetRefHandle& Entry) { return Entry.GetId() == PendingBatchData.Owner.GetId(); }) != nullptr);
 				if (bIsBroken)
 				{
 					continue;
@@ -2772,18 +2913,18 @@ void FReplicationReader::ProcessQueuedBatches()
 				StreamDebugFeatures = CurrentChunk.StreamDebugFeatures;
 
 				// $IRIS: $TODO: Implement special dispatch to defer RepNotifies if we are processing multiple batches for the same object.
-				ReadObjectsInBatch(Context, PendingBatchData.Handle, CurrentChunk.bHasBatchOwnerData, CurrentChunk.NumBits);
+				ReadObjectsInBatch(Context, PendingBatchData.Owner, CurrentChunk.bHasBatchOwnerData, CurrentChunk.NumBits);
 
 				if (Context.HasErrorOrOverflow())
 				{
-					BrokenObjects.AddUnique(PendingBatchData.Handle);
+					BrokenObjects.AddUnique(PendingBatchData.Owner);
 
 					// Log error  if get more incoming data for an object in the broken state we will skip it.
-					UE_LOG(LogIris, Error, TEXT("FReplicationReader::ProcessQueuedBatches Failed to process object batch handle: %s skipping batch data. bIsEndReplication: %d"), ToCStr(PendingBatchData.Handle.ToString()), int(CurrentChunk.bIsEndReplicationChunk));
+					UE_LOG(LogIris, Error, TEXT("FReplicationReader::ProcessQueuedBatches Failed to process object batch handle: %s skipping batch data. bIsEndReplication: %d"), ToCStr(PendingBatchData.Owner.ToString()), int(CurrentChunk.bIsEndReplicationChunk));
 
 					if (Context.GetError() == GNetError_BrokenNetHandle)
 					{
-						ReplicationBridge->SendErrorWithNetRefHandle(UE::Net::ENetRefHandleError::ReplicationDisabled, PendingBatchData.Handle, Parameters.ConnectionId);
+						ReplicationBridge->SendErrorWithNetRefHandle(UE::Net::ENetRefHandleError::ReplicationDisabled, PendingBatchData.Owner, Parameters.ConnectionId);
 
 						Context.ResetErrorContext();
 					}
@@ -2826,50 +2967,75 @@ void FReplicationReader::ProcessQueuedBatches()
 				{
 					const double BlockedTime = FPlatformTime::ToSeconds64(FPlatformTime::Cycles64() - PendingBatchData.PendingBatchStartCycles);
 
-					TArray<FNetRefHandle> PendingMustBeMapped;
-					PendingMustBeMapped.Reserve(MaxMustBeMappedHandleArray);
-
-					for (const FNetRefHandle& AsyncLoadingRef : PendingBatchData.PendingMustBeMappedReferences)
+					if (!bAllParentsExist)
 					{
-						if (PendingMustBeMapped.Num() >= MaxMustBeMappedHandleArray)
-						{
-							// Waiting on too many assets already, ignore the rest
-							break;
-						}
+						UE_LOG(LogIris, Warning, TEXT("FReplicationWriter::ProcessQueuedBatches: Replication is blocked for longer than normal (%f secs) for %s, Waiting on %d/%d parents to be created.")
+							   , BlockedTime
+							   , *NetRefHandleManager->PrintObjectFromNetRefHandle(PendingBatchData.Owner)
+							   , NumUnknownParents
+							   , PendingBatchData.CreationDependentParents.Num());
 
-						// Ignore this asset if we recently warned about it
-						const double LastWarningTime = BlockedMustBeMappedLastWarningTime.FindOrAdd(AsyncLoadingRef, 0.0);
-						if (LastWarningTime == 0.0 || (CurrentTime - LastWarningTime) > QueuedBatchWarningRepeatTime)
+						for (const FNetRefHandle CreationDependentParent : PendingBatchData.CreationDependentParents)
 						{
-							PendingMustBeMapped.Emplace(AsyncLoadingRef);
-
-							BlockedMustBeMappedLastWarningTime[AsyncLoadingRef] = CurrentTime;
+							if (DoesParentExist(CreationDependentParent))
+							{
+								UE_LOG(LogIris, Warning, TEXT("FReplicationWriter::ProcessQueuedBatches: (valid) CreationDependent parent: %s exists."), *NetRefHandleManager->PrintObjectFromNetRefHandle(CreationDependentParent));
+							}
+							else
+							{
+								UE_LOG(LogIris, Warning, TEXT("FReplicationWriter::ProcessQueuedBatches: (wrong) CreationDependent parent: %s does not exist yet."), *NetRefHandleManager->PrintObjectFromNetRefHandle(CreationDependentParent));
+							}
 						}
 					}
 
-					// Warn and tell server about the blockage
-					if (!PendingMustBeMapped.IsEmpty())
+					if (!bAllReferencesResolved)
 					{
-						const FString BlockedObjectName = *NetRefHandleManager->PrintObjectFromNetRefHandle(PendingBatchData.Handle);
-
-						UE_LOG(LogIris, Warning, TEXT("FReplicationWriter::ProcessQueuedBatches: Replication is blocked for longer than normal for: %s, QueuedBunches: %i, PendingNetRefHandleResolves: %i, BlockedTime: %f")
-							, *BlockedObjectName
-							, PendingBatchData.QueuedDataChunks.Num()
-							, PendingBatchData.PendingMustBeMappedReferences.Num()
-							, BlockedTime);
+						TArray<FNetRefHandle> PendingMustBeMapped;
+						PendingMustBeMapped.Reserve(MaxMustBeMappedHandleArray);
 
-						for (const FNetRefHandle& NetRefHandle : PendingBatchData.PendingMustBeMappedReferences)
+						for (const FNetRefHandle& AsyncLoadingRef : PendingBatchData.PendingMustBeMappedReferences)
 						{
-							UE_LOG(LogIris, Log, TEXT("  NetRefHandle Pending resolve [%s]"), *ObjectReferenceCache->DescribeObjectReference(FObjectReferenceCache::MakeNetObjectReference(NetRefHandle), ResolveContext));
+							if (PendingMustBeMapped.Num() >= MaxMustBeMappedHandleArray)
+							{
+								// Waiting on too many assets already, ignore the rest
+								break;
+							}
+
+							// Ignore this asset if we recently warned about it
+							const double LastWarningTime = BlockedMustBeMappedLastWarningTime.FindOrAdd(AsyncLoadingRef, 0.0);
+							if (LastWarningTime == 0.0 || (CurrentTime - LastWarningTime) > QueuedBatchWarningRepeatTime)
+							{
+								PendingMustBeMapped.Emplace(AsyncLoadingRef);
+
+								BlockedMustBeMappedLastWarningTime[AsyncLoadingRef] = CurrentTime;
+							}
 						}
 
-#if IRIS_CLIENT_PROFILER_ENABLE
-						FClientProfiler::RecordBlockedReplication(*BlockedObjectName, PendingBatchData.PendingMustBeMappedReferences.Num(), (float)BlockedTime);
-#endif
+						// Warn and tell server about the blockage
+						if (!PendingMustBeMapped.IsEmpty())
+						{
+							const FString BlockedObjectName = *NetRefHandleManager->PrintObjectFromNetRefHandle(PendingBatchData.Owner);
 
-						ReplicationBridge->SendErrorWithNetRefHandle(UE::Net::ENetRefHandleError::BlockedByMustBeMapped, PendingBatchData.Handle, Parameters.ConnectionId, PendingMustBeMapped);
+							UE_LOG(LogIris, Warning, TEXT("FReplicationWriter::ProcessQueuedBatches: Replication is blocked for longer than normal for: %s, QueuedBunches: %i, PendingNetRefHandleResolves: %i, BlockedTime: %f")
+								   , *BlockedObjectName
+								   , PendingBatchData.QueuedDataChunks.Num()
+								   , PendingBatchData.PendingMustBeMappedReferences.Num()
+								   , BlockedTime);
+
+							for (const FNetRefHandle& NetRefHandle : PendingBatchData.PendingMustBeMappedReferences)
+							{
+								UE_LOG(LogIris, Log, TEXT("  NetRefHandle Pending resolve [%s]"), *ObjectReferenceCache->DescribeObjectReference(FObjectReferenceCache::MakeNetObjectReference(NetRefHandle), ResolveContext));
+							}
+
+							#if IRIS_CLIENT_PROFILER_ENABLE
+							FClientProfiler::RecordBlockedReplication(*BlockedObjectName, PendingBatchData.PendingMustBeMappedReferences.Num(), (float)BlockedTime);
+							#endif
+
+							ReplicationBridge->SendErrorWithNetRefHandle(UE::Net::ENetRefHandleError::BlockedByMustBeMapped, PendingBatchData.Owner, Parameters.ConnectionId, PendingMustBeMapped);
 						
+						}
 					}
+					
 
 					PendingBatchData.NextWarningTimeout = CurrentTime + QueuedBatchWarningRepeatTime;
 				}
diff --git a/Engine/Source/Runtime/Experimental/Iris/Core/Private/Iris/ReplicationSystem/ReplicationReader.h b/Engine/Source/Runtime/Experimental/Iris/Core/Private/Iris/ReplicationSystem/ReplicationReader.h
--- a/Engine/Source/Runtime/Experimental/Iris/Core/Private/Iris/ReplicationSystem/ReplicationReader.h
+++ b/Engine/Source/Runtime/Experimental/Iris/Core/Private/Iris/ReplicationSystem/ReplicationReader.h
@@ -219,7 +219,13 @@ private:
 	void DeserializeObjectStateDelta(FNetSerializationContext& Context, uint32 InternalIndex, FDispatchObjectInfo& Info, FReplicatedObjectInfo& ObjectInfo, const FNetRefHandleManager::FReplicatedObjectData& ObjectData, uint32& OutNewBaselineIndex);
 
 	// If async loading is enabled this function will verify if we can resolve all PendingMustBeMappedReferences
-	FPendingBatchData* UpdateUnresolvedMustBeMappedReferences(FNetRefHandle Handle, TArray<FNetRefHandle>& MustBeMappedReferences, EIrisAsyncLoadingPriority AsyncLoadingPriority);
+	FPendingBatchData* UpdateUnresolvedMustBeMappedReferences(FNetRefHandle OwnerHandle, TArray<FNetRefHandle>& MustBeMappedReferences, EIrisAsyncLoadingPriority InIrisAsyncLoadingPriority);
+
+	FPendingBatchData* CreatePendingBatch(FNetRefHandle Owner);
+	FPendingBatchData* FindOrCreatePendingBatch(FNetRefHandle Owner);
+
+	// Returns true if the object assigned to this handle has been instantiated locally
+	bool DoesParentExist(FNetRefHandle ObjectHandle) const;
 
 	// If we are queuing data for a batch we must also defer calls to EndReplication
 	// This method writes this method in the form of a QueuedChunk
@@ -256,7 +262,7 @@ private:
 
 	// Temporary buffers used by ResolveAndDispatchUnresolvedReferences().
 	TSet<FNetRefHandle> VisitedUnresolvedHandles;
-	TSet<uint32> InternalObjectsToResolve;
+	TSet<FInternalNetRefIndex> InternalObjectsToResolve;
 
 	// We track some data about incoming objects
 	// Stored in a map for now
@@ -270,7 +276,7 @@ private:
 	TArray<uint32> ObjectsWithAttachmentPendingResolve;
 
 	// Track all objects waiting for this handle to be resolvable
-	TMultiMap<FNetRefHandle, uint32> UnresolvedHandleToDependents;
+	TMultiMap<FNetRefHandle, FInternalNetRefIndex /*OwnerInternalIndex*/> UnresolvedHandleToDependents;
 	
 	// Track all objects with a dynamic handle in case it becomes unresolvable
 	TMultiMap<FNetRefHandle, uint32> ResolvedDynamicHandleToDependents;
diff --git a/Engine/Source/Runtime/Experimental/Iris/Core/Private/Iris/ReplicationSystem/ReplicationWriter.cpp b/Engine/Source/Runtime/Experimental/Iris/Core/Private/Iris/ReplicationSystem/ReplicationWriter.cpp
--- a/Engine/Source/Runtime/Experimental/Iris/Core/Private/Iris/ReplicationSystem/ReplicationWriter.cpp
+++ b/Engine/Source/Runtime/Experimental/Iris/Core/Private/Iris/ReplicationSystem/ReplicationWriter.cpp
@@ -1,32 +1,45 @@
 // Copyright Epic Games, Inc. All Rights Reserved.
 
 #include "Iris/ReplicationSystem/ReplicationWriter.h"
+
 #include "Containers/Array.h"
-#include "Iris/IrisConfigInternal.h"
+
+#include "HAL/IConsoleManager.h"
+
 #include "Iris/Core/IrisLog.h"
 #include "Iris/Core/IrisProfiler.h"
-#include "Misc/ScopeExit.h"
-#include "Net/Core/Trace/NetTrace.h"
-#include "Net/Core/Trace/NetDebugName.h"
+
+#include "Iris/IrisConfigInternal.h"
+
 #include "Iris/PacketControl/PacketNotification.h"
-#include "Iris/ReplicationSystem/ReplicationSystemTypes.h"
+
 #include "Iris/ReplicationSystem/DeltaCompression/DeltaCompressionBaselineManager.h"
 #include "Iris/ReplicationSystem/Filtering/ReplicationFiltering.h"
+#include "Iris/ReplicationSystem/NetBlob/NetObjectBlobHandler.h"
+#include "Iris/ReplicationSystem/NetBlob/PartialNetObjectAttachmentHandler.h"
 #include "Iris/ReplicationSystem/NetRefHandleManager.h"
 #include "Iris/ReplicationSystem/ReplicationOperations.h"
 #include "Iris/ReplicationSystem/ReplicationOperationsInternal.h"
 #include "Iris/ReplicationSystem/ReplicationSystem.h"
 #include "Iris/ReplicationSystem/ReplicationSystemInternal.h"
-#include "Iris/ReplicationSystem/NetBlob/NetObjectBlobHandler.h"
-#include "Iris/ReplicationSystem/NetBlob/PartialNetObjectAttachmentHandler.h"
-#include "Iris/Serialization/NetReferenceCollector.h"
+#include "Iris/ReplicationSystem/ReplicationSystemTypes.h"
+
 #include "Iris/Serialization/InternalNetSerializationContext.h"
 #include "Iris/Serialization/NetBitStreamUtil.h"
-#include "Iris/Serialization/NetSerializer.h"
 #include "Iris/Serialization/NetExportContext.h"
+#include "Iris/Serialization/NetReferenceCollector.h"
+#include "Iris/Serialization/NetSerializer.h"
+#include "Iris/Serialization/ObjectNetSerializer.h"
+
 #include "Iris/Stats/NetStatsContext.h"
-#include "HAL/IConsoleManager.h"
+
+#include "Misc/ScopeExit.h"
+
+#include "Net/Core/Trace/NetDebugName.h"
+#include "Net/Core/Trace/NetTrace.h"
+
 #include "ProfilingDebugging/CsvProfiler.h"
+
 #include <algorithm>
 #include <cmath> // std::nextafter
 
@@ -2052,7 +2065,7 @@ uint32 FReplicationWriter::WriteRootObjectsPendingDestroy(FNetSerializationConte
 		// We do not support destroying an object that is currently being sent as a huge object.
 		if (IsObjectPartOfActiveHugeObject(InternalIndex))
 		{
-			UE_LOG(LogIris, Verbose, TEXT("Skipping writing destroy for object ( InternalIndex: %u ) which is part of active huge object."), InternalIndex);
+			UE_LOG(LogIris, Verbose, TEXT("Skipping writing destroy for object %s which is part of active huge object."), *NetRefHandleManager->PrintObjectFromIndex(InternalIndex));
 			bWroteAllDestroyedObjects = false;
 			continue;
 		}
@@ -2216,7 +2229,7 @@ bool FReplicationWriter::CanSendObject(uint32 InternalIndex) const
 				// if we cannot send the initial dependent object we must wait until we can.
 				if (!CanSendObject(DependentInternalIndex))
 				{
-					UE_LOG(LogIris, Verbose, TEXT("ReplicationWriter: Cannot send internal index (%u) due to waiting on init dependency internal index (%d)"), InternalIndex, DependentInternalIndex);
+					UE_LOG(LogIris, Verbose, TEXT("ReplicationWriter: Cannot send: %s due to waiting on init dependency of: %s"), *NetRefHandleManager->PrintObjectFromIndex(InternalIndex), *NetRefHandleManager->PrintObjectFromIndex(DependentInternalIndex));
 					UE_NET_REPLICATIONWRITER_WARN_IF_CANNOT_SEND(TEXT("Waiting on init dependency %s "), *NetRefHandleManager->PrintObjectFromIndex(DependentInternalIndex));
 
 					return false;
@@ -2225,7 +2238,7 @@ bool FReplicationWriter::CanSendObject(uint32 InternalIndex) const
 				// if the dependent object are scheduled before parent and did not fit in this packet, we cannot write the parent either and have to wait until creation is confirmed
 				if ((DependentObjectInfo.SchedulingHint == EDependentObjectSchedulingHint::ScheduleBeforeParent) && ObjectsWithDirtyChanges.GetBit(DependentInternalIndex))
 				{
-					UE_LOG(LogIris, Verbose, TEXT("ReplicationWriter: Cannot send internal index (%u) due to waiting on ScheduleBefore dependency internal index (%d)"), InternalIndex, DependentInternalIndex);
+					UE_LOG(LogIris, Verbose, TEXT("ReplicationWriter: Cannot send: %s due to waiting on ScheduleBefore dependency of: %s"), *NetRefHandleManager->PrintObjectFromIndex(InternalIndex), *NetRefHandleManager->PrintObjectFromIndex(DependentInternalIndex));
 					UE_NET_REPLICATIONWRITER_WARN_IF_CANNOT_SEND(TEXT("Waiting on ScheduleBefore dependency %s "), *NetRefHandleManager->PrintObjectFromIndex(DependentInternalIndex));
 
 					return false;
@@ -2303,13 +2316,13 @@ FReplicationWriter::EWriteObjectStatus FReplicationWriter::WriteObjectAndSubObje
 	UE_NET_TRACE_NAMED_OBJECT_SCOPE(ObjectTraceScope, NetRefHandleForTraceScope, Writer, Context.GetTraceCollector(), ENetTraceVerbosity::Trace);
 #endif
 
-	// We only need to write batch info for root objects
-	const bool bWriteBatchInfo = !Info.IsSubObject;
+	// We only need to write batch header for the root object
+	const bool bWriteBatchHeader = !Info.IsSubObject;
 	uint32 InitialStateHeaderPos = 0U;
 	const uint32 NumBitsUsedForBatchSize = (WriteObjectFlags & EWriteObjectFlag::WriteObjectFlag_IsWritingHugeObjectBatch) == 0U ? Parameters.NumBitsUsedForBatchSize : Parameters.NumBitsUsedForHugeObjectBatchSize;
 
 	// This is the beginning of what we treat as a batch on the receiving end
-	if (bWriteBatchInfo)
+	if (bWriteBatchHeader)
 	{
 		// Write bit indicating that we are not a destruction info.
 		constexpr bool bIsDestructionInfo = false;
@@ -2317,7 +2330,7 @@ FReplicationWriter::EWriteObjectStatus FReplicationWriter::WriteObjectAndSubObje
 
 		WriteSentinel(&Writer, TEXT("RootObjectBatchInfo"));
 
-		// A batch starts with (RefHandleId | BatchSize | bHasBatchObjectData | bHasExports)
+		// A batch starts with (RefHandleId | BatchSize | bHasBatchObjectData | bHasExports | bHasCreationDependencyHandles | ?CreationParentHandleList)
 		// We write the header up front, and then we seek back and update relevant info is the object + subobjects is successfully serialized along with necessary exports
 
 		// We send the Index of the handle to the remote end
@@ -2331,11 +2344,48 @@ FReplicationWriter::EWriteObjectStatus FReplicationWriter::WriteObjectAndSubObje
 		}
 
 		// Did we write serialize any data related to batch owner
-		Writer.WriteBool(false);
+		constexpr bool bHasBatchOwnerData = false;
+		Writer.WriteBool(bHasBatchOwnerData);
 
 		// If the batch has exports, they are at the end of the batch
 		// We handle this on the reading side to avoid rewriting the entire object to insert exports up front.	
-		Writer.WriteBool(false);
+		constexpr bool bHasExports = false;
+		Writer.WriteBool(bHasExports);
+
+		// Write the object creation dependencies
+		//TODO: Only send the CreationDependency data when it's an InitState creation data payload
+
+		TConstArrayView<const FInternalNetRefIndex> CreationDependencies = NetRefHandleManager->GetCreationDependencies(InternalIndex);
+
+		if (CreationDependencies.IsEmpty())
+		{
+			constexpr bool bHasCreationDependencyHandles = false;
+			Writer.WriteBool(bHasCreationDependencyHandles);
+		}
+		else
+		{
+			for (FInternalNetRefIndex ParentInternalIndex : CreationDependencies)
+			{
+				constexpr bool bHasCreationDependencyHandles = true;
+				Writer.WriteBool(bHasCreationDependencyHandles);
+
+				const FNetRefHandle ParentHandle = NetRefHandleManager->GetNetRefHandleFromInternalIndex(ParentInternalIndex);
+				ensureMsgf(ParentHandle.IsValid(), TEXT("Invalid parent found in CreationDependency of child: %s. Stored parent was:  (InternalIndex: %u)"), *NetRefHandleManager->PrintObjectFromIndex(InternalIndex), ParentInternalIndex);
+
+				//UE_LOG(LogIris, Verbose, TEXT("Writing CreationDependency to child: %s | Parent: %s"), *NetRefHandleManager->PrintObjectFromIndex(InternalIndex), *NetRefHandleManager->PrintObjectFromIndex(ParentInternalIndex));
+				
+				WriteNetRefHandle(Context, ParentHandle);
+			}
+
+			// Flag that this was the last handle
+			constexpr bool bHasCreationDependencyHandles = false;
+			Writer.WriteBool(bHasCreationDependencyHandles);
+		}
+
+		if (Writer.IsOverflown())
+		{
+			return EWriteObjectStatus::BitStreamOverflow;
+		}
 	}
 	
 	// Create a temporary batch entry. We don't want to push it to the batch info unless we're successful.
@@ -2630,7 +2680,7 @@ FReplicationWriter::EWriteObjectStatus FReplicationWriter::WriteObjectAndSubObje
 						UE_LOG(LogIris, Verbose, TEXT("Failed to write huge attachment for object %s ( InternalIndex: %u ), forcing fallback on hugeobject for attachments"), *NetRefHandle.ToString(), InternalIndex);
 						Writer.DoOverflow();
 					}
-					else if (!(BatchEntry.bSentState || bSentTearOff || Info.SubObjectPendingDestroy || Info.HasDirtySubObjects || bWriteBatchInfo))
+					else if (!(BatchEntry.bSentState || bSentTearOff || Info.SubObjectPendingDestroy || Info.HasDirtySubObjects || bWriteBatchHeader))
 					{
 						// If we didn't send state and didn't send any attachments let's rollback
 						ObjectRollbackScope.Rollback();
@@ -2685,7 +2735,7 @@ FReplicationWriter::EWriteObjectStatus FReplicationWriter::WriteObjectAndSubObje
 				FNetBitStreamWriteScope WriteScope(Writer, ReplicatedDestroyHeaderBitPos);
 				Writer.WriteBits(ReplicatedDestroyHeaderFlags, ReplicatedDestroyHeaderFlags_BitCount);
 			}
-			else if (bWriteBatchInfo && !(BatchEntry.bSentState || BatchEntry.bSentAttachments))
+			else if (bWriteBatchHeader && !(BatchEntry.bSentState || BatchEntry.bSentAttachments))
 			{
 				// No need for the destroy header as we did not write any data at all for the batch.
 				Writer.Seek(ReplicatedDestroyHeaderBitPos);
@@ -2805,7 +2855,7 @@ FReplicationWriter::EWriteObjectStatus FReplicationWriter::WriteObjectAndSubObje
 			const bool bWroteExports = WriteExportResult == FObjectReferenceCache::EWriteExportsResult::WroteExports;
 
 			// Update header
-			if (ensure(bWriteBatchInfo))
+			if (ensure(bWriteBatchHeader))
 			{
 				FNetBitStreamWriteScope SizeScope(Writer, InitialStateHeaderPos);
 				Writer.WriteBits(WrittenBitsInBatch, NumBitsUsedForBatchSize);
@@ -3370,10 +3420,13 @@ uint32 FReplicationWriter::WriteObjects(FNetSerializationContext& Context)
 				checkSlow(InternalIndex != ObjectIndexForOOBAttachment);
 				ensureMsgf(GetReplicationInfo(InternalIndex).GetState() != EReplicatedObjectState::Invalid, TEXT("DependentObject with InternalIndex %u is not in scope"), InternalIndex);
 
+				//UE_LOG(LogIris, Verbose, TEXT("Writing DependentObjectsPendingSend: %s"), *NetRefHandleManager->PrintObjectFromIndex(InternalIndex));
+
 				const EWriteStatus Status = SendObjectFunction(InternalIndex);
 
 				if (Status == EWriteStatus::Abort)
 				{
+					//UE_LOG(LogIris, Verbose, TEXT("Aborted writing of DependentObjectsPendingSend: %s. Added back to the list."), *NetRefHandleManager->PrintObjectFromIndex(InternalIndex));
 					// If we fail, we put the object back on the pending send stack and try again in the next packet of the batch
 					// The reason for pop before using the index is that the SendObjectFunction will push new dependent objects on the stack
 					WriteContext.DependentObjectsPendingSend.Push(InternalIndex);
@@ -3384,6 +3437,7 @@ uint32 FReplicationWriter::WriteObjects(FNetSerializationContext& Context)
 				// If we overflown and could not fit the object in the packet
 				if (Status == EWriteStatus::Skipped)
 				{
+					//UE_LOG(LogIris, Verbose, TEXT("Skipping over DependentObjectsPendingSend: %s"), *NetRefHandleManager->PrintObjectFromIndex(InternalIndex));
 					// TODO: We need to ensure it's ok to skip dependant objects and wait for the next frame to resend them if want want Dependent objects to be truly atomic
 				}
 			}
@@ -3395,10 +3449,13 @@ uint32 FReplicationWriter::WriteObjects(FNetSerializationContext& Context)
 
 				checkSlow(InternalIndex != ObjectIndexForOOBAttachment);
 
+				//UE_LOG(LogIris, Verbose, TEXT("Writing regular object: %s"), *NetRefHandleManager->PrintObjectFromIndex(InternalIndex));
+
 				const EWriteStatus Status = SendObjectFunction(InternalIndex);
 				
 				if (Status == EWriteStatus::Abort)
 				{
+					//UE_LOG(LogIris, Verbose, TEXT("Aborted writing of regular object: %s"), *NetRefHandleManager->PrintObjectFromIndex(InternalIndex));
 					bContinue = false;
 					break;
 				}
@@ -3407,6 +3464,7 @@ uint32 FReplicationWriter::WriteObjects(FNetSerializationContext& Context)
 				{
 					// Continue trying to send objects as we might find an object with a smaller footprint that fits in the available buffer space.
 					// Note that skipped objects won't be reserialized again even if the more packets are available after this batch.
+					//UE_LOG(LogIris, Verbose, TEXT("Skipped writing of regular object: %s"), *NetRefHandleManager->PrintObjectFromIndex(InternalIndex));
 				}
 
 				++ObjectListIt;
@@ -3930,6 +3988,8 @@ UDataStream::EWriteResult FReplicationWriter::Write(FNetSerializationContext& Co
 		// Trigger rollback as we did not write any data
 		Writer.DoOverflow();
 		WriteResult = UDataStream::EWriteResult::NoData;
+
+		UE_CLOG(WrittenObjectCount > 0, LogIris, Verbose, TEXT("Packet overflow when writing"));
 	}
 
 	// Report packet stats
diff --git a/Engine/Source/Runtime/Experimental/Iris/Core/Public/Iris/IrisConfig.h b/Engine/Source/Runtime/Experimental/Iris/Core/Public/Iris/IrisConfig.h
--- a/Engine/Source/Runtime/Experimental/Iris/Core/Public/Iris/IrisConfig.h
+++ b/Engine/Source/Runtime/Experimental/Iris/Core/Public/Iris/IrisConfig.h
@@ -39,3 +39,8 @@ IRISCORE_API EReplicationSystem GetUseIrisReplicationCmdlineValue();
 #ifndef UE_NET_THREAD_SAFETY_CHECK
 #	define UE_NET_THREAD_SAFETY_CHECK DO_CHECK
 #endif
+
+/** Enables code to simulate async loading stalls on specific objects on the client */
+#ifndef UE_NET_ASYNCLOADING_DEBUG
+#	define UE_NET_ASYNCLOADING_DEBUG !(UE_BUILD_SHIPPING || UE_BUILD_TEST)
+#endif
\ No newline at end of file
diff --git a/Engine/Source/Runtime/Experimental/Iris/Core/Public/Iris/ReplicationSystem/ObjectReferenceTypes.h b/Engine/Source/Runtime/Experimental/Iris/Core/Public/Iris/ReplicationSystem/ObjectReferenceTypes.h
--- a/Engine/Source/Runtime/Experimental/Iris/Core/Public/Iris/ReplicationSystem/ObjectReferenceTypes.h
+++ b/Engine/Source/Runtime/Experimental/Iris/Core/Public/Iris/ReplicationSystem/ObjectReferenceTypes.h
@@ -20,6 +20,8 @@ enum class EIrisAsyncLoadingPriority : uint8
 	VeryHigh,
 	/** Maximum possible value to the calculate minimum bits needed to serialize the enum */
 	Max = VeryHigh,
+	/** Don't use this in the config. Used in code to tell if a value was read and set. */
+	Invalid,
 };
 
 static uint32 GetIrisAsyncLoadingPriorityBits()
@@ -48,6 +50,11 @@ static const TCHAR* LexToString(EIrisAsyncLoadingPriority Priority)
 			return TEXT("VeryHigh");
 		} break;
 
+		case EIrisAsyncLoadingPriority::Invalid:
+		{
+			return TEXT("Invalid");
+		} break;
+
 		default:
 		{
 			ensure(false);
diff --git a/Engine/Source/Runtime/Experimental/Iris/Core/Public/Iris/ReplicationSystem/ObjectReplicationBridge.h b/Engine/Source/Runtime/Experimental/Iris/Core/Public/Iris/ReplicationSystem/ObjectReplicationBridge.h
--- a/Engine/Source/Runtime/Experimental/Iris/Core/Public/Iris/ReplicationSystem/ObjectReplicationBridge.h
+++ b/Engine/Source/Runtime/Experimental/Iris/Core/Public/Iris/ReplicationSystem/ObjectReplicationBridge.h
@@ -147,8 +147,11 @@ public:
 	/** Remove dependent object from parent. The dependent object will function as a standard standalone replicated object. */
 	IRISCORE_API void RemoveDependentObject(FNetRefHandle Parent, FNetRefHandle DependentObject);
 
-	/** Make an object dependent but it can only be relevant if it's parent is also relevant to the connection.  */
-	IRISCORE_API void AddDependentObjectWithParentRelevancy(FNetRefHandle Parent, FNetRefHandle DependentObject);
+	/** Force the child to only be created after the parent exists on a client. This also forces the child to never be relevant if their parents are not also relevant at the same time. */
+	IRISCORE_API void AddCreationDependencyLink(FNetRefHandle Parent, FNetRefHandle Child);
+
+	/** Remove a creation dependency link on the child */
+	IRISCORE_API void RemoveCreationDependencyLink(FNetRefHandle Parent, FNetRefHandle Child);
 
 	// Dormancy support
 
diff --git a/Engine/Source/Runtime/Experimental/Iris/Core/Public/Iris/ReplicationSystem/ReplicationBridge.h b/Engine/Source/Runtime/Experimental/Iris/Core/Public/Iris/ReplicationSystem/ReplicationBridge.h
--- a/Engine/Source/Runtime/Experimental/Iris/Core/Public/Iris/ReplicationSystem/ReplicationBridge.h
+++ b/Engine/Source/Runtime/Experimental/Iris/Core/Public/Iris/ReplicationSystem/ReplicationBridge.h
@@ -9,6 +9,8 @@
 
 #include "Iris/Core/NetObjectReference.h"
 
+#include "Iris/IrisConfig.h"
+
 #include "Iris/ReplicationSystem/NetObjectFactory.h"
 #include "Iris/ReplicationSystem/NetObjectFactoryRegistry.h"
 #include "Iris/ReplicationSystem/NetObjectGroupHandle.h"
@@ -25,10 +27,6 @@
 
 #include "ReplicationBridge.generated.h"
 
-#ifndef UE_NET_ASYNCLOADING_DEBUG
-	#define UE_NET_ASYNCLOADING_DEBUG !(UE_BUILD_SHIPPING || UE_BUILD_TEST)
-#endif
-
 class UObjectReplicationBridge;
 class UReplicationSystem;
 class UNetDriver;
