Did some minor fixes while investigating query perf:
- The dynamic aabb tree was constructed with a dirty tree (passed the wrong bool through)
- The queries were incorrectly doubly transforming the face normal

diff --git a/Engine/Source/Runtime/Experimental/Chaos/Private/Chaos/PBDRigidsEvolution.cpp b/Engine/Source/Runtime/Experimental/Chaos/Private/Chaos/PBDRigidsEvolution.cpp
--- a/Engine/Source/Runtime/Experimental/Chaos/Private/Chaos/PBDRigidsEvolution.cpp
+++ b/Engine/Source/Runtime/Experimental/Chaos/Private/Chaos/PBDRigidsEvolution.cpp
@@ -172,11 +172,14 @@ namespace Chaos
 				}
 				else if (BroadPhaseConfig.BroadphaseType == FBroadPhaseConfig::Tree || BroadPhaseConfig.BroadphaseType == FBroadPhaseConfig::TreeAndGrid)
 				{
+					bool bUseDirtyTree = false;
+					const int MaxNumToProcess = ForceFullBuild ? 0 : BroadPhaseConfig.IterationsPerTimeSlice;
 					if (bDynamicTree)
 					{
-						return MakeUnique<AABBDynamicTreeType>(Particles, BroadPhaseConfig.MaxChildrenInLeaf, BroadPhaseConfig.MaxTreeDepth, BroadPhaseConfig.MaxPayloadSize, ForceFullBuild ? 0 : BroadPhaseConfig.IterationsPerTimeSlice, true, bBuildOverlapCache);
+						return MakeUnique<AABBDynamicTreeType>(Particles, BroadPhaseConfig.MaxChildrenInLeaf, BroadPhaseConfig.MaxTreeDepth, BroadPhaseConfig.MaxPayloadSize, MaxNumToProcess, bDynamicTree, bUseDirtyTree, bBuildOverlapCache);
 					}
-					return MakeUnique<AABBTreeType>(Particles, BroadPhaseConfig.MaxChildrenInLeaf, BroadPhaseConfig.MaxTreeDepth, BroadPhaseConfig.MaxPayloadSize, ForceFullBuild ? 0 : BroadPhaseConfig.IterationsPerTimeSlice, false, AccelerationStructureUseDirtyTreeInsteadOfGrid == 1, bBuildOverlapCache);
+					bUseDirtyTree = AccelerationStructureUseDirtyTreeInsteadOfGrid == 1;
+					return MakeUnique<AABBTreeType>(Particles, BroadPhaseConfig.MaxChildrenInLeaf, BroadPhaseConfig.MaxTreeDepth, BroadPhaseConfig.MaxPayloadSize, MaxNumToProcess, bDynamicTree, bUseDirtyTree, bBuildOverlapCache);
 				}
 				else if (BroadPhaseConfig.BroadphaseType == FBroadPhaseConfig::TreeOfGridAndGrid || BroadPhaseConfig.BroadphaseType == FBroadPhaseConfig::TreeOfGrid)
 				{
diff --git a/Engine/Source/Runtime/Experimental/Chaos/Public/Chaos/GeometryQueries.h b/Engine/Source/Runtime/Experimental/Chaos/Public/Chaos/GeometryQueries.h
--- a/Engine/Source/Runtime/Experimental/Chaos/Public/Chaos/GeometryQueries.h
+++ b/Engine/Source/Runtime/Experimental/Chaos/Public/Chaos/GeometryQueries.h
@@ -316,7 +316,7 @@ namespace Chaos
 		{
 			OutNormal = ATM.TransformVectorNoScale(LocalNormal).GetSafeNormal(UE_KINDA_SMALL_NUMBER, FVec3::AxisVector(0));
 			OutPosition = ATM.TransformPositionNoScale(LocalPosition);
-			OutFaceNormal = ATM.TransformVectorNoScale(A.FindGeometryOpposingNormal(LocalDir, OutFaceIndex, OutNormal));
+			OutFaceNormal = ATM.TransformVectorNoScale(A.FindGeometryOpposingNormal(LocalDir, OutFaceIndex, LocalNormal));
 		}
 
 		return bResult;
