Fixed an issue where the IAX debug display was showing CDN cache hit statistics even when that info was not being tracked.

#rnx

- We don't track the CDN hit/miss rate in test or shipping config due to perf concerns but we were displaying that info in the 'ias.DisplayOnScreenStatistics' display which was confusing people.
- Added a new define UE_TRACK_CDN_HIT_STATUS which can be used to enable/disable CDN hit/miss tracking. By default this is set to be off for Test and Shipping configurations but can be overriden by Build.cs/Target.cs.
- Do not display any info about the CDN hit/miss rate if UE_TRACK_CDN_HIT_STATUS is 0.

diff --git a/Engine/Source/Runtime/Experimental/IoStore/OnDemand/Private/OnDemandHttpClient.cpp b/Engine/Source/Runtime/Experimental/IoStore/OnDemand/Private/OnDemandHttpClient.cpp
--- a/Engine/Source/Runtime/Experimental/IoStore/OnDemand/Private/OnDemandHttpClient.cpp
+++ b/Engine/Source/Runtime/Experimental/IoStore/OnDemand/Private/OnDemandHttpClient.cpp
@@ -10,6 +10,7 @@
 #include "IO/IoStoreOnDemand.h"
 #include "Misc/ConfigCacheIni.h"
 #include "Misc/Paths.h"
+#include "Statistics.h"
 
 #if UE_HTTPCLIENT_THREADSAFE_CHECKS
 	#define UE_CHECK_THREADSAFETY(MethodName) checkf(OwningThread == 0 || OwningThread == FPlatformTLS::GetCurrentThreadId() || !FGenericPlatformProcess::SupportsMultithreading(), #MethodName TEXT("was called on the on the wrong thread!"))
@@ -56,7 +57,7 @@ static int8 TrackCdnCacheStats(const HTTP::FResponse& Response)
 {
 	int8 Result = -1;
 
-#if !UE_BUILD_TEST
+#if UE_TRACK_CDN_HIT_STATUS
 	// All header fields are considered, with the assumption that later fields are
 	// more accurate than prior ones (e.g. a caching proxy sits infront if CDNs)
 	Response.ReadHeaders([&Result] (FAnsiStringView Key, FAnsiStringView Value)
@@ -67,14 +68,18 @@ static int8 TrackCdnCacheStats(const HTTP::FResponse& Response)
 		}
 		return true;
 	});
-#endif
+#endif //UE_TRACK_CDN_HIT_STATUS
 
 	return Result;
 }
 
 static const TCHAR* CDNCacheStatusToString(int8 Status)
 {
+#if UE_TRACK_CDN_HIT_STATUS
 	return Status > 0 ? TEXT("HIT") : Status == 0 ? TEXT("MISS") : TEXT("???");
+#else
+	return TEXT("???");
+#endif //UE_TRACK_CDN_HIT_STATUS
 }
 
 FIoStatus LoadDefaultHttpCertificates(bool& bWasLoaded)
diff --git a/Engine/Source/Runtime/Experimental/IoStore/OnDemand/Private/Statistics.cpp b/Engine/Source/Runtime/Experimental/IoStore/OnDemand/Private/Statistics.cpp
--- a/Engine/Source/Runtime/Experimental/IoStore/OnDemand/Private/Statistics.cpp
+++ b/Engine/Source/Runtime/Experimental/IoStore/OnDemand/Private/Statistics.cpp
@@ -801,6 +801,7 @@ void FOnDemandIoBackendStats::PrintOnScreenStatistics(FCoreDelegates::FSeverityM
 	}
 
 	// Print IAX CDN Cache rates
+#if UE_TRACK_CDN_HIT_STATUS
 	{
 		TStringBuilder<256> Builder;
 		Builder << TEXT("IAX CDN: Hit/Miss/NoHdr: ");
@@ -820,6 +821,7 @@ void FOnDemandIoBackendStats::PrintOnScreenStatistics(FCoreDelegates::FSeverityM
 
 		OutMessages.Add(FCoreDelegates::EOnScreenMessageSeverity::Info, FText::FromStringView(Builder.ToView()));
 	}
+#endif //#if UE_TRACK_CDN_HIT_STATUS
 
 	// Print IAS Cache
 	{
diff --git a/Engine/Source/Runtime/Experimental/IoStore/OnDemand/Private/Statistics.h b/Engine/Source/Runtime/Experimental/IoStore/OnDemand/Private/Statistics.h
--- a/Engine/Source/Runtime/Experimental/IoStore/OnDemand/Private/Statistics.h
+++ b/Engine/Source/Runtime/Experimental/IoStore/OnDemand/Private/Statistics.h
@@ -14,6 +14,14 @@
 #	define IAS_WITH_STATISTICS 0
 #endif
 
+#ifndef UE_TRACK_CDN_HIT_STATUS
+	#if UE_BUILD_TEST == 0 && UE_BUILD_SHIPPING == 0
+		#define UE_TRACK_CDN_HIT_STATUS 1
+	#else
+		#define UE_TRACK_CDN_HIT_STATUS 0
+	#endif
+#endif
+
 struct FAnalyticsEventAttribute;
 class FIoStatus;
 enum class EIoErrorCode;
