Fix a NetStats ensure by making sure we have enough ChildNetStatsContexts to serve all possible threads. The Scheduler allocates a number of backup threads defined by GTaskGraphOversubscriptionRatio which can be used dynamically whenever a task uses FOversubscriptionScope to indicate it's going to hold onto a thread whilst waiting, such as FGenericPlatformManualResetEvent.

#tests Built and tested against Dedicated Server and PIE
#rnx

diff --git a/Engine/Source/Runtime/Core/Public/Async/Fundamental/Scheduler.h b/Engine/Source/Runtime/Core/Public/Async/Fundamental/Scheduler.h
--- a/Engine/Source/Runtime/Core/Public/Async/Fundamental/Scheduler.h
+++ b/Engine/Source/Runtime/Core/Public/Async/Fundamental/Scheduler.h
@@ -123,6 +123,9 @@ namespace LowLevelTasks
 		//number of instantiated workers
 		inline uint32 GetNumWorkers() const;
 
+		//maximum number of workers, including standby workers (Oversubscription)
+		inline uint32 GetMaxNumWorkers() const;
+
 		//get the worker priority set when workers were started
 		inline EThreadPriority GetWorkerPriority() const { return WorkerPriority; }
 
@@ -197,6 +200,12 @@ namespace LowLevelTasks
 	{
 		return ActiveWorkers.load(std::memory_order_relaxed);
 	}
+	
+	// Return the maximum number of worker threads, including Standby Workers
+	inline uint32 FScheduler::GetMaxNumWorkers() const
+	{
+		return WorkerLocalQueues.Num();
+	}
 
 	inline bool FScheduler::WakeUpWorker(bool bBackgroundWorker)
 	{
diff --git a/Engine/Source/Runtime/Experimental/Iris/Core/Private/Iris/Stats/NetStats.cpp b/Engine/Source/Runtime/Experimental/Iris/Core/Private/Iris/Stats/NetStats.cpp
--- a/Engine/Source/Runtime/Experimental/Iris/Core/Private/Iris/Stats/NetStats.cpp
+++ b/Engine/Source/Runtime/Experimental/Iris/Core/Private/Iris/Stats/NetStats.cpp
@@ -7,7 +7,7 @@
 #include "HAL/IConsoleManager.h"
 #include "Misc/ScopeLock.h"
 #include "ProfilingDebugging/CsvProfiler.h"
-#include "Async/TaskGraphInterfaces.h"
+#include "Async/Fundamental/Scheduler.h"
 
 DEFINE_LOG_CATEGORY(LogNetStats);
 
@@ -133,10 +133,10 @@ void FNetTypeStats::Init(FInitParams& InitParams)
 void FNetTypeStats::PreUpdateSetup()
 {
 	//Create child contexts if needed
-	FTaskGraphInterface& TaskGraphInterface = FTaskGraphInterface::Get();
-	const int32 MaxParallelThreads = TaskGraphInterface.GetNumWorkerThreads() + 1;//Include GameThread as TaskGraph can also run there in addition to Worker Threads
+	LowLevelTasks::FScheduler& Scheduler = LowLevelTasks::FScheduler::Get();
+	const int32 MaxParallelThreads = Scheduler.GetMaxNumWorkers() + 1;//Include GameThread as Tasks can also run there in addition to Worker Threads
 	const int32 NewChildCount = MaxParallelThreads - ChildStatsContext.Num();
-	UE_CLOG(NewChildCount > 0, LogNetStats, Log, TEXT("FNetTypeStats Creating %i CreateChildNetStatsContexts"), NewChildCount);
+	UE_CLOG(NewChildCount > 0, LogNetStats, Log, TEXT("FNetTypeStats Creating %i CreateChildNetStatsContexts. Total: %i"), NewChildCount, MaxParallelThreads);
 	for (int32 i = 0; i < NewChildCount; ++i)
 	{
 		CreateChildNetStatsContext();
