Fixing physics missing collisions when auto welding children with different collision filters.
Preprefiltering for sim incorrectly used the union of the collision channel index.

diff --git a/Engine/Source/Runtime/Experimental/Chaos/Public/Chaos/Collision/SpatialAccelerationBroadPhase.h b/Engine/Source/Runtime/Experimental/Chaos/Public/Chaos/Collision/SpatialAccelerationBroadPhase.h
--- a/Engine/Source/Runtime/Experimental/Chaos/Public/Chaos/Collision/SpatialAccelerationBroadPhase.h
+++ b/Engine/Source/Runtime/Experimental/Chaos/Public/Chaos/Collision/SpatialAccelerationBroadPhase.h
@@ -238,11 +238,21 @@ namespace Chaos
 		 */
 		struct FSimOverlapVisitor
 		{
-			FSimOverlapVisitor(FGeometryParticleHandle* ParticleHandle, const FCollisionFilterData& InSimFilterData, Private::FBroadPhaseContext& InContext)
+			FSimOverlapVisitor(FGeometryParticleHandle* ParticleHandle, Private::FBroadPhaseContext& InContext)
 				: Context(InContext)
-				, SimFilterData(InSimFilterData)
 				, ParticleUniqueIdx(ParticleHandle ? ParticleHandle->UniqueIdx() : FUniqueIdx(0))
 				, AccelerationHandle(ParticleHandle)
+			{
+				// Should never be null. Also suppresses a warning about possible null ptr.
+				check(ParticleHandle != nullptr);
+				PRAGMA_DISABLE_INTERNAL_WARNINGS
+				FAccelerationStructureHandle::ComputeParticleSimFilterDataFromShapes(*ParticleHandle, SimFilterData);
+				PRAGMA_ENABLE_INTERNAL_WARNINGS
+			}
+
+			UE_DEPRECATED(5.7, "Use the constructor that does not take the sim filter.")
+			FSimOverlapVisitor(FGeometryParticleHandle* ParticleHandle, const FCollisionFilterData& InSimFilterData, Private::FBroadPhaseContext& InContext)
+				: FSimOverlapVisitor(ParticleHandle, InContext)
 			{
 			}
 
@@ -287,8 +297,20 @@ namespace Chaos
 				return false;
 			}
 
+			UE_INTERNAL bool PrePreFilter(const FAccelerationStructureHandle& Payload) const
+			{
+				if (ShouldIgnore(Payload))
+				{
+					return true;
+				}
+				PRAGMA_DISABLE_INTERNAL_WARNINGS
+				return Payload.PrePreSimFilter(SimFilterData);
+				PRAGMA_ENABLE_INTERNAL_WARNINGS
+			}
+
 		private:
 			Private::FBroadPhaseContext& Context;
+			// @todo(chaos): cache this on the particle?
 			FCollisionFilterData SimFilterData;
 			FUniqueIdx ParticleUniqueIdx; // unique id of the particle visiting, used to skip self intersection as early as possible
 
@@ -297,6 +319,14 @@ namespace Chaos
 		};
 	}
 
+	template <>
+	inline bool PrePreFilterHelper(const FAccelerationStructureHandle& Payload, const Private::FSimOverlapVisitor& Visitor)
+	{
+		PRAGMA_DISABLE_INTERNAL_WARNINGS
+		return Visitor.PrePreFilter(Payload);
+		PRAGMA_ENABLE_INTERNAL_WARNINGS
+	}
+
 	/**
 	 * A broad phase that iterates over particle and uses a spatial acceleration structure to output
 	 * potentially overlapping SpatialAccelerationHandles.
@@ -802,15 +832,11 @@ namespace Chaos
 			const bool bBody1Bounded = Particle1->HasBounds();
 			if (bBody1Bounded)
 			{
-				// @todo(chaos): cache this on the particle?
-				FCollisionFilterData ParticleSimData;
-				FAccelerationStructureHandle::ComputeParticleSimFilterDataFromShapes(*Particle1, ParticleSimData);
-
 				const FAABB3 Box1 = Particle1->WorldSpaceInflatedBounds();
 
 				{
 					PHYSICS_CSV_SCOPED_EXPENSIVE(PhysicsVerbose, DetectCollisions_BroadPhase);
-					Private::FSimOverlapVisitor OverlapVisitor(Particle1, ParticleSimData, Context);
+					Private::FSimOverlapVisitor OverlapVisitor(Particle1, Context);
 					InSpatialAcceleration.Overlap(Box1, OverlapVisitor);
 				}
 			}
diff --git a/Engine/Source/Runtime/Experimental/Chaos/Public/Chaos/ParticleHandle.h b/Engine/Source/Runtime/Experimental/Chaos/Public/Chaos/ParticleHandle.h
--- a/Engine/Source/Runtime/Experimental/Chaos/Public/Chaos/ParticleHandle.h
+++ b/Engine/Source/Runtime/Experimental/Chaos/Public/Chaos/ParticleHandle.h
@@ -272,6 +272,23 @@ public:
 		return false;
 	}
 
+	UE_INTERNAL bool PrePreSimFilter(const FCollisionFilterData& QuerierFilterData) const
+	{
+		if (bCanPrePreFilter)
+		{
+			const uint32 QuerierBlockMask = QuerierFilterData.Word1;
+			const uint32 QuerierChannelMask = QuerierFilterData.Word2;
+			const uint32 SimBlockMask = UnionSimFilterData.Word1;
+			const uint32 SimChannelMask = UnionSimFilterData.Word2;
+
+			// check if they can collide ( same logic as DoCollide in CollisionResolution.cpp )
+			const bool CanCollide = (QuerierChannelMask & SimBlockMask) && (SimChannelMask & QuerierBlockMask);
+			return !CanCollide;
+		}
+
+		return false;
+	}
+
 	void UpdateFrom(const FAccelerationStructureHandle& InOther)
 	{
 		UnionQueryFilterData.Word0 = InOther.UnionQueryFilterData.Word0;
@@ -295,10 +312,11 @@ public:
 
 	/**
 	* compute the aggregated sim collision filter from all associated shapes
-	* NOTE: may not initialize OutQueryFilterData, so it must be cleared before calling if necessary (e.g., when changing blocking channels)
+	* NOTE: may not initialize OutSimFilterData, so it must be cleared before calling if necessary (e.g., when changing blocking channels).
+	* NOTE: This sim filter is slightly different than the individual filters as word2 is stolen to be the decompressed channel bitmask instead of the component ID.
 	**/
 	template <typename TParticle>
-	static void ComputeParticleSimFilterDataFromShapes(const TParticle& Particle, FCollisionFilterData& OutSimFilterData);
+	UE_INTERNAL static void ComputeParticleSimFilterDataFromShapes(const TParticle& Particle, FCollisionFilterData& OutSimFilterData);
 
 private:
 	FGeometryParticle* ExternalGeometryParticle;
@@ -306,6 +324,8 @@ private:
 
 	FUniqueIdx CachedUniqueIdx;
 	FCollisionFilterData UnionQueryFilterData;
+	// Note: The union data is a binary or of all of the shapes in the particle. Some data in the normal sim filter is not valid if or'd, such as the channel index in word3, the component id in word2, and the body index in word0.
+	// In order to make pre-pre filtering work, the sim filter's word2 is replaced with the unpacked collision channel mask which can be or'd.
 	FCollisionFilterData UnionSimFilterData;
 	bool bCanPrePreFilter;
 
@@ -3959,15 +3979,23 @@ template <typename TParticle>
 /* static */ void FAccelerationStructureHandle::ComputeParticleSimFilterDataFromShapes(const TParticle& Particle, FCollisionFilterData& OutSimFilterData)
 {
 	const auto& Shapes = Particle.ShapesArray();
+	OutSimFilterData.Word2 = 0;
 	for (const auto& Shape : Shapes)
 	{
 		// In the simulation we can ignore query-only shapes, but need to know about sim- and probe-enabled shapes
 		if (Shape->GetSimEnabled() || Shape->GetIsProbe())
 		{
 			const FCollisionFilterData& ShapeSimData = Shape->GetSimData();
+
+			// Word3 contains the channel index, this cannot be or'd together. Instead, expand the index to the bitmask and or that. 
+			// No one currently uses word2 (the component ID) which would also be invalid to or, so re-use that memory.
+			uint8 UnusedMaskFilter;
+			const uint32 ShapeChannel = GetChaosCollisionChannelAndExtraFilter(ShapeSimData.Word3, UnusedMaskFilter);
+			const uint32 ShapeChannelMask = 1 << ShapeChannel;
+
 			OutSimFilterData.Word0 |= ShapeSimData.Word0;
 			OutSimFilterData.Word1 |= ShapeSimData.Word1;
-			OutSimFilterData.Word2 |= ShapeSimData.Word2;
+			OutSimFilterData.Word2 |= ShapeChannelMask;
 			OutSimFilterData.Word3 |= ShapeSimData.Word3;
 		}
 	}
@@ -3977,7 +4005,9 @@ template <typename TParticle>
 FORCEINLINE_DEBUGGABLE void FAccelerationStructureHandle::UpdatePrePreFilter(const TParticle& Particle)
 {
 	ComputeParticleQueryFilterDataFromShapes(Particle, UnionQueryFilterData);
+	PRAGMA_DISABLE_INTERNAL_WARNINGS
 	ComputeParticleSimFilterDataFromShapes(Particle, UnionSimFilterData);
+	PRAGMA_ENABLE_INTERNAL_WARNINGS
 	bCanPrePreFilter = true;
 }
 
