Iris
* Fix creation dependency filtering logic using the previous dependent list instead of the new creation list
* Added unit tests to cover basic creation dependency features
* Added AEngineTestAlwaysRelevant and AEngineTestSpatial basic actors to use in misc unit tests

#rnx

diff --git a/Engine/Source/Runtime/Experimental/Iris/Core/Private/Iris/ReplicationSystem/Filtering/ReplicationFiltering.cpp b/Engine/Source/Runtime/Experimental/Iris/Core/Private/Iris/ReplicationSystem/Filtering/ReplicationFiltering.cpp
--- a/Engine/Source/Runtime/Experimental/Iris/Core/Private/Iris/ReplicationSystem/Filtering/ReplicationFiltering.cpp
+++ b/Engine/Source/Runtime/Experimental/Iris/Core/Private/Iris/ReplicationSystem/Filtering/ReplicationFiltering.cpp
@@ -1351,16 +1351,16 @@ void FReplicationFiltering::PreUpdateDynamicFiltering()
 	}
 }
 
-void FReplicationFiltering::UpdateCreationDependentParent(uint32 ChildINdex, const FNetBitArrayView ObjectsWithCreationDependencies,  FNetBitArrayView OutConnectionObjectsInScope, bool bRecursive) const
+void FReplicationFiltering::UpdateCreationDependentParent(uint32 ChildIndex, const FNetBitArrayView ObjectsWithCreationDependencies,  FNetBitArrayView OutConnectionObjectsInScope, bool bRecursive) const
 {
-	const bool bStartedRelevant = OutConnectionObjectsInScope.IsBitSet(ChildINdex);
+	const bool bStartedRelevant = OutConnectionObjectsInScope.IsBitSet(ChildIndex);
 	if (!bStartedRelevant)
 	{
 		// No need to check it's parents, the object must be relevant by itself first.
 		return;
 	}
 
-	const TArrayView<const FInternalNetRefIndex> Parents = NetRefHandleManager->GetDependentObjectParents(ChildINdex);
+	TConstArrayView<const FInternalNetRefIndex> Parents = NetRefHandleManager->GetCreationDependencies(ChildIndex);
 
 	for (FInternalNetRefIndex ParentIndex : Parents)
 	{
@@ -1375,14 +1375,14 @@ void FReplicationFiltering::UpdateCreationDependentParent(uint32 ChildINdex, con
 		}
 
 		const bool bIsParentRelevant = OutConnectionObjectsInScope.IsBitSet(ParentIndex);
-		OutConnectionObjectsInScope.AndBitValue(ChildINdex, bIsParentRelevant);
+		OutConnectionObjectsInScope.AndBitValue(ChildIndex, bIsParentRelevant);
 	}
 
 	// Check if we are still relevant, if not clear our subobjects
-	const bool bIsStillRelevant = OutConnectionObjectsInScope.IsBitSet(ChildINdex);
+	const bool bIsStillRelevant = OutConnectionObjectsInScope.IsBitSet(ChildIndex);
 	if (!bRecursive && !bIsStillRelevant)
 	{
-		for (const FInternalNetRefIndex SubObjectIndex : NetRefHandleManager->GetSubObjects(ChildINdex))
+		for (const FInternalNetRefIndex SubObjectIndex : NetRefHandleManager->GetSubObjects(ChildIndex))
 		{
 			OutConnectionObjectsInScope.ClearBit(SubObjectIndex);
 		}
@@ -1668,10 +1668,10 @@ void FReplicationFiltering::UpdateDynamicFiltering()
 		{
 			IRIS_PROFILER_SCOPE(FReplicationFiltering_UpdateCreationDependencies);
 			FNetBitArrayView ConnectionObjectsInScope = MakeNetBitArrayView(ConnectionInfo.ObjectsInScope);
-			FNetBitArrayView::ForAllSetBits(ObjectsWithCreationDependencies, ConnectionObjectsInScope, FNetBitArrayView::AndOp, [&](uint32 DependentIndex)
+			FNetBitArrayView::ForAllSetBits(ObjectsWithCreationDependencies, ConnectionObjectsInScope, FNetBitArrayView::AndOp, [&](uint32 ChildIndex)
 			{
 				constexpr bool bIsNotRecursive = false;
-				this->UpdateCreationDependentParent(DependentIndex, ObjectsWithCreationDependencies, ConnectionObjectsInScope, bIsNotRecursive);
+				this->UpdateCreationDependentParent(ChildIndex, ObjectsWithCreationDependencies, ConnectionObjectsInScope, bIsNotRecursive);
 			});
 		}
 
