Extend the 'ChunkPlugin' option for IASTool so that we can now chunk IAS data in addition to IAD.

#rnx

- Moved the code for chunking container files to a new function ::ChunkContainer and code for writing out the .uondemandtoc files to another new function ::WriteOnDemandToc.
-- The moved code is relatively unchanged.
- This lets us gather one toc for installed ondemand content (IAD) and another for streamed ondemand content (IAS).
- At the moment we always create and write out the IAD toc. We always attempt to create the IAS toc but only write it out if we detect that at least one container was found.
-- In the future we should pass this info into the tool rather than relying on conventions.

diff --git a/Engine/Source/Runtime/Experimental/IoStore/OnDemand/Private/Tool/ChunkPlugin.cpp b/Engine/Source/Runtime/Experimental/IoStore/OnDemand/Private/Tool/ChunkPlugin.cpp
--- a/Engine/Source/Runtime/Experimental/IoStore/OnDemand/Private/Tool/ChunkPlugin.cpp
+++ b/Engine/Source/Runtime/Experimental/IoStore/OnDemand/Private/Tool/ChunkPlugin.cpp
@@ -231,7 +231,260 @@ static FIoStatus WriteChunk(const FString& Directory, FMemoryView Chunk, const F
 		<< TEXT("'");
 }
 
+static int32 ChunkContainer(const FString& ContainerUTocPath, const FChunkPluginSettings& Settings, const TMap<FGuid, FAES::FAESKey>& EncryptionKeys, FChunkWriterInterface* ChunkWriter, const FString& ChunksRelativeFolder, bool bIgnoreContainerHeader,
+	FOnDemandToc& OutOnDemandToc, TMap<FString, TArray<FOnDemandTocTagSetPackageList>>* OutFoundTagSets, TArray<FString>* OutFilesToDelete, FChunkPluginStats& OutStats)
+{
+	check(ChunkWriter != nullptr);
+
+	FIoStoreReader ContainerFileReader;
+	{
+		FIoStatus Status = ContainerFileReader.Initialize(*FPaths::ChangeExtension(ContainerUTocPath, TEXT("")), EncryptionKeys);
+		if (!Status.IsOk())
+		{
+			UE_LOG(LogIoStoreOnDemand, Error, TEXT("Failed to open container '%s' for reading"), *ContainerUTocPath);
+			return 0; // TODO - Should this be -1 and stop execution?
+		}
+	}
+
+	UE_LOG(LogIoStoreOnDemand, Display, TEXT("Processing container '%s'"), *ContainerUTocPath);
+
+	/*TArray<FString> ContainerFileList;
+	ContainerFileReader.GetFilenames(ContainerFileList);
+	for (const FString& File : ContainerFileList)
+	{
+		UE_LOG(LogIoStoreOnDemand, Display, TEXT("Found file in container '%s'"), *File);
+	}*/
+
+	const uint32 BlockSize = ContainerFileReader.GetCompressionBlockSize();
+	if (OutOnDemandToc.Header.BlockSize == 0)
+	{
+		OutOnDemandToc.Header.BlockSize = ContainerFileReader.GetCompressionBlockSize();
+	}
+
+	check(OutOnDemandToc.Header.BlockSize == ContainerFileReader.GetCompressionBlockSize());
+
+	const int32 ContainerIndex = OutOnDemandToc.Containers.AddDefaulted();
+	FOnDemandTocContainerEntry& ContainerEntry = OutOnDemandToc.Containers[ContainerIndex];
+	ContainerEntry.ContainerId = ContainerFileReader.GetContainerId();
+	ContainerEntry.ContainerName = FPaths::GetBaseFilename(ContainerUTocPath);
+
+	if (EnumHasAnyFlags(ContainerFileReader.GetContainerFlags(), EIoContainerFlags::Encrypted))
+	{
+		ContainerEntry.EncryptionKeyGuid = LexToString(ContainerFileReader.GetEncryptionKeyGuid());
+	}
+
+	if (Settings.PackageSets.Num() > 0 && OutFoundTagSets != nullptr)
+	{
+		TArray<FOnDemandTocTagSet> TagSets;
+
+		const FIoChunkId HeaderChunkId = CreateContainerHeaderChunkId(ContainerFileReader.GetContainerId());
+		TIoStatusOr<FIoBuffer> Status = ContainerFileReader.Read(HeaderChunkId, FIoReadOptions());
+		if (!Status.IsOk())
+		{
+			if (Status.Status().GetErrorCode() == EIoErrorCode::NotFound)
+			{
+				UE_LOG(LogIoStoreOnDemand, Display, TEXT("No header chunk for container '%s'"), *ContainerUTocPath);
+			}
+			else
+			{
+				UE_LOG(LogIoStoreOnDemand, Error, TEXT("Failed to read header chunk for container '%s'"), *ContainerUTocPath);
+				return -1;
+			}
+		}
+		else
+		{
+			FIoContainerHeader ContainerHeader;
+			FMemoryReaderView Ar(Status.ValueOrDie().GetView());
+			Ar << ContainerHeader;
+			Ar.Close();
+			if (Ar.IsError() || Ar.IsCriticalError())
+			{
+				UE_LOG(LogIoStoreOnDemand, Error, TEXT("Failed to deserialize header chunk for container '%s'"), *ContainerUTocPath);
+				return -1;
+			}
+
+			for (const TPair<FString, TArray<FString>>& Pair : Settings.PackageSets)
+			{
+				const FString& Tag = Pair.Key;
+
+				FOnDemandTocTagSetPackageList TagSetPackageList;
+				TagSetPackageList.ContainerIndex = ContainerIndex;
 
+				for (const FString& Package : Pair.Value)
+				{
+					const FPackageId PackageId = FPackageId::FromName(FName(Package));
+					int32 PackageIndex = INDEX_NONE;
+					if (ContainerHeader.PackageIds.Find(PackageId, PackageIndex))
+					{
+						UE_LOG(LogIoStoreOnDemand, Display, TEXT("Found package %s for tag %s in container %s"), *Package, *Tag, *ContainerUTocPath);
+						TagSetPackageList.PackageIndicies.Add(PackageIndex);
+					}
+				}
+
+				if (TagSetPackageList.PackageIndicies.Num() > 0)
+				{
+					OutFoundTagSets->FindOrAdd(Tag).Add(MoveTemp(TagSetPackageList));
+				}
+			}
+		}
+	}
+
+	TArray<FIoStoreTocChunkInfo> ChunkInfos;
+	ContainerFileReader.EnumerateChunks([&ChunkInfos](FIoStoreTocChunkInfo&& Info)
+		{
+			ChunkInfos.Emplace(MoveTemp(Info));
+			return true;
+		});
+
+	UE_LOG(LogIoStoreOnDemand, Display, TEXT("Serializing %d chunks"), ChunkInfos.Num());
+
+	for (const FIoStoreTocChunkInfo& ChunkInfo : ChunkInfos)
+	{
+		const bool bDecrypt = false;
+		TIoStatusOr<FIoStoreCompressedReadResult> Status = ContainerFileReader.ReadCompressed(ChunkInfo.Id, FIoReadOptions(), bDecrypt);
+		if (!Status.IsOk())
+		{
+			UE_LOG(LogIoStoreOnDemand, Error, TEXT("Failed to read container chunk, Container='%s', Reason='%s'"),
+				*ContainerUTocPath, *Status.Status().ToString());
+			return -1;
+		}
+
+		FIoStoreCompressedReadResult ReadResult = Status.ConsumeValueOrDie();
+
+		const uint32 BlockOffset = ContainerEntry.BlockSizes.Num();
+		const uint32 BlockCount = ReadResult.Blocks.Num();
+		const FIoHash ChunkHash = FIoHash::HashBuffer(ReadResult.IoBuffer.GetView());
+		const FString HashString = LexToString(ChunkHash);
+
+		FMemoryView EncodedBlocks = ReadResult.IoBuffer.GetView();
+		uint64 RawChunkSize = 0;
+		uint64 EncodedChunkSize = 0;
+		for (const FIoStoreCompressedBlockInfo& BlockInfo : ReadResult.Blocks)
+		{
+			check(Align(BlockInfo.CompressedSize, FAES::AESBlockSize) == BlockInfo.AlignedSize);
+			const uint64 EncodedBlockSize = BlockInfo.AlignedSize;
+			ContainerEntry.BlockSizes.Add(uint32(BlockInfo.CompressedSize));
+
+			FMemoryView EncodedBlock = EncodedBlocks.Left(EncodedBlockSize);
+			EncodedBlocks += EncodedBlock.GetSize();
+
+			check(((char*)EncodedBlock.GetData() + EncodedBlock.GetSize()) <= ((char*)ReadResult.IoBuffer.GetView().GetData() + ReadResult.IoBuffer.GetView().GetSize()));
+			
+			FIoBlockHash Hash = FIoChunkEncoding::HashBlock(EncodedBlock);
+			ContainerEntry.BlockHashes.Add(Hash);
+
+			EncodedChunkSize += EncodedBlockSize;
+			RawChunkSize += BlockInfo.UncompressedSize;
+
+			if (OutOnDemandToc.Header.CompressionFormat.IsEmpty() && BlockInfo.CompressionMethod != NAME_None)
+			{
+				OutOnDemandToc.Header.CompressionFormat = BlockInfo.CompressionMethod.ToString();
+			}
+		}
+
+		if (EncodedChunkSize != ReadResult.IoBuffer.GetSize())
+		{
+			UE_LOG(LogIoStoreOnDemand, Error, TEXT("Chunk size mismatch, Container='%s', ChunkId='%s'"),
+				*ContainerUTocPath, *LexToString(ChunkInfo.Id));
+			return -1;
+		}
+
+		{
+			FIoChunkDecodingParams DecodingParams;
+			DecodingParams.TotalRawSize = RawChunkSize;
+			DecodingParams.EncodedBlockSize = TConstArrayView<uint32>(ContainerEntry.BlockSizes.GetData() + BlockOffset, BlockCount);
+			DecodingParams.BlockHash = TConstArrayView<FIoBlockHash>(ContainerEntry.BlockHashes.GetData() + BlockOffset, BlockCount);
+
+			FIoBuffer OutputBuffer(RawChunkSize);
+			bool bDecoded = FIoChunkEncoding::Decode(DecodingParams, ReadResult.IoBuffer.GetView(), OutputBuffer.GetMutableView(), EIoDecodeFlags::None);
+			check(bDecoded);
+		}
+
+		FIoStatus WriteStatus = ChunkWriter->WriteChunk(ChunksRelativeFolder, ReadResult.IoBuffer, ChunkHash);
+		if (WriteStatus.IsOk() == false)
+		{
+			UE_LOG(LogIoStoreOnDemand, Error, TEXT("%s"), *WriteStatus.ToString());
+			return -1;
+		}
+
+		FOnDemandTocEntry& TocEntry = ContainerEntry.Entries.AddDefaulted_GetRef();
+		TocEntry.ChunkId = ChunkInfo.Id;
+		TocEntry.Hash = ChunkHash;
+		TocEntry.RawSize = RawChunkSize;
+		TocEntry.EncodedSize = EncodedChunkSize;
+		TocEntry.BlockOffset = BlockOffset;
+		TocEntry.BlockCount = BlockCount;
+
+		OutStats.TotalChunksSizeKb += (double)EncodedChunkSize / 1024.0;
+	}
+
+	if (bIgnoreContainerHeader)
+	{
+		UE_LOG(LogIoStoreOnDemand, Display, TEXT("Ignoring container header for '%s'"), *ContainerUTocPath)
+	}
+	else
+	{
+		const bool bDecrypt = false;
+		const FIoChunkId ChunkId = CreateContainerHeaderChunkId(ContainerFileReader.GetContainerId());
+		TIoStatusOr<FIoStoreCompressedReadResult> ReadResult = ContainerFileReader.ReadCompressed(ChunkId, FIoReadOptions(), bDecrypt);
+		if (ReadResult.IsOk())
+		{
+			const FIoBuffer& Chunk = ReadResult.ValueOrDie().IoBuffer;
+			ContainerEntry.Header = MakeArrayView<const uint8>(Chunk.GetData(), static_cast<int32>(Chunk.GetSize()));
+		}
+		else
+		{
+			const EIoErrorCode Err = ReadResult.Status().GetErrorCode();
+			if (Err != EIoErrorCode::UnknownChunkID && Err != EIoErrorCode::NotFound)
+			{
+				UE_LOG(LogIoStoreOnDemand, Error, TEXT("Failed to read container header, reason '%s'"), *ReadResult.Status().ToString())
+					return -1;
+			}
+		}
+	}
+
+	if (OutFilesToDelete != nullptr)
+	{
+		OutFilesToDelete->Add(ContainerUTocPath);
+		ContainerFileReader.GetContainerFilePaths(*OutFilesToDelete);
+	}
+
+	return 0;
+}
+
+int32 WriteOnDemandToc(FOnDemandToc& OnDemandToc, const FString& ContainerName, const FString& OutputFolder)
+{
+	IFileManager& FileMgr = IFileManager::Get();
+
+	FString Filename = FPathViews::SetExtension(ContainerName, TEXT(".uondemandtoc"));
+	Filename.ToLowerInline();
+
+	const FString TocPath = OutputFolder / Filename;
+
+	if (TUniquePtr<FArchive> Ar(FileMgr.CreateFileWriter(*TocPath)); Ar.IsValid())
+	{
+		*Ar << OnDemandToc;
+		Ar->Close();
+		if (Ar->IsError())
+		{
+			UE_LOG(LogIoStoreOnDemand, Error, TEXT("Failed to serialize TOC '%s'"), *TocPath);
+			return -1;
+		}
+		else
+		{
+			const int64 TocSize = Ar->Tell();
+			UE_LOG(LogIoStoreOnDemand, Display, TEXT("Writing file '%s' (%.2lf KiB)"),
+				*TocPath, double(TocSize) / 1024);
+
+			return 0;
+		}
+	}
+	else
+	{
+		UE_LOG(LogIoStoreOnDemand, Display, TEXT("Failed writing file '%s'"), *TocPath);
+		return -1;
+	}
+}
 
 ////////////////////////////////////////////////////////////////////////////////
 static int32 ChunkPluginCommandEntry(const FContext& Context)
@@ -380,224 +633,69 @@ static int32 ChunkPluginCommandEntry(const FContext& Context)
 	FileMgr.FindFiles(ContainerFilenames, *ContainerFolder, TEXT("*.utoc"));
 	UE_LOG(LogIoStoreOnDemand, Display, TEXT("Found %d container files(s)"), ContainerFilenames.Num());
 
+	// TODO: This info should be passed into the tool or always check for both
+	bool bSupportInstalledContainers = true; // IAD
+	bool bSupportStreamingContainers = true; // IAS
+
+	if (!bSupportInstalledContainers && !bSupportStreamingContainers)
+	{
+		UE_LOG(LogIoStoreOnDemand, Error, TEXT("Neither Installed or Streaming containers were requested"));
+		return -1;
+	}
+
 	FOnDemandToc OnDemandToc;
 	OnDemandToc.Header.Flags = static_cast<uint32>(TocFlags);
 	OnDemandToc.Header.HostGroupName = HostGroupName;
-	//OnDemandToc.Header.ChunksDirectory = TODO 
+	OnDemandToc.Header.ChunksDirectory = TEXT(""); // TODO?
 	OnDemandToc.Containers.Reserve(ContainerFilenames.Num());
 
 	TMap<FString, TArray<FOnDemandTocTagSetPackageList>> FoundTagSets;
 	FoundTagSets.Reserve(Settings.PackageSets.Num());
 
+	FOnDemandToc StreamingOnDemandToc;
+	StreamingOnDemandToc.Header.Flags = static_cast<uint32>(EOnDemandTocFlags::StreamOnDemand);
+	StreamingOnDemandToc.Header.HostGroupName = HostGroupName;
+	StreamingOnDemandToc.Header.ChunksDirectory = TEXT(""); // TODO?
+	StreamingOnDemandToc.Containers.Reserve(ContainerFilenames.Num());
+
 	TArray<FString> FilesToDelete;
 
 	for (const FString& Filename : ContainerFilenames)
 	{
 		const FString FullPath = ContainerFolder / Filename;
-		FIoStoreReader ContainerFileReader;
-		{
-			FIoStatus Status = ContainerFileReader.Initialize(*FPaths::ChangeExtension(FullPath, TEXT("")), EncryptionKeys);
-			if (!Status.IsOk())
-			{
-				UE_LOG(LogIoStoreOnDemand, Error, TEXT("Failed to open container '%s' for reading"), *FullPath);
-				continue;
-			}
-		}
-
-		UE_LOG(LogIoStoreOnDemand, Display, TEXT("Processing container '%s'"), *FullPath);
-
-		/*TArray<FString> ContainerFileList;
-		ContainerFileReader.GetFilenames(ContainerFileList);
-		for (const FString& File : ContainerFileList)
-		{
-			UE_LOG(LogIoStoreOnDemand, Display, TEXT("Found file in container '%s'"), *File);
-		}*/
-
-		const uint32 BlockSize = ContainerFileReader.GetCompressionBlockSize();
-		if (OnDemandToc.Header.BlockSize == 0)
-		{
-			OnDemandToc.Header.BlockSize = ContainerFileReader.GetCompressionBlockSize();
-		}
-		check(OnDemandToc.Header.BlockSize == ContainerFileReader.GetCompressionBlockSize());
-
-		const int32 ContainerIndex = OnDemandToc.Containers.AddDefaulted();
-		FOnDemandTocContainerEntry& ContainerEntry = OnDemandToc.Containers[ContainerIndex];
-		ContainerEntry.ContainerId = ContainerFileReader.GetContainerId();
-		ContainerEntry.ContainerName = FPaths::GetBaseFilename(FullPath);
-
-		if (EnumHasAnyFlags(ContainerFileReader.GetContainerFlags(), EIoContainerFlags::Encrypted))
-		{
-			ContainerEntry.EncryptionKeyGuid = LexToString(ContainerFileReader.GetEncryptionKeyGuid());
-		}
-
-		TArray<FOnDemandTocTagSet> TagSets;
-		if (Settings.PackageSets.Num() > 0)
-		{
-			const FIoChunkId HeaderChunkId = CreateContainerHeaderChunkId(ContainerFileReader.GetContainerId());
-			TIoStatusOr<FIoBuffer> Status = ContainerFileReader.Read(HeaderChunkId, FIoReadOptions());
-			if (!Status.IsOk())
-			{
-				if (Status.Status().GetErrorCode() == EIoErrorCode::NotFound)
-				{
-					UE_LOG(LogIoStoreOnDemand, Display, TEXT("No header chunk for container '%s'"), *FullPath);
-				}
-				else
-				{
-					UE_LOG(LogIoStoreOnDemand, Error, TEXT("Failed to read header chunk for container '%s'"), *FullPath);
-					return -1;
-				}
-			}
-			else
-			{
-				FIoContainerHeader ContainerHeader;
-				FMemoryReaderView Ar(Status.ValueOrDie().GetView());
-				Ar << ContainerHeader;
-				Ar.Close();
-				if (Ar.IsError() || Ar.IsCriticalError())
-				{
-					UE_LOG(LogIoStoreOnDemand, Error, TEXT("Failed to deserialize header chunk for container '%s'"), *FullPath);
-					return -1;
-				}
-
-				for (const TPair<FString, TArray<FString>>& Pair : Settings.PackageSets)
-				{
-					const FString& Tag = Pair.Key;
-
-					FOnDemandTocTagSetPackageList TagSetPackageList;
-					TagSetPackageList.ContainerIndex = ContainerIndex;
-
-					for (const FString& Package : Pair.Value)
-					{
-						const FPackageId PackageId = FPackageId::FromName(FName(Package));
-						int32 PackageIndex = INDEX_NONE;
-						if (ContainerHeader.PackageIds.Find(PackageId, PackageIndex))
-						{
-							UE_LOG(LogIoStoreOnDemand, Display, TEXT("Found package %s for tag %s in container %s"), *Package, *Tag, *Filename)
-							TagSetPackageList.PackageIndicies.Add(PackageIndex);
-						}
-					}
-
-					if (TagSetPackageList.PackageIndicies.Num() > 0)
-					{
-						FoundTagSets.FindOrAdd(Tag).Add(MoveTemp(TagSetPackageList));
-					}
-				}
-			}
-		}
-
-		TArray<FIoStoreTocChunkInfo> ChunkInfos;
-		ContainerFileReader.EnumerateChunks([&ChunkInfos](FIoStoreTocChunkInfo&& Info)
-		{ 
-			ChunkInfos.Emplace(MoveTemp(Info));
-			return true;
-		});
 
-		UE_LOG(LogIoStoreOnDemand, Display, TEXT("Serializing %d chunks"), ChunkInfos.Num());
+		TArray<FString>* FilesToDeletePtr = bDeleteContainerFiles ? &FilesToDelete : nullptr;
 
-		for (const FIoStoreTocChunkInfo& ChunkInfo : ChunkInfos)
+		// TODO: At the moment the only way we can identify an OnDemand container is by the name. Might be better to pass this info into the tool
+		if (bSupportStreamingContainers && FullPath.Contains(TEXT("OnDemand")))
 		{
-			const bool bDecrypt = false;
-			TIoStatusOr<FIoStoreCompressedReadResult> Status = ContainerFileReader.ReadCompressed(ChunkInfo.Id, FIoReadOptions(), bDecrypt);
-			if (!Status.IsOk())
+			// Attempt to add to IAS output container
+			int32 Result = ChunkContainer(FullPath, Settings, EncryptionKeys, ChunkWriter.Get(), ChunksRelativeFolder, bIgnoreContainerHeader, StreamingOnDemandToc, nullptr /*OutFoundTagSets*/, FilesToDeletePtr, Stats);
+			if (Result < 0)
 			{
-				UE_LOG(LogIoStoreOnDemand, Error, TEXT("Failed to read container chunk, Container='%s', Reason='%s'"),
-					*FullPath, *Status.Status().ToString());
-				return -1;
+				return Result;
 			}
-
-			FIoStoreCompressedReadResult ReadResult = Status.ConsumeValueOrDie();
-
-			const uint32 BlockOffset	= ContainerEntry.BlockSizes.Num();
-			const uint32 BlockCount		= ReadResult.Blocks.Num();
-			const FIoHash ChunkHash		= FIoHash::HashBuffer(ReadResult.IoBuffer.GetView());
-			const FString HashString	= LexToString(ChunkHash);
-
-			FMemoryView EncodedBlocks = ReadResult.IoBuffer.GetView();
-			uint64 RawChunkSize = 0;
-			uint64 EncodedChunkSize = 0;
-			for (const FIoStoreCompressedBlockInfo& BlockInfo : ReadResult.Blocks)
-			{
-				check(Align(BlockInfo.CompressedSize, FAES::AESBlockSize) == BlockInfo.AlignedSize);
-				const uint64 EncodedBlockSize = BlockInfo.AlignedSize;
-				ContainerEntry.BlockSizes.Add(uint32(BlockInfo.CompressedSize));
-
-				FMemoryView EncodedBlock = EncodedBlocks.Left(EncodedBlockSize);
-				EncodedBlocks += EncodedBlock.GetSize();
-				ContainerEntry.BlockHashes.Add(FIoChunkEncoding::HashBlock(EncodedBlock));
-
-				EncodedChunkSize += EncodedBlockSize;
-				RawChunkSize += BlockInfo.UncompressedSize;
-
-				if (OnDemandToc.Header.CompressionFormat.IsEmpty() && BlockInfo.CompressionMethod != NAME_None)
-				{
-					OnDemandToc.Header.CompressionFormat = BlockInfo.CompressionMethod.ToString();
-				}
-			}
-
-			if (EncodedChunkSize != ReadResult.IoBuffer.GetSize())
-			{
-				UE_LOG(LogIoStoreOnDemand, Error, TEXT("Chunk size mismatch, Container='%s', ChunkId='%s'"),
-					*FullPath, *LexToString(ChunkInfo.Id));
-				return -1;
-			}
-
-			FIoStatus WriteStatus = ChunkWriter->WriteChunk(ChunksRelativeFolder, ReadResult.IoBuffer, ChunkHash);
-			if (WriteStatus.IsOk() == false)
-			{
-				UE_LOG(LogIoStoreOnDemand, Error, TEXT("%s"), *WriteStatus.ToString());
-				return -1;
-			}
-
-			FOnDemandTocEntry& TocEntry = ContainerEntry.Entries.AddDefaulted_GetRef();
-			TocEntry.ChunkId = ChunkInfo.Id;
-			TocEntry.Hash = ChunkHash;
-			TocEntry.RawSize = RawChunkSize;
-			TocEntry.EncodedSize = EncodedChunkSize;
-			TocEntry.BlockOffset = BlockOffset;
-			TocEntry.BlockCount = BlockCount;
-
-			Stats.TotalChunksSizeKb += (double)EncodedChunkSize / 1024.0;
-		}
-
-
-		if (bIgnoreContainerHeader)
-		{
-			UE_LOG(LogIoStoreOnDemand, Display, TEXT("Ignoring container header for '%s'"), *Filename)
 		}
 		else
 		{
-			const bool bDecrypt = false;
-			const FIoChunkId ChunkId = CreateContainerHeaderChunkId(ContainerFileReader.GetContainerId());
-			TIoStatusOr<FIoStoreCompressedReadResult> ReadResult = ContainerFileReader.ReadCompressed(ChunkId, FIoReadOptions(), bDecrypt);
-			if (ReadResult.IsOk())
-			{
-				const FIoBuffer& Chunk = ReadResult.ValueOrDie().IoBuffer;
-				ContainerEntry.Header = MakeArrayView<const uint8>(Chunk.GetData(), static_cast<int32>(Chunk.GetSize()));
-			}
-			else
+			// Attempt to add to IAD output container
+			int32 Result = ChunkContainer(FullPath, Settings, EncryptionKeys, ChunkWriter.Get(), ChunksRelativeFolder, bIgnoreContainerHeader, OnDemandToc, &FoundTagSets, FilesToDeletePtr, Stats);
+			if (Result < 0)
 			{
-				const EIoErrorCode Err = ReadResult.Status().GetErrorCode();
-				if (Err != EIoErrorCode::UnknownChunkID && Err != EIoErrorCode::NotFound)
-				{
-					UE_LOG(LogIoStoreOnDemand, Error, TEXT("Failed to read container header, reason '%s'"), *ReadResult.Status().ToString())
-					return -1; 
-				}
+				return Result;
 			}
 		}
-
-		if (bDeleteContainerFiles)
-		{
-			FilesToDelete.Add(FullPath);
-			ContainerFileReader.GetContainerFilePaths(FilesToDelete);
-		}
 	}
 
-	OnDemandToc.TagSets.Reserve(FoundTagSets.Num());
-	for (TPair<FString, TArray<FOnDemandTocTagSetPackageList>>& Pair : FoundTagSets)
+	if (bSupportInstalledContainers)
 	{
-		FOnDemandTocTagSet& TocTagSet = OnDemandToc.TagSets.Emplace_GetRef();
-		TocTagSet.Tag = Pair.Key;
-		TocTagSet.Packages = MoveTemp(Pair.Value);
+		OnDemandToc.TagSets.Reserve(FoundTagSets.Num());
+		for (TPair<FString, TArray<FOnDemandTocTagSetPackageList>>& Pair : FoundTagSets)
+		{
+			FOnDemandTocTagSet& TocTagSet = OnDemandToc.TagSets.Emplace_GetRef();
+			TocTagSet.Tag = Pair.Key;
+			TocTagSet.Packages = MoveTemp(Pair.Value);
+		}
 	}
 
 	for (const FString& Path : FilesToDelete)
@@ -669,36 +767,35 @@ static int32 ChunkPluginCommandEntry(const FContext& Context)
 		}
 	}
 
-	OnDemandToc.Meta.EpochTimestamp = FDateTime::UtcNow().ToUnixTimestamp();
-	OnDemandToc.Meta.TargetPlatform = Platform;
-	OnDemandToc.Meta.BuildVersion	= BuildVersion;
+	const int64 Timestamp = FDateTime::UtcNow().ToUnixTimestamp();
+	// todo, we actually want the .uondemandtoc files in the base directory
+	// const FString TocPath = OutputFolder / Filename; // like this yo
+	const FString OnDemandTocOutputFolder = OutputFolder / IoStoreRelativeFolder;
+
+	if (bSupportInstalledContainers)
 	{
-		FString Filename = FPathViews::SetExtension(OnDemandTocName, TEXT(".uondemandtoc"));
-		Filename.ToLowerInline();
+		OnDemandToc.Meta.EpochTimestamp = Timestamp;
+		OnDemandToc.Meta.TargetPlatform = Platform;
+		OnDemandToc.Meta.BuildVersion = BuildVersion;
 
-		// todo, we actually want this file in the base directory
-		// const FString TocPath = OutputFolder / Filename; // like this yo
-		const FString TocPath = OutputFolder / IoStoreRelativeFolder / Filename;
-		if (TUniquePtr<FArchive> Ar(FileMgr.CreateFileWriter(*TocPath)); Ar.IsValid())
+		int32 Result = WriteOnDemandToc(OnDemandToc, OnDemandTocName, OnDemandTocOutputFolder);
+		if (Result < 0)
 		{
-			*Ar << OnDemandToc;
-			Ar->Close();
-			if (Ar->IsError())
-			{
-				UE_LOG(LogIoStoreOnDemand, Error, TEXT("Failed to serialize TOC '%s'"), *TocPath);
-				return -1;
-			}
-			else
-			{
-				const int64 TocSize = Ar->Tell();
-				UE_LOG(LogIoStoreOnDemand, Display, TEXT("Writing file '%s' (%.2lf KiB)"),
-					*TocPath, double(TocSize) / 1024);
-			}
+			return Result;
 		}
-		else
+	}
+
+	if (bSupportStreamingContainers && !StreamingOnDemandToc.Containers.IsEmpty())
+	{
+		StreamingOnDemandToc.Meta.EpochTimestamp = Timestamp;
+		StreamingOnDemandToc.Meta.TargetPlatform = Platform;
+		StreamingOnDemandToc.Meta.BuildVersion = BuildVersion;
+
+		// TODO: OnDemandTocName
+		int32 Result = WriteOnDemandToc(StreamingOnDemandToc, OnDemandTocName + TEXT("ondemand"), OnDemandTocOutputFolder);
+		if (Result < 0)
 		{
-			UE_LOG(LogIoStoreOnDemand, Display, TEXT("Failed writing file '%s'"), *TocPath );
-			return -1;
+			return Result;
 		}
 	}
 
