Prevent some potential divide by 0 issues when collecting IAD analytics.

#rnx

- If a user did not download or install anything via IAD during an analytics recording session then 'CurrentInstallerAnalytics.InstallCount' will be zero causing both 'IadAvgInstallDurationMs' and 'IadAvgCacheHitRatio' to be invalid numbers.

diff --git a/Engine/Source/Runtime/Experimental/IoStore/OnDemand/Private/Statistics.cpp b/Engine/Source/Runtime/Experimental/IoStore/OnDemand/Private/Statistics.cpp
--- a/Engine/Source/Runtime/Experimental/IoStore/OnDemand/Private/Statistics.cpp
+++ b/Engine/Source/Runtime/Experimental/IoStore/OnDemand/Private/Statistics.cpp
@@ -1373,13 +1373,16 @@ void FOnDemandContentInstallerStats::ReportAnalytics(TArray<FAnalyticsEventAttri
 	const FDateTime OldestBlockAccess(CurrentInstallCacheAnalytics.OldestBlockAccess);
 	const FTimespan OldestBlockAge = (Now >= OldestBlockAccess) ? (Now - OldestBlockAccess) : FTimespan::MaxValue();
 
+	const double AvgInstallDurationMs = CurrentInstallerAnalytics.InstallCount > 0 ? (double)CurrentInstallerAnalytics.TotalInstallDurationMs / (double)CurrentInstallerAnalytics.InstallCount : 0.0;
+	const double AvgCacheHitRatio = CurrentInstallerAnalytics.InstallCount > 0 ? (double)CurrentInstallerAnalytics.TotalCacheHitRatio / (double)CurrentInstallerAnalytics.InstallCount : 0.0;
+
 	AppendAnalyticsEventAttributeArray(OutAnalyticsArray
 		, TEXT("IadTotalInstallCount"), CurrentInstallerAnalytics.InstallCount
 		, TEXT("IadTotalInstallErrorCount"), CurrentInstallerAnalytics.InstallErrorCount 
 		, TEXT("IadTotalDownloadedBytes"), CurrentInstallerAnalytics.DownloadedBytes 
 		, TEXT("IadTotalInstallDurationMs"), CurrentInstallerAnalytics.TotalInstallDurationMs 
-		, TEXT("IadAvgInstallDurationMs"), (double(CurrentInstallerAnalytics.TotalInstallDurationMs) / double(CurrentInstallerAnalytics.InstallCount))
-		, TEXT("IadAvgCacheHitRatio"), (CurrentInstallerAnalytics.TotalCacheHitRatio / double(CurrentInstallerAnalytics.InstallCount))
+		, TEXT("IadAvgInstallDurationMs"), AvgInstallDurationMs
+		, TEXT("IadAvgCacheHitRatio"), AvgCacheHitRatio
 		, TEXT("IadInstallCacheStartupErrorCode"), CurrentInstallCacheAnalytics.StartupErrorCode
 		, TEXT("IadInstallCacheVerificationRemovedBlockCount"), CurrentInstallCacheAnalytics.VerificationRemovedBlockCount 
 		, TEXT("IadInstallCacheFlushCount"), CurrentInstallCacheAnalytics.FlushCount 
