Chaos Managed Array - Removed static log declaration in header.

#rnx

diff --git a/Engine/Source/Runtime/Experimental/Chaos/Private/ChaosLog.cpp b/Engine/Source/Runtime/Experimental/Chaos/Private/ChaosLog.cpp
--- a/Engine/Source/Runtime/Experimental/Chaos/Private/ChaosLog.cpp
+++ b/Engine/Source/Runtime/Experimental/Chaos/Private/ChaosLog.cpp
@@ -13,6 +13,11 @@ DEFINE_LOG_CATEGORY(LogChaosDebug);
 DEFINE_LOG_CATEGORY(LogChaos);
 DEFINE_LOG_CATEGORY(LogChaosDataflow);
 
+// UE_DEPRECATED(5.7, "Use LogChaos instead")
+PRAGMA_DISABLE_DEPRECATION_WARNINGS
+DEFINE_LOG_CATEGORY(UManagedArrayLogging);
+PRAGMA_ENABLE_DEPRECATION_WARNINGS
+
 CSV_DEFINE_CATEGORY_MODULE(CHAOS_API, Chaos, true);
 CSV_DEFINE_CATEGORY_MODULE(CHAOS_API, PhysicsVerbose, false);
 CSV_DEFINE_CATEGORY_MODULE(CHAOS_API, PhysicsCounters, false);
\ No newline at end of file
diff --git a/Engine/Source/Runtime/Experimental/Chaos/Public/ChaosLog.h b/Engine/Source/Runtime/Experimental/Chaos/Public/ChaosLog.h
--- a/Engine/Source/Runtime/Experimental/Chaos/Public/ChaosLog.h
+++ b/Engine/Source/Runtime/Experimental/Chaos/Public/ChaosLog.h
@@ -13,7 +13,10 @@ CHAOS_API DECLARE_LOG_CATEGORY_EXTERN(LogChaosDataflow, Log, All);
 
 CHAOS_API DECLARE_LOG_CATEGORY_EXTERN(LogChaos, Log, All);
 
+UE_DEPRECATED(5.7, "Use LogChaos instead")
+CHAOS_API DECLARE_LOG_CATEGORY_EXTERN(UManagedArrayLogging, NoLogging, All);
+
 CSV_DECLARE_CATEGORY_MODULE_EXTERN(CHAOS_API, Chaos);
 
 CSV_DECLARE_CATEGORY_MODULE_EXTERN(CHAOS_API, PhysicsVerbose);
-CSV_DECLARE_CATEGORY_MODULE_EXTERN(CHAOS_API, PhysicsCounters);
\ No newline at end of file
+CSV_DECLARE_CATEGORY_MODULE_EXTERN(CHAOS_API, PhysicsCounters);
diff --git a/Engine/Source/Runtime/Experimental/Chaos/Public/GeometryCollection/ManagedArray.h b/Engine/Source/Runtime/Experimental/Chaos/Public/GeometryCollection/ManagedArray.h
--- a/Engine/Source/Runtime/Experimental/Chaos/Public/GeometryCollection/ManagedArray.h
+++ b/Engine/Source/Runtime/Experimental/Chaos/Public/GeometryCollection/ManagedArray.h
@@ -15,7 +15,6 @@
 #include "Templates/TypeHash.h"
 
 struct FManagedArrayCollection;
-DEFINE_LOG_CATEGORY_STATIC(UManagedArrayLogging, NoLogging, All);
 
 template <typename T>
 void TryBulkSerializeManagedArray(Chaos::FChaosArchive& Ar, TArray<T>& Array)
@@ -1151,7 +1150,7 @@ public:
 
 	virtual void Reindex(const TArray<int32> & Offsets, const int32 & FinalSize, const TArray<int32> & SortedDeletionList, const TSet<int32>& DeletionSet) override
 	{
-		UE_LOG(UManagedArrayLogging, Log, TEXT("TManagedArray<int32>[%p]::Reindex()"),this);
+		UE_LOG(LogChaos, VeryVerbose, TEXT("TManagedArray<int32>[%p]::Reindex()"),this);
 
 		int32 ArraySize = Num(), MaskSize = Offsets.Num();
 		for (int32 Index = 0; Index < ArraySize; Index++)
@@ -1283,7 +1282,7 @@ public:
 	
 	virtual void Reindex(const TArray<int32> & Offsets, const int32 & FinalSize, const TArray<int32> & SortedDeletionList, const TSet<int32>& DeletionSet) override
 	{
-		UE_LOG(UManagedArrayLogging, Log, TEXT("TManagedArray<TArray<int32>>[%p]::Reindex()"), this);
+		UE_LOG(LogChaos, VeryVerbose, TEXT("TManagedArray<TArray<int32>>[%p]::Reindex()"), this);
 		
 		int32 ArraySize = Num(), MaskSize = Offsets.Num();
 
@@ -1352,7 +1351,7 @@ public:
 
 	virtual void Reindex(const TArray<int32> & Offsets, const int32 & FinalSize, const TArray<int32> & SortedDeletionList, const TSet<int32>& DeletionSet) override
 	{
-		UE_LOG(UManagedArrayLogging, Log, TEXT("TManagedArray<FIntVector>[%p]::Reindex()"), this);
+		UE_LOG(LogChaos, VeryVerbose, TEXT("TManagedArray<FIntVector>[%p]::Reindex()"), this);
 		int32 ArraySize = Num(), MaskSize = Offsets.Num();
 		for (int32 Index = 0; Index < ArraySize; Index++)
 		{
@@ -1429,7 +1428,7 @@ public:
 
 	virtual void Reindex(const TArray<int32>& Offsets, const int32& FinalSize, const TArray<int32>& SortedDeletionList, const TSet<int32>& DeletionSet) override
 	{
-		UE_LOG(UManagedArrayLogging, Log, TEXT("TManagedArray<FIntVector>[%p]::Reindex()"), this);
+		UE_LOG(LogChaos, VeryVerbose, TEXT("TManagedArray<FIntVector>[%p]::Reindex()"), this);
 		int32 ArraySize = Num(), MaskSize = Offsets.Num();
 		for (int32 Index = 0; Index < ArraySize; Index++)
 		{
@@ -1506,7 +1505,7 @@ public:
 
 	virtual void Reindex(const TArray<int32>& Offsets, const int32& FinalSize, const TArray<int32>& SortedDeletionList, const TSet<int32>& DeletionSet) override
 	{
-		UE_LOG(UManagedArrayLogging, Log, TEXT("TManagedArray<FIntVector>[%p]::Reindex()"), this);
+		UE_LOG(LogChaos, VeryVerbose, TEXT("TManagedArray<FIntVector>[%p]::Reindex()"), this);
 		int32 ArraySize = Num(), MaskSize = Offsets.Num();
 		for (int32 Index = 0; Index < ArraySize; Index++)
 		{
@@ -1580,7 +1579,7 @@ public:
 
 	virtual void Reindex(const TArray<int32>& Offsets, const int32& FinalSize, const TArray<int32>& SortedDeletionList, const TSet<int32>& DeletionSet) override
 	{
-		UE_LOG(UManagedArrayLogging, Log, TEXT("TManagedArray<FIntVector>[%p]::Reindex()"), this);
+		UE_LOG(LogChaos, VeryVerbose, TEXT("TManagedArray<FIntVector>[%p]::Reindex()"), this);
 		int32 ArraySize = Num(), MaskSize = Offsets.Num();
 		for (int32 Index = 0; Index < ArraySize; Index++)
 		{
@@ -1657,7 +1656,7 @@ public:
 
 	virtual void Reindex(const TArray<int32>& Offsets, const int32& FinalSize, const TArray<int32>& SortedDeletionList, const TSet<int32>& DeletionSet) override
 	{
-		UE_LOG(UManagedArrayLogging, Log, TEXT("TManagedArray<FIntVector>[%p]::Reindex()"), this);
+		UE_LOG(LogChaos, VeryVerbose, TEXT("TManagedArray<FIntVector>[%p]::Reindex()"), this);
 		int32 ArraySize = Num(), MaskSize = Offsets.Num();
 		for (int32 Index = 0; Index < ArraySize; Index++)
 		{
@@ -1722,7 +1721,7 @@ public:
 
 	virtual void Reindex(const TArray<int32>& Offsets, const int32& FinalSize, const TArray<int32>& SortedDeletionList, const TSet<int32>& DeletionSet) override
 	{
-		UE_LOG(UManagedArrayLogging, Log, TEXT("TManagedArray<FIntVector>[%p]::Reindex()"), this);
+		UE_LOG(LogChaos, VeryVerbose, TEXT("TManagedArray<FIntVector>[%p]::Reindex()"), this);
 		const int32 ArraySize = Num();
 		const int32 MaskSize = Offsets.Num();
 		for (int32 Index = 0; Index < ArraySize; ++Index)
