Iris
* Refactored ReplicationReader ProcessQueuedBatch
* Object pending data now held in a map by a TUniquePtr
* Now keeping track of the objects processed in the queued list, and reprocessing other queued objects that were waiting on them instead of waiting on the next time
* This will ensure immediate creation of all pending objects that all depend on each other no matter what replication order they get sent in.

#rnx

diff --git a/Engine/Source/Runtime/Experimental/Iris/Core/Private/Iris/ReplicationSystem/ObjectReferenceCache.cpp b/Engine/Source/Runtime/Experimental/Iris/Core/Private/Iris/ReplicationSystem/ObjectReferenceCache.cpp
--- a/Engine/Source/Runtime/Experimental/Iris/Core/Private/Iris/ReplicationSystem/ObjectReferenceCache.cpp
+++ b/Engine/Source/Runtime/Experimental/Iris/Core/Private/Iris/ReplicationSystem/ObjectReferenceCache.cpp
@@ -736,13 +736,13 @@ bool FObjectReferenceCache::IsNetRefHandlePreRegistered(FNetRefHandle RefHandle)
 	return CacheObjectPtr ? CacheObjectPtr->bIsPreRegistered : false;
 }
 
-bool FObjectReferenceCache::IsNetRefHandlePending(FNetRefHandle NetRefHandle, const FPendingBatches& PendingBatches) const
+bool FObjectReferenceCache::IsNetRefHandlePending(FNetRefHandle NetRefHandle, const FPendingBatchHolder& PendingBatchHolder) const
 {
 	// Check Outer chain
 	while (NetRefHandle.IsValid())
 	{
 		// Need lambda to figure this one out
-		if (PendingBatches.Find(NetRefHandle))
+		if (PendingBatchHolder.Contains(NetRefHandle))
 		{
 			return true;
 		}
diff --git a/Engine/Source/Runtime/Experimental/Iris/Core/Private/Iris/ReplicationSystem/ObjectReferenceCache.h b/Engine/Source/Runtime/Experimental/Iris/Core/Private/Iris/ReplicationSystem/ObjectReferenceCache.h
--- a/Engine/Source/Runtime/Experimental/Iris/Core/Private/Iris/ReplicationSystem/ObjectReferenceCache.h
+++ b/Engine/Source/Runtime/Experimental/Iris/Core/Private/Iris/ReplicationSystem/ObjectReferenceCache.h
@@ -23,7 +23,7 @@ namespace UE::Net
 	{
 		class FNetRefHandleManager;
 		class FNetExportContext;
-		struct FPendingBatches;
+		struct FPendingBatchHolder;
 		typedef uint32 FInternalNetRefIndex;
 	}
 }
@@ -73,7 +73,7 @@ public:
 	bool IsNetRefHandleBroken(FNetRefHandle Handle, bool bMustBeRegistered) const;
 
 	// Returns true of the provided NetRefHandle or one of its outers is pending async loading.
-	bool IsNetRefHandlePending(FNetRefHandle NetRefHandle, const FPendingBatches& PendingBatches) const;
+	bool IsNetRefHandlePending(FNetRefHandle NetRefHandle, const FPendingBatchHolder& PendingBatchHolder) const;
 
 	// Returns true of the provided NetRefHandle was assigned via PreRegisterObjectReferenceHandle or AddPreregisteredReference.
 	bool IsNetRefHandlePreRegistered(FNetRefHandle NetRefHandle) const;
diff --git a/Engine/Source/Runtime/Experimental/Iris/Core/Private/Iris/ReplicationSystem/PendingBatchData.h b/Engine/Source/Runtime/Experimental/Iris/Core/Private/Iris/ReplicationSystem/PendingBatchData.h
--- a/Engine/Source/Runtime/Experimental/Iris/Core/Private/Iris/ReplicationSystem/PendingBatchData.h
+++ b/Engine/Source/Runtime/Experimental/Iris/Core/Private/Iris/ReplicationSystem/PendingBatchData.h
@@ -7,6 +7,8 @@
 #include "Iris/ReplicationSystem/NetRefHandle.h"
 #include "Iris/ReplicationSystem/ObjectReferenceTypes.h"
 
+#include "Templates/UniquePtr.h"
+
 namespace UE::Net::Private
 {
 
@@ -66,13 +68,66 @@ struct FPendingBatchData
 	EIrisAsyncLoadingPriority IrisAsyncLoadingPriority = EIrisAsyncLoadingPriority::Invalid;
 };
 
-struct FPendingBatches
+typedef TUniquePtr<FPendingBatchData> FPendingBatchDataPtr;
+
+struct FPendingBatchHolder
 {
-	FPendingBatchData* Find(FNetRefHandle NetRefHandle) { return PendingBatches.FindByPredicate([&NetRefHandle](const FPendingBatchData& Entry) { return Entry.Owner == NetRefHandle; });}
-	const FPendingBatchData* Find(FNetRefHandle NetRefHandle) const { return PendingBatches.FindByPredicate([&NetRefHandle](const FPendingBatchData& Entry) { return Entry.Owner == NetRefHandle; });}
-	bool GetHasPendingBatches() const { return !PendingBatches.IsEmpty(); }
+public:
+
+	bool Contains(FNetRefHandle NetRefHandle) const
+	{
+		return PendingBatches.Contains(NetRefHandle);
+	}
 
-	TArray<FPendingBatchData> PendingBatches;
+	FPendingBatchData* Find(FNetRefHandle NetRefHandle)
+	{
+		FPendingBatchDataPtr* Ptr = PendingBatches.Find(NetRefHandle);
+		return Ptr ? Ptr->Get() : nullptr;
+	}
+
+	const FPendingBatchData* Find(FNetRefHandle NetRefHandle) const
+	{
+		const FPendingBatchDataPtr* Ptr = PendingBatches.Find(NetRefHandle);
+		return Ptr ? Ptr->Get() : nullptr;
+	}
+
+	FPendingBatchData* FindOrCreate(FNetRefHandle NetRefHandle)
+	{
+		FPendingBatchDataPtr* Ptr = PendingBatches.Find(NetRefHandle);
+		return Ptr ? Ptr->Get() : CreatePendingBatch(NetRefHandle);
+	}
+
+	void Remove(FNetRefHandle NetRefHandle)
+	{
+		PendingBatches.Remove(NetRefHandle);
+	}
+	
+	bool IsEmpty() const
+	{
+		return PendingBatches.IsEmpty();
+	}
+
+	int32 Num() const
+	{
+		return PendingBatches.Num();
+	}
+
+	auto CreateConstIterator() const
+	{
+		return PendingBatches.CreateConstIterator();
+	}
+
+	void Empty()
+	{
+		PendingBatches.Empty();
+	}
+
+private:
+
+	FPendingBatchData* CreatePendingBatch(FNetRefHandle Owner);
+
+	TMap<FNetRefHandle, FPendingBatchDataPtr> PendingBatches;
 };
 
-}
+} // namespace UE::Net::Private
+
diff --git a/Engine/Source/Runtime/Experimental/Iris/Core/Private/Iris/ReplicationSystem/ReplicationReader.cpp b/Engine/Source/Runtime/Experimental/Iris/Core/Private/Iris/ReplicationSystem/ReplicationReader.cpp
--- a/Engine/Source/Runtime/Experimental/Iris/Core/Private/Iris/ReplicationSystem/ReplicationReader.cpp
+++ b/Engine/Source/Runtime/Experimental/Iris/Core/Private/Iris/ReplicationSystem/ReplicationReader.cpp
@@ -354,7 +354,7 @@ FReplicationReader::FReplicationReader()
 FReplicationReader::~FReplicationReader()
 {
 	checkf(ReplicatedObjects.IsEmpty(), TEXT("Possible leak detected in FReplicationReader. Nothing should be registered after Deinit()"));
-	checkf(PendingBatches.PendingBatches.IsEmpty(), TEXT("Possible leak detected in FReplicationReader. Nothing should be registered after Deinit()"));
+	checkf(PendingBatchHolder.IsEmpty(), TEXT("Possible leak detected in FReplicationReader. Nothing should be registered after Deinit()"));
 }
 
 void FReplicationReader::Init(const FReplicationParameters& InParameters)
@@ -393,20 +393,22 @@ void FReplicationReader::Init(const FReplicationParameters& InParameters)
 
 void FReplicationReader::Deinit()
 {
-	for (FPendingBatchData& PendingBatchData : PendingBatches.PendingBatches)
+	for (auto BatchIt = PendingBatchHolder.CreateConstIterator(); BatchIt; ++BatchIt)
 	{
-		UE_LOG(LogIris, Warning, TEXT("FReplicationReader::Deinit NetHandle %s has %d unprocessed data batches"), *PendingBatchData.Owner.ToString(), PendingBatchData.QueuedDataChunks.Num());
+		const FPendingBatchData* PendingBatchData = BatchIt.Value().Get();
+		UE_LOG(LogIris, Warning, TEXT("FReplicationReader::Deinit NetHandle %s has %d unprocessed data batches"), *PendingBatchData->Owner.ToString(), PendingBatchData->QueuedDataChunks.Num());
 
 		// Make sure to release all references that we are holding on to
 		if (ObjectReferenceCache)
 		{
-			for (const FNetRefHandle& RefHandle : PendingBatchData.ResolvedReferences)
+			for (const FNetRefHandle& RefHandle : PendingBatchData->ResolvedReferences)
 			{
 				ObjectReferenceCache->RemoveTrackedQueuedBatchObjectReference(RefHandle);
 			}
 		}		
 	}
-	PendingBatches.PendingBatches.Empty();
+
+	PendingBatchHolder.Empty();
 
 	// Cleanup any allocation stored in the per object info
 	for (auto& ObjectIt : ReplicatedObjects)
@@ -472,8 +474,6 @@ uint32 FReplicationReader::ReadObjectsAndSubObjectsPendingDestroy(FNetSerializat
 		return 0;
 	}
 
-	const bool bHasPendingBatches = PendingBatches.GetHasPendingBatches();
-	
 	for (uint32 It = 0; It < ObjectsToRead; ++It)
 	{
 		UE_NET_TRACE_NAMED_OBJECT_SCOPE(DestroyedObjectScope, FNetRefHandle::GetInvalid(), Reader, Context.GetTraceCollector(), ENetTraceVerbosity::Trace);
@@ -490,7 +490,7 @@ uint32 FReplicationReader::ReadObjectsAndSubObjectsPendingDestroy(FNetSerializat
 			break;
 		}
 
-		if (FPendingBatchData* PendingBatchData = bHasPendingBatches ? PendingBatches.Find(SubObjectRootOrHandle) : nullptr)
+		if (FPendingBatchData* PendingBatchData = PendingBatchHolder.Find(SubObjectRootOrHandle))
 		{
 			// We expect end replication requests for subobjects too so we should only handle one object at a time.
 			constexpr bool bShouldProcessHierarchy = false;
@@ -554,8 +554,6 @@ uint32 FReplicationReader::ReadRootObjectsPendingDestroy(FNetSerializationContex
 		return 0;
 	}
 
-	const bool bHasPendingBatches = PendingBatches.GetHasPendingBatches();
-	
 	for (uint32 It = 0; It < ObjectsToRead; ++It)
 	{
 		UE_NET_TRACE_NAMED_OBJECT_SCOPE(DestroyedObjectScope, FNetRefHandle::GetInvalid(), Reader, Context.GetTraceCollector(), ENetTraceVerbosity::Trace);
@@ -567,7 +565,7 @@ uint32 FReplicationReader::ReadRootObjectsPendingDestroy(FNetSerializationContex
 			break;
 		}
 
-		if (FPendingBatchData* PendingBatchData = bHasPendingBatches ? PendingBatches.Find(IncompleteHandle) : nullptr)
+		if (FPendingBatchData* PendingBatchData = PendingBatchHolder.Find(IncompleteHandle))
 		{
 			constexpr bool bShouldProcessHierarchy = true;
 			EnqueueEndReplication(PendingBatchData, bShouldDestroyInstance, bShouldProcessHierarchy, IncompleteHandle);
@@ -617,10 +615,7 @@ uint32 FReplicationReader::ReadRootObjectsPendingDestroy(FNetSerializationContex
 				// If we did not find the object or associated bridge, the packet that would have created the object may have been lost.
 				UE_LOG(LogIris, Verbose, TEXT("FReplicationReader::Read Tried to destroy object %s. This can occur if the server sends destroy for an object that has not yet been confirmed as created."), *IncompleteHandle.ToString());
 				// Remove from broken list, so that we can try again if the server asks us to instantiate to object again.
-				BrokenObjects.SetNum(Algo::RemoveIf(BrokenObjects, [&IncompleteHandle](const FNetRefHandle& Handle)
-				{
-					return Handle.GetId() == IncompleteHandle.GetId();
-				}));
+				BrokenObjects.RemoveSwap(IncompleteHandle);
 			}
 		}
 	}
@@ -759,7 +754,7 @@ void FReplicationReader::DeserializeObjectStateDelta(FNetSerializationContext& C
 
 FPendingBatchData* FReplicationReader::UpdateUnresolvedMustBeMappedReferences(FNetRefHandle OwnerHandle, TArray<FNetRefHandle>& MustBeMappedReferences, EIrisAsyncLoadingPriority InIrisAsyncLoadingPriority)
 {
-	FPendingBatchData* PendingBatch = PendingBatches.Find(OwnerHandle);
+	FPendingBatchData* PendingBatch = PendingBatchHolder.Find(OwnerHandle);
 
 	// Take the priority if it was sent with the MustBeMapped reference list
 	EIrisAsyncLoadingPriority OwnerAsyncLoadingPriority = InIrisAsyncLoadingPriority != EIrisAsyncLoadingPriority::Invalid ? InIrisAsyncLoadingPriority : EIrisAsyncLoadingPriority::Default;
@@ -826,7 +821,7 @@ FPendingBatchData* FReplicationReader::UpdateUnresolvedMustBeMappedReferences(FN
 		// We must create a new batch
 		if (PendingBatch == nullptr)
 		{
-			PendingBatch = CreatePendingBatch(OwnerHandle);
+			PendingBatch = PendingBatchHolder.FindOrCreate(OwnerHandle);
 		}
 
 		// Update the list of references we are waiting on
@@ -854,23 +849,6 @@ FPendingBatchData* FReplicationReader::UpdateUnresolvedMustBeMappedReferences(FN
 	return PendingBatch;
 }
 
-FPendingBatchData* FReplicationReader::CreatePendingBatch(FNetRefHandle Owner)
-{
-	FPendingBatchData* NewPendingBatch = &PendingBatches.PendingBatches.AddDefaulted_GetRef();
-	
-	NewPendingBatch->Owner = Owner;
-	NewPendingBatch->PendingBatchStartCycles = FPlatformTime::Cycles64();
-	NewPendingBatch->NextWarningTimeout = FPlatformTime::ToSeconds64(NewPendingBatch->PendingBatchStartCycles) + QueuedBatchTimeoutSeconds;
-	
-	return NewPendingBatch;
-}
-
-FPendingBatchData* FReplicationReader::FindOrCreatePendingBatch(FNetRefHandle Owner)
-{
-	FPendingBatchData* PendingBatch = PendingBatches.Find(Owner);
-	return PendingBatch ? PendingBatch : CreatePendingBatch(Owner);
-}
-
 bool FReplicationReader::DoesParentExist(FNetRefHandle ParentHandle) const
 {
 	if (!ParentHandle.IsValid())
@@ -1171,7 +1149,7 @@ uint32 FReplicationReader::ReadObjectBatch(FNetSerializationContext& Context, ui
 	{
 		if (!PendingBatchData)
 		{
-			PendingBatchData = FindOrCreatePendingBatch(IncompleteHandle);
+			PendingBatchData = PendingBatchHolder.FindOrCreate(IncompleteHandle);
 		}
 
 		PendingBatchData->CreationDependentParents = CreationDependentParents;
@@ -2785,9 +2763,9 @@ void FReplicationReader::RemoveFromUnresolvedCache(const FNetRefHandle Handle)
 
 void FReplicationReader::ProcessQueuedBatches()
 {
-	UE_NET_TRACE_FRAME_STATSCOUNTER(Parameters.ReplicationSystem->GetId(), ReplicationReader.PendingQueuedBatches, PendingBatches.PendingBatches.Num(), ENetTraceVerbosity::Trace);
+	UE_NET_TRACE_FRAME_STATSCOUNTER(Parameters.ReplicationSystem->GetId(), ReplicationReader.PendingQueuedBatches, PendingBatchHolder.Num(), ENetTraceVerbosity::Trace);
 
-	if (!PendingBatches.GetHasPendingBatches())
+	if (PendingBatchHolder.IsEmpty())
 	{
 		//Nothing to do.
 		return;
@@ -2807,184 +2785,197 @@ void FReplicationReader::ProcessQueuedBatches()
 	Context.SetInternalContext(&InternalContext);
 	Context.SetNetBlobReceiver(&ReplicationSystemInternal->GetNetBlobHandlerManager());
 
-	//TODO: Add a system to retry pending batches waiting on entities resolved this frame.
+	// List of objects that were processed
+	TQueue<FNetRefHandle, EQueueMode::SingleThreaded> ObjectsProcessed;
 
-	for (int BatchIt = 0; BatchIt < PendingBatches.PendingBatches.Num(); )
-	{
-		FPendingBatchData& PendingBatchData = PendingBatches.PendingBatches[BatchIt];
-
-		// Try to resolve remaining must be mapped references
-		TempMustBeMappedReferences.Reset();
-		UpdateUnresolvedMustBeMappedReferences(PendingBatchData.Owner, TempMustBeMappedReferences, PendingBatchData.IrisAsyncLoadingPriority);
+	// Map of dependent parent handles and the pending batch owners that are waiting on them.
+	TMap<FNetRefHandle /*CreationDependentParent*/, TArray<FNetRefHandle, TInlineAllocator<16>>/* OwnersWaitingOnParent*/> WaitedOnParentHandles;
 
-		int32 NumUnknownParents = 0;
-		
-		for (const FNetRefHandle CreationDependentParent : PendingBatchData.CreationDependentParents)
-		{
-			if (!DoesParentExist(CreationDependentParent))
-			{
-				++NumUnknownParents;
-			}
-		}
-		
-		const bool bAllParentsExist = NumUnknownParents == 0;
-		const bool bAllReferencesResolved = PendingBatchData.PendingMustBeMappedReferences.IsEmpty();
+	auto ProcessBatch = [&, this](FPendingBatchData* PendingBatchData, FNetRefHandle OwnerHandle)
+	{
+		UE_LOG(LogIris, Verbose, TEXT("ProcessQueuedBatches processing %d queued batches for Handle %s "), PendingBatchData->QueuedDataChunks.Num(), *OwnerHandle.ToString());
 
-		//UE_LOG(LogIris, Display, TEXT("ProcessQueuedBatches found %s to have %d unresolved references and %d non-existing parent"), 
-		//	*PendingBatchData.Owner.ToString(), PendingBatchData.PendingMustBeMappedReferences.Num(), NumUnknownParents);
+		// Reset for each batch
+		Reader = FNetBitStreamReader();
+		Context.ResetErrorContext();
+		Context.SetErrorHandleContext(OwnerHandle);
+		UE_RESET_READ_JOURNAL(Context);
+		UE_ADD_READ_JOURNAL_ENTRY(Context, TEXT("ProcessQueuedBatches"));
 
-		// If we no longer waiting on references or the parent, let's process the data!
-		if (bAllParentsExist && bAllReferencesResolved)
+		// Process batched data & dispatch data
+		for (const FQueuedDataChunk& CurrentChunk : PendingBatchData->QueuedDataChunks)
 		{
-			UE_LOG(LogIris, Verbose, TEXT("ProcessQueuedBatches processing %d queued batches for Handle %s "), PendingBatchData.QueuedDataChunks.Num(), *PendingBatchData.Owner.ToString());
-
-			// Reset for batch
-			Reader = FNetBitStreamReader();
-			Context.ResetErrorContext();
-			Context.SetErrorHandleContext(PendingBatchData.Owner);
-			UE_RESET_READ_JOURNAL(Context);
-			UE_ADD_READ_JOURNAL_ENTRY(Context, TEXT("ProcessQueuedBatches"));
+			Reader.InitBits(PendingBatchData->DataChunkStorage.GetData() + CurrentChunk.StorageOffset, CurrentChunk.NumBits);
 
-			// Process batched data & dispatch data
-			for (const FQueuedDataChunk& CurrentChunk : PendingBatchData.QueuedDataChunks)
+			// Chunks marked as bIsEndReplicationChunk are dispatched immediately as we do not know if the next chunk tries to re-create the instance
+			if (CurrentChunk.bIsEndReplicationChunk)
 			{
-				Reader.InitBits(PendingBatchData.DataChunkStorage.GetData() + CurrentChunk.StorageOffset, CurrentChunk.NumBits);
-
-				// Chunks marked as bIsEndReplicationChunk are dispatched immediately as we do not know if the next chunk tries to re-create the instance
-				if (CurrentChunk.bIsEndReplicationChunk)
-				{
-					// Read data stored for objects ending replication, 
-					// this can be the batch root or a subobject owned by the batched root.
-					const uint64 NetRefHandleIdToEndReplication = ReadUint64(&Reader);
-					const bool bShouldDestroyInstance = Reader.ReadBool();
-					const bool bShouldProcessHierarchy = Reader.ReadBool();
+				// Read data stored for objects ending replication, 
+				// this can be the batch root or a subobject owned by the batched root.
+				const uint64 NetRefHandleIdToEndReplication = ReadUint64(&Reader);
+				const bool bShouldDestroyInstance = Reader.ReadBool();
+				const bool bShouldProcessHierarchy = Reader.ReadBool();
 					
-					const FNetRefHandle NetRefHandleToEndReplication = FNetRefHandleManager::MakeNetRefHandleFromId(NetRefHandleIdToEndReplication);
+				const FNetRefHandle NetRefHandleToEndReplication = FNetRefHandleManager::MakeNetRefHandleFromId(NetRefHandleIdToEndReplication);
 					
-					// End replication for object
-					const uint32 InternalIndex = NetRefHandleManager->GetInternalIndex(NetRefHandleToEndReplication);
-					if (InternalIndex != FNetRefHandleManager::InvalidInternalIndex)
+				// End replication for object
+				const uint32 InternalIndex = NetRefHandleManager->GetInternalIndex(NetRefHandleToEndReplication);
+				if (InternalIndex != FNetRefHandleManager::InvalidInternalIndex)
+				{
+					constexpr bool bTearOff = false;
+					if (bShouldProcessHierarchy)
 					{
-						constexpr bool bTearOff = false;
-						if (bShouldProcessHierarchy)
+						EndReplication(InternalIndex, bTearOff, bShouldDestroyInstance);
+						for (const FInternalNetRefIndex SubObjectIndex : NetRefHandleManager->GetSubObjects(InternalIndex))
 						{
-							EndReplication(InternalIndex, bTearOff, bShouldDestroyInstance);
-							for (const FInternalNetRefIndex SubObjectIndex : NetRefHandleManager->GetSubObjects(InternalIndex))
-							{
-								// End replication for all subobjects and destroy dynamic ones
-								const FNetRefHandle SubObjectRefHandle = NetRefHandleManager->GetNetRefHandleFromInternalIndex(SubObjectIndex);
-								const bool bShouldDestroySubObject = SubObjectRefHandle.IsDynamic();
-								EndReplication(SubObjectIndex, bTearOff, bShouldDestroySubObject);
-							}
-						}
-						else
-						{
-							EndReplication(InternalIndex, bTearOff, bShouldDestroyInstance);
+							// End replication for all subobjects and destroy dynamic ones
+							const FNetRefHandle SubObjectRefHandle = NetRefHandleManager->GetNetRefHandleFromInternalIndex(SubObjectIndex);
+							const bool bShouldDestroySubObject = SubObjectRefHandle.IsDynamic();
+							EndReplication(SubObjectIndex, bTearOff, bShouldDestroySubObject);
 						}
 					}
-
-					// Remove from broken list
-					BrokenObjects.SetNum(Algo::RemoveIf(BrokenObjects, [&NetRefHandleToEndReplication](const FNetRefHandle& Handle)
+					else
 					{
-						return Handle.GetId() == NetRefHandleToEndReplication.GetId();
-					}));
-
-					UE_LOG(LogIris, Verbose, TEXT("FReplicationReader::ProcessQueuedBatches EndReplication for %s while processing queued batches for %s"), *NetRefHandleToEndReplication.ToString(), *PendingBatchData.Owner.ToString());
-					continue;
+						EndReplication(InternalIndex, bTearOff, bShouldDestroyInstance);
+					}
 				}
 
-				// Skip over broken objects, we still process remaining chunk if the object has been destroyed.
-				const bool bIsBroken = (BrokenObjects.FindByPredicate([&PendingBatchData](const FNetRefHandle& Entry) { return Entry.GetId() == PendingBatchData.Owner.GetId(); }) != nullptr);
-				if (bIsBroken)
-				{
-					continue;
-				}
+				// Remove from broken list
+				BrokenObjects.RemoveSwap(NetRefHandleToEndReplication);
 
-				// Read and process chunk as it was a received packet for now at least
-				FMemMark TempAllocatorScope(TempLinearAllocator);
+				UE_LOG(LogIris, Verbose, TEXT("FReplicationReader::ProcessQueuedBatches EndReplication for %s while processing queued batches for %s"), *NetRefHandleToEndReplication.ToString(), *OwnerHandle.ToString());
+				continue;
+			}
 
-				// We need to set this up to store temporary dispatch data, the array will grow if needed
-				FObjectsToDispatchArray TempObjectsToDispatchArray(ObjectsToDispatchSlackCount, TempLinearAllocator);
+			// Skip over broken objects, we still process remaining chunk if the object has been destroyed.
+			const bool bIsBroken = BrokenObjects.Contains(OwnerHandle);
+			if (bIsBroken)
+			{
+				continue;
+			}
 
-				// Need to set this pointer as we are dealing with temporary linear allocations
-				ObjectsToDispatchArray = &TempObjectsToDispatchArray;
+			// Read and process chunk as it was a received packet for now at least
+			FMemMark TempAllocatorScope(TempLinearAllocator);
 
-				// Use whatver StreamDebugFeatures were set when receiving the chunk
-				StreamDebugFeatures = CurrentChunk.StreamDebugFeatures;
+			// We need to set this up to store temporary dispatch data, the array will grow if needed
+			FObjectsToDispatchArray TempObjectsToDispatchArray(ObjectsToDispatchSlackCount, TempLinearAllocator);
 
-				// $IRIS: $TODO: Implement special dispatch to defer RepNotifies if we are processing multiple batches for the same object.
-				ReadObjectsInBatch(Context, PendingBatchData.Owner, CurrentChunk.bHasBatchOwnerData, CurrentChunk.NumBits);
+			// Need to set this pointer as we are dealing with temporary linear allocations
+			ObjectsToDispatchArray = &TempObjectsToDispatchArray;
 
-				if (Context.HasErrorOrOverflow())
-				{
-					BrokenObjects.AddUnique(PendingBatchData.Owner);
+			// Use whatver StreamDebugFeatures were set when receiving the chunk
+			StreamDebugFeatures = CurrentChunk.StreamDebugFeatures;
 
-					// Log error  if get more incoming data for an object in the broken state we will skip it.
-					UE_LOG(LogIris, Error, TEXT("FReplicationReader::ProcessQueuedBatches Failed to process object batch handle: %s skipping batch data. bIsEndReplication: %d"), ToCStr(PendingBatchData.Owner.ToString()), int(CurrentChunk.bIsEndReplicationChunk));
+			// $IRIS: $TODO: Implement special dispatch to defer RepNotifies if we are processing multiple batches for the same object.
+			ReadObjectsInBatch(Context, OwnerHandle, CurrentChunk.bHasBatchOwnerData, CurrentChunk.NumBits);
 
-					if (Context.GetError() == GNetError_BrokenNetHandle)
-					{
-						ReplicationBridge->SendErrorWithNetRefHandle(UE::Net::ENetRefHandleError::ReplicationDisabled, PendingBatchData.Owner, Parameters.ConnectionId);
+			if (Context.HasErrorOrOverflow())
+			{
+				BrokenObjects.AddUnique(OwnerHandle);
 
-						Context.ResetErrorContext();
-					}
+				// Log error  if get more incoming data for an object in the broken state we will skip it.
+				UE_LOG(LogIris, Error, TEXT("FReplicationReader::ProcessQueuedBatches Failed to process object batch handle: %s skipping batch data. bIsEndReplication: %d"), ToCStr(OwnerHandle.ToString()), int(CurrentChunk.bIsEndReplicationChunk));
+
+				if (Context.GetError() == GNetError_BrokenNetHandle)
+				{
+					ReplicationBridge->SendErrorWithNetRefHandle(UE::Net::ENetRefHandleError::ReplicationDisabled, OwnerHandle, Parameters.ConnectionId);
+
+					Context.ResetErrorContext();
 				}
+			}
 			
-				// Apply received data and resolve dependencies
-				DispatchStateData(Context);
+			// Apply received data and resolve dependencies
+			DispatchStateData(Context);
 
-				// Resolve
-				ResolveAndDispatchUnresolvedReferences();
+			// Resolve
+			ResolveAndDispatchUnresolvedReferences();
 
-				// EndReplication for all objects in the batch that should no longer replicate
-				DispatchEndReplication(Context);
+			// EndReplication for all objects in the batch that should no longer replicate
+			DispatchEndReplication(Context);
 
-				// Drop temporary data
-				ObjectsToDispatchArray = nullptr;
-			}
+			// Drop temporary data
+			ObjectsToDispatchArray = nullptr;
+		}
 	
-			// Make sure to release all references that we hold on to
-			for (const FNetRefHandle& RefHandle : PendingBatchData.ResolvedReferences)
+		// Make sure to release all references that we hold on to
+		for (const FNetRefHandle& RefHandle : PendingBatchData->ResolvedReferences)
+		{
+			ObjectReferenceCache->RemoveTrackedQueuedBatchObjectReference(RefHandle);
+		}
+	};
+
+	for (auto BatchIt = PendingBatchHolder.CreateConstIterator(); BatchIt; ++BatchIt)
+	{
+		const FPendingBatchDataPtr& DataPtr = BatchIt.Value();
+		FPendingBatchData* PendingBatchData = DataPtr.Get();
+		check(PendingBatchData);
+
+		const FNetRefHandle OwnerHandle = PendingBatchData->Owner;
+
+		// Try to resolve remaining must be mapped references
+		{
+			TempMustBeMappedReferences.Reset();
+			UpdateUnresolvedMustBeMappedReferences(OwnerHandle, TempMustBeMappedReferences, PendingBatchData->IrisAsyncLoadingPriority);
+		}
+
+		const bool bAllReferencesResolved = PendingBatchData->PendingMustBeMappedReferences.IsEmpty();
+
+		// Look if all CreationDependencies are resolved
+		int32 NumUnknownParents = 0;		
+		for (const FNetRefHandle CreationDependentParent : PendingBatchData->CreationDependentParents)
+		{
+			if (!DoesParentExist(CreationDependentParent))
 			{
-				ObjectReferenceCache->RemoveTrackedQueuedBatchObjectReference(RefHandle);
+				++NumUnknownParents;
+
+				// Add the owner to the wait list if only creation dependencies are left
+				if (bAllReferencesResolved)
+				{
+					auto& WaitedOnList = WaitedOnParentHandles.FindOrAdd(CreationDependentParent);
+					WaitedOnList.Add(OwnerHandle);
+				}
 			}
+		}
+		
+		const bool bAllParentsExist = NumUnknownParents == 0;
 
-			// Not optimal, but we want to preserve the order if we can as there might be batches waiting for the same reference
-			PendingBatches.PendingBatches.RemoveAt(BatchIt);
+		//UE_LOG(LogIris, Display, TEXT("ProcessQueuedBatches found %s to have %d unresolved references and %d non-existing parent"), 
+		//	*OwnerHandle.ToString(), PendingBatchData.PendingMustBeMappedReferences.Num(), NumUnknownParents);
+
+		// If we are no longer waiting on references or on creation parents, let's process the data!
+		if (bAllParentsExist && bAllReferencesResolved)
+		{
+			ProcessBatch(PendingBatchData, OwnerHandle);
+			ObjectsProcessed.Enqueue(OwnerHandle);
 		}
+		// Batch is still pending something, track how long its been waiting
 		else
 		{
-			// Increment counter
-			++PendingBatchData.PendingBatchTryProcessCount;
-
 			// Check how long we've been blocked only every 30 ticks
-			if ((QueuedBatchTimeoutWarningInterval > 0) && (PendingBatchData.PendingBatchTryProcessCount > QueuedBatchTimeoutWarningInterval) && FPlatformProperties::RequiresCookedData() && !FPlatformMisc::IsDebuggerPresent())
+			++PendingBatchData->PendingBatchTryProcessCount;
+
+			if ((QueuedBatchTimeoutWarningInterval > 0) && (PendingBatchData->PendingBatchTryProcessCount > QueuedBatchTimeoutWarningInterval) && FPlatformProperties::RequiresCookedData() && !FPlatformMisc::IsDebuggerPresent())
 			{
 				const double CurrentTime = FPlatformTime::Seconds();
 
 				// Log a warning and tell the server about the blockage every X seconds (first after 30secs, again every 10secs after)
-				if (CurrentTime >= PendingBatchData.NextWarningTimeout)
+				if (CurrentTime >= PendingBatchData->NextWarningTimeout)
 				{
-					const double BlockedTime = FPlatformTime::ToSeconds64(FPlatformTime::Cycles64() - PendingBatchData.PendingBatchStartCycles);
+					const double BlockedTime = FPlatformTime::ToSeconds64(FPlatformTime::Cycles64() - PendingBatchData->PendingBatchStartCycles);
 
 					if (!bAllParentsExist)
 					{
 						UE_LOG(LogIris, Warning, TEXT("FReplicationWriter::ProcessQueuedBatches: Replication is blocked for longer than normal (%f secs) for %s, Waiting on %d/%d parents to be created.")
 							   , BlockedTime
-							   , *NetRefHandleManager->PrintObjectFromNetRefHandle(PendingBatchData.Owner)
+							   , *NetRefHandleManager->PrintObjectFromNetRefHandle(OwnerHandle)
 							   , NumUnknownParents
-							   , PendingBatchData.CreationDependentParents.Num());
+							   , PendingBatchData->CreationDependentParents.Num());
 
-						for (const FNetRefHandle CreationDependentParent : PendingBatchData.CreationDependentParents)
+						for (const FNetRefHandle CreationDependentParent : PendingBatchData->CreationDependentParents)
 						{
-							if (DoesParentExist(CreationDependentParent))
-							{
-								UE_LOG(LogIris, Warning, TEXT("FReplicationWriter::ProcessQueuedBatches: (valid) CreationDependent parent: %s exists."), *NetRefHandleManager->PrintObjectFromNetRefHandle(CreationDependentParent));
-							}
-							else
-							{
-								UE_LOG(LogIris, Warning, TEXT("FReplicationWriter::ProcessQueuedBatches: (wrong) CreationDependent parent: %s does not exist yet."), *NetRefHandleManager->PrintObjectFromNetRefHandle(CreationDependentParent));
-							}
+							const bool bParentExists = DoesParentExist(CreationDependentParent);
+
+							UE_LOG(LogIris, Warning, TEXT("FReplicationWriter::ProcessQueuedBatches: %s CreationDependent parent: %s ."), bParentExists?TEXT("Found"):TEXT("Waiting"), *NetRefHandleManager->PrintObjectFromNetRefHandle(CreationDependentParent));
 						}
 					}
 
@@ -2993,7 +2984,7 @@ void FReplicationReader::ProcessQueuedBatches()
 						TArray<FNetRefHandle> PendingMustBeMapped;
 						PendingMustBeMapped.Reserve(MaxMustBeMappedHandleArray);
 
-						for (const FNetRefHandle& AsyncLoadingRef : PendingBatchData.PendingMustBeMappedReferences)
+						for (const FNetRefHandle& AsyncLoadingRef : PendingBatchData->PendingMustBeMappedReferences)
 						{
 							if (PendingMustBeMapped.Num() >= MaxMustBeMappedHandleArray)
 							{
@@ -3014,37 +3005,77 @@ void FReplicationReader::ProcessQueuedBatches()
 						// Warn and tell server about the blockage
 						if (!PendingMustBeMapped.IsEmpty())
 						{
-							const FString BlockedObjectName = *NetRefHandleManager->PrintObjectFromNetRefHandle(PendingBatchData.Owner);
+							const FString BlockedObjectName = *NetRefHandleManager->PrintObjectFromNetRefHandle(OwnerHandle);
 
 							UE_LOG(LogIris, Warning, TEXT("FReplicationWriter::ProcessQueuedBatches: Replication is blocked for longer than normal for: %s, QueuedBunches: %i, PendingNetRefHandleResolves: %i, BlockedTime: %f")
 								   , *BlockedObjectName
-								   , PendingBatchData.QueuedDataChunks.Num()
-								   , PendingBatchData.PendingMustBeMappedReferences.Num()
+								   , PendingBatchData->QueuedDataChunks.Num()
+								   , PendingBatchData->PendingMustBeMappedReferences.Num()
 								   , BlockedTime);
 
-							for (const FNetRefHandle& NetRefHandle : PendingBatchData.PendingMustBeMappedReferences)
+							for (const FNetRefHandle& NetRefHandle : PendingBatchData->PendingMustBeMappedReferences)
 							{
 								UE_LOG(LogIris, Log, TEXT("  NetRefHandle Pending resolve [%s]"), *ObjectReferenceCache->DescribeObjectReference(FObjectReferenceCache::MakeNetObjectReference(NetRefHandle), ResolveContext));
 							}
 
 							#if IRIS_CLIENT_PROFILER_ENABLE
-							FClientProfiler::RecordBlockedReplication(*BlockedObjectName, PendingBatchData.PendingMustBeMappedReferences.Num(), (float)BlockedTime);
+							FClientProfiler::RecordBlockedReplication(*BlockedObjectName, PendingBatchData->PendingMustBeMappedReferences.Num(), (float)BlockedTime);
 							#endif
 
-							ReplicationBridge->SendErrorWithNetRefHandle(UE::Net::ENetRefHandleError::BlockedByMustBeMapped, PendingBatchData.Owner, Parameters.ConnectionId, PendingMustBeMapped);
+							ReplicationBridge->SendErrorWithNetRefHandle(UE::Net::ENetRefHandleError::BlockedByMustBeMapped, OwnerHandle, Parameters.ConnectionId, PendingMustBeMapped);
 						
 						}
 					}
-					
 
-					PendingBatchData.NextWarningTimeout = CurrentTime + QueuedBatchWarningRepeatTime;
+					PendingBatchData->NextWarningTimeout = CurrentTime + QueuedBatchWarningRepeatTime;
 				}
 
 				// Reset counter
-				PendingBatchData.PendingBatchTryProcessCount = 0;
+				PendingBatchData->PendingBatchTryProcessCount = 0;
 			}
+		}
+	}
+
+	FNetRefHandle ProcessedObject;
+	while (ObjectsProcessed.Dequeue(ProcessedObject))
+	{
+		UE_LOG(LogIris, Verbose, TEXT("Clean-up processed pending object: %s"), *NetRefHandleManager->PrintObjectFromNetRefHandle(ProcessedObject));
 
-			++BatchIt;
+		// Delete the previously processed batch data for this object
+		PendingBatchHolder.Remove(ProcessedObject);
+
+		const bool bIsWaitedOn = WaitedOnParentHandles.Contains(ProcessedObject);
+		
+		// This newly created object had others waiting on it, let's check if they can be processed now
+		if (bIsWaitedOn)
+		{
+			const auto& WaitingOwners = WaitedOnParentHandles.FindChecked(ProcessedObject);
+
+			for (FNetRefHandle WaitingOwner : WaitingOwners)
+			{
+				//UE_LOG(LogIris, Verbose, TEXT("Pending object: %s was waiting on: %s"), *WaitingOwner.ToString() ,*NetRefHandleManager->PrintObjectFromNetRefHandle(ProcessedObject));
+
+				if (FPendingBatchData* PendingBatchData = PendingBatchHolder.Find(WaitingOwner))
+				{
+					// Make sure all it's parent do exist
+					bool bAllParentsExist = true;
+					for (const FNetRefHandle CreationDependentParent : PendingBatchData->CreationDependentParents)
+					{
+						bAllParentsExist = DoesParentExist(CreationDependentParent);
+						if (!bAllParentsExist)
+						{
+							break;
+						}
+					}
+
+					// No reasons to wait anymore, process it's queued data
+					if (bAllParentsExist)
+					{
+						ProcessBatch(PendingBatchData, WaitingOwner);
+						ObjectsProcessed.Enqueue(WaitingOwner);
+					}
+				}
+			}
 		}
 	}
 }
@@ -3213,7 +3244,7 @@ void FReplicationReader::ResolveAndDispatchAttachments(FNetSerializationContext&
 					// The difference is that the old system might be able to resolve incoming dynamic references from later packets processed for the same tick, but as this is far from guaranteed we currently do not try to mimic this.
 					for (const FNetReferenceCollector::FReferenceInfo& Info : MakeArrayView(Collector.GetCollectedReferences()))
 					{
-						if (ObjectReferenceCache->IsNetRefHandlePending(Info.Reference.GetRefHandle(), PendingBatches))
+						if (ObjectReferenceCache->IsNetRefHandlePending(Info.Reference.GetRefHandle(), PendingBatchHolder))
 						{
 							bDelayRpc = true;
 							break;
@@ -3314,15 +3345,22 @@ FString FReplicationReader::PrintObjectInfo(FInternalNetRefIndex ObjectIndex, FN
 	const bool bIsSubObject = ObjectData.SubObjectRootIndex != FNetRefHandleManager::InvalidInternalIndex;
 	InfoBuilder.Appendf(TEXT("Status info: 0x%llx (%s)"), ObjectInfo->Value, bIsSubObject?TEXT("SubObject"):TEXT("RootObject"));
 
-	const bool bHasPendingBatches = PendingBatches.GetHasPendingBatches();
 	const FInternalNetRefIndex RootInternalIndex = bIsSubObject ? ObjectData.SubObjectRootIndex : ObjectIndex;
-	if (const FPendingBatchData* PendingBatchData = bHasPendingBatches ? PendingBatches.Find(NetRefHandleManager->GetNetRefHandleFromInternalIndex(RootInternalIndex)) : nullptr)
+	const FNetRefHandle RootObjectHandle = bIsSubObject ? NetRefHandleManager->GetNetRefHandleFromInternalIndex(RootInternalIndex) : NetRefHandle;
+	if (const FPendingBatchData* PendingBatchData = PendingBatchHolder.Find(RootObjectHandle))
 	{
-		InfoBuilder.Appendf(TEXT("| PendingBatches: QueuedChunks: %d | PendingReferences(%d): "), PendingBatchData->QueuedDataChunks.Num(), PendingBatchData->PendingMustBeMappedReferences.Num());
+		InfoBuilder.Appendf(TEXT("| PendingBatches: QueuedChunks: %d | PendingReferences(%d): "),PendingBatchData->QueuedDataChunks.Num(), PendingBatchData->PendingMustBeMappedReferences.Num());
+
 		for (const FNetRefHandle& PendingHandle : PendingBatchData->PendingMustBeMappedReferences)
 		{
 			InfoBuilder << PendingHandle.ToString() << ", ";
 		}
+
+		InfoBuilder.Appendf(TEXT(" | CreationDependentParents(%d): "), PendingBatchData->CreationDependentParents.Num());
+		for (FNetRefHandle CreationDependentParent : PendingBatchData->CreationDependentParents)
+		{
+			InfoBuilder << CreationDependentParent.ToString() << ", ";
+		}
 	}
 	else
 	{
@@ -3332,4 +3370,24 @@ FString FReplicationReader::PrintObjectInfo(FInternalNetRefIndex ObjectIndex, FN
 	return InfoBuilder.ToString();
 }
 
+//------------------------------------------------------------------------
+// FPendingBatchHolder
+//------------------------------------------------------------------------
+
+FPendingBatchData* FPendingBatchHolder::CreatePendingBatch(FNetRefHandle Owner)
+{
+	// TODO: Would be more efficient to base this off FReplicationSystem::GetElapsedTime() instead.
+	const uint64 CurrentTime = FPlatformTime::Cycles64();
+	
+	FPendingBatchData* NewData(new FPendingBatchData());
+	
+	NewData->Owner = Owner;
+	NewData->PendingBatchStartCycles = CurrentTime;
+	NewData->NextWarningTimeout = FPlatformTime::ToSeconds64(CurrentTime) + QueuedBatchTimeoutSeconds;
+
+	PendingBatches.Add(Owner, FPendingBatchDataPtr(NewData));
+	
+	return NewData;
+}
+
 } // end namespace UE::Net::Private
diff --git a/Engine/Source/Runtime/Experimental/Iris/Core/Private/Iris/ReplicationSystem/ReplicationReader.h b/Engine/Source/Runtime/Experimental/Iris/Core/Private/Iris/ReplicationSystem/ReplicationReader.h
--- a/Engine/Source/Runtime/Experimental/Iris/Core/Private/Iris/ReplicationSystem/ReplicationReader.h
+++ b/Engine/Source/Runtime/Experimental/Iris/Core/Private/Iris/ReplicationSystem/ReplicationReader.h
@@ -221,9 +221,6 @@ private:
 	// If async loading is enabled this function will verify if we can resolve all PendingMustBeMappedReferences
 	FPendingBatchData* UpdateUnresolvedMustBeMappedReferences(FNetRefHandle OwnerHandle, TArray<FNetRefHandle>& MustBeMappedReferences, EIrisAsyncLoadingPriority InIrisAsyncLoadingPriority);
 
-	FPendingBatchData* CreatePendingBatch(FNetRefHandle Owner);
-	FPendingBatchData* FindOrCreatePendingBatch(FNetRefHandle Owner);
-
 	// Returns true if the object assigned to this handle has been instantiated locally
 	bool DoesParentExist(FNetRefHandle ObjectHandle) const;
 
@@ -282,7 +279,7 @@ private:
 	TMultiMap<FNetRefHandle, uint32> ResolvedDynamicHandleToDependents;
 
 	// We do not expect to have many objects in this state
-	FPendingBatches PendingBatches;
+	FPendingBatchHolder PendingBatchHolder;
 
 	// We do not expect many objects to be broken
 	TArray<FNetRefHandle> BrokenObjects;
