Momentum Propagation

Partial wake events will not only wake up connected particles in higher levels (level-based waking).
Additionally, the impact momentum of the trigger particle is propagated through the constraint graph.
Traversed particles are woken up if the received momentum is greater than their mass/inertia x wake-up velocity threshold.
We transmit a percentage of the momentum from one constraint level to the next tunable by a CVar.

diff --git a/Engine/Source/Programs/HeadlessChaos/Private/HeadlessChaosTestGraphEvolution.cpp b/Engine/Source/Programs/HeadlessChaos/Private/HeadlessChaosTestGraphEvolution.cpp
--- a/Engine/Source/Programs/HeadlessChaos/Private/HeadlessChaosTestGraphEvolution.cpp
+++ b/Engine/Source/Programs/HeadlessChaos/Private/HeadlessChaosTestGraphEvolution.cpp
@@ -48,10 +48,15 @@ namespace ChaosTest
 
 			CVarPartialSleeping = IConsoleManager::Get().FindConsoleVariable(TEXT("p.Chaos.Solver.Sleep.PartialIslandSleep"), false);
 			check(CVarPartialSleeping && CVarPartialSleeping->IsVariableBool());
-
+			PartialSleepDefault = CVarPartialSleeping->GetBool();
 			CVarPartialSleeping->Set(bPartialSleeping);
 		}
 
+		~FGraphEvolutionTest()
+		{
+			CVarPartialSleeping->Set(PartialSleepDefault);
+		}
+
 		// Connect all the particles in a chain
 		void MakeChain()
 		{
@@ -114,6 +119,7 @@ namespace ChaosTest
 		int32 TickCount;
 
 		IConsoleVariable* CVarPartialSleeping;
+		bool PartialSleepDefault;
 	};
 
 	// Veryify that the Null Constraint mockup is working as intended. We can create the container and constraints, and they are correctly bound to the evolution
@@ -1589,6 +1595,12 @@ namespace ChaosTest
 
 		if (Test.CVarPartialSleeping->GetBool())
 		{
+			// Turn off mommentum propagation to test level-based waking only
+			IConsoleVariable* CVarMomentumPropagation = IConsoleManager::Get().FindConsoleVariable(TEXT("p.Chaos.Solver.Sleep.MomentumPropagation"), false);
+			check(CVarMomentumPropagation && CVarMomentumPropagation->IsVariableFloat());
+			const FRealSingle MomentumPropagationDefault = CVarMomentumPropagation->GetFloat();
+			CVarMomentumPropagation->Set(0);
+
 			Test.MakeChain();
 			Test.Evolution.SetParticleObjectState(Test.ParticleHandles[0], EObjectStateType::Kinematic);
 
@@ -1643,6 +1655,8 @@ namespace ChaosTest
 			EXPECT_TRUE(Test.ConstraintHandles[1]->IsSleeping());
 			EXPECT_FALSE(Test.ConstraintHandles[2]->IsSleeping());
 			EXPECT_FALSE(Test.ConstraintHandles[3]->IsSleeping());
+
+			CVarMomentumPropagation->Set(MomentumPropagationDefault);
 		}
 	}
 
@@ -1659,6 +1673,12 @@ namespace ChaosTest
 
 		if (Test.CVarPartialSleeping->GetBool())
 		{
+			// Turn off mommentum propagation to test level-based waking only
+			IConsoleVariable* CVarMomentumPropagation = IConsoleManager::Get().FindConsoleVariable(TEXT("p.Chaos.Solver.Sleep.MomentumPropagation"), false);
+			check(CVarMomentumPropagation && CVarMomentumPropagation->IsVariableFloat());
+			const FRealSingle MomentumPropagationDefault = CVarMomentumPropagation->GetFloat();
+			CVarMomentumPropagation->Set(0);
+
 			Test.MakeChain();
 			Test.Evolution.SetParticleObjectState(Test.ParticleHandles[0], EObjectStateType::Kinematic);
 
@@ -1713,6 +1733,8 @@ namespace ChaosTest
 			EXPECT_FALSE(Test.ConstraintHandles[1]->IsSleeping());
 			EXPECT_FALSE(Test.ConstraintHandles[2]->IsSleeping());
 			EXPECT_FALSE(Test.ConstraintHandles[3]->IsSleeping());
+
+			CVarMomentumPropagation->Set(MomentumPropagationDefault);
 		}
 	}
 
@@ -1738,6 +1760,12 @@ namespace ChaosTest
 		FGraphEvolutionTest Test(7, GetParam());
 		if (Test.CVarPartialSleeping->GetBool())
 		{
+			// Turn off mommentum propagation to test level-based waking only
+			IConsoleVariable* CVarMomentumPropagation = IConsoleManager::Get().FindConsoleVariable(TEXT("p.Chaos.Solver.Sleep.MomentumPropagation"), false);
+			check(CVarMomentumPropagation && CVarMomentumPropagation->IsVariableFloat());
+			const FRealSingle MomentumPropagationDefault = CVarMomentumPropagation->GetFloat();
+			CVarMomentumPropagation->Set(0);
+
 			Test.Evolution.SetParticleObjectState(Test.ParticleHandles[0], EObjectStateType::Kinematic);
 			Test.ConstraintHandles.Add(Test.Constraints.AddConstraint({ Test.ParticleHandles[0], Test.ParticleHandles[1] }));
 			Test.ConstraintHandles.Add(Test.Constraints.AddConstraint({ Test.ParticleHandles[0], Test.ParticleHandles[2] }));
@@ -1810,6 +1838,8 @@ namespace ChaosTest
 			EXPECT_FALSE(Test.ConstraintHandles[7]->IsSleeping());
 			EXPECT_FALSE(Test.ConstraintHandles[8]->IsSleeping());
 			EXPECT_FALSE(Test.ConstraintHandles[9]->IsSleeping());
+
+			CVarMomentumPropagation->Set(MomentumPropagationDefault);
 		}
 	}
 
@@ -1835,6 +1865,12 @@ namespace ChaosTest
 		FGraphEvolutionTest Test(7, GetParam());
 		if (Test.CVarPartialSleeping->GetBool())
 		{
+			// Turn off mommentum propagation to test level-based waking only
+			IConsoleVariable* CVarMomentumPropagation = IConsoleManager::Get().FindConsoleVariable(TEXT("p.Chaos.Solver.Sleep.MomentumPropagation"), false);
+			check(CVarMomentumPropagation && CVarMomentumPropagation->IsVariableFloat());
+			const FRealSingle MomentumPropagationDefault = CVarMomentumPropagation->GetFloat();
+			CVarMomentumPropagation->Set(0);
+
 			Test.Evolution.SetParticleObjectState(Test.ParticleHandles[0], EObjectStateType::Kinematic);
 			Test.ConstraintHandles.Add(Test.Constraints.AddConstraint({ Test.ParticleHandles[0], Test.ParticleHandles[1] }));
 			Test.ConstraintHandles.Add(Test.Constraints.AddConstraint({ Test.ParticleHandles[0], Test.ParticleHandles[2] }));
@@ -1917,6 +1953,8 @@ namespace ChaosTest
 			EXPECT_FALSE(Test.ConstraintHandles[10]->IsSleeping());
 			EXPECT_FALSE(Test.ConstraintHandles[11]->IsSleeping());
 			EXPECT_FALSE(Test.ConstraintHandles[12]->IsSleeping());
+
+			CVarMomentumPropagation->Set(MomentumPropagationDefault);
 		}
 	}
 
@@ -1939,6 +1977,12 @@ namespace ChaosTest
 	//                 Ak
 	TEST_P(GraphEvolutionTests, TestConstraintGraph_WakeUpPropagation_AddToBlockWithoutSpacingSleeping)
 	{
+		// Turn off mommentum propagation to test level-based waking only
+		IConsoleVariable* CVarMomentumPropagation = IConsoleManager::Get().FindConsoleVariable(TEXT("p.Chaos.Solver.Sleep.MomentumPropagation"), false);
+		check(CVarMomentumPropagation && CVarMomentumPropagation->IsVariableFloat());
+		const FRealSingle MomentumPropagationDefault = CVarMomentumPropagation->GetFloat();
+		CVarMomentumPropagation->Set(0);
+
 		FGraphEvolutionTest Test(10, GetParam());
 		Test.Evolution.SetParticleObjectState(Test.ParticleHandles[0], EObjectStateType::Kinematic);
 
@@ -2031,6 +2075,8 @@ namespace ChaosTest
 		EXPECT_TRUE(Test.ConstraintHandles[13]->IsSleeping());
 		EXPECT_TRUE(Test.ConstraintHandles[14]->IsSleeping());
 		EXPECT_TRUE(Test.ConstraintHandles[15]->IsSleeping());
+
+		CVarMomentumPropagation->Set(MomentumPropagationDefault);
 	}
 
 	// Test validating the wake-up propagation for a rectangular block wall without lateral spacing.
@@ -2062,6 +2108,12 @@ namespace ChaosTest
 	//                 As
 	TEST_P(GraphEvolutionTests, TestConstraintGraph_WakeUpPropagation_AddToBlockWithoutSpacingAwake)
 	{
+		// Turn off mommentum propagation to test level-based waking only
+		IConsoleVariable* CVarMomentumPropagation = IConsoleManager::Get().FindConsoleVariable(TEXT("p.Chaos.Solver.Sleep.MomentumPropagation"), false);
+		check(CVarMomentumPropagation && CVarMomentumPropagation->IsVariableFloat());
+		const FRealSingle MomentumPropagationDefault = CVarMomentumPropagation->GetFloat();
+		CVarMomentumPropagation->Set(0);
+
 		FGraphEvolutionTest Test(10, GetParam());
 		Test.Evolution.SetParticleObjectState(Test.ParticleHandles[0], EObjectStateType::Kinematic);
 
@@ -2198,6 +2250,8 @@ namespace ChaosTest
 			EXPECT_FALSE(Test.ConstraintHandles[14]->IsSleeping());
 			EXPECT_FALSE(Test.ConstraintHandles[15]->IsSleeping());
 		}
+
+		CVarMomentumPropagation->Set(MomentumPropagationDefault);
 	}
 
 	TEST_P(GraphEvolutionTests, TestConstraintGraph_SleepCounterReset)
@@ -2419,9 +2473,72 @@ namespace ChaosTest
 			EXPECT_TRUE(FConstGenericParticleHandle(Test.ParticleHandles[0])->Sleeping());
 			EXPECT_TRUE(FConstGenericParticleHandle(Test.ParticleHandles[1])->Sleeping());
 			EXPECT_FALSE(FConstGenericParticleHandle(Test.ParticleHandles[2])->Sleeping());
+			EXPECT_FALSE(FConstGenericParticleHandle(Test.ParticleHandles[3])->Sleeping());
+			EXPECT_FALSE(FConstGenericParticleHandle(Test.ParticleHandles[4])->Sleeping());
+		}
+	}
+
+	TEST_P(GraphEvolutionTests, TestConstraintGraph_MomentumPropagation)
+	{
+		const bool bEnablePartialSleep = true;
+		const FRealSingle MomentumPropagation = GetParam() ? 0.5f : 0.025f;
+
+		IConsoleVariable* CVarMomentumPropagation = IConsoleManager::Get().FindConsoleVariable(TEXT("p.Chaos.Solver.Sleep.MomentumPropagation"), false);
+		check(CVarMomentumPropagation && CVarMomentumPropagation->IsVariableFloat());
+		const FRealSingle MomentumPropagationDefault = CVarMomentumPropagation->GetFloat();
+		CVarMomentumPropagation->Set(MomentumPropagation);
+
+		FGraphEvolutionTest Test(5, bEnablePartialSleep);
+		Test.MakeChain();
+		Test.Evolution.SetParticleObjectState(Test.ParticleHandles[0], EObjectStateType::Kinematic);
+
+		Test.AdvanceUntilSleeping();
+
+		// Explicitly wake up one of the top particle
+		Test.Evolution.SetParticleObjectState(Test.ParticleHandles[4], EObjectStateType::Dynamic);
+		Test.Advance();
+
+		// One the particle explicitly woken and the particle next to it are awake.
+		// Momentum of the woken particle is zero, i.e. no momentum propagation.
+		EXPECT_TRUE(FConstGenericParticleHandle(Test.ParticleHandles[0])->IsKinematic());
+		EXPECT_TRUE(FConstGenericParticleHandle(Test.ParticleHandles[1])->Sleeping());
+		EXPECT_TRUE(FConstGenericParticleHandle(Test.ParticleHandles[2])->Sleeping());
+		EXPECT_FALSE(FConstGenericParticleHandle(Test.ParticleHandles[3])->Sleeping());
+		EXPECT_FALSE(FConstGenericParticleHandle(Test.ParticleHandles[4])->Sleeping());
+
+		Test.AdvanceUntilSleeping();
+
+		// Explicitly wake up one of the particles and add an impulse
+		Test.Evolution.SetParticleObjectState(Test.ParticleHandles[4], EObjectStateType::Dynamic);
+		FParticleDynamics Dynamics;
+		Dynamics.SetAcceleration(FVec3(0, 0, 0));
+		Dynamics.SetAngularAcceleration(FVec3(0, 0, 0));
+		Dynamics.SetLinearImpulseVelocity(FVec3(0, 0.1, 0));
+		Dynamics.SetAngularImpulseVelocity(FVec3(0, 0, 0));
+		Test.ParticleHandles[4]->SetDynamics(Dynamics);
+		Test.Evolution.ResetVSmoothFromForces(*Test.ParticleHandles[4]);
+
+		Test.Advance();
+
+		if (GetParam())
+		{
+			// All particles wake up, momentum is propagated through
+			EXPECT_TRUE(FConstGenericParticleHandle(Test.ParticleHandles[0])->IsKinematic());
+			EXPECT_FALSE(FConstGenericParticleHandle(Test.ParticleHandles[1])->Sleeping());
 			EXPECT_FALSE(FConstGenericParticleHandle(Test.ParticleHandles[2])->Sleeping());
+			EXPECT_FALSE(FConstGenericParticleHandle(Test.ParticleHandles[3])->Sleeping());
+			EXPECT_FALSE(FConstGenericParticleHandle(Test.ParticleHandles[4])->Sleeping());
+		}
+		else
+		{
+			// Only the first two particles wake up, momentum is too small to be propagated through
+			EXPECT_TRUE(FConstGenericParticleHandle(Test.ParticleHandles[0])->IsKinematic());
+			EXPECT_TRUE(FConstGenericParticleHandle(Test.ParticleHandles[1])->Sleeping());
+			EXPECT_TRUE(FConstGenericParticleHandle(Test.ParticleHandles[2])->Sleeping());
+			EXPECT_FALSE(FConstGenericParticleHandle(Test.ParticleHandles[3])->Sleeping());
 			EXPECT_FALSE(FConstGenericParticleHandle(Test.ParticleHandles[4])->Sleeping());
 		}
+		CVarMomentumPropagation->Set(MomentumPropagationDefault);
 	}
 
 	INSTANTIATE_TEST_SUITE_P(, GraphEvolutionTests,	testing::Bool());
diff --git a/Engine/Source/Programs/HeadlessChaos/Private/HeadlessChaosTestSleeping.cpp b/Engine/Source/Programs/HeadlessChaos/Private/HeadlessChaosTestSleeping.cpp
--- a/Engine/Source/Programs/HeadlessChaos/Private/HeadlessChaosTestSleeping.cpp
+++ b/Engine/Source/Programs/HeadlessChaos/Private/HeadlessChaosTestSleeping.cpp
@@ -47,11 +47,14 @@ namespace ChaosTest
 
 			CVarPartialSleeping = IConsoleManager::Get().FindConsoleVariable(TEXT("p.Chaos.Solver.Sleep.PartialIslandSleep"), false);
 			check(CVarPartialSleeping && CVarPartialSleeping->IsVariableBool());
+			PartialSleepDefault = CVarPartialSleeping->GetBool();
 			CVarPartialSleeping->Set(bPartialSleeping);
 		}
 
 		~FSleepingTest()
 		{
+			CVarPartialSleeping->Set(PartialSleepDefault);
+
 			for (FSingleParticlePhysicsProxy* ParticleProxy : ParticleProxies)
 			{
 				Solver->UnregisterObject(ParticleProxy);
@@ -304,6 +307,7 @@ namespace ChaosTest
 		int32 TickCount;
 
 		IConsoleVariable* CVarPartialSleeping;
+		bool PartialSleepDefault;
 	};
 
 	// Dummy test to make sure a stack of cubes will fall asleep.
@@ -356,6 +360,12 @@ namespace ChaosTest
 		const FReal Dimension = 200;
 		const bool bPartialSleepEnabled = GetParam();
 
+		// Turn off mommentum propagation to test level-based waking only
+		IConsoleVariable* CVarMomentumPropagation = IConsoleManager::Get().FindConsoleVariable(TEXT("p.Chaos.Solver.Sleep.MomentumPropagation"), false);
+		check(CVarMomentumPropagation && CVarMomentumPropagation->IsVariableFloat());
+		const FRealSingle MomentumPropagationDefault = CVarMomentumPropagation->GetFloat();
+		CVarMomentumPropagation->Set(0);
+
 		// Create a stack of cubes
 		FSleepingTest Test(bPartialSleepEnabled);
 		Test.MakeStackOfCubes(NumOfStackedObjects, /*Size=*/Dimension);
@@ -397,6 +407,7 @@ namespace ChaosTest
 			EXPECT_TRUE(Test.ParticleHandles[5]->ObjectState() == EObjectStateType::Dynamic);
 		}
 
+		CVarMomentumPropagation->Set(MomentumPropagationDefault);
 	}
 
 	TEST_P(SleepingTests, CollideWithStack)
@@ -405,6 +416,12 @@ namespace ChaosTest
 		const FReal Dimension = 200;
 		const bool bPartialSleepEnabled = GetParam();
 
+		// Turn off mommentum propagation to test level-based waking only
+		IConsoleVariable* CVarMomentumPropagation = IConsoleManager::Get().FindConsoleVariable(TEXT("p.Chaos.Solver.Sleep.MomentumPropagation"), false);
+		check(CVarMomentumPropagation && CVarMomentumPropagation->IsVariableFloat());
+		const FRealSingle MomentumPropagationDefault = CVarMomentumPropagation->GetFloat();
+		CVarMomentumPropagation->Set(0);
+
 		// Create a stack of cubes
 		FSleepingTest Test(bPartialSleepEnabled);
 		Test.MakeStackOfCubes(NumOfStackedObjects, /*Size=*/Dimension);
@@ -443,6 +460,8 @@ namespace ChaosTest
 			EXPECT_TRUE(Test.ParticleHandles[4]->ObjectState() == EObjectStateType::Dynamic);
 			EXPECT_TRUE(Test.ParticleHandles[5]->ObjectState() == EObjectStateType::Dynamic);
 		}
+
+		CVarMomentumPropagation->Set(MomentumPropagationDefault);
 	}
 
 	TEST_P(SleepingTests, PullCubeOutOfStack)
@@ -451,6 +470,12 @@ namespace ChaosTest
 		const FReal Dimension = 200;
 		const bool bPartialSleepEnabled = GetParam();
 
+		// Turn off mommentum propagation to test level-based waking only
+		IConsoleVariable* CVarMomentumPropagation = IConsoleManager::Get().FindConsoleVariable(TEXT("p.Chaos.Solver.Sleep.MomentumPropagation"), false);
+		check(CVarMomentumPropagation && CVarMomentumPropagation->IsVariableFloat());
+		const FRealSingle MomentumPropagationDefault = CVarMomentumPropagation->GetFloat();
+		CVarMomentumPropagation->Set(0);
+
 		// Create a stack of cubes
 		FSleepingTest Test(bPartialSleepEnabled);
 		Test.MakeStackOfCubes(NumOfStackedObjects, /*Size=*/Dimension);
@@ -487,6 +512,7 @@ namespace ChaosTest
 			EXPECT_TRUE(Test.ParticleHandles[4]->ObjectState() == EObjectStateType::Dynamic);
 		}
 
+		CVarMomentumPropagation->Set(MomentumPropagationDefault);
 	}
 
 	TEST_P(SleepingTests, TeleportCubeOutOfStack)
@@ -540,6 +566,12 @@ namespace ChaosTest
 		const FVec3 Dimensions = FVec3(100, 200, 50);
 		const bool bPartialSleepEnabled = GetParam();
 
+		// Turn off mommentum propagation to test level-based waking only
+		IConsoleVariable* CVarMomentumPropagation = IConsoleManager::Get().FindConsoleVariable(TEXT("p.Chaos.Solver.Sleep.MomentumPropagation"), false);
+		check(CVarMomentumPropagation && CVarMomentumPropagation->IsVariableFloat());
+		const FRealSingle MomentumPropagationDefault = CVarMomentumPropagation->GetFloat();
+		CVarMomentumPropagation->Set(0);
+
 		// Create a stack of cubes
 		FSleepingTest Test(bPartialSleepEnabled);
 		Test.MakeBrickWall(/*NumY=*/ 5, /*NumZ=*/3, /*Dimensions=*/Dimensions);
@@ -599,6 +631,8 @@ namespace ChaosTest
 			EXPECT_TRUE(Test.ParticleHandles[13]->ObjectState() == EObjectStateType::Sleeping);
 			EXPECT_TRUE(Test.ParticleHandles[14]->ObjectState() == EObjectStateType::Sleeping);
 		}
+
+		CVarMomentumPropagation->Set(MomentumPropagationDefault);
 	}
 
 	// Uses a block of bricks without shifting (see SleepingTests::MakeBrickWall)
@@ -608,6 +642,12 @@ namespace ChaosTest
 		const FVec3 Dimensions = FVec3(100, 200, 50);
 		const bool bPartialSleepEnabled = GetParam();
 
+		// Turn off mommentum propagation to test level-based waking only
+		IConsoleVariable* CVarMomentumPropagation = IConsoleManager::Get().FindConsoleVariable(TEXT("p.Chaos.Solver.Sleep.MomentumPropagation"), false);
+		check(CVarMomentumPropagation && CVarMomentumPropagation->IsVariableFloat());
+		const FRealSingle MomentumPropagationDefault = CVarMomentumPropagation->GetFloat();
+		CVarMomentumPropagation->Set(0);
+
 		// Create a stack of cubes
 		FSleepingTest Test(bPartialSleepEnabled);
 		Test.MakeBrickWall(/*NumY=*/ 5, /*NumZ=*/3, /*Dimensions=*/Dimensions, /*Mass=*/2, /*bShiftUnevenRows*/false);
@@ -667,6 +707,8 @@ namespace ChaosTest
 			EXPECT_TRUE(Test.ParticleHandles[13]->ObjectState() == EObjectStateType::Dynamic);
 			EXPECT_TRUE(Test.ParticleHandles[14]->ObjectState() == EObjectStateType::Sleeping);
 		}
+
+		CVarMomentumPropagation->Set(MomentumPropagationDefault);
 	}
 
 	TEST_P(SleepingTests, CollideWithQueueOfCubes)
@@ -676,6 +718,12 @@ namespace ChaosTest
 		const bool bPartialSleepEnabled = true;
 		const bool bAllowPartialSleep = GetParam();
 
+		// Turn off mommentum propagation to test level-based waking only
+		IConsoleVariable* CVarMomentumPropagation = IConsoleManager::Get().FindConsoleVariable(TEXT("p.Chaos.Solver.Sleep.MomentumPropagation"), false);
+		check(CVarMomentumPropagation && CVarMomentumPropagation->IsVariableFloat());
+		const FRealSingle MomentumPropagationDefault = CVarMomentumPropagation->GetFloat();
+		CVarMomentumPropagation->Set(0);
+
 		// Create a stack of cubes
 		FSleepingTest Test(bPartialSleepEnabled);
 		Test.MakeQueueOfCubes(NumOfStackedObjects, /*Size=*/Dimension);
@@ -722,6 +770,129 @@ namespace ChaosTest
 			EXPECT_TRUE(Test.ParticleHandles[3]->ObjectState() == EObjectStateType::Sleeping);
 			EXPECT_TRUE(Test.ParticleHandles[4]->ObjectState() == EObjectStateType::Sleeping);
 		}
+
+		CVarMomentumPropagation->Set(MomentumPropagationDefault);
+	}
+
+	TEST_P(SleepingTests, PropagateMomentumThroughQueue)
+	{
+		const int32 NumOfStackedObjects = 5;
+		const FReal Dimension = 200;
+		const bool bPartialSleepEnabled = true;
+		const FRealSingle MomentumPropagation = GetParam() ? 0.5f : 0.1f;
+		IConsoleVariable* CVarMomentumPropagation = IConsoleManager::Get().FindConsoleVariable(TEXT("p.Chaos.Solver.Sleep.MomentumPropagation"), false);
+		check(CVarMomentumPropagation && CVarMomentumPropagation->IsVariableFloat());
+		const FRealSingle MomentumPropagationDefault = CVarMomentumPropagation->GetFloat();
+		CVarMomentumPropagation->Set(MomentumPropagation);
+		// Create a stack of cubes
+		FSleepingTest Test(bPartialSleepEnabled);
+		Test.MakeQueueOfCubes(NumOfStackedObjects, /*Size=*/Dimension);
+		// Simulate until all particles are sleeping
+		Test.AdvanceUntilSleeping();
+		// Check that everything is sleeping
+		for (FRigidBodyHandle_External* ParticleHandle : Test.ParticleHandles)
+		{
+			EXPECT_TRUE(ParticleHandle->ObjectState() == EObjectStateType::Sleeping);
+		}
+		// Create another particle which will collide with the queue
+		FRigidBodyHandle_External* CollidingCube = Test.MakeCube(FVec3(0, -200, 100), 200);
+		// Note: We use the VSmooth for momentum propagation (0.3 * v_old + 0.7 * v_new)
+		// Thus, only 30% of the velocity are actually transmitted.
+		CollidingCube->SetLinearImpulse(FVec3(0, 200, 0), true);
+		Test.Advance();
+		//// Debug print of all particle positions
+		//for (int32 Id = 0; Id < Test.ParticleHandles.Num(); ++Id)
+		//{
+		//	FRigidBodyHandle_External* ParticleHandle = Test.ParticleHandles[Id];
+		//	printf("Particle Id: %d, State: %d, Position: %f, %f, %f\n", Id,
+		//		ParticleHandle->ObjectState(), ParticleHandle->GetX().X, ParticleHandle->GetX().Y, ParticleHandle->GetX().Z);
+		//}
+		if (GetParam()) // 50% propagation
+		{
+			// Check that everything is awake
+			for (FRigidBodyHandle_External* ParticleHandle : Test.ParticleHandles)
+			{
+				EXPECT_TRUE(ParticleHandle->ObjectState() == EObjectStateType::Dynamic);
+			}
+		}
+		else // 10% propagation
+		{
+			// Only the colliding cube and first two particles in the queue wake up
+			EXPECT_TRUE(Test.ParticleHandles[5]->ObjectState() == EObjectStateType::Dynamic);
+			EXPECT_TRUE(Test.ParticleHandles[0]->ObjectState() == EObjectStateType::Dynamic);
+			EXPECT_TRUE(Test.ParticleHandles[1]->ObjectState() == EObjectStateType::Dynamic);
+			EXPECT_TRUE(Test.ParticleHandles[2]->ObjectState() == EObjectStateType::Sleeping);
+			EXPECT_TRUE(Test.ParticleHandles[3]->ObjectState() == EObjectStateType::Sleeping);
+			EXPECT_TRUE(Test.ParticleHandles[4]->ObjectState() == EObjectStateType::Sleeping);
+		}
+		CVarMomentumPropagation->Set(MomentumPropagationDefault);
+	}
+	// Uses a block of bricks without shifting (see SleepingTests::MakeBrickWall)
+	TEST_P(SleepingTests, PropagateMomentumThroughBlockOfBricks)
+	{
+		const int32 NumOfStackedObjects = 5;
+		const FVec3 Dimensions = FVec3(100, 200, 50);
+		const bool bPartialSleepEnabled = true;
+		const FRealSingle MomentumPropagation = GetParam() ? 0.5f : 0.1f;
+		IConsoleVariable* CVarMomentumPropagation = IConsoleManager::Get().FindConsoleVariable(TEXT("p.Chaos.Solver.Sleep.MomentumPropagation"), false);
+		check(CVarMomentumPropagation && CVarMomentumPropagation->IsVariableFloat());
+		const FRealSingle MomentumPropagationDefault = CVarMomentumPropagation->GetFloat();
+		CVarMomentumPropagation->Set(MomentumPropagation);
+		// Create a stack of cubes
+		FSleepingTest Test(bPartialSleepEnabled);
+		Test.MakeBrickWall(/*NumY=*/ 5, /*NumZ=*/3, /*Dimensions=*/Dimensions, /*Mass=*/2, /*bShiftUnevenRows*/false);
+		// Simulate until all particles are sleeping
+		Test.AdvanceUntilSleeping(/*MaxIterations=*/500);
+		// Check that everything is sleeping
+		for (FRigidBodyHandle_External* ParticleHandle : Test.ParticleHandles)
+		{
+			EXPECT_TRUE(ParticleHandle->ObjectState() == EObjectStateType::Sleeping);
+		}
+		// Cube spawned about particles 11 & 12 (i.e. the 2nd and 3rd of the top row)
+		FRigidBodyHandle_External* CollidingCube = Test.MakeCube(FVec3(0, 300, 150), 50);
+		// Note: We use the VSmooth for momentum propagation (0.3 * v_old + 0.7 * v_new)
+		// Thus, only 30% of the velocity are actually transmitted.
+		CollidingCube->SetLinearImpulse(FVec3(0, 0, -200), true);
+		Test.Advance();
+		//// Debug print of all particle positions
+		//for (int32 Id = 0; Id < Test.ParticleHandles.Num(); ++Id)
+		//{
+		//	FRigidBodyHandle_External* ParticleHandle = Test.ParticleHandles[Id];
+		//	printf("Particle Id: %d, State: %d, Position: %f, %f, %f\n", Id, 
+		//		ParticleHandle->ObjectState(), ParticleHandle->GetX().X, ParticleHandle->GetX().Y, ParticleHandle->GetX().Z);
+		//}
+		if (GetParam()) // 50% propagation
+		{
+			// The entire brick wall will wake up
+			for (FRigidBodyHandle_External* ParticleHandle : Test.ParticleHandles)
+			{
+				EXPECT_TRUE(ParticleHandle->ObjectState() == EObjectStateType::Dynamic);
+			}
+		}
+		else // 10% propagation
+		{
+			// Colliding with particles 11 & 12
+			// Only particles with constraint distance <= 2 wake up.
+			// 1st row
+			EXPECT_TRUE(Test.ParticleHandles[0]->ObjectState() == EObjectStateType::Sleeping);
+			EXPECT_TRUE(Test.ParticleHandles[1]->ObjectState() == EObjectStateType::Sleeping);
+			EXPECT_TRUE(Test.ParticleHandles[2]->ObjectState() == EObjectStateType::Sleeping);
+			EXPECT_TRUE(Test.ParticleHandles[3]->ObjectState() == EObjectStateType::Sleeping);
+			EXPECT_TRUE(Test.ParticleHandles[4]->ObjectState() == EObjectStateType::Sleeping);
+			// 2nd row
+			EXPECT_TRUE(Test.ParticleHandles[5]->ObjectState() == EObjectStateType::Dynamic);
+			EXPECT_TRUE(Test.ParticleHandles[6]->ObjectState() == EObjectStateType::Dynamic);
+			EXPECT_TRUE(Test.ParticleHandles[7]->ObjectState() == EObjectStateType::Dynamic);
+			EXPECT_TRUE(Test.ParticleHandles[8]->ObjectState() == EObjectStateType::Dynamic);
+			EXPECT_TRUE(Test.ParticleHandles[9]->ObjectState() == EObjectStateType::Sleeping);
+			// 3rd row
+			EXPECT_TRUE(Test.ParticleHandles[10]->ObjectState() == EObjectStateType::Dynamic);
+			EXPECT_TRUE(Test.ParticleHandles[11]->ObjectState() == EObjectStateType::Dynamic);
+			EXPECT_TRUE(Test.ParticleHandles[12]->ObjectState() == EObjectStateType::Dynamic);
+			EXPECT_TRUE(Test.ParticleHandles[13]->ObjectState() == EObjectStateType::Dynamic);
+			EXPECT_TRUE(Test.ParticleHandles[14]->ObjectState() == EObjectStateType::Dynamic);
+		}
+		CVarMomentumPropagation->Set(MomentumPropagationDefault);
 	}
 
 	INSTANTIATE_TEST_SUITE_P(, SleepingTests, testing::Bool());
diff --git a/Engine/Source/Runtime/Experimental/Chaos/Private/Chaos/Island/IslandManager.cpp b/Engine/Source/Runtime/Experimental/Chaos/Private/Chaos/Island/IslandManager.cpp
--- a/Engine/Source/Runtime/Experimental/Chaos/Private/Chaos/Island/IslandManager.cpp
+++ b/Engine/Source/Runtime/Experimental/Chaos/Private/Chaos/Island/IslandManager.cpp
@@ -67,6 +67,10 @@ namespace Chaos::CVars
 	bool bChaosSolverWakeOnMovingOnly = true;
 	FAutoConsoleVariableRef CVarChaosSolverWakeOnMovingOnly(TEXT("p.Chaos.Solver.Sleep.WakeOnMovingOnly"), bChaosSolverWakeOnMovingOnly, TEXT("[EXPERIMENTAL] Only newly detected constraints connected to a moving particle will trigger wake-up event (partial island sleep only)."));
 
+	/** Cvar to enable/disable momentum propagation through the constraint graph (partial island sleeping only) */
+	FRealSingle bChaosSolverMomentumPropagation = 0.25f;
+	FAutoConsoleVariableRef CVarChaosSolverMomentumPropagation(TEXT("p.Chaos.Solver.Sleep.MomentumPropagation"), bChaosSolverMomentumPropagation, TEXT("[EXPERIMENTAL] Ratio of the impact momentum transferred from one sleeping particle to another through a constraint during partial wake event (partial island sleeping only, 0 to disable, 1 to propagate with dissipation).[def:0.25]"));
+
 	/** Cvar to override the sleep counter threshold if necessary */
 	int32 ChaosSolverCollisionDefaultSleepCounterThreshold = 20;
 	FAutoConsoleVariableRef CVarChaosSolverCollisionDefaultSleepCounterThreshold(TEXT("p.Chaos.Solver.Sleep.Defaults.SleepCounterThreshold"), ChaosSolverCollisionDefaultSleepCounterThreshold, TEXT("Default counter threshold for sleeping.[def:20]"));
@@ -310,6 +314,15 @@ namespace Chaos::Private
 		}
 	}
 
+	// Validate that the number of elements in NodeQueue and MomentumQueue is compatible with the wake algorithm.
+	// a) bChaosSolverMomentumPropagation <= 0.0f: No momentum information is stored.
+	// b) bChaosSolverMomentumPropagation > 0.0f: Both queues must have the same number of elements.
+	inline void ValidateWakeQueues(TArray<FPBDIslandParticle*>& NodeQueue, TArray<FPBDParticleMomentumSq>& MomentumQueue)
+	{
+		check((CVars::bChaosSolverMomentumPropagation > 0.0f && NodeQueue.Num() == MomentumQueue.Num())
+			|| (CVars::bChaosSolverMomentumPropagation <= 0.0f && MomentumQueue.Num() == 0));
+	}
+
 	//////////////////////////////////////////////////////////////////////////////////////////////////
 	//////////////////////////////////////////////////////////////////////////////////////////////////
 	// 
@@ -2655,7 +2668,12 @@ namespace Chaos::Private
 	{
 		TArray<FGeometryParticleHandle*> ParticlesToWakeUp; // Particles which will be woken up after processing all partial wake events.
 		TArray<FPBDIslandParticle*> NodeQueue;				// Queue of island nodes to be traversed by partial waking.
+		TArray<FPBDParticleMomentumSq> MomentumQueue;		// Linear and angular momentum of island nodes in NodeQueue.
 		NodeQueue.Reserve(Island->Nodes.Num());
+		if (CVars::bChaosSolverMomentumPropagation > 0.0f)
+		{
+			MomentumQueue.Reserve(Island->Nodes.Num());
+		}
 
 		// Phase 1 - Waking triggered by new constraints
 		// Step 1a - Find trigger particles (i.e. particles connected to new constraints with significant motion)
@@ -2667,21 +2685,21 @@ namespace Chaos::Private
 				{
 					FPBDIslandParticle* Node0 = IslandEdge->Nodes[0];
 					FPBDIslandParticle* Node1 = IslandEdge->Nodes[1];
-
 					if (Node0 && !Node0->Flags.bWasProcessed)
 					{
-						ProcessTriggerParticle(Node0, NodeQueue, /*bAlwaysTrigger=*/ !CVars::bChaosSolverWakeOnMovingOnly);
+						ProcessTriggerParticle(Node0, NodeQueue, MomentumQueue, /*bAlwaysTrigger=*/ !CVars::bChaosSolverWakeOnMovingOnly);
 						Node0->Flags.bWasProcessed = true;
 					}
 					if (Node1 && !Node1->Flags.bWasProcessed)
 					{
-						ProcessTriggerParticle(Node1, NodeQueue, /*bAlwaysTrigger=*/ !CVars::bChaosSolverWakeOnMovingOnly);
+						ProcessTriggerParticle(Node1, NodeQueue, MomentumQueue, /*bAlwaysTrigger=*/ !CVars::bChaosSolverWakeOnMovingOnly);
 						Node1->Flags.bWasProcessed = true;
 					}
 				}
 			}
 		}
 		const int32 NumNodesConstraintWake = NodeQueue.Num();
+		ValidateWakeQueues(NodeQueue, MomentumQueue);
 		bool bResetEdgeFlags = false;
 
 		// Step 1b - Propagate wake event
@@ -2695,20 +2713,23 @@ namespace Chaos::Private
 		}
 		else // Otherwise, propagate the partial wake event through the constraint graph.
 		{
-			bResetEdgeFlags = PropagatePartialIslandWake(NodeQueue, ParticlesToWakeUp, NumNodesConstraintWake, /*bTriggerNodesAlwaysPropagate=*/ false);
+			bResetEdgeFlags = PropagatePartialIslandWake(NodeQueue, MomentumQueue, ParticlesToWakeUp, NumNodesConstraintWake, /*bTriggerNodesAlwaysPropagate=*/ false);
 		}
 
 		// Phase 2 - Waking triggered by particle state changes
 		// Step 2a - Find trigger particles (i.e. particles which were just woken up the user) 
 		NodeQueue.Reset();
+		MomentumQueue.Reset();
 		for (FPBDIslandParticle* Node : Island->Nodes)
 		{
 			if (Node->Flags.bWasUpdated)
 			{
-				ProcessTriggerParticle(Node, NodeQueue, /*bAlwaysTrigger=*/ true);
+				ProcessTriggerParticle(Node, NodeQueue, MomentumQueue, /*bAlwaysTrigger=*/ true);
 			}
 		}
 		const int32 NumNodesNodeWake = NodeQueue.Num();
+		ValidateWakeQueues(NodeQueue, MomentumQueue);
+
 
 		// Step 2b - Propagate wake event
 		// Partial island wake disabled, but at least one trigger particle found? Then, perform a full island wake.
@@ -2726,14 +2747,14 @@ namespace Chaos::Private
 			{
 				ResetEdgeFlags(Island);
 			}
-			PropagatePartialIslandWake(NodeQueue, ParticlesToWakeUp, NumNodesNodeWake, /*bTriggerNodesAlwaysPropagate=*/ true);
+			PropagatePartialIslandWake(NodeQueue, MomentumQueue, ParticlesToWakeUp, NumNodesNodeWake, /*bTriggerNodesAlwaysPropagate=*/ true);
 		}
 
 		// Phase 3 - Wake up all relevant particles and reset sleep counter
-		WakeUpParticles(ParticlesToWakeUp);
+		WakeUpParticles(ParticlesToWakeUp, /*bResetSleepCounter=*/ false);
 	}
 
-	bool FPBDIslandManager::PropagatePartialIslandWake(TArray<FPBDIslandParticle*>& NodeQueue,
+	bool FPBDIslandManager::PropagatePartialIslandWake(TArray<FPBDIslandParticle*>& NodeQueue, TArray<FPBDParticleMomentumSq>& MomentumQueue,
 		TArray<FGeometryParticleHandle*>& ParticlesToWakeUp, const int32 NumNodesInQueue, const bool bTriggerNodesAlwaysPropagate)
 	{
 		bool bResetEdgeFlags = false;
@@ -2753,15 +2774,25 @@ namespace Chaos::Private
 					if ((OtherNode != nullptr) && OtherNode->Flags.bIsDynamic)
 					{
 						// Wake and propagate always to the next constraint if
-						// 1) We process one of the trigger nodes and propagate through a new constraint (if bTriggerNodesAlwaysPropagate == false).
-						// 2) We propagate upward (away from kinematic) and bChaosSolverLevelBasedWake is enabled.
 						const bool bIsTriggerNode = NodeIndex < NumNodesInQueue;
-						const bool bPropagateThrough = bIsTriggerNode && (Edge->Flags.bIsNew || bTriggerNodesAlwaysPropagate);
+						const bool bIsConfirmedPropagation = bIsTriggerNode && (Edge->Flags.bIsNew || bTriggerNodesAlwaysPropagate);
 						const bool bIsUpwardPropagation = OtherNode->Level >= Edge->Level;	// Propagation upward through the levels, i.e. away from closest kinematic
-						if (bPropagateThrough || bIsUpwardPropagation)
+
+						// 1) We process one of the trigger nodes and propagate through a new constraint (if bTriggerNodesAlwaysPropagate == false).
+						// 2) We propagate upward in the constraint graph levels (away from kinematic).
+						if (bIsConfirmedPropagation || bIsUpwardPropagation)
+						{
+							const FPBDParticleMomentumSq* NodeMomentum = MomentumQueue.IsEmpty() ? nullptr : &MomentumQueue[NodeIndex];
+							ProcessParticleWake(OtherNode, NodeQueue[NodeIndex], NodeMomentum, NodeQueue, MomentumQueue, ParticlesToWakeUp, /*bAlwaysPropagate=*/ true);
+						}
+						// 3) We propagate the impact momentum of the trigger particle through the constraint graph.
+						else if (CVars::bChaosSolverMomentumPropagation > 0.0f)
 						{
-							ProcessParticleWake(OtherNode, NodeQueue[NodeIndex], NodeQueue, ParticlesToWakeUp, /*bAlwaysPropagate=*/ true);
+							check(!MomentumQueue.IsEmpty());
+							ProcessParticleWake(OtherNode, NodeQueue[NodeIndex], &MomentumQueue[NodeIndex], NodeQueue, MomentumQueue, ParticlesToWakeUp, /*bAlwaysPropagate=*/ false);
 						}
+
+						ValidateWakeQueues(NodeQueue, MomentumQueue);
 					}
 				}
 			}
@@ -2769,7 +2800,7 @@ namespace Chaos::Private
 		return bResetEdgeFlags;
 	}
 
-	bool FPBDIslandManager::ProcessTriggerParticle(FPBDIslandParticle* Node, TArray<FPBDIslandParticle*>& NodeQueue, const bool bAlwaysTrigger)
+	bool FPBDIslandManager::ProcessTriggerParticle(FPBDIslandParticle* Node, TArray<FPBDIslandParticle*>& NodeQueue, TArray<FPBDParticleMomentumSq>& MomentumQueue, const bool bAlwaysTrigger)
 	{
 		const auto IsMoving = [](FPBDIslandParticle* Node) -> bool
 			{
@@ -2812,6 +2843,57 @@ namespace Chaos::Private
 				return Node && !!Node->Flags.bIsDynamic;
 			};
 
+		const auto ComputeLinearMomentumSq = [](const FPBDIslandParticle* Node) -> FRealSingle
+			{
+				if (Node == nullptr)
+				{
+					return FRealSingle(0);
+				}
+				const FPBDRigidParticleHandle* Rigid = Node->GetParticle()->CastToRigidParticle();
+				if (Rigid == nullptr)
+				{
+					return FRealSingle(0);
+				}
+
+				// Use the smoothed velocity to check for motion, unless it is zero (i.e. was never set).
+				FVec3f V = Rigid->VSmooth();
+				if (V.SizeSquared() == 0)
+				{
+					// Use the regular velocity from the last time step
+					V = Rigid->GetPreV();
+				}
+				const FRealSingle Mass = FRealSingle(Rigid->M());
+				const FVec3f LinearMomentum = Mass * V;
+				return LinearMomentum.SizeSquared();
+			};
+
+		const auto ComputeAngularMomentumSq = [](const FPBDIslandParticle* Node)->FRealSingle
+			{
+				if (Node == nullptr)
+				{
+					return FRealSingle(0);
+				}
+				const FPBDRigidParticleHandle* Rigid = Node->GetParticle()->CastToRigidParticle();
+				if (Rigid == nullptr)
+				{
+					return FRealSingle(0);
+				}
+
+				// Use the smoothed velocity to check for motion, unless it is zero (i.e. was never set).
+				FVec3f W = Rigid->WSmooth();
+				if (W.SizeSquared() == 0)
+				{
+					// Use the regular velocity from the last time step
+					W = Rigid->GetPreW();
+				}
+
+				// Compute the magnitude of the angular momentum in body-space (equivalent and more efficient)
+				const FMatrix33 Transform = (Rigid->GetR() * Rigid->RotationOfMass()).ToMatrix().GetTransposed();
+				const FVec3f WLocal = Transform * W;
+				const FVec3f AngularMomentum = Rigid->I() * WLocal;
+				return AngularMomentum.SizeSquared();
+			};
+
 		// Stop wake-up event if the node or particle is empty
 		if (Node == nullptr)
 		{
@@ -2835,8 +2917,15 @@ namespace Chaos::Private
 		// The particle must be awake to trigger a partial wake-up event.
 		if (bIsAwake)
 		{
+			// Add node to the queue if (a) above the motion wake threshold or (b) always triggering.
 			if (bAlwaysTrigger || IsMoving(Node))
 			{
+				if (CVars::bChaosSolverMomentumPropagation > 0.0f)
+				{
+					const FRealSingle LinearMomentumSq = ComputeLinearMomentumSq(Node);
+					const FRealSingle AngularMomentumSq = ComputeAngularMomentumSq(Node);
+					MomentumQueue.Add(FPBDParticleMomentumSq(LinearMomentumSq, AngularMomentumSq));
+				}
 				NodeQueue.Add(Node);
 				return true;
 			}
@@ -2845,7 +2934,9 @@ namespace Chaos::Private
 		return false;
 	}
 
-	bool FPBDIslandManager::ProcessParticleWake(FPBDIslandParticle* Node, const FPBDIslandParticle* OtherNode, TArray<FPBDIslandParticle*>& NodeQueue, TArray<FGeometryParticleHandle*>& WakeUpQueue, const bool bAlwaysPropagate)
+	bool FPBDIslandManager::ProcessParticleWake(FPBDIslandParticle* Node, const FPBDIslandParticle* OtherNode, const FPBDParticleMomentumSq* OtherMomentumSq,
+		TArray<FPBDIslandParticle*>& NodeQueue, TArray<FPBDParticleMomentumSq>& MomentumQueue, TArray<FGeometryParticleHandle*>& WakeUpQueue,
+		const bool bAlwaysPropagate)
 	{
 		const auto IsAwake = [](const FPBDIslandParticle* Node) -> bool
 			{
@@ -2870,6 +2961,7 @@ namespace Chaos::Private
 
 		const bool bIsAwake = IsAwake(Node);
 		const bool bIsDynamic = IsDynamic(Node);
+		const int32 SleepCounterValue = 0;
 
 		// Stop wake-up event if the particle is not dynamic
 		if (!bIsDynamic)
@@ -2877,15 +2969,56 @@ namespace Chaos::Private
 			return false;
 		}
 
-		// Otherwise, the particle is added to the node queue and woken up if needed.
-		NodeQueue.Add(Node);
-		Rigid->SetSleepCounter(0); // Force reset in partial wake events to prevent immediate re-sleep
-		if (!bIsAwake)
+		// No need for momentum computation, if momentum propagation disabled and we always propagate to the next particle.
+		if (CVars::bChaosSolverMomentumPropagation <= 0.0f && bAlwaysPropagate)
 		{
-			Node->Flags.bIsSleeping = false;
-			WakeUpQueue.Add(Rigid);
+			NodeQueue.Add(Node);
+			Rigid->SetSleepCounter(int8(SleepCounterValue)); // Force reset in partial wake events to prevent immediate re-sleep
+			if (!bIsAwake)
+			{
+				Node->Flags.bIsSleeping = false;
+				WakeUpQueue.Add(Rigid);
+			}
+			return true;
 		}
-		return true;
+
+		// Compute the minimum momentum required to make the particle move significantly.
+		// Approximation of the squared angular momentum using the smallest inertia diagonal in body-space
+		// This gives us the lower bound of when to consider the particle to be moving.
+		// linear: ||p||^2 = m^2 v^2
+		// angular: ||L||^2 = w^T I^2 w ~= w^2 * I_min^2
+		const FRealSingle LinearMomentumSq = FMath::Square(FRealSingle(Rigid->M())) * Node->WakeLinearThresholdSq;
+		const FRealSingle AngularMomentumSq = FMath::Square(FRealSingle(Rigid->I().Min())) * Node->WakeAngularThresholdSq;
+
+		// Other momentum needed to perform momentum propagation
+		if (CVars::bChaosSolverMomentumPropagation > 0.0f && OtherMomentumSq)
+		{
+			const FRealSingle ReceivedLinearMomentumSq = OtherMomentumSq->Linear;
+			const FRealSingle ReceivedAngularMomentumSq = OtherMomentumSq->Angular;
+
+			if (bAlwaysPropagate
+				|| ReceivedLinearMomentumSq > LinearMomentumSq
+				|| ReceivedAngularMomentumSq > AngularMomentumSq)
+			{
+				// NOTE: Since we deal with the squared momentum, we also square the propagation discount factor.
+				// @todo (Chaos): It would be great to use the direction of the momentum as opposed to the absolute value. 
+				// For contacts, if the received momentum acts in the opposing direction of the relative position of the particle,
+				// there is no need to propagate the wake through the particle.
+				const FRealSingle ResidualLinearMomentumSq = ReceivedLinearMomentumSq * FMath::Square(CVars::bChaosSolverMomentumPropagation);
+				const FRealSingle ResidualAngularMomentumSq = ReceivedAngularMomentumSq * FMath::Square(CVars::bChaosSolverMomentumPropagation);
+
+				MomentumQueue.Add(FPBDParticleMomentumSq(ResidualLinearMomentumSq, ResidualAngularMomentumSq));
+				NodeQueue.Add(Node);
+				Rigid->SetSleepCounter(int8(SleepCounterValue)); // Force reset in partial wake events to prevent immediate re-sleep
+				if (!bIsAwake)
+				{
+					Node->Flags.bIsSleeping = false;
+					WakeUpQueue.Add(Rigid);
+				}
+				return true;
+			}
+		}
+		return false;
 	}
 
 	void FPBDIslandManager::SyncConstraintSleepStateToParticles(FPBDIsland* Island)
diff --git a/Engine/Source/Runtime/Experimental/Chaos/Public/Chaos/Island/IslandManager.h b/Engine/Source/Runtime/Experimental/Chaos/Public/Chaos/Island/IslandManager.h
--- a/Engine/Source/Runtime/Experimental/Chaos/Public/Chaos/Island/IslandManager.h
+++ b/Engine/Source/Runtime/Experimental/Chaos/Public/Chaos/Island/IslandManager.h
@@ -24,6 +24,16 @@ namespace Chaos::Private
 {
 	class FPBDIslandMergeSet;
 
+	struct FPBDParticleMomentumSq
+	{
+		FPBDParticleMomentumSq(const FRealSingle Linear = 0.0f, const FRealSingle Angular = 0.0f)
+			: Linear(Linear)
+			, Angular(Angular)
+		{}
+		FRealSingle Linear;
+		FRealSingle Angular;
+	};
+
 	/**
 	* A particle in the constraint graph.
 	*/
@@ -607,9 +617,12 @@ namespace Chaos::Private
 
 		// Partial Island Waking
 		CHAOS_API void ProcessPartialIslandWake(FPBDIsland* Island);
-		CHAOS_API bool PropagatePartialIslandWake(TArray<FPBDIslandParticle*>& NodeQueue, TArray<FGeometryParticleHandle*>& ParticlesToWakeUp, const int32 NumNodesConstraintWake, const bool bTriggerNodesAlwaysPropagate);
-		CHAOS_API bool ProcessTriggerParticle(FPBDIslandParticle* Node, TArray<FPBDIslandParticle*>& NodeQueue, const bool bAlwaysTrigger);
-		CHAOS_API bool ProcessParticleWake(FPBDIslandParticle* Node, const FPBDIslandParticle* OtherNode, TArray<FPBDIslandParticle*>& NodeQueue, TArray<FGeometryParticleHandle*>& WakeUpQueue, const bool bAlwaysPropagate);
+		CHAOS_API bool PropagatePartialIslandWake(TArray<FPBDIslandParticle*>& NodeQueue, TArray<FPBDParticleMomentumSq>& MomentumQueue,
+			TArray<FGeometryParticleHandle*>& ParticlesToWakeUp, const int32 NumNodesConstraintWake, const bool bTriggerNodesAlwaysPropagate);
+		CHAOS_API bool ProcessTriggerParticle(FPBDIslandParticle* Node, TArray<FPBDIslandParticle*>& NodeQueue, TArray<FPBDParticleMomentumSq>& MomentumQueue, const bool bAlwaysTrigger);
+		CHAOS_API bool ProcessParticleWake(FPBDIslandParticle* Node, const FPBDIslandParticle* OtherNode, const FPBDParticleMomentumSq* OtherMomentum,
+			TArray<FPBDIslandParticle*>& NodeQueue, TArray<FPBDParticleMomentumSq>& MomentumQueue,
+			TArray<FGeometryParticleHandle*>& WakeUpQueue, const bool bAlwaysPropagate);
 		CHAOS_API void SyncConstraintSleepStateToParticles(FPBDIsland* Island);
 		CHAOS_API void WakeUpParticles(TArray<FGeometryParticleHandle*>& ParticlesToWake, const bool bResetSleepCounter = true);
 		CHAOS_API void ResetParticleSleepCounter(FPBDIsland* Island);
