Chaos - cvar to control radius of debris collision sphere

This is an experiment at the moment. Cvar will be moved to a config property or per-body property if useful.

#rnx

diff --git a/Engine/Source/Runtime/Experimental/Chaos/Private/Chaos/Collision/ParticlePairMidPhase.cpp b/Engine/Source/Runtime/Experimental/Chaos/Private/Chaos/Collision/ParticlePairMidPhase.cpp
--- a/Engine/Source/Runtime/Experimental/Chaos/Private/Chaos/Collision/ParticlePairMidPhase.cpp
+++ b/Engine/Source/Runtime/Experimental/Chaos/Private/Chaos/Collision/ParticlePairMidPhase.cpp
@@ -51,8 +51,8 @@ namespace Chaos
 
 	namespace CVars
 	{
-		Chaos::FRealSingle Chaos_Collision_CullDistanceScaleInverseSize = 0.01f;	// 100cm
-		Chaos::FRealSingle Chaos_Collision_MinCullDistanceScale = 1.0f;
+		FRealSingle Chaos_Collision_CullDistanceScaleInverseSize = 0.01f;	// 100cm
+		FRealSingle Chaos_Collision_MinCullDistanceScale = 1.0f;
 		FAutoConsoleVariableRef CVarChaos_Collision_CullDistanceReferenceSize(TEXT("p.Chaos.Collision.CullDistanceReferenceSize"), Chaos_Collision_CullDistanceScaleInverseSize, TEXT(""));
 		FAutoConsoleVariableRef CVarChaos_Collision_MinCullDistanecScale(TEXT("p.Chaos.Collision.MinCullDistanceScale"), Chaos_Collision_MinCullDistanceScale, TEXT(""));
 
@@ -72,6 +72,9 @@ namespace Chaos
 		bool bChaosMidPhaseModifierConvexOptimizerFix = true;
 		FAutoConsoleVariableRef CVarChaosMidPhaseModifierConvexOptimizerFix(TEXT("p.Chaos.Bugs.MidPhaseModifierConvexOptimizerFix"), bChaosMidPhaseModifierConvexOptimizerFix, TEXT(""));
 
+		FRealSingle ChaosOneWayPairSphereRadiusAlpha = 1.0f;
+		FAutoConsoleVariableRef CVarChaosOneWayPairSphereRadiusAlpha(TEXT("p.Chaos.Solver.OneWayPairSphereRadiusAlpha"), ChaosOneWayPairSphereRadiusAlpha, TEXT("If one-way interactions are treated as spheres, the diameter is lerp(Size.Min, Size.Max, RadiusAlpha)"));
+
 		extern int32 ChaosOneWayInteractionPairCollisionMode;
 
 		extern bool bChaosForceMACD;
@@ -2234,7 +2237,8 @@ namespace Chaos
 		// @todo(chaos): maybe we should only consider the bounds of the shapes that pass the ShapePairNarrowPhaseFilter?
 		const FVec3 Center = InParticle->LocalBounds().Center();
 		const FVec3 Extents = InParticle->LocalBounds().Extents();
-		const FReal Radius = 0.5 * Extents.GetAbsMax();
+		const FReal RadiusAlpha = FMath::Clamp(ChaosOneWayPairSphereRadiusAlpha, 0.0f, 1.0f);
+		const FReal Radius = 0.5f * FMath::Lerp(Extents.GetMin(), Extents.GetMax(), RadiusAlpha);
 
 		OutSphere = FImplicitSphere3(Center, Radius);
 	}
