Allow streaming ondemand containers to register to use the correct hostgroup when not using the default.

#rnx

- Prevent the EOnDemandContainerFlags flag from being overridden in a container if they already contain that flag. I was seeing cases where IAS containers (with EOnDemandContainerFlags::StreamOnDemand) were having the IAD flag (EOnDemandContainerFlags::InstallOnDemand) being forced on via the FOnDemandMountArgs which would force the IAS data to be pre-installed via IAS.
-- This will need a re-work to clean up the code paths, as it stands it is very unclear when the container flags should be respected vs the mount arg flags. Although I cannot think of any scenario where we would want to force IAS data to be installed via IAD instead.
-- The two flags are pretty mutually exclusive so we might want to rethink how we are encoding the infomation.
- In the older code paths IAS and IAD have different default hostgroups so IAS was hardcoded to always ignore the MountArgs and always use "DefaultOnDemand". Instead we should check to see if the "Default" hostgroup is being applied and if so then change that to "DefaultOnDemand" and if not then use the hostgroup provided by MountArgs.

diff --git a/Engine/Source/Runtime/Experimental/IoStore/OnDemand/Private/OnDemandIoStore.cpp b/Engine/Source/Runtime/Experimental/IoStore/OnDemand/Private/OnDemandIoStore.cpp
--- a/Engine/Source/Runtime/Experimental/IoStore/OnDemand/Private/OnDemandIoStore.cpp
+++ b/Engine/Source/Runtime/Experimental/IoStore/OnDemand/Private/OnDemandIoStore.cpp
@@ -1273,26 +1273,14 @@ FIoStatus FOnDemandIoStore::TickMountRequest(FMountRequest& MountRequest)
 
 	for (const FSharedOnDemandContainer& Container : RequestedContainers)
 	{
-		if (EnumHasAnyFlags(Args.Options, EOnDemandMountOptions::StreamOnDemand))
+		if (EnumHasAnyFlags(Container->Flags, EOnDemandContainerFlags::StreamOnDemand | EOnDemandContainerFlags::InstallOnDemand) == false)
 		{
-			if (!EnumHasAnyFlags(Container->Flags, EOnDemandContainerFlags::StreamOnDemand))
-			{
-				UE_LOG(LogIoStoreOnDemand, Log, TEXT("Overriding container flag 'StreamOnDemand' from mount option"));
-				EnumAddFlags(Container->Flags, EOnDemandContainerFlags::StreamOnDemand);
-			}
+			return FIoStatusBuilder(EIoErrorCode::InvalidParameter) << TEXT("Does not have either a 'StreamOnDemand' or 'InstallOnDemand' EOnDemandContainerFlags flag");
 		}
-		else if (EnumHasAnyFlags(Args.Options, EOnDemandMountOptions::InstallOnDemand))
-		{
-			if (!EnumHasAnyFlags(Container->Flags, EOnDemandContainerFlags::InstallOnDemand))
-			{
-				UE_LOG(LogIoStoreOnDemand, Log, TEXT("Overriding container flag 'InstallOnDemand' from mount option"));
-				EnumAddFlags(Container->Flags, EOnDemandContainerFlags::InstallOnDemand);
-			}
 
-			if (EnumHasAnyFlags(Args.Options, EOnDemandMountOptions::WithSoftReferences))
-			{
-				EnumAddFlags(Container->Flags, EOnDemandContainerFlags::WithSoftReferences);
-			}
+		if (EnumHasAnyFlags(Args.Options, EOnDemandMountOptions::WithSoftReferences))
+		{
+			EnumAddFlags(Container->Flags, EOnDemandContainerFlags::WithSoftReferences);
 		}
 
 		if (const FIoStatus Status = SetupHostGroup(Container, Args); !Status.IsOk())
@@ -1696,7 +1684,9 @@ FIoStatus FOnDemandIoStore::SetupHostGroup(const FSharedOnDemandContainer& Conta
 	}
 
 	// This path will be removed once the HTTP I/O dispatcher is on by default
-	if (EnumHasAnyFlags(Container->Flags, EOnDemandContainerFlags::StreamOnDemand))
+	// StreamOnDemand (IAS) uses a different default hostgroup in this code path so if the MountArgs are requesting 'Default' we need
+	// to route that to the IAS default instead.
+	if (EnumHasAnyFlags(Container->Flags, EOnDemandContainerFlags::StreamOnDemand) && MountArgs.HostGroupName == TEXT("Default"))
 	{
 		//TODO: Remove this special case for streaming containers
 		Container->HostGroupName	= FName("DefaultOnDemand");
