Iris
* Added FindById option to Net.Iris.PrintDebugInfoOfNetHandle
* Usage exemple:  Net.Iris.PrintDebugInfoOfNetHandle FindById=10001,32355
* Improved PrintDebugInfoOfNetHandle to print the name of the dependents attached to the object (childs, parents, creation dependencies)
* Don't automatically print the status of dependents
* FReplicationFiltering::PrintFilterObjectInfo will now print if the object has a OwnerFilter, ConnectionFilter and a inclusive or exclusive filter group affecting the given connection

#rnx

diff --git a/Engine/Source/Runtime/Experimental/Iris/Core/Private/Iris/ReplicationSystem/Filtering/NetObjectFilter.cpp b/Engine/Source/Runtime/Experimental/Iris/Core/Private/Iris/ReplicationSystem/Filtering/NetObjectFilter.cpp
--- a/Engine/Source/Runtime/Experimental/Iris/Core/Private/Iris/ReplicationSystem/Filtering/NetObjectFilter.cpp
+++ b/Engine/Source/Runtime/Experimental/Iris/Core/Private/Iris/ReplicationSystem/Filtering/NetObjectFilter.cpp
@@ -7,6 +7,31 @@
 
 #include UE_INLINE_GENERATED_CPP_BY_NAME(NetObjectFilter)
 
+namespace UE::Net
+{
+	const TCHAR* LexToString(ENetFilterStatus Status)
+	{
+		switch(Status)
+		{
+			case ENetFilterStatus::Disallow:
+			{
+				return TEXT("Disallow");
+			} break;
+		
+			case ENetFilterStatus::Allow:
+			{
+				return TEXT("Allow");
+			} break;
+
+			default:
+			{
+				ensure(false);
+				return TEXT("Missing");
+			} break;
+		}
+	}
+}
+
 UNetObjectFilter::UNetObjectFilter()
 {
 }
diff --git a/Engine/Source/Runtime/Experimental/Iris/Core/Private/Iris/ReplicationSystem/Filtering/ReplicationFiltering.cpp b/Engine/Source/Runtime/Experimental/Iris/Core/Private/Iris/ReplicationSystem/Filtering/ReplicationFiltering.cpp
--- a/Engine/Source/Runtime/Experimental/Iris/Core/Private/Iris/ReplicationSystem/Filtering/ReplicationFiltering.cpp
+++ b/Engine/Source/Runtime/Experimental/Iris/Core/Private/Iris/ReplicationSystem/Filtering/ReplicationFiltering.cpp
@@ -2774,12 +2774,88 @@ TArrayView<FNetObjectFilteringInfo> FReplicationFiltering::GetNetObjectFiltering
 
 FString FReplicationFiltering::PrintFilterObjectInfo(FInternalNetRefIndex ObjectIndex, uint32 ConnectionId) const
 {
-	// $IRIS TODO: Print info if the object is part of a non-dynamic filter.
+	// $IRIS TODO: Support printing filter info for subobjects
+
+	TStringBuilder<512> PrintBuilder;
+
+	if (ValidConnections.IsBitSet(ConnectionId))
+	{
+		if (HasOwnerFilter(ObjectIndex))
+		{
+			const uint32 OwningConnection = ObjectIndexToOwningConnection[ObjectIndex];
+			const bool bIsOwner = (ConnectionId == OwningConnection);
+			
+			PrintBuilder.Appendf(TEXT("\r\t OwnerFilter to connectionId:%u | %s"), ConnectionId, bIsOwner?TEXT("relevant"):TEXT("not relevant"));
+		}
+		else if (HasConnectionFilter(ObjectIndex))
+		{
+			const PerObjectInfoIndexType ObjectInfoIndex = ObjectIndexToPerObjectInfoIndex[ObjectIndex];
+			const FPerObjectInfo* ObjectInfo = GetPerObjectInfo(ObjectInfoIndex);
+
+			ENetFilterStatus FilterStatus = GetConnectionFilterStatus(*ObjectInfo, ConnectionId);
+
+			PrintBuilder.Appendf(TEXT("\r\t ConnectionFilter to connectionId:%u | %s"), ConnectionId, LexToString(FilterStatus));
+		}
+
+		const FPerConnectionInfo& ConnectionInfo = ConnectionInfos[ConnectionId];
+
+		const TArrayView<const FNetObjectGroupHandle::FGroupIndexType> GroupIndexes = Groups->GetGroupIndexesOfNetObject(ObjectIndex);
+
+		// List of objects excluded to this connection via a group
+		const FNetBitArray& GroupExcludedObjects = ConnectionInfo.GroupExcludedObjects;
+
+		// Is Excluded via group
+		if (GroupExcludedObjects.IsBitSet(ObjectIndex))
+		{
+			PrintBuilder.Append(TEXT("\r\t ExclusionGroups:"));
+			for (const FNetObjectGroupHandle::FGroupIndexType GroupIndex : GroupIndexes)
+			{
+				if (ExclusionFilterGroups.GetBit(GroupIndex))
+				{
+					const FNetObjectGroup* Group = Groups->GetGroupFromIndex(GroupIndex);
+					check(Group);
+					const FPerObjectInfo* ConnectionState = GetPerObjectInfo(GroupInfos[GroupIndex].ConnectionStateIndex);
+					check(ConnectionState);
+					
+					ENetFilterStatus FilterStatus = GetConnectionFilterStatus(*ConnectionState, ConnectionId);
+					if (FilterStatus == ENetFilterStatus::Disallow)
+					{
+						PrintBuilder.Appendf(TEXT(" | %s status: %s"), ToCStr(Group->GroupName.ToString()), LexToString(FilterStatus));
+					}
+				}
+			}
+		}
+
+		// List of objects included to this connection via a group
+		const FNetBitArray& GroupIncludedObjects = ConnectionInfo.GroupIncludedObjects;
+
+		// Is Included via group
+		if (GroupIncludedObjects.IsBitSet(ObjectIndex))
+		{
+			PrintBuilder.Append(TEXT("\r\t InclusionGroups:"));
+			for (const FNetObjectGroupHandle::FGroupIndexType GroupIndex : GroupIndexes)
+			{
+				if (InclusionFilterGroups.GetBit(GroupIndex))
+				{
+					const FNetObjectGroup* Group = Groups->GetGroupFromIndex(GroupIndex);
+					check(Group);
+					const FPerObjectInfo* ConnectionState = GetPerObjectInfo(GroupInfos[GroupIndex].ConnectionStateIndex);
+					check(ConnectionState);
+					
+					ENetFilterStatus FilterStatus = GetConnectionFilterStatus(*ConnectionState, ConnectionId);
+					if (FilterStatus == ENetFilterStatus::Allow)
+					{
+						PrintBuilder.Appendf(TEXT(" | %s status: %s"), ToCStr(Group->GroupName.ToString()), LexToString(FilterStatus));
+					}
+				}
+			}
+		}
+	}
 
 	const uint8 DynamicFilterIndex = ObjectIndexToDynamicFilterIndex[ObjectIndex];
 	if (DynamicFilterIndex == InvalidDynamicFilterIndex)
 	{
-		return TEXT("[NoDynamicFilter]");
+		return FString::Printf(TEXT("[NoDynamicFilter] %s"), PrintBuilder.ToString());
 	}
 
 	const FFilterInfo& FilterInfo = DynamicFilterInfos[DynamicFilterIndex];
@@ -2787,7 +2863,7 @@ FString FReplicationFiltering::PrintFilterObjectInfo(FInternalNetRefIndex Object
 	if (!FilterInfo.Filter->GetFilteredObjects().IsBitSet(ObjectIndex))
 	{
 		ensureMsgf(false, TEXT("Problem with Filter configs for %s.  DynamicIndex %u Filter %s but not in FilteredObjects list"), *NetRefHandleManager->PrintObjectFromIndex(ObjectIndex), DynamicFilterIndex, *FilterInfo.Name.ToString());
-		return TEXT("[WrongSettings]");
+		return FString::Printf(TEXT("[WrongFilter Setting!] %s"), PrintBuilder.ToString());
 	}
 
 	UNetObjectFilter::FDebugInfoParams DebugParams;
@@ -2796,7 +2872,10 @@ FString FReplicationFiltering::PrintFilterObjectInfo(FInternalNetRefIndex Object
 	DebugParams.ConnectionId = ConnectionId;
 	DebugParams.View = Connections->GetReplicationView(ConnectionId);
 
-	return FString::Printf(TEXT("FilterName: %s (%s) | %s"), *FilterInfo.Name.ToString(), *FilterInfo.Filter->GetClass()->GetName(), *FilterInfo.Filter->PrintDebugInfoForObject(DebugParams, ObjectIndex));
+	return FString::Printf(TEXT("[DynamicFilter]: %s (%s)\r\t %s%s"), 
+						   *FilterInfo.Name.ToString(), *FilterInfo.Filter->GetClass()->GetName(),
+						   *FilterInfo.Filter->PrintDebugInfoForObject(DebugParams, ObjectIndex),
+						   PrintBuilder.ToString());
 }
 
 void FReplicationFiltering::BuildObjectsInFilterList(FNetBitArrayView OutObjectsInFilter, FName FilterName) const
diff --git a/Engine/Source/Runtime/Experimental/Iris/Core/Private/Iris/ReplicationSystem/Filtering/ReplicationFiltering.h b/Engine/Source/Runtime/Experimental/Iris/Core/Private/Iris/ReplicationSystem/Filtering/ReplicationFiltering.h
--- a/Engine/Source/Runtime/Experimental/Iris/Core/Private/Iris/ReplicationSystem/Filtering/ReplicationFiltering.h
+++ b/Engine/Source/Runtime/Experimental/Iris/Core/Private/Iris/ReplicationSystem/Filtering/ReplicationFiltering.h
@@ -353,6 +353,8 @@ private:
 	FNetBitArray ObjectsWithDirtyOwner;
 
 	FNetBitArray ObjectsWithOwnerFilter;
+
+	//$TODO: This should be a on-demand data allocation. Very few objects have an owner and the memory usage would be minimal
 	TArray<uint16> ObjectIndexToOwningConnection;
 
 	// For non-owner filtered objects
diff --git a/Engine/Source/Runtime/Experimental/Iris/Core/Private/Iris/ReplicationSystem/ObjectReplicationBridgeDebugging.cpp b/Engine/Source/Runtime/Experimental/Iris/Core/Private/Iris/ReplicationSystem/ObjectReplicationBridgeDebugging.cpp
--- a/Engine/Source/Runtime/Experimental/Iris/Core/Private/Iris/ReplicationSystem/ObjectReplicationBridgeDebugging.cpp
+++ b/Engine/Source/Runtime/Experimental/Iris/Core/Private/Iris/ReplicationSystem/ObjectReplicationBridgeDebugging.cpp
@@ -347,6 +347,24 @@ TArray<FNetRefHandle> FindNetRefHandlesFromArg(UReplicationSystem* RepSystem, co
 			}
 		}
 	}
+	
+	if (const FString* FindById = Args.FindByPredicate([](const FString& Str) { return Str.Contains(TEXT("FindById=")); }))
+	{
+		FString IdSubString;
+		constexpr bool bShouldStopOnSeperator = false;
+		FParse::Value(**FindById, TEXT("FindById="), IdSubString, bShouldStopOnSeperator);
+
+		TArray<FString> StringIds;
+		IdSubString.ParseIntoArray(StringIds, TEXT(","));
+
+		for (FString& StringId : StringIds)
+		{
+			StringId.TrimStartAndEndInline();
+			const int64 HandleId = FCString::Atoi64(*StringId);
+			const FNetRefHandle Handle = FNetRefHandleManager::MakeNetRefHandle(HandleId, RepSystem->GetId());
+			NetRefHandles.Emplace(Handle);
+		}
+	}
 
 	return NetRefHandles;
 }
@@ -1277,19 +1295,6 @@ void UObjectReplicationBridge::PrintDebugInfoForNetRefHandle(const FNetRefHandle
 
 	if (GetReplicationSystem()->IsServer())
 	{
-		// World location
-		if (const TOptional<FWorldLocations::FWorldInfo> WorldInfo = WorldLocations.GetWorldInfo(ObjectIndex))
-		{
-			UE_LOG(LogIris, Display, TEXT("\tWorldLocation: %s | CullDistance: %f"), *WorldInfo->WorldLocation.ToCompactString(), WorldInfo->CullDistance);
-		}
-		else
-		{
-			UE_LOG(LogIris, Display, TEXT("\tWorldLocation: None"));
-		}
-
-		// Filter Info
-		UE_LOG(LogIris, Display, TEXT("\tFilterInfo: %s"), *Filtering.PrintFilterObjectInfo(ObjectIndex, ConnectionId));
-
 		// Connection info
 		if (const FReplicationConnection* Connection = Connections.GetConnection(ConnectionId))
 		{
@@ -1330,10 +1335,49 @@ void UObjectReplicationBridge::PrintDebugInfoForNetRefHandle(const FNetRefHandle
 			UE_LOG(LogIris, Display, TEXT("\tNo valid connection with id: %u"), ConnectionId);
 		}
 
-		// Dependents
+		// World location
+		if (const TOptional<FWorldLocations::FWorldInfo> WorldInfo = WorldLocations.GetWorldInfo(ObjectIndex))
+		{
+			UE_LOG(LogIris, Display, TEXT("\tWorldLocation: %s | CullDistance: %f"), *WorldInfo->WorldLocation.ToCompactString(), WorldInfo->CullDistance);
+		}
+		else
+		{
+			UE_LOG(LogIris, Display, TEXT("\tWorldLocation: None"));
+		}
+
+		// Filter Info
+		UE_LOG(LogIris, Display, TEXT("\tFilterInfo: %s"), *Filtering.PrintFilterObjectInfo(ObjectIndex, ConnectionId));
+
+		// Child Dependents
 		{
-			UE_LOG(LogIris, Display, TEXT("\tChildDependents: %d objects"), NetRefHandleManager->GetDependentObjectInfos(ObjectIndex).Num());
-			UE_LOG(LogIris, Display, TEXT("\tParentDependents: %d objects"), NetRefHandleManager->GetDependentObjectParents(ObjectIndex).Num());
+			const TArrayView<const FDependentObjectInfo> DependentInfos = NetRefHandleManager->GetDependentObjectInfos(ObjectIndex);
+			
+			UE_CLOG(!DependentInfos.IsEmpty(), LogIris, Display, TEXT("\tChildDependents: %d childs"), DependentInfos.Num());
+			for (int32 i=0; i < DependentInfos.Num(); ++i)
+			{
+				UE_LOG(LogIris, Display, TEXT("\t\t[%i] %s (%s)"), i, *NetRefHandleManager->PrintObjectFromIndex(DependentInfos[i].NetRefIndex), LexToString(DependentInfos[i].SchedulingHint));
+			}
+
+		}
+
+		// Parent Dependents
+		{
+			const TArrayView<const FInternalNetRefIndex> ParentDependents = NetRefHandleManager->GetDependentObjectParents(ObjectIndex);
+			UE_CLOG(!ParentDependents.IsEmpty(), LogIris, Display, TEXT("\tParentDependents: %d parents"), ParentDependents.Num());
+			for (int32 i = 0; i < ParentDependents.Num(); ++i)
+			{
+				UE_LOG(LogIris, Display, TEXT("\t\t[%i] %s"), i, *NetRefHandleManager->PrintObjectFromIndex(ParentDependents[i]));
+			}
+		}
+
+		// Creation dependencies
+		{
+			TConstArrayView<const FInternalNetRefIndex> CreationDependencies = NetRefHandleManager->GetCreationDependencies(ObjectIndex);
+			UE_CLOG(!CreationDependencies.IsEmpty(), LogIris, Display, TEXT("\tCreationDependencies: %d parents"), CreationDependencies.Num());
+			for (int32 i=0; i < CreationDependencies.Num(); i++)
+			{
+				UE_LOG(LogIris, Display, TEXT("\t\t[%i] %s"), i, *NetRefHandleManager->PrintObjectFromIndex(CreationDependencies[i]));
+			}
 		}
 	}
 	else
@@ -1359,7 +1403,7 @@ void UObjectReplicationBridge::PrintDebugInfoForNetRefHandle(const FNetRefHandle
 	UE_LOG(LogIris, Display, TEXT("################ Stopped Printing debug info for: %s ################"), *NetRefHandleManager->PrintObjectFromNetRefHandle(NetRefHandle));
 
 	// Log dependents if requested
-	if (GetReplicationSystem()->IsServer() && !EnumHasAnyFlags(PrintTraits, EPrintDebugInfoTraits::NoDependents))
+	if (GetReplicationSystem()->IsServer() && EnumHasAnyFlags(PrintTraits, EPrintDebugInfoTraits::WithDependents))
 	{
 		const TArrayView<const FDependentObjectInfo> ChildDependents = NetRefHandleManager->GetDependentObjectInfos(ObjectIndex);
 		for (const FDependentObjectInfo& DependentInfo : ChildDependents)
diff --git a/Engine/Source/Runtime/Experimental/Iris/Core/Public/Iris/ReplicationSystem/Filtering/NetObjectFilter.h b/Engine/Source/Runtime/Experimental/Iris/Core/Public/Iris/ReplicationSystem/Filtering/NetObjectFilter.h
--- a/Engine/Source/Runtime/Experimental/Iris/Core/Public/Iris/ReplicationSystem/Filtering/NetObjectFilter.h
+++ b/Engine/Source/Runtime/Experimental/Iris/Core/Public/Iris/ReplicationSystem/Filtering/NetObjectFilter.h
@@ -41,6 +41,8 @@ enum class ENetFilterStatus : uint32
 	Allow,
 };
 
+IRISCORE_API const TCHAR* LexToString(ENetFilterStatus Status);
+
 } // end namespace UE::Net
 
 /**
diff --git a/Engine/Source/Runtime/Experimental/Iris/Core/Public/Iris/ReplicationSystem/ObjectReplicationBridge.h b/Engine/Source/Runtime/Experimental/Iris/Core/Public/Iris/ReplicationSystem/ObjectReplicationBridge.h
--- a/Engine/Source/Runtime/Experimental/Iris/Core/Public/Iris/ReplicationSystem/ObjectReplicationBridge.h
+++ b/Engine/Source/Runtime/Experimental/Iris/Core/Public/Iris/ReplicationSystem/ObjectReplicationBridge.h
@@ -266,8 +266,8 @@ public:
 		Default = 0x0000,
 		// Do not print the entire protocol state when this is set 
 		NoProtocolState = 0x0001,
-		// Do not automatically print the debug information of child dependents
-		NoDependents = 0x0002,
+		// Also print the debug information of child dependents
+		WithDependents = 0x0002,
 	};
 
 	// Debug functions that are triggered via console commands declared in ObjectReplicationBridgeDebugging.cpp
