Make FCollectionPropertyFacade to use FNames instead of FStrings. This removes a lot of string hashing that's slowing down GT tick.

Legacy cloth with lots of Cloth interactor usage:
* GT cloth tick on Test Client: 37us -> 22us
* GT cloth tick in PIE: 54us -> 40us

diff --git a/Engine/Plugins/ChaosCloth/Source/ChaosCloth/Private/ChaosCloth/ChaosClothConstraints.cpp b/Engine/Plugins/ChaosCloth/Source/ChaosCloth/Private/ChaosCloth/ChaosClothConstraints.cpp
--- a/Engine/Plugins/ChaosCloth/Source/ChaosCloth/Private/ChaosCloth/ChaosClothConstraints.cpp
+++ b/Engine/Plugins/ChaosCloth/Source/ChaosCloth/Private/ChaosCloth/ChaosClothConstraints.cpp
@@ -1,5 +1,6 @@
 // Copyright Epic Games, Inc. All Rights Reserved.
 #include "ChaosCloth/ChaosClothConstraints.h"
+#include "ChaosCloth/ChaosClothingSimulationConfig.h"
 #include "ChaosCloth/ChaosWeightMapTarget.h"
 #include "ChaosCloth/ChaosClothingPatternData.h"
 #include "ChaosCloth/ChaosClothComplexColliders.h"
@@ -870,9 +871,7 @@ void FClothConstraints::CreateSelfCollisionConstraints(const Softs::FCollectionP
 	const TMap<FString, const TSet<int32>*>& FaceSets,
 	const TMap<FString, TConstArrayView<int32>>& FaceIntMaps, const FTriangleMesh& TriangleMesh)
 {
-	const bool bUseSelfCollisions = ConfigProperties.GetValue<bool>(TEXT("UseSelfCollisions"));
-
-	if (bUseSelfCollisions)
+	if (Softs::FPBDCollisionSpringConstraints::IsEnabled(ConfigProperties))
 	{
 		SelfCollisionInit = MakeShared<Softs::FPBDTriangleMeshCollisions>(
 			ParticleOffset,
@@ -1355,10 +1354,10 @@ void FClothConstraints::CreateMaxDistanceConstraints(
 	const TMap<FString, TConstArrayView<FRealSingle>>& WeightMaps,
 	Softs::FSolverReal MeshScale)
 {
-	if (ConfigProperties.GetValue("UseLegacyConfig", false))
+	if (ConfigProperties.GetValue(FClothingSimulationConfig::UseLegacyConfigName, false))
 	{
 		FString MaxDistanceString = TEXT("MaxDistance");
-		MaxDistanceString = ConfigProperties.GetStringValue(MaxDistanceString, MaxDistanceString);  // Uses the same string for both the default weight map name and the property name
+		MaxDistanceString = Softs::FPBDSphericalConstraint::GetMaxDistanceString(ConfigProperties, MaxDistanceString);  // Uses the same string for both the default weight map name and the property name
 		const TConstArrayView<FRealSingle> MaxDistances = WeightMaps.FindRef(MaxDistanceString);
 		if (MaxDistances.Num() != NumParticles)
 		{
@@ -1385,12 +1384,12 @@ void FClothConstraints::CreateBackstopConstraints(
 	const TMap<FString, TConstArrayView<FRealSingle>>& WeightMaps,
 	Softs::FSolverReal MeshScale)
 {
-	if (ConfigProperties.GetValue("UseLegacyConfig", false))
+	if (ConfigProperties.GetValue(FClothingSimulationConfig::UseLegacyConfigName, false))
 	{
 		FString BackstopRadiusString = TEXT("BackstopRadius");
 		FString BackstopDistanceString = TEXT("BackstopDistance");
-		BackstopRadiusString = ConfigProperties.GetStringValue(BackstopRadiusString, BackstopRadiusString);        // Uses the same string for both the default weight map name and the property name
-		BackstopDistanceString = ConfigProperties.GetStringValue(BackstopDistanceString, BackstopDistanceString);  //
+		BackstopRadiusString = Softs::FPBDSphericalBackstopConstraint::GetBackstopRadiusString(ConfigProperties, BackstopRadiusString);        // Uses the same string for both the default weight map name and the property name
+		BackstopDistanceString = Softs::FPBDSphericalBackstopConstraint::GetBackstopDistanceString(ConfigProperties, BackstopDistanceString);  //
 		const TConstArrayView<FRealSingle> BackstopRadiuses = WeightMaps.FindRef(BackstopRadiusString);
 		const TConstArrayView<FRealSingle> BackstopDistances = WeightMaps.FindRef(BackstopDistanceString);
 
diff --git a/Engine/Plugins/ChaosCloth/Source/ChaosCloth/Private/ChaosCloth/ChaosClothVisualization.cpp b/Engine/Plugins/ChaosCloth/Source/ChaosCloth/Private/ChaosCloth/ChaosClothVisualization.cpp
--- a/Engine/Plugins/ChaosCloth/Source/ChaosCloth/Private/ChaosCloth/ChaosClothVisualization.cpp
+++ b/Engine/Plugins/ChaosCloth/Source/ChaosCloth/Private/ChaosCloth/ChaosClothVisualization.cpp
@@ -379,7 +379,7 @@ namespace Chaos
 			const int32 Offset = Solver->IsLegacySolver() ? ParticleRangeId : 0;
 			const TConstArrayView<TVec3<int32>> Elements = Cloth->GetTriangleMesh(Solver).GetElements();
 			const TConstArrayView<Softs::FSolverVec3> Positions = Cloth->GetParticlePositions(Solver);
-			const TConstArrayView<int32>& WeightMap = Cloth->GetFaceIntMapByProperty(Solver, TEXT("SelfCollisionLayers"));
+			const TConstArrayView<int32>& WeightMap = Cloth->GetFaceIntMapByProperty(Solver, Softs::FPBDCollisionSpringConstraints::SelfCollisionLayersName);
 
 			for (int32 ElementIndex = 0; ElementIndex < Elements.Num(); ++ElementIndex, VertexIndex += 3)
 			{
@@ -519,8 +519,8 @@ namespace Chaos
 
 			const int32 NumParticles = Cloth->GetNumParticles(Solver);
 			const Softs::FPBDFlatWeightMapView MaxDistances(
-				Cloth->GetConfig()->GetProperties(Cloth->GetLODIndex(Solver)).GetWeightedFloatValue(TEXT("MaxDistance"), FVector2f(0.f, 1.f)),
-				Cloth->GetWeightMapByProperty(Solver, TEXT("MaxDistance")),
+				(Softs::FSolverVec2)(Softs::FPBDSphericalConstraint::GetWeightedFloatMaxDistance(Cloth->GetConfig()->GetProperties(Cloth->GetLODIndex(Solver)), FVector2f(0.f, 1.f))),
+				Cloth->GetWeightMapByProperty(Solver, Softs::FPBDSphericalConstraint::MaxDistanceName),
 				NumParticles);
 
 			const TConstArrayView<Softs::FSolverVec3> Positions = Cloth->GetAnimationPositions(Solver);
@@ -1852,8 +1852,8 @@ namespace Chaos
 
 			const int32 NumParticles = Cloth->GetNumParticles(Solver);
 			const Softs::FPBDFlatWeightMapView MaxDistances(
-				Cloth->GetConfig()->GetProperties(Cloth->GetLODIndex(Solver)).GetWeightedFloatValue(TEXT("MaxDistance"), FVector2f(0.f, 1.f)),
-				Cloth->GetWeightMapByProperty(Solver, TEXT("MaxDistance")), 
+				(Softs::FSolverVec2)(Softs::FPBDSphericalConstraint::GetWeightedFloatMaxDistance(Cloth->GetConfig()->GetProperties(Cloth->GetLODIndex(Solver)), FVector2f(0.f, 1.f))),
+				Cloth->GetWeightMapByProperty(Solver, Softs::FPBDSphericalConstraint::MaxDistanceName),
 				NumParticles);
 
 			const TConstArrayView<Softs::FSolverReal> InvMasses = Cloth->GetParticleInvMasses(Solver);
@@ -1901,7 +1901,7 @@ namespace Chaos
 			const FClothConstraints& ClothConstraints = Solver->GetClothConstraints(ParticleRangeId);
 			if (const Softs::FPBDAnimDriveConstraint* const AnimDriveConstraint = ClothConstraints.GetAnimDriveConstraints().Get())
 			{
-				const TConstArrayView<FRealSingle>& AnimDriveStiffnessMultipliers = Cloth->GetWeightMapByProperty(Solver, TEXT("AnimDriveStiffness"));
+				const TConstArrayView<FRealSingle>& AnimDriveStiffnessMultipliers = Cloth->GetWeightMapByProperty(Solver, Softs::FPBDAnimDriveConstraint::AnimDriveStiffnessName);
 				const TConstArrayView<Softs::FSolverVec3> AnimationPositions = Cloth->GetAnimationPositions(Solver);
 				const TConstArrayView<Softs::FSolverVec3> ParticlePositions = Cloth->GetParticlePositions(Solver);
 				check(ParticlePositions.Num() == AnimationPositions.Num());
@@ -2869,7 +2869,7 @@ namespace Chaos
 					const int32 Offset = Solver->IsLegacySolver() ? ParticleRangeId : 0;
 					const TConstArrayView<Softs::FSolverVec3> Positions = Cloth->GetParticlePositions(Solver);
 					const TConstArrayView<Softs::FSolverReal> InvMasses = Cloth->GetParticleInvMasses(Solver);
-					const TConstArrayView<int32>& WeightMap = Cloth->GetFaceIntMapByProperty(Solver, TEXT("SelfCollisionLayers"));
+					const TConstArrayView<int32>& WeightMap = Cloth->GetFaceIntMapByProperty(Solver, Softs::FPBDCollisionSpringConstraints::SelfCollisionLayersName);
 
 					const TConstArrayView<TVec2<int32>> Edges = Cloth->GetTriangleMesh(Solver).GetSegmentMesh().GetElements();
 					const TArray<TVec2<int32>>& EdgeToFaces = Cloth->GetTriangleMesh(Solver).GetEdgeToFaces();
diff --git a/Engine/Plugins/ChaosCloth/Source/ChaosCloth/Private/ChaosCloth/ChaosClothingSimulation.cpp b/Engine/Plugins/ChaosCloth/Source/ChaosCloth/Private/ChaosCloth/ChaosClothingSimulation.cpp
--- a/Engine/Plugins/ChaosCloth/Source/ChaosCloth/Private/ChaosCloth/ChaosClothingSimulation.cpp
+++ b/Engine/Plugins/ChaosCloth/Source/ChaosCloth/Private/ChaosCloth/ChaosClothingSimulation.cpp
@@ -469,17 +469,17 @@ void FClothingSimulation::UpdateSimulationFromSharedSimConfig()
 
 void FClothingSimulation::SetNumIterations(int32 InNumIterations)
 {
-	Solver->GetConfig()->GetProperties(Solver->GetSolverLOD()).SetValue(TEXT("NumIterations"), InNumIterations);
+	Solver->GetConfig()->GetProperties(Solver->GetSolverLOD()).SetValue(FClothingSimulationSolver::NumIterationsName, InNumIterations);
 }
 
 void FClothingSimulation::SetMaxNumIterations(int32 MaxNumIterations)
 {
-	Solver->GetConfig()->GetProperties(Solver->GetSolverLOD()).SetValue(TEXT("MaxNumIterations"), MaxNumIterations);
+	Solver->GetConfig()->GetProperties(Solver->GetSolverLOD()).SetValue(FClothingSimulationSolver::MaxNumIterationsName, MaxNumIterations);
 }
 
 void FClothingSimulation::SetNumSubsteps(int32 InNumSubsteps)
 {
-	Solver->GetConfig()->GetProperties(Solver->GetSolverLOD()).SetValue(TEXT("NumSubsteps"), InNumSubsteps);
+	Solver->GetConfig()->GetProperties(Solver->GetSolverLOD()).SetValue(FClothingSimulationSolver::NumSubstepsName, InNumSubsteps);
 }
 
 bool FClothingSimulation::ShouldSimulate() const
diff --git a/Engine/Plugins/ChaosCloth/Source/ChaosCloth/Private/ChaosCloth/ChaosClothingSimulationCloth.cpp b/Engine/Plugins/ChaosCloth/Source/ChaosCloth/Private/ChaosCloth/ChaosClothingSimulationCloth.cpp
--- a/Engine/Plugins/ChaosCloth/Source/ChaosCloth/Private/ChaosCloth/ChaosClothingSimulationCloth.cpp
+++ b/Engine/Plugins/ChaosCloth/Source/ChaosCloth/Private/ChaosCloth/ChaosClothingSimulationCloth.cpp
@@ -11,8 +11,10 @@
 #include "Chaos/CollectionPropertyFacade.h"
 #include "Chaos/PBDFlatWeightMap.h"
 #include "Chaos/PBDSoftBodyCollisionConstraint.h"
+#include "Chaos/PBDSphericalConstraint.h"
 #include "Chaos/SoftsExternalForces.h"
 #include "Chaos/SoftsSimulationSpace.h"
+#include "Chaos/VelocityField.h"
 #include "Containers/ArrayView.h"
 #include "GeometryCollection/ManagedArrayCollection.h"
 #include "HAL/IConsoleManager.h"
@@ -30,6 +32,7 @@ namespace ClothingSimulationClothDefault
 	constexpr float CollisionThickness = Softs::FPBDSoftBodyCollisionConstraint::DefaultCollisionThickness;
 	constexpr float FrictionCoefficient = Softs::FPBDSoftBodyCollisionConstraint::DefaultFrictionCoefficient;
 	constexpr float DampingCoefficient = 0.01f;
+	constexpr float LocalDampingCoefficient = 0.0f;
 	constexpr float Drag = 0.035f;
 	constexpr float Lift = 0.035f;
 	constexpr float Pressure = 0.f;
@@ -120,11 +123,8 @@ struct FClothingSimulationCloth::FLODData
 
 
 	// These are only used in Add, so no need to cache them
-	UE_CHAOS_DECLARE_INDEXLESS_PROPERTYCOLLECTION_NAME(MaxDistance, float);
-	UE_CHAOS_DECLARE_INDEXLESS_PROPERTYCOLLECTION_NAME(MassValue, float);
 	UE_CHAOS_DECLARE_INDEXLESS_PROPERTYCOLLECTION_NAME(MultiResCoarseLODIndex, int32);
 	UE_CHAOS_DECLARE_INDEXLESS_PROPERTYCOLLECTION_NAME(IsCoarseMultiResLOD, bool);
-	UE_CHAOS_DECLARE_INDEXLESS_PROPERTYCOLLECTION_NAME(KinematicVertices3D, bool); // Selection set name string property; bool value is not actually used.
 
 	// This can update every frame, so should be cached
 	UE_CHAOS_DECLARE_PROPERTYCOLLECTION_NAME(ActiveMorphTarget, float);
@@ -241,11 +241,11 @@ void FClothingSimulationCloth::FLODData::AddParticles(FClothingSimulationSolver*
 
 	// Retrieve MaxDistance information (weight map and Low/High values)
 	const Softs::FPBDFlatWeightMapView MaxDistances(
-		GetWeightedFloatMaxDistance(ConfigProperties, FVector2f(0.f, 1.f)),
-		WeightMaps.FindRef(GetMaxDistanceString(ConfigProperties, MaxDistanceName.ToString())),
+		Softs::FPBDSphericalConstraint::GetWeightedFloatMaxDistance(ConfigProperties, FVector2f(0.f, 1.f)),
+		WeightMaps.FindRef(Softs::FPBDSphericalConstraint::GetMaxDistanceString(ConfigProperties, Softs::FPBDSphericalConstraint::MaxDistanceName.ToString())),
 		NumParticles);
 
-	const TSet<int32>* const KinematicVertices3DSet = VertexSets.FindRef(GetKinematicVertices3DString(ConfigProperties, KinematicVertices3DName.ToString()), nullptr);
+	const TSet<int32>* const KinematicVertices3DSet = VertexSets.FindRef(FClothingSimulationCloth::GetKinematicVertices3DString(ConfigProperties, KinematicVertices3DName.ToString()), nullptr);
 
 	// Set the particle masses
 	static const FRealSingle KinematicDistanceThreshold = 0.1f;  // TODO: This is not the same value as set in the painting UI but we might want to expose this value as parameter
@@ -255,10 +255,10 @@ void FClothingSimulationCloth::FLODData::AddParticles(FClothingSimulationSolver*
 		return MaxDistances.GetValue(Index) < KinematicDistanceThreshold || (KinematicVertices3DSet && KinematicVertices3DSet->Contains(Index));
 	};
 
-	const int32 MassMode = ConfigProperties.GetValue<int32>(TEXT("MassMode"), ClothingSimulationClothDefault::MassMode);
+	const int32 MassMode = FClothingSimulationCloth::GetMassMode(ConfigProperties, ClothingSimulationClothDefault::MassMode);
 
 	constexpr FRealSingle MinPerParticleMassClampMin = UE_SMALL_NUMBER;
-	const FRealSingle MinPerParticleMass = FMath::Max((FRealSingle)ConfigProperties.GetValue<float>(TEXT("MinPerParticleMass"), ClothingSimulationClothDefault::MinPerParticleMass), MinPerParticleMassClampMin);
+	const FRealSingle MinPerParticleMass = FMath::Max((FRealSingle)FClothingSimulationCloth::GetMinPerParticleMass(ConfigProperties, ClothingSimulationClothDefault::MinPerParticleMass), MinPerParticleMassClampMin);
 
 	switch (MassMode)
 	{
@@ -266,21 +266,21 @@ void FClothingSimulationCloth::FLODData::AddParticles(FClothingSimulationSolver*
 		check(false);
 	case EMassMode::UniformMass:
 	{
-		const FVector2f MassValue = GetWeightedFloatMassValue(ConfigProperties, ClothingSimulationClothDefault::MassValue);
-		const TConstArrayView<float> MassValueMultipliers = WeightMaps.FindRef(GetMassValueString(ConfigProperties, MassValueName.ToString()));
+		const FVector2f MassValue = FClothingSimulationCloth::GetWeightedFloatMassValue(ConfigProperties, ClothingSimulationClothDefault::MassValue);
+		const TConstArrayView<float> MassValueMultipliers = WeightMaps.FindRef(FClothingSimulationCloth::GetMassValueString(ConfigProperties, MassValueName.ToString()));
 		Solver->SetParticleMassUniform(SolverDatum.ParticleRangeId, MassValue, MassValueMultipliers, MinPerParticleMass, TriangleMesh, KinematicPredicate);
 	}
 	break;
 	case EMassMode::TotalMass:
 	{
-		const FRealSingle MassValue = GetMassValue(ConfigProperties, ClothingSimulationClothDefault::MassValue);
+		const FRealSingle MassValue = FClothingSimulationCloth::GetMassValue(ConfigProperties, ClothingSimulationClothDefault::MassValue);
 		Solver->SetParticleMassFromTotalMass(SolverDatum.ParticleRangeId, MassValue, MinPerParticleMass, TriangleMesh, KinematicPredicate);
 	}
 	break;
 	case EMassMode::Density:
 	{
-		const FVector2f MassValue = GetWeightedFloatMassValue(ConfigProperties, ClothingSimulationClothDefault::MassValue);
-		const TConstArrayView<float> MassValueMultipliers = WeightMaps.FindRef(GetMassValueString(ConfigProperties, MassValueName.ToString()));
+		const FVector2f MassValue = FClothingSimulationCloth::GetWeightedFloatMassValue(ConfigProperties, ClothingSimulationClothDefault::MassValue);
+		const TConstArrayView<float> MassValueMultipliers = WeightMaps.FindRef(FClothingSimulationCloth::GetMassValueString(ConfigProperties, MassValueName.ToString()));
 		Solver->SetParticleMassFromDensity(SolverDatum.ParticleRangeId, MassValue, MassValueMultipliers, MinPerParticleMass, TriangleMesh, KinematicPredicate);
 	}
 	break;
@@ -504,7 +504,7 @@ void FClothingSimulationCloth::SetMesh(FClothingSimulationMesh* InMesh)
 		TMap<FString, const TSet<int32>*> FaceSets = Mesh->GetFaceSets(LODIndex);
 		TMap<FString, TConstArrayView<int32>> FaceIntMaps = Mesh->GetFaceIntMaps(LODIndex);
 
-		const bool bUseGeodesicTethers = Config->GetProperties(LODIndex).GetValue<bool>(TEXT("UseGeodesicTethers"), ClothingSimulationClothDefault::bUseGeodesicTethers);
+		const bool bUseGeodesicTethers = GetUseGeodesicTethers(Config->GetProperties(LODIndex), ClothingSimulationClothDefault::bUseGeodesicTethers);
 
 		// Add LOD data
 		LODData.Add(MakeUnique<FLODData>(
@@ -718,9 +718,9 @@ TVec3<FRealSingle> FClothingSimulationCloth::GetGravity(const FClothingSimulatio
 		check(Config);
 		const Softs::FCollectionPropertyFacade& ConfigProperties = Config->GetProperties(GetLODIndex(Solver));
 
-		const bool bUseGravityOverride = ConfigProperties.GetValue<bool>(TEXT("UseGravityOverride"));
-		const TVec3<FRealSingle> GravityOverride = (TVec3<FRealSingle>)ConfigProperties.GetValue<FVector3f>(TEXT("GravityOverride"), FVector3f(0.f, 0.f, ClothingSimulationClothDefault::GravityZOverride));
-		const FRealSingle GravityScale = (FRealSingle)ConfigProperties.GetValue<float>(TEXT("GravityScale"), 1.f);
+		const bool bUseGravityOverride = GetUseGravityOverride(ConfigProperties, false);
+		const TVec3<FRealSingle> GravityOverride = (TVec3<FRealSingle>)GetGravityOverride(ConfigProperties, FVector3f(0.f, 0.f, ClothingSimulationClothDefault::GravityZOverride));
+		const FRealSingle GravityScale = (FRealSingle)GetGravityScale(ConfigProperties, 1.f);
 		const FRealSingle GravityMultiplier = (FRealSingle)ClothingSimulationClothConsoleVariables::CVarGravityMultiplier.GetValueOnAnyThread();
 
 		return (Solver->IsClothGravityOverrideEnabled() && bUseGravityOverride ? GravityOverride : Solver->GetGravity() * GravityScale) * GravityMultiplier;
@@ -770,7 +770,7 @@ TConstArrayView<FRealSingle> FClothingSimulationCloth::GetWeightMapByName(const
 	return LODData.IsValidIndex(LODIndex) ? LODData[LODIndex]->WeightMaps.FindRef(Name) : TConstArrayView<FRealSingle>();
 }
 
-TConstArrayView<FRealSingle> FClothingSimulationCloth::GetWeightMapByProperty(const FClothingSimulationSolver* Solver, const FString& Property) const
+TConstArrayView<FRealSingle> FClothingSimulationCloth::GetWeightMapByProperty(const FClothingSimulationSolver* Solver, const FName& Property) const
 {
 	check(Config);
 	const FString PropertyString = Config->GetProperties(GetLODIndex(Solver)).GetStringValue(Property);
@@ -795,7 +795,7 @@ TConstArrayView<int32> FClothingSimulationCloth::GetFaceIntMapByName(const FClot
 	return LODData.IsValidIndex(LODIndex) ? LODData[LODIndex]->FaceIntMaps.FindRef(Name) : TConstArrayView<int32>();
 }
 
-TConstArrayView<int32> FClothingSimulationCloth::GetFaceIntMapByProperty(const FClothingSimulationSolver* Solver, const FString& Property) const
+TConstArrayView<int32> FClothingSimulationCloth::GetFaceIntMapByProperty(const FClothingSimulationSolver* Solver, const FName& Property) const
 {
 	check(Config);
 	const FString PropertyString = Config->GetProperties(GetLODIndex(Solver)).GetStringValue(Property);
@@ -1020,18 +1020,18 @@ void FClothingSimulationCloth::Update(FClothingSimulationSolver* Solver)
 		else
 		{
 			// Use the cloth config parameters
-			VelocityScaleSpace = (EChaosSoftsSimulationSpace)ConfigProperties.GetValue<int32>(TEXT("VelocityScaleSpace"), (int32)ClothingSimulationClothDefault::VelocityScaleSpace);
-			OutLinearVelocityScale = ConfigProperties.GetValue<FVector3f>(TEXT("LinearVelocityScale"), FVector3f(ClothingSimulationClothDefault::VelocityScale));
-			OutAngularVelocityScale = ConfigProperties.GetValue<float>(TEXT("AngularVelocityScale"), ClothingSimulationClothDefault::VelocityScale);
-			OutMaxVelocityScale = ConfigProperties.GetValue<float>(TEXT("MaxVelocityScale"), ClothingSimulationClothDefault::MaxVelocityScale);
-			MaxLinearVelocity = ConfigProperties.GetValue<FVector3f>(TEXT("MaxLinearVelocity"), MaxLinearVelocity);
-			MaxLinearAcceleration = ConfigProperties.GetValue<FVector3f>(TEXT("MaxLinearAcceleration"), MaxLinearAcceleration);
-			MaxAngularVelocity = ConfigProperties.GetValue<float>(TEXT("MaxAngularVelocity"), MaxAngularVelocity);
-			MaxAngularAcceleration = ConfigProperties.GetValue<float>(TEXT("MaxAngularAcceleration"), MaxAngularAcceleration);
+			VelocityScaleSpace = (EChaosSoftsSimulationSpace)GetVelocityScaleSpace(ConfigProperties, (int32)ClothingSimulationClothDefault::VelocityScaleSpace);
+			OutLinearVelocityScale = GetLinearVelocityScale(ConfigProperties, FVector3f(ClothingSimulationClothDefault::VelocityScale));
+			OutAngularVelocityScale = GetAngularVelocityScale(ConfigProperties, ClothingSimulationClothDefault::VelocityScale);
+			OutMaxVelocityScale = GetMaxVelocityScale(ConfigProperties, ClothingSimulationClothDefault::MaxVelocityScale);
+			MaxLinearVelocity = GetMaxLinearVelocity(ConfigProperties, MaxLinearVelocity);
+			MaxLinearAcceleration = GetMaxLinearAcceleration(ConfigProperties, MaxLinearAcceleration);
+			MaxAngularVelocity = GetMaxAngularVelocity(ConfigProperties, MaxAngularVelocity);
+			MaxAngularAcceleration = GetMaxAngularAcceleration(ConfigProperties, MaxAngularAcceleration);
 		}
 
 		// NOTE: Force-based solver doesn't actually use FictitiousAngularScale here. It gets it from the property collection directly.
-		const FRealSingle FictitiousAngularScale = ConfigProperties.GetValue<float>(TEXT("FictitiousAngularScale"), ClothingSimulationClothDefault::FictitiousAngularScale);
+		const FRealSingle FictitiousAngularScale = GetFictitiousAngularScale(ConfigProperties, ClothingSimulationClothDefault::FictitiousAngularScale);
 		Solver->SetReferenceVelocityScale(
 			GroupId,
 			OldReferenceSpaceTransform,
@@ -1072,7 +1072,7 @@ void FClothingSimulationCloth::Update(FClothingSimulationSolver* Solver)
 			Solver->SetGravity(GroupId, GetGravity(Solver));
 
 			// External forces (legacy wind+field)
-			const bool bUsePointBasedWindModel = ConfigProperties.GetValue<bool>(TEXT("UsePointBasedWindModel"));
+			const bool bUsePointBasedWindModel = GetUsePointBasedWindModel(ConfigProperties, false);
 			Solver->AddExternalForces(GroupId, bUsePointBasedWindModel);
 
 			const bool bPointBasedWindDisablesAccurateWind = ClothingSimulationClothConsoleVariables::CVarLegacyDisablesAccurateWind.GetValueOnAnyThread();
@@ -1080,7 +1080,7 @@ void FClothingSimulationCloth::Update(FClothingSimulationSolver* Solver)
 			Solver->SetWindAndPressureProperties(GroupId, ConfigProperties, LODData[LODIndex]->WeightMaps, bEnableAerodynamics);
 
 			constexpr float WorldScale = 100.f;  // VelocityField wind is in m/s in the config (same as the wind unit), but cm/s in the solver  TODO: Cleanup the Solver SetWindVelocity functions to be consistent with the unit
-			const FVec3f WindVelocity = ConfigProperties.GetValue<FVector3f>(TEXT("WindVelocity")) * WorldScale;
+			const FVec3f WindVelocity = Softs::FVelocityAndPressureField::GetWindVelocity(ConfigProperties, FVector3f(0.f)) * WorldScale;
 			Solver->SetWindVelocity(GroupId, WindVelocity + Solver->GetWindVelocity());
 
 			// Update general solver properties
@@ -1089,14 +1089,14 @@ void FClothingSimulationCloth::Update(FClothingSimulationSolver* Solver)
 			const FReal LocalSpaceScaleInv = 1. / LocalSpaceScale;
 			const Softs::FSolverReal MeshScale = Mesh->GetScale() * LocalSpaceScaleInv;
 
-			const FRealSingle DampingCoefficient = ConfigProperties.GetValue<float>(TEXT("DampingCoefficient"), ClothingSimulationClothDefault::DampingCoefficient);
-			const FRealSingle LocalDampingCoefficient = ConfigProperties.GetValue<float>(TEXT("LocalDampingCoefficient"));
-			const FRealSingle CollisionThickness = ConfigProperties.GetValue<float>(TEXT("CollisionThickness"), ClothingSimulationClothDefault::CollisionThickness);
-			const FRealSingle FrictionCoefficient = ConfigProperties.GetValue<float>(TEXT("FrictionCoefficient"), ClothingSimulationClothDefault::FrictionCoefficient);
+			const FRealSingle DampingCoefficient = GetDampingCoefficient(ConfigProperties, ClothingSimulationClothDefault::DampingCoefficient);
+			const FRealSingle LocalDampingCoefficient = GetLocalDampingCoefficient(ConfigProperties, ClothingSimulationClothDefault::LocalDampingCoefficient);
+			const FRealSingle CollisionThickness = GetCollisionThickness(ConfigProperties, ClothingSimulationClothDefault::CollisionThickness);
+			const FRealSingle FrictionCoefficient = GetFrictionCoefficient(ConfigProperties, ClothingSimulationClothDefault::FrictionCoefficient);
 			Solver->SetProperties(GroupId, DampingCoefficient, LocalDampingCoefficient, CollisionThickness * MeshScale, FrictionCoefficient);
 
 			// Update use of continuous collision detection
-			const bool bUseCCD = ConfigProperties.GetValue<bool>(TEXT("UseCCD"));
+			const bool bUseCCD = GetUseCCD(ConfigProperties, false);
 			Solver->SetUseCCD(GroupId, bUseCCD);
 		}
 
diff --git a/Engine/Plugins/ChaosCloth/Source/ChaosCloth/Private/ChaosCloth/ChaosClothingSimulationConfig.cpp b/Engine/Plugins/ChaosCloth/Source/ChaosCloth/Private/ChaosCloth/ChaosClothingSimulationConfig.cpp
--- a/Engine/Plugins/ChaosCloth/Source/ChaosCloth/Private/ChaosCloth/ChaosClothingSimulationConfig.cpp
+++ b/Engine/Plugins/ChaosCloth/Source/ChaosCloth/Private/ChaosCloth/ChaosClothingSimulationConfig.cpp
@@ -2,8 +2,19 @@
 #include "ChaosCloth/ChaosClothingSimulationConfig.h"
 #include "ChaosCloth/ChaosClothConfig.h"
 #include "Chaos/CollectionPropertyFacade.h"
-#include "Chaos/PBDLongRangeConstraints.h"  // For Tether modes
-#include "Chaos/PBDBendingConstraintsBase.h" // For ERestAngleConstructionType
+
+#include "ChaosCloth/ChaosClothingSimulationCloth.h"
+#include "ChaosCloth/ChaosClothingSimulationSolver.h"
+#include "Chaos/PBDAnimDriveConstraint.h"
+#include "Chaos/PBDAxialSpringConstraints.h"
+#include "Chaos/PBDBendingConstraints.h"
+#include "Chaos/PBDCollisionSpringConstraints.h"
+#include "Chaos/PBDLongRangeConstraints.h"
+#include "Chaos/PBDSelfCollisionSphereConstraints.h"
+#include "Chaos/PBDSphericalConstraint.h"
+#include "Chaos/PBDSpringConstraints.h"
+#include "Chaos/PBDTriangleMeshCollisions.h"
+#include "Chaos/VelocityField.h"
 #include "GeometryCollection/ManagedArrayCollection.h"
 
 namespace Chaos
@@ -40,9 +51,9 @@ namespace Chaos
 		// Solver properties
 		if (ClothSharedConfig)
 		{
-			Property->AddValue(TEXT("NumIterations"), ClothSharedConfig->IterationCount, AnimatablePropertyFlags);
-			Property->AddValue(TEXT("MaxNumIterations"), ClothSharedConfig->MaxIterationCount, AnimatablePropertyFlags);
-			Property->AddValue(TEXT("NumSubsteps"), ClothSharedConfig->SubdivisionCount, AnimatablePropertyFlags);
+			Property->AddValue(FClothingSimulationSolver::NumIterationsName, ClothSharedConfig->IterationCount, AnimatablePropertyFlags);
+			Property->AddValue(FClothingSimulationSolver::MaxNumIterationsName, ClothSharedConfig->MaxIterationCount, AnimatablePropertyFlags);
+			Property->AddValue(FClothingSimulationSolver::NumSubstepsName, ClothSharedConfig->SubdivisionCount, AnimatablePropertyFlags);
 		}
 
 		// Cloth properties
@@ -64,15 +75,15 @@ namespace Chaos
 					MassValue = ClothConfig->Density;
 					break;
 				}
-				Property->AddValue(TEXT("MassMode"), (int32)ClothConfig->MassMode, NonAnimatablePropertyFlags | ECollectionPropertyFlags::Intrinsic);
-				Property->AddValue(TEXT("MassValue"), MassValue, NonAnimatablePropertyFlags | ECollectionPropertyFlags::Intrinsic);
-				Property->AddValue(TEXT("MinPerParticleMass"), ClothConfig->MinPerParticleMass, NonAnimatablePropertyFlags | ECollectionPropertyFlags::Intrinsic);
+				Property->AddValue(FClothingSimulationCloth::MassModeName, (int32)ClothConfig->MassMode, NonAnimatablePropertyFlags | ECollectionPropertyFlags::Intrinsic);
+				Property->AddValue(FClothingSimulationCloth::MassValueName, MassValue, NonAnimatablePropertyFlags | ECollectionPropertyFlags::Intrinsic);
+				Property->AddValue(FClothingSimulationCloth::MinPerParticleMassName, ClothConfig->MinPerParticleMass, NonAnimatablePropertyFlags | ECollectionPropertyFlags::Intrinsic);
 			}
 
 			// Edge constraint
 			if (ClothConfig->EdgeStiffnessWeighted.Low > 0.f || ClothConfig->EdgeStiffnessWeighted.High > 0.f)
 			{
-				const int32 EdgeSpringStiffnessIndex = Property->AddProperty(TEXT("EdgeSpringStiffness"), AnimatablePropertyFlags);
+				const int32 EdgeSpringStiffnessIndex = Property->AddProperty(Softs::FPBDEdgeSpringConstraints::EdgeSpringStiffnessName, AnimatablePropertyFlags);
 				Property->SetWeightedValue(EdgeSpringStiffnessIndex, ClothConfig->EdgeStiffnessWeighted.Low, ClothConfig->EdgeStiffnessWeighted.High);
 				Property->SetStringValue(EdgeSpringStiffnessIndex, TEXT("EdgeStiffness"));
 			}
@@ -83,26 +94,26 @@ namespace Chaos
 			{
 				if (ClothConfig->bUseBendingElements)
 				{
-					const int32 BendingElementStiffnessIndex = Property->AddProperty(TEXT("BendingElementStiffness"), AnimatablePropertyFlags);
+					const int32 BendingElementStiffnessIndex = Property->AddProperty(Softs::FPBDBendingConstraints::BendingElementStiffnessName, AnimatablePropertyFlags);
 					Property->SetWeightedValue(BendingElementStiffnessIndex, ClothConfig->BendingStiffnessWeighted.Low, ClothConfig->BendingStiffnessWeighted.High);
 					Property->SetStringValue(BendingElementStiffnessIndex, TEXT("BendingStiffness"));
 
-					Property->AddValue(TEXT("BucklingRatio"), ClothConfig->BucklingRatio, AnimatablePropertyFlags);
+					Property->AddValue(Softs::FPBDBendingConstraints::BucklingRatioName, ClothConfig->BucklingRatio, AnimatablePropertyFlags);
 
-					const int32 BucklingStiffnessIndex = Property->AddProperty(TEXT("BucklingStiffness"), AnimatablePropertyFlags);
+					const int32 BucklingStiffnessIndex = Property->AddProperty(Softs::FPBDBendingConstraints::BucklingStiffnessName, AnimatablePropertyFlags);
 					Property->SetWeightedValue(BucklingStiffnessIndex, ClothConfig->BucklingStiffnessWeighted.Low, ClothConfig->BucklingStiffnessWeighted.High);
 					Property->SetStringValue(BucklingStiffnessIndex, TEXT("BucklingStiffness"));
 
-					const int32 RestAngleTypeIndex = Property->AddProperty(TEXT("RestAngleType"), NonAnimatablePropertyFlags);
+					const int32 RestAngleTypeIndex = Property->AddProperty(Softs::FPBDBendingConstraints::RestAngleTypeName, NonAnimatablePropertyFlags);
 					Property->SetValue(RestAngleTypeIndex, (int32)Chaos::Softs::FPBDBendingConstraintsBase::ERestAngleConstructionType::FlatnessRatio);
 
-					const int32 FlatnessRatioIndex = Property->AddProperty(TEXT("FlatnessRatio"), NonAnimatablePropertyFlags);
+					const int32 FlatnessRatioIndex = Property->AddProperty(Softs::FPBDBendingConstraints::FlatnessRatioName, NonAnimatablePropertyFlags);
 					Property->SetWeightedValue(FlatnessRatioIndex, ClothConfig->FlatnessRatio.Low, ClothConfig->FlatnessRatio.High);
 					Property->SetStringValue(FlatnessRatioIndex, TEXT("FlatnessRatio"));
 				}
 				else  // Not using bending elements
 				{
-					const int32 BendingSpringStiffnessIndex = Property->AddProperty(TEXT("BendingSpringStiffness"), AnimatablePropertyFlags);
+					const int32 BendingSpringStiffnessIndex = Property->AddProperty(Softs::FPBDBendingSpringConstraints::BendingSpringStiffnessName, AnimatablePropertyFlags);
 					Property->SetWeightedValue(BendingSpringStiffnessIndex, ClothConfig->BendingStiffnessWeighted.Low, ClothConfig->BendingStiffnessWeighted.High);
 					Property->SetStringValue(BendingSpringStiffnessIndex, TEXT("BendingStiffness"));
 				}
@@ -111,7 +122,7 @@ namespace Chaos
 			// Area constraint
 			if (ClothConfig->AreaStiffnessWeighted.Low > 0.f || ClothConfig->AreaStiffnessWeighted.High > 0.f)
 			{
-				const int32 AreaSpringStiffnessIndex = Property->AddProperty(TEXT("AreaSpringStiffness"), AnimatablePropertyFlags);
+				const int32 AreaSpringStiffnessIndex = Property->AddProperty(Softs::FPBDAreaSpringConstraints::AreaSpringStiffnessName, AnimatablePropertyFlags);
 				Property->SetWeightedValue(AreaSpringStiffnessIndex, ClothConfig->AreaStiffnessWeighted.Low, ClothConfig->AreaStiffnessWeighted.High);
 				Property->SetStringValue(AreaSpringStiffnessIndex, TEXT("AreaStiffness"));
 			}
@@ -119,13 +130,13 @@ namespace Chaos
 			// Long range attachment
 			if (ClothConfig->TetherStiffness.Low > 0.f || ClothConfig->TetherStiffness.High > 0.f)
 			{
-				Property->AddValue(TEXT("UseGeodesicTethers"), ClothConfig->bUseGeodesicDistance, NonAnimatablePropertyFlags);
+				Property->AddValue(FClothingSimulationCloth::UseGeodesicTethersName, ClothConfig->bUseGeodesicDistance, NonAnimatablePropertyFlags);
 
-				const int32 TetherStiffnessIndex = Property->AddProperty(TEXT("TetherStiffness"), AnimatablePropertyFlags);
+				const int32 TetherStiffnessIndex = Property->AddProperty(Softs::FPBDLongRangeConstraints::TetherStiffnessName, AnimatablePropertyFlags);
 				Property->SetWeightedValue(TetherStiffnessIndex, ClothConfig->TetherStiffness.Low, ClothConfig->TetherStiffness.High);
 				Property->SetStringValue(TetherStiffnessIndex, TEXT("TetherStiffness"));
 
-				const int32 TetherScaleIndex = Property->AddProperty(TEXT("TetherScale"), AnimatablePropertyFlags);
+				const int32 TetherScaleIndex = Property->AddProperty(Softs::FPBDLongRangeConstraints::TetherScaleName, AnimatablePropertyFlags);
 				Property->SetWeightedValue(TetherScaleIndex, ClothConfig->TetherScale.Low, ClothConfig->TetherScale.High);
 				Property->SetStringValue(TetherScaleIndex, TEXT("TetherScale"));
 			}
@@ -133,118 +144,117 @@ namespace Chaos
 			// AnimDrive
 			if (ClothConfig->AnimDriveStiffness.Low > 0.f || ClothConfig->AnimDriveStiffness.High > 0.f)
 			{
-				const int32 AnimDriveStiffnessIndex = Property->AddProperty(TEXT("AnimDriveStiffness"), AnimatablePropertyFlags);
+				const int32 AnimDriveStiffnessIndex = Property->AddProperty(Softs::FPBDAnimDriveConstraint::AnimDriveStiffnessName, AnimatablePropertyFlags);
 				Property->SetWeightedValue(AnimDriveStiffnessIndex, ClothConfig->AnimDriveStiffness.Low, ClothConfig->AnimDriveStiffness.High);
 				Property->SetStringValue(AnimDriveStiffnessIndex, TEXT("AnimDriveStiffness"));
 
-				const int32 AnimDriveDampingIndex = Property->AddProperty(TEXT("AnimDriveDamping"), AnimatablePropertyFlags);
+				const int32 AnimDriveDampingIndex = Property->AddProperty(Softs::FPBDAnimDriveConstraint::AnimDriveDampingName, AnimatablePropertyFlags);
 				Property->SetWeightedValue(AnimDriveDampingIndex, ClothConfig->AnimDriveDamping.Low, ClothConfig->AnimDriveDamping.High);
 				Property->SetStringValue(AnimDriveDampingIndex, TEXT("AnimDriveDamping"));
 			}
 
 			// Gravity
 			{
-				Property->AddValue(TEXT("GravityScale"), ClothConfig->GravityScale, AnimatablePropertyFlags);
-				Property->AddValue(TEXT("UseGravityOverride"), ClothConfig->bUseGravityOverride, AnimatablePropertyFlags);
-				Property->AddValue(TEXT("GravityOverride"), FVector3f(ClothConfig->Gravity), AnimatablePropertyFlags);
+				Property->AddValue(FClothingSimulationCloth::GravityScaleName, ClothConfig->GravityScale, AnimatablePropertyFlags);
+				Property->AddValue(FClothingSimulationCloth::UseGravityOverrideName, ClothConfig->bUseGravityOverride, AnimatablePropertyFlags);
+				Property->AddValue(FClothingSimulationCloth::GravityOverrideName, FVector3f(ClothConfig->Gravity), AnimatablePropertyFlags);
 			}
 
 			// Velocity scale
 			{
-				Property->AddValue(TEXT("VelocityScaleSpace"), (int32)ClothConfig->VelocityScaleSpace, AnimatablePropertyFlags);
-				Property->AddValue(TEXT("LinearVelocityScale"), FVector3f(ClothConfig->LinearVelocityScale), AnimatablePropertyFlags);
-				Property->AddValue(TEXT("AngularVelocityScale"), ClothConfig->AngularVelocityScale, AnimatablePropertyFlags);
-				Property->AddValue(TEXT("FictitiousAngularScale"), ClothConfig->FictitiousAngularScale, AnimatablePropertyFlags);
+				Property->AddValue(FClothingSimulationCloth::VelocityScaleSpaceName, (int32)ClothConfig->VelocityScaleSpace, AnimatablePropertyFlags);
+				Property->AddValue(FClothingSimulationCloth::LinearVelocityScaleName, FVector3f(ClothConfig->LinearVelocityScale), AnimatablePropertyFlags);
+				Property->AddValue(FClothingSimulationCloth::AngularVelocityScaleName, ClothConfig->AngularVelocityScale, AnimatablePropertyFlags);
+				Property->AddValue(FClothingSimulationCloth::FictitiousAngularScaleName, ClothConfig->FictitiousAngularScale, AnimatablePropertyFlags);
 
 				static const FVector3f DefaultLinearClamp = { TNumericLimits<float>::Max(), TNumericLimits<float>::Max(), TNumericLimits<float>::Max() };
-				Property->AddValue(TEXT("MaxLinearVelocity"), ClothConfig->bEnableLinearVelocityClamping ? ClothConfig->MaxLinearVelocity : DefaultLinearClamp, AnimatablePropertyFlags);
-				Property->AddValue(TEXT("MaxLinearAcceleration"), ClothConfig->bEnableLinearAccelerationClamping ? ClothConfig->MaxLinearAcceleration : DefaultLinearClamp, AnimatablePropertyFlags);
-				Property->AddValue(TEXT("MaxAngularVelocity"), ClothConfig->bEnableAngularVelocityClamping ? ClothConfig->MaxAngularVelocity : TNumericLimits<float>::Max(), AnimatablePropertyFlags);
-				Property->AddValue(TEXT("MaxAngularAcceleration"), ClothConfig->bEnableAngularAccelerationClamping ? ClothConfig->MaxAngularAcceleration : TNumericLimits<float>::Max(), AnimatablePropertyFlags);
+				Property->AddValue(FClothingSimulationCloth::MaxLinearVelocityName, ClothConfig->bEnableLinearVelocityClamping ? ClothConfig->MaxLinearVelocity : DefaultLinearClamp, AnimatablePropertyFlags);
+				Property->AddValue(FClothingSimulationCloth::MaxLinearAccelerationName, ClothConfig->bEnableLinearAccelerationClamping ? ClothConfig->MaxLinearAcceleration : DefaultLinearClamp, AnimatablePropertyFlags);
+				Property->AddValue(FClothingSimulationCloth::MaxAngularVelocityName, ClothConfig->bEnableAngularVelocityClamping ? ClothConfig->MaxAngularVelocity : TNumericLimits<float>::Max(), AnimatablePropertyFlags);
+				Property->AddValue(FClothingSimulationCloth::MaxAngularAccelerationName, ClothConfig->bEnableAngularAccelerationClamping ? ClothConfig->MaxAngularAcceleration : TNumericLimits<float>::Max(), AnimatablePropertyFlags);
 			}
 
 			// Aerodynamics
-			Property->AddValue(TEXT("UsePointBasedWindModel"), ClothConfig->bUsePointBasedWindModel, NonAnimatablePropertyFlags);
+			Property->AddValue(FClothingSimulationCloth::UsePointBasedWindModelName, ClothConfig->bUsePointBasedWindModel, NonAnimatablePropertyFlags);
 			{
-				const int32 DragIndex = Property->AddProperty(TEXT("Drag"), AnimatablePropertyFlags);
+				const int32 DragIndex = Property->AddProperty(Softs::FVelocityAndPressureField::DragName, AnimatablePropertyFlags);
 				Property->SetWeightedValue(DragIndex, ClothConfig->Drag.Low, ClothConfig->Drag.High);
 				Property->SetStringValue(DragIndex, TEXT("Drag"));
 
 				if (ClothConfig->bEnableOuterDrag)
 				{
-					const int32 OuterDragIndex = Property->AddProperty(TEXT("OuterDrag"), AnimatablePropertyFlags);
+					const int32 OuterDragIndex = Property->AddProperty(Softs::FVelocityAndPressureField::OuterDragName, AnimatablePropertyFlags);
 					Property->SetWeightedValue(OuterDragIndex, ClothConfig->OuterDrag.Low, ClothConfig->OuterDrag.High);
 					Property->SetStringValue(OuterDragIndex, TEXT("OuterDrag"));
 				}
 
-				const int32 LiftIndex = Property->AddProperty(TEXT("Lift"), AnimatablePropertyFlags);
+				const int32 LiftIndex = Property->AddProperty(Softs::FVelocityAndPressureField::LiftName, AnimatablePropertyFlags);
 				Property->SetWeightedValue(LiftIndex, ClothConfig->Lift.Low, ClothConfig->Lift.High);
 				Property->SetStringValue(LiftIndex, TEXT("Lift"));
 
 				if (ClothConfig->bEnableOuterLift)
 				{
-					const int32 OuterLiftIndex = Property->AddProperty(TEXT("OuterLift"), AnimatablePropertyFlags);
+					const int32 OuterLiftIndex = Property->AddProperty(Softs::FVelocityAndPressureField::OuterLiftName, AnimatablePropertyFlags);
 					Property->SetWeightedValue(OuterLiftIndex, ClothConfig->OuterLift.Low, ClothConfig->OuterLift.High);
 					Property->SetStringValue(OuterLiftIndex, TEXT("OuterLift"));
 				}
 
 				constexpr float AirDensity = 1.225f;  // Air density in kg/m^3
-				Property->AddValue(TEXT("FluidDensity"), AirDensity, AnimatablePropertyFlags);
+				Property->AddValue(Softs::FVelocityAndPressureField::FluidDensityName, AirDensity, AnimatablePropertyFlags);
 
-				Property->AddValue(TEXT("WindVelocity"), FVector3f(0.f), AnimatablePropertyFlags);  // Wind velocity must exist to be animatable
+				Property->AddValue(Softs::FVelocityAndPressureField::WindVelocityName, FVector3f(0.f), AnimatablePropertyFlags);  // Wind velocity must exist to be animatable
 			}
 
 			// Pressure
 			{
-				const int32 PressureIndex = Property->AddProperty(TEXT("Pressure"), AnimatablePropertyFlags);
+				const int32 PressureIndex = Property->AddProperty(Softs::FVelocityAndPressureField::PressureName, AnimatablePropertyFlags);
 				Property->SetWeightedValue(PressureIndex, ClothConfig->Pressure.Low, ClothConfig->Pressure.High);
 				Property->SetStringValue(PressureIndex, TEXT("Pressure"));
 			}
 
 			// Damping
-			Property->AddValue(TEXT("DampingCoefficient"), ClothConfig->DampingCoefficient, AnimatablePropertyFlags);
-			Property->AddValue(TEXT("LocalDampingCoefficient"), ClothConfig->LocalDampingCoefficient, AnimatablePropertyFlags);
+			Property->AddValue(FClothingSimulationCloth::DampingCoefficientName, ClothConfig->DampingCoefficient, AnimatablePropertyFlags);
+			Property->AddValue(FClothingSimulationCloth::LocalDampingCoefficientName, ClothConfig->LocalDampingCoefficient, AnimatablePropertyFlags);
 
 			// Collision
-			Property->AddValue(TEXT("CollisionThickness"), ClothConfig->CollisionThickness, AnimatablePropertyFlags);
-			Property->AddValue(TEXT("FrictionCoefficient"), ClothConfig->FrictionCoefficient, AnimatablePropertyFlags);
-			Property->AddValue(TEXT("UseCCD"), ClothConfig->bUseCCD, AnimatablePropertyFlags);
-			Property->AddValue(TEXT("UseSelfCollisions"), ClothConfig->bUseSelfCollisions, NonAnimatablePropertyFlags);
-			Property->AddValue(TEXT("SelfCollisionThickness"), ClothConfig->SelfCollisionThickness, NonAnimatablePropertyFlags);
-			Property->AddValue(TEXT("UseSelfIntersections"), ClothConfig->bUseSelfIntersections, NonAnimatablePropertyFlags);
-			Property->AddValue(TEXT("SelfCollisionFriction"), ClothConfig->SelfCollisionFriction, NonAnimatablePropertyFlags);
+			Property->AddValue(FClothingSimulationCloth::CollisionThicknessName, ClothConfig->CollisionThickness, AnimatablePropertyFlags);
+			Property->AddValue(FClothingSimulationCloth::FrictionCoefficientName, ClothConfig->FrictionCoefficient, AnimatablePropertyFlags);
+			Property->AddValue(FClothingSimulationCloth::UseCCDName, ClothConfig->bUseCCD, AnimatablePropertyFlags);
+			Property->AddValue(Softs::FPBDCollisionSpringConstraints::UseSelfCollisionsName, ClothConfig->bUseSelfCollisions, NonAnimatablePropertyFlags);
+			Property->AddValue(Softs::FPBDCollisionSpringConstraints::SelfCollisionThicknessName, ClothConfig->SelfCollisionThickness, NonAnimatablePropertyFlags);
+			Property->AddValue(Softs::FPBDTriangleMeshCollisions::UseSelfIntersectionsName, ClothConfig->bUseSelfIntersections, NonAnimatablePropertyFlags);
+			Property->AddValue(Softs::FPBDCollisionSpringConstraints::SelfCollisionFrictionName, ClothConfig->SelfCollisionFriction, NonAnimatablePropertyFlags);
 
 			if (ClothConfig->bUseSelfCollisionSpheres)
 			{
-				Property->AddValue(TEXT("SelfCollisionSphereRadius"), ClothConfig->SelfCollisionSphereRadius, NonAnimatablePropertyFlags);
-				Property->AddValue(TEXT("SelfCollisionSphereStiffness"), ClothConfig->SelfCollisionSphereStiffness, AnimatablePropertyFlags);
-				static const FString SelfCollisionSphereSetName(TEXT("SelfCollisionSphereSetName"));
-				Property->AddStringValue(SelfCollisionSphereSetName, SelfCollisionSphereSetName, NonAnimatablePropertyFlags);
+				Property->AddValue(Softs::FPBDSelfCollisionSphereConstraints::SelfCollisionSphereRadiusName, ClothConfig->SelfCollisionSphereRadius, NonAnimatablePropertyFlags);
+				Property->AddValue(Softs::FPBDSelfCollisionSphereConstraints::SelfCollisionSphereStiffnessName, ClothConfig->SelfCollisionSphereStiffness, AnimatablePropertyFlags);
+				Property->AddStringValue(Softs::FPBDSelfCollisionSphereConstraints::SelfCollisionSphereSetNameName, Softs::FPBDSelfCollisionSphereConstraints::SelfCollisionSphereSetNameName.ToString(), NonAnimatablePropertyFlags);
 			}
 
 			// Max distance
 			{
-				const int32 MaxDistanceIndex = Property->AddProperty(TEXT("MaxDistance"), AnimatablePropertyFlags);
+				const int32 MaxDistanceIndex = Property->AddProperty(Softs::FPBDSphericalConstraint::MaxDistanceName, AnimatablePropertyFlags);
 				Property->SetWeightedValue(MaxDistanceIndex, 0.f, 1.f);  // Backward compatibility with legacy mask must use a unit range since the multiplier is in the mask
 				Property->SetStringValue(MaxDistanceIndex, TEXT("MaxDistance"));
 			}
 
 			// Backstop
 			{
-				const int32 BackstopDistanceIndex = Property->AddProperty(TEXT("BackstopDistance"), AnimatablePropertyFlags);
+				const int32 BackstopDistanceIndex = Property->AddProperty(Softs::FPBDSphericalBackstopConstraint::BackstopDistanceName, AnimatablePropertyFlags);
 				Property->SetWeightedValue(BackstopDistanceIndex, 0.f, 1.f);  // Backward compatibility with legacy mask must use a unit range since the multiplier is in the mask
 				Property->SetStringValue(BackstopDistanceIndex, TEXT("BackstopDistance"));
 
-				const int32 BackstopRadiusIndex = Property->AddProperty(TEXT("BackstopRadius"), AnimatablePropertyFlags);
+				const int32 BackstopRadiusIndex = Property->AddProperty(Softs::FPBDSphericalBackstopConstraint::BackstopRadiusName, AnimatablePropertyFlags);
 				Property->SetWeightedValue(BackstopRadiusIndex, 0.f, 1.f);  // Backward compatibility with legacy mask must use a unit range since the multiplier is in the mask
 				Property->SetStringValue(BackstopRadiusIndex, TEXT("BackstopRadius"));
 
-				Property->AddValue(TEXT("UseLegacyBackstop"), ClothConfig->bUseLegacyBackstop, NonAnimatablePropertyFlags);
+				Property->AddValue(Softs::FPBDSphericalBackstopConstraint::UseLegacyBackstopName, ClothConfig->bUseLegacyBackstop, NonAnimatablePropertyFlags);
 			}
 		}
 
 		// Mark this as a potential legacy config, but leave the behavior control to the client code (usually means constraint are removed with 0 stiffness, or missing weight maps)
-		Property->AddValue(TEXT("UseLegacyConfig"), bUseLegacyConfig, NonAnimatablePropertyFlags);
+		Property->AddValue(UseLegacyConfigName, bUseLegacyConfig, NonAnimatablePropertyFlags);
 	}
 
 	void FClothingSimulationConfig::Initialize(const TArray<TSharedPtr<const FManagedArrayCollection>>& InPropertyCollections)
diff --git a/Engine/Plugins/ChaosCloth/Source/ChaosCloth/Private/ChaosCloth/ChaosClothingSimulationInteractor.cpp b/Engine/Plugins/ChaosCloth/Source/ChaosCloth/Private/ChaosCloth/ChaosClothingSimulationInteractor.cpp
--- a/Engine/Plugins/ChaosCloth/Source/ChaosCloth/Private/ChaosCloth/ChaosClothingSimulationInteractor.cpp
+++ b/Engine/Plugins/ChaosCloth/Source/ChaosCloth/Private/ChaosCloth/ChaosClothingSimulationInteractor.cpp
@@ -6,6 +6,17 @@
 #include "ChaosCloth/ChaosClothingSimulationSolver.h"
 #include "ChaosCloth/ChaosClothingSimulationConfig.h"
 
+#include "Chaos/PBDAnimDriveConstraint.h"
+#include "Chaos/PBDAxialSpringConstraints.h"
+#include "Chaos/PBDBendingConstraints.h"
+#include "Chaos/PBDCollisionSpringConstraints.h"
+#include "Chaos/PBDLongRangeConstraints.h"
+#include "Chaos/PBDSelfCollisionSphereConstraints.h"
+#include "Chaos/PBDSphericalConstraint.h"
+#include "Chaos/PBDSpringConstraints.h"
+#include "Chaos/PBDTriangleMeshCollisions.h"
+#include "Chaos/VelocityField.h"
+
 #include UE_INLINE_GENERATED_CPP_BY_NAME(ChaosClothingSimulationInteractor)
 
 namespace ChaosClothingInteractor
@@ -17,6 +28,7 @@ void UChaosClothingInteractor::Sync(IClothingSimulation* Simulation)
 {
 	check(Simulation);
 
+	TRACE_CPUPROFILER_EVENT_SCOPE(UChaosClothingInteractor_Sync);
 	if (Chaos::FClothingSimulationCloth* const Cloth = static_cast<Chaos::FClothingSimulation*>(Simulation)->GetCloth(ClothingId))
 	{
 		for (FChaosClothingInteractorCommand& Command : Commands)
@@ -44,9 +56,9 @@ void UChaosClothingInteractor::SetMaterialLinear(float EdgeStiffnessLinear, floa
 
 	ConfigCommands.Add(FChaosClothingInteractorConfigCommand::CreateLambda([EdgeStiffness, BendingStiffness, AreaStiffness](Chaos::FClothingSimulationConfig* Config, int32 LODIndex)
 	{
-		Config->GetProperties(LODIndex).SetWeightedFloatValue(TEXT("EdgeSpringStiffness"), EdgeStiffness);
-		Config->GetProperties(LODIndex).SetWeightedFloatValue(TEXT("BendingSpringStiffness"), BendingStiffness);
-		Config->GetProperties(LODIndex).SetWeightedFloatValue(TEXT("AreaSpringStiffness"), AreaStiffness);
+		Config->GetProperties(LODIndex).SetWeightedFloatValue(Chaos::Softs::FPBDEdgeSpringConstraints::EdgeSpringStiffnessName, EdgeStiffness);
+		Config->GetProperties(LODIndex).SetWeightedFloatValue(Chaos::Softs::FPBDBendingSpringConstraints::BendingSpringStiffnessName, BendingStiffness);
+		Config->GetProperties(LODIndex).SetWeightedFloatValue(Chaos::Softs::FPBDAreaSpringConstraints::AreaSpringStiffnessName, AreaStiffness);
 	}));
 }
 
@@ -54,9 +66,9 @@ void UChaosClothingInteractor::SetMaterial(FVector2D EdgeStiffness, FVector2D Be
 {
 	ConfigCommands.Add(FChaosClothingInteractorConfigCommand::CreateLambda([EdgeStiffness, BendingStiffness, AreaStiffness](Chaos::FClothingSimulationConfig* Config, int32 LODIndex)
 	{
-		Config->GetProperties(LODIndex).SetWeightedFloatValue(TEXT("EdgeSpringStiffness"), FVector2f(EdgeStiffness));
-		Config->GetProperties(LODIndex).SetWeightedFloatValue(TEXT("BendingSpringStiffness"), FVector2f(BendingStiffness));
-		Config->GetProperties(LODIndex).SetWeightedFloatValue(TEXT("AreaSpringStiffness"), FVector2f(AreaStiffness));
+		Config->GetProperties(LODIndex).SetWeightedFloatValue(Chaos::Softs::FPBDEdgeSpringConstraints::EdgeSpringStiffnessName, FVector2f(EdgeStiffness));
+		Config->GetProperties(LODIndex).SetWeightedFloatValue(Chaos::Softs::FPBDBendingSpringConstraints::BendingSpringStiffnessName, FVector2f(BendingStiffness));
+		Config->GetProperties(LODIndex).SetWeightedFloatValue(Chaos::Softs::FPBDAreaSpringConstraints::AreaSpringStiffnessName, FVector2f(AreaStiffness));
 	}));
 }
 
@@ -64,8 +76,8 @@ void UChaosClothingInteractor::SetMaterialBuckling(FVector2D BucklingRatio, FVec
 {
 	ConfigCommands.Add(FChaosClothingInteractorConfigCommand::CreateLambda([BucklingRatio, BucklingStiffness](Chaos::FClothingSimulationConfig* Config, int32 LODIndex)
 	{
-		Config->GetProperties(LODIndex).SetValue(TEXT("BucklingRatio"), (float)BucklingRatio[0]);  // TODO: Make BuckingRatio weighted
-		Config->GetProperties(LODIndex).SetWeightedFloatValue(TEXT("BucklingStiffness"), FVector2f(BucklingStiffness));
+		Config->GetProperties(LODIndex).SetValue(Chaos::Softs::FPBDBendingConstraints::BucklingRatioName, (float)BucklingRatio[0]);  // TODO: Make BuckingRatio weighted
+		Config->GetProperties(LODIndex).SetWeightedFloatValue(Chaos::Softs::FPBDBendingConstraints::BucklingStiffnessName, FVector2f(BucklingStiffness));
 	}));
 }
 
@@ -75,8 +87,8 @@ void UChaosClothingInteractor::SetLongRangeAttachmentLinear(float TetherStiffnes
 	const FVector2f TetherStiffness((FMath::Clamp(FMath::Loge(TetherStiffnessLinear) * ChaosClothingInteractor::InvStiffnessLogBase + 1.f, 0.f, 1.f)), 1.f);
 	ConfigCommands.Add(FChaosClothingInteractorConfigCommand::CreateLambda([TetherStiffness, TetherScale](Chaos::FClothingSimulationConfig* Config, int32 LODIndex)
 	{
-		Config->GetProperties(LODIndex).SetWeightedFloatValue(TEXT("TetherStiffness"), TetherStiffness);
-		Config->GetProperties(LODIndex).SetWeightedFloatValue(TEXT("TetherScale"), FVector2f(TetherScale));
+		Config->GetProperties(LODIndex).SetWeightedFloatValue(Chaos::Softs::FPBDLongRangeConstraints::TetherStiffnessName, TetherStiffness);
+		Config->GetProperties(LODIndex).SetWeightedFloatValue(Chaos::Softs::FPBDLongRangeConstraints::TetherScaleName, FVector2f(TetherScale));
 	}));
 }
 
@@ -84,8 +96,8 @@ void UChaosClothingInteractor::SetLongRangeAttachment(FVector2D TetherStiffness,
 {
 	ConfigCommands.Add(FChaosClothingInteractorConfigCommand::CreateLambda([TetherStiffness, TetherScale](Chaos::FClothingSimulationConfig* Config, int32 LODIndex)
 	{
-		Config->GetProperties(LODIndex).SetWeightedFloatValue(TEXT("TetherStiffness"), FVector2f(TetherStiffness));
-		Config->GetProperties(LODIndex).SetWeightedFloatValue(TEXT("TetherScale"), FVector2f(TetherScale));
+		Config->GetProperties(LODIndex).SetWeightedFloatValue(Chaos::Softs::FPBDLongRangeConstraints::TetherStiffnessName, FVector2f(TetherStiffness));
+		Config->GetProperties(LODIndex).SetWeightedFloatValue(Chaos::Softs::FPBDLongRangeConstraints::TetherScaleName, FVector2f(TetherScale));
 	}));
 }
 
@@ -93,10 +105,10 @@ void UChaosClothingInteractor::SetCollision(float CollisionThickness, float Fric
 {
 	ConfigCommands.Add(FChaosClothingInteractorConfigCommand::CreateLambda([CollisionThickness, FrictionCoefficient, bUseCCD, SelfCollisionThickness](Chaos::FClothingSimulationConfig* Config, int32 LODIndex)
 	{
-		Config->GetProperties(LODIndex).SetValue(TEXT("CollisionThickness"), CollisionThickness);
-		Config->GetProperties(LODIndex).SetValue(TEXT("FrictionCoefficient"), FrictionCoefficient);
-		Config->GetProperties(LODIndex).SetValue(TEXT("UseCCD"), bUseCCD);
-		Config->GetProperties(LODIndex).SetValue(TEXT("SelfCollisionThickness"), SelfCollisionThickness);
+		Config->GetProperties(LODIndex).SetValue(Chaos::FClothingSimulationCloth::CollisionThicknessName, CollisionThickness);
+		Config->GetProperties(LODIndex).SetValue(Chaos::FClothingSimulationCloth::FrictionCoefficientName, FrictionCoefficient);
+		Config->GetProperties(LODIndex).SetValue(Chaos::FClothingSimulationCloth::UseCCDName, bUseCCD);
+		Config->GetProperties(LODIndex).SetValue(Chaos::Softs::FPBDCollisionSpringConstraints::SelfCollisionThicknessName, SelfCollisionThickness);
 	}));
 }
 
@@ -104,15 +116,15 @@ void UChaosClothingInteractor::SetBackstop(bool bEnabled)
 {
 	ConfigCommands.Add(FChaosClothingInteractorConfigCommand::CreateLambda([bEnabled](Chaos::FClothingSimulationConfig* Config, int32 LODIndex)
 	{
-		Config->GetProperties(LODIndex).SetEnabled(TEXT("BackstopRadius"), bEnabled);  // BackstopRadius controls whether the backstop is enabled or not
+		Config->GetProperties(LODIndex).SetEnabled(Chaos::Softs::FPBDSphericalBackstopConstraint::BackstopRadiusName, bEnabled);  // BackstopRadius controls whether the backstop is enabled or not
 	}));
 }
 void UChaosClothingInteractor::SetDamping(float DampingCoefficient, float LocalDampingCoefficient)
 {
 	ConfigCommands.Add(FChaosClothingInteractorConfigCommand::CreateLambda([DampingCoefficient, LocalDampingCoefficient](Chaos::FClothingSimulationConfig* Config, int32 LODIndex)
 	{
-		Config->GetProperties(LODIndex).SetValue(TEXT("DampingCoefficient"), DampingCoefficient);
-		Config->GetProperties(LODIndex).SetValue(TEXT("LocalDampingCoefficient"), LocalDampingCoefficient);
+		Config->GetProperties(LODIndex).SetValue(Chaos::FClothingSimulationCloth::DampingCoefficientName, DampingCoefficient);
+		Config->GetProperties(LODIndex).SetValue(Chaos::FClothingSimulationCloth::LocalDampingCoefficientName, LocalDampingCoefficient);
 	}));
 }
 
@@ -122,10 +134,10 @@ void UChaosClothingInteractor::SetAerodynamics(float DragCoefficient, float Lift
 	{
 		constexpr float AirDensity = 1.225f;
 		constexpr float WorldScale = 100.f;  // Unreal's world unit is the cm
-		Config->GetProperties(LODIndex).SetWeightedFloatValue(TEXT("Drag"), FVector2f(DragCoefficient));
-		Config->GetProperties(LODIndex).SetWeightedFloatValue(TEXT("Lift"), FVector2f(LiftCoefficient));
-		Config->GetProperties(LODIndex).SetValue(TEXT("FluidDensity"), AirDensity);
-		Config->GetProperties(LODIndex).SetValue(TEXT("WindVelocity"), FVector3f(WindVelocity) / WorldScale);
+		Config->GetProperties(LODIndex).SetWeightedFloatValue(Chaos::Softs::FVelocityAndPressureField::DragName, FVector2f(DragCoefficient));
+		Config->GetProperties(LODIndex).SetWeightedFloatValue(Chaos::Softs::FVelocityAndPressureField::LiftName, FVector2f(LiftCoefficient));
+		Config->GetProperties(LODIndex).SetValue(Chaos::Softs::FVelocityAndPressureField::FluidDensityName, AirDensity);
+		Config->GetProperties(LODIndex).SetValue(Chaos::Softs::FVelocityAndPressureField::WindVelocityName, FVector3f(WindVelocity) / WorldScale);
 	}));
 }
 
@@ -134,12 +146,12 @@ void UChaosClothingInteractor::SetWind(FVector2D Drag, FVector2D Lift, float Air
 	ConfigCommands.Add(FChaosClothingInteractorConfigCommand::CreateLambda([Drag, Lift, AirDensity, WindVelocity, OuterDrag, OuterLift](Chaos::FClothingSimulationConfig* Config, int32 LODIndex)
 	{
 		constexpr float WorldScale = 100.f;  // Unreal's world unit is the cm
-		Config->GetProperties(LODIndex).SetWeightedFloatValue(TEXT("Drag"), FVector2f(Drag));
-		Config->GetProperties(LODIndex).SetWeightedFloatValue(TEXT("OuterDrag"), FVector2f(OuterDrag));
-		Config->GetProperties(LODIndex).SetWeightedFloatValue(TEXT("Lift"), FVector2f(Lift));
-		Config->GetProperties(LODIndex).SetWeightedFloatValue(TEXT("OuterLift"), FVector2f(OuterLift));
-		Config->GetProperties(LODIndex).SetValue(TEXT("FluidDensity"), (float)AirDensity * FMath::Cube(WorldScale));  // AirDensity is here in kg/cm^3 for legacy reason but must be in kg/m^3 in the config UI
-		Config->GetProperties(LODIndex).SetValue(TEXT("WindVelocity"), FVector3f(WindVelocity) / WorldScale);
+		Config->GetProperties(LODIndex).SetWeightedFloatValue(Chaos::Softs::FVelocityAndPressureField::DragName, FVector2f(Drag));
+		Config->GetProperties(LODIndex).SetWeightedFloatValue(Chaos::Softs::FVelocityAndPressureField::OuterDragName, FVector2f(OuterDrag));
+		Config->GetProperties(LODIndex).SetWeightedFloatValue(Chaos::Softs::FVelocityAndPressureField::LiftName, FVector2f(Lift));
+		Config->GetProperties(LODIndex).SetWeightedFloatValue(Chaos::Softs::FVelocityAndPressureField::OuterLiftName, FVector2f(OuterLift));
+		Config->GetProperties(LODIndex).SetValue(Chaos::Softs::FVelocityAndPressureField::FluidDensityName, (float)AirDensity * FMath::Cube(WorldScale));  // AirDensity is here in kg/cm^3 for legacy reason but must be in kg/m^3 in the config UI
+		Config->GetProperties(LODIndex).SetValue(Chaos::Softs::FVelocityAndPressureField::WindVelocityName, FVector3f(WindVelocity) / WorldScale);
 	}));
 }
 
@@ -147,7 +159,7 @@ void UChaosClothingInteractor::SetPressure(FVector2D Pressure)
 {
 	ConfigCommands.Add(FChaosClothingInteractorConfigCommand::CreateLambda([Pressure](Chaos::FClothingSimulationConfig* Config, int32 LODIndex)
 	{
-		Config->GetProperties(LODIndex).SetWeightedFloatValue(TEXT("Pressure"), FVector2f(Pressure));
+		Config->GetProperties(LODIndex).SetWeightedFloatValue(Chaos::Softs::FVelocityAndPressureField::PressureName, FVector2f(Pressure));
 	}));
 }
 
@@ -155,9 +167,9 @@ void UChaosClothingInteractor::SetGravity(float GravityScale, bool bIsGravityOve
 {
 	ConfigCommands.Add(FChaosClothingInteractorConfigCommand::CreateLambda([GravityScale, bIsGravityOverridden, GravityOverride](Chaos::FClothingSimulationConfig* Config, int32 LODIndex)
 	{
-		Config->GetProperties(LODIndex).SetValue(TEXT("GravityScale"), GravityScale);
-		Config->GetProperties(LODIndex).SetValue(TEXT("UseGravityOverride"), bIsGravityOverridden);
-		Config->GetProperties(LODIndex).SetValue(TEXT("GravityOverride"), FVector3f(GravityOverride));
+		Config->GetProperties(LODIndex).SetValue(Chaos::FClothingSimulationCloth::GravityScaleName, GravityScale);
+		Config->GetProperties(LODIndex).SetValue(Chaos::FClothingSimulationCloth::UseGravityOverrideName, bIsGravityOverridden);
+		Config->GetProperties(LODIndex).SetValue(Chaos::FClothingSimulationCloth::GravityOverrideName, FVector3f(GravityOverride));
 	}));
 }
 
@@ -169,8 +181,8 @@ void UChaosClothingInteractor::SetAnimDriveLinear(float AnimDriveStiffnessLinear
 	{
 		// The Anim Drive stiffness Low value needs to be 0 in order to keep backward compatibility with existing mask (this wouldn't be an issue if this property had no legacy mask)
 		static const FVector2f AnimDriveDamping(0.f, 1.f);
-		Config->GetProperties(LODIndex).SetWeightedFloatValue(TEXT("AnimDriveStiffness"), AnimDriveStiffness);
-		Config->GetProperties(LODIndex).SetWeightedFloatValue(TEXT("AnimDriveDamping"), AnimDriveDamping);
+		Config->GetProperties(LODIndex).SetWeightedFloatValue(Chaos::Softs::FPBDAnimDriveConstraint::AnimDriveStiffnessName, AnimDriveStiffness);
+		Config->GetProperties(LODIndex).SetWeightedFloatValue(Chaos::Softs::FPBDAnimDriveConstraint::AnimDriveDampingName, AnimDriveDamping);
 	}));
 }
 
@@ -178,8 +190,8 @@ void UChaosClothingInteractor::SetAnimDrive(FVector2D AnimDriveStiffness, FVecto
 {
 	ConfigCommands.Add(FChaosClothingInteractorConfigCommand::CreateLambda([AnimDriveStiffness, AnimDriveDamping](Chaos::FClothingSimulationConfig* Config, int32 LODIndex)
 	{
-		Config->GetProperties(LODIndex).SetWeightedFloatValue(TEXT("AnimDriveStiffness"), FVector2f(AnimDriveStiffness));
-		Config->GetProperties(LODIndex).SetWeightedFloatValue(TEXT("AnimDriveDamping"), FVector2f(AnimDriveDamping));
+		Config->GetProperties(LODIndex).SetWeightedFloatValue(Chaos::Softs::FPBDAnimDriveConstraint::AnimDriveStiffnessName, FVector2f(AnimDriveStiffness));
+		Config->GetProperties(LODIndex).SetWeightedFloatValue(Chaos::Softs::FPBDAnimDriveConstraint::AnimDriveDampingName, FVector2f(AnimDriveDamping));
 	}));
 }
 
@@ -187,9 +199,9 @@ void UChaosClothingInteractor::SetVelocityScale(FVector LinearVelocityScale, flo
 {
 	ConfigCommands.Add(FChaosClothingInteractorConfigCommand::CreateLambda([LinearVelocityScale, AngularVelocityScale, FictitiousAngularScale](Chaos::FClothingSimulationConfig* Config, int32 LODIndex)
 	{
-		Config->GetProperties(LODIndex).SetValue(TEXT("LinearVelocityScale"), FVector3f(LinearVelocityScale));
-		Config->GetProperties(LODIndex).SetValue(TEXT("AngularVelocityScale"), AngularVelocityScale);
-		Config->GetProperties(LODIndex).SetValue(TEXT("FictitiousAngularScale"), FictitiousAngularScale);
+		Config->GetProperties(LODIndex).SetValue(Chaos::FClothingSimulationCloth::LinearVelocityScaleName, FVector3f(LinearVelocityScale));
+		Config->GetProperties(LODIndex).SetValue(Chaos::FClothingSimulationCloth::AngularVelocityScaleName, AngularVelocityScale);
+		Config->GetProperties(LODIndex).SetValue(Chaos::FClothingSimulationCloth::FictitiousAngularScaleName, FictitiousAngularScale);
 	}));
 }
 
@@ -199,10 +211,10 @@ void UChaosClothingInteractor::SetVelocityClamps(bool bEnableLinearVelocityClamp
 	ConfigCommands.Add(FChaosClothingInteractorConfigCommand::CreateLambda([bEnableLinearVelocityClamping, MaxLinearVelocity, bEnableLinearAccelerationClamping, MaxLinearAcceleration, bEnableAngularVelocityClamping, MaxAngularVelocity, bEnableAngularAccelerationClamping, MaxAngularAcceleration](Chaos::FClothingSimulationConfig* Config, int32 LODIndex)
 		{
 			static const FVector3f DefaultLinearClamp = { TNumericLimits<float>::Max(), TNumericLimits<float>::Max(), TNumericLimits<float>::Max() };
-			Config->GetProperties(LODIndex).SetValue(TEXT("MaxLinearVelocity"), bEnableLinearVelocityClamping ? FVector3f(MaxLinearVelocity) : DefaultLinearClamp);
-			Config->GetProperties(LODIndex).SetValue(TEXT("MaxLinearAcceleration"), bEnableLinearAccelerationClamping ? FVector3f(MaxLinearAcceleration) : DefaultLinearClamp);
-			Config->GetProperties(LODIndex).SetValue(TEXT("MaxAngularVelocity"), bEnableAngularVelocityClamping ? MaxAngularVelocity : TNumericLimits<float>::Max());
-			Config->GetProperties(LODIndex).SetValue(TEXT("MaxAngularAcceleration"), bEnableAngularAccelerationClamping ? MaxAngularAcceleration : TNumericLimits<float>::Max());
+			Config->GetProperties(LODIndex).SetValue(Chaos::FClothingSimulationCloth::MaxLinearVelocityName, bEnableLinearVelocityClamping ? FVector3f(MaxLinearVelocity) : DefaultLinearClamp);
+			Config->GetProperties(LODIndex).SetValue(Chaos::FClothingSimulationCloth::MaxLinearAccelerationName, bEnableLinearAccelerationClamping ? FVector3f(MaxLinearAcceleration) : DefaultLinearClamp);
+			Config->GetProperties(LODIndex).SetValue(Chaos::FClothingSimulationCloth::MaxAngularVelocityName, bEnableAngularVelocityClamping ? MaxAngularVelocity : TNumericLimits<float>::Max());
+			Config->GetProperties(LODIndex).SetValue(Chaos::FClothingSimulationCloth::MaxAngularAccelerationName, bEnableAngularAccelerationClamping ? MaxAngularAcceleration : TNumericLimits<float>::Max());
 		}));
 }
 
@@ -293,7 +305,7 @@ void UChaosClothingSimulationInteractor::SetNumIterations(int32 NumIterations)
 {
 	ConfigCommands.Add(FChaosClothingInteractorConfigCommand::CreateLambda([NumIterations](Chaos::FClothingSimulationConfig* Config, int32 LODIndex)
 	{
-		Config->GetProperties(LODIndex).SetValue(TEXT("NumIterations"), NumIterations);
+		Config->GetProperties(LODIndex).SetValue(Chaos::FClothingSimulationSolver::NumIterationsName, NumIterations);
 	}));
 }
 
@@ -301,7 +313,7 @@ void UChaosClothingSimulationInteractor::SetMaxNumIterations(int32 MaxNumIterati
 {
 	ConfigCommands.Add(FChaosClothingInteractorConfigCommand::CreateLambda([MaxNumIterations](Chaos::FClothingSimulationConfig* Config, int32 LODIndex)
 	{
-		Config->GetProperties(LODIndex).SetValue(TEXT("MaxNumIterations"), MaxNumIterations);
+		Config->GetProperties(LODIndex).SetValue(Chaos::FClothingSimulationSolver::MaxNumIterationsName, MaxNumIterations);
 	}));
 }
 
@@ -309,7 +321,7 @@ void UChaosClothingSimulationInteractor::SetNumSubsteps(int32 NumSubsteps)
 {
 	ConfigCommands.Add(FChaosClothingInteractorConfigCommand::CreateLambda([NumSubsteps](Chaos::FClothingSimulationConfig* Config, int32 LODIndex)
 	{
-		Config->GetProperties(LODIndex).SetValue(TEXT("NumSubsteps"), NumSubsteps);
+		Config->GetProperties(LODIndex).SetValue(Chaos::FClothingSimulationSolver::NumSubstepsName, NumSubsteps);
 	}));
 }
 
diff --git a/Engine/Plugins/ChaosCloth/Source/ChaosCloth/Private/ChaosCloth/ChaosClothingSimulationSolver.cpp b/Engine/Plugins/ChaosCloth/Source/ChaosCloth/Private/ChaosCloth/ChaosClothingSimulationSolver.cpp
--- a/Engine/Plugins/ChaosCloth/Source/ChaosCloth/Private/ChaosCloth/ChaosClothingSimulationSolver.cpp
+++ b/Engine/Plugins/ChaosCloth/Source/ChaosCloth/Private/ChaosCloth/ChaosClothingSimulationSolver.cpp
@@ -10,6 +10,7 @@
 #include "Chaos/SoftsEvolution.h"
 #include "Chaos/PBDEvolution.h"
 #include "Chaos/PBDFlatWeightMap.h"
+#include "Chaos/PBDSphericalConstraint.h"
 #include "Chaos/Levelset.h"
 #include "Chaos/MLLevelset.h" 
 #include "Chaos/SkinnedTriangleMesh.h"
@@ -1471,22 +1472,21 @@ void FClothingSimulationSolver::SetWindVelocity(const TVec3<FRealSingle>& InWind
 
 int32 FClothingSimulationSolver::GetNumIterations() const
 {
-	return Config ?
-		Config->GetProperties(SolverLOD).GetValue<int32>(TEXT("NumIterations"), ClothingSimulationSolverDefault::NumIterations) :
+	return Config ? GetNumIterations(Config->GetProperties(SolverLOD), ClothingSimulationSolverDefault::NumIterations) :
 		ClothingSimulationSolverDefault::NumIterations;
 }
 
 int32 FClothingSimulationSolver::GetMaxNumIterations() const
 {
 	return Config ?
-		Config->GetProperties(SolverLOD).GetValue<int32>(TEXT("MaxNumIterations"), ClothingSimulationSolverDefault::MaxNumIterations) :
+		GetMaxNumIterations(Config->GetProperties(SolverLOD), ClothingSimulationSolverDefault::MaxNumIterations) :
 		ClothingSimulationSolverDefault::MaxNumIterations;
 }
 
 int32 FClothingSimulationSolver::GetNumSubsteps() const
 {
 	return Config ?
-		Config->GetProperties(SolverLOD).GetValue<int32>(TEXT("NumSubsteps"), ClothingSimulationSolverDefault::NumSubsteps) :
+		GetNumSubsteps(Config->GetProperties(SolverLOD), ClothingSimulationSolverDefault::NumSubsteps) :
 		ClothingSimulationSolverDefault::NumSubsteps;
 }
 
@@ -2150,18 +2150,17 @@ void FClothingSimulationSolver::Update(Softs::FSolverReal InDeltaTime)
 			check(Config);
 			const Softs::FCollectionPropertyFacade& Properties = Config->GetProperties(SolverLOD);
 
-			const int32 ConfigNumSubsteps = FMath::Max(
-				Properties.GetValue<int32>(TEXT("NumSubsteps"), ClothingSimulationSolverDefault::NumSubsteps),
+			const int32 ConfigNumSubsteps = FMath::Max(GetNumSubsteps(Properties, ClothingSimulationSolverDefault::NumSubsteps),
 				ClothingSimulationSolverDefault::MinNumSubsteps);
 
-			const Softs::FSolverReal ConfigDynamicSubstepDeltaTime = Properties.GetValue<float>(TEXT("DynamicSubstepDeltaTime"), ClothingSimulationSolverDefault::DynamicSubstepDeltaTime); //  This is in ms
+			const Softs::FSolverReal ConfigDynamicSubstepDeltaTime = GetDynamicSubstepDeltaTime(Properties, ClothingSimulationSolverDefault::DynamicSubstepDeltaTime); //  This is in ms
 
 			// Calculate NumSubsteps for dynamic substepping.
 			NumUsedSubsteps = ConfigDynamicSubstepDeltaTime > (Softs::FSolverReal)0.f ? FMath::Clamp(FMath::RoundToInt32(DeltaTime * 1000.f / ConfigDynamicSubstepDeltaTime), 1, ConfigNumSubsteps) : ConfigNumSubsteps;
 
-			const bool bConfigEnableNumSelfCollisionSubsteps = Properties.GetValue<bool>(TEXT("EnableNumSelfCollisionSubsteps"), ClothingSimulationSolverDefault::bEnableNumSelfCollisionSubsteps);
+			const bool bConfigEnableNumSelfCollisionSubsteps = GetEnableNumSelfCollisionSubsteps(Properties, ClothingSimulationSolverDefault::bEnableNumSelfCollisionSubsteps);
 
-			const int32 ConfigNumSelfCollisionSubsteps = bConfigEnableNumSelfCollisionSubsteps ? FMath::Clamp(Properties.GetValue<int32>(TEXT("NumSelfCollisionSubsteps"), ClothingSimulationSolverDefault::NumSelfCollisionSubsteps), ClothingSimulationSolverDefault::MinNumSubsteps, NumUsedSubsteps) : NumUsedSubsteps;
+			const int32 ConfigNumSelfCollisionSubsteps = bConfigEnableNumSelfCollisionSubsteps ? FMath::Clamp(GetNumSelfCollisionSubsteps(Properties, ClothingSimulationSolverDefault::NumSelfCollisionSubsteps), ClothingSimulationSolverDefault::MinNumSubsteps, NumUsedSubsteps) : NumUsedSubsteps;
 
 			const int32 NumSubstepsPerSelfCollisionSubstep = (NumUsedSubsteps - 1) / ConfigNumSelfCollisionSubsteps + 1;
 
@@ -2173,11 +2172,11 @@ void FClothingSimulationSolver::Update(Softs::FSolverReal InDeltaTime)
 			else
 			{
 				const int32 ConfigMaxNumIterations = FMath::Max(
-					Properties.GetValue<int32>(TEXT("MaxNumIterations"), ClothingSimulationSolverDefault::MaxNumIterations),
+					GetMaxNumIterations(Properties, ClothingSimulationSolverDefault::MaxNumIterations),
 					ClothingSimulationSolverDefault::MinNumIterations);
 
 				const int32 ConfigNumIterations = FMath::Clamp(
-					Properties.GetValue<int32>(TEXT("NumIterations"), ClothingSimulationSolverDefault::NumIterations),
+					GetNumIterations(Properties, ClothingSimulationSolverDefault::NumIterations),
 					ClothingSimulationSolverDefault::MinNumIterations,
 					ConfigMaxNumIterations);
 
diff --git a/Engine/Plugins/ChaosCloth/Source/ChaosCloth/Public/ChaosCloth/ChaosClothingSimulationCloth.h b/Engine/Plugins/ChaosCloth/Source/ChaosCloth/Public/ChaosCloth/ChaosClothingSimulationCloth.h
--- a/Engine/Plugins/ChaosCloth/Source/ChaosCloth/Public/ChaosCloth/ChaosClothingSimulationCloth.h
+++ b/Engine/Plugins/ChaosCloth/Source/ChaosCloth/Public/ChaosCloth/ChaosClothingSimulationCloth.h
@@ -1,6 +1,7 @@
 // Copyright Epic Games, Inc. All Rights Reserved.
 #pragma once
 
+#include "Chaos/CollectionPropertyFacade.h"
 #include "Chaos/Transform.h"
 #include "Chaos/TriangleMesh.h"
 #include "Chaos/AABB.h"
@@ -106,13 +107,25 @@ namespace Chaos
 		// Return the weight map of the specified name if available on the current LOD, or an empty array view otherwise.
 		CHAOSCLOTH_API TConstArrayView<FRealSingle> GetWeightMapByName(const FClothingSimulationSolver* Solver, const FString& Name) const;
 		// Return the weight map of the specified property name if it exists and is available on the current LOD, or an empty array view otherwise.
-		CHAOSCLOTH_API TConstArrayView<FRealSingle> GetWeightMapByProperty(const FClothingSimulationSolver* Solver, const FString& Property) const;
+		CHAOSCLOTH_API TConstArrayView<FRealSingle> GetWeightMapByProperty(const FClothingSimulationSolver* Solver, const FName& Property) const;
+		template<typename StringType UE_REQUIRES(Chaos::Softs::TIsStringAndNameConvertibleType<StringType>::Value)>
+		UE_DEPRECATED(5.7, "Use FName property version of this method instead")
+		TConstArrayView<FRealSingle> GetWeightMapByProperty(const FClothingSimulationSolver* Solver, const StringType Property) const
+		{
+			return GetWeightMapByProperty(Solver, FName(Property));
+		}
 		// Return list of weight map names available across all LODs
 		CHAOSCLOTH_API TSet<FString> GetAllWeightMapNames() const;
 		// Return the face int map of the specified name if available on the current LOD, or an empty array view otherwise.
 		CHAOSCLOTH_API TConstArrayView<int32> GetFaceIntMapByName(const FClothingSimulationSolver* Solver, const FString& Name) const;
 		// Return the face int map of the specified property name if it exists and is available on the current LOD, or an empty array view otherwise.
-		CHAOSCLOTH_API TConstArrayView<int32> GetFaceIntMapByProperty(const FClothingSimulationSolver* Solver, const FString& Property) const;
+		CHAOSCLOTH_API TConstArrayView<int32> GetFaceIntMapByProperty(const FClothingSimulationSolver* Solver, const FName& Property) const;
+		template<typename StringType UE_REQUIRES(Chaos::Softs::TIsStringAndNameConvertibleType<StringType>::Value)>
+		UE_DEPRECATED(5.7, "Use FName property version of this method instead")
+		TConstArrayView<int32> GetFaceIntMapByProperty(const FClothingSimulationSolver* Solver, const StringType Property) const
+		{
+			return GetFaceIntMapByProperty(Solver, FName(Property));
+		}
 		// Return the current LOD tethers.
 		CHAOSCLOTH_API const TArray<TConstArrayView<TTuple<int32, int32, float>>>& GetTethers(const FClothingSimulationSolver* Solver) const;
 		// Return the reference bone index for this cloth.
@@ -139,6 +152,37 @@ namespace Chaos
 		CHAOSCLOTH_API void UpdateFromCache(const FClothingSimulationCacheData& CacheData);
 		// ---- End of the Solver interface ----
 
+
+		UE_CHAOS_DECLARE_INDEXLESS_PROPERTYCOLLECTION_NAME(UseGeodesicTethers, bool);
+		UE_CHAOS_DECLARE_INDEXLESS_PROPERTYCOLLECTION_NAME(MassMode, int32);
+		UE_CHAOS_DECLARE_INDEXLESS_PROPERTYCOLLECTION_NAME(MinPerParticleMass, float);
+		UE_CHAOS_DECLARE_INDEXLESS_PROPERTYCOLLECTION_NAME(MassValue, float);
+
+		UE_CHAOS_DECLARE_INDEXLESS_PROPERTYCOLLECTION_NAME(UseGravityOverride, bool);
+		UE_CHAOS_DECLARE_INDEXLESS_PROPERTYCOLLECTION_NAME(GravityOverride, FVector3f);
+		UE_CHAOS_DECLARE_INDEXLESS_PROPERTYCOLLECTION_NAME(GravityScale, float);
+
+		UE_CHAOS_DECLARE_INDEXLESS_PROPERTYCOLLECTION_NAME(VelocityScaleSpace, int32);
+		UE_CHAOS_DECLARE_INDEXLESS_PROPERTYCOLLECTION_NAME(LinearVelocityScale, FVector3f);
+		UE_CHAOS_DECLARE_INDEXLESS_PROPERTYCOLLECTION_NAME(AngularVelocityScale, float);
+		UE_CHAOS_DECLARE_INDEXLESS_PROPERTYCOLLECTION_NAME(MaxVelocityScale, float);
+		UE_CHAOS_DECLARE_INDEXLESS_PROPERTYCOLLECTION_NAME(MaxLinearVelocity, FVector3f);
+		UE_CHAOS_DECLARE_INDEXLESS_PROPERTYCOLLECTION_NAME(MaxLinearAcceleration, FVector3f);
+		UE_CHAOS_DECLARE_INDEXLESS_PROPERTYCOLLECTION_NAME(MaxAngularVelocity, float);
+		UE_CHAOS_DECLARE_INDEXLESS_PROPERTYCOLLECTION_NAME(MaxAngularAcceleration, float);
+		UE_CHAOS_DECLARE_INDEXLESS_PROPERTYCOLLECTION_NAME(FictitiousAngularScale, float);
+
+		UE_CHAOS_DECLARE_INDEXLESS_PROPERTYCOLLECTION_NAME(UsePointBasedWindModel, bool);
+
+		UE_CHAOS_DECLARE_INDEXLESS_PROPERTYCOLLECTION_NAME(DampingCoefficient, float);
+		UE_CHAOS_DECLARE_INDEXLESS_PROPERTYCOLLECTION_NAME(LocalDampingCoefficient, float);
+
+		UE_CHAOS_DECLARE_INDEXLESS_PROPERTYCOLLECTION_NAME(CollisionThickness, float);
+		UE_CHAOS_DECLARE_INDEXLESS_PROPERTYCOLLECTION_NAME(FrictionCoefficient, float);
+		UE_CHAOS_DECLARE_INDEXLESS_PROPERTYCOLLECTION_NAME(UseCCD, bool);
+
+		UE_CHAOS_DECLARE_INDEXLESS_PROPERTYCOLLECTION_NAME(KinematicVertices3D, bool); // Selection set name string property; bool value is not actually used.
+
 	private:
 		CHAOSCLOTH_API int32 GetNumParticles(int32 InLODIndex) const;
 		CHAOSCLOTH_API int32 GetParticleRangeId(const FClothingSimulationSolver* Solver, int32 InLODIndex) const;
diff --git a/Engine/Plugins/ChaosCloth/Source/ChaosCloth/Public/ChaosCloth/ChaosClothingSimulationConfig.h b/Engine/Plugins/ChaosCloth/Source/ChaosCloth/Public/ChaosCloth/ChaosClothingSimulationConfig.h
--- a/Engine/Plugins/ChaosCloth/Source/ChaosCloth/Public/ChaosCloth/ChaosClothingSimulationConfig.h
+++ b/Engine/Plugins/ChaosCloth/Source/ChaosCloth/Public/ChaosCloth/ChaosClothingSimulationConfig.h
@@ -95,6 +95,8 @@ namespace Chaos
 			}
 		}
 
+		inline static const FName UseLegacyConfigName = TEXT("UseLegacyConfig");
+
 	private:
 		bool bIsLegacySingleLOD = false;
 		TArray<TSharedPtr<FManagedArrayCollection>> PropertyCollections;
diff --git a/Engine/Plugins/ChaosCloth/Source/ChaosCloth/Public/ChaosCloth/ChaosClothingSimulationSolver.h b/Engine/Plugins/ChaosCloth/Source/ChaosCloth/Public/ChaosCloth/ChaosClothingSimulationSolver.h
--- a/Engine/Plugins/ChaosCloth/Source/ChaosCloth/Public/ChaosCloth/ChaosClothingSimulationSolver.h
+++ b/Engine/Plugins/ChaosCloth/Source/ChaosCloth/Public/ChaosCloth/ChaosClothingSimulationSolver.h
@@ -2,6 +2,7 @@
 #pragma once
 
 #include "ChaosCloth/ChaosClothConstraints.h"
+#include "Chaos/CollectionPropertyFacade.h"
 #include "Chaos/PBDSoftsEvolutionFwd.h"
 #include "Chaos/ArrayCollectionArray.h"
 #include "Chaos/SoftsSimulationSpace.h"
@@ -430,6 +431,14 @@ namespace Chaos
 		const FPerSolverFieldSystem& GetPerSolverField() const { return PerSolverField; }
 		// ---- End of the Field interface ----
 
+		// Solver Properties
+		UE_CHAOS_DECLARE_INDEXLESS_PROPERTYCOLLECTION_NAME(NumIterations, int32);
+		UE_CHAOS_DECLARE_INDEXLESS_PROPERTYCOLLECTION_NAME(MaxNumIterations, int32);
+		UE_CHAOS_DECLARE_INDEXLESS_PROPERTYCOLLECTION_NAME(NumSubsteps, int32);
+		UE_CHAOS_DECLARE_INDEXLESS_PROPERTYCOLLECTION_NAME(DynamicSubstepDeltaTime, float);
+		UE_CHAOS_DECLARE_INDEXLESS_PROPERTYCOLLECTION_NAME(EnableNumSelfCollisionSubsteps, bool);
+		UE_CHAOS_DECLARE_INDEXLESS_PROPERTYCOLLECTION_NAME(NumSelfCollisionSubsteps, int32);
+
 	private:
 		void Reset();
 
diff --git a/Engine/Plugins/ChaosClothAsset/Source/ChaosClothAssetEngine/Private/ChaosClothAsset/ClothAsset.cpp b/Engine/Plugins/ChaosClothAsset/Source/ChaosClothAssetEngine/Private/ChaosClothAsset/ClothAsset.cpp
--- a/Engine/Plugins/ChaosClothAsset/Source/ChaosClothAssetEngine/Private/ChaosClothAsset/ClothAsset.cpp
+++ b/Engine/Plugins/ChaosClothAsset/Source/ChaosClothAssetEngine/Private/ChaosClothAsset/ClothAsset.cpp
@@ -123,6 +123,11 @@ namespace UE::Chaos::ClothAsset::Private
 				{
 					TSharedRef<FManagedArrayCollection> ClothCollection = MakeShared<FManagedArrayCollection>();
 					ClothCollection->Serialize(Ar);
+					
+					// Property Facade may need to upgrade
+					::Chaos::Softs::FCollectionPropertyMutableFacade PropertyFacade(ClothCollection);
+					PropertyFacade.PostSerialize(Ar);
+
 					ClothCollections.Emplace(MoveTemp(ClothCollection));
 				}
 			}
diff --git a/Engine/Plugins/ChaosClothAsset/Source/ChaosClothAssetEngine/Private/ChaosClothAsset/ClothAssetInteractor.cpp b/Engine/Plugins/ChaosClothAsset/Source/ChaosClothAssetEngine/Private/ChaosClothAsset/ClothAssetInteractor.cpp
--- a/Engine/Plugins/ChaosClothAsset/Source/ChaosClothAssetEngine/Private/ChaosClothAsset/ClothAssetInteractor.cpp
+++ b/Engine/Plugins/ChaosClothAsset/Source/ChaosClothAssetEngine/Private/ChaosClothAsset/ClothAssetInteractor.cpp
@@ -11,104 +11,104 @@ namespace UE::Chaos::ClothAsset::Private
 	// the property that is actually used by the solver is used here.
 
 	// SimulationBendingConfigNode
-	static const TArray<FString> BendingStiffnessWarpAliases =
+	static const TArray<FName> BendingStiffnessWarpAliases =
 	{
 		TEXT("XPBDAnisoBendingStiffnessWarp"),
 	};
-	static const TArray<FString> BendingStiffnessWeftAliases =
+	static const TArray<FName> BendingStiffnessWeftAliases =
 	{
 		TEXT("XPBDAnisoBendingStiffnessWeft"),
 	};
-	static const TArray<FString> BendingStiffnessBiasAliases =
+	static const TArray<FName> BendingStiffnessBiasAliases =
 	{
 		TEXT("XPBDAnisoBendingStiffnessBias"),
 	};
-	static const TArray<FString> BendingDampingAliases =
+	static const TArray<FName> BendingDampingAliases =
 	{
 		TEXT("XPBDAnisoBendingDamping"),
 		TEXT("XPBDBendingElementDamping"),
 		TEXT("XPBDBendingSpringDamping"),
 	};
-	static const TArray<FString> BucklingRatioAliases =
+	static const TArray<FName> BucklingRatioAliases =
 	{
 		TEXT("XPBDAnisoBucklingRatio"),
 		TEXT("XPBDBucklingRatio"),
 		TEXT("BucklingRatio"),
 	};
-	static const TArray<FString> BucklingStiffnessWarpAliases =
+	static const TArray<FName> BucklingStiffnessWarpAliases =
 	{
 		TEXT("XPBDAnisoBucklingStiffnessWarp"),
 	};
-	static const TArray<FString> BucklingStiffnessWeftAliases =
+	static const TArray<FName> BucklingStiffnessWeftAliases =
 	{
 		TEXT("XPBDAnisoBucklingStiffnessWeft"),
 	};
-	static const TArray<FString> BucklingStiffnessBiasAliases =
+	static const TArray<FName> BucklingStiffnessBiasAliases =
 	{
 		TEXT("XPBDAnisoBucklingStiffnessBias"),
 	};
-	static const TArray<FString> BendingStiffnessAliases =
+	static const TArray<FName> BendingStiffnessAliases =
 	{
 		TEXT("XPBDBendingElementStiffness"),
 		TEXT("BendingElementStiffness"),
 		TEXT("XPBDBendingSpringStiffness"),
 		TEXT("BendingSpringStiffness"),
 	};
-	static const TArray<FString> BucklingStiffnessAliases =
+	static const TArray<FName> BucklingStiffnessAliases =
 	{
 		TEXT("XPBDBucklingStiffness"),
 		TEXT("BucklingStiffness"),
 	};
 
 	// SimulationStretchConfigNode
-	static const TArray<FString> StretchStiffnessWarpAliases =
+	static const TArray<FName> StretchStiffnessWarpAliases =
 	{
 		TEXT("XPBDAnisoStretchStiffnessWarp"),
 		TEXT("XPBDAnisoSpringStiffnessWarp"),
 	};
-	static const TArray<FString> StretchStiffnessWeftAliases =
+	static const TArray<FName> StretchStiffnessWeftAliases =
 	{
 		TEXT("XPBDAnisoStretchStiffnessWeft"),
 		TEXT("XPBDAnisoSpringStiffnessWeft"),
 	};
-	static const TArray<FString> StretchStiffnessBiasAliases =
+	static const TArray<FName> StretchStiffnessBiasAliases =
 	{
 		TEXT("XPBDAnisoStretchStiffnessBias"),
 		TEXT("XPBDAnisoSpringStiffnessBias"),
 	};
-	static const TArray<FString> StretchDampingAliases =
+	static const TArray<FName> StretchDampingAliases =
 	{
 		TEXT("XPBDAnisoStretchDamping"),
 		TEXT("XPBDEdgeSpringDamping"),
 		TEXT("XPBDAnisoSpringDamping"),
 	};
-	static const TArray<FString> StretchStiffnessAliases =
+	static const TArray<FName> StretchStiffnessAliases =
 	{
 		TEXT("XPBDEdgeSpringStiffness"),
 		TEXT("EdgeSpringStiffness"),
 	};
-	static const TArray<FString> StretchWarpScaleAliases =
+	static const TArray<FName> StretchWarpScaleAliases =
 	{
 		TEXT("XPBDAnisoStretchWarpScale"),
 		TEXT("XPBDAnisoSpringWarpScale"),
 		TEXT("EdgeSpringWarpScale"),
 		TEXT("AreaSpringWarpScale"),
 	};
-	static const TArray<FString> StretchWeftScaleAliases =
+	static const TArray<FName> StretchWeftScaleAliases =
 	{
 		TEXT("XPBDAnisoStretchWeftScale"),
 		TEXT("XPBDAnisoSpringWeftScale"),
 		TEXT("EdgeSpringWeftScale"),
 		TEXT("AreaSpringWeftScale"),
 	};
-	static const TArray<FString> AreaStiffnessAliases =
+	static const TArray<FName> AreaStiffnessAliases =
 	{
 		TEXT("XPBDAreaSpringStiffness"),
 		TEXT("AreaSpringStiffness"),
 	};
 	
 
-	static const TMap<FString, TArray<FString>> Aliases =
+	static const TMap<FName, TArray<FName>> Aliases =
 	{
 		{TEXT("BendingStiffnessWarp"), BendingStiffnessWarpAliases},
 		{TEXT("BendingStiffnessWeft"), BendingStiffnessWeftAliases},
@@ -132,22 +132,22 @@ namespace UE::Chaos::ClothAsset::Private
 
 
 	template<typename T>
-	static T GetValueWithAlias(TSharedPtr<::Chaos::Softs::FCollectionPropertyFacade>& PropertyFacade, const FString& PropertyName, const T& DefaultValue, const TFunctionRef<T(int32 KeyIndex)>& GetValue)
+	static T GetValueWithAlias(TSharedPtr<::Chaos::Softs::FCollectionPropertyFacade>& PropertyFacade, const FName& PropertyName, const T& DefaultValue, const TFunctionRef<T(int32 KeyIndex)>& GetValue)
 	{
 		check(PropertyFacade);
 		{
-			const int32 KeyIndex = PropertyFacade->GetKeyIndex(PropertyName);
+			const int32 KeyIndex = PropertyFacade->GetKeyNameIndex(PropertyName);
 			if (KeyIndex != INDEX_NONE)
 			{
 				return GetValue(KeyIndex);
 			}
 		}
 
-		if (const TArray<FString>* FoundAliases = Aliases.Find(PropertyName))
+		if (const TArray<FName>* FoundAliases = Aliases.Find(PropertyName))
 		{
-			for (const FString& FoundAlias : *FoundAliases)
+			for (const FName& FoundAlias : *FoundAliases)
 			{
-				const int32 KeyIndex = PropertyFacade->GetKeyIndex(FoundAlias);
+				const int32 KeyIndex = PropertyFacade->GetKeyNameIndex(FoundAlias);
 				if (KeyIndex != INDEX_NONE)
 				{
 					return GetValue(KeyIndex);
@@ -158,12 +158,12 @@ namespace UE::Chaos::ClothAsset::Private
 		return DefaultValue;
 	}
 
-	void SetValueWithAlias(TSharedPtr<::Chaos::Softs::FCollectionPropertyFacade>& PropertyFacade, const FString& PropertyName, const TFunctionRef<void(const FString&)>& SetValue)
+	void SetValueWithAlias(TSharedPtr<::Chaos::Softs::FCollectionPropertyFacade>& PropertyFacade, const FName& PropertyName, const TFunctionRef<void(const FName&)>& SetValue)
 	{
 		check(PropertyFacade);
-		if (const TArray<FString>* FoundAliases = Aliases.Find(PropertyName))
+		if (const TArray<FName>* FoundAliases = Aliases.Find(PropertyName))
 		{
-			for (const FString& FoundAlias : *FoundAliases)
+			for (const FName& FoundAlias : *FoundAliases)
 			{
 				SetValue(FoundAlias);
 			}
@@ -189,11 +189,11 @@ void UChaosClothAssetInteractor::ResetProperties()
 	CollectionPropertyFacades.Reset();
 }
 
-TArray<FString> UChaosClothAssetInteractor::GetAllProperties(int32 LODIndex) const
+TArray<FName> UChaosClothAssetInteractor::GetAllPropertyNames(int32 LODIndex) const
 {
 	using namespace UE::Chaos::ClothAsset::Private;
 
-	TSet<FString> Keys;
+	TSet<FName> Keys;
 	if (LODIndex == INDEX_NONE)
 	{
 		for (int32 Index = 0; Index < CollectionPropertyFacades.Num(); ++Index)
@@ -206,14 +206,14 @@ TArray<FString> UChaosClothAssetInteractor::GetAllProperties(int32 LODIndex) con
 					Keys.Reserve(PropertyFacade->Num() + Aliases.Num());
 					for (int32 KeyIndex = 0; KeyIndex < PropertyFacade->Num(); ++KeyIndex)
 					{
-						Keys.Add(PropertyFacade->GetKey(KeyIndex));
+						Keys.Add(PropertyFacade->GetKeyName(KeyIndex));
 					}
 				}
 				else
 				{
 					for (int32 KeyIndex = 0; KeyIndex < PropertyFacade->Num(); ++KeyIndex)
 					{
-						Keys.Add(PropertyFacade->GetKey(KeyIndex));
+						Keys.Add(PropertyFacade->GetKeyName(KeyIndex));
 					}
 				}
 			}
@@ -228,15 +228,15 @@ TArray<FString> UChaosClothAssetInteractor::GetAllProperties(int32 LODIndex) con
 				Keys.Reserve(PropertyFacade->Num());
 				for (int32 KeyIndex = 0; KeyIndex < PropertyFacade->Num(); ++KeyIndex)
 				{
-					Keys.Add(PropertyFacade->GetKey(KeyIndex));
+					Keys.Add(PropertyFacade->GetKeyName(KeyIndex));
 				}
 			}
 		}
 	}
 
-	for (TMap<FString, TArray<FString>>::TConstIterator AliasIter = Aliases.CreateConstIterator(); AliasIter; ++AliasIter)
+	for (TMap<FName, TArray<FName>>::TConstIterator AliasIter = Aliases.CreateConstIterator(); AliasIter; ++AliasIter)
 	{
-		for (const FString& OtherName : AliasIter.Value())
+		for (const FName& OtherName : AliasIter.Value())
 		{
 			if (Keys.Contains(OtherName))
 			{
@@ -247,7 +247,20 @@ TArray<FString> UChaosClothAssetInteractor::GetAllProperties(int32 LODIndex) con
 	return Keys.Array();
 }
 
-float UChaosClothAssetInteractor::GetFloatValue(const FString& PropertyName, int32 LODIndex, float DefaultValue) const
+TArray<FString> UChaosClothAssetInteractor::GetAllProperties(int32 LODIndex) const
+{
+	using namespace UE::Chaos::ClothAsset::Private;
+	const TArray<FName> PropertyNames = GetAllPropertyNames(LODIndex);
+	TArray<FString> Result;
+	Result.Reserve(PropertyNames.Num());
+	for (const FName& Name : PropertyNames)
+	{
+		Result.Emplace(Name.ToString());
+	}
+	return Result;
+}
+
+float UChaosClothAssetInteractor::GetFloatPropertyValue(const FName& PropertyName, int32 LODIndex, float DefaultValue) const
 {
 	using namespace UE::Chaos::ClothAsset::Private;
 	if (CollectionPropertyFacades.IsValidIndex(LODIndex))
@@ -264,7 +277,7 @@ float UChaosClothAssetInteractor::GetFloatValue(const FString& PropertyName, int
 	return DefaultValue;
 }
 
-float UChaosClothAssetInteractor::GetLowFloatValue(const FString& PropertyName, int32 LODIndex, float DefaultValue) const
+float UChaosClothAssetInteractor::GetLowFloatPropertyValue(const FName& PropertyName, int32 LODIndex, float DefaultValue) const
 {
 	using namespace UE::Chaos::ClothAsset::Private;
 	if (CollectionPropertyFacades.IsValidIndex(LODIndex))
@@ -281,7 +294,7 @@ float UChaosClothAssetInteractor::GetLowFloatValue(const FString& PropertyName,
 	return DefaultValue;
 }
 
-float UChaosClothAssetInteractor::GetHighFloatValue(const FString& PropertyName, int32 LODIndex, float DefaultValue) const
+float UChaosClothAssetInteractor::GetHighFloatPropertyValue(const FName& PropertyName, int32 LODIndex, float DefaultValue) const
 {
 	using namespace UE::Chaos::ClothAsset::Private;
 	if (CollectionPropertyFacades.IsValidIndex(LODIndex))
@@ -298,7 +311,7 @@ float UChaosClothAssetInteractor::GetHighFloatValue(const FString& PropertyName,
 	return DefaultValue;
 }
 
-FVector2D UChaosClothAssetInteractor::GetWeightedFloatValue(const FString& PropertyName, int32 LODIndex, FVector2D DefaultValue) const
+FVector2D UChaosClothAssetInteractor::GetWeightedFloatPropertyValue(const FName& PropertyName, int32 LODIndex, FVector2D DefaultValue) const
 {
 	using namespace UE::Chaos::ClothAsset::Private;
 	if (CollectionPropertyFacades.IsValidIndex(LODIndex))
@@ -315,7 +328,7 @@ FVector2D UChaosClothAssetInteractor::GetWeightedFloatValue(const FString& Prope
 	return DefaultValue;
 }
 
-int32 UChaosClothAssetInteractor::GetIntValue(const FString& PropertyName, int32 LODIndex, int32 DefaultValue) const
+int32 UChaosClothAssetInteractor::GetIntPropertyValue(const FName& PropertyName, int32 LODIndex, int32 DefaultValue) const
 {
 	using namespace UE::Chaos::ClothAsset::Private;
 	if (CollectionPropertyFacades.IsValidIndex(LODIndex))
@@ -332,7 +345,7 @@ int32 UChaosClothAssetInteractor::GetIntValue(const FString& PropertyName, int32
 	return DefaultValue;
 }
 
-FVector UChaosClothAssetInteractor::GetVectorValue(const FString& PropertyName, int32 LODIndex, FVector DefaultValue) const
+FVector UChaosClothAssetInteractor::GetVectorPropertyValue(const FName& PropertyName, int32 LODIndex, FVector DefaultValue) const
 {
 	using namespace UE::Chaos::ClothAsset::Private;
 	if (CollectionPropertyFacades.IsValidIndex(LODIndex))
@@ -349,7 +362,7 @@ FVector UChaosClothAssetInteractor::GetVectorValue(const FString& PropertyName,
 	return DefaultValue;
 }
 
-FString UChaosClothAssetInteractor::GetStringValue(const FString& PropertyName, int32 LODIndex, const FString& DefaultValue) const
+FString UChaosClothAssetInteractor::GetStringPropertyValue(const FName& PropertyName, int32 LODIndex, const FString& DefaultValue) const
 {
 	using namespace UE::Chaos::ClothAsset::Private;
 	if (CollectionPropertyFacades.IsValidIndex(LODIndex))
@@ -366,7 +379,7 @@ FString UChaosClothAssetInteractor::GetStringValue(const FString& PropertyName,
 	return DefaultValue;
 }
 
-void UChaosClothAssetInteractor::SetFloatValue(const FString& PropertyName, int32 LODIndex, float Value)
+void UChaosClothAssetInteractor::SetFloatPropertyValue(const FName& PropertyName, int32 LODIndex, float Value)
 {
 	using namespace UE::Chaos::ClothAsset::Private;
 	if (LODIndex == INDEX_NONE)
@@ -376,7 +389,7 @@ void UChaosClothAssetInteractor::SetFloatValue(const FString& PropertyName, int3
 			if (TSharedPtr<::Chaos::Softs::FCollectionPropertyFacade> PropertyFacade = CollectionPropertyFacades[Index].Pin())
 			{
 				SetValueWithAlias(PropertyFacade, PropertyName,
-					[&PropertyFacade, Value](const FString& Name)
+					[&PropertyFacade, Value](const FName& Name)
 					{
 						PropertyFacade->SetValue(Name, Value);
 					});
@@ -388,7 +401,7 @@ void UChaosClothAssetInteractor::SetFloatValue(const FString& PropertyName, int3
 		if (TSharedPtr<::Chaos::Softs::FCollectionPropertyFacade> PropertyFacade = CollectionPropertyFacades[LODIndex].Pin())
 		{
 			SetValueWithAlias(PropertyFacade, PropertyName,
-				[&PropertyFacade, Value](const FString& Name)
+				[&PropertyFacade, Value](const FName& Name)
 				{
 					PropertyFacade->SetValue(Name, Value);
 				});
@@ -396,7 +409,7 @@ void UChaosClothAssetInteractor::SetFloatValue(const FString& PropertyName, int3
 	}
 }
 
-void UChaosClothAssetInteractor::SetLowFloatValue(const FString& PropertyName, int32 LODIndex, float Value)
+void UChaosClothAssetInteractor::SetLowFloatPropertyValue(const FName& PropertyName, int32 LODIndex, float Value)
 {
 	using namespace UE::Chaos::ClothAsset::Private;
 	if (LODIndex == INDEX_NONE)
@@ -406,7 +419,7 @@ void UChaosClothAssetInteractor::SetLowFloatValue(const FString& PropertyName, i
 			if (TSharedPtr<::Chaos::Softs::FCollectionPropertyFacade> PropertyFacade = CollectionPropertyFacades[Index].Pin())
 			{
 				SetValueWithAlias(PropertyFacade, PropertyName,
-					[&PropertyFacade, Value](const FString& Name)
+					[&PropertyFacade, Value](const FName& Name)
 					{
 						PropertyFacade->SetLowValue(Name, Value);
 					});
@@ -418,7 +431,7 @@ void UChaosClothAssetInteractor::SetLowFloatValue(const FString& PropertyName, i
 		if (TSharedPtr<::Chaos::Softs::FCollectionPropertyFacade> PropertyFacade = CollectionPropertyFacades[LODIndex].Pin())
 		{
 			SetValueWithAlias(PropertyFacade, PropertyName,
-				[&PropertyFacade, Value](const FString& Name)
+				[&PropertyFacade, Value](const FName& Name)
 				{
 					PropertyFacade->SetLowValue(Name, Value);
 				});
@@ -426,7 +439,7 @@ void UChaosClothAssetInteractor::SetLowFloatValue(const FString& PropertyName, i
 	}
 }
 
-void UChaosClothAssetInteractor::SetHighFloatValue(const FString& PropertyName, int32 LODIndex, float Value)
+void UChaosClothAssetInteractor::SetHighFloatPropertyValue(const FName& PropertyName, int32 LODIndex, float Value)
 {
 	using namespace UE::Chaos::ClothAsset::Private;
 	if (LODIndex == INDEX_NONE)
@@ -436,7 +449,7 @@ void UChaosClothAssetInteractor::SetHighFloatValue(const FString& PropertyName,
 			if (TSharedPtr<::Chaos::Softs::FCollectionPropertyFacade> PropertyFacade = CollectionPropertyFacades[Index].Pin())
 			{
 				SetValueWithAlias(PropertyFacade, PropertyName,
-					[&PropertyFacade, Value](const FString& Name)
+					[&PropertyFacade, Value](const FName& Name)
 					{
 						PropertyFacade->SetHighValue(Name, Value);
 					});
@@ -448,7 +461,7 @@ void UChaosClothAssetInteractor::SetHighFloatValue(const FString& PropertyName,
 		if (TSharedPtr<::Chaos::Softs::FCollectionPropertyFacade> PropertyFacade = CollectionPropertyFacades[LODIndex].Pin())
 		{
 			SetValueWithAlias(PropertyFacade, PropertyName,
-				[&PropertyFacade, Value](const FString& Name)
+				[&PropertyFacade, Value](const FName& Name)
 				{
 					PropertyFacade->SetHighValue(Name, Value);
 				});
@@ -456,7 +469,7 @@ void UChaosClothAssetInteractor::SetHighFloatValue(const FString& PropertyName,
 	}
 }
 
-void UChaosClothAssetInteractor::SetWeightedFloatValue(const FString& PropertyName, int32 LODIndex, FVector2D Value)
+void UChaosClothAssetInteractor::SetWeightedFloatPropertyValue(const FName& PropertyName, int32 LODIndex, FVector2D Value)
 {
 	using namespace UE::Chaos::ClothAsset::Private;
 	if (LODIndex == INDEX_NONE)
@@ -466,7 +479,7 @@ void UChaosClothAssetInteractor::SetWeightedFloatValue(const FString& PropertyNa
 			if (TSharedPtr<::Chaos::Softs::FCollectionPropertyFacade> PropertyFacade = CollectionPropertyFacades[Index].Pin())
 			{
 				SetValueWithAlias(PropertyFacade, PropertyName,
-					[&PropertyFacade, &Value](const FString& Name)
+					[&PropertyFacade, &Value](const FName& Name)
 					{
 						PropertyFacade->SetWeightedFloatValue(Name, FVector2f(Value));
 					});
@@ -478,7 +491,7 @@ void UChaosClothAssetInteractor::SetWeightedFloatValue(const FString& PropertyNa
 		if (TSharedPtr<::Chaos::Softs::FCollectionPropertyFacade> PropertyFacade = CollectionPropertyFacades[LODIndex].Pin())
 		{
 			SetValueWithAlias(PropertyFacade, PropertyName,
-				[&PropertyFacade, &Value](const FString& Name)
+				[&PropertyFacade, &Value](const FName& Name)
 				{
 					PropertyFacade->SetWeightedFloatValue(Name, FVector2f(Value));
 				});
@@ -486,7 +499,7 @@ void UChaosClothAssetInteractor::SetWeightedFloatValue(const FString& PropertyNa
 	}
 }
 
-void UChaosClothAssetInteractor::SetIntValue(const FString& PropertyName, int32 LODIndex, int32 Value)
+void UChaosClothAssetInteractor::SetIntPropertyValue(const FName& PropertyName, int32 LODIndex, int32 Value)
 {
 	using namespace UE::Chaos::ClothAsset::Private;
 	if (LODIndex == INDEX_NONE)
@@ -496,7 +509,7 @@ void UChaosClothAssetInteractor::SetIntValue(const FString& PropertyName, int32
 			if (TSharedPtr<::Chaos::Softs::FCollectionPropertyFacade> PropertyFacade = CollectionPropertyFacades[Index].Pin())
 			{
 				SetValueWithAlias(PropertyFacade, PropertyName,
-					[&PropertyFacade, Value](const FString& Name)
+					[&PropertyFacade, Value](const FName& Name)
 					{
 						PropertyFacade->SetValue(Name, Value);
 					});
@@ -508,7 +521,7 @@ void UChaosClothAssetInteractor::SetIntValue(const FString& PropertyName, int32
 		if (TSharedPtr<::Chaos::Softs::FCollectionPropertyFacade> PropertyFacade = CollectionPropertyFacades[LODIndex].Pin())
 		{
 			SetValueWithAlias(PropertyFacade, PropertyName,
-				[&PropertyFacade, Value](const FString& Name)
+				[&PropertyFacade, Value](const FName& Name)
 				{
 					PropertyFacade->SetValue(Name, Value);
 				});
@@ -516,7 +529,7 @@ void UChaosClothAssetInteractor::SetIntValue(const FString& PropertyName, int32
 	}
 }
 
-void UChaosClothAssetInteractor::SetVectorValue(const FString& PropertyName, int32 LODIndex, FVector Value)
+void UChaosClothAssetInteractor::SetVectorPropertyValue(const FName& PropertyName, int32 LODIndex, FVector Value)
 {
 	using namespace UE::Chaos::ClothAsset::Private;
 	if (LODIndex == INDEX_NONE)
@@ -526,7 +539,7 @@ void UChaosClothAssetInteractor::SetVectorValue(const FString& PropertyName, int
 			if (TSharedPtr<::Chaos::Softs::FCollectionPropertyFacade> PropertyFacade = CollectionPropertyFacades[Index].Pin())
 			{
 				SetValueWithAlias(PropertyFacade, PropertyName,
-					[&PropertyFacade, &Value](const FString& Name)
+					[&PropertyFacade, &Value](const FName& Name)
 					{
 						PropertyFacade->SetValue(Name, FVector3f(Value));
 					});
@@ -538,7 +551,7 @@ void UChaosClothAssetInteractor::SetVectorValue(const FString& PropertyName, int
 		if (TSharedPtr<::Chaos::Softs::FCollectionPropertyFacade> PropertyFacade = CollectionPropertyFacades[LODIndex].Pin())
 		{
 			SetValueWithAlias(PropertyFacade, PropertyName,
-				[&PropertyFacade, &Value](const FString& Name)
+				[&PropertyFacade, &Value](const FName& Name)
 				{
 					PropertyFacade->SetValue(Name, FVector3f(Value));
 				});
@@ -546,7 +559,7 @@ void UChaosClothAssetInteractor::SetVectorValue(const FString& PropertyName, int
 	}
 }
 
-void UChaosClothAssetInteractor::SetStringValue(const FString& PropertyName, int32 LODIndex, const FString& Value)
+void UChaosClothAssetInteractor::SetStringPropertyValue(const FName& PropertyName, int32 LODIndex, const FString& Value)
 {
 	using namespace UE::Chaos::ClothAsset::Private;
 	if (LODIndex == INDEX_NONE)
@@ -556,7 +569,7 @@ void UChaosClothAssetInteractor::SetStringValue(const FString& PropertyName, int
 			if (TSharedPtr<::Chaos::Softs::FCollectionPropertyFacade> PropertyFacade = CollectionPropertyFacades[Index].Pin())
 			{
 				SetValueWithAlias(PropertyFacade, PropertyName,
-					[&PropertyFacade, &Value](const FString& Name)
+					[&PropertyFacade, &Value](const FName& Name)
 					{
 						PropertyFacade->SetStringValue(Name, Value);
 					});
@@ -568,7 +581,7 @@ void UChaosClothAssetInteractor::SetStringValue(const FString& PropertyName, int
 		if (TSharedPtr<::Chaos::Softs::FCollectionPropertyFacade> PropertyFacade = CollectionPropertyFacades[LODIndex].Pin())
 		{
 			SetValueWithAlias(PropertyFacade, PropertyName,
-				[&PropertyFacade, &Value](const FString& Name)
+				[&PropertyFacade, &Value](const FName& Name)
 				{
 					PropertyFacade->SetStringValue(Name, Value);
 				});
diff --git a/Engine/Plugins/ChaosClothAsset/Source/ChaosClothAssetEngine/Private/ChaosClothAsset/ClothEngineTools.cpp b/Engine/Plugins/ChaosClothAsset/Source/ChaosClothAssetEngine/Private/ChaosClothAsset/ClothEngineTools.cpp
--- a/Engine/Plugins/ChaosClothAsset/Source/ChaosClothAssetEngine/Private/ChaosClothAsset/ClothEngineTools.cpp
+++ b/Engine/Plugins/ChaosClothAsset/Source/ChaosClothAssetEngine/Private/ChaosClothAsset/ClothEngineTools.cpp
@@ -133,12 +133,13 @@ FPointWeightMap FClothEngineTools::GetMaxDistanceWeightMap(const TSharedRef<cons
 FPointWeightMap FClothEngineTools::GetMaxDistanceWeightMap(const FCollectionClothConstFacade& ClothFacade, const ::Chaos::Softs::FCollectionPropertyConstFacade& PropertyFacade, const int32 NumLodSimVertices)
 {
 	int32 MaxDistancePropertyKeyIndex;
-	FString MaxDistanceString = TEXT("MaxDistance");
-	MaxDistanceString = PropertyFacade.GetStringValue(MaxDistanceString, MaxDistanceString, &MaxDistancePropertyKeyIndex);
+	static const FName MaxDistanceName(TEXT("MaxDistance"), EFindName::FNAME_Find);
+	check(MaxDistanceName != NAME_None);
+	const FString MaxDistanceString = PropertyFacade.GetStringValue(MaxDistanceName, MaxDistanceName.ToString(), &MaxDistancePropertyKeyIndex);
 	const bool bHasMaxDistanceProperty = (MaxDistancePropertyKeyIndex != INDEX_NONE);
 	const float MaxDistanceOffset = bHasMaxDistanceProperty ? PropertyFacade.GetLowValue<float>(MaxDistancePropertyKeyIndex) : TNumericLimits<float>::Max();  // Uses infinite distance when no MaxDistance properties are set
 	const float MaxDistanceScale = bHasMaxDistanceProperty ? PropertyFacade.GetHighValue<float>(MaxDistancePropertyKeyIndex) - MaxDistanceOffset : 0.f;
-	const TConstArrayView<float> MaxDistanceWeightMap = ClothFacade.GetWeightMap(FName(MaxDistanceString));
+	const TConstArrayView<float> MaxDistanceWeightMap = ClothFacade.GetWeightMap(MaxDistanceName);
 
 	return (MaxDistanceWeightMap.Num() == NumLodSimVertices) ?
 		FPointWeightMap(MaxDistanceWeightMap, MaxDistanceOffset, MaxDistanceScale) :
diff --git a/Engine/Plugins/ChaosClothAsset/Source/ChaosClothAssetEngine/Public/ChaosClothAsset/ClothAssetInteractor.h b/Engine/Plugins/ChaosClothAsset/Source/ChaosClothAssetEngine/Public/ChaosClothAsset/ClothAssetInteractor.h
--- a/Engine/Plugins/ChaosClothAsset/Source/ChaosClothAssetEngine/Public/ChaosClothAsset/ClothAssetInteractor.h
+++ b/Engine/Plugins/ChaosClothAsset/Source/ChaosClothAssetEngine/Public/ChaosClothAsset/ClothAssetInteractor.h
@@ -25,63 +25,166 @@ public:
 
 	/** Generate a list of all properties held by this interactor. Properties for all LODs will be returned if LODIndex = -1.*/
 	UFUNCTION(BlueprintCallable, Category = ClothProperty)
-	CHAOSCLOTHASSETENGINE_API TArray<FString> GetAllProperties(int32 LODIndex = -1) const;
+	CHAOSCLOTHASSETENGINE_API TArray<FName> GetAllPropertyNames(int32 LODIndex = -1) const;
 
 	/** Get the value for a property cast to float. DefaultValue will be returned if the property is not found.*/
 	UFUNCTION(BlueprintCallable, Category = ClothProperty)
-	CHAOSCLOTHASSETENGINE_API float GetFloatValue(const FString& PropertyName, int32 LODIndex = 0, float DefaultValue = 0.f) const;
+	CHAOSCLOTHASSETENGINE_API float GetFloatPropertyValue(const FName& PropertyName, int32 LODIndex = 0, float DefaultValue = 0.f) const;
 
 	/** Get the low value for a weighted property value (same as GetFloatValue). DefaultValue will be returned if the property is not found.*/
 	UFUNCTION(BlueprintCallable, Category = ClothProperty)
-	CHAOSCLOTHASSETENGINE_API float GetLowFloatValue(const FString& PropertyName, int32 LODIndex = 0, float DefaultValue = 0.f) const;
+	CHAOSCLOTHASSETENGINE_API float GetLowFloatPropertyValue(const FName& PropertyName, int32 LODIndex = 0, float DefaultValue = 0.f) const;
 
 	/** Get the high value for a weighted property value. DefaultValue will be returned if the property is not found.*/
 	UFUNCTION(BlueprintCallable, Category = ClothProperty)
-	CHAOSCLOTHASSETENGINE_API float GetHighFloatValue(const FString& PropertyName, int32 LODIndex = 0, float DefaultValue = 0.f) const;
+	CHAOSCLOTHASSETENGINE_API float GetHighFloatPropertyValue(const FName& PropertyName, int32 LODIndex = 0, float DefaultValue = 0.f) const;
 
 	/** Get the low and high values for a weighted property value. DefaultValue will be returned if the property is not found.*/
 	UFUNCTION(BlueprintCallable, Category = ClothProperty)
-	CHAOSCLOTHASSETENGINE_API FVector2D GetWeightedFloatValue(const FString& PropertyName, int32 LODIndex = 0, FVector2D DefaultValue = FVector2D(0., 0.)) const;
+	CHAOSCLOTHASSETENGINE_API FVector2D GetWeightedFloatPropertyValue(const FName& PropertyName, int32 LODIndex = 0, FVector2D DefaultValue = FVector2D(0., 0.)) const;
 
 	/** Get the value for a property cast to int. DefaultValue will be returned if the property is not found.*/
 	UFUNCTION(BlueprintCallable, Category = ClothProperty)
-	CHAOSCLOTHASSETENGINE_API int32 GetIntValue(const FString& PropertyName, int32 LODIndex = 0, int32 DefaultValue = 0) const;
+	CHAOSCLOTHASSETENGINE_API int32 GetIntPropertyValue(const FName& PropertyName, int32 LODIndex = 0, int32 DefaultValue = 0) const;
 
 	/** Get the value for a property cast to vector. DefaultValue will be returned if the property is not found.*/
 	UFUNCTION(BlueprintCallable, Category = ClothProperty)
-	CHAOSCLOTHASSETENGINE_API FVector GetVectorValue(const FString& PropertyName, int32 LODIndex = 0, FVector DefaultValue = FVector(0.)) const;
+	CHAOSCLOTHASSETENGINE_API FVector GetVectorPropertyValue(const FName& PropertyName, int32 LODIndex = 0, FVector DefaultValue = FVector(0.)) const;
 
 	/** Get the string value for a property (typically the associated map name for weighted values). DefaultValue will be returned if the property is not found.*/
 	UFUNCTION(BlueprintCallable, Category = ClothProperty)
-	CHAOSCLOTHASSETENGINE_API FString GetStringValue(const FString& PropertyName, int32 LODIndex = 0, const FString& DefaultValue = "") const;
+	CHAOSCLOTHASSETENGINE_API FString GetStringPropertyValue(const FName& PropertyName, int32 LODIndex = 0, const FString& DefaultValue = "") const;
 
 	/**Set the value for a property (if it exists). This sets the Low and High values for weighted values. All LODs will be set when LODIndex = -1.*/
 	UFUNCTION(BlueprintCallable, Category = ClothProperty)
-	CHAOSCLOTHASSETENGINE_API void SetFloatValue(const FString& PropertyName, int32 LODIndex = -1, float Value = 0.f);
+	CHAOSCLOTHASSETENGINE_API void SetFloatPropertyValue(const FName& PropertyName, int32 LODIndex = -1, float Value = 0.f);
 
 	/**Set the low value for a weighted property (if it exists). All LODs will be set when LODIndex = -1.*/
 	UFUNCTION(BlueprintCallable, Category = ClothProperty)
-	CHAOSCLOTHASSETENGINE_API void SetLowFloatValue(const FString& PropertyName, int32 LODIndex = -1, float Value = 0.f);
+	CHAOSCLOTHASSETENGINE_API void SetLowFloatPropertyValue(const FName& PropertyName, int32 LODIndex = -1, float Value = 0.f);
 
 	/**Set the high value for a weighted property (if it exists). All LODs will be set when LODIndex = -1.*/
 	UFUNCTION(BlueprintCallable, Category = ClothProperty)
-	CHAOSCLOTHASSETENGINE_API void SetHighFloatValue(const FString& PropertyName, int32 LODIndex = -1, float Value = 0.f);
+	CHAOSCLOTHASSETENGINE_API void SetHighFloatPropertyValue(const FName& PropertyName, int32 LODIndex = -1, float Value = 0.f);
 
 	/**Set the low and high values for a weighted property (if it exists). All LODs will be set when LODIndex = -1.*/
 	UFUNCTION(BlueprintCallable, Category = ClothProperty)
-	CHAOSCLOTHASSETENGINE_API void SetWeightedFloatValue(const FString& PropertyName, int32 LODIndex = -1, FVector2D Value = FVector2D(0., 0.));
+	CHAOSCLOTHASSETENGINE_API void SetWeightedFloatPropertyValue(const FName& PropertyName, int32 LODIndex = -1, FVector2D Value = FVector2D(0., 0.));
 
 	/**Set the value for a property (if it exists). All LODs will be set when LODIndex = -1.*/
 	UFUNCTION(BlueprintCallable, Category = ClothProperty)
-	CHAOSCLOTHASSETENGINE_API void SetIntValue(const FString& PropertyName, int32 LODIndex = -1, int32 Value = 0);
+	CHAOSCLOTHASSETENGINE_API void SetIntPropertyValue(const FName& PropertyName, int32 LODIndex = -1, int32 Value = 0);
 
 	/**Set the value for a property (if it exists). All LODs will be set when LODIndex = -1.*/
 	UFUNCTION(BlueprintCallable, Category = ClothProperty)
-	CHAOSCLOTHASSETENGINE_API void SetVectorValue(const FString& PropertyName, int32 LODIndex = -1, FVector Value = FVector(0.));
+	CHAOSCLOTHASSETENGINE_API void SetVectorPropertyValue(const FName& PropertyName, int32 LODIndex = -1, FVector Value = FVector(0.));
 
 	/**Set the string value for a property (if it exists). This is typically the map name associated with a property. All LODs will be set when LODIndex = -1.*/
 	UFUNCTION(BlueprintCallable, Category = ClothProperty)
-	CHAOSCLOTHASSETENGINE_API void SetStringValue(const FString& PropertyName, int32 LODIndex = -1, const FString& Value = "");
+	CHAOSCLOTHASSETENGINE_API void SetStringPropertyValue(const FName& PropertyName, int32 LODIndex = -1, const FString& Value = "");
+
+	// Deprecated string-based key versions.
+	/** Generate a list of all properties held by this interactor. Properties for all LODs will be returned if LODIndex = -1.*/
+	UFUNCTION(BlueprintCallable, meta = (DeprecatedFunction = "Note", DeprecationMessage = "Use GetAllPropertyNames instead"))
+	CHAOSCLOTHASSETENGINE_API TArray<FString> GetAllProperties(int32 LODIndex = -1) const;
+
+	/** Get the value for a property cast to float. DefaultValue will be returned if the property is not found.*/
+	UFUNCTION(BlueprintCallable, Category = "ClothProperty|Deprecated", meta = (DeprecatedFunction = "Note", DeprecationMessage = "Use GetFloatPropertyValue instead"))
+	float GetFloatValue(const FString& PropertyName, int32 LODIndex = 0, float DefaultValue = 0.f) const
+	{
+		return GetFloatPropertyValue(FName(PropertyName), LODIndex, DefaultValue);
+	}
+
+	/** Get the low value for a weighted property value (same as GetFloatValue). DefaultValue will be returned if the property is not found.*/
+	UFUNCTION(BlueprintCallable, Category = "ClothProperty|Deprecated", meta = (DeprecatedFunction = "Note", DeprecationMessage = "Use GetLowFloatPropertyValue instead"))
+	float GetLowFloatValue(const FString& PropertyName, int32 LODIndex = 0, float DefaultValue = 0.f) const
+	{
+		return GetLowFloatPropertyValue(FName(PropertyName), LODIndex, DefaultValue);
+	}
+
+	/** Get the high value for a weighted property value. DefaultValue will be returned if the property is not found.*/
+	UFUNCTION(BlueprintCallable, Category = "ClothProperty|Deprecated", meta = (DeprecatedFunction = "Note", DeprecationMessage = "Use GetHighFloatPropertyValue instead"))
+	float GetHighFloatValue(const FString& PropertyName, int32 LODIndex = 0, float DefaultValue = 0.f) const
+	{
+		return GetHighFloatPropertyValue(FName(PropertyName), LODIndex, DefaultValue);
+	}
+
+	/** Get the low and high values for a weighted property value. DefaultValue will be returned if the property is not found.*/
+	UFUNCTION(BlueprintCallable, Category = "ClothProperty|Deprecated", meta = (DeprecatedFunction = "Note", DeprecationMessage = "Use GetWeightedFloatPropertyValue instead"))
+	FVector2D GetWeightedFloatValue(const FString& PropertyName, int32 LODIndex = 0, FVector2D DefaultValue = FVector2D(0., 0.)) const
+	{
+		return GetWeightedFloatPropertyValue(FName(PropertyName), LODIndex, DefaultValue);
+	}
+
+	/** Get the value for a property cast to int. DefaultValue will be returned if the property is not found.*/
+	UFUNCTION(BlueprintCallable, Category = "ClothProperty|Deprecated", meta = (DeprecatedFunction = "Note", DeprecationMessage = "Use GetIntPropertyValue instead"))
+	int32 GetIntValue(const FString& PropertyName, int32 LODIndex = 0, int32 DefaultValue = 0) const
+	{
+		return GetIntPropertyValue(FName(PropertyName), LODIndex, DefaultValue);
+	}
+
+	/** Get the value for a property cast to vector. DefaultValue will be returned if the property is not found.*/
+	UFUNCTION(BlueprintCallable, Category = "ClothProperty|Deprecated", meta = (DeprecatedFunction = "Note", DeprecationMessage = "Use GetVectorPropertyValue instead"))
+	FVector GetVectorValue(const FString& PropertyName, int32 LODIndex = 0, FVector DefaultValue = FVector(0.)) const
+	{
+		return GetVectorPropertyValue(FName(PropertyName), LODIndex, DefaultValue);
+	}
+
+	/** Get the string value for a property (typically the associated map name for weighted values). DefaultValue will be returned if the property is not found.*/
+	UFUNCTION(BlueprintCallable, Category = "ClothProperty|Deprecated", meta = (DeprecatedFunction = "Note", DeprecationMessage = "Use GetStringPropertyValue instead"))
+	FString GetStringValue(const FString& PropertyName, int32 LODIndex = 0, const FString& DefaultValue = "") const
+	{
+		return GetStringPropertyValue(FName(PropertyName), LODIndex, DefaultValue);
+	}
+
+	/**Set the value for a property (if it exists). This sets the Low and High values for weighted values. All LODs will be set when LODIndex = -1.*/
+	UFUNCTION(BlueprintCallable, Category = "ClothProperty|Deprecated", meta = (DeprecatedFunction = "Note", DeprecationMessage = "Use SetFloatPropertyValue instead"))
+	void SetFloatValue(const FString& PropertyName, int32 LODIndex = -1, float Value = 0.f)
+	{
+		SetFloatPropertyValue(FName(PropertyName), LODIndex, Value);
+	}
+
+	/**Set the low value for a weighted property (if it exists). All LODs will be set when LODIndex = -1.*/
+	UFUNCTION(BlueprintCallable, Category = "ClothProperty|Deprecated", meta = (DeprecatedFunction = "Note", DeprecationMessage = "Use SetLowFloatPropertyValue instead"))
+	void SetLowFloatValue(const FString& PropertyName, int32 LODIndex = -1, float Value = 0.f)
+	{
+		SetLowFloatPropertyValue(FName(PropertyName), LODIndex, Value);
+	}
+
+	/**Set the high value for a weighted property (if it exists). All LODs will be set when LODIndex = -1.*/
+	UFUNCTION(BlueprintCallable, Category = "ClothProperty|Deprecated", meta = (DeprecatedFunction = "Note", DeprecationMessage = "Use SetHighFloatPropertyValue instead"))
+	void SetHighFloatValue(const FString& PropertyName, int32 LODIndex = -1, float Value = 0.f)
+	{
+		SetHighFloatPropertyValue(FName(PropertyName), LODIndex, Value);
+	}
+
+	/**Set the low and high values for a weighted property (if it exists). All LODs will be set when LODIndex = -1.*/
+	UFUNCTION(BlueprintCallable, Category = "ClothProperty|Deprecated", meta = (DeprecatedFunction = "Note", DeprecationMessage = "Use SetWeightedFloatPropertyValue instead"))
+	void SetWeightedFloatValue(const FString& PropertyName, int32 LODIndex = -1, FVector2D Value = FVector2D(0., 0.))
+	{
+		SetWeightedFloatPropertyValue(FName(PropertyName), LODIndex, Value);
+	}
+
+	/**Set the value for a property (if it exists). All LODs will be set when LODIndex = -1.*/
+	UFUNCTION(BlueprintCallable, Category = "ClothProperty|Deprecated", meta = (DeprecatedFunction = "Note", DeprecationMessage = "Use SetIntPropertyValue instead"))
+	void SetIntValue(const FString& PropertyName, int32 LODIndex = -1, int32 Value = 0)
+	{
+		SetIntPropertyValue(FName(PropertyName), LODIndex, Value);
+	}
+
+	/**Set the value for a property (if it exists). All LODs will be set when LODIndex = -1.*/
+	UFUNCTION(BlueprintCallable, Category = "ClothProperty|Deprecated", meta = (DeprecatedFunction = "Note", DeprecationMessage = "Use SetVectorPropertyValue instead"))
+	void SetVectorValue(const FString& PropertyName, int32 LODIndex = -1, FVector Value = FVector(0.))
+	{
+		SetVectorPropertyValue(FName(PropertyName), LODIndex, Value);
+	}
+
+	/**Set the string value for a property (if it exists). This is typically the map name associated with a property. All LODs will be set when LODIndex = -1.*/
+	UFUNCTION(BlueprintCallable, Category = "ClothProperty|Deprecated", meta = (DeprecatedFunction = "Note", DeprecationMessage = "Use SetStringPropertyValue instead"))
+	void SetStringValue(const FString& PropertyName, int32 LODIndex = -1, const FString& Value = "")
+	{
+		SetStringPropertyValue(FName(PropertyName), LODIndex, Value);
+	}
 
 private:
 	TArray<TWeakPtr<::Chaos::Softs::FCollectionPropertyFacade>> CollectionPropertyFacades;
diff --git a/Engine/Plugins/ChaosClothAssetEditor/Source/ChaosClothAssetDataflowNodes/Private/ChaosClothAsset/ClothCollectionQueryNode.cpp b/Engine/Plugins/ChaosClothAssetEditor/Source/ChaosClothAssetDataflowNodes/Private/ChaosClothAsset/ClothCollectionQueryNode.cpp
--- a/Engine/Plugins/ChaosClothAssetEditor/Source/ChaosClothAssetDataflowNodes/Private/ChaosClothAsset/ClothCollectionQueryNode.cpp
+++ b/Engine/Plugins/ChaosClothAssetEditor/Source/ChaosClothAssetDataflowNodes/Private/ChaosClothAsset/ClothCollectionQueryNode.cpp
@@ -48,7 +48,7 @@ void FChaosClothAssetCollectionQueryNode::Evaluate(UE::Dataflow::FContext& Conte
 		bool bBooleanPropertyResult = bBooleanPropertyValue;
 		if (PropertyFacade.IsValid())
 		{
-			bBooleanPropertyResult = PropertyFacade.GetValue(BooleanPropertyName, bBooleanPropertyValue);
+			bBooleanPropertyResult = PropertyFacade.GetValue(FName(BooleanPropertyName), bBooleanPropertyValue);
 		}
 		SetValue(Context, bBooleanPropertyResult, &bBooleanPropertyValue);
 	}
diff --git a/Engine/Plugins/ChaosClothAssetEditor/Source/ChaosClothAssetDataflowNodes/Private/ChaosClothAsset/MergeClothCollectionsNode.cpp b/Engine/Plugins/ChaosClothAssetEditor/Source/ChaosClothAssetDataflowNodes/Private/ChaosClothAsset/MergeClothCollectionsNode.cpp
--- a/Engine/Plugins/ChaosClothAssetEditor/Source/ChaosClothAssetDataflowNodes/Private/ChaosClothAsset/MergeClothCollectionsNode.cpp
+++ b/Engine/Plugins/ChaosClothAssetEditor/Source/ChaosClothAssetDataflowNodes/Private/ChaosClothAsset/MergeClothCollectionsNode.cpp
@@ -16,7 +16,7 @@
 
 namespace UE::Chaos::ClothAsset::Private
 {
-	static void LogAndToastDifferentWeightMapNames(const FDataflowNode& DataflowNode, const FString& PropertyName, const FString& InWeightMapName, const FString& OutWeightMapName, const FString& WeightMapName)
+	static void LogAndToastDifferentWeightMapNames(const FDataflowNode& DataflowNode, const FName& PropertyName, const FString& InWeightMapName, const FString& OutWeightMapName, const FString& WeightMapName)
 	{
 		using namespace UE::Chaos::ClothAsset;
 
@@ -26,7 +26,7 @@ namespace UE::Chaos::ClothAsset::Private
 			LOCTEXT(
 				"DifferentWeightMapNamesDetails",
 				"Two identical Cloth Collection properties '{0}' are being merged but have different weight map names '{1}' and '{2}'. The weight map named '{3}' will be used in the resulting merge."),
-			FText::FromString(PropertyName),
+			FText::FromName(PropertyName),
 			FText::FromString(OutWeightMapName),
 			FText::FromString(InWeightMapName),
 			FText::FromString(WeightMapName));
@@ -78,7 +78,7 @@ namespace UE::Chaos::ClothAsset::Private
 		bool bAppendedCloth,
 		const FCollectionClothConstFacade& InClothFacade, FCollectionClothFacade& OutClothFacade,
 		const FVector2f& InPropertyBounds, const FVector2f& OutPropertyBounds,
-		const FVector2f& PropertyBounds, const FString& PropertyName,
+		const FVector2f& PropertyBounds, const FName& PropertyName,
 		const FString& InWeightMapName, const FString& OutWeightMapName, TMap<FString,FMergedProperty>& MergedPropertyMaps)
 	{
 		const FMergedProperty MergedProperty = {OutWeightMapName + FString(TEXT("_")) + InWeightMapName,
@@ -92,15 +92,15 @@ namespace UE::Chaos::ClothAsset::Private
 				return MergedPropertyMap.Key;
 			}
 		}
-		FString WeightMapName = PropertyName;
+		FString WeightMapName = PropertyName.ToString();
 		int32 WeightMapCount = 0;
 		
 		// the weight map could already been stored on the out collection and linked to different bounds
 		// coming from the out collection itself or from previous merge with in collection
 		// Since we don't want to break them we need to create a new one on the first available slot
-		while(OutClothFacade.GetWeightMap(FName(WeightMapName)).Num() > 0)
+		while(OutClothFacade.GetWeightMap(PropertyName).Num() > 0)
 		{
-			WeightMapName = PropertyName;
+			WeightMapName = PropertyName.ToString();
 			WeightMapName.AppendInt(++WeightMapCount);
 		}
 		
@@ -197,7 +197,7 @@ namespace UE::Chaos::ClothAsset::Private
 		::Chaos::Softs::FEmbeddedSpringFacade& OutSpringFacade,
 		const FConstraintMergeData& ConstraintMergeData,
 		const FVector2f& InPropertyBounds, const FVector2f& OutPropertyBounds,
-		const FVector2f& PropertyBounds, const FString& PropertyName)
+		const FVector2f& PropertyBounds, const FName& PropertyName)
 	{
 		if (PropertyBounds[0] == PropertyBounds[1])
 		{
@@ -205,30 +205,29 @@ namespace UE::Chaos::ClothAsset::Private
 			return;
 		}
 		using namespace ::Chaos::Softs;
-		const FName PropertyNameName(*PropertyName);
 
 
 		TArrayView<float> WeightMap;
 		TConstArrayView<float> InWeightMap;
-		if (VertexSpringConstraintPropertyNames.Contains(PropertyNameName))
+		if (VertexSpringConstraintPropertyNames.Contains(PropertyName))
 		{
 			FEmbeddedSpringConstraintFacade OutConstraintFacade = OutSpringFacade.GetSpringConstraint(ConstraintMergeData.VertexSpringConstraintIndex);
 			check(ConstraintMergeData.OtherVertexSpringConstraintIndex != INDEX_NONE);
 			FEmbeddedSpringConstraintFacade InConstraintFacade = InSpringFacade.GetSpringConstraintConst(ConstraintMergeData.OtherVertexSpringConstraintIndex);
 
-			if (PropertyNameName == FName(TEXT("VertexSpringExtensionStiffness")))
+			if (PropertyName == FName(TEXT("VertexSpringExtensionStiffness")))
 			{
 				WeightMap = OutConstraintFacade.GetExtensionStiffness();
 				InWeightMap = InConstraintFacade.GetExtensionStiffnessConst();
 			}
-			else if (PropertyNameName == FName(TEXT("VertexSpringCompressionStiffness")))
+			else if (PropertyName == FName(TEXT("VertexSpringCompressionStiffness")))
 			{
 				WeightMap = OutConstraintFacade.GetCompressionStiffness();
 				InWeightMap = InConstraintFacade.GetCompressionStiffnessConst();
 			}
 			else
 			{
-				check(PropertyNameName == FName(TEXT("VertexSpringDamping")));
+				check(PropertyName == FName(TEXT("VertexSpringDamping")));
 				WeightMap = OutConstraintFacade.GetDamping();
 				InWeightMap = InConstraintFacade.GetDampingConst();
 			}
@@ -238,25 +237,25 @@ namespace UE::Chaos::ClothAsset::Private
 			FillWeightMap(WeightMap.Left(OutNumSprings), WeightMap.Left(OutNumSprings), PropertyBounds, OutPropertyBounds);
 			FillWeightMap(WeightMap.Right(InNumSprings), InWeightMap, PropertyBounds, InPropertyBounds);
 		}
-		else if (VertexFaceSpringConstraintPropertyNames.Contains(PropertyNameName))
+		else if (VertexFaceSpringConstraintPropertyNames.Contains(PropertyName))
 		{
 			FEmbeddedSpringConstraintFacade OutConstraintFacade = OutSpringFacade.GetSpringConstraint(ConstraintMergeData.VertexFaceSpringConstraintIndex);
 			check(ConstraintMergeData.OtherVertexFaceSpringConstraintIndex != INDEX_NONE);
 			FEmbeddedSpringConstraintFacade InConstraintFacade = InSpringFacade.GetSpringConstraintConst(ConstraintMergeData.OtherVertexFaceSpringConstraintIndex);
 
-			if (PropertyNameName == FName(TEXT("VertexFaceSpringExtensionStiffness")))
+			if (PropertyName == FName(TEXT("VertexFaceSpringExtensionStiffness")))
 			{
 				WeightMap = OutConstraintFacade.GetExtensionStiffness();
 				InWeightMap = InConstraintFacade.GetExtensionStiffnessConst();
 			}
-			else if (PropertyNameName == FName(TEXT("VertexFaceSpringCompressionStiffness")))
+			else if (PropertyName == FName(TEXT("VertexFaceSpringCompressionStiffness")))
 			{
 				WeightMap = OutConstraintFacade.GetCompressionStiffness();
 				InWeightMap = InConstraintFacade.GetCompressionStiffnessConst();
 			}
 			else
 			{
-				check(PropertyNameName == FName(TEXT("VertexFaceSpringDamping")));
+				check(PropertyName == FName(TEXT("VertexFaceSpringDamping")));
 				WeightMap = OutConstraintFacade.GetDamping();
 				InWeightMap = InConstraintFacade.GetDampingConst();
 			}
@@ -272,19 +271,19 @@ namespace UE::Chaos::ClothAsset::Private
 			check(ConstraintMergeData.OtherFaceSpringConstraintIndex != INDEX_NONE);
 			FEmbeddedSpringConstraintFacade InConstraintFacade = InSpringFacade.GetSpringConstraintConst(ConstraintMergeData.OtherFaceSpringConstraintIndex);
 
-			if (PropertyNameName == FName(TEXT("FaceSpringExtensionStiffness")))
+			if (PropertyName == FName(TEXT("FaceSpringExtensionStiffness")))
 			{
 				WeightMap = OutConstraintFacade.GetExtensionStiffness();
 				InWeightMap = InConstraintFacade.GetExtensionStiffnessConst();
 			}
-			else if (PropertyNameName == FName(TEXT("FaceSpringCompressionStiffness")))
+			else if (PropertyName == FName(TEXT("FaceSpringCompressionStiffness")))
 			{
 				WeightMap = OutConstraintFacade.GetCompressionStiffness();
 				InWeightMap = InConstraintFacade.GetCompressionStiffnessConst();
 			}
 			else
 			{
-				check(PropertyNameName == FName(TEXT("FaceSpringDamping")));
+				check(PropertyName == FName(TEXT("FaceSpringDamping")));
 				WeightMap = OutConstraintFacade.GetDamping();
 				InWeightMap = InConstraintFacade.GetDampingConst();
 			}
@@ -302,7 +301,7 @@ namespace UE::Chaos::ClothAsset::Private
 			  const int32 InKeyIndex,
 			  const int32 OutKeyIndex,
 			  const ::Chaos::Softs::ECollectionPropertyFlags InPropertyFlags,
-			  const FString& PropertyName)
+			  const FName& PropertyName)
 	{
 		PRAGMA_DISABLE_DEPRECATION_WARNINGS
 		// TODO: GetFlags needs to return an ECollectionPropertyFlags, not an uint8, but the uint8 getter needs to be deprecated first
@@ -327,15 +326,18 @@ namespace UE::Chaos::ClothAsset::Private
 			}
 			if(!ensure(OutPropertyFacade.IsIntrinsic(OutKeyIndex) == InPropertyFacade.IsIntrinsic(InKeyIndex)))
 			{
-				UE_LOG(LogChaosClothAssetDataflowNodes, Warning, TEXT("MergeClothCollectionsNode: Mismatch in intrinsic flag onto %s property"), *PropertyName);
+				const FString PropertyNameString = PropertyName.ToString();
+				UE_LOG(LogChaosClothAssetDataflowNodes, Warning, TEXT("MergeClothCollectionsNode: Mismatch in intrinsic flag onto %s property"), *PropertyNameString);
 			}
 			if(!ensure(OutPropertyFacade.IsLegacy(OutKeyIndex) == InPropertyFacade.IsLegacy(InKeyIndex)))
 			{
-				UE_LOG(LogChaosClothAssetDataflowNodes, Warning, TEXT("MergeClothCollectionsNode: Mismatch in legacy flag onto %s property"), *PropertyName);
+				const FString PropertyNameString = PropertyName.ToString();
+				UE_LOG(LogChaosClothAssetDataflowNodes, Warning, TEXT("MergeClothCollectionsNode: Mismatch in legacy flag onto %s property"), *PropertyNameString);
 			}
 			if(!ensure(OutPropertyFacade.IsInterpolable(OutKeyIndex) == InPropertyFacade.IsInterpolable(InKeyIndex)))
 			{
-				UE_LOG(LogChaosClothAssetDataflowNodes, Warning, TEXT("MergeClothCollectionsNode: Mismatch in interpolable flag onto %s property"), *PropertyName);
+				const FString PropertyNameString = PropertyName.ToString();
+				UE_LOG(LogChaosClothAssetDataflowNodes, Warning, TEXT("MergeClothCollectionsNode: Mismatch in interpolable flag onto %s property"), *PropertyNameString);
 			}
 		}
 		return PropertyFlags;
@@ -362,8 +364,8 @@ namespace UE::Chaos::ClothAsset::Private
 			PRAGMA_ENABLE_DEPRECATION_WARNINGS
 
 			// Get the matching output key for the given input one 
-			const FString& InPropertyKey = InPropertyFacade.GetKey(InKeyIndex);
-			int32 OutKeyIndex = OutPropertyFacade.GetKeyIndex(InPropertyKey);
+			const FName& InPropertyKey = InPropertyFacade.GetKeyName(InKeyIndex);
+			int32 OutKeyIndex = OutPropertyFacade.GetKeyNameIndex(InPropertyKey);
 
 			// We first check if the output key exists into the output facade
 			bool bOverrideProperty = true;
@@ -382,7 +384,7 @@ namespace UE::Chaos::ClothAsset::Private
 					OutPropertyFacade.SetFlags(OutKeyIndex, PropertyFlags);
 					OutPropertyFacade.SetWeightedFloatValue(OutKeyIndex, PropertyBounds);
 
-					if (IsSpringConstraintProperty(FName(*InPropertyKey)))
+					if (IsSpringConstraintProperty(InPropertyKey))
 					{
 						UpdateSpringConstraintWeights(bAppendedCloth, InSpringFacade, OutSpringFacade, ConstraintMergeData, InPropertyBounds, OutPropertyBounds, PropertyBounds, InPropertyKey);
 					}
diff --git a/Engine/Plugins/ChaosClothAssetEditor/Source/ChaosClothAssetDataflowNodes/Private/ChaosClothAsset/RemeshNode.cpp b/Engine/Plugins/ChaosClothAssetEditor/Source/ChaosClothAssetDataflowNodes/Private/ChaosClothAsset/RemeshNode.cpp
--- a/Engine/Plugins/ChaosClothAssetEditor/Source/ChaosClothAssetDataflowNodes/Private/ChaosClothAsset/RemeshNode.cpp
+++ b/Engine/Plugins/ChaosClothAssetEditor/Source/ChaosClothAssetDataflowNodes/Private/ChaosClothAsset/RemeshNode.cpp
@@ -1084,25 +1084,25 @@ namespace UE::Chaos::ClothAsset::Private
 		FCollectionPropertyConstFacade InProperties(InClothCollection);
 
 		// Reconstruct KinematicVertexSet
-		const FString MaxDistanceString = GET_MEMBER_NAME_STRING_CHECKED(FChaosClothAssetSimulationMaxDistanceConfigNode, MaxDistance);
+		const FName MaxDistanceName = GET_MEMBER_NAME_CHECKED(FChaosClothAssetSimulationMaxDistanceConfigNode, MaxDistance);
 		const FName KinematicVertices3DName = GET_MEMBER_NAME_CHECKED(FChaosClothAssetSimulationMaxDistanceConfigNode, KinematicVertices3D);
-		if (InProperties.GetKeyIndex(MaxDistanceString) != INDEX_NONE)
+		if (InProperties.GetKeyNameIndex(MaxDistanceName) != INDEX_NONE)
 		{
-			const FName MaxDistanceMapName(InProperties.GetStringValue(MaxDistanceString, MaxDistanceString));
+			const FName MaxDistanceMapName(InProperties.GetStringValue(MaxDistanceName, MaxDistanceName.ToString()));
 			OutSelectionFacade.FindOrAddSelectionSet(KinematicVertices3DName, ClothCollectionGroup::SimVertices3D) =
-				FClothGeometryTools::GenerateKinematicVertices3D(OutClothCollection, MaxDistanceMapName, InProperties.GetWeightedFloatValue(MaxDistanceString, FVector2f(0.f, 1.f)), NAME_None);
+				FClothGeometryTools::GenerateKinematicVertices3D(OutClothCollection, MaxDistanceMapName, InProperties.GetWeightedFloatValue(MaxDistanceName, FVector2f(0.f, 1.f)), NAME_None);
 		}
 
 		// Reconstruct collision spheres
-		const FString SelfCollisionSphereStiffnessString = GET_MEMBER_NAME_STRING_CHECKED(FChaosClothAssetSimulationSelfCollisionSpheresConfigNode, SelfCollisionSphereStiffness);
-		const FString SelfCollisionSphereRadiusString = GET_MEMBER_NAME_STRING_CHECKED(FChaosClothAssetSimulationSelfCollisionSpheresConfigNode, SelfCollisionSphereRadius);
-		const FString SelfCollisionSphereRadiusCullMultiplierString = GET_MEMBER_NAME_STRING_CHECKED(FChaosClothAssetSimulationSelfCollisionSpheresConfigNode, SelfCollisionSphereRadiusCullMultiplier);
-		const FString SelfCollisionSphereSetNameString = GET_MEMBER_NAME_STRING_CHECKED(FChaosClothAssetSimulationSelfCollisionSpheresConfigNode, SelfCollisionSphereSetName);
+		const FName SelfCollisionSphereStiffnessName = GET_MEMBER_NAME_CHECKED(FChaosClothAssetSimulationSelfCollisionSpheresConfigNode, SelfCollisionSphereStiffness);
+		const FName SelfCollisionSphereRadiusName = GET_MEMBER_NAME_CHECKED(FChaosClothAssetSimulationSelfCollisionSpheresConfigNode, SelfCollisionSphereRadius);
+		const FName SelfCollisionSphereRadiusCullMultiplierName = GET_MEMBER_NAME_CHECKED(FChaosClothAssetSimulationSelfCollisionSpheresConfigNode, SelfCollisionSphereRadiusCullMultiplier);
+		const FName SelfCollisionSphereSetNameName = GET_MEMBER_NAME_CHECKED(FChaosClothAssetSimulationSelfCollisionSpheresConfigNode, SelfCollisionSphereSetName);
 
-		if (InProperties.GetKeyIndex(SelfCollisionSphereStiffnessString) != INDEX_NONE)
+		if (InProperties.GetKeyNameIndex(SelfCollisionSphereStiffnessName) != INDEX_NONE)
 		{
-			const float SelfCollisionSphereRadius = InProperties.GetValue<float>(SelfCollisionSphereRadiusString);
-			const float SelfCollisionSphereRadiusCullMultiplier = InProperties.GetValue<float>(SelfCollisionSphereRadiusCullMultiplierString);
+			const float SelfCollisionSphereRadius = InProperties.GetValue<float>(SelfCollisionSphereRadiusName);
+			const float SelfCollisionSphereRadiusCullMultiplier = InProperties.GetValue<float>(SelfCollisionSphereRadiusCullMultiplierName);
 			const float CullDiameterSq = FMath::Square(SelfCollisionSphereRadius * SelfCollisionSphereRadiusCullMultiplier * 2.f);
 
 			if (OutClothFacade.IsValid() && CullDiameterSq > 0.f)
@@ -1111,7 +1111,7 @@ namespace UE::Chaos::ClothAsset::Private
 				TSet<int32> VertexSet;
 				FClothGeometryTools::SampleVertices(SimPositions, CullDiameterSq, VertexSet);
 
-				const FName SelectionSetName(*InProperties.GetStringValue(SelfCollisionSphereSetNameString, SelfCollisionSphereSetNameString));
+				const FName SelectionSetName(*InProperties.GetStringValue(SelfCollisionSphereSetNameName, SelfCollisionSphereSetNameName.ToString()));
 				OutSelectionFacade.FindOrAddSelectionSet(SelectionSetName, ClothCollectionGroup::SimVertices3D) = VertexSet;
 			}
 		}
@@ -1120,26 +1120,24 @@ namespace UE::Chaos::ClothAsset::Private
 
 		// v1 (weight map)
 		PRAGMA_DISABLE_DEPRECATION_WARNINGS
-			const FString FixedEndWeightMapString = GET_MEMBER_NAME_STRING_CHECKED(FChaosClothAssetSimulationLongRangeAttachmentConfigNode, FixedEndWeightMap);
-		if (InProperties.GetKeyIndex(FixedEndWeightMapString) != INDEX_NONE)
+		const FName FixedEndWeightMapName = GET_MEMBER_NAME_CHECKED(FChaosClothAssetSimulationLongRangeAttachmentConfigNode, FixedEndWeightMap);
+		static const FName UseGeodesicTethersName(TEXT("UseGeodesicTethers"), EFindName::FNAME_Find);
+		check(UseGeodesicTethersName != NAME_None);
+		if (InProperties.GetKeyNameIndex(FixedEndWeightMapName) != INDEX_NONE)
 		{
-			FString UseGeodesicTethersString = GET_MEMBER_NAME_STRING_CHECKED(FChaosClothAssetSimulationLongRangeAttachmentConfigNode, bUseGeodesicTethers);
-			UseGeodesicTethersString.RemoveFromStart(TEXT("b"), ESearchCase::CaseSensitive);  // Property collection names doesn't use the b prefix for booleans
-			const bool bUseGeodesicTethers = InProperties.GetValue<bool>(UseGeodesicTethersString);
-			const FName FixedEndWeightMap(InProperties.GetStringValue(FixedEndWeightMapString));
+			const bool bUseGeodesicTethers = InProperties.GetValue<bool>(UseGeodesicTethersName);
+			const FName FixedEndWeightMap(InProperties.GetStringValue(FixedEndWeightMapName));
 
 			UE::Chaos::ClothAsset::FClothEngineTools::GenerateTethers(OutClothCollection, FixedEndWeightMap, bUseGeodesicTethers);
 		}
 		PRAGMA_ENABLE_DEPRECATION_WARNINGS
 
-			// v2 (vertex set)
-			const FString FixedEndSetString = GET_MEMBER_NAME_STRING_CHECKED(FChaosClothAssetSimulationLongRangeAttachmentConfigNode_v2, FixedEndSet);
-		if (InProperties.GetKeyIndex(FixedEndSetString) != INDEX_NONE)
+		// v2 (vertex set)
+		const FName FixedEndSetName = GET_MEMBER_NAME_CHECKED(FChaosClothAssetSimulationLongRangeAttachmentConfigNode_v2, FixedEndSet);
+		if (InProperties.GetKeyNameIndex(FixedEndSetName) != INDEX_NONE)
 		{
 			// Regenerate using the Kinematic vertices.
-			FString UseGeodesicTethersString = GET_MEMBER_NAME_STRING_CHECKED(FChaosClothAssetSimulationLongRangeAttachmentConfigNode_v2, bUseGeodesicTethers);
-			UseGeodesicTethersString.RemoveFromStart(TEXT("b"), ESearchCase::CaseSensitive);  // Property collection names doesn't use the b prefix for booleans
-			const bool bUseGeodesicTethers = InProperties.GetValue<bool>(UseGeodesicTethersString);
+			const bool bUseGeodesicTethers = InProperties.GetValue<bool>(UseGeodesicTethersName);
 			UE::Chaos::ClothAsset::FClothEngineTools::GenerateTethersFromSelectionSet(OutClothCollection, KinematicVertices3DName, bUseGeodesicTethers);
 		}
 	}
diff --git a/Engine/Plugins/ChaosClothAssetEditor/Source/ChaosClothAssetDataflowNodes/Private/ChaosClothAsset/SimulationBaseConfigNode.cpp b/Engine/Plugins/ChaosClothAssetEditor/Source/ChaosClothAssetDataflowNodes/Private/ChaosClothAsset/SimulationBaseConfigNode.cpp
--- a/Engine/Plugins/ChaosClothAssetEditor/Source/ChaosClothAssetDataflowNodes/Private/ChaosClothAsset/SimulationBaseConfigNode.cpp
+++ b/Engine/Plugins/ChaosClothAssetEditor/Source/ChaosClothAssetDataflowNodes/Private/ChaosClothAsset/SimulationBaseConfigNode.cpp
@@ -204,10 +204,10 @@ int32 FChaosClothAssetSimulationBaseConfigNode::AddPropertyHelper(
 {
 	using namespace UE::Chaos::ClothAsset;
 
-	int32 KeyIndex = Properties.GetKeyIndex(PropertyName.ToString());
+	int32 KeyIndex = Properties.GetKeyNameIndex(PropertyName);
 	if (KeyIndex == INDEX_NONE)
 	{
-		KeyIndex = Properties.AddProperty(PropertyName.ToString());
+		KeyIndex = Properties.AddProperty(PropertyName);
 	}
 	else if (bWarnDuplicateProperty && !Properties.IsLegacy(KeyIndex))  // Only warns of duplicates when the property hasn't been set using a legacy Chaos config
 	{
@@ -221,7 +221,7 @@ int32 FChaosClothAssetSimulationBaseConfigNode::AddPropertyHelper(
 	// Check for similar properties
 	for (const FName& SimilarPropertyName : SimilarPropertyNames)
 	{
-		const int32 SimilarPropertyKeyIndex = Properties.GetKeyIndex(SimilarPropertyName.ToString());
+		const int32 SimilarPropertyKeyIndex = Properties.GetKeyNameIndex(SimilarPropertyName);
 		if (SimilarPropertyKeyIndex != INDEX_NONE && !Properties.IsLegacy(SimilarPropertyKeyIndex))  // Only warns of overrides when the property hasn't been set using a legacy Chaos config
 		{
 			Private::LogAndToastSimilarProperty(*this, PropertyName, SimilarPropertyName);
@@ -339,7 +339,7 @@ void FChaosClothAssetSimulationBaseConfigNode::FPropertyHelper::OverrideProperti
 {
 	for (const FName& PropertyName : PropertyNames)
 	{
-		const int32 PropertyKeyIndex = Properties.GetKeyIndex(PropertyName.ToString());
+		const int32 PropertyKeyIndex = Properties.GetKeyNameIndex(PropertyName);
 		if (PropertyKeyIndex != INDEX_NONE)
 		{
 			Properties.SetValue(PropertyKeyIndex, bPropertyValue);
@@ -355,7 +355,7 @@ void FChaosClothAssetSimulationBaseConfigNode::FPropertyHelper::OverrideProperti
 	}
 	for (const FName& PropertyName : PropertyNames)
 	{
-		const int32 PropertyKeyIndex = Properties.GetKeyIndex(PropertyName.ToString());
+		const int32 PropertyKeyIndex = Properties.GetKeyNameIndex(PropertyName);
 		if (PropertyKeyIndex != INDEX_NONE)
 		{
 			if (OverrideType == EChaosClothAssetConstraintOverrideType::Override)
@@ -378,7 +378,7 @@ void FChaosClothAssetSimulationBaseConfigNode::FPropertyHelper::OverrideProperti
 	}
 	for (const FName& PropertyName : PropertyNames)
 	{
-		const int32 PropertyKeyIndex = Properties.GetKeyIndex(PropertyName.ToString());
+		const int32 PropertyKeyIndex = Properties.GetKeyNameIndex(PropertyName);
 		if (PropertyKeyIndex != INDEX_NONE)
 		{
 			if (OverrideType == EChaosClothAssetConstraintOverrideType::Override)
diff --git a/Engine/Plugins/ChaosClothAssetEditor/Source/ChaosClothAssetDataflowNodes/Private/ChaosClothAsset/SimulationDefaultConfigNode.cpp b/Engine/Plugins/ChaosClothAssetEditor/Source/ChaosClothAssetDataflowNodes/Private/ChaosClothAsset/SimulationDefaultConfigNode.cpp
--- a/Engine/Plugins/ChaosClothAssetEditor/Source/ChaosClothAssetDataflowNodes/Private/ChaosClothAsset/SimulationDefaultConfigNode.cpp
+++ b/Engine/Plugins/ChaosClothAssetEditor/Source/ChaosClothAssetDataflowNodes/Private/ChaosClothAsset/SimulationDefaultConfigNode.cpp
@@ -5,9 +5,11 @@
 #include "ChaosClothAsset/ClothEngineTools.h"
 #include "ChaosClothAsset/CollectionClothFacade.h"
 #include "ChaosCloth/ChaosClothConfig.h"
+#include "ChaosCloth/ChaosClothingSimulationCloth.h"
 #include "ChaosCloth/ChaosClothingSimulationConfig.h"
 #include "Chaos/CollectionPropertyFacade.h"
-#include "Chaos/PBDLongRangeConstraints.h"  // For Tether modes
+#include "Chaos/PBDLongRangeConstraints.h"
+#include "Chaos/PBDSphericalConstraint.h"
 #include "Dataflow/DataflowInputOutput.h"
 #include "UObject/FortniteMainBranchObjectVersion.h"
 
@@ -67,10 +69,10 @@ void FChaosClothAssetSimulationDefaultConfigNode::Evaluate(UE::Dataflow::FContex
 				// Generate tethers
 				const FCollectionPropertyConstFacade& Properties = ClothingSimulationConfig.GetProperties(0);
 				constexpr bool bUseGeodesicTethersDefault = true;
-				const bool bUseGeodesicTethers = Properties.GetValue<bool>(TEXT("UseGeodesicTethers"), bUseGeodesicTethersDefault);
+				const bool bUseGeodesicTethers = Properties.GetValue<bool>(FClothingSimulationCloth::UseGeodesicTethersName, bUseGeodesicTethersDefault);
 				// Use the "MaxDistance" weight map to generate tethers. This follows legacy behavior.
-				static const FName MaxDistanceName(TEXT("MaxDistance"));
-				const FVector2f MaxDistanceValue = Properties.GetWeightedFloatValue(TEXT("MaxDistance"), FVector2f(0.f, 1.f));
+				const FName& MaxDistanceName = FPBDSphericalConstraint::MaxDistanceName;
+				const FVector2f MaxDistanceValue = Properties.GetWeightedFloatValue(MaxDistanceName, FVector2f(0.f, 1.f));
 
 				UE::Chaos::ClothAsset::FClothEngineTools::GenerateTethers(ClothCollection, MaxDistanceName, bUseGeodesicTethers, MaxDistanceValue);
 			}
diff --git a/Engine/Plugins/ChaosClothAssetEditor/Source/ChaosClothAssetDataflowNodes/Private/ChaosClothAsset/TerminalNode.cpp b/Engine/Plugins/ChaosClothAssetEditor/Source/ChaosClothAssetDataflowNodes/Private/ChaosClothAsset/TerminalNode.cpp
--- a/Engine/Plugins/ChaosClothAssetEditor/Source/ChaosClothAssetDataflowNodes/Private/ChaosClothAsset/TerminalNode.cpp
+++ b/Engine/Plugins/ChaosClothAssetEditor/Source/ChaosClothAssetDataflowNodes/Private/ChaosClothAsset/TerminalNode.cpp
@@ -83,7 +83,7 @@ namespace UE::Chaos::ClothAsset::Private
 			}
 			for (int32 PropertyIndex = 0; PropertyIndex < Property0.Num(); ++PropertyIndex)
 			{
-				if (Property0.GetKey(PropertyIndex) != Property1.GetKey(PropertyIndex))
+				if (Property0.GetKeyName(PropertyIndex) != Property1.GetKeyName(PropertyIndex))
 				{
 					return false;
 				}
diff --git a/Engine/Plugins/Experimental/Dataflow/Source/DataflowEditor/Private/DataflowEditorTools/DataflowEditorCorrectSkinWeightsNode.cpp b/Engine/Plugins/Experimental/Dataflow/Source/DataflowEditor/Private/DataflowEditorTools/DataflowEditorCorrectSkinWeightsNode.cpp
--- a/Engine/Plugins/Experimental/Dataflow/Source/DataflowEditor/Private/DataflowEditorTools/DataflowEditorCorrectSkinWeightsNode.cpp
+++ b/Engine/Plugins/Experimental/Dataflow/Source/DataflowEditor/Private/DataflowEditorTools/DataflowEditorCorrectSkinWeightsNode.cpp
@@ -801,23 +801,23 @@ void FDataflowSetSkinningSelectionNode::Evaluate(UE::Dataflow::FContext& Context
 
 		if(CorrectionType == ESkinWeightsCorrectionType::Prune)
 		{
-			PropertyFacade.AddStringValue(SelectionMapKeyValue.Group + FString("::") + FDataflowCorrectSkinWeightsNode::PruneSkinWeightsSelectionName.ToString(), SelectionMapKeyValue.Attribute);
+			PropertyFacade.AddStringValue(FName(SelectionMapKeyValue.Group + FString("::") + FDataflowCorrectSkinWeightsNode::PruneSkinWeightsSelectionName.ToString()), SelectionMapKeyValue.Attribute);
 		}
 		else if(CorrectionType == ESkinWeightsCorrectionType::Relax)
 		{
-			PropertyFacade.AddStringValue(SelectionMapKeyValue.Group + FString("::") + FDataflowCorrectSkinWeightsNode::RelaxSkinWeightsSelectionName.ToString(), SelectionMapKeyValue.Attribute);
+			PropertyFacade.AddStringValue(FName(SelectionMapKeyValue.Group + FString("::") + FDataflowCorrectSkinWeightsNode::RelaxSkinWeightsSelectionName.ToString()), SelectionMapKeyValue.Attribute);
 		}
 		else if(CorrectionType == ESkinWeightsCorrectionType::Hammer)
 		{
-			PropertyFacade.AddStringValue(SelectionMapKeyValue.Group + FString("::") + FDataflowCorrectSkinWeightsNode::HammerSkinWeightsSelectionName.ToString(), SelectionMapKeyValue.Attribute);
+			PropertyFacade.AddStringValue(FName(SelectionMapKeyValue.Group + FString("::") + FDataflowCorrectSkinWeightsNode::HammerSkinWeightsSelectionName.ToString()), SelectionMapKeyValue.Attribute);
 		}
 		else if(CorrectionType == ESkinWeightsCorrectionType::Clamp)
 		{
-			PropertyFacade.AddStringValue(SelectionMapKeyValue.Group + FString("::") + FDataflowCorrectSkinWeightsNode::ClampSkinWeightsSelectionName.ToString(), SelectionMapKeyValue.Attribute);
+			PropertyFacade.AddStringValue(FName(SelectionMapKeyValue.Group + FString("::") + FDataflowCorrectSkinWeightsNode::ClampSkinWeightsSelectionName.ToString()), SelectionMapKeyValue.Attribute);
 		}
 		else if(CorrectionType == ESkinWeightsCorrectionType::Normalize)
 		{
-			PropertyFacade.AddStringValue(SelectionMapKeyValue.Group + FString("::") + FDataflowCorrectSkinWeightsNode::NormalizeSkinWeightsSelectionName.ToString(), SelectionMapKeyValue.Attribute);
+			PropertyFacade.AddStringValue(FName(SelectionMapKeyValue.Group + FString("::") + FDataflowCorrectSkinWeightsNode::NormalizeSkinWeightsSelectionName.ToString()), SelectionMapKeyValue.Attribute);
 		}
 
 		SetValue(Context, MoveTemp(*OutCollection), &Collection);
@@ -872,23 +872,23 @@ void FDataflowGetSkinningSelectionNode::Evaluate(UE::Dataflow::FContext& Context
 
 			if(CorrectionType == ESkinWeightsCorrectionType::Prune)
 			{
-				SelectionMapKeyValue.Attribute = PropertyFacade.GetStringValue(SelectionMapKeyValue.Group + FString("::") + FDataflowCorrectSkinWeightsNode::PruneSkinWeightsSelectionName.ToString());
+				SelectionMapKeyValue.Attribute = PropertyFacade.GetStringValue(FName(SelectionMapKeyValue.Group + FString("::") + FDataflowCorrectSkinWeightsNode::PruneSkinWeightsSelectionName.ToString()));
 			}
 			else if(CorrectionType == ESkinWeightsCorrectionType::Relax)
 			{
-				SelectionMapKeyValue.Attribute = PropertyFacade.GetStringValue(SelectionMapKeyValue.Group + FString("::") + FDataflowCorrectSkinWeightsNode::RelaxSkinWeightsSelectionName.ToString());
+				SelectionMapKeyValue.Attribute = PropertyFacade.GetStringValue(FName(SelectionMapKeyValue.Group + FString("::") + FDataflowCorrectSkinWeightsNode::RelaxSkinWeightsSelectionName.ToString()));
 			}
 			else if(CorrectionType == ESkinWeightsCorrectionType::Hammer)
 			{
-				SelectionMapKeyValue.Attribute = PropertyFacade.GetStringValue(SelectionMapKeyValue.Group + FString("::") + FDataflowCorrectSkinWeightsNode::HammerSkinWeightsSelectionName.ToString());
+				SelectionMapKeyValue.Attribute = PropertyFacade.GetStringValue(FName(SelectionMapKeyValue.Group + FString("::") + FDataflowCorrectSkinWeightsNode::HammerSkinWeightsSelectionName.ToString()));
 			}
 			else if(CorrectionType == ESkinWeightsCorrectionType::Clamp)
 			{
-				SelectionMapKeyValue.Attribute = PropertyFacade.GetStringValue(SelectionMapKeyValue.Group + FString("::") + FDataflowCorrectSkinWeightsNode::ClampSkinWeightsSelectionName.ToString());
+				SelectionMapKeyValue.Attribute = PropertyFacade.GetStringValue(FName(SelectionMapKeyValue.Group + FString("::") + FDataflowCorrectSkinWeightsNode::ClampSkinWeightsSelectionName.ToString()));
 			}
 			else if(CorrectionType == ESkinWeightsCorrectionType::Normalize)
 			{
-				SelectionMapKeyValue.Attribute = PropertyFacade.GetStringValue(SelectionMapKeyValue.Group + FString("::") + FDataflowCorrectSkinWeightsNode::NormalizeSkinWeightsSelectionName.ToString());
+				SelectionMapKeyValue.Attribute = PropertyFacade.GetStringValue(FName(SelectionMapKeyValue.Group + FString("::") + FDataflowCorrectSkinWeightsNode::NormalizeSkinWeightsSelectionName.ToString()));
 			}
 			SetValue(Context, MoveTemp(SelectionMapKeyValue), &SelectionMapKey);
 		}
diff --git a/Engine/Source/Runtime/Experimental/Chaos/Private/Chaos/CollectionPropertyFacade.cpp b/Engine/Source/Runtime/Experimental/Chaos/Private/Chaos/CollectionPropertyFacade.cpp
--- a/Engine/Source/Runtime/Experimental/Chaos/Private/Chaos/CollectionPropertyFacade.cpp
+++ b/Engine/Source/Runtime/Experimental/Chaos/Private/Chaos/CollectionPropertyFacade.cpp
@@ -9,13 +9,20 @@ namespace Chaos::Softs
 	{
 		// Attribute groups, predefined data member of the collection
 		static const FName PropertyGroup("Property");
-		static const FName KeyName("Key");  // Property key, name to look for
+		static const FName KeyNameName("KeyName");  // Property key, name to look for
 		static const FName LowValueName("LowValue");  // Boolean, 24 bit integer (max 16777215), float, or vector value, or value of the lowest weight on the weight map if any
 		static const FName HighValueName("HighValue");  // Boolean, 24 bit integer (max 16777215), float, or vector value of the highest weight on the weight map if any
 		static const FName StringValueName("StringValue");  // String value, or weight map name, ...etc.
 		static const FName FlagsName("Flags");  // Whether this property is enabled, animatable, ...etc.
+
+		// Previous name for String-based Key
+		static const FName KeyName("Key");
 	}
 
+	PRAGMA_DISABLE_DEPRECATION_WARNINGS
+	FCollectionPropertyConstFacade::~FCollectionPropertyConstFacade() = default;
+	PRAGMA_ENABLE_DEPRECATION_WARNINGS
+
 	FCollectionPropertyConstFacade::FCollectionPropertyConstFacade(const TSharedPtr<const FManagedArrayCollection>& InManagedArrayCollection)
 		: ManagedArrayCollection(InManagedArrayCollection)
 	{
@@ -31,7 +38,7 @@ namespace Chaos::Softs
 	bool FCollectionPropertyConstFacade::IsValid() const
 	{
 		return
-			ManagedArrayCollection->HasAttribute(PropertyFacadeNames::KeyName, PropertyFacadeNames::PropertyGroup) &&
+			ManagedArrayCollection->HasAttribute(PropertyFacadeNames::KeyNameName, PropertyFacadeNames::PropertyGroup) &&
 			ManagedArrayCollection->HasAttribute(PropertyFacadeNames::LowValueName, PropertyFacadeNames::PropertyGroup) &&
 			ManagedArrayCollection->HasAttribute(PropertyFacadeNames::HighValueName, PropertyFacadeNames::PropertyGroup) &&
 			ManagedArrayCollection->HasAttribute(PropertyFacadeNames::StringValueName, PropertyFacadeNames::PropertyGroup) &&
@@ -40,7 +47,7 @@ namespace Chaos::Softs
 
 	void FCollectionPropertyConstFacade::UpdateArrays()
 	{
-		KeyArray = GetArray<FString>(PropertyFacadeNames::KeyName);
+		KeyNameArray = GetArray<FName>(PropertyFacadeNames::KeyNameName);
 		LowValueArray = GetArray<FVector3f>(PropertyFacadeNames::LowValueName);
 		HighValueArray = GetArray<FVector3f>(PropertyFacadeNames::HighValueName);
 		StringValueArray = GetArray<FString>(PropertyFacadeNames::StringValueName);
@@ -51,11 +58,11 @@ namespace Chaos::Softs
 	{
 		TRACE_CPUPROFILER_EVENT_SCOPE(FCollectionPropertyConstFacade_RebuildKeyIndices);
 		// Create a fast access search map (although it might only be faster for a large enough number of properties)
-		const int32 NumKeys = KeyArray.Num();
-		KeyIndices.Empty(NumKeys);
+		const int32 NumKeys = KeyNameArray.Num();
+		KeyNameIndices.Empty(NumKeys);
 		for (int32 Index = 0; Index < NumKeys; ++Index)
 		{
-			KeyIndices.Emplace(KeyArray[Index], Index);
+			KeyNameIndices.Emplace(KeyNameArray[Index], Index);
 		}
 	}
 
@@ -155,8 +162,8 @@ namespace Chaos::Softs
 			const int32 NumInKeys = InPropertyFacade.Num();
 			for (int32 InKeyIndex = 0; InKeyIndex < NumInKeys; ++InKeyIndex)
 			{
-				const FString& PropertyName = InPropertyFacade.GetKey(InKeyIndex);
-				const int32 PropertyIndex = GetKeyIndex(PropertyName);
+				const FName& PropertyName = InPropertyFacade.GetKeyName(InKeyIndex);
+				const int32 PropertyIndex = GetKeyNameIndex(PropertyName);
 				if (PropertyIndex == INDEX_NONE)
 				{
 					continue;
@@ -180,7 +187,7 @@ namespace Chaos::Softs
 	void FCollectionPropertyMutableFacade::DefineSchema()
 	{
 		TRACE_CPUPROFILER_EVENT_SCOPE(FCollectionPropertyMutableFacade_DefineSchema);
-		GetManagedArrayCollection()->AddAttribute<FString>(PropertyFacadeNames::KeyName, PropertyFacadeNames::PropertyGroup);
+		GetManagedArrayCollection()->AddAttribute<FName>(PropertyFacadeNames::KeyNameName, PropertyFacadeNames::PropertyGroup);
 		GetManagedArrayCollection()->AddAttribute<FVector3f>(PropertyFacadeNames::LowValueName, PropertyFacadeNames::PropertyGroup);
 		GetManagedArrayCollection()->AddAttribute<FVector3f>(PropertyFacadeNames::HighValueName, PropertyFacadeNames::PropertyGroup);
 		GetManagedArrayCollection()->AddAttribute<FString>(PropertyFacadeNames::StringValueName, PropertyFacadeNames::PropertyGroup);
@@ -190,7 +197,7 @@ namespace Chaos::Softs
 		RebuildKeyIndices();
 	}
 
-	int32 FCollectionPropertyMutableFacade::AddProperty(const FString& Key, bool bEnabled, bool bAnimatable, bool bIntrinsic)
+	int32 FCollectionPropertyMutableFacade::AddProperty(const FName& Key, bool bEnabled, bool bAnimatable, bool bIntrinsic)
 	{
 		const ECollectionPropertyFlags Flags =
 			(bEnabled ? ECollectionPropertyFlags::Enabled : ECollectionPropertyFlags::None) |
@@ -199,7 +206,7 @@ namespace Chaos::Softs
 		return AddProperty(Key, Flags);
 	}
 
-	int32 FCollectionPropertyMutableFacade::AddProperty(const FString& Key, ECollectionPropertyFlags Flags)
+	int32 FCollectionPropertyMutableFacade::AddProperty(const FName& Key, ECollectionPropertyFlags Flags)
 	{
 		const int32 Index = GetManagedArrayCollection()->AddElements(1, PropertyFacadeNames::PropertyGroup);
 
@@ -207,17 +214,17 @@ namespace Chaos::Softs
 		UpdateArrays();
 
 		// Setup the new element's default value and enable the property by default
-		GetKeyArray()[Index] = Key;
+		GetKeyNameArray()[Index] = Key;
 		GetLowValueArray()[Index] = GetHighValueArray()[Index] = FVector3f::ZeroVector;
 		GetFlagsArray()[Index] = Flags;
 
 		// Update search map
-		KeyIndices.Emplace(KeyArray[Index], Index);
+		KeyNameIndices.Emplace(KeyNameArray[Index], Index);
 
 		return Index;
 	}
 
-	int32 FCollectionPropertyMutableFacade::AddProperties(const TArray<FString>& Keys, bool bEnabled, bool bAnimatable, bool bIntrinsic)
+	int32 FCollectionPropertyMutableFacade::AddProperties(const TArray<FName>& Keys, bool bEnabled, bool bAnimatable, bool bIntrinsic)
 	{
 		const ECollectionPropertyFlags Flags =
 			(bEnabled ? ECollectionPropertyFlags::Enabled : ECollectionPropertyFlags::None) |
@@ -227,7 +234,7 @@ namespace Chaos::Softs
 		return AddProperties(Keys, Flags);
 	}
 
-	int32 FCollectionPropertyMutableFacade::AddProperties(const TArray<FString>& Keys, ECollectionPropertyFlags Flags)
+	int32 FCollectionPropertyMutableFacade::AddProperties(const TArray<FName>& Keys, ECollectionPropertyFlags Flags)
 	{
 		if (const int32 NumProperties = Keys.Num())
 		{
@@ -239,12 +246,12 @@ namespace Chaos::Softs
 			for (int32 Index = StartIndex; Index < NumProperties + StartIndex; ++Index)
 			{
 				// Setup the new elements' default value and enable the property by default
-				GetKeyArray()[Index] = Keys[Index - StartIndex];
+				GetKeyNameArray()[Index] = Keys[Index - StartIndex];
 				GetLowValueArray()[Index] = GetHighValueArray()[Index] = FVector3f::ZeroVector;
 				GetFlagsArray()[Index] = Flags;
 
 				// Update search map
-				KeyIndices.Emplace(KeyArray[Index], Index);
+				KeyNameIndices.Emplace(KeyNameArray[Index], Index);
 			}
 
 			return StartIndex;
@@ -261,7 +268,7 @@ namespace Chaos::Softs
 	{
 		TRACE_CPUPROFILER_EVENT_SCOPE(FCollectionPropertyMutableFacade_Copy);
 
-		if (InManagedArrayCollection.HasAttribute(PropertyFacadeNames::KeyName, PropertyFacadeNames::PropertyGroup) &&
+		if (InManagedArrayCollection.HasAttribute(PropertyFacadeNames::KeyNameName, PropertyFacadeNames::PropertyGroup) &&
 			InManagedArrayCollection.HasAttribute(PropertyFacadeNames::LowValueName, PropertyFacadeNames::PropertyGroup) &&
 			InManagedArrayCollection.HasAttribute(PropertyFacadeNames::HighValueName, PropertyFacadeNames::PropertyGroup) &&
 			InManagedArrayCollection.HasAttribute(PropertyFacadeNames::StringValueName, PropertyFacadeNames::PropertyGroup) &&
@@ -269,15 +276,15 @@ namespace Chaos::Softs
 		{
 			// Compare keys
 			bool bKeyIndicesMatch = false;
-			const TManagedArray<FString>* InKeyManagedArray = InManagedArrayCollection.FindAttributeTyped<FString>(PropertyFacadeNames::KeyName, PropertyFacadeNames::PropertyGroup);
-			const TManagedArray<FString>* OrigKeyManagedArray = GetManagedArrayCollection()->FindAttributeTyped<FString>(PropertyFacadeNames::KeyName, PropertyFacadeNames::PropertyGroup);
+			const TManagedArray<FName>* InKeyManagedArray = InManagedArrayCollection.FindAttributeTyped<FName>(PropertyFacadeNames::KeyNameName, PropertyFacadeNames::PropertyGroup);
+			const TManagedArray<FName>* OrigKeyManagedArray = GetManagedArrayCollection()->FindAttributeTyped<FName>(PropertyFacadeNames::KeyNameName, PropertyFacadeNames::PropertyGroup);
 			if (InKeyManagedArray && OrigKeyManagedArray && InKeyManagedArray->Num() == OrigKeyManagedArray->Num())
 			{
 				TRACE_CPUPROFILER_EVENT_SCOPE(FCollectionPropertyMutableFacade_Copy_CompareKeys);
 				bKeyIndicesMatch = true;
 				for (int32 Index = 0; Index < OrigKeyManagedArray->Num(); ++Index)
 				{
-					if (!(*InKeyManagedArray)[Index].Equals((*OrigKeyManagedArray)[Index], ESearchCase::CaseSensitive))
+					if ((*InKeyManagedArray)[Index] != (*OrigKeyManagedArray)[Index])
 					{
 						bKeyIndicesMatch = false;
 						break;
@@ -286,7 +293,7 @@ namespace Chaos::Softs
 			}
 			if (!bKeyIndicesMatch)
 			{
-				GetManagedArrayCollection()->CopyAttribute(InManagedArrayCollection, PropertyFacadeNames::KeyName, PropertyFacadeNames::PropertyGroup);
+				GetManagedArrayCollection()->CopyAttribute(InManagedArrayCollection, PropertyFacadeNames::KeyNameName, PropertyFacadeNames::PropertyGroup);
 			}
 			GetManagedArrayCollection()->CopyAttribute(InManagedArrayCollection, PropertyFacadeNames::LowValueName, PropertyFacadeNames::PropertyGroup);
 			GetManagedArrayCollection()->CopyAttribute(InManagedArrayCollection, PropertyFacadeNames::HighValueName, PropertyFacadeNames::PropertyGroup);
@@ -333,9 +340,9 @@ namespace Chaos::Softs
 				const int32 NumInKeys = InPropertyFacade.Num();
 				for (int32 InKeyIndex = 0; InKeyIndex < NumInKeys; ++InKeyIndex)
 				{
-					const FString& PropertyName = InPropertyFacade.GetKey(InKeyIndex);
+					const FName& PropertyName = InPropertyFacade.GetKeyName(InKeyIndex);
 					const ECollectionPropertyFlags PropertyFlags = InPropertyFacade.GetFlags(InKeyIndex);
-					int32 NewPropertyIndex = GetKeyIndex(PropertyName);
+					int32 NewPropertyIndex = GetKeyNameIndex(PropertyName);
 					if (NewPropertyIndex == INDEX_NONE)
 					{
 						if (bAppendNewProperties)
@@ -367,7 +374,7 @@ namespace Chaos::Softs
 				TArray<int32> KeyIndicesToRemove;
 				for (int32 ExistingKeyIndex = 0; ExistingKeyIndex < PrevNumKeys; ++ExistingKeyIndex)
 				{
-					if (InPropertyFacade.GetKeyIndex(GetKey(ExistingKeyIndex)) == INDEX_NONE)
+					if (InPropertyFacade.GetKeyNameIndex(GetKeyName(ExistingKeyIndex)) == INDEX_NONE)
 					{
 						if (bRemoveMissingProperties)
 						{
@@ -390,31 +397,55 @@ namespace Chaos::Softs
 		}
 	}
 
-	int32 FCollectionPropertyMutableFacade::AddWeightedFloatValue(const FString& Key, const FVector2f& Value, bool bEnabled, bool bAnimatable, bool bIntrinsic)
+	int32 FCollectionPropertyMutableFacade::AddWeightedFloatValue(const FName& Key, const FVector2f& Value, bool bEnabled, bool bAnimatable, bool bIntrinsic)
 	{
 		const int32 KeyIndex = AddProperty(Key, bEnabled, bAnimatable, bIntrinsic);
 		SetWeightedFloatValue(KeyIndex, Value);
 		return KeyIndex;
 	}
 
-	int32 FCollectionPropertyMutableFacade::AddWeightedFloatValue(const FString& Key, const FVector2f& Value, ECollectionPropertyFlags Flags)
+	int32 FCollectionPropertyMutableFacade::AddWeightedFloatValue(const FName& Key, const FVector2f& Value, ECollectionPropertyFlags Flags)
 	{
 		const int32 KeyIndex = AddProperty(Key, Flags);
 		SetWeightedFloatValue(KeyIndex, Value);
 		return KeyIndex;
 	}
 
-	int32 FCollectionPropertyMutableFacade::AddStringValue(const FString& Key, const FString& Value, bool bEnabled, bool bAnimatable, bool bIntrinsic)
+	int32 FCollectionPropertyMutableFacade::AddStringValue(const FName& Key, const FString& Value, bool bEnabled, bool bAnimatable, bool bIntrinsic)
 	{
 		const int32 KeyIndex = AddProperty(Key, bEnabled, bAnimatable, bIntrinsic);
 		SetStringValue(Key, Value);
 		return KeyIndex;
 	}
 
-	int32 FCollectionPropertyMutableFacade::AddStringValue(const FString& Key, const FString& Value, ECollectionPropertyFlags Flags)
+	int32 FCollectionPropertyMutableFacade::AddStringValue(const FName& Key, const FString& Value, ECollectionPropertyFlags Flags)
 	{
 		const int32 KeyIndex = AddProperty(Key, Flags);
 		SetStringValue(Key, Value);
 		return KeyIndex;
 	}
+
+	void FCollectionPropertyMutableFacade::PostSerialize(const FArchive& Ar)
+	{
+		if (Ar.IsLoading())
+		{
+			if (const TManagedArray<FString>* const InKeyArray = ManagedArrayCollection->FindAttributeTyped<FString>(PropertyFacadeNames::KeyName, PropertyFacadeNames::PropertyGroup))
+			{
+				if (KeyNameArray.IsEmpty())
+				{
+					// This is an old property facade using string rather than name-based keys. Transfer strings over to names.
+					GetManagedArrayCollection()->AddAttribute<FName>(PropertyFacadeNames::KeyNameName, PropertyFacadeNames::PropertyGroup);
+					UpdateArrays();
+					check(InKeyArray->Num() == KeyNameArray.Num());
+					for (int32 Index = 0; Index < KeyNameArray.Num(); ++Index)
+					{
+						GetKeyNameArray()[Index] = FName((*InKeyArray)[Index]);
+					}
+					RebuildKeyIndices();
+				}
+
+				GetManagedArrayCollection()->RemoveAttribute(PropertyFacadeNames::KeyName, PropertyFacadeNames::PropertyGroup);
+			}
+		}
+	}
 }  // End namespace Chaos::Softs
diff --git a/Engine/Source/Runtime/Experimental/Chaos/Public/Chaos/CollectionPropertyFacade.h b/Engine/Source/Runtime/Experimental/Chaos/Public/Chaos/CollectionPropertyFacade.h
--- a/Engine/Source/Runtime/Experimental/Chaos/Public/Chaos/CollectionPropertyFacade.h
+++ b/Engine/Source/Runtime/Experimental/Chaos/Public/Chaos/CollectionPropertyFacade.h
@@ -32,6 +32,14 @@ namespace Chaos::Softs
 	template<> struct TIsWeightedType<float> { static constexpr bool Value = true; };
 	template<> struct TIsWeightedType<FVector3f> { static constexpr bool Value = true; };
 
+	/** Helper to disambiguate new FName based keys from deprecated FString-based keys.*/
+	template<typename T> struct TIsStringAndNameConvertibleType { static constexpr bool Value = false; };
+	template<> struct TIsStringAndNameConvertibleType<FString> { static constexpr bool Value = true; };
+	template<> struct TIsStringAndNameConvertibleType<const FString&> { static constexpr bool Value = true; };
+	template<> struct TIsStringAndNameConvertibleType<const WIDECHAR*> { static constexpr bool Value = true; };
+	template<> struct TIsStringAndNameConvertibleType<const ANSICHAR*> { static constexpr bool Value = true; };
+	template<> struct TIsStringAndNameConvertibleType<const UTF8CHAR*> { static constexpr bool Value = true; };
+
 	/**
 	 * Defines common API for reading simulation properties data and metadata.
 	 * This is mainly used for the cloth simulation properties to provide
@@ -41,7 +49,7 @@ namespace Chaos::Softs
 	{
 	public:
 		CHAOS_API explicit FCollectionPropertyConstFacade(const TSharedPtr<const FManagedArrayCollection>& InManagedArrayCollection);
-		virtual ~FCollectionPropertyConstFacade() = default;
+		CHAOS_API virtual ~FCollectionPropertyConstFacade();
 
 		FCollectionPropertyConstFacade() = delete;
 		FCollectionPropertyConstFacade(const FCollectionPropertyConstFacade&) = delete;
@@ -54,14 +62,18 @@ namespace Chaos::Softs
 		CHAOS_API bool IsValid() const;
 
 		/** Return the number of properties in this collection. */
-		int32 Num() const { return KeyIndices.Num(); }
+		int32 Num() const { return KeyNameIndices.Num(); }
 
 		/** Return the property index for the specified key if it exists, or INDEX_NONE otherwise. */
-		int32 GetKeyIndex(const FString& Key) const { const int32* const Index = KeyIndices.Find(Key); return Index ? *Index : INDEX_NONE; }
+		int32 GetKeyNameIndex(const FName& Key) const { const int32* const Index = KeyNameIndices.Find(Key); return Index ? *Index : INDEX_NONE; }
+		UE_DEPRECATED(5.7, "Use GetKeyNameIndex instead")
+		int32 GetKeyIndex(const FString& Key) const { const int32* const Index = KeyNameIndices.Find(FName(Key)); return Index ? *Index : INDEX_NONE; }
 
 		/** Return the property name (key) for the specified index. */
 		//~ Values access per index, fast, no check, index must be valid (0 <= KeyIndex < Num())
-		const FString& GetKey(int32 KeyIndex) const { return GetValue<const FString&>(KeyIndex, KeyArray); }
+		const FName& GetKeyName(int32 KeyIndex) const { return GetValue<const FName&>(KeyIndex, KeyNameArray); }
+		UE_DEPRECATED(5.7, "Use GetKeyName instead")
+		FString GetKey(int32 KeyIndex) const { return GetValue<const FName&>(KeyIndex, KeyNameArray).ToString(); }
 
 		template<typename T UE_REQUIRES(TIsWeightedType<T>::Value)>
 		T GetLowValue(int32 KeyIndex) const { return GetValue<T>(KeyIndex, LowValueArray); }
@@ -91,90 +103,196 @@ namespace Chaos::Softs
 
 		//~ Values access per key
 		template<typename T UE_REQUIRES(TIsWeightedType<T>::Value)>
-		T GetLowValue(const FString& Key, const T& Default = T(0), int32* OutKeyIndex = nullptr) const
+		T GetLowValue(const FName& Key, const T& Default = T(0), int32* OutKeyIndex = nullptr) const
 		{
 			return SafeGet(Key, [this](int32 KeyIndex)->T { return GetLowValue<T>(KeyIndex); }, Default, OutKeyIndex);
 		}
 
 		template<typename T UE_REQUIRES(TIsWeightedType<T>::Value)>
-		T GetHighValue(const FString& Key, const T& Default = T(0), int32* OutKeyIndex = nullptr) const
+		T GetHighValue(const FName& Key, const T& Default = T(0), int32* OutKeyIndex = nullptr) const
 		{
 			return SafeGet(Key, [this](int32 KeyIndex)->T { return GetHighValue<T>(KeyIndex); }, Default, OutKeyIndex);
 		}
 
 		template<typename T UE_REQUIRES(TIsWeightedType<T>::Value)>
-		TPair<T, T> GetWeightedValue(const FString& Key, const T& Default = T(0), int32* OutKeyIndex = nullptr) const
+		TPair<T, T> GetWeightedValue(const FName& Key, const T& Default = T(0), int32* OutKeyIndex = nullptr) const
 		{
 			return SafeGet(Key, [this](int32 KeyIndex)->TPair<T, T> { return GetWeightedValue<T>(KeyIndex); }, MakeTuple(Default, Default), OutKeyIndex);
 		}
 
-		FVector2f GetWeightedFloatValue(const FString& Key, const float& Default = 0.f, int32* OutKeyIndex = nullptr) const
+		FVector2f GetWeightedFloatValue(const FName& Key, const float& Default = 0.f, int32* OutKeyIndex = nullptr) const
 		{
 			return SafeGet(Key, [this](int32 KeyIndex)->FVector2f { return GetWeightedFloatValue(KeyIndex); }, FVector2f(Default), OutKeyIndex);
 		}
 
-		FVector2f GetWeightedFloatValue(const FString& Key, const FVector2f& Default, int32* OutKeyIndex = nullptr) const
+		FVector2f GetWeightedFloatValue(const FName& Key, const FVector2f& Default, int32* OutKeyIndex = nullptr) const
 		{
 			return SafeGet(Key, [this](int32 KeyIndex)->FVector2f { return GetWeightedFloatValue(KeyIndex); }, Default, OutKeyIndex);
 		}
 
 		template<typename T UE_REQUIRES(TIsWeightedType<T>::Value)>
-		T GetValue(const FString& Key, const T& Default = T(0), int32* OutKeyIndex = nullptr) const
+		T GetValue(const FName& Key, const T& Default = T(0), int32* OutKeyIndex = nullptr) const
 		{
 			return SafeGet(Key, [this](int32 KeyIndex)->T { return GetValue<T>(KeyIndex); }, Default, OutKeyIndex);
 		}
 
-		FString GetStringValue(const FString& Key, const FString& Default = "", int32* OutKeyIndex = nullptr) const
+		FString GetStringValue(const FName& Key, const FString& Default = "", int32* OutKeyIndex = nullptr) const
 		{
 			return SafeGet(Key, [this](int32 KeyIndex)->FString { return GetStringValue(KeyIndex); }, Default, OutKeyIndex);
 		}
 
-		ECollectionPropertyFlags GetFlags(const FString& Key, uint8 Default = 0, int32* OutKeyIndex = nullptr) const
+		ECollectionPropertyFlags GetFlags(const FName& Key, uint8 Default = 0, int32* OutKeyIndex = nullptr) const
 		{
 			return SafeGet(Key, [this](int32 KeyIndex)->ECollectionPropertyFlags { return GetFlags(KeyIndex); }, (ECollectionPropertyFlags)Default, OutKeyIndex);
 		}
 
-		bool IsEnabled(const FString& Key, bool bDefault = false, int32* OutKeyIndex = nullptr) const
+		bool IsEnabled(const FName& Key, bool bDefault = false, int32* OutKeyIndex = nullptr) const
 		{
 			return SafeGet(Key, [this](int32 KeyIndex)->bool { return IsEnabled(KeyIndex); }, bDefault, OutKeyIndex);
 		}
 
-		bool IsAnimatable(const FString& Key, bool bDefault = false, int32* OutKeyIndex = nullptr) const
+		bool IsAnimatable(const FName& Key, bool bDefault = false, int32* OutKeyIndex = nullptr) const
 		{
 			return SafeGet(Key, [this](int32 KeyIndex)->bool { return IsAnimatable(KeyIndex); }, bDefault, OutKeyIndex);
 		}
 
-		bool IsLegacy(const FString& Key, bool bDefault = false, int32* OutKeyIndex = nullptr) const
+		bool IsLegacy(const FName& Key, bool bDefault = false, int32* OutKeyIndex = nullptr) const
 		{
 			return SafeGet(Key, [this](int32 KeyIndex)->bool { return IsLegacy(KeyIndex); }, bDefault, OutKeyIndex);
 		}
 
-		bool IsIntrinsic(const FString& Key, bool bDefault = false, int32* OutKeyIndex = nullptr) const
+		bool IsIntrinsic(const FName& Key, bool bDefault = false, int32* OutKeyIndex = nullptr) const
 		{
 			return SafeGet(Key, [this](int32 KeyIndex)->bool { return IsIntrinsic(KeyIndex); }, bDefault, OutKeyIndex);
 		}
 
-		bool IsStringDirty(const FString& Key, bool bDefault = false, int32* OutKeyIndex = nullptr) const
+		bool IsStringDirty(const FName& Key, bool bDefault = false, int32* OutKeyIndex = nullptr) const
 		{
 			return SafeGet(Key, [this](int32 KeyIndex)->bool { return IsStringDirty(KeyIndex); }, bDefault, OutKeyIndex);
 		}
 
-		bool IsDirty(const FString& Key, bool bDefault = false, int32* OutKeyIndex = nullptr) const
+		bool IsDirty(const FName& Key, bool bDefault = false, int32* OutKeyIndex = nullptr) const
 		{
 			return SafeGet(Key, [this](int32 KeyIndex)->bool { return IsDirty(KeyIndex); }, bDefault, OutKeyIndex);
 		}
 		
-		bool IsInterpolable(const FString& Key, bool bDefault = false, int32* OutKeyIndex = nullptr) const
+		bool IsInterpolable(const FName& Key, bool bDefault = false, int32* OutKeyIndex = nullptr) const
 		{
 			return SafeGet(Key, [this](int32 KeyIndex)->bool { return IsInterpolable(KeyIndex); }, bDefault, OutKeyIndex);
 		}
+		
+		//~ Deprecated string key versions: Values access per key
+		template<typename T, typename StringType UE_REQUIRES(TIsWeightedType<T>::Value && TIsStringAndNameConvertibleType<StringType>::Value)>
+		UE_DEPRECATED(5.7, "Use FName Key version of this method")
+		T GetLowValue(const StringType Key, const T& Default = T(0), int32* OutKeyIndex = nullptr) const
+		{
+			return SafeGet(FName(Key), [this](int32 KeyIndex)->T { return GetLowValue<T>(KeyIndex); }, Default, OutKeyIndex);
+		}
+
+		template<typename T, typename StringType UE_REQUIRES(TIsWeightedType<T>::Value && TIsStringAndNameConvertibleType<StringType>::Value)>
+		UE_DEPRECATED(5.7, "Use FName Key version of this method")
+		T GetHighValue(const StringType Key, const T& Default = T(0), int32* OutKeyIndex = nullptr) const
+		{
+			return SafeGet(FName(Key), [this](int32 KeyIndex)->T { return GetHighValue<T>(KeyIndex); }, Default, OutKeyIndex);
+		}
+
+		template<typename T, typename StringType UE_REQUIRES(TIsWeightedType<T>::Value && TIsStringAndNameConvertibleType<StringType>::Value)>
+		UE_DEPRECATED(5.7, "Use FName Key version of this method")
+		TPair<T, T> GetWeightedValue(const StringType Key, const T& Default = T(0), int32* OutKeyIndex = nullptr) const
+		{
+			return SafeGet(FName(Key), [this](int32 KeyIndex)->TPair<T, T> { return GetWeightedValue<T>(KeyIndex); }, MakeTuple(Default, Default), OutKeyIndex);
+		}
+
+		template<typename StringType UE_REQUIRES(TIsStringAndNameConvertibleType<StringType>::Value)>
+		UE_DEPRECATED(5.7, "Use FName Key version of this method")
+		FVector2f GetWeightedFloatValue(const StringType Key, const float& Default = 0.f, int32* OutKeyIndex = nullptr) const
+		{
+			return SafeGet(FName(Key), [this](int32 KeyIndex)->FVector2f { return GetWeightedFloatValue(KeyIndex); }, FVector2f(Default), OutKeyIndex);
+		}
+
+		template<typename StringType UE_REQUIRES(TIsStringAndNameConvertibleType<StringType>::Value)>
+		UE_DEPRECATED(5.7, "Use FName Key version of this method")
+		FVector2f GetWeightedFloatValue(const StringType Key, const FVector2f& Default, int32* OutKeyIndex = nullptr) const
+		{
+			return SafeGet(FName(Key), [this](int32 KeyIndex)->FVector2f { return GetWeightedFloatValue(KeyIndex); }, Default, OutKeyIndex);
+		}
+
+		template<typename T, typename StringType UE_REQUIRES(TIsWeightedType<T>::Value && TIsStringAndNameConvertibleType<StringType>::Value)>
+		UE_DEPRECATED(5.7, "Use FName Key version of this method")
+		T GetValue(const StringType Key, const T& Default = T(0), int32* OutKeyIndex = nullptr) const
+		{
+			return SafeGet(FName(Key), [this](int32 KeyIndex)->T { return GetValue<T>(KeyIndex); }, Default, OutKeyIndex);
+		}
+
+		template<typename StringType UE_REQUIRES(TIsStringAndNameConvertibleType<StringType>::Value)>
+		UE_DEPRECATED(5.7, "Use FName Key version of this method")
+		FString GetStringValue(const StringType Key, const FString& Default = "", int32* OutKeyIndex = nullptr) const
+		{
+			return SafeGet(FName(Key), [this](int32 KeyIndex)->FString { return GetStringValue(KeyIndex); }, Default, OutKeyIndex);
+		}
+
+		template<typename StringType UE_REQUIRES(TIsStringAndNameConvertibleType<StringType>::Value)>
+		UE_DEPRECATED(5.7, "Use FName Key version of this method")
+		ECollectionPropertyFlags GetFlags(const StringType Key, uint8 Default = 0, int32* OutKeyIndex = nullptr) const
+		{
+			return SafeGet(FName(Key), [this](int32 KeyIndex)->ECollectionPropertyFlags { return GetFlags(KeyIndex); }, (ECollectionPropertyFlags)Default, OutKeyIndex);
+		}
+
+		template<typename StringType UE_REQUIRES(TIsStringAndNameConvertibleType<StringType>::Value)>
+		UE_DEPRECATED(5.7, "Use FName Key version of this method")
+		bool IsEnabled(const StringType Key, bool bDefault = false, int32* OutKeyIndex = nullptr) const
+		{
+			return SafeGet(FName(Key), [this](int32 KeyIndex)->bool { return IsEnabled(KeyIndex); }, bDefault, OutKeyIndex);
+		}
+
+		template<typename StringType UE_REQUIRES(TIsStringAndNameConvertibleType<StringType>::Value)>
+		UE_DEPRECATED(5.7, "Use FName Key version of this method")
+		bool IsAnimatable(const StringType Key, bool bDefault = false, int32* OutKeyIndex = nullptr) const
+		{
+			return SafeGet(FName(Key), [this](int32 KeyIndex)->bool { return IsAnimatable(KeyIndex); }, bDefault, OutKeyIndex);
+		}
+
+		template<typename StringType UE_REQUIRES(TIsStringAndNameConvertibleType<StringType>::Value)>
+		UE_DEPRECATED(5.7, "Use FName Key version of this method")
+		bool IsLegacy(const StringType Key, bool bDefault = false, int32* OutKeyIndex = nullptr) const
+		{
+			return SafeGet(FName(Key), [this](int32 KeyIndex)->bool { return IsLegacy(KeyIndex); }, bDefault, OutKeyIndex);
+		}
+
+		template<typename StringType UE_REQUIRES(TIsStringAndNameConvertibleType<StringType>::Value)>
+		UE_DEPRECATED(5.7, "Use FName Key version of this method")
+		bool IsIntrinsic(const StringType Key, bool bDefault = false, int32* OutKeyIndex = nullptr) const
+		{
+			return SafeGet(FName(Key), [this](int32 KeyIndex)->bool { return IsIntrinsic(KeyIndex); }, bDefault, OutKeyIndex);
+		}
+
+		template<typename StringType UE_REQUIRES(TIsStringAndNameConvertibleType<StringType>::Value)>
+		UE_DEPRECATED(5.7, "Use FName Key version of this method")
+		bool IsStringDirty(const StringType Key, bool bDefault = false, int32* OutKeyIndex = nullptr) const
+		{
+			return SafeGet(FName(Key), [this](int32 KeyIndex)->bool { return IsStringDirty(KeyIndex); }, bDefault, OutKeyIndex);
+		}
+
+		template<typename StringType UE_REQUIRES(TIsStringAndNameConvertibleType<StringType>::Value)>
+		UE_DEPRECATED(5.7, "Use FName Key version of this method")
+		bool IsDirty(const StringType Key, bool bDefault = false, int32* OutKeyIndex = nullptr) const
+		{
+			return SafeGet(FName(Key), [this](int32 KeyIndex)->bool { return IsDirty(KeyIndex); }, bDefault, OutKeyIndex);
+		}
+
+		template<typename StringType UE_REQUIRES(TIsStringAndNameConvertibleType<StringType>::Value)>
+		UE_DEPRECATED(5.7, "Use FName Key version of this method")
+		bool IsInterpolable(const StringType Key, bool bDefault = false, int32* OutKeyIndex = nullptr) const
+		{
+			return SafeGet(FName(Key), [this](int32 KeyIndex)->bool { return IsInterpolable(KeyIndex); }, bDefault, OutKeyIndex);
+		}
 
 		friend ::uint32 GetTypeHash(const Chaos::Softs::FCollectionPropertyConstFacade& PropertyFacade)
 		{
 			uint32 Hash = 0;
 			if (PropertyFacade.IsValid())
 			{
-				Hash = GetArrayHash(PropertyFacade.KeyArray.GetData(), PropertyFacade.KeyArray.Num(), Hash);
+				Hash = GetArrayHash(PropertyFacade.KeyNameArray.GetData(), PropertyFacade.KeyNameArray.Num(), Hash);
 				Hash = GetArrayHash(PropertyFacade.LowValueArray.GetData(), PropertyFacade.LowValueArray.Num(), Hash);
 				Hash = GetArrayHash(PropertyFacade.HighValueArray.GetData(), PropertyFacade.HighValueArray.Num(), Hash);
 				Hash = GetArrayHash(PropertyFacade.StringValueArray.GetData(), PropertyFacade.StringValueArray.Num(), Hash);
@@ -183,6 +301,7 @@ namespace Chaos::Softs
 			return Hash;
 		}
 
+
 	protected:
 		// No init constructor for FCollectionPropertyFacade
 		CHAOS_API FCollectionPropertyConstFacade(const TSharedPtr<const FManagedArrayCollection>& InManagedArrayCollection, ENoInit);
@@ -194,9 +313,9 @@ namespace Chaos::Softs
 		CHAOS_API void RebuildKeyIndices();
 
 		template<typename CallableType, typename ReturnType>
-		ReturnType SafeGet(const FString& Key, CallableType&& Callable, ReturnType Default, int32* OutKeyIndex) const
+		ReturnType SafeGet(const FName& Key, CallableType&& Callable, ReturnType Default, int32* OutKeyIndex) const
 		{
-			const int32 KeyIndex = GetKeyIndex(Key);
+			const int32 KeyIndex = GetKeyNameIndex(Key);
 			if (OutKeyIndex)
 			{
 				*OutKeyIndex = KeyIndex;
@@ -204,6 +323,13 @@ namespace Chaos::Softs
 			return KeyIndex != INDEX_NONE ? Callable(KeyIndex) : Default;
 		}
 
+		template<typename CallableType, typename ReturnType, typename StringKeyType UE_REQUIRES(TIsStringAndNameConvertibleType<StringKeyType>::Value)>
+		UE_DEPRECATED(5.7, "Use FName Key version of this method")
+		ReturnType SafeGet(const StringKeyType Key, CallableType&& Callable, ReturnType Default, int32* OutKeyIndex) const
+		{
+			return SafeGet(FName(Key), Callable, Default, OutKeyIndex);
+		}
+
 		template<typename T, typename ElementType>
 		CHAOS_API T GetValue(int32 KeyIndex, const TConstArrayView<ElementType>& ValueArray) const;
 
@@ -214,24 +340,31 @@ namespace Chaos::Softs
 
 		// Attribute groups, predefined data member of the collection
 		static CHAOS_API const FName PropertyGroup;
-		static CHAOS_API const FName KeyName;  // Property key, name to look for
+		static CHAOS_API const FName KeyNameName;  // Property key, name to look for
 		static CHAOS_API const FName LowValueName;  // Boolean, 24 bit integer (max 16777215), float, or vector value, or value of the lowest weight on the weight map if any
 		static CHAOS_API const FName HighValueName;  // Boolean, 24 bit integer (max 16777215), float, or vector value of the highest weight on the weight map if any
 		static CHAOS_API const FName StringValueName;  // String value, or weight map name, ...etc.
 		static CHAOS_API const FName FlagsName;  // Whether this property is enabled, animatable, ...etc.
 
 		// Property Group array views
-		TConstArrayView<FString> KeyArray;
+		TConstArrayView<FName> KeyNameArray;
 		TConstArrayView<FVector3f> LowValueArray;
 		TConstArrayView<FVector3f> HighValueArray;
 		TConstArrayView<FString> StringValueArray;
 		TConstArrayView<uint8> FlagsArray;
 
 		// Key to index search map
-		TMap<FString, int32> KeyIndices;
+		TMap<FName, int32> KeyNameIndices;
 
 		// Property collection
 		TSharedPtr<const FManagedArrayCollection> ManagedArrayCollection;
+
+		UE_DEPRECATED(5.7, "Use KeyNameName instead")
+		static CHAOS_API const FName KeyName;
+		UE_DEPRECATED(5.7, "Use KeyNameArray instead")
+		TConstArrayView<FString> KeyArray;
+		UE_DEPRECATED(5.7, "Use KeyNameIndices instead")
+		TMap<FString, int32> KeyIndices;
 	};
 
 	/**
@@ -286,82 +419,181 @@ namespace Chaos::Softs
 
 		//~ Values set per key
 		template<typename T UE_REQUIRES(TIsWeightedType<T>::Value)>
-		int32 SetLowValue(const FString& Key, const T& Value)
+		int32 SetLowValue(const FName& Key, const T& Value)
 		{
 			return SafeSet(Key, [this, &Value](int32 KeyIndex) { SetLowValue(KeyIndex, Value); });
 		}
 
 		template<typename T UE_REQUIRES(TIsWeightedType<T>::Value)>
-		int32 SetHighValue(const FString& Key, const T& Value)
+		int32 SetHighValue(const FName& Key, const T& Value)
 		{
 			return SafeSet(Key, [this, &Value](int32 KeyIndex) { SetHighValue(KeyIndex, Value); });
 		}
 
 		template<typename T UE_REQUIRES(TIsWeightedType<T>::Value)>
-		int32 SetWeightedValue(const FString& Key, const T& LowValue, const T& HighValue)
+		int32 SetWeightedValue(const FName& Key, const T& LowValue, const T& HighValue)
 		{
 			return SafeSet(Key, [this, &LowValue, &HighValue](int32 KeyIndex) { SetWeightedValue(KeyIndex, LowValue, HighValue); });
 		}
 
-		int32 SetWeightedFloatValue(const FString& Key, const FVector2f& Value)
+		int32 SetWeightedFloatValue(const FName& Key, const FVector2f& Value)
 		{
 			return SafeSet(Key, [this, &Value](int32 KeyIndex) { SetWeightedFloatValue(KeyIndex, Value); });
 		}
 
 		template<typename T UE_REQUIRES(TIsWeightedType<T>::Value)>
-		int32 SetValue(const FString& Key, const T& Value)
+		int32 SetValue(const FName& Key, const T& Value)
 		{
 			return SafeSet(Key, [this, &Value](int32 KeyIndex) { SetValue(KeyIndex, Value); });
 		}
 
-		int32 SetStringValue(const FString& Key, const FString& Value)
+		int32 SetStringValue(const FName& Key, const FString& Value)
 		{
 			return SafeSet(Key, [this, &Value](int32 KeyIndex) { SetStringValue(KeyIndex, Value); });
 		}
 
-		int32 SetFlags(const FString& Key, ECollectionPropertyFlags Flags)
+		int32 SetFlags(const FName& Key, ECollectionPropertyFlags Flags)
 		{
 			return SafeSet(Key, [this, Flags](int32 KeyIndex) { SetFlags(KeyIndex, Flags); });
 		}
 
-		int32 SetEnabled(const FString& Key, bool bEnabled)
+		int32 SetEnabled(const FName& Key, bool bEnabled)
 		{
 			return SafeSet(Key, [this, bEnabled](int32 KeyIndex) { SetEnabled(KeyIndex, bEnabled); });
 		}
 
-		int32 SetAnimatable(const FString& Key, bool bAnimatable)
+		int32 SetAnimatable(const FName& Key, bool bAnimatable)
 		{
 			return SafeSet(Key, [this, bAnimatable](int32 KeyIndex) { SetAnimatable(KeyIndex, bAnimatable); });
 		}
 
-		int32 SetLegacy(const FString& Key, bool bLegacy)
+		int32 SetLegacy(const FName& Key, bool bLegacy)
 		{
 			return SafeSet(Key, [this, bLegacy](int32 KeyIndex) { SetLegacy(KeyIndex, bLegacy); });
 		}
 
 		/** Set the intrinsic flag for this property. This flag cannot be removed and implies non Animatable. */
-		int32 SetIntrinsic(const FString& Key)
+		int32 SetIntrinsic(const FName& Key)
 		{
 			return SafeSet(Key, [this](int32 KeyIndex) { SetIntrinsic(KeyIndex); });
 		}
 
-		int32 SetDirty(const FString& Key)
+		int32 SetDirty(const FName& Key)
 		{
 			return SafeSet(Key, [this](int32 KeyIndex) { SetDirty(KeyIndex); });
 		}
 
-		int32 SetStringDirty(const FString& Key)
+		int32 SetStringDirty(const FName& Key)
 		{
 			return SafeSet(Key, [this](int32 KeyIndex) { SetStringDirty(KeyIndex); });
 		}
 		
-		int32 SetInterpolable(const FString& Key)
+		int32 SetInterpolable(const FName& Key)
 		{
 			return SafeSet(Key, [this](int32 KeyIndex) { SetInterpolable(KeyIndex); });
 		}
 
-		CHAOS_API void ClearDirtyFlags();
+		//~ Deprecated string key versions: Values set per key
+		template<typename T, typename StringType UE_REQUIRES(TIsWeightedType<T>::Value && TIsStringAndNameConvertibleType<StringType>::Value)>
+		UE_DEPRECATED(5.7, "Use FName Key version of this method")
+		int32 SetLowValue(const StringType Key, const T& Value)
+		{
+			return SafeSet(FName(Key), [this, &Value](int32 KeyIndex) { SetLowValue(KeyIndex, Value); });
+		}
+
+		template<typename T, typename StringType UE_REQUIRES(TIsWeightedType<T>::Value && TIsStringAndNameConvertibleType<StringType>::Value)>
+		UE_DEPRECATED(5.7, "Use FName Key version of this method")
+		int32 SetHighValue(const StringType Key, const T& Value)
+		{
+			return SafeSet(FName(Key), [this, &Value](int32 KeyIndex) { SetHighValue(KeyIndex, Value); });
+		}
+
+		template<typename T, typename StringType UE_REQUIRES(TIsWeightedType<T>::Value && TIsStringAndNameConvertibleType<StringType>::Value)>
+		UE_DEPRECATED(5.7, "Use FName Key version of this method")
+		int32 SetWeightedValue(const StringType Key, const T& LowValue, const T& HighValue)
+		{
+			return SafeSet(FName(Key), [this, &LowValue, &HighValue](int32 KeyIndex) { SetWeightedValue(KeyIndex, LowValue, HighValue); });
+		}
+
+		template<typename StringType UE_REQUIRES(TIsStringAndNameConvertibleType<StringType>::Value)>
+		UE_DEPRECATED(5.7, "Use FName Key version of this method")
+		int32 SetWeightedFloatValue(const StringType Key, const FVector2f& Value)
+		{
+			return SafeSet(FName(Key), [this, &Value](int32 KeyIndex) { SetWeightedFloatValue(KeyIndex, Value); });
+		}
+
+		template<typename T, typename StringType UE_REQUIRES(TIsWeightedType<T>::Value && TIsStringAndNameConvertibleType<StringType>::Value)>
+		UE_DEPRECATED(5.7, "Use FName Key version of this method")
+		int32 SetValue(const StringType Key, const T& Value)
+		{
+			return SafeSet(FName(Key), [this, &Value](int32 KeyIndex) { SetValue(KeyIndex, Value); });
+		}
+
+		template<typename StringType UE_REQUIRES(TIsStringAndNameConvertibleType<StringType>::Value)>
+		UE_DEPRECATED(5.7, "Use FName Key version of this method")
+		int32 SetStringValue(const StringType Key, const FString& Value)
+		{
+			return SafeSet(FName(Key), [this, &Value](int32 KeyIndex) { SetStringValue(KeyIndex, Value); });
+		}
 
+		template<typename StringType UE_REQUIRES(TIsStringAndNameConvertibleType<StringType>::Value)>
+		UE_DEPRECATED(5.7, "Use FName Key version of this method")
+		int32 SetFlags(const StringType Key, ECollectionPropertyFlags Flags)
+		{
+			return SafeSet(FName(Key), [this, Flags](int32 KeyIndex) { SetFlags(KeyIndex, Flags); });
+		}
+
+		template<typename StringType UE_REQUIRES(TIsStringAndNameConvertibleType<StringType>::Value)>
+		UE_DEPRECATED(5.7, "Use FName Key version of this method")
+		int32 SetEnabled(const StringType Key, bool bEnabled)
+		{
+			return SafeSet(FName(Key), [this, bEnabled](int32 KeyIndex) { SetEnabled(KeyIndex, bEnabled); });
+		}
+
+		template<typename StringType UE_REQUIRES(TIsStringAndNameConvertibleType<StringType>::Value)>
+		UE_DEPRECATED(5.7, "Use FName Key version of this method")
+		int32 SetAnimatable(const StringType Key, bool bAnimatable)
+		{
+			return SafeSet(FName(Key), [this, bAnimatable](int32 KeyIndex) { SetAnimatable(KeyIndex, bAnimatable); });
+		}
+
+		template<typename StringType UE_REQUIRES(TIsStringAndNameConvertibleType<StringType>::Value)>
+		UE_DEPRECATED(5.7, "Use FName Key version of this method")
+		int32 SetLegacy(const StringType Key, bool bLegacy)
+		{
+			return SafeSet(FName(Key), [this, bLegacy](int32 KeyIndex) { SetLegacy(KeyIndex, bLegacy); });
+		}
+
+		/** Set the intrinsic flag for this property. This flag cannot be removed and implies non Animatable. */
+		template<typename StringType UE_REQUIRES(TIsStringAndNameConvertibleType<StringType>::Value)>
+		UE_DEPRECATED(5.7, "Use FName Key version of this method")
+		int32 SetIntrinsic(const StringType Key)
+		{
+			return SafeSet(FName(Key), [this](int32 KeyIndex) { SetIntrinsic(KeyIndex); });
+		}
+
+		template<typename StringType UE_REQUIRES(TIsStringAndNameConvertibleType<StringType>::Value)>
+		UE_DEPRECATED(5.7, "Use FName Key version of this method")
+		int32 SetDirty(const StringType Key)
+		{
+			return SafeSet(FName(Key), [this](int32 KeyIndex) { SetDirty(KeyIndex); });
+		}
+
+		template<typename StringType UE_REQUIRES(TIsStringAndNameConvertibleType<StringType>::Value)>
+		UE_DEPRECATED(5.7, "Use FName Key version of this method")
+		int32 SetStringDirty(const StringType Key)
+		{
+			return SafeSet(FName(Key), [this](int32 KeyIndex) { SetStringDirty(KeyIndex); });
+		}
+
+		template<typename StringType UE_REQUIRES(TIsStringAndNameConvertibleType<StringType>::Value)>
+		UE_DEPRECATED(5.7, "Use FName Key version of this method")
+		int32 SetInterpolable(const StringType Key)
+		{
+			return SafeSet(FName(Key), [this](int32 KeyIndex) { SetInterpolable(KeyIndex); });
+		}
+
+		CHAOS_API void ClearDirtyFlags();
 		CHAOS_API void UpdateProperties(const TSharedPtr<const FManagedArrayCollection>& InManagedArrayCollection);
 
 	protected:
@@ -370,18 +602,32 @@ namespace Chaos::Softs
 
 		// Access to a writeable ManagedArrayCollection is protected, use an FPropertyCollectionMutableAdapter if needed to get a non const pointer
 		TSharedPtr<FManagedArrayCollection> GetManagedArrayCollection() { return ConstCastSharedPtr<FManagedArrayCollection>(ManagedArrayCollection); }
-
-		const TArrayView<FString>& GetKeyArray() { return reinterpret_cast<TArrayView<FString>&>(KeyArray); }
+		const TArrayView<FName>& GetKeyNameArray() { return reinterpret_cast<TArrayView<FName>&>(KeyNameArray); }
 		const TArrayView<FVector3f>& GetLowValueArray() { return reinterpret_cast<TArrayView<FVector3f>&>(LowValueArray); }
 		const TArrayView<FVector3f>& GetHighValueArray() { return reinterpret_cast<TArrayView<FVector3f>&>(HighValueArray); }
 		const TArrayView<FString>& GetStringValueArray() { return reinterpret_cast<TArrayView<FString>&>(StringValueArray); }
 		const TArrayView<ECollectionPropertyFlags>& GetFlagsArray() { return reinterpret_cast<const TArrayView<ECollectionPropertyFlags>&>(FlagsArray); }
 
+		UE_DEPRECATED(5.7, "Use GetKeyNameArray instead")
+		const TArrayView<FString>& GetKeyArray() 
+		{ 
+			PRAGMA_DISABLE_DEPRECATION_WARNINGS
+			return reinterpret_cast<TArrayView<FString>&>(KeyArray); 
+			PRAGMA_ENABLE_DEPRECATION_WARNINGS
+		}
+
 	private:
+		template<typename CallableType, typename StringType UE_REQUIRES(TIsStringAndNameConvertibleType<StringType>::Value)>
+		UE_DEPRECATED(5.7, "Use FName Key version of this method")
+		int32 SafeSet(const StringType Key, CallableType&& Callable)
+		{
+			return SafeSet(FName(Key), Callable);
+		}
+
 		template<typename CallableType>
-		int32 SafeSet(const FString& Key, CallableType&& Callable)
+		int32 SafeSet(const FName& Key, CallableType&& Callable)
 		{
-			const int32 KeyIndex = GetKeyIndex(Key);
+			const int32 KeyIndex = GetKeyNameIndex(Key);
 			if (KeyIndex != INDEX_NONE)
 			{
 				Callable(KeyIndex);
@@ -428,16 +674,52 @@ namespace Chaos::Softs
 		CHAOS_API void DefineSchema();
 
 		/** Add a single property, and return its index. */
-		CHAOS_API int32 AddProperty(const FString& Key, ECollectionPropertyFlags Flags = ECollectionPropertyFlags::Enabled);
+		CHAOS_API int32 AddProperty(const FName& Key, ECollectionPropertyFlags Flags = ECollectionPropertyFlags::Enabled);
+		template<typename StringType UE_REQUIRES(TIsStringAndNameConvertibleType<StringType>::Value)>
+		UE_DEPRECATED(5.7, "Use FName Key version of this method")
+		int32 AddProperty(const StringType Key, ECollectionPropertyFlags Flags = ECollectionPropertyFlags::Enabled)
+		{
+			return AddProperty(FName(Key), Flags);
+		}
 
 		/** Add a single property, and return its index. */
-		CHAOS_API int32 AddProperty(const FString& Key, bool bEnabled, bool bAnimatable = false, bool bIntrinsic = false);
+		CHAOS_API int32 AddProperty(const FName& Key, bool bEnabled, bool bAnimatable = false, bool bIntrinsic = false);
+		template<typename StringType UE_REQUIRES(TIsStringAndNameConvertibleType<StringType>::Value)>
+		UE_DEPRECATED(5.7, "Use FName Key version of this method")
+		int32 AddProperty(const StringType Key, bool bEnabled, bool bAnimatable = false, bool bIntrinsic = false)
+		{
+			return AddProperty(FName(Key), bEnabled, bAnimatable, bIntrinsic);
+		}
 
 		/** Add new properties, and return the index of the first added property. */
-		CHAOS_API int32 AddProperties(const TArray<FString>& Keys, ECollectionPropertyFlags Flags = ECollectionPropertyFlags::Enabled);
+		CHAOS_API int32 AddProperties(const TArray<FName>& Keys, ECollectionPropertyFlags Flags = ECollectionPropertyFlags::Enabled);
+		template<typename StringType UE_REQUIRES(TIsStringAndNameConvertibleType<StringType>::Value)>
+		UE_DEPRECATED(5.7, "Use FName Keys version of this method")
+		int32 AddProperties(const TArray<StringType>& InKeys, ECollectionPropertyFlags Flags = ECollectionPropertyFlags::Enabled)
+		{
+			TArray<FName> Keys;
+			Keys.Reserve(Keys.Num());
+			for (const StringType& Key : InKeys)
+			{
+				Keys.Emplace(FName(Key));
+			}
+			return AddProperties(Keys, Flags);
+		}
 
 		/** Add new properties, and return the index of the first added property. */
-		CHAOS_API int32 AddProperties(const TArray<FString>& Keys, bool bEnabled, bool bAnimatable = false, bool bIntrinsic = false);
+		CHAOS_API int32 AddProperties(const TArray<FName>& Keys, bool bEnabled, bool bAnimatable = false, bool bIntrinsic = false);
+		template<typename StringType UE_REQUIRES(TIsStringAndNameConvertibleType<StringType>::Value)>
+		UE_DEPRECATED(5.7, "Use FName Keys version of this method")
+		int32 AddProperties(const TArray<StringType>& InKeys, bool bEnabled, bool bAnimatable = false, bool bIntrinsic = false)
+		{
+			TArray<FName> Keys;
+			Keys.Reserve(Keys.Num());
+			for (const StringType& Key : InKeys)
+			{
+				Keys.Emplace(FName(Key));
+			}
+			return AddProperties(Keys, bEnabled, bAnimatable, bIntrinsic);
+		}
 
 		/**
 		 * Append all properties and values from an existing collection to this property collection.
@@ -465,20 +747,79 @@ namespace Chaos::Softs
 
 		//~ Add values
 		template<typename T UE_REQUIRES(TIsWeightedType<T>::Value)>
-		inline int32 AddWeightedValue(const FString& Key, const T& LowValue, const T& HighValue, ECollectionPropertyFlags Flags = ECollectionPropertyFlags::Enabled);
+		inline int32 AddWeightedValue(const FName& Key, const T& LowValue, const T& HighValue, ECollectionPropertyFlags Flags = ECollectionPropertyFlags::Enabled);
 		template<typename T UE_REQUIRES(TIsWeightedType<T>::Value)>
-		inline int32 AddWeightedValue(const FString& Key, const T& LowValue, const T& HighValue, bool bEnabled, bool bAnimatable = false, bool bIntrinsic = false);
+		inline int32 AddWeightedValue(const FName& Key, const T& LowValue, const T& HighValue, bool bEnabled, bool bAnimatable = false, bool bIntrinsic = false);
 
-		CHAOS_API int32 AddWeightedFloatValue(const FString& Key, const FVector2f& Value, ECollectionPropertyFlags Flags = ECollectionPropertyFlags::Enabled);
-		CHAOS_API int32 AddWeightedFloatValue(const FString& Key, const FVector2f& Value, bool bEnabled, bool bAnimatable, bool bIntrinsic = false);
+		CHAOS_API int32 AddWeightedFloatValue(const FName& Key, const FVector2f& Value, ECollectionPropertyFlags Flags = ECollectionPropertyFlags::Enabled);
+		CHAOS_API int32 AddWeightedFloatValue(const FName& Key, const FVector2f& Value, bool bEnabled, bool bAnimatable, bool bIntrinsic = false);
 
 		template<typename T UE_REQUIRES(TIsWeightedType<T>::Value)>
-		int32 AddValue(const FString& Key, const T& Value, ECollectionPropertyFlags Flags = ECollectionPropertyFlags::Enabled) { return AddWeightedValue(Key, Value, Value, Flags); }
+		int32 AddValue(const FName& Key, const T& Value, ECollectionPropertyFlags Flags = ECollectionPropertyFlags::Enabled) { return AddWeightedValue(Key, Value, Value, Flags); }
 		template<typename T UE_REQUIRES(TIsWeightedType<T>::Value)>
-		int32 AddValue(const FString& Key, const T& Value, bool bEnabled, bool bAnimatable = false) { return AddWeightedValue(Key, Value, Value, bEnabled, bAnimatable); }
+		int32 AddValue(const FName& Key, const T& Value, bool bEnabled, bool bAnimatable = false) { return AddWeightedValue(Key, Value, Value, bEnabled, bAnimatable); }
+
+		CHAOS_API int32 AddStringValue(const FName& Key, const FString& Value, ECollectionPropertyFlags Flags = ECollectionPropertyFlags::Enabled);
+		CHAOS_API int32 AddStringValue(const FName& Key, const FString& Value, bool bEnabled, bool bAnimatable = false, bool bIntrinsic = false);
+
+		//~ Deprecated string key versions: Add values
+		template<typename T, typename StringType UE_REQUIRES(TIsWeightedType<T>::Value&& TIsStringAndNameConvertibleType<StringType>::Value)>
+		UE_DEPRECATED(5.7, "Use FName Key version of this method")
+		int32 AddWeightedValue(const StringType Key, const T& LowValue, const T& HighValue, ECollectionPropertyFlags Flags = ECollectionPropertyFlags::Enabled)
+		{
+			return AddWeightedValue(FName(Key), LowValue, HighValue, Flags);
+		}
+
+		template<typename T, typename StringType UE_REQUIRES(TIsWeightedType<T>::Value&& TIsStringAndNameConvertibleType<StringType>::Value)>
+		UE_DEPRECATED(5.7, "Use FName Key version of this method")
+		int32 AddWeightedValue(const StringType Key, const T& LowValue, const T& HighValue, bool bEnabled, bool bAnimatable = false, bool bIntrinsic = false)
+		{
+			return AddWeightedValue(FName(Key), LowValue, HighValue, bEnabled, bAnimatable, bIntrinsic);
+		}
+
+		template<typename StringType UE_REQUIRES(TIsStringAndNameConvertibleType<StringType>::Value)>
+		UE_DEPRECATED(5.7, "Use FName Key version of this method")
+		int32 AddWeightedFloatValue(const StringType Key, const FVector2f& Value, ECollectionPropertyFlags Flags = ECollectionPropertyFlags::Enabled)
+		{
+			return AddWeightedFloatValue(FName(Key), Value, Flags);
+		}
+
+		template<typename StringType UE_REQUIRES(TIsStringAndNameConvertibleType<StringType>::Value)>
+		UE_DEPRECATED(5.7, "Use FName Key version of this method")
+		int32 AddWeightedFloatValue(const StringType Key, const FVector2f& Value, bool bEnabled, bool bAnimatable, bool bIntrinsic = false)
+		{
+			return AddWeightedFloatValue(FName(Key), Value, bEnabled, bAnimatable, bIntrinsic);
+		}
+
+		template<typename T, typename StringType UE_REQUIRES(TIsWeightedType<T>::Value&& TIsStringAndNameConvertibleType<StringType>::Value)>
+		UE_DEPRECATED(5.7, "Use FName Key version of this method")
+		int32 AddValue(const StringType Key, const T& Value, ECollectionPropertyFlags Flags = ECollectionPropertyFlags::Enabled)
+		{
+			return AddWeightedValue(FName(Key), Value, Value, Flags);
+		}
+
+		template<typename T, typename StringType UE_REQUIRES(TIsWeightedType<T>::Value&& TIsStringAndNameConvertibleType<StringType>::Value)>
+		UE_DEPRECATED(5.7, "Use FName Key version of this method")
+		int32 AddValue(const StringType Key, const T& Value, bool bEnabled, bool bAnimatable = false)
+		{ 
+			return AddWeightedValue(FName(Key), Value, Value, bEnabled, bAnimatable); 
+		}
+
+		template<typename StringType UE_REQUIRES(TIsStringAndNameConvertibleType<StringType>::Value)>
+		UE_DEPRECATED(5.7, "Use FName Key version of this method")
+		int32 AddStringValue(const StringType Key, const FString& Value, ECollectionPropertyFlags Flags = ECollectionPropertyFlags::Enabled)
+		{
+			return AddStringValue(FName(Key), Value, Flags);
+		}
 
-		CHAOS_API int32 AddStringValue(const FString& Key, const FString& Value, ECollectionPropertyFlags Flags = ECollectionPropertyFlags::Enabled);
-		CHAOS_API int32 AddStringValue(const FString& Key, const FString& Value, bool bEnabled, bool bAnimatable = false, bool bIntrinsic = false);
+		template<typename StringType UE_REQUIRES(TIsStringAndNameConvertibleType<StringType>::Value)>
+		UE_DEPRECATED(5.7, "Use FName Key version of this method")
+		int32 AddStringValue(const StringType Key, const FString& Value, bool bEnabled, bool bAnimatable = false, bool bIntrinsic = false)
+		{
+			return AddStringValue(FName(Key), Value, bEnabled, bAnimatable, bIntrinsic);
+		}
+
+		CHAOS_API void PostSerialize(const FArchive& Ar);
 	};
 
 	template<typename T>
@@ -492,7 +833,7 @@ namespace Chaos::Softs
 	}
 
 	template<typename T UE_REQUIRES_DEFINITION(TIsWeightedType<T>::Value)>
-	inline int32 FCollectionPropertyMutableFacade::AddWeightedValue(const FString& Key, const T& LowValue, const T& HighValue, ECollectionPropertyFlags Flags)
+	inline int32 FCollectionPropertyMutableFacade::AddWeightedValue(const FName& Key, const T& LowValue, const T& HighValue, ECollectionPropertyFlags Flags)
 	{
 		const int32 KeyIndex = AddProperty(Key, Flags);
 		SetWeightedValue(KeyIndex, LowValue, HighValue);
@@ -500,7 +841,7 @@ namespace Chaos::Softs
 	}
 
 	template<typename T UE_REQUIRES_DEFINITION(TIsWeightedType<T>::Value)>
-	inline int32 FCollectionPropertyMutableFacade::AddWeightedValue(const FString& Key, const T& LowValue, const T& HighValue, bool bEnabled, bool bAnimatable, bool bIntrinsic)
+	inline int32 FCollectionPropertyMutableFacade::AddWeightedValue(const FName& Key, const T& LowValue, const T& HighValue, bool bEnabled, bool bAnimatable, bool bIntrinsic)
 	{
 		const int32 KeyIndex = AddProperty(Key, bEnabled, bAnimatable, bIntrinsic);
 		SetWeightedValue(KeyIndex, LowValue, HighValue);
@@ -511,107 +852,106 @@ namespace Chaos::Softs
 // Use this macro to add shorthands for property getters without a key index
 #define UE_CHAOS_DECLARE_INDEXLESS_PROPERTYCOLLECTION_NAME(PropertyName, Type) \
 	inline static const FName PropertyName##Name = TEXT(#PropertyName); \
-	static bool Is##PropertyName##Enabled(const FCollectionPropertyConstFacade& InPropertyCollection, bool bDefault) \
+	static bool Is##PropertyName##Enabled(const ::Chaos::Softs::FCollectionPropertyConstFacade& InPropertyCollection, bool bDefault) \
 	{ \
-		return InPropertyCollection.IsEnabled(PropertyName##Name.ToString(), bDefault); \
+		return InPropertyCollection.IsEnabled(PropertyName##Name, bDefault); \
 	} \
-	static bool Is##PropertyName##Animatable(const FCollectionPropertyConstFacade& InPropertyCollection, bool bDefault) \
+	static bool Is##PropertyName##Animatable(const ::Chaos::Softs::FCollectionPropertyConstFacade& InPropertyCollection, bool bDefault) \
 	{ \
-		return InPropertyCollection.IsAnimatable(PropertyName##Name.ToString(), bDefault); \
+		return InPropertyCollection.IsAnimatable(PropertyName##Name, bDefault); \
 	} \
-	static Type GetLow##PropertyName(const FCollectionPropertyConstFacade& InPropertyCollection, const Type& Default) \
+	static Type GetLow##PropertyName(const ::Chaos::Softs::FCollectionPropertyConstFacade& InPropertyCollection, const Type& Default) \
 	{ \
-		return InPropertyCollection.GetLowValue<Type>(PropertyName##Name.ToString(), Default); \
+		return InPropertyCollection.GetLowValue<Type>(PropertyName##Name, Default); \
 	} \
-	static Type GetHigh##PropertyName(const FCollectionPropertyConstFacade& InPropertyCollection, const Type& Default) \
+	static Type GetHigh##PropertyName(const ::Chaos::Softs::FCollectionPropertyConstFacade& InPropertyCollection, const Type& Default) \
 	{ \
-		return InPropertyCollection.GetHighValue<Type>(PropertyName##Name.ToString(), Default); \
+		return InPropertyCollection.GetHighValue<Type>(PropertyName##Name, Default); \
 	} \
-	static TPair<Type, Type> GetWeighted##PropertyName(const FCollectionPropertyConstFacade& InPropertyCollection, const Type& Default) \
+	static TPair<Type, Type> GetWeighted##PropertyName(const ::Chaos::Softs::FCollectionPropertyConstFacade& InPropertyCollection, const Type& Default) \
 	{ \
-		return InPropertyCollection.GetWeightedValue<Type>(PropertyName##Name.ToString(), Default); \
+		return InPropertyCollection.GetWeightedValue<Type>(PropertyName##Name, Default); \
 	} \
-	static FVector2f GetWeightedFloat##PropertyName(const FCollectionPropertyConstFacade& InPropertyCollection, const float& Default) \
+	static FVector2f GetWeightedFloat##PropertyName(const ::Chaos::Softs::FCollectionPropertyConstFacade& InPropertyCollection, const float& Default) \
 	{ \
-		return InPropertyCollection.GetWeightedFloatValue(PropertyName##Name.ToString(), Default); \
+		return InPropertyCollection.GetWeightedFloatValue(PropertyName##Name, Default); \
 	} \
-	static FVector2f GetWeightedFloat##PropertyName(const FCollectionPropertyConstFacade& InPropertyCollection, const FVector2f& Default) \
+	static FVector2f GetWeightedFloat##PropertyName(const ::Chaos::Softs::FCollectionPropertyConstFacade& InPropertyCollection, const FVector2f& Default) \
 	{ \
-		return InPropertyCollection.GetWeightedFloatValue(PropertyName##Name.ToString(), Default); \
+		return InPropertyCollection.GetWeightedFloatValue(PropertyName##Name, Default); \
 	} \
-	static Type Get##PropertyName(const FCollectionPropertyConstFacade& InPropertyCollection, const Type& Default) \
+	static Type Get##PropertyName(const ::Chaos::Softs::FCollectionPropertyConstFacade& InPropertyCollection, const Type& Default) \
 	{ \
-		return InPropertyCollection.GetValue<Type>(PropertyName##Name.ToString(), Default); \
+		return InPropertyCollection.GetValue<Type>(PropertyName##Name, Default); \
 	} \
-	static FString Get##PropertyName##String(const FCollectionPropertyConstFacade& InPropertyCollection, const FString& Default) \
+	static FString Get##PropertyName##String(const ::Chaos::Softs::FCollectionPropertyConstFacade& InPropertyCollection, const FString& Default) \
 	{ \
-		return InPropertyCollection.GetStringValue(PropertyName##Name.ToString(), Default); \
+		return InPropertyCollection.GetStringValue(PropertyName##Name, Default); \
 	} \
 
-// Use this macro to add shorthands for property getters and direct access through the declared key index
-#define UE_CHAOS_DECLARE_PROPERTYCOLLECTION_NAME(PropertyName, Type) \
-	UE_CHAOS_DECLARE_INDEXLESS_PROPERTYCOLLECTION_NAME(PropertyName, Type) \
-	Type GetLow##PropertyName(const FCollectionPropertyConstFacade& InPropertyCollection) const \
+// Use this macro to add shorthands for property getters and direct access through the declared key index. 
+#define UE_CHAOS_DECLARE_INDEXED_PROPERTYCOLLECTION_NAME(PropertyName, Type) \
+	Type GetLow##PropertyName(const ::Chaos::Softs::FCollectionPropertyConstFacade& InPropertyCollection) const \
 	{ \
-		checkSlow(PropertyName##Index == InPropertyCollection.GetKeyIndex(PropertyName##Name.ToString())); \
+		checkSlow(PropertyName##Index == InPropertyCollection.GetKeyNameIndex(PropertyName##Name)); \
 		checkf(PropertyName##Index != INDEX_NONE, TEXT("The default value getter that sets the property index must be called once prior to calling this function.")); \
 		return InPropertyCollection.GetLowValue<Type>(PropertyName##Index); \
 	} \
-	Type GetHigh##PropertyName(const FCollectionPropertyConstFacade& InPropertyCollection) const \
+	Type GetHigh##PropertyName(const ::Chaos::Softs::FCollectionPropertyConstFacade& InPropertyCollection) const \
 	{ \
-		checkSlow(PropertyName##Index == InPropertyCollection.GetKeyIndex(PropertyName##Name.ToString())); \
+		checkSlow(PropertyName##Index == InPropertyCollection.GetKeyNameIndex(PropertyName##Name)); \
 		checkf(PropertyName##Index != INDEX_NONE, TEXT("The default value getter that sets the property index must be called once prior to calling this function.")); \
 		return InPropertyCollection.GetHighValue<Type>(PropertyName##Index); \
 	} \
-	TPair<Type, Type> GetWeighted##PropertyName(const FCollectionPropertyConstFacade& InPropertyCollection) const \
+	TPair<Type, Type> GetWeighted##PropertyName(const ::Chaos::Softs::FCollectionPropertyConstFacade& InPropertyCollection) const \
 	{ \
-		checkSlow(PropertyName##Index == InPropertyCollection.GetKeyIndex(PropertyName##Name.ToString())); \
+		checkSlow(PropertyName##Index == InPropertyCollection.GetKeyNameIndex(PropertyName##Name)); \
 		checkf(PropertyName##Index != INDEX_NONE, TEXT("The default value getter that sets the property index must be called once prior to calling this function.")); \
 		return InPropertyCollection.GetWeightedValue<Type>(PropertyName##Index); \
 	} \
-	FVector2f GetWeightedFloat##PropertyName(const FCollectionPropertyConstFacade& InPropertyCollection) const \
+	FVector2f GetWeightedFloat##PropertyName(const ::Chaos::Softs::FCollectionPropertyConstFacade& InPropertyCollection) const \
 	{ \
-		checkSlow(PropertyName##Index == InPropertyCollection.GetKeyIndex(PropertyName##Name.ToString())); \
+		checkSlow(PropertyName##Index == InPropertyCollection.GetKeyNameIndex(PropertyName##Name)); \
 		checkf(PropertyName##Index != INDEX_NONE, TEXT("The default value getter that sets the property index must be called once prior to calling this function.")); \
 		return InPropertyCollection.GetWeightedFloatValue(PropertyName##Index); \
 	} \
-	Type Get##PropertyName(const FCollectionPropertyConstFacade& InPropertyCollection) const \
+	Type Get##PropertyName(const ::Chaos::Softs::FCollectionPropertyConstFacade& InPropertyCollection) const \
 	{ \
-		checkSlow(PropertyName##Index == InPropertyCollection.GetKeyIndex(PropertyName##Name.ToString())); \
+		checkSlow(PropertyName##Index == InPropertyCollection.GetKeyNameIndex(PropertyName##Name)); \
 		checkf(PropertyName##Index != INDEX_NONE, TEXT("The default value getter that sets the property index must be called once prior to calling this function.")); \
 		return InPropertyCollection.GetValue<Type>(PropertyName##Index); \
 	} \
-	const FString& Get##PropertyName##String(const FCollectionPropertyConstFacade& InPropertyCollection) \
+	const FString& Get##PropertyName##String(const ::Chaos::Softs::FCollectionPropertyConstFacade& InPropertyCollection) \
 	{ \
-		checkSlow(PropertyName##Index == InPropertyCollection.GetKeyIndex(PropertyName##Name.ToString())); \
+		checkSlow(PropertyName##Index == InPropertyCollection.GetKeyNameIndex(PropertyName##Name)); \
 		checkf(PropertyName##Index != INDEX_NONE, TEXT("The default value getter that sets the property index must be called once prior to calling this function.")); \
 		return InPropertyCollection.GetStringValue(PropertyName##Index); \
 	} \
-	bool Is##PropertyName##Enabled(const FCollectionPropertyConstFacade& InPropertyCollection) const \
+	bool Is##PropertyName##Enabled(const ::Chaos::Softs::FCollectionPropertyConstFacade& InPropertyCollection) const \
 	{ \
-		checkSlow(PropertyName##Index == InPropertyCollection.GetKeyIndex(PropertyName##Name.ToString())); \
+		checkSlow(PropertyName##Index == InPropertyCollection.GetKeyNameIndex(PropertyName##Name)); \
 		checkf(PropertyName##Index != INDEX_NONE, TEXT("The default value getter that sets the property index must be called once prior to calling this function.")); \
 		return InPropertyCollection.IsEnabled(PropertyName##Index); \
 	} \
-	bool Is##PropertyName##Animatable(const FCollectionPropertyConstFacade& InPropertyCollection) const\
+	bool Is##PropertyName##Animatable(const ::Chaos::Softs::FCollectionPropertyConstFacade& InPropertyCollection) const\
 	{ \
-		checkSlow(PropertyName##Index == InPropertyCollection.GetKeyIndex(PropertyName##Name.ToString())); \
+		checkSlow(PropertyName##Index == InPropertyCollection.GetKeyNameIndex(PropertyName##Name)); \
 		checkf(PropertyName##Index != INDEX_NONE, TEXT("The default value getter that sets the property index must be called once prior to calling this function.")); \
 		return InPropertyCollection.IsAnimatable(PropertyName##Index); \
 	} \
-	bool Is##PropertyName##Dirty(const FCollectionPropertyConstFacade& InPropertyCollection) const \
+	bool Is##PropertyName##Dirty(const ::Chaos::Softs::FCollectionPropertyConstFacade& InPropertyCollection) const \
 	{ \
-		checkSlow(PropertyName##Index == InPropertyCollection.GetKeyIndex(PropertyName##Name.ToString())); \
+		checkSlow(PropertyName##Index == InPropertyCollection.GetKeyNameIndex(PropertyName##Name)); \
 		checkf(PropertyName##Index != INDEX_NONE, TEXT("The default value getter that sets the property index must be called once prior to calling this function.")); \
 		return InPropertyCollection.IsDirty(PropertyName##Index); \
 	} \
-	bool Is##PropertyName##StringDirty(const FCollectionPropertyConstFacade& InPropertyCollection) const \
+	bool Is##PropertyName##StringDirty(const ::Chaos::Softs::FCollectionPropertyConstFacade& InPropertyCollection) const \
 	{ \
-		checkSlow(PropertyName##Index == InPropertyCollection.GetKeyIndex(PropertyName##Name.ToString())); \
+		checkSlow(PropertyName##Index == InPropertyCollection.GetKeyNameIndex(PropertyName##Name)); \
 		checkf(PropertyName##Index != INDEX_NONE, TEXT("The default value getter that sets the property index must be called once prior to calling this function.")); \
 		return InPropertyCollection.IsStringDirty(PropertyName##Index); \
 	} \
-	bool Is##PropertyName##Mutable(const FCollectionPropertyConstFacade& InPropertyCollection) const \
+	bool Is##PropertyName##Mutable(const ::Chaos::Softs::FCollectionPropertyConstFacade& InPropertyCollection) const \
 	{ \
 		return PropertyName##Index != INDEX_NONE && \
 			InPropertyCollection.IsAnimatable(PropertyName##Index) && (InPropertyCollection.IsDirty(PropertyName##Index)); \
@@ -620,7 +960,13 @@ namespace Chaos::Softs
 	{ \
 		int32 Index = INDEX_NONE; \
 		explicit F##PropertyName##Index(EForceInit) : Index(INDEX_NONE) {} \
-		explicit F##PropertyName##Index(const FCollectionPropertyConstFacade& InPropertyCollection) : Index(InPropertyCollection.GetKeyIndex(PropertyName##Name.ToString())) {} \
+		explicit F##PropertyName##Index(const ::Chaos::Softs::FCollectionPropertyConstFacade& InPropertyCollection) : Index(InPropertyCollection.GetKeyNameIndex(PropertyName##Name)) {} \
 		operator int32() const { return Index; } \
 	} PropertyName##Index;
 
+
+// Use this macro to add both indexless and indexed property getters.
+#define UE_CHAOS_DECLARE_PROPERTYCOLLECTION_NAME(PropertyName, Type) \
+	UE_CHAOS_DECLARE_INDEXLESS_PROPERTYCOLLECTION_NAME(PropertyName, Type) \
+	UE_CHAOS_DECLARE_INDEXED_PROPERTYCOLLECTION_NAME(PropertyName, Type) \
+
diff --git a/Engine/Source/Runtime/Experimental/Chaos/Public/Chaos/PBDAnimDriveConstraint.h b/Engine/Source/Runtime/Experimental/Chaos/Public/Chaos/PBDAnimDriveConstraint.h
--- a/Engine/Source/Runtime/Experimental/Chaos/Public/Chaos/PBDAnimDriveConstraint.h
+++ b/Engine/Source/Runtime/Experimental/Chaos/Public/Chaos/PBDAnimDriveConstraint.h
@@ -188,6 +188,9 @@ namespace Chaos::Softs
 			}
 		}
 
+		UE_CHAOS_DECLARE_INDEXLESS_PROPERTYCOLLECTION_NAME(AnimDriveStiffness, float);
+		UE_CHAOS_DECLARE_INDEXLESS_PROPERTYCOLLECTION_NAME(AnimDriveDamping, float);
+
 	private:
 
 
@@ -220,8 +223,8 @@ namespace Chaos::Softs
 		FPBDStiffness Stiffness;
 		FPBDStiffness Damping;
 
-		UE_CHAOS_DECLARE_PROPERTYCOLLECTION_NAME(AnimDriveStiffness, float);
-		UE_CHAOS_DECLARE_PROPERTYCOLLECTION_NAME(AnimDriveDamping, float);
+		UE_CHAOS_DECLARE_INDEXED_PROPERTYCOLLECTION_NAME(AnimDriveStiffness, float);
+		UE_CHAOS_DECLARE_INDEXED_PROPERTYCOLLECTION_NAME(AnimDriveDamping, float);
 	};
 
 }  // End namespace Chaos::Softs
diff --git a/Engine/Source/Runtime/Experimental/Chaos/Public/Chaos/PBDAxialSpringConstraints.h b/Engine/Source/Runtime/Experimental/Chaos/Public/Chaos/PBDAxialSpringConstraints.h
--- a/Engine/Source/Runtime/Experimental/Chaos/Public/Chaos/PBDAxialSpringConstraints.h
+++ b/Engine/Source/Runtime/Experimental/Chaos/Public/Chaos/PBDAxialSpringConstraints.h
@@ -152,6 +152,10 @@ public:
 
 	CHAOS_API void ApplyProperties(const FSolverReal Dt, const int32 NumIterations);
 
+	UE_CHAOS_DECLARE_INDEXLESS_PROPERTYCOLLECTION_NAME(AreaSpringStiffness, float);
+	UE_CHAOS_DECLARE_INDEXLESS_PROPERTYCOLLECTION_NAME(AreaSpringWarpScale, float);
+	UE_CHAOS_DECLARE_INDEXLESS_PROPERTYCOLLECTION_NAME(AreaSpringWeftScale, float);
+
 private:
 
 	void InitFromPatternData(const TArray<TVec3<FVec2f>>& FaceVertexPatternPositions, const FTriangleMesh& TriangleMesh);
@@ -169,9 +173,9 @@ private:
 	TArray<FSolverReal> BaseDists; // Without Warp/Weft Scale applied
 	TArray<FSolverVec2> WarpWeftScaleBaseMultipliers;
 
-	UE_CHAOS_DECLARE_PROPERTYCOLLECTION_NAME(AreaSpringStiffness, float);
-	UE_CHAOS_DECLARE_PROPERTYCOLLECTION_NAME(AreaSpringWarpScale, float);
-	UE_CHAOS_DECLARE_PROPERTYCOLLECTION_NAME(AreaSpringWeftScale, float);
+	UE_CHAOS_DECLARE_INDEXED_PROPERTYCOLLECTION_NAME(AreaSpringStiffness, float);
+	UE_CHAOS_DECLARE_INDEXED_PROPERTYCOLLECTION_NAME(AreaSpringWarpScale, float);
+	UE_CHAOS_DECLARE_INDEXED_PROPERTYCOLLECTION_NAME(AreaSpringWeftScale, float);
 };
 
 }  // End namespace Chaos::Softs
diff --git a/Engine/Source/Runtime/Experimental/Chaos/Public/Chaos/PBDBendingConstraints.h b/Engine/Source/Runtime/Experimental/Chaos/Public/Chaos/PBDBendingConstraints.h
--- a/Engine/Source/Runtime/Experimental/Chaos/Public/Chaos/PBDBendingConstraints.h
+++ b/Engine/Source/Runtime/Experimental/Chaos/Public/Chaos/PBDBendingConstraints.h
@@ -167,6 +167,13 @@ public:
 			(ERestAngleConstructionType)GetRestAngleType(PropertyCollection, (int32)ERestAngleConstructionType::Use3DRestAngles));
 	}
 
+	UE_CHAOS_DECLARE_INDEXLESS_PROPERTYCOLLECTION_NAME(BendingElementStiffness, float);
+	UE_CHAOS_DECLARE_INDEXLESS_PROPERTYCOLLECTION_NAME(BucklingRatio, float);
+	UE_CHAOS_DECLARE_INDEXLESS_PROPERTYCOLLECTION_NAME(BucklingStiffness, float);
+	UE_CHAOS_DECLARE_INDEXLESS_PROPERTYCOLLECTION_NAME(FlatnessRatio, float);
+	UE_CHAOS_DECLARE_INDEXLESS_PROPERTYCOLLECTION_NAME(RestAngle, float);
+	UE_CHAOS_DECLARE_INDEXLESS_PROPERTYCOLLECTION_NAME(RestAngleType, int32);
+
 private:
 	template<typename SolverParticlesOrRange>
 	CHAOS_API void InitColor(const SolverParticlesOrRange& InParticles);
@@ -216,12 +223,12 @@ private:
 
 	TArray<int32> ConstraintsPerColorStartIndex; // Constraints are ordered so each batch is contiguous. This is ColorNum + 1 length so it can be used as start and end.
 
-	UE_CHAOS_DECLARE_PROPERTYCOLLECTION_NAME(BendingElementStiffness, float);
-	UE_CHAOS_DECLARE_PROPERTYCOLLECTION_NAME(BucklingRatio, float);
-	UE_CHAOS_DECLARE_PROPERTYCOLLECTION_NAME(BucklingStiffness, float);
-	UE_CHAOS_DECLARE_PROPERTYCOLLECTION_NAME(FlatnessRatio, float);
-	UE_CHAOS_DECLARE_PROPERTYCOLLECTION_NAME(RestAngle, float);
-	UE_CHAOS_DECLARE_PROPERTYCOLLECTION_NAME(RestAngleType, int32);
+	UE_CHAOS_DECLARE_INDEXED_PROPERTYCOLLECTION_NAME(BendingElementStiffness, float);
+	UE_CHAOS_DECLARE_INDEXED_PROPERTYCOLLECTION_NAME(BucklingRatio, float);
+	UE_CHAOS_DECLARE_INDEXED_PROPERTYCOLLECTION_NAME(BucklingStiffness, float);
+	UE_CHAOS_DECLARE_INDEXED_PROPERTYCOLLECTION_NAME(FlatnessRatio, float);
+	UE_CHAOS_DECLARE_INDEXED_PROPERTYCOLLECTION_NAME(RestAngle, float);
+	UE_CHAOS_DECLARE_INDEXED_PROPERTYCOLLECTION_NAME(RestAngleType, int32);
 };
 
 }  // End namespace Chaos::Softs
diff --git a/Engine/Source/Runtime/Experimental/Chaos/Public/Chaos/PBDCollisionSpringConstraints.h b/Engine/Source/Runtime/Experimental/Chaos/Public/Chaos/PBDCollisionSpringConstraints.h
--- a/Engine/Source/Runtime/Experimental/Chaos/Public/Chaos/PBDCollisionSpringConstraints.h
+++ b/Engine/Source/Runtime/Experimental/Chaos/Public/Chaos/PBDCollisionSpringConstraints.h
@@ -19,8 +19,9 @@ public:
 
 	static bool IsEnabled(const FCollectionPropertyConstFacade& PropertyCollection)
 	{
-		return IsSelfCollisionStiffnessEnabled(PropertyCollection, false);
+		return GetUseSelfCollisions(PropertyCollection, false);
 	}
+
 	FPBDCollisionSpringConstraints(
 		const int32 InOffset,
 		const int32 InNumParticles,
@@ -182,6 +183,11 @@ public:
 		SetProperties(PropertyCollection, TMap<FString, TConstArrayView<FRealSingle>>(), TMap<FString, TConstArrayView<int32>>());
 	}
 
+	UE_CHAOS_DECLARE_INDEXLESS_PROPERTYCOLLECTION_NAME(UseSelfCollisions, bool);
+	UE_CHAOS_DECLARE_INDEXLESS_PROPERTYCOLLECTION_NAME(SelfCollisionThickness, float);
+	UE_CHAOS_DECLARE_INDEXLESS_PROPERTYCOLLECTION_NAME(SelfCollisionStiffness, float);
+	UE_CHAOS_DECLARE_INDEXLESS_PROPERTYCOLLECTION_NAME(SelfCollisionFriction, float);
+	UE_CHAOS_DECLARE_INDEXLESS_PROPERTYCOLLECTION_NAME(SelfCollisionLayers, float);
 private:
 	using Base::ThicknessWeighted;
 	using Base::Stiffness;
@@ -209,12 +215,13 @@ private:
 		return DisabledCollisionElements;
 	}
 
-	UE_CHAOS_DECLARE_PROPERTYCOLLECTION_NAME(SelfCollisionThickness, float);
-	UE_CHAOS_DECLARE_PROPERTYCOLLECTION_NAME(SelfCollisionStiffness, float);
-	UE_CHAOS_DECLARE_PROPERTYCOLLECTION_NAME(SelfCollisionFriction, float);
+	UE_CHAOS_DECLARE_INDEXED_PROPERTYCOLLECTION_NAME(SelfCollisionThickness, float);
+	UE_CHAOS_DECLARE_INDEXED_PROPERTYCOLLECTION_NAME(SelfCollisionStiffness, float);
+	UE_CHAOS_DECLARE_INDEXED_PROPERTYCOLLECTION_NAME(SelfCollisionFriction, float);
+
 	UE_CHAOS_DECLARE_PROPERTYCOLLECTION_NAME(SelfCollisionProximityStiffness, float);
 	UE_CHAOS_DECLARE_INDEXLESS_PROPERTYCOLLECTION_NAME(SelfCollisionDisableNeighborDistance, int32);
-	UE_CHAOS_DECLARE_PROPERTYCOLLECTION_NAME(SelfCollisionLayers, float);
+	UE_CHAOS_DECLARE_INDEXED_PROPERTYCOLLECTION_NAME(SelfCollisionLayers, float);
 	UE_CHAOS_DECLARE_PROPERTYCOLLECTION_NAME(SelfCollideAgainstKinematicCollidersOnly, bool);
 	UE_CHAOS_DECLARE_PROPERTYCOLLECTION_NAME(SelfCollisionKinematicColliderThickness, float);
 	UE_CHAOS_DECLARE_PROPERTYCOLLECTION_NAME(SelfCollisionKinematicColliderStiffness, float);
diff --git a/Engine/Source/Runtime/Experimental/Chaos/Public/Chaos/PBDLongRangeConstraints.h b/Engine/Source/Runtime/Experimental/Chaos/Public/Chaos/PBDLongRangeConstraints.h
--- a/Engine/Source/Runtime/Experimental/Chaos/Public/Chaos/PBDLongRangeConstraints.h
+++ b/Engine/Source/Runtime/Experimental/Chaos/Public/Chaos/PBDLongRangeConstraints.h
@@ -99,6 +99,9 @@ public:
 	template<typename SolverParticlesOrRange>
 	CHAOS_API void Apply(SolverParticlesOrRange& Particles, const FSolverReal Dt) const;
 
+	UE_CHAOS_DECLARE_INDEXLESS_PROPERTYCOLLECTION_NAME(TetherStiffness, float);
+	UE_CHAOS_DECLARE_INDEXLESS_PROPERTYCOLLECTION_NAME(TetherScale, float);
+
 private:
 	using Base::MinTetherScale;
 	using Base::MaxTetherScale;
@@ -108,8 +111,8 @@ private:
 	using Base::ParticleOffset;
 	using Base::ParticleCount;
 
-	UE_CHAOS_DECLARE_PROPERTYCOLLECTION_NAME(TetherStiffness, float);
-	UE_CHAOS_DECLARE_PROPERTYCOLLECTION_NAME(TetherScale, float);
+	UE_CHAOS_DECLARE_INDEXED_PROPERTYCOLLECTION_NAME(TetherStiffness, float);
+	UE_CHAOS_DECLARE_INDEXED_PROPERTYCOLLECTION_NAME(TetherScale, float);
 };
 
 }  // End namespace Chaos::Softs
diff --git a/Engine/Source/Runtime/Experimental/Chaos/Public/Chaos/PBDSelfCollisionSphereConstraints.h b/Engine/Source/Runtime/Experimental/Chaos/Public/Chaos/PBDSelfCollisionSphereConstraints.h
--- a/Engine/Source/Runtime/Experimental/Chaos/Public/Chaos/PBDSelfCollisionSphereConstraints.h
+++ b/Engine/Source/Runtime/Experimental/Chaos/Public/Chaos/PBDSelfCollisionSphereConstraints.h
@@ -75,14 +75,17 @@ namespace Chaos::Softs
 			const FCollectionPropertyConstFacade& PropertyCollection,
 			const TMap<FString, const TSet<int32>*>& VertexSets);
 
+		UE_CHAOS_DECLARE_INDEXLESS_PROPERTYCOLLECTION_NAME(SelfCollisionSphereSetName, bool);  // Selection set name string property, the bool value is not actually used
+		UE_CHAOS_DECLARE_INDEXLESS_PROPERTYCOLLECTION_NAME(SelfCollisionSphereRadius, float);
+		UE_CHAOS_DECLARE_INDEXLESS_PROPERTYCOLLECTION_NAME(SelfCollisionSphereStiffness, float);
+
 	private:
 		using Base::Radius;
 		using Base::Stiffness;
 		using Base::VertexSetNoOffset;
 
-		UE_CHAOS_DECLARE_INDEXLESS_PROPERTYCOLLECTION_NAME(SelfCollisionSphereSetName, bool);  // Selection set name string property, the bool value is not actually used
-		UE_CHAOS_DECLARE_PROPERTYCOLLECTION_NAME(SelfCollisionSphereRadius, float);
-		UE_CHAOS_DECLARE_PROPERTYCOLLECTION_NAME(SelfCollisionSphereStiffness, float);
+		UE_CHAOS_DECLARE_INDEXED_PROPERTYCOLLECTION_NAME(SelfCollisionSphereRadius, float);
+		UE_CHAOS_DECLARE_INDEXED_PROPERTYCOLLECTION_NAME(SelfCollisionSphereStiffness, float);
 	};
 }
 #endif
diff --git a/Engine/Source/Runtime/Experimental/Chaos/Public/Chaos/PBDSphericalConstraint.h b/Engine/Source/Runtime/Experimental/Chaos/Public/Chaos/PBDSphericalConstraint.h
--- a/Engine/Source/Runtime/Experimental/Chaos/Public/Chaos/PBDSphericalConstraint.h
+++ b/Engine/Source/Runtime/Experimental/Chaos/Public/Chaos/PBDSphericalConstraint.h
@@ -108,6 +108,8 @@ public:
 	void SetScale(FSolverReal InScale) { Scale = InScale; }
 	FSolverReal GetScale() const { return Scale; }
 
+	UE_CHAOS_DECLARE_INDEXLESS_PROPERTYCOLLECTION_NAME(MaxDistance, float);
+
 private:
 
 	template<bool bHasMaxDistance, typename SolverParticlesOrRange>
@@ -153,7 +155,7 @@ private:
 	FSolverReal MaxDistanceBase = (FSolverReal)0.;
 	FSolverReal MaxDistanceRange = (FSolverReal)1.;
 
-	UE_CHAOS_DECLARE_PROPERTYCOLLECTION_NAME(MaxDistance, float);
+	UE_CHAOS_DECLARE_INDEXED_PROPERTYCOLLECTION_NAME(MaxDistance, float);
 };
 
 class FPBDSphericalBackstopConstraint final
@@ -326,6 +328,10 @@ public:
 		return SphereOffsetDistances.Num() == ParticleCount ? BackstopDistanceBase + BackstopDistanceRange * SphereOffsetDistances[ConstraintIndex] : BackstopDistanceBase;
 	}
 
+	UE_CHAOS_DECLARE_INDEXLESS_PROPERTYCOLLECTION_NAME(BackstopDistance, float);
+	UE_CHAOS_DECLARE_INDEXLESS_PROPERTYCOLLECTION_NAME(BackstopRadius, float);
+	UE_CHAOS_DECLARE_INDEXLESS_PROPERTYCOLLECTION_NAME(UseLegacyBackstop, bool);
+
 private:
 
 	template<bool bHasBackstopDistance, bool bHasBackstopRadius, typename SolverParticlesOrRange>
@@ -424,9 +430,9 @@ private:
 	bool bEnabled = true;
 	bool bUseLegacyBackstop = false;
 
-	UE_CHAOS_DECLARE_PROPERTYCOLLECTION_NAME(BackstopDistance, float);
-	UE_CHAOS_DECLARE_PROPERTYCOLLECTION_NAME(BackstopRadius, float);
-	UE_CHAOS_DECLARE_PROPERTYCOLLECTION_NAME(UseLegacyBackstop, bool);
+	UE_CHAOS_DECLARE_INDEXED_PROPERTYCOLLECTION_NAME(BackstopDistance, float);
+	UE_CHAOS_DECLARE_INDEXED_PROPERTYCOLLECTION_NAME(BackstopRadius, float);
+	UE_CHAOS_DECLARE_INDEXED_PROPERTYCOLLECTION_NAME(UseLegacyBackstop, bool);
 };
 
 }  // End namespace Chaos::Softs
diff --git a/Engine/Source/Runtime/Experimental/Chaos/Public/Chaos/PBDSpringConstraints.h b/Engine/Source/Runtime/Experimental/Chaos/Public/Chaos/PBDSpringConstraints.h
--- a/Engine/Source/Runtime/Experimental/Chaos/Public/Chaos/PBDSpringConstraints.h
+++ b/Engine/Source/Runtime/Experimental/Chaos/Public/Chaos/PBDSpringConstraints.h
@@ -163,6 +163,10 @@ public:
 
 	CHAOS_API void ResetRestLengths(const TConstArrayView<FSolverVec3>& Positions);
 
+	UE_CHAOS_DECLARE_INDEXLESS_PROPERTYCOLLECTION_NAME(EdgeSpringStiffness, float);
+	UE_CHAOS_DECLARE_INDEXLESS_PROPERTYCOLLECTION_NAME(EdgeSpringWarpScale, float);
+	UE_CHAOS_DECLARE_INDEXLESS_PROPERTYCOLLECTION_NAME(EdgeSpringWeftScale, float);
+
 private:
 
 	void InitFromPatternData(const TArray<TVec3<FVec2f>>& FaceVertexPatternPositions, const FTriangleMesh& TriangleMesh);
@@ -176,9 +180,9 @@ private:
 	TArray<FSolverReal> BaseDists; // Without Warp/Weft Scale applied
 	TArray<FSolverVec2> WarpWeftScaleBaseMultipliers;
 
-	UE_CHAOS_DECLARE_PROPERTYCOLLECTION_NAME(EdgeSpringStiffness, float);
-	UE_CHAOS_DECLARE_PROPERTYCOLLECTION_NAME(EdgeSpringWarpScale, float);
-	UE_CHAOS_DECLARE_PROPERTYCOLLECTION_NAME(EdgeSpringWeftScale, float);
+	UE_CHAOS_DECLARE_INDEXED_PROPERTYCOLLECTION_NAME(EdgeSpringStiffness, float);
+	UE_CHAOS_DECLARE_INDEXED_PROPERTYCOLLECTION_NAME(EdgeSpringWarpScale, float);
+	UE_CHAOS_DECLARE_INDEXED_PROPERTYCOLLECTION_NAME(EdgeSpringWeftScale, float);
 };
 
 class FPBDBendingSpringConstraints final : public FPBDSpringConstraints
@@ -233,10 +237,12 @@ public:
 		const FCollectionPropertyConstFacade& PropertyCollection,
 		const TMap<FString, TConstArrayView<FRealSingle>>& WeightMaps);
 
+	UE_CHAOS_DECLARE_INDEXLESS_PROPERTYCOLLECTION_NAME(BendingSpringStiffness, float);
+
 private:
 	using FPBDSpringConstraints::ParticleCount;
 
-	UE_CHAOS_DECLARE_PROPERTYCOLLECTION_NAME(BendingSpringStiffness, float);
+	UE_CHAOS_DECLARE_INDEXED_PROPERTYCOLLECTION_NAME(BendingSpringStiffness, float);
 };
 
 }  // End namespace Chaos::Softs
diff --git a/Engine/Source/Runtime/Experimental/Chaos/Public/Chaos/PBDTriangleMeshCollisions.h b/Engine/Source/Runtime/Experimental/Chaos/Public/Chaos/PBDTriangleMeshCollisions.h
--- a/Engine/Source/Runtime/Experimental/Chaos/Public/Chaos/PBDTriangleMeshCollisions.h
+++ b/Engine/Source/Runtime/Experimental/Chaos/Public/Chaos/PBDTriangleMeshCollisions.h
@@ -190,7 +190,7 @@ public:
 
 	static bool IsEnabled(const FCollectionPropertyConstFacade& PropertyCollection)
 	{
-		return PropertyCollection.IsEnabled(FName(TEXT("SelfCollisionStiffness")).ToString(), false);  // Don't use UE_CHAOS_DECLARE_PROPERTYCOLLECTION_NAME here, SelfCollisionStiffness is only needed for activation
+		return GetUseSelfCollisions(PropertyCollection, false);
 	}
 
 	FPBDTriangleMeshCollisions(
@@ -403,6 +403,9 @@ public:
 	// for debug draw purposes.
 	const TArray<FContourMinimizationIntersection>& GetPostStepContourMinimizationIntersections() const { return PostStepContourMinimizationIntersections; }
 	const TArray<TArray<FBarycentricPoint>>& GetPostStepIntersectionContourPoints() const { return PostStepIntersectionContourPoints; }
+
+
+	UE_CHAOS_DECLARE_INDEXLESS_PROPERTYCOLLECTION_NAME(UseSelfIntersections, bool);
 private:
 
 	FTriangleSubMesh CollidableSubMesh;
@@ -439,7 +442,8 @@ private:
 	TArray<FContourMinimizationIntersection> PostStepContourMinimizationIntersections;
 	TArray<TArray<FBarycentricPoint>> PostStepIntersectionContourPoints;
 
-	UE_CHAOS_DECLARE_PROPERTYCOLLECTION_NAME(UseSelfIntersections, bool);
+	UE_CHAOS_DECLARE_INDEXLESS_PROPERTYCOLLECTION_NAME(UseSelfCollisions, bool);
+	UE_CHAOS_DECLARE_INDEXED_PROPERTYCOLLECTION_NAME(UseSelfIntersections, bool);
 	UE_CHAOS_DECLARE_PROPERTYCOLLECTION_NAME(UseGlobalIntersectionAnalysis, bool);
 	UE_CHAOS_DECLARE_PROPERTYCOLLECTION_NAME(UseContourMinimization, bool);
 	UE_CHAOS_DECLARE_PROPERTYCOLLECTION_NAME(NumContourMinimizationPostSteps, int32);
diff --git a/Engine/Source/Runtime/Experimental/Chaos/Public/Chaos/PBDTriangleMeshIntersections.h b/Engine/Source/Runtime/Experimental/Chaos/Public/Chaos/PBDTriangleMeshIntersections.h
--- a/Engine/Source/Runtime/Experimental/Chaos/Public/Chaos/PBDTriangleMeshIntersections.h
+++ b/Engine/Source/Runtime/Experimental/Chaos/Public/Chaos/PBDTriangleMeshIntersections.h
@@ -14,7 +14,7 @@ class FPBDTriangleMeshIntersections
 public:
 	static bool IsEnabled(const FCollectionPropertyConstFacade& PropertyCollection)
 	{
-		return PropertyCollection.IsEnabled(FName(TEXT("SelfCollisionStiffness")).ToString(), false);  // Don't use UE_CHAOS_DECLARE_PROPERTYCOLLECTION_NAME here, SelfCollisionStiffness is only needed for activation
+		return FPBDTriangleMeshCollisions::IsEnabled(PropertyCollection);
 	}
 
 	FPBDTriangleMeshIntersections(
diff --git a/Engine/Source/Runtime/Experimental/Chaos/Public/Chaos/VelocityField.h b/Engine/Source/Runtime/Experimental/Chaos/Public/Chaos/VelocityField.h
--- a/Engine/Source/Runtime/Experimental/Chaos/Public/Chaos/VelocityField.h
+++ b/Engine/Source/Runtime/Experimental/Chaos/Public/Chaos/VelocityField.h
@@ -198,6 +198,15 @@ public:
 	// This method is currently used for debug drawing.
 	CHAOS_API FSolverVec3 CalculateForce(const TConstArrayView<FSolverVec3>& Xs, const TConstArrayView<FSolverVec3>& Vs, int32 ElementIndex) const;
 
+	UE_CHAOS_DECLARE_INDEXLESS_PROPERTYCOLLECTION_NAME(Drag, float);
+	UE_CHAOS_DECLARE_INDEXLESS_PROPERTYCOLLECTION_NAME(OuterDrag, float);
+	UE_CHAOS_DECLARE_INDEXLESS_PROPERTYCOLLECTION_NAME(Lift, float);
+	UE_CHAOS_DECLARE_INDEXLESS_PROPERTYCOLLECTION_NAME(OuterLift, float);
+	UE_CHAOS_DECLARE_INDEXLESS_PROPERTYCOLLECTION_NAME(FluidDensity, float);
+	UE_CHAOS_DECLARE_INDEXLESS_PROPERTYCOLLECTION_NAME(Pressure, float);
+	UE_CHAOS_DECLARE_INDEXLESS_PROPERTYCOLLECTION_NAME(WindVelocity, FVector3f);
+	UE_CHAOS_DECLARE_INDEXLESS_PROPERTYCOLLECTION_NAME(WindVelocitySpace, int32);
+
 private:
 	bool AreAerodynamicsEnabled() const { return QuarterRho > (FSolverReal)0.; }
 
@@ -307,14 +316,14 @@ private:
 	FSolverReal Rho;
 	FSolverReal QuarterRho;
 
-	UE_CHAOS_DECLARE_PROPERTYCOLLECTION_NAME(Drag, float);
-	UE_CHAOS_DECLARE_PROPERTYCOLLECTION_NAME(OuterDrag, float);
-	UE_CHAOS_DECLARE_PROPERTYCOLLECTION_NAME(Lift, float);
-	UE_CHAOS_DECLARE_PROPERTYCOLLECTION_NAME(OuterLift, float);
-	UE_CHAOS_DECLARE_PROPERTYCOLLECTION_NAME(FluidDensity, float);
-	UE_CHAOS_DECLARE_PROPERTYCOLLECTION_NAME(Pressure, float);
-	UE_CHAOS_DECLARE_PROPERTYCOLLECTION_NAME(WindVelocity, FVector3f);
-	UE_CHAOS_DECLARE_PROPERTYCOLLECTION_NAME(WindVelocitySpace, int32);
+	UE_CHAOS_DECLARE_INDEXED_PROPERTYCOLLECTION_NAME(Drag, float);
+	UE_CHAOS_DECLARE_INDEXED_PROPERTYCOLLECTION_NAME(OuterDrag, float);
+	UE_CHAOS_DECLARE_INDEXED_PROPERTYCOLLECTION_NAME(Lift, float);
+	UE_CHAOS_DECLARE_INDEXED_PROPERTYCOLLECTION_NAME(OuterLift, float);
+	UE_CHAOS_DECLARE_INDEXED_PROPERTYCOLLECTION_NAME(FluidDensity, float);
+	UE_CHAOS_DECLARE_INDEXED_PROPERTYCOLLECTION_NAME(Pressure, float);
+	UE_CHAOS_DECLARE_INDEXED_PROPERTYCOLLECTION_NAME(WindVelocity, FVector3f);
+	UE_CHAOS_DECLARE_INDEXED_PROPERTYCOLLECTION_NAME(WindVelocitySpace, int32);
 };
 
 }  // End namespace Chaos::Softs
diff --git a/Engine/Source/Runtime/Experimental/Chaos/Public/GeometryCollection/ManagedArrayTypeValues.inl b/Engine/Source/Runtime/Experimental/Chaos/Public/GeometryCollection/ManagedArrayTypeValues.inl
--- a/Engine/Source/Runtime/Experimental/Chaos/Public/GeometryCollection/ManagedArrayTypeValues.inl
+++ b/Engine/Source/Runtime/Experimental/Chaos/Public/GeometryCollection/ManagedArrayTypeValues.inl
@@ -62,6 +62,7 @@ MANAGED_ARRAY_TYPE(TArray<TArray<FVector3f>>, FVector3fNestedArray)
 MANAGED_ARRAY_TYPE(FUintVector2, UintVector2)
 MANAGED_ARRAY_TYPE(TObjectPtr<UObject>, UObjectArray)
 MANAGED_ARRAY_TYPE(Chaos::FLinearCurve, LinearCurve)
+MANAGED_ARRAY_TYPE(FName, Name)
 // NOTE: new types must be added at the bottom to keep serialization from breaking
 
 
