Fixing perf regression from UE4 for some queries in Physics.
- Swept SphereVsSphere
- Swept SphereVsCapsule
- Swept CapsuleVsSphere

diff --git a/Engine/Source/Runtime/Experimental/Chaos/Private/Chaos/GeometryQueries.cpp b/Engine/Source/Runtime/Experimental/Chaos/Private/Chaos/GeometryQueries.cpp
new file mode 100644
--- /dev/null
+++ b/Engine/Source/Runtime/Experimental/Chaos/Private/Chaos/GeometryQueries.cpp
@@ -0,0 +1,94 @@
+// Copyright Epic Games, Inc. All Rights Reserved.
+#include "Chaos/GeometryQueries.h"
+
+#include "Chaos/Box.h"
+#include "Chaos/Capsule.h"
+#include "Chaos/CastingUtilities.h"
+#include "Chaos/ImplicitObject.h"
+#include "Chaos/Sphere.h"
+#include "Chaos/Sweeps.h"
+#include "Chaos/Transform.h"
+
+namespace Chaos
+{
+	void TransformSweepResultsToWorld(const bool bResult, const FReal Time, const bool bComputeMTD, const FImplicitObject& TestGeom, const FVec3& TestGeomLocation, const FVec3& LocalDir, const FVec3& LocalPosition, const FVec3& LocalNormal, const int32 FaceIndex,
+		FVec3& OutWorldPosition, FVec3& OutWorldNormal, FVec3& OutWorldFaceNormal)
+	{
+		if (bResult && (Time >= 0 || bComputeMTD))
+		{
+			OutWorldPosition = TestGeomLocation + LocalPosition;
+			OutWorldFaceNormal = TestGeom.FindGeometryOpposingNormal(LocalDir, FaceIndex, LocalNormal);
+			OutWorldNormal = LocalNormal.GetSafeNormal(UE_KINDA_SMALL_NUMBER, FVec3::AxisVector(0));
+		}
+	}
+
+	void TransformSweepResultsToWorld(const bool bResult, const FReal Time, const bool bComputeMTD, const FImplicitObject& TestGeom, const FRigidTransform3& TestGeomTM, const FVec3& LocalDir, const FVec3& LocalPosition, const FVec3& LocalNormal, const int32 FaceIndex,
+		FVec3& OutWorldPosition, FVec3& OutWorldNormal, FVec3& OutWorldFaceNormal)
+	{
+		if (bResult && (Time >= 0 || bComputeMTD))
+		{
+			OutWorldPosition = TestGeomTM.TransformPositionNoScale(LocalPosition);
+			OutWorldFaceNormal = TestGeomTM.TransformVectorNoScale(TestGeom.FindGeometryOpposingNormal(LocalDir, FaceIndex, LocalNormal));
+			OutWorldNormal = TestGeomTM.TransformVectorNoScale(LocalNormal).GetSafeNormal(UE_KINDA_SMALL_NUMBER, FVec3::AxisVector(0));
+		}
+	}
+
+	bool SweepSphereVsSphere(const FSphere& SweptSphere, const FRigidTransform3& SweptSphereTM, const FSphere& TestSphere, const FRigidTransform3& TestSphereTM,
+		const FVec3& SweepDir, const FReal Length, const FReal Thickness, const bool bComputeMTD,
+		FReal& OutTime, FVec3& OutPosition, FVec3& OutNormal, int32& OutFaceIndex, FVec3& OutFaceNormal)
+	{
+		const Sweeps::ESweepFlags Flags = bComputeMTD ? Sweeps::ESweepFlags::MTD : Sweeps::ESweepFlags::None;
+		// Transform to the test shape's local space. Note: We can skip rotation which is much faster
+		const FVec3 TestGeomLocalToWorld = TestSphereTM.GetLocation();
+		const FVec3 SweptShapeLocalToWorld = SweptSphereTM.GetLocation();
+		const FVec3 LocalSweepDir = SweepDir;
+		const FVec3 LocalSweepStart = SweptSphere.GetCenterf() + SweptShapeLocalToWorld - TestGeomLocalToWorld;
+		const FReal SweepRadius = Thickness + SweptSphere.GetRadiusf();
+
+		const bool bResult = Sweeps::SweepSphereVsSphere(LocalSweepStart, LocalSweepDir, Length, SweepRadius, TestSphere.GetCenterf(), TestSphere.GetRadiusf(), Flags, OutTime, OutPosition, OutNormal);
+		TransformSweepResultsToWorld(bResult, OutTime, bComputeMTD, TestSphere, TestGeomLocalToWorld, LocalSweepDir, OutPosition, OutNormal, OutFaceIndex, OutPosition, OutNormal, OutFaceNormal);
+		return bResult;
+	}
+
+	bool SweepSphereVsCapsule(const FSphere& SweptSphere, const FRigidTransform3& SweptSphereTM, const FCapsule& TestCapsule, const FRigidTransform3& TestCapsuleTM,
+		const FVec3& SweepDir, const FReal Length, const FReal Thickness, const bool bComputeMTD,
+		FReal& OutTime, FVec3& OutPosition, FVec3& OutNormal, int32& OutFaceIndex, FVec3& OutFaceNormal)
+	{
+		const Sweeps::ESweepFlags Flags = bComputeMTD ? Sweeps::ESweepFlags::MTD : Sweeps::ESweepFlags::None;
+		// Bring both shapes into "world space" but offset so the capsule is at the origin
+		const FVec3 SphereLocation = SweptSphereTM.GetLocation();
+		const FVec3 CapsuleLocation = TestCapsuleTM.GetLocation();
+		const FVec3 CapsuleWorldAxis = TestCapsuleTM.TransformVectorNoScale(FVec3(TestCapsule.GetAxisf()));
+		const FVec3 CapsuleWorldX1 = TestCapsuleTM.TransformPositionNoScale(FVec3(TestCapsule.GetX1f())) - CapsuleLocation;
+		const FVec3 CapsuleWorldX2 = TestCapsuleTM.TransformPositionNoScale(FVec3(TestCapsule.GetX2f())) - CapsuleLocation;
+		const FVec3 WorldSweepDir = SweepDir;
+		const FVec3 WorldSweepStart = SweptSphere.GetCenterf() + SphereLocation - CapsuleLocation;
+		const FReal SweepRadius = Thickness + SweptSphere.GetRadiusf();
+
+		const bool bResult = Sweeps::SweepSphereVsCapsule(WorldSweepStart, WorldSweepDir, Length, SweepRadius, TestCapsule.GetRadiusf(), TestCapsule.GetHeightf(), CapsuleWorldAxis, CapsuleWorldX1, CapsuleWorldX2, Flags, OutTime, OutPosition, OutNormal);
+		TransformSweepResultsToWorld(bResult, OutTime, bComputeMTD, TestCapsule, CapsuleLocation, WorldSweepDir, OutPosition, OutNormal, OutFaceIndex, OutPosition, OutNormal, OutFaceNormal);
+		return bResult;
+	}
+
+	bool SweepCapsuleVsSphere(const FCapsule& SweptCapsule, const FRigidTransform3& SweptCapsuleTM, const FSphere& TestSphere, const FRigidTransform3& TestSphereTM,
+		const FVec3& SweepDir, const FReal Length, const FReal Thickness, const bool bComputeMTD,
+		FReal& OutTime, FVec3& OutPosition, FVec3& OutNormal, int32& OutFaceIndex, FVec3& OutFaceNormal)
+	{
+		// Convert this into sweeping the sphere against the capsule. This is identical except for the resultant impact info.
+		const FVec3 WorldSweepDir = -SweepDir;
+		bool bResult = SweepSphereVsCapsule(TestSphere, TestSphereTM, SweptCapsule, SweptCapsuleTM, WorldSweepDir, Length, Thickness, bComputeMTD, OutTime, OutPosition, OutNormal, OutFaceIndex, OutFaceNormal);
+
+		// Now convert the results from being "on the capsule" to being "on the sphere"
+		if (bResult)
+		{
+			const FVec3 SphereLocation = TestSphereTM.GetLocation();
+			const FReal SphereRadius = TestSphere.GetRadiusf();
+
+			OutNormal *= -1;
+			// Reconstruct the position on the sphere in local space (we transform everything below)
+			OutPosition = WorldSweepDir * SphereRadius;
+			TransformSweepResultsToWorld(bResult, OutTime, bComputeMTD, SweptCapsule, SphereLocation, WorldSweepDir, OutPosition, OutNormal, OutFaceIndex, OutPosition, OutNormal, OutFaceNormal);
+		}
+		return bResult;
+	}
+}
diff --git a/Engine/Source/Runtime/Experimental/Chaos/Public/Chaos/GeometryQueries.h b/Engine/Source/Runtime/Experimental/Chaos/Public/Chaos/GeometryQueries.h
--- a/Engine/Source/Runtime/Experimental/Chaos/Public/Chaos/GeometryQueries.h
+++ b/Engine/Source/Runtime/Experimental/Chaos/Public/Chaos/GeometryQueries.h
@@ -152,6 +152,45 @@ namespace Chaos
 		}
 	}
 
+	UE_INTERNAL CHAOS_API void TransformSweepResultsToWorld(const bool bResult, const FReal Time, const bool bComputeMTD, const FImplicitObject& TestGeom, const FRigidTransform3& TestGeomTM, const FVec3& LocalDir, const FVec3& LocalPosition, const FVec3& LocalNormal, const int32 FaceIndex, FVec3& OutWorldPosition, FVec3& OutWorldNormal, FVec3& OutWorldFaceNormal);
+	UE_INTERNAL CHAOS_API bool SweepSphereVsSphere(const FSphere& SweptSphere, const FRigidTransform3& SweptSphereTM, const FSphere& TestSphere, const FRigidTransform3& TestSphereTM, const FVec3& SweepDir, const FReal Length, const FReal Thickness, const bool bComputeMTD, FReal& OutTime, FVec3& OutPosition, FVec3& OutNormal, int32& OutFaceIndex, FVec3& OutFaceNormal);
+	UE_INTERNAL CHAOS_API bool SweepSphereVsCapsule(const FSphere& SweptSphere, const FRigidTransform3& SweptSphereTM, const FCapsule& TestCapsule, const FRigidTransform3& TestCapsuleTM, const FVec3& SweepDir, const FReal Length, const FReal Thickness, const bool bComputeMTD, FReal& OutTime, FVec3& OutPosition, FVec3& OutNormal, int32& OutFaceIndex, FVec3& OutFaceNormal);
+	UE_INTERNAL CHAOS_API bool SweepCapsuleVsSphere(const FCapsule& SweptCapsule, const FRigidTransform3& SweptCapsuleTM, const FSphere& TestSphere, const FRigidTransform3& TestSphereTM, const FVec3& SweepDir, const FReal Length, const FReal Thickness, const bool bComputeMTD, FReal& OutTime, FVec3& OutPosition, FVec3& OutNormal, int32& OutFaceIndex, FVec3& OutFaceNormal);
+
+	template <EImplicitObjectType SweptType>
+	UE_INTERNAL bool TryRunSpecializedSweep(const FImplicitObject& SweptShape, const FRigidTransform3& SweptShapeTM, const FImplicitObject& TestGeom, const FRigidTransform3& TestGeomTM, const FVec3& SweepDir, const FReal Length, const FReal Thickness, const bool bComputeMTD, bool& bOutResult, FReal& OutTime, FVec3& OutPosition, FVec3& OutNormal, int32& OutFaceIndex, FVec3& OutFaceNormal)
+	{
+		const EImplicitObjectType TestType = TestGeom.GetType();
+		if (SweptType == ImplicitObjectType::Sphere)
+		{
+			if (TestType == ImplicitObjectType::Sphere)
+			{
+				const FSphere& SweptSphere = SweptShape.GetObjectChecked<FSphere>();
+				const FSphere& TestSphere = TestGeom.GetObjectChecked<FSphere>();
+				bOutResult = SweepSphereVsSphere(SweptSphere, SweptShapeTM, TestSphere, TestGeomTM, SweepDir, Length, Thickness, bComputeMTD, OutTime, OutPosition, OutNormal, OutFaceIndex, OutFaceNormal);
+				return true;
+			}
+			else if (TestType == ImplicitObjectType::Capsule)
+			{
+				const FSphere& SweptSphere = SweptShape.GetObjectChecked<FSphere>();
+				const FCapsule& TestCapsule = TestGeom.GetObjectChecked<FCapsule>();
+				bOutResult = SweepSphereVsCapsule(SweptSphere, SweptShapeTM, TestCapsule, TestGeomTM, SweepDir, Length, Thickness, bComputeMTD, OutTime, OutPosition, OutNormal, OutFaceIndex, OutFaceNormal);
+				return true;
+			}
+		}
+		else if (SweptType == ImplicitObjectType::Capsule)
+		{
+			if (TestType == ImplicitObjectType::Sphere)
+			{
+				const FCapsule& SweptCapsule = SweptShape.GetObjectChecked<FCapsule>();
+				const FSphere& TestSphere = TestGeom.GetObjectChecked<FSphere>();
+				bOutResult = SweepCapsuleVsSphere(SweptCapsule, SweptShapeTM, TestSphere, TestGeomTM, SweepDir, Length, Thickness, bComputeMTD, OutTime, OutPosition, OutNormal, OutFaceIndex, OutFaceNormal);
+				return true;
+			}
+		}
+		return false;
+	}
+
 	template <typename SweptGeometry>
 	bool SweepQuery(const FImplicitObject& A, const FRigidTransform3& ATM, const SweptGeometry& B, const FRigidTransform3& BTM, const FVec3& Dir, const FReal Length, FReal& OutTime, FVec3& OutPosition, FVec3& OutNormal, int32& OutFaceIndex, FVec3& OutFaceNormal, const FReal Thickness, const bool bComputeMTD)
 	{
@@ -168,7 +207,13 @@ namespace Chaos
 		}
 
 		OutFaceIndex = INDEX_NONE;
-		
+		PRAGMA_DISABLE_INTERNAL_WARNINGS
+		if (TryRunSpecializedSweep<BType>(B, BTM, A, ATM, Dir, Length, Thickness, bComputeMTD, bResult, OutTime, OutPosition, OutNormal, OutFaceIndex, OutFaceNormal))
+		{
+			return bResult;
+		}
+		PRAGMA_ENABLE_INTERNAL_WARNINGS
+
 		FVec3 LocalPosition(-TNumericLimits<FReal>::Max()); // Make it obvious when things go wrong
 		FVec3 LocalNormal(0);
 
@@ -312,12 +357,9 @@ namespace Chaos
 		}
 
 		//put back into world space
-		if (bResult && (OutTime >= 0 || bComputeMTD))
-		{
-			OutNormal = ATM.TransformVectorNoScale(LocalNormal).GetSafeNormal(UE_KINDA_SMALL_NUMBER, FVec3::AxisVector(0));
-			OutPosition = ATM.TransformPositionNoScale(LocalPosition);
-			OutFaceNormal = ATM.TransformVectorNoScale(A.FindGeometryOpposingNormal(LocalDir, OutFaceIndex, LocalNormal));
-		}
+		PRAGMA_DISABLE_INTERNAL_WARNINGS
+		TransformSweepResultsToWorld(bResult, OutTime, bComputeMTD, A, ATM, LocalDir, LocalPosition, LocalNormal, OutFaceIndex, OutPosition, OutNormal, OutFaceNormal);
+		PRAGMA_ENABLE_INTERNAL_WARNINGS
 
 		return bResult;
 	}
diff --git a/Engine/Source/Runtime/Experimental/ChaosCore/Private/Chaos/Sweeps.cpp b/Engine/Source/Runtime/Experimental/ChaosCore/Private/Chaos/Sweeps.cpp
new file mode 100644
--- /dev/null
+++ b/Engine/Source/Runtime/Experimental/ChaosCore/Private/Chaos/Sweeps.cpp
@@ -0,0 +1,42 @@
+// Copyright Epic Games, Inc. All Rights Reserved.
+#include "Chaos/Sweeps.h"
+
+#include "Chaos/Raycasts.h"
+#include "Chaos/SweepsMTD.h"
+
+namespace Chaos::Sweeps
+{
+	bool SweepSphereVsSphere(const FVec3& SweepStart, const FVec3& SweepDir, const FReal SweepLength, const FReal Sphere0Radius, const FVec3& Sphere1Center, const FReal Sphere1Radius, const ESweepFlags Flags, FReal& OutTime, FVec3& OutPosition, FVec3& OutNormal)
+	{
+		const bool bResult = Raycasts::RaySphere<FReal, 3>(SweepStart, SweepDir, SweepLength, Sphere0Radius, Sphere1Center, FRealSingle(Sphere1Radius), OutTime, OutPosition, OutNormal);
+		if (bResult && OutTime <= UE_DOUBLE_SMALL_NUMBER && EnumHasAnyFlags(Flags, ESweepFlags::MTD))
+		{
+			SphereSphereMTD(Sphere1Center, Sphere1Radius, SweepStart, Sphere0Radius, OutTime, OutPosition, OutNormal);
+		}
+		return bResult;
+	}
+
+	bool SweepSphereVsAabb(const FVec3& SweepStart, const FVec3& SweepDir, const FReal SweepLength, const FReal SphereRadius, const FVec3& AabbMin, const FVec3& AabbMax, const ESweepFlags Flags, FReal& OutTime, FVec3& OutPosition, FVec3& OutNormal, int32& OutFaceIndex)
+	{
+		const bool bResult = Raycasts::RayAabb<FReal, 3>(SweepStart, SweepDir, SweepLength, SphereRadius, AabbMin, AabbMax, OutTime, OutPosition, OutNormal, OutFaceIndex);
+		if (bResult && OutTime <= UE_DOUBLE_SMALL_NUMBER && EnumHasAnyFlags(Flags, ESweepFlags::MTD))
+		{
+			AabbSphereMTD(AabbMin, AabbMax, SweepStart, SphereRadius, OutTime, OutPosition, OutNormal);
+		}
+
+		return bResult;
+	}
+
+	bool SweepSphereVsCapsule(const FVec3& SweepStart, const FVec3& SweepDir, const FReal SweepLength, const FReal SphereRadius,
+		const FReal CapsuleRadius, const FReal CapsuleHeight, const FVec3& CapsuleAxis, const FVec3& CapsuleX1, const FVec3& CapsuleX2,
+		const ESweepFlags Flags, FReal& OutTime, FVec3& OutPosition, FVec3& OutNormal)
+	{
+		const bool bResult = Raycasts::RayCapsule(SweepStart, SweepDir, SweepLength, SphereRadius, CapsuleRadius, CapsuleHeight, CapsuleAxis, CapsuleX1, CapsuleX2, OutTime, OutPosition, OutNormal);
+		if (bResult && OutTime <= UE_DOUBLE_SMALL_NUMBER && EnumHasAnyFlags(Flags, ESweepFlags::MTD))
+		{
+			CapsuleSphereMTD(CapsuleX1, CapsuleX2, CapsuleRadius, SweepStart, SphereRadius, OutTime, OutPosition, OutNormal);
+		}
+
+		return bResult;
+	}
+} // namespace Chaos::Sweeps
diff --git a/Engine/Source/Runtime/Experimental/ChaosCore/Private/Chaos/SweepsMTD.cpp b/Engine/Source/Runtime/Experimental/ChaosCore/Private/Chaos/SweepsMTD.cpp
new file mode 100644
--- /dev/null
+++ b/Engine/Source/Runtime/Experimental/ChaosCore/Private/Chaos/SweepsMTD.cpp
@@ -0,0 +1,75 @@
+// Copyright Epic Games, Inc. All Rights Reserved.
+#include "Chaos/SweepsMTD.h"
+
+namespace Chaos::Sweeps
+{
+	void SphereSphereMTD(const FVec3& Sphere0Center, const FReal Sphere0Radius, const FVec3& Sphere1Center, const FReal Sphere1Radius, FReal& OutTime, FVec3& OutPosition, FVec3& OutNormal)
+	{
+		const FReal RadiusSum = Sphere0Radius + Sphere1Radius;
+		const FVec3 C01 = Sphere1Center - Sphere0Center;
+		const FReal Dist01 = C01.Size();
+		if (Dist01 > UE_DOUBLE_SMALL_NUMBER)
+		{
+			OutNormal = C01 / Dist01;
+		}
+		else
+		{
+			OutNormal = FVec3(0, 0, 1);
+		}
+		OutTime = Dist01 - RadiusSum;
+		OutPosition = Sphere0Center + OutNormal * Sphere0Radius;
+	}
+
+	void AabbSphereMTD(const FVec3& AabbMin, const FVec3& AabbMax, const FVec3& SphereCenter, const FReal SphereRadius, FReal& OutTime, FVec3& OutPosition, FVec3& OutNormal)
+	{
+		FVec3 ClampedCenter;
+		ClampedCenter.X = FMath::Clamp(SphereCenter.X, AabbMin.X, AabbMax.X);
+		ClampedCenter.Y = FMath::Clamp(SphereCenter.Y, AabbMin.Y, AabbMax.Y);
+		ClampedCenter.Z = FMath::Clamp(SphereCenter.Z, AabbMin.Z, AabbMax.Z);
+
+		const FVec3 DistVec = SphereCenter - ClampedCenter;
+		const FReal DistSq = DistVec.SquaredLength();
+		if (DistSq > UE_DOUBLE_SMALL_NUMBER)
+		{
+			// Sphere center was outside the aabb. The MTD is along the dist vec.
+			const FReal Dist = FMath::Sqrt(DistSq);
+			OutPosition = ClampedCenter;
+			OutNormal = DistVec / Dist;
+			OutTime = Dist - SphereRadius;
+		}
+		else
+		{
+			// Sphere center was inside the aabb. The MTD is along the aabb axis of deepest penetration.
+			const FVec3 AabbCenter = (AabbMax + AabbMin) * 0.5;
+			const FVec3 AabbHalfExtent = (AabbMax - AabbMin) * 0.5;
+			const FVec3 R = SphereCenter - AabbCenter;
+			const FVec3 AbsR = R.GetAbs();
+
+			int32 LargestAxis = 0;
+			FReal LargestAxisValue = AbsR.X;
+			for (int32 I = 1; I < 3; ++I)
+			{
+				if (AbsR[I] > LargestAxisValue)
+				{
+					LargestAxis = I;
+					LargestAxisValue = AbsR[I];
+				}
+			}
+			OutNormal = FVec3::Zero();
+			OutNormal[LargestAxis] = 1 * FMath::Sign(R[LargestAxis]);
+			OutPosition = AabbCenter + AabbHalfExtent * OutNormal;
+			// OutTime = (AabbR + SphereR) - Penetration
+			OutTime = -(AabbHalfExtent[LargestAxis] + SphereRadius - AbsR[LargestAxis]);
+		}
+	}
+
+	void CapsuleSphereMTD(const FVec3& CapsuleX1, const FVec3& CapsuleX2, const FReal CapsuleRadius, const FVec3& SphereCenter, const FReal SphereRadius, FReal& OutTime, FVec3& OutPosition, FVec3& OutNormal)
+	{
+		const FVec3 Axis = CapsuleX2 - CapsuleX1;
+		const FReal AxisLengthSq = Axis.SquaredLength();
+		const FReal T = FVec3::DotProduct(SphereCenter - CapsuleX1, Axis) / AxisLengthSq;
+		const FReal ClampedT = FMath::Clamp(T, 0, 1);
+		const FVec3 ProjectedPoint = CapsuleX1 + Axis * ClampedT;
+		return SphereSphereMTD(ProjectedPoint, CapsuleRadius, SphereCenter, SphereRadius, OutTime, OutPosition, OutNormal);
+	}
+} // namespace Chaos::Sweeps
diff --git a/Engine/Source/Runtime/Experimental/ChaosCore/Public/Chaos/Sweeps.h b/Engine/Source/Runtime/Experimental/ChaosCore/Public/Chaos/Sweeps.h
new file mode 100644
--- /dev/null
+++ b/Engine/Source/Runtime/Experimental/ChaosCore/Public/Chaos/Sweeps.h
@@ -0,0 +1,22 @@
+// Copyright Epic Games, Inc. All Rights Reserved.
+#pragma once
+
+#include "Chaos/Core.h"
+
+namespace Chaos::Sweeps
+{
+	enum class ESweepFlags
+	{
+		None = 0,
+		MTD = (1 << 9),
+	};
+	ENUM_CLASS_FLAGS(ESweepFlags);
+
+	CHAOSCORE_API bool SweepSphereVsSphere(const FVec3& SweepStart, const FVec3& SweepDir, const FReal SweepLength, const FReal Sphere0Radius, const FVec3& Sphere1Center, const FReal Sphere1Radius, const ESweepFlags Flags, FReal& OutTime, FVec3& OutPosition, FVec3& OutNormal);
+
+	CHAOSCORE_API bool SweepSphereVsAabb(const FVec3& SweepStart, const FVec3& SweepDir, const FReal SweepLength, const FReal SphereRadius, const FVec3& AabbMin, const FVec3& AabbMax, const ESweepFlags Flags, FReal& OutTime, FVec3& OutPosition, FVec3& OutNormal, int32& OutFaceIndex);
+
+	CHAOSCORE_API bool SweepSphereVsCapsule(const FVec3& SweepStart, const FVec3& SweepDir, const FReal SweepLength, const FReal SphereRadius,
+		const FReal CapsuleRadius, const FReal CapsuleHeight, const FVec3& CapsuleAxis, const FVec3& CapsuleX1, const FVec3& CapsuleX2,
+		const ESweepFlags Flags, FReal& OutTime, FVec3& OutPosition, FVec3& OutNormal);
+} // namespace Chaos::Sweeps
diff --git a/Engine/Source/Runtime/Experimental/ChaosCore/Public/Chaos/SweepsMTD.h b/Engine/Source/Runtime/Experimental/ChaosCore/Public/Chaos/SweepsMTD.h
new file mode 100644
--- /dev/null
+++ b/Engine/Source/Runtime/Experimental/ChaosCore/Public/Chaos/SweepsMTD.h
@@ -0,0 +1,13 @@
+// Copyright Epic Games, Inc. All Rights Reserved.
+#pragma once
+
+#include "Chaos/Core.h"
+
+namespace Chaos::Sweeps
+{
+	void SphereSphereMTD(const FVec3& Sphere0Center, const FReal Sphere0Radius, const FVec3& Sphere1Center, const FReal Sphere1Radius, FReal& OutTime, FVec3& OutPosition, FVec3& OutNormal);
+
+	void AabbSphereMTD(const FVec3& AabbMin, const FVec3& AabbMax, const FVec3& SphereCenter, const FReal SphereRadius, FReal& OutTime, FVec3& OutPosition, FVec3& OutNormal);
+
+	void CapsuleSphereMTD(const FVec3& CapsuleX1, const FVec3& CapsuleX2, const FReal CapsuleRadius, const FVec3& SphereCenter, const FReal SphereRadius, FReal& OutTime, FVec3& OutPosition, FVec3& OutNormal);
+} // namespace Chaos::Sweeps
