Fix CIS error from CollectionPropertyFacade changes.

diff --git a/Engine/Source/Runtime/Experimental/Chaos/Private/Chaos/CollectionPropertyFacade.cpp b/Engine/Source/Runtime/Experimental/Chaos/Private/Chaos/CollectionPropertyFacade.cpp
--- a/Engine/Source/Runtime/Experimental/Chaos/Private/Chaos/CollectionPropertyFacade.cpp
+++ b/Engine/Source/Runtime/Experimental/Chaos/Private/Chaos/CollectionPropertyFacade.cpp
@@ -75,6 +75,8 @@ namespace Chaos::Softs
 	template CHAOS_API const FVector3f& FCollectionPropertyConstFacade::GetValue<const FVector3f&, FVector3f>(int32 KeyIndex, const TConstArrayView<FVector3f>& ValueArray) const;
 	template CHAOS_API FString FCollectionPropertyConstFacade::GetValue<FString, FString>(int32 KeyIndex, const TConstArrayView<FString>& ValueArray) const;
 	template CHAOS_API const FString& FCollectionPropertyConstFacade::GetValue<const FString&, FString>(int32 KeyIndex, const TConstArrayView<FString>& ValueArray) const;
+	template CHAOS_API FName FCollectionPropertyConstFacade::GetValue<FName, FName>(int32 KeyIndex, const TConstArrayView<FName>& ValueArray) const;
+	template CHAOS_API const FName& FCollectionPropertyConstFacade::GetValue<const FName&, FName>(int32 KeyIndex, const TConstArrayView<FName>& ValueArray) const;
 	template CHAOS_API uint8 FCollectionPropertyConstFacade::GetValue<uint8, uint8>(int32 KeyIndex, const TConstArrayView<uint8>& ValueArray) const;
 	template CHAOS_API ECollectionPropertyFlags FCollectionPropertyConstFacade::GetValue<ECollectionPropertyFlags, uint8>(int32 KeyIndex, const TConstArrayView<uint8>& ValueArray) const;
 
