[Backout] - CL44488080
#rnx
Original CL Desc
-----------------------------------------------------------------
Fixing perf regression from UE4 for some queries in Physics.
- Swept SphereVsSphere
- Swept SphereVsCapsule
- Swept CapsuleVsSphere

diff --git a/Engine/Source/Runtime/Experimental/Chaos/Private/Chaos/GeometryQueries.cpp b/Engine/Source/Runtime/Experimental/Chaos/Private/Chaos/GeometryQueries.cpp
deleted file mode 100644
diff --git a/Engine/Source/Runtime/Experimental/Chaos/Public/Chaos/GeometryQueries.h b/Engine/Source/Runtime/Experimental/Chaos/Public/Chaos/GeometryQueries.h
--- a/Engine/Source/Runtime/Experimental/Chaos/Public/Chaos/GeometryQueries.h
+++ b/Engine/Source/Runtime/Experimental/Chaos/Public/Chaos/GeometryQueries.h
@@ -152,45 +152,6 @@ namespace Chaos
 		}
 	}
 
-	UE_INTERNAL CHAOS_API void TransformSweepResultsToWorld(const bool bResult, const FReal Time, const bool bComputeMTD, const FImplicitObject& TestGeom, const FRigidTransform3& TestGeomTM, const FVec3& LocalDir, const FVec3& LocalPosition, const FVec3& LocalNormal, const int32 FaceIndex, FVec3& OutWorldPosition, FVec3& OutWorldNormal, FVec3& OutWorldFaceNormal);
-	UE_INTERNAL CHAOS_API bool SweepSphereVsSphere(const FSphere& SweptSphere, const FRigidTransform3& SweptSphereTM, const FSphere& TestSphere, const FRigidTransform3& TestSphereTM, const FVec3& SweepDir, const FReal Length, const FReal Thickness, const bool bComputeMTD, FReal& OutTime, FVec3& OutPosition, FVec3& OutNormal, int32& OutFaceIndex, FVec3& OutFaceNormal);
-	UE_INTERNAL CHAOS_API bool SweepSphereVsCapsule(const FSphere& SweptSphere, const FRigidTransform3& SweptSphereTM, const FCapsule& TestCapsule, const FRigidTransform3& TestCapsuleTM, const FVec3& SweepDir, const FReal Length, const FReal Thickness, const bool bComputeMTD, FReal& OutTime, FVec3& OutPosition, FVec3& OutNormal, int32& OutFaceIndex, FVec3& OutFaceNormal);
-	UE_INTERNAL CHAOS_API bool SweepCapsuleVsSphere(const FCapsule& SweptCapsule, const FRigidTransform3& SweptCapsuleTM, const FSphere& TestSphere, const FRigidTransform3& TestSphereTM, const FVec3& SweepDir, const FReal Length, const FReal Thickness, const bool bComputeMTD, FReal& OutTime, FVec3& OutPosition, FVec3& OutNormal, int32& OutFaceIndex, FVec3& OutFaceNormal);
-
-	template <EImplicitObjectType SweptType>
-	UE_INTERNAL bool TryRunSpecializedSweep(const FImplicitObject& SweptShape, const FRigidTransform3& SweptShapeTM, const FImplicitObject& TestGeom, const FRigidTransform3& TestGeomTM, const FVec3& SweepDir, const FReal Length, const FReal Thickness, const bool bComputeMTD, bool& bOutResult, FReal& OutTime, FVec3& OutPosition, FVec3& OutNormal, int32& OutFaceIndex, FVec3& OutFaceNormal)
-	{
-		const EImplicitObjectType TestType = TestGeom.GetType();
-		if (SweptType == ImplicitObjectType::Sphere)
-		{
-			if (TestType == ImplicitObjectType::Sphere)
-			{
-				const FSphere& SweptSphere = SweptShape.GetObjectChecked<FSphere>();
-				const FSphere& TestSphere = TestGeom.GetObjectChecked<FSphere>();
-				bOutResult = SweepSphereVsSphere(SweptSphere, SweptShapeTM, TestSphere, TestGeomTM, SweepDir, Length, Thickness, bComputeMTD, OutTime, OutPosition, OutNormal, OutFaceIndex, OutFaceNormal);
-				return true;
-			}
-			else if (TestType == ImplicitObjectType::Capsule)
-			{
-				const FSphere& SweptSphere = SweptShape.GetObjectChecked<FSphere>();
-				const FCapsule& TestCapsule = TestGeom.GetObjectChecked<FCapsule>();
-				bOutResult = SweepSphereVsCapsule(SweptSphere, SweptShapeTM, TestCapsule, TestGeomTM, SweepDir, Length, Thickness, bComputeMTD, OutTime, OutPosition, OutNormal, OutFaceIndex, OutFaceNormal);
-				return true;
-			}
-		}
-		else if (SweptType == ImplicitObjectType::Capsule)
-		{
-			if (TestType == ImplicitObjectType::Sphere)
-			{
-				const FCapsule& SweptCapsule = SweptShape.GetObjectChecked<FCapsule>();
-				const FSphere& TestSphere = TestGeom.GetObjectChecked<FSphere>();
-				bOutResult = SweepCapsuleVsSphere(SweptCapsule, SweptShapeTM, TestSphere, TestGeomTM, SweepDir, Length, Thickness, bComputeMTD, OutTime, OutPosition, OutNormal, OutFaceIndex, OutFaceNormal);
-				return true;
-			}
-		}
-		return false;
-	}
-
 	template <typename SweptGeometry>
 	bool SweepQuery(const FImplicitObject& A, const FRigidTransform3& ATM, const SweptGeometry& B, const FRigidTransform3& BTM, const FVec3& Dir, const FReal Length, FReal& OutTime, FVec3& OutPosition, FVec3& OutNormal, int32& OutFaceIndex, FVec3& OutFaceNormal, const FReal Thickness, const bool bComputeMTD)
 	{
@@ -207,13 +168,7 @@ namespace Chaos
 		}
 
 		OutFaceIndex = INDEX_NONE;
-		PRAGMA_DISABLE_INTERNAL_WARNINGS
-		if (TryRunSpecializedSweep<BType>(B, BTM, A, ATM, Dir, Length, Thickness, bComputeMTD, bResult, OutTime, OutPosition, OutNormal, OutFaceIndex, OutFaceNormal))
-		{
-			return bResult;
-		}
-		PRAGMA_ENABLE_INTERNAL_WARNINGS
-
+		
 		FVec3 LocalPosition(-TNumericLimits<FReal>::Max()); // Make it obvious when things go wrong
 		FVec3 LocalNormal(0);
 
@@ -357,9 +312,12 @@ namespace Chaos
 		}
 
 		//put back into world space
-		PRAGMA_DISABLE_INTERNAL_WARNINGS
-		TransformSweepResultsToWorld(bResult, OutTime, bComputeMTD, A, ATM, LocalDir, LocalPosition, LocalNormal, OutFaceIndex, OutPosition, OutNormal, OutFaceNormal);
-		PRAGMA_ENABLE_INTERNAL_WARNINGS
+		if (bResult && (OutTime >= 0 || bComputeMTD))
+		{
+			OutNormal = ATM.TransformVectorNoScale(LocalNormal).GetSafeNormal(UE_KINDA_SMALL_NUMBER, FVec3::AxisVector(0));
+			OutPosition = ATM.TransformPositionNoScale(LocalPosition);
+			OutFaceNormal = ATM.TransformVectorNoScale(A.FindGeometryOpposingNormal(LocalDir, OutFaceIndex, LocalNormal));
+		}
 
 		return bResult;
 	}
diff --git a/Engine/Source/Runtime/Experimental/ChaosCore/Private/Chaos/Sweeps.cpp b/Engine/Source/Runtime/Experimental/ChaosCore/Private/Chaos/Sweeps.cpp
deleted file mode 100644
diff --git a/Engine/Source/Runtime/Experimental/ChaosCore/Private/Chaos/SweepsMTD.cpp b/Engine/Source/Runtime/Experimental/ChaosCore/Private/Chaos/SweepsMTD.cpp
deleted file mode 100644
diff --git a/Engine/Source/Runtime/Experimental/ChaosCore/Public/Chaos/Sweeps.h b/Engine/Source/Runtime/Experimental/ChaosCore/Public/Chaos/Sweeps.h
deleted file mode 100644
diff --git a/Engine/Source/Runtime/Experimental/ChaosCore/Public/Chaos/SweepsMTD.h b/Engine/Source/Runtime/Experimental/ChaosCore/Public/Chaos/SweepsMTD.h
deleted file mode 100644
