Iris
* Refactored NetRefHandle error reports into seperate functions and more distinct error messages
* Replicated objects blocked from spawning for 30secs on the client will now send an error message back to the server with the missing parents

#rnx

diff --git a/Engine/Source/Runtime/Engine/Private/DataChannel.cpp b/Engine/Source/Runtime/Engine/Private/DataChannel.cpp
--- a/Engine/Source/Runtime/Engine/Private/DataChannel.cpp
+++ b/Engine/Source/Runtime/Engine/Private/DataChannel.cpp
@@ -1759,7 +1759,6 @@ IMPLEMENT_CONTROL_CHANNEL_MESSAGE(BeaconAssignGUID);
 IMPLEMENT_CONTROL_CHANNEL_MESSAGE(BeaconNetGUIDAck);
 IMPLEMENT_CONTROL_CHANNEL_MESSAGE(IrisProtocolMismatch);
 IMPLEMENT_CONTROL_CHANNEL_MESSAGE(IrisNetRefHandleError);
-IMPLEMENT_CONTROL_CHANNEL_MESSAGE(IrisNetRefHandleErrorWithArray);
 
 void UControlChannel::Init( UNetConnection* InConnection, int32 InChannelIndex, EChannelCreateFlags CreateFlags )
 {
@@ -1985,25 +1984,11 @@ void UControlChannel::ReceivedBunch( FInBunch& Bunch )
 			}
 		}
 		else if (MessageType == NMT_IrisNetRefHandleError)
-		{
-			UE::Net::ENetRefHandleError NetRefHandleError = UE::Net::ENetRefHandleError::None;
-			uint64 NetRefHandleId = 0;
-			if (FNetControlMessage<NMT_IrisNetRefHandleError>::Receive(Bunch, NetRefHandleError, NetRefHandleId))
-			{
-#if UE_WITH_IRIS
-				if (UReplicationSystem* IrisRepSystem = Connection->Driver->GetReplicationSystem())
-				{
-					IrisRepSystem->ReportErrorWithNetRefHandle(NetRefHandleError, NetRefHandleId, Connection->GetConnectionHandle().GetParentConnectionId());
-				}
-#endif
-			}
-		}
-		else if (MessageType == NMT_IrisNetRefHandleErrorWithArray)
 		{
 			UE::Net::ENetRefHandleError NetRefHandleError = UE::Net::ENetRefHandleError::None;
 			uint64 NetRefHandleId = 0;
 			TArray<uint64> ExtraNetRefs;
-			if (FNetControlMessage<NMT_IrisNetRefHandleErrorWithArray>::Receive(Bunch, NetRefHandleError, NetRefHandleId, ExtraNetRefs))
+			if (FNetControlMessage<NMT_IrisNetRefHandleError>::Receive(Bunch, NetRefHandleError, NetRefHandleId, ExtraNetRefs))
 			{
 #if UE_WITH_IRIS
 				if (UReplicationSystem* IrisRepSystem = Connection->Driver->GetReplicationSystem())
@@ -2109,9 +2094,6 @@ void UControlChannel::ReceivedBunch( FInBunch& Bunch )
 				case NMT_IrisNetRefHandleError:
 					FNetControlMessage<NMT_IrisNetRefHandleError>::Discard(Bunch);
 					break;
-				case NMT_IrisNetRefHandleErrorWithArray:
-					FNetControlMessage<NMT_IrisNetRefHandleErrorWithArray>::Discard(Bunch);
-					break;
 				case NMT_JoinNoPawn:
 					FNetControlMessage<NMT_JoinNoPawn>::Discard(Bunch);
 					break;
diff --git a/Engine/Source/Runtime/Engine/Private/Net/Iris/ReplicationSystem/EngineReplicationBridge.cpp b/Engine/Source/Runtime/Engine/Private/Net/Iris/ReplicationSystem/EngineReplicationBridge.cpp
--- a/Engine/Source/Runtime/Engine/Private/Net/Iris/ReplicationSystem/EngineReplicationBridge.cpp
+++ b/Engine/Source/Runtime/Engine/Private/Net/Iris/ReplicationSystem/EngineReplicationBridge.cpp
@@ -808,23 +808,7 @@ void UEngineReplicationBridge::OnProtocolMismatchReported(FNetRefHandle RefHandl
 	}
 }
 
-void UEngineReplicationBridge::SendErrorWithNetRefHandle(UE::Net::ENetRefHandleError ErrorType, FNetRefHandle RefHandle, uint32 ConnectionId)
-{
-	if (NetDriver)
-	{
-		if (UNetConnection* ClientConnection = NetDriver->GetConnectionByHandle(UE::Net::FConnectionHandle(ConnectionId)))
-		{
-			uint64 RawHandleId = RefHandle.GetId();
-			FNetControlMessage<NMT_IrisNetRefHandleError>::Send(ClientConnection, ErrorType, RawHandleId);
-		}
-		else
-		{
-			UE_LOG(LogIris, Error, TEXT("UEngineReplicationBridge::SendErrorWithNetRefHandle could not find Connection for id:%u"), ConnectionId);
-		}
-	}
-}
-
-void UEngineReplicationBridge::SendErrorWithNetRefHandle(UE::Net::ENetRefHandleError ErrorType, FNetRefHandle RefHandle, uint32 ConnectionId, const TArray<FNetRefHandle>& ExtraNetRefHandle)
+void UEngineReplicationBridge::SendErrorWithNetRefHandle(UE::Net::ENetRefHandleError ErrorType, FNetRefHandle RefHandle, uint32 ConnectionId, TConstArrayView<const FNetRefHandle> ExtraNetRefHandle)
 {
 	if (NetDriver)
 	{
@@ -834,12 +818,12 @@ void UEngineReplicationBridge::SendErrorWithNetRefHandle(UE::Net::ENetRefHandleE
 
 			TArray<uint64> ArrayIds;
 			ArrayIds.Reserve(ExtraNetRefHandle.Num());
-			for(FNetRefHandle ExtraHandle : ExtraNetRefHandle)
+			for(const FNetRefHandle ExtraHandle : ExtraNetRefHandle)
 			{
 				ArrayIds.Emplace(ExtraHandle.GetId());
 			}
 
-			FNetControlMessage<NMT_IrisNetRefHandleErrorWithArray>::Send(ClientConnection, ErrorType, HandleId, ArrayIds);
+			FNetControlMessage<NMT_IrisNetRefHandleError>::Send(ClientConnection, ErrorType, HandleId, ArrayIds);
 		}
 		else
 		{
diff --git a/Engine/Source/Runtime/Engine/Public/Net/DataChannel.h b/Engine/Source/Runtime/Engine/Public/Net/DataChannel.h
--- a/Engine/Source/Runtime/Engine/Public/Net/DataChannel.h
+++ b/Engine/Source/Runtime/Engine/Public/Net/DataChannel.h
@@ -216,8 +216,7 @@ DEFINE_CONTROL_CHANNEL_MESSAGE(BeaconAssignGUID, 27, UE::Net::FNetIDVariant); //
 DEFINE_CONTROL_CHANNEL_MESSAGE(BeaconNetGUIDAck, 28, FString); // server received NetGUIDAck from client, connection established successfully
 
 DEFINE_CONTROL_CHANNEL_MESSAGE(IrisProtocolMismatch, 29, uint64); // client has a different protocol hash from the server
-DEFINE_CONTROL_CHANNEL_MESSAGE(IrisNetRefHandleError, 30, UE::Net::ENetRefHandleError, uint64); // a specific handle caused an error and we want the remote connection to log all information it has on it
-DEFINE_CONTROL_CHANNEL_MESSAGE(IrisNetRefHandleErrorWithArray, 33, UE::Net::ENetRefHandleError, uint64, TArray<uint64>); // a specific handle caused 
+DEFINE_CONTROL_CHANNEL_MESSAGE(IrisNetRefHandleError, 30, UE::Net::ENetRefHandleError, uint64, TArray<uint64>); // a specific handle caused an error and we want the remote connection to log all information it has on it
 
 // Called in place of Join or JoinSplit during connection handshake, and will created a connection that has the following properties:
 //	* Won't spawn an instance of APawn.
diff --git a/Engine/Source/Runtime/Engine/Public/Net/Iris/ReplicationSystem/EngineReplicationBridge.h b/Engine/Source/Runtime/Engine/Public/Net/Iris/ReplicationSystem/EngineReplicationBridge.h
--- a/Engine/Source/Runtime/Engine/Public/Net/Iris/ReplicationSystem/EngineReplicationBridge.h
+++ b/Engine/Source/Runtime/Engine/Public/Net/Iris/ReplicationSystem/EngineReplicationBridge.h
@@ -100,10 +100,7 @@ public:
 	UIrisObjectReferencePackageMap* GetObjectReferencePackageMap() const { return ObjectReferencePackageMap; }
 
 	/** Tell the remote connection that we detected a reading error with a specific replicated object */
-	ENGINE_API virtual void SendErrorWithNetRefHandle(UE::Net::ENetRefHandleError ErrorType, FNetRefHandle RefHandle, uint32 ConnectionId) override;
-
-	/** Tell the remote connection about an error with extra information */
-	ENGINE_API virtual void SendErrorWithNetRefHandle(UE::Net::ENetRefHandleError ErrorType, FNetRefHandle RefHandle, uint32 ConnectionId, const TArray<FNetRefHandle>& ExtraNetRefHandle) override;
+	ENGINE_API virtual void SendErrorWithNetRefHandle(UE::Net::ENetRefHandleError ErrorType, FNetRefHandle RefHandle, uint32 ConnectionId, TConstArrayView<const FNetRefHandle> ExtraNetRefHandle = {}) override;
 
 	/** Add the rootobject to the level's filter group so it will only be relevant if the connection has that level streamed in. */
 	ENGINE_API void AddRootObjectToLevelGroup(const UObject* RootObject, const ULevel* Level);
diff --git a/Engine/Source/Runtime/Experimental/Iris/Core/Private/Iris/ReplicationSystem/ObjectReplicationBridge.cpp b/Engine/Source/Runtime/Experimental/Iris/Core/Private/Iris/ReplicationSystem/ObjectReplicationBridge.cpp
--- a/Engine/Source/Runtime/Experimental/Iris/Core/Private/Iris/ReplicationSystem/ObjectReplicationBridge.cpp
+++ b/Engine/Source/Runtime/Experimental/Iris/Core/Private/Iris/ReplicationSystem/ObjectReplicationBridge.cpp
@@ -109,13 +109,6 @@ static FAutoConsoleVariableRef CVarInvalidLocationWarnings(
 	TEXT("Chooses how to warn when we detect an invalid world location on a rootobject. 0=never warn. 1=warn once per class. 2=warn once per object. 3=warn on every update")
 );
 
-static int32 GEnsureNetRefHandleError = 0;
-static FAutoConsoleVariableRef CvarEnsureNetRefHandleError(
-	TEXT("net.Iris.EnsureNetRefHandleError"),
-	GEnsureNetRefHandleError,
-	TEXT("Chooses if we should ensure when a NetRefHandleError was reported. -1=never ensure. 0=always ensure. 1..X=ensure only for specific error type")
-);
-
 #if UE_BUILD_SHIPPING
 static bool bUseVerboseIrisCsvStats = false;
 #else
@@ -134,7 +127,25 @@ static FAutoConsoleVariableRef CVarUseVeryVerboseCsvStats(
 	bUseVeryVerboseIrisCsvStats,
 	TEXT("Whether to use very verbose per-class csv stats. This causes all untracked classes to get their own CSV stats. Will incur significant CSV bloat.")
 );
-#endif // !UE_BUILD_SHIPPING
+#endif
+
+#if DO_ENSURE
+static bool bLogBlockedByAsyncLoadingError = true;
+static FAutoConsoleVariableRef CVarLogBlockedByAsyncLoadingErrorr(
+	TEXT("net.Iris.Error.LogBlockedByAsyncLoading"),
+	bLogBlockedByAsyncLoadingError,
+	TEXT("Log and ensure on client reports of objects blocked by must be mapped references stuck in a async loading queue.")
+);
+#endif
+
+#if DO_ENSURE
+static bool bLogMissingParentsError = true;
+static FAutoConsoleVariableRef CVarLogMissingParentsError(
+	TEXT("net.Iris.Error.LogMissingParentsError"),
+	bLogMissingParentsError,
+	TEXT("Log and ensure on client reports of objects not spawning due to missing creation dependent parents.")
+);
+#endif
 
 namespace UE::Net::Private
 {
@@ -2645,115 +2656,147 @@ void UObjectReplicationBridge::OnProtocolMismatchReported(FNetRefHandle RefHandl
 	}
 }
 
-void UObjectReplicationBridge::OnErrorWithNetRefHandleReported(UE::Net::ENetRefHandleError ErrorType, FNetRefHandle RefHandle, uint32 ConnectionId)
+void UObjectReplicationBridge::OnErrorWithNetRefHandleReported(UE::Net::ENetRefHandleError ErrorType, FNetRefHandle RefHandle, uint32 ConnectionId, const TArray<FNetRefHandle>& ExtraHandles)
 {
 	using namespace UE::Net;
 	using namespace UE::Net::Private;
 
-	// Ensure at the end so the log contains all the relevant information
-	ON_SCOPE_EXIT
+	switch(ErrorType)
 	{
-		if (GEnsureNetRefHandleError != -1 && (GEnsureNetRefHandleError==0 || EnumHasAnyFlags((ENetRefHandleError)GEnsureNetRefHandleError,ErrorType)))
+		case ENetRefHandleError::BitstreamCorrupted:
 		{
-			// Use different calls for every error type to prevent only reporting the first error that occured.
-			switch(ErrorType)
-			{
-				case ENetRefHandleError::BitstreamCorrupted:
-				{
-					ensureMsgf(false, TEXT("NetRefHandle error: Bitstream corrupted reported with %s. Look at the log for important information on the object tied to the handle."), *NetRefHandleManager->PrintObjectFromNetRefHandle(RefHandle));
-				} break;
-				case ENetRefHandleError::ReplicationDisabled:
-				{
-					ensureMsgf(false, TEXT("NetRefHandle error: Replication disabled for %s. Client was unable to read the creation data and did not spawn the netobject."), *NetRefHandleManager->PrintObjectFromNetRefHandle(RefHandle));
-				} break;
-				case ENetRefHandleError::BlockedByMustBeMapped:
-				{
-					ensureMsgf(false, TEXT("NetRefHandle error: Replication blocked for %s due to ? unmapped references."), *NetRefHandleManager->PrintObjectFromNetRefHandle(RefHandle));
-				} break;
-				default:
-				{
-					ensureMsgf(false, TEXT("NetRefHandle error: %s reported for %s. This should not occur. Missing switch case or wrong error type used."), LexToString(ErrorType), *NetRefHandleManager->PrintObjectFromNetRefHandle(RefHandle));
-				} break;
-			}
-		}
-	};
+			OnErrorBitstreamCorrupted(RefHandle, ConnectionId);
+		} break;
+		
+		case ENetRefHandleError::ReplicationDisabled:
+		{
+			OnErrorCannotReplicateObject(RefHandle, ConnectionId);
+		} break;
+		
+		case ENetRefHandleError::BlockedByMustBeMapped:
+		{
+			OnErrorBlockedByAsyncLoading(RefHandle, ConnectionId, ExtraHandles);
+		} break;
+		
+		case ENetRefHandleError::BlockedByCreationParents:
+		{
+			OnErrorBlockedByParents(RefHandle, ConnectionId, ExtraHandles);
+		} break;
+		
+		default:
+		{
+			ensureMsgf(false, TEXT("Not implemented"));
+		} break;
+	}
+
+	// Print all internal info we have on the blocked object
+	//PrintDebugInfoForNetRefHandle(RefHandle, ConnectionId, UObjectReplicationBridge::EPrintDebugInfoTraits::NoProtocolState);
+}
+
+void UObjectReplicationBridge::OnErrorBitstreamCorrupted(FNetRefHandle RefHandle, uint32 ConnectionId)
+{
+	using namespace UE::Net;
+	using namespace UE::Net::Private;
 
 	const FInternalNetRefIndex ObjectInternalIndex = NetRefHandleManager->GetInternalIndex(RefHandle);
 	if (ObjectInternalIndex == FNetRefHandleManager::InvalidInternalIndex)
 	{
-		UE_LOG_BRIDGEID(LogIris, Warning, TEXT("OnErrorWithNetRefHandleReported: %s from client %s for %s but object has no InternalIndex."), LexToString(ErrorType), *PrintConnectionInfo(ConnectionId), *RefHandle.ToString());
-		return;
+		UE_LOG_BRIDGEID(LogIris, Error, TEXT("Bitstream corruption from unmapped handle: %s. Reported by client: %s."), *RefHandle.ToString(), *PrintConnectionInfo(ConnectionId));
+	}
+	else
+	{
+		UE_LOG_BRIDGEID(LogIris, Error, TEXT("Bitstream corruption from object: %s. Reported by client: %s"), *NetRefHandleManager->PrintObjectFromIndex(ObjectInternalIndex), *PrintConnectionInfo(ConnectionId));
 	}
 
-	UE_LOG_BRIDGEID(LogIris, Error, TEXT("OnErrorWithNetRefHandleReported: %s from client %s. Problematic object was %s"),  LexToString(ErrorType), *PrintConnectionInfo(ConnectionId), *NetRefHandleManager->PrintObjectFromIndex(ObjectInternalIndex));
+	ensureMsgf(false, TEXT("Bitstream corruption reported with %s. Look at the log for important information on the object tied to the handle."), *NetRefHandleManager->PrintObjectFromNetRefHandle(RefHandle));
 }
 
-void UObjectReplicationBridge::OnErrorWithNetRefHandleReported(UE::Net::ENetRefHandleError ErrorType, FNetRefHandle RefHandle, uint32 ConnectionId, const TArray<FNetRefHandle>& ExtraHandles)
+void UObjectReplicationBridge::OnErrorCannotReplicateObject(FNetRefHandle RefHandle, uint32 ConnectionId)
 {
 	using namespace UE::Net;
 	using namespace UE::Net::Private;
 
-	// Ensure at the end so the log contains all the relevant information
-	ON_SCOPE_EXIT
+	const FInternalNetRefIndex ObjectInternalIndex = NetRefHandleManager->GetInternalIndex(RefHandle);
+	if (ObjectInternalIndex == FNetRefHandleManager::InvalidInternalIndex)
 	{
-		if (GEnsureNetRefHandleError != -1 && (GEnsureNetRefHandleError == 0 || EnumHasAnyFlags((ENetRefHandleError)GEnsureNetRefHandleError,ErrorType)))
-		{
-			// Use different calls for every error type to prevent only reporting the first error that occured.
-			switch (ErrorType)
-			{
-				case ENetRefHandleError::BitstreamCorrupted:
-				{
-					ensureMsgf(false, TEXT("NetRefHandle error: Bitstream corrupted reported with %s. Look at the log for important information on the object tied to the handle."), *NetRefHandleManager->PrintObjectFromNetRefHandle(RefHandle));
-				} break;
-				case ENetRefHandleError::ReplicationDisabled:
-				{
-					ensureMsgf(false, TEXT("NetRefHandle error: Replication disabled for %s. Client was unable to read the creation data and did not spawn the netobject."), *NetRefHandleManager->PrintObjectFromNetRefHandle(RefHandle));
-				} break;
-				case ENetRefHandleError::BlockedByMustBeMapped:
-				{
-					ensureMsgf(false, TEXT("NetRefHandle error: Replication blocked for %s due to %d unmapped references."), *NetRefHandleManager->PrintObjectFromNetRefHandle(RefHandle), ExtraHandles.Num());
-				} break;
-				default:
-				{
-					ensureMsgf(false, TEXT("NetRefHandle error: %s reported for %s. This should not occur. Missing switch case or wrong error type used."), LexToString(ErrorType), *NetRefHandleManager->PrintObjectFromNetRefHandle(RefHandle));
-				} break;
-			}
-		}
-	};
+		UE_LOG_BRIDGEID(LogIris, Error, TEXT("Client cannot replicate unmapped handle: %s. Reported by client: %s."), *RefHandle.ToString(), *PrintConnectionInfo(ConnectionId));
+	}
+	else
+	{
+		UE_LOG_BRIDGEID(LogIris, Error, TEXT("Client cannot replicated object: %s. Reported by client: %s"), *NetRefHandleManager->PrintObjectFromIndex(ObjectInternalIndex), *PrintConnectionInfo(ConnectionId));
+	}
+
+	ensureMsgf(false, TEXT("Client disabled replication for %s. Client was unable to read the creation data and did not spawn the netobject."), *NetRefHandleManager->PrintObjectFromNetRefHandle(RefHandle));
+}
 
-	// Print the blocked object
+void UObjectReplicationBridge::OnErrorBlockedByAsyncLoading(FNetRefHandle RefHandle, uint32 ConnectionId, const TArray<FNetRefHandle>& MustBeMappedRefs)
+{
+	using namespace UE::Net;
+	using namespace UE::Net::Private;
+
+#if DO_ENSURE
+
+	if (!bLogBlockedByAsyncLoadingError)
 	{
-		const FInternalNetRefIndex ObjectInternalIndex = NetRefHandleManager->GetInternalIndex(RefHandle);
-		if (ObjectInternalIndex == FNetRefHandleManager::InvalidInternalIndex)
-		{
-			UE_LOG_BRIDGEID(LogIris, Warning, TEXT("OnErrorWithNetRefHandleReported: %s from client %s for %s but object has no InternalIndex."), LexToString(ErrorType), *PrintConnectionInfo(ConnectionId), *RefHandle.ToString());
-			return;
-		}
+		return;
+	}
+
+	// A client cannot process incoming data for the Object because it has been waiting on the MustBeMappedRefs to be async loaded for too long (30secs by default)
+
+	const FInternalNetRefIndex ObjectInternalIndex = NetRefHandleManager->GetInternalIndex(RefHandle);
 
-		UE_LOG_BRIDGEID(LogIris, Error, TEXT("OnErrorWithNetRefHandleReported: %s for %s. Client: %s"), LexToString(ErrorType), *NetRefHandleManager->PrintObjectFromIndex(ObjectInternalIndex), *PrintConnectionInfo(ConnectionId));
+	if (ObjectInternalIndex == FNetRefHandleManager::InvalidInternalIndex)
+	{
+		UE_LOG_BRIDGEID(LogIris, Error, TEXT("Replication Blocked by AsyncLoading for unmapped handle: %s. Client: %s"), *RefHandle.ToString(), *PrintConnectionInfo(ConnectionId));
+	}
+	else
+	{
+		UE_LOG_BRIDGEID(LogIris, Error, TEXT("Replication Blocked by AsyncLoading for: %s. Client: %s"), *NetRefHandleManager->PrintObjectFromIndex(ObjectInternalIndex), *PrintConnectionInfo(ConnectionId));
 	}
+	
+	for (FNetRefHandle MustBeMappedRef : MustBeMappedRefs)
+	{
+		UE_LOG_BRIDGEID(LogIris, Error, TEXT("AsyncLoading reference: %s | PathName: %s"), *NetRefHandleManager->PrintObjectFromNetRefHandle(MustBeMappedRef), GetObjectReferenceCache()->GetObjectRelativePath(MustBeMappedRef));
+	}
+		
+	ensureMsgf(false, TEXT("Blocked by AsyncLoading error: Replication blocked for %s due to %d unmapped references."), *NetRefHandleManager->PrintObjectFromNetRefHandle(RefHandle), MustBeMappedRefs.Num());
 
-	// Make a custom printer for the specific error
-	auto PrintExtraObjectInfo = [&](FNetRefHandle ObjectHandle) -> FString
+#endif // DO_ENSURE
+}
+
+
+void UObjectReplicationBridge::OnErrorBlockedByParents(FNetRefHandle RefHandle, uint32 ConnectionId, const TArray<FNetRefHandle>& BlockedParents)
+{
+	using namespace UE::Net;
+	using namespace UE::Net::Private;
+
+#if DO_ENSURE
+
+	if (!bLogMissingParentsError)
 	{
-		// Error specific message
-		if (ErrorType == ENetRefHandleError::BlockedByMustBeMapped)
-		{
-			return FString::Printf(TEXT("MustBeMapped reference: %s | PathName: %s"), *NetRefHandleManager->PrintObjectFromNetRefHandle(ObjectHandle), GetObjectReferenceCache()->GetObjectRelativePath(ObjectHandle));
-		}
+		return;
+	}
 
-		// For not yet implemented errors
-		return FString::Printf(TEXT("Extra object: %s"), *NetRefHandleManager->PrintObjectFromNetRefHandle(ObjectHandle));
-	};
+	// A client cannot instantiate the Object because it has been waiting on the creation parents to be spawned for too long (default is 30secs)
+	const FInternalNetRefIndex ObjectInternalIndex = NetRefHandleManager->GetInternalIndex(RefHandle);
+	
+	if (ObjectInternalIndex == FNetRefHandleManager::InvalidInternalIndex)
+	{
+		UE_LOG_BRIDGEID(LogIris, Error, TEXT("Client cannot spawn unmapped object: %s. Reported by client: %s"), *RefHandle.ToString(), *PrintConnectionInfo(ConnectionId));
+	}
+	else
+	{
+		UE_LOG_BRIDGEID(LogIris, Error, TEXT("Client cannot spawn object: %s. Reported by client: %s"), *NetRefHandleManager->PrintObjectFromIndex(ObjectInternalIndex), *PrintConnectionInfo(ConnectionId));
+	}
 
-	// Print the extra objects sent with the error
-	for (FNetRefHandle ExtraHandle : ExtraHandles)
+	for (FNetRefHandle BlockedParent : BlockedParents)
 	{
-		UE_LOG_BRIDGEID(LogIris, Error, TEXT("\t%s"), *PrintExtraObjectInfo(ExtraHandle));
+		UE_LOG_BRIDGEID(LogIris, Error, TEXT("Missing Creation Dependency Parent: %s"), *NetRefHandleManager->PrintObjectFromNetRefHandle(BlockedParent));
 	}
+		
+	ensureMsgf(false, TEXT("Client cannot spawn: %s. Blocked by %d missing Parents."), *NetRefHandleManager->PrintObjectFromNetRefHandle(RefHandle), BlockedParents.Num());
 
-	// All all internal info we have on the blocked object
-	//PrintDebugInfoForNetRefHandle(RefHandle, ConnectionId, UObjectReplicationBridge::EPrintDebugInfoTraits::NoProtocolState);
+#endif // DO_ENSURE
 }
 
 TArray<uint32> UObjectReplicationBridge::FindConnectionsFromArgs(const TArray<FString>& Args) const
diff --git a/Engine/Source/Runtime/Experimental/Iris/Core/Private/Iris/ReplicationSystem/ReplicationReader.cpp b/Engine/Source/Runtime/Experimental/Iris/Core/Private/Iris/ReplicationSystem/ReplicationReader.cpp
--- a/Engine/Source/Runtime/Experimental/Iris/Core/Private/Iris/ReplicationSystem/ReplicationReader.cpp
+++ b/Engine/Source/Runtime/Experimental/Iris/Core/Private/Iris/ReplicationSystem/ReplicationReader.cpp
@@ -2254,7 +2254,7 @@ void FReplicationReader::DispatchStateData(FNetSerializationContext& Context)
 
 		const FNetRefHandleManager::FReplicatedObjectData& ObjectData = NetRefHandleManager->GetReplicatedObjectDataNoCheck(Info.InternalIndex);
 
-		// Before starting a potentially new batch we want to flush rpc:s and legacy callbacks belonging to the previous batch
+		// Before starting a potentially new batch we want to flush rpc's and legacy callbacks belonging to the previous batch
 		const FInternalNetRefIndex RootInternalIndex = ObjectData.SubObjectRootIndex == FNetRefHandleManager::InvalidInternalIndex ? Info.InternalIndex : ObjectData.SubObjectRootIndex;
 		if (RootInternalIndex != LastDispatchedRootInternalIndex && NumObjectsPendingPostDistpatch)
 		{
@@ -2971,12 +2971,20 @@ void FReplicationReader::ProcessQueuedBatches()
 							   , NumUnknownParents
 							   , PendingBatchData->CreationDependentParents.Num());
 
+						TArray<FNetRefHandle> BlockedParents;
 						for (const FNetRefHandle CreationDependentParent : PendingBatchData->CreationDependentParents)
 						{
 							const bool bParentExists = DoesParentExist(CreationDependentParent);
 
 							UE_LOG(LogIris, Warning, TEXT("FReplicationWriter::ProcessQueuedBatches: %s CreationDependent parent: %s ."), bParentExists?TEXT("Found"):TEXT("Waiting"), *NetRefHandleManager->PrintObjectFromNetRefHandle(CreationDependentParent));
+
+							if (!bParentExists)
+							{
+								BlockedParents.Add(CreationDependentParent);
+							}
 						}
+
+						ReplicationBridge->SendErrorWithNetRefHandle(ENetRefHandleError::BlockedByCreationParents, OwnerHandle, Parameters.ConnectionId, BlockedParents);
 					}
 
 					if (!bAllReferencesResolved)
@@ -3022,8 +3030,7 @@ void FReplicationReader::ProcessQueuedBatches()
 							FClientProfiler::RecordBlockedReplication(*BlockedObjectName, PendingBatchData->PendingMustBeMappedReferences.Num(), (float)BlockedTime);
 							#endif
 
-							ReplicationBridge->SendErrorWithNetRefHandle(UE::Net::ENetRefHandleError::BlockedByMustBeMapped, OwnerHandle, Parameters.ConnectionId, PendingMustBeMapped);
-						
+							ReplicationBridge->SendErrorWithNetRefHandle(ENetRefHandleError::BlockedByMustBeMapped, OwnerHandle, Parameters.ConnectionId, PendingMustBeMapped);
 						}
 					}
 
diff --git a/Engine/Source/Runtime/Experimental/Iris/Core/Private/Iris/ReplicationSystem/ReplicationSystem.cpp b/Engine/Source/Runtime/Experimental/Iris/Core/Private/Iris/ReplicationSystem/ReplicationSystem.cpp
--- a/Engine/Source/Runtime/Experimental/Iris/Core/Private/Iris/ReplicationSystem/ReplicationSystem.cpp
+++ b/Engine/Source/Runtime/Experimental/Iris/Core/Private/Iris/ReplicationSystem/ReplicationSystem.cpp
@@ -2067,14 +2067,6 @@ void UReplicationSystem::ReportProtocolMismatch(uint64 NetRefHandleId, uint32 Co
 	Impl->ReplicationSystemInternal.GetReplicationBridge()->OnProtocolMismatchReported(NetRefHandle, ConnectionId);
 }
 
-void UReplicationSystem::ReportErrorWithNetRefHandle(UE::Net::ENetRefHandleError ErrorType, uint64 NetRefHandleId, uint32 ConnectionId)
-{
-	using namespace UE::Net::Private;
-	const FNetRefHandle NetRefHandle = FNetRefHandleManager::MakeNetRefHandle(NetRefHandleId, GetId());
-
-	Impl->ReplicationSystemInternal.GetReplicationBridge()->OnErrorWithNetRefHandleReported(ErrorType, NetRefHandle, ConnectionId);
-}
-
 void UReplicationSystem::ReportErrorWithNetRefHandle(UE::Net::ENetRefHandleError ErrorType, uint64 NetRefHandleId, uint32 ConnectionId, const TArray<uint64>& ExtraNetRefHandles)
 {
 	using namespace UE::Net::Private;
diff --git a/Engine/Source/Runtime/Experimental/Iris/Core/Public/Iris/ReplicationSystem/ObjectReplicationBridge.h b/Engine/Source/Runtime/Experimental/Iris/Core/Public/Iris/ReplicationSystem/ObjectReplicationBridge.h
--- a/Engine/Source/Runtime/Experimental/Iris/Core/Public/Iris/ReplicationSystem/ObjectReplicationBridge.h
+++ b/Engine/Source/Runtime/Experimental/Iris/Core/Public/Iris/ReplicationSystem/ObjectReplicationBridge.h
@@ -304,8 +304,14 @@ protected:
 	IRISCORE_API virtual void DetachInstanceFromRemote(FNetRefHandle Handle, EReplicationBridgeDestroyInstanceReason DestroyReason, EReplicationBridgeDestroyInstanceFlags DestroyFlags,  UE::Net::FNetObjectFactoryId NetFactoryId) override;
 	IRISCORE_API virtual void DetachInstance(FNetRefHandle Handle) override;
 	IRISCORE_API virtual void OnProtocolMismatchReported(FNetRefHandle RefHandle, uint32 ConnectionId) override;
-	IRISCORE_API virtual void OnErrorWithNetRefHandleReported(UE::Net::ENetRefHandleError ErrorType, FNetRefHandle RefHandle, uint32 ConnectionId) override;
 	IRISCORE_API virtual void OnErrorWithNetRefHandleReported(UE::Net::ENetRefHandleError ErrorType, FNetRefHandle RefHandle, uint32 ConnectionId, const TArray<FNetRefHandle>& ExtraHandles) override;
+
+private:
+
+	void OnErrorBitstreamCorrupted(FNetRefHandle RefHandle, uint32 ConnectionId);
+	void OnErrorCannotReplicateObject(FNetRefHandle RefHandle, uint32 ConnectionId);
+	void OnErrorBlockedByAsyncLoading(FNetRefHandle RefHandle, uint32 ConnectionId, const TArray<FNetRefHandle>& MustBeMappedRefs);
+	void OnErrorBlockedByParents(FNetRefHandle RefHandle, uint32 ConnectionId, const TArray<FNetRefHandle>& BlockedParents);
 	
 protected:
 	
diff --git a/Engine/Source/Runtime/Experimental/Iris/Core/Public/Iris/ReplicationSystem/ReplicationBridge.h b/Engine/Source/Runtime/Experimental/Iris/Core/Public/Iris/ReplicationSystem/ReplicationBridge.h
--- a/Engine/Source/Runtime/Experimental/Iris/Core/Public/Iris/ReplicationSystem/ReplicationBridge.h
+++ b/Engine/Source/Runtime/Experimental/Iris/Core/Public/Iris/ReplicationSystem/ReplicationBridge.h
@@ -263,12 +263,10 @@ protected:
 	virtual void OnProtocolMismatchReported(FNetRefHandle RefHandle, uint32 ConnectionId) {}
 
 	/** Called when a remote connection has a critical error caused by a specific NetRefHandle */
-	virtual void OnErrorWithNetRefHandleReported(UE::Net::ENetRefHandleError ErrorType, FNetRefHandle RefHandle, uint32 ConnectionId) {}
 	virtual void OnErrorWithNetRefHandleReported(UE::Net::ENetRefHandleError ErrorType, FNetRefHandle RefHandle, uint32 ConnectionId, const TArray<FNetRefHandle>& ExtraNetRefHandle) {}
 
 	/** Tell the remote connection that we detected a reading error with a specific replicated object */
-	virtual void SendErrorWithNetRefHandle(UE::Net::ENetRefHandleError ErrorType, FNetRefHandle RefHandle, uint32 ConnectionId) {}
-	virtual void SendErrorWithNetRefHandle(UE::Net::ENetRefHandleError ErrorType, FNetRefHandle RefHandle, uint32 ConnectionId, const TArray<FNetRefHandle>& ExtraNetRefHandle) {}
+	virtual void SendErrorWithNetRefHandle(UE::Net::ENetRefHandleError ErrorType, FNetRefHandle RefHandle, uint32 ConnectionId, TConstArrayView<const FNetRefHandle> ExtraNetRefHandle = {}) {}
 
 	/** Called from PreSeamlessTravelGarbageCollect. */
 	IRISCORE_API virtual void OnPreSeamlessTravelGarbageCollect();
diff --git a/Engine/Source/Runtime/Experimental/Iris/Core/Public/Iris/ReplicationSystem/ReplicationSystem.h b/Engine/Source/Runtime/Experimental/Iris/Core/Public/Iris/ReplicationSystem/ReplicationSystem.h
--- a/Engine/Source/Runtime/Experimental/Iris/Core/Public/Iris/ReplicationSystem/ReplicationSystem.h
+++ b/Engine/Source/Runtime/Experimental/Iris/Core/Public/Iris/ReplicationSystem/ReplicationSystem.h
@@ -730,10 +730,7 @@ public:
 	/** Called when a connection finds a protocol divergence when instantiating a replicated object. */
 	IRISCORE_API void ReportProtocolMismatch(uint64 NetRefHandleId, uint32 ConnectionId);
 
-	/** Called when a connection reports a critical error with a netrefhandle object */
-	IRISCORE_API void ReportErrorWithNetRefHandle(UE::Net::ENetRefHandleError ErrorType, uint64 NetRefHandleId, uint32 ConnectionId);
-
-	/** Called when a connection reports a critical error with a netrefhandle object and extra handles */
+		/** Called when a connection reports a critical error with a netrefhandle object */
 	IRISCORE_API void ReportErrorWithNetRefHandle(UE::Net::ENetRefHandleError ErrorType, uint64 NetRefHandleId, uint32 ConnectionId, const TArray<uint64>& ExtraNetRefHandles);
 
 	/** Collect relevant metrics that could be sent to analytics */
diff --git a/Engine/Source/Runtime/Net/Core/Public/Net/Core/Connection/NetEnums.h b/Engine/Source/Runtime/Net/Core/Public/Net/Core/Connection/NetEnums.h
--- a/Engine/Source/Runtime/Net/Core/Public/Net/Core/Connection/NetEnums.h
+++ b/Engine/Source/Runtime/Net/Core/Public/Net/Core/Connection/NetEnums.h
@@ -129,11 +129,11 @@ namespace UE::Net
 enum class ENetRefHandleError : uint32
 {
 	None = 0x00,
-	BitstreamCorrupted = 0x01, // An error while reading detected an irrecoverable bitstream corruption. Client was forced to disconnect
-	ReplicationDisabled = 0x02, // An error while reading was recoverable but prevents this NetObject from replicating.
-	BlockedByMustBeMapped = 0x03, // An object is blocked due to MustBeMapped references not resolving fast enough
+	BitstreamCorrupted, // An error while reading detected an irrecoverable bitstream corruption. Client was forced to disconnect
+	ReplicationDisabled, // An error while reading was recoverable but prevents this NetObject from replicating.
+	BlockedByMustBeMapped, // An object is blocked due to MustBeMapped references not resolving fast enough
+	BlockedByCreationParents, // An object is blocked waiting on creation parents that do not exist yet on the client.
 };
-ENUM_CLASS_FLAGS(ENetRefHandleError);
 
 inline const TCHAR* LexToString(ENetRefHandleError NetRefHandleError)
 {
