[ChaosVD] Replaced particle debug names with a new particle metadata object

Until now, we were only serializing the chaos debug name for each particle. The problem with that is these are only available in Development builds, and it does not have the ideal format to be used as display name in CVD's scene outliner.

With CL, CVD will now gather on demand several pieces of data when it first sees a particle, which then can be used to build a display name when the CVD file is loaded. As all the extra metadata is stored as FNames, it means this also takes advantage of the serialized name table CVD already has and as they are pre-existing FNames, we are not allocation new strings in the process.

Part of this new metadata is the full asset path for the owning object of a particle, which allow us to implement a feature to allow users find the asset from which a particle was created from in the content browser (assuming CVD is open in the same editor).

In addition this CL adds:

- Moved most of the global variables in the ChaosVisualDebuggerTrace file into a single struct that holds the entire recording session state
- Removed an additional lock in the geometry serialization path. We were using two maps to keep track of if we traced an object, and a cache of geometry hashes, but we can infer that we traced an object if we already have a cached hash.
- Switched most maps that were guarded with a lock to TStrippedMap, instead of having a single lock per map.

diff --git a/Engine/Plugins/ChaosVD/Source/ChaosVD/Private/ChaosVDSceneParticle.cpp b/Engine/Plugins/ChaosVD/Source/ChaosVD/Private/ChaosVDSceneParticle.cpp
--- a/Engine/Plugins/ChaosVD/Source/ChaosVD/Private/ChaosVDSceneParticle.cpp
+++ b/Engine/Plugins/ChaosVD/Source/ChaosVD/Private/ChaosVDSceneParticle.cpp
@@ -201,12 +201,12 @@ void FChaosVDSceneParticle::PreUpdateFromRecordedParticleData(const TSharedPtr<F
 	const bool bStateChanged = !ParticleDataPtr || (ParticleDataPtr->ParticleDynamicsMisc.MObjectState != InRecordedData->ParticleDynamicsMisc.MObjectState);
 	// Particle name shouldn't change, but CVD keeps particle instances alive even when a particle is destroyed so they can be reused if a user scrubs back to a frame where the particle existed
 	// But it could also be the case that the particle ID was re-used, therefore the name could have changed as this is a new particle
-	const bool bHasNewName = !ParticleDataPtr || (ParticleDataPtr->DebugNameHash != InRecordedData->DebugNameHash);
+	const bool bHasNewName = !ParticleDataPtr || (ParticleDataPtr->MetadataId != InRecordedData->MetadataId);
 
 	static FString UnnamedParticle("UnnamedParticle");
 	if (bHasNewName)
 	{
-		SetDisplayName(InRecordedData->DebugName.IsEmpty() ? UnnamedParticle : InRecordedData->DebugName);
+		SetDisplayName(InRecordedData->GetDebugName().IsEmpty() ? UnnamedParticle : InRecordedData->GetDebugName());
 		EnumAddFlags(DirtyFlags, EChaosVDSceneParticleDirtyFlags::TEDS);
 	}
 
@@ -225,7 +225,7 @@ void FChaosVDSceneParticle::PreUpdateFromRecordedParticleData(const TSharedPtr<F
 				if (!ensure(CurrentRootGeometry))
 				{
 					// We intentionally let the code continue, as passing down a null geometry  will take care of removing any existing mesh representation for this particle
-					UE_LOG(LogChaosVDEditor, Warning, TEXT("Failed to find Geometry for Particle ID [%d] | Geometry Hash [%u] | Debug Name [%s]"), ParticleDataPtr->ParticleIndex, GeometryHash, *ParticleDataPtr->DebugName);
+					UE_LOG(LogChaosVDEditor, Warning, TEXT("Failed to find Geometry for Particle ID [%d] | Geometry Hash [%u] | Debug Name [%s]"), ParticleDataPtr->ParticleIndex, GeometryHash, *ParticleDataPtr->GetDebugName());
 				}
 			}
 			else
@@ -603,11 +603,6 @@ bool FChaosVDSceneParticle::HasCollisionData()
 	return false;
 }
 
-FName FChaosVDSceneParticle::GetProviderName()
-{
-	return ParticleDataPtr ? FName(ParticleDataPtr->DebugName) : NAME_None;
-}
-
 void FChaosVDSceneParticle::UpdateMeshInstancesSelectionState()
 {
 	TSharedPtr<FChaosVDInstancedMeshData> CurrentSelectedGeometry = CurrentSelectedGeometryInstance.Pin();
diff --git a/Engine/Plugins/ChaosVD/Source/ChaosVD/Private/ChaosVDSceneParticleCustomization.cpp b/Engine/Plugins/ChaosVD/Source/ChaosVD/Private/ChaosVDSceneParticleCustomization.cpp
--- a/Engine/Plugins/ChaosVD/Source/ChaosVD/Private/ChaosVDSceneParticleCustomization.cpp
+++ b/Engine/Plugins/ChaosVD/Source/ChaosVD/Private/ChaosVDSceneParticleCustomization.cpp
@@ -9,6 +9,7 @@
 #include "ChaosVDObjectDetailsTab.h"
 #include "ChaosVDSceneParticle.h"
 #include "ChaosVDScene.h"
+#include "ChaosVDStyle.h"
 #include "ChaosVDTabsIDs.h"
 #include "DetailCategoryBuilder.h"
 #include "DetailLayoutBuilder.h"
@@ -19,6 +20,7 @@
 #include "Widgets/SChaosVDCollisionDataInspector.h"
 #include "Widgets/SChaosVDMainTab.h"
 #include "Widgets/SToolTip.h"
+#include "Widgets/Images/SImage.h"
 #include "Widgets/Input/SButton.h"
 #include "Widgets/Text/STextBlock.h"
 
@@ -117,6 +119,16 @@ void FChaosVDSceneParticleCustomization::CustomizeDetails(IDetailLayoutBuilder&
 
 			HandleSceneUpdated();
 
+			if (TSharedPtr<const FChaosVDParticleDataWrapper> ParticleData = SelectedParticleInstance->GetParticleData())
+			{
+				// If bHasDebugName it means this is an old CVD recording where we didn't have the metadata system (or that it was a particle without a valid name)
+				// In which case would be ok showing an empty metadata structure, even if it is an old CVD file
+				if (!ParticleData->HasLegacyDebugName())
+				{
+					AddExternalStructure(CachedParticleMetadata, DetailBuilder, FName("Particle Metadata"), LOCTEXT("ParticleMetadataStructName", "Particle Metadata"));
+				}
+			}
+
 			TSharedPtr<IPropertyHandle> InspectedDataPropertyHandlePtr;
 
 			if (TSharedPtr<FChaosVDInstancedMeshData> SelectedGeometryInstance = SelectedParticleInstance->GetSelectedMeshInstance().Pin())
@@ -173,6 +185,24 @@ void FChaosVDSceneParticleCustomization::HandleSceneUpdated()
 	{
 		TSharedPtr<const FChaosVDParticleDataWrapper> ParticleDataPtr = CurrentObservedParticle->GetParticleData();
 		CachedParticleData = ParticleDataPtr ? *ParticleDataPtr : FChaosVDParticleDataWrapper();
+
+		if (ParticleDataPtr)
+		{
+			CachedParticleData = *ParticleDataPtr;
+			if (const TSharedPtr<FChaosVDParticleMetadata>& MetadataInstance = ParticleDataPtr->GetMetadataInstance())
+			{
+				CachedParticleMetadata = *MetadataInstance;
+			}
+			else
+			{
+				CachedParticleMetadata = FChaosVDParticleMetadata();
+			}
+		}
+		else
+		{
+			CachedParticleData = FChaosVDParticleDataWrapper();
+			CachedParticleMetadata = FChaosVDParticleMetadata();
+		}
 	}
 }
 
@@ -249,6 +279,7 @@ void FChaosVDSceneParticleCustomization::ResetCachedView()
 	CurrentObservedParticle = nullptr;
 	CachedParticleData = FChaosVDParticleDataWrapper();
 	CachedGeometryDataInstanceCopy = FChaosVDMeshDataInstanceState();
+	CachedParticleMetadata = FChaosVDParticleMetadata();
 }
 
 void FChaosVDSceneParticleCustomization::RegisterCVDScene(const TSharedPtr<FChaosVDScene>& InScene)
diff --git a/Engine/Plugins/ChaosVD/Source/ChaosVD/Private/ChaosVDSceneParticleCustomization.h b/Engine/Plugins/ChaosVD/Source/ChaosVD/Private/ChaosVDSceneParticleCustomization.h
--- a/Engine/Plugins/ChaosVD/Source/ChaosVD/Private/ChaosVDSceneParticleCustomization.h
+++ b/Engine/Plugins/ChaosVD/Source/ChaosVD/Private/ChaosVDSceneParticleCustomization.h
@@ -61,6 +61,8 @@ private:
 	FChaosVDParticleDataWrapper CachedParticleData;
 	/* Copy of the last known particle data structure of a selected particle -  Used to avoid rebuild the layout every time we change frame in CVD */
 	FChaosVDMeshDataInstanceState CachedGeometryDataInstanceCopy;
+	/* Copy of the last known particle metadata data structure of a selected particle -  Used to avoid rebuild the layout every time we change frame in CVD */
+	FChaosVDParticleMetadata CachedParticleMetadata;
 	
 	TWeakPtr<FChaosVDScene> SceneWeakPtr = nullptr;
 
diff --git a/Engine/Plugins/ChaosVD/Source/ChaosVD/Private/DetailsCustomizations/ChaosVDParticleDataWrapperCustomization.cpp b/Engine/Plugins/ChaosVD/Source/ChaosVD/Private/DetailsCustomizations/ChaosVDParticleDataWrapperCustomization.cpp
--- a/Engine/Plugins/ChaosVD/Source/ChaosVD/Private/DetailsCustomizations/ChaosVDParticleDataWrapperCustomization.cpp
+++ b/Engine/Plugins/ChaosVD/Source/ChaosVD/Private/DetailsCustomizations/ChaosVDParticleDataWrapperCustomization.cpp
@@ -39,12 +39,17 @@ void FChaosVDParticleDataWrapperCustomization::CustomizeChildren(TSharedRef<IPro
 	{
 		TSharedPtr<IPropertyHandle> Handle = StructPropertyHandle->GetChildHandle(ChildIndex);
 
+		if (!Handle)
+		{
+			continue;
+		}
+
 		bool bIsParticleDataStruct = false;
 		if (FChaosVDDetailsCustomizationUtils::HasValidCVDWrapperData(Handle, bIsParticleDataStruct))
 		{
 			StructBuilder.AddProperty(Handle.ToSharedRef());
 		}
-		else if (!bIsParticleDataStruct) // If it is not a valid particle data but because it is not a data struct, just add the property anyways
+		else if (!bIsParticleDataStruct) // If it is not a valid particle data because it is not a data struct, just add the property anyways
 		{
 			StructBuilder.AddProperty(Handle.ToSharedRef());
 		}
diff --git a/Engine/Plugins/ChaosVD/Source/ChaosVD/Private/DetailsCustomizations/ChaosVDParticleMetadataCustomization.cpp b/Engine/Plugins/ChaosVD/Source/ChaosVD/Private/DetailsCustomizations/ChaosVDParticleMetadataCustomization.cpp
new file mode 100644
--- /dev/null
+++ b/Engine/Plugins/ChaosVD/Source/ChaosVD/Private/DetailsCustomizations/ChaosVDParticleMetadataCustomization.cpp
@@ -0,0 +1,118 @@
+// Copyright Epic Games, Inc. All Rights Reserved.
+#include "ChaosVDParticleMetadataCustomization.h"
+
+#include "ChaosVDDetailsCustomizationUtils.h"
+#include "ChaosVDModule.h"
+#include "DetailWidgetRow.h"
+#include "Editor.h"
+#include "IDetailChildrenBuilder.h"
+#include "PropertyHandle.h"
+#include "DataWrappers/ChaosVDParticleDataWrapper.h"
+#include "Misc/MessageDialog.h"
+#include "Widgets/SToolTip.h"
+#include "Widgets/Input/SButton.h"
+#include "Widgets/Text/STextBlock.h"
+
+#define LOCTEXT_NAMESPACE "ChaosVisualDebugger"
+
+TSharedRef<IPropertyTypeCustomization> FChaosVDParticleMetadataCustomization::MakeInstance()
+{
+	return MakeShareable(new FChaosVDParticleMetadataCustomization());
+}
+
+void FChaosVDParticleMetadataCustomization::CustomizeHeader(TSharedRef<IPropertyHandle> StructPropertyHandle, FDetailWidgetRow& HeaderRow,
+	IPropertyTypeCustomizationUtils& StructCustomizationUtils)
+{
+	HeaderRow
+	.NameContent()
+	[
+		StructPropertyHandle->CreatePropertyNameWidget()
+	];
+}
+
+void FChaosVDParticleMetadataCustomization::CustomizeChildren(TSharedRef<IPropertyHandle> StructPropertyHandle, IDetailChildrenBuilder& StructBuilder,
+	IPropertyTypeCustomizationUtils& StructCustomizationUtils)
+{
+
+	FChaosVDDetailsPropertyDataHandle<FChaosVDParticleMetadata> ParticleMetadataHandle(StructPropertyHandle);
+	FChaosVDParticleMetadata* ParticleMetadataInstance = ParticleMetadataHandle.GetDataInstance();
+
+	if (!ParticleMetadataInstance)
+	{
+		return;
+	}
+	
+	uint32 NumChildren;
+	StructPropertyHandle->GetNumChildren(NumChildren);
+
+	if (NumChildren == 0)
+	{
+		return;
+	}
+
+	for (uint32 ChildIndex = 0; ChildIndex < NumChildren; ++ChildIndex)
+	{
+		TSharedPtr<IPropertyHandle> Handle = StructPropertyHandle->GetChildHandle(ChildIndex);
+
+		if (!Handle)
+		{
+			continue;
+		}
+
+		StructBuilder.AddProperty(Handle.ToSharedRef());
+	}
+
+	StructBuilder.AddCustomRow(FText::GetEmpty()).WholeRowContent()
+	[
+		GenerateBrowseParticleOwnerButton(ParticleMetadataInstance->OwnerAssetPath)
+	];
+}
+
+TSharedRef<SWidget> FChaosVDParticleMetadataCustomization::GenerateBrowseParticleOwnerButton(FTopLevelAssetPath AssetPath)
+{
+	TSharedRef<SWidget> BrowseButton = SNew(SHorizontalBox)
+	+SHorizontalBox::Slot()
+	.VAlign(VAlign_Center)
+	.HAlign(HAlign_Center)
+	.Padding(12.0f, 7.0f, 12.0f, 7.0f)
+	.FillWidth(1.0f)
+	[
+		SNew(SButton)
+		.ToolTip(SNew(SToolTip).Text(LOCTEXT("BrowseParticleOwnerDesc", "Click here to open attempt to find the asset owning this particle in the content browser. This only works if CVD is open in the editor of the project where the recording comes from.")))
+		.ContentPadding(FMargin(0, 5.f, 0, 4.f))
+		.OnClicked_Raw(this, &FChaosVDParticleMetadataCustomization::BrowseParticleOwner, AssetPath)
+		.Content()
+		[
+			SNew(SHorizontalBox)
+			+SHorizontalBox::Slot()
+			.Padding(FMargin(3, 0, 0, 0))
+			.VAlign(VAlign_Center)
+			.AutoWidth()
+			[
+				SNew(STextBlock)
+				.TextStyle(FAppStyle::Get(), "SmallButtonText")
+				.Text(LOCTEXT("BrowseParticleOwnerText","Attempt To Find Find in Content Browser"))
+			]
+		]
+	];
+
+	return BrowseButton;
+}
+
+FReply FChaosVDParticleMetadataCustomization::BrowseParticleOwner(FTopLevelAssetPath AssetPath)
+{
+	if (GEditor)
+	{
+		if (UObject* OwnerObject = FindObject<UObject>(AssetPath))
+		{
+			GEditor->SyncBrowserToObject(OwnerObject);
+		}
+		else
+		{
+			FMessageDialog::Open(EAppMsgType::Ok, LOCTEXT("FailedToFindOwnerMessage", " Failed to find the asset from which this particle was created"), LOCTEXT("FailedToFindOwnerTitle", "CVD Asset Finder"));
+		}
+	}
+	
+	return FReply::Handled();
+}
+#undef LOCTEXT_NAMESPACE
diff --git a/Engine/Plugins/ChaosVD/Source/ChaosVD/Private/DetailsCustomizations/ChaosVDParticleMetadataCustomization.h b/Engine/Plugins/ChaosVD/Source/ChaosVD/Private/DetailsCustomizations/ChaosVDParticleMetadataCustomization.h
new file mode 100644
--- /dev/null
+++ b/Engine/Plugins/ChaosVD/Source/ChaosVD/Private/DetailsCustomizations/ChaosVDParticleMetadataCustomization.h
@@ -0,0 +1,24 @@
+// Copyright Epic Games, Inc. All Rights Reserved.
+
+#pragma once
+
+#include "IPropertyTypeCustomization.h"
+#include "PropertyEditorModule.h"
+#include "Templates/SharedPointer.h"
+
+class FDetailWidgetRow;
+class IDetailChildrenBuilder;
+
+/** Property Customization for particle metadata in the Chaos Visual Debugger*/
+class FChaosVDParticleMetadataCustomization : public IPropertyTypeCustomization
+{
+public:
+	static TSharedRef<IPropertyTypeCustomization> MakeInstance();
+
+	virtual void CustomizeHeader(TSharedRef<IPropertyHandle> StructPropertyHandle, FDetailWidgetRow& HeaderRow, IPropertyTypeCustomizationUtils& StructCustomizationUtils) override;
+	virtual void CustomizeChildren(TSharedRef<IPropertyHandle> StructPropertyHandle, IDetailChildrenBuilder& StructBuilder, IPropertyTypeCustomizationUtils& StructCustomizationUtils) override;
+
+private:
+	TSharedRef<SWidget> GenerateBrowseParticleOwnerButton(FTopLevelAssetPath AssetPath);
+	FReply BrowseParticleOwner(FTopLevelAssetPath AssetPath);
+};
diff --git a/Engine/Plugins/ChaosVD/Source/ChaosVD/Private/Trace/ChaosVDTraceProvider.cpp b/Engine/Plugins/ChaosVD/Source/ChaosVD/Private/Trace/ChaosVDTraceProvider.cpp
--- a/Engine/Plugins/ChaosVD/Source/ChaosVD/Private/Trace/ChaosVDTraceProvider.cpp
+++ b/Engine/Plugins/ChaosVD/Source/ChaosVD/Private/Trace/ChaosVDTraceProvider.cpp
@@ -9,6 +9,7 @@
 #include "Compression/OodleDataCompressionUtil.h"
 #include "Containers/Ticker.h"
 #include "DataProcessors/ChaosVDCollisionChannelsInfoDataProcessor.h"
+#include "DataProcessors/ChaosVDParticleMetadataProcessor.h"
 #include "ExtensionsSystem/ChaosVDExtensionsManager.h"
 #include "Misc/MessageDialog.h"
 #include "Settings/ChaosVDGeneralSettings.h"
@@ -165,6 +166,21 @@ int32 FChaosVDTraceProvider::GetRemappedSolverID(int32 SolverID)
 	return RemappedSolverID;
 }
 
+void FChaosVDTraceProvider::AddParticleMetadata(uint64 MetadaId, const TSharedPtr<FChaosVDParticleMetadata>& InMetadata)
+{
+	if (!SerializedParticleMetadata.Contains(MetadaId))
+	{
+		SerializedParticleMetadata.Add(MetadaId, InMetadata);
+	}
+}
+
+TSharedPtr<FChaosVDParticleMetadata> FChaosVDTraceProvider::GetParticleMetadata(uint64 MetadataId)
+{
+	TSharedPtr<FChaosVDParticleMetadata>* MetadataInstance = SerializedParticleMetadata.Find(MetadataId);
+
+	return MetadataInstance ? *MetadataInstance : nullptr;
+}
+
 void FChaosVDTraceProvider::CommitProcessedGameFramesToRecording()
 {
 	TArray<int32, TInlineAllocator<16>> SolverIDs;
@@ -498,6 +514,10 @@ void FChaosVDTraceProvider::RegisterDefaultDataProcessorsIfNeeded()
 	CollisionChannelsInfoDataProcessor->SetTraceProvider(AsShared());
 	RegisterDataProcessor(CollisionChannelsInfoDataProcessor);
 
+	TSharedPtr<FChaosVDParticleMetadataProcessor> ParticleMetadataProcessor = MakeShared<FChaosVDParticleMetadataProcessor>();
+	ParticleMetadataProcessor->SetTraceProvider(AsShared());
+	RegisterDataProcessor(ParticleMetadataProcessor);
+
 	FChaosVDExtensionsManager::Get().EnumerateExtensions([this](const TSharedRef<FChaosVDExtension>& Extension)
 	{
 		Extension->RegisterDataProcessorsInstancesForProvider(StaticCastSharedRef<FChaosVDTraceProvider>(AsShared()));
diff --git a/Engine/Plugins/ChaosVD/Source/ChaosVD/Private/Trace/DataProcessors/ChaosVDParticleMetadataProcessor.cpp b/Engine/Plugins/ChaosVD/Source/ChaosVD/Private/Trace/DataProcessors/ChaosVDParticleMetadataProcessor.cpp
new file mode 100644
--- /dev/null
+++ b/Engine/Plugins/ChaosVD/Source/ChaosVD/Private/Trace/DataProcessors/ChaosVDParticleMetadataProcessor.cpp
@@ -0,0 +1,28 @@
+// Copyright Epic Games, Inc. All Rights Reserved.
+
+#include "ChaosVDParticleMetadataProcessor.h"
+
+FChaosVDParticleMetadataProcessor::FChaosVDParticleMetadataProcessor() : FChaosVDDataProcessorBase(FChaosVDParticleMetadata::WrapperTypeName)
+{
+}
+
+bool FChaosVDParticleMetadataProcessor::ProcessRawData(const TArray<uint8>& InData)
+{
+	FChaosVDDataProcessorBase::ProcessRawData(InData);
+
+	TSharedPtr<FChaosVDTraceProvider> ProviderSharedPtr = TraceProvider.Pin();
+	if (!ensure(ProviderSharedPtr.IsValid()))
+	{
+		return false;
+	}
+
+	TSharedPtr<FChaosVDParticleMetadata> ParticleMetaData = MakeShared<FChaosVDParticleMetadata>();
+	const bool bSuccess = Chaos::VisualDebugger::ReadDataFromBuffer(InData, *ParticleMetaData, ProviderSharedPtr.ToSharedRef());
+
+	if (bSuccess)
+	{
+		ProviderSharedPtr->AddParticleMetadata(ParticleMetaData->MetadataID, ParticleMetaData);
+	}
+
+	return bSuccess;
+}
diff --git a/Engine/Plugins/ChaosVD/Source/ChaosVD/Private/Trace/DataProcessors/ChaosVDParticleMetadataProcessor.h b/Engine/Plugins/ChaosVD/Source/ChaosVD/Private/Trace/DataProcessors/ChaosVDParticleMetadataProcessor.h
new file mode 100644
--- /dev/null
+++ b/Engine/Plugins/ChaosVD/Source/ChaosVD/Private/Trace/DataProcessors/ChaosVDParticleMetadataProcessor.h
@@ -0,0 +1,17 @@
+// Copyright Epic Games, Inc. All Rights Reserved.
+
+#pragma once
+
+#include "Trace/DataProcessors/ChaosVDDataProcessorBase.h"
+
+/**
+ * Data processor implementation that is able to deserialize traced Particles meta data
+ */
+class FChaosVDParticleMetadataProcessor final : public FChaosVDDataProcessorBase
+{
+public:
+	virtual ~FChaosVDParticleMetadataProcessor() override = default;
+	explicit FChaosVDParticleMetadataProcessor();
+
+	virtual bool ProcessRawData(const TArray<uint8>& InData) override;
+};
diff --git a/Engine/Plugins/ChaosVD/Source/ChaosVD/Private/Trace/DataProcessors/ChaosVDTraceParticleDataProcessor.cpp b/Engine/Plugins/ChaosVD/Source/ChaosVD/Private/Trace/DataProcessors/ChaosVDTraceParticleDataProcessor.cpp
--- a/Engine/Plugins/ChaosVD/Source/ChaosVD/Private/Trace/DataProcessors/ChaosVDTraceParticleDataProcessor.cpp
+++ b/Engine/Plugins/ChaosVD/Source/ChaosVD/Private/Trace/DataProcessors/ChaosVDTraceParticleDataProcessor.cpp
@@ -36,26 +36,27 @@ bool FChaosVDTraceParticleDataProcessor::ProcessRawData(const TArray<uint8>& InD
 		{
 			int32 NewParticleDataIndex = INDEX_NONE;
 
-			FString& DebugName = ParticleData->DebugName;
-			ParticleData->DebugNameHash = CityHash64(reinterpret_cast<const char*>(*DebugName), DebugName.Len() * sizeof(TCHAR));
+			if (ParticleData->HasLegacyDebugName())
+			{
+				const FString& DebugName = ParticleData->GetDebugName();
+				ParticleData->MetadataId = CityHash64(reinterpret_cast<const char*>(*DebugName), DebugName.Len() * sizeof(TCHAR));
+			}
+			else
+			{
+				ParticleData->SetMetadataInstance(ProviderSharedPtr->GetParticleMetadata(ParticleData->MetadataId));
+			}
 
-			// Non-Explicit stages might have duplicated data if a particle was traced twice. As CVD only shows the last state per stage, we need to just keep the last seen state
+			// Some stages like the auto generated "In-Between" stage, or "Rewind Frame" might have duplicated data if a particle was traced twice.
+			// As CVD only shows the last state per stage, we need to just keep the last seen state
 			// We do this on trace analysis to not do the de-duplication process during playback
-			if (!EnumHasAnyFlags(CurrentSolverStage->StageFlags, EChaosVDSolverStageFlags::ExplicitStage))
+			if (int32* ExistingParticleDataIndex = CurrentSolverStage->CurrentRecordedParticlesIndexes.Find(ParticleData->ParticleIndex))
 			{
-				if (int32* ExistingParticleDataIndex = CurrentSolverStage->CurrentRecordedParticlesIndexes.Find(ParticleData->ParticleIndex))
-				{
-					CurrentSolverStage->RecordedParticlesData[*ExistingParticleDataIndex] = ParticleData;
-				}
-				else
-				{
-					NewParticleDataIndex = CurrentSolverStage->RecordedParticlesData.Add(ParticleData);
-				}
+				CurrentSolverStage->RecordedParticlesData[*ExistingParticleDataIndex] = ParticleData;
 			}
 			else
 			{
 				NewParticleDataIndex = CurrentSolverStage->RecordedParticlesData.Add(ParticleData);
-			}
+			}	
 
 			if (NewParticleDataIndex != INDEX_NONE)
 			{
diff --git a/Engine/Plugins/ChaosVD/Source/ChaosVD/Private/Widgets/SChaosVDCollisionDataInspector.cpp b/Engine/Plugins/ChaosVD/Source/ChaosVD/Private/Widgets/SChaosVDCollisionDataInspector.cpp
--- a/Engine/Plugins/ChaosVD/Source/ChaosVD/Private/Widgets/SChaosVDCollisionDataInspector.cpp
+++ b/Engine/Plugins/ChaosVD/Source/ChaosVD/Private/Widgets/SChaosVDCollisionDataInspector.cpp
@@ -157,7 +157,7 @@ FText SChaosVDCollisionDataInspector::GetParticleName(EChaosVDParticlePairIndex
 		{
 			if (TSharedPtr<const FChaosVDParticleDataWrapper> ParticleData = ParticleActor->GetParticleData())
 			{
-				return FText::AsCultureInvariant(ParticleData->DebugName);
+				return FText::AsCultureInvariant(ParticleData->GetDebugName());
 			}
 		}
 	}
diff --git a/Engine/Plugins/ChaosVD/Source/ChaosVD/Private/Widgets/SChaosVDConstraintDataInspector.cpp b/Engine/Plugins/ChaosVD/Source/ChaosVD/Private/Widgets/SChaosVDConstraintDataInspector.cpp
--- a/Engine/Plugins/ChaosVD/Source/ChaosVD/Private/Widgets/SChaosVDConstraintDataInspector.cpp
+++ b/Engine/Plugins/ChaosVD/Source/ChaosVD/Private/Widgets/SChaosVDConstraintDataInspector.cpp
@@ -309,7 +309,7 @@ FText SChaosVDConstraintDataInspector::GetParticleName_Internal(int32 SolverID,
 		{
 			if (TSharedPtr<const FChaosVDParticleDataWrapper> ParticleData = ParticleActor->GetParticleData())
 			{
-				return FText::AsCultureInvariant(ParticleData->DebugName);
+				return FText::AsCultureInvariant(ParticleData->GetDebugName());
 			}
 		}
 	}
diff --git a/Engine/Plugins/ChaosVD/Source/ChaosVD/Private/Widgets/SChaosVDMainTab.cpp b/Engine/Plugins/ChaosVD/Source/ChaosVD/Private/Widgets/SChaosVDMainTab.cpp
--- a/Engine/Plugins/ChaosVD/Source/ChaosVD/Private/Widgets/SChaosVDMainTab.cpp
+++ b/Engine/Plugins/ChaosVD/Source/ChaosVD/Private/Widgets/SChaosVDMainTab.cpp
@@ -44,6 +44,7 @@
 #include "Components/ChaosVDStaticMeshComponent.h"
 #include "DetailsCustomizations/ChaosVDGeometryComponentCustomization.h"
 #include "DetailsCustomizations/ChaosVDParticleDataWrapperCustomization.h"
+#include "DetailsCustomizations/ChaosVDParticleMetadataCustomization.h"
 #include "DetailsCustomizations/ChaosVDQueryDataWrappersCustomizationDetails.h"
 #include "DetailsCustomizations/ChaosVDSelectionMultipleViewCustomization.h"
 #include "DetailsCustomizations/ChaosVDShapeDataCustomization.h"
@@ -860,6 +861,7 @@ void SChaosVDMainTab::SetCustomPropertyLayouts(IDetailsView* DetailsView)
 	DetailsView->RegisterInstancedCustomPropertyTypeLayout(TEXT("ChaosVDCollisionObjectQueryParams"), FOnGetPropertyTypeCustomizationInstance::CreateStatic(&FChaosVDCollisionObjectParamsCustomization::MakeInstance,StaticCastWeakPtr<SChaosVDMainTab>(AsWeak())));
 	DetailsView->RegisterInstancedCustomPropertyTypeLayout(TEXT("ChaosVDShapeCollisionData"), FOnGetPropertyTypeCustomizationInstance::CreateStatic(&FChaosVDShapeDataCustomization::MakeInstance, StaticCastWeakPtr<SChaosVDMainTab>(AsWeak())));
 	DetailsView->RegisterInstancedCustomPropertyTypeLayout(TEXT("ChaosVDQueryVisitStep"), FOnGetPropertyTypeCustomizationInstance::CreateStatic(&FChaosVDQueryVisitDataPropertyCustomization::MakeInstance, StaticCastWeakPtr<SChaosVDMainTab>(AsWeak())));
+	DetailsView->RegisterInstancedCustomPropertyTypeLayout(TEXT("ChaosVDParticleMetadata"), FOnGetPropertyTypeCustomizationInstance::CreateStatic(&FChaosVDParticleMetadataCustomization::MakeInstance));
 
 	// We don't need to validate properties, and trying to do so seems to be costing between 2-4ms per tick!
 	DetailsView->SetCustomValidatePropertyNodesFunction(FOnValidateDetailsViewPropertyNodes::CreateLambda([](const FRootPropertyNodeList& Root){ return true;}));
diff --git a/Engine/Plugins/ChaosVD/Source/ChaosVD/Public/ChaosVDSceneParticle.h b/Engine/Plugins/ChaosVD/Source/ChaosVD/Public/ChaosVDSceneParticle.h
--- a/Engine/Plugins/ChaosVD/Source/ChaosVD/Public/ChaosVDSceneParticle.h
+++ b/Engine/Plugins/ChaosVD/Source/ChaosVD/Public/ChaosVDSceneParticle.h
@@ -93,11 +93,8 @@ public:
 	FBox GetInflatedBoundingBox() const;
 	Chaos::TAABB<double, 3> GetChaosBoundingBox() const;
 
-	// BEGIN IChaosVDCollisionDataProvider Interface
 	TConstArrayView<TSharedPtr<FChaosVDParticlePairMidPhase>> GetCollisionData();
 	bool HasCollisionData();
-	FName GetProviderName();
-	// END IChaosVDCollisionDataProvider Interface
 
 	// BEGIN IChaosVDCharacterGroundConstraintDataProviderInterface
 	void GetCharacterGroundConstraintData(TArray<TSharedPtr<FChaosVDCharacterGroundConstraint>>& OutConstraintsFound);
diff --git a/Engine/Plugins/ChaosVD/Source/ChaosVD/Public/Trace/ChaosVDTraceProvider.h b/Engine/Plugins/ChaosVD/Source/ChaosVD/Public/Trace/ChaosVDTraceProvider.h
--- a/Engine/Plugins/ChaosVD/Source/ChaosVD/Public/Trace/ChaosVDTraceProvider.h
+++ b/Engine/Plugins/ChaosVD/Source/ChaosVD/Public/Trace/ChaosVDTraceProvider.h
@@ -170,6 +170,9 @@ public:
 		return TypesFailedToSerialize;
 	}
 
+	void AddParticleMetadata(uint64 MetadaId, const TSharedPtr<FChaosVDParticleMetadata>& InMetadata);
+	TSharedPtr<FChaosVDParticleMetadata> GetParticleMetadata(uint64 MetadataId);
+
 private:
 	
 	int32 GetCurrentGameThreadTrackID() const
@@ -223,6 +226,8 @@ private:
 	bool bShouldTrimOutStartEmptyFrames = false;
 	int32 MaxGameFramesToQueueNum = 10;
 
+	TMap<uint64, TSharedPtr<FChaosVDParticleMetadata>> SerializedParticleMetadata;
+
 	friend class FChaosVDTraceAnalyzer;
 
 	// Not idea, but needed to for the serialization fix up logic of the collision data channel info
diff --git a/Engine/Plugins/Experimental/ControlRigPhysics/Source/ControlRigPhysics/Private/RigPhysicsSimulation.cpp b/Engine/Plugins/Experimental/ControlRigPhysics/Source/ControlRigPhysics/Private/RigPhysicsSimulation.cpp
--- a/Engine/Plugins/Experimental/ControlRigPhysics/Source/ControlRigPhysics/Private/RigPhysicsSimulation.cpp
+++ b/Engine/Plugins/Experimental/ControlRigPhysics/Source/ControlRigPhysics/Private/RigPhysicsSimulation.cpp
@@ -35,6 +35,7 @@
 #include "Chaos/PBDJointConstraintUtilities.h"
 #include "Chaos/Evolution/SimulationSpace.h"
 #include "Chaos/PBDCollisionConstraints.h"
+#include "Chaos/ChaosDebugNameDefines.h"
 
 #include "PhysicsEngine/BodyInstance.h"
 #include "PhysicsEngine/ConstraintInstance.h"
@@ -124,7 +125,7 @@ void FRigPhysicsSimulation::InitialiseSimulation(const FRigPhysicsSolverComponen
 
 	Simulation = MakeShared<ImmediatePhysics::FSimulation>();
 
-#if CHAOS_DEBUG_NAME
+#if CHAOS_SOLVER_DEBUG_NAME
 	FString SimName = TEXT("ControlRigPhysics-") + OwnerName.ToString();
 	Simulation->SetDebugName(FName(SimName));
 #endif
diff --git a/Engine/Source/Runtime/Core/Public/UObject/FortniteMainBranchObjectVersions.inl b/Engine/Source/Runtime/Core/Public/UObject/FortniteMainBranchObjectVersions.inl
--- a/Engine/Source/Runtime/Core/Public/UObject/FortniteMainBranchObjectVersions.inl
+++ b/Engine/Source/Runtime/Core/Public/UObject/FortniteMainBranchObjectVersions.inl
@@ -669,5 +669,9 @@ FORTNITE_MAIN_VERSION(CompositeActorSceneCaptureRefactor, 218)
 
 // Moved HLOD Layer properties to an editor only optional object
 FORTNITE_MAIN_VERSION(HLODLayerEditorOnlyObject, 219)
+
+// Deduplicated particle debug names serialization in the Chaos Visual Debugger
+FORTNITE_MAIN_VERSION(DeduplicatedDebugNameSerializationInCVD, 220)
+
 // -----<new versions can be added above this line>-------------------------------------------------
 #undef FORTNITE_MAIN_VERSION
\ No newline at end of file
diff --git a/Engine/Source/Runtime/Engine/Private/PhysicsEngine/Experimental/PhysScene_Chaos.cpp b/Engine/Source/Runtime/Engine/Private/PhysicsEngine/Experimental/PhysScene_Chaos.cpp
--- a/Engine/Source/Runtime/Engine/Private/PhysicsEngine/Experimental/PhysScene_Chaos.cpp
+++ b/Engine/Source/Runtime/Engine/Private/PhysicsEngine/Experimental/PhysScene_Chaos.cpp
@@ -536,13 +536,13 @@ struct FPhysScenePendingComponentTransform_Chaos
 };
 
 FPhysScene_Chaos::FPhysScene_Chaos(AActor* InSolverActor
-#if CHAOS_DEBUG_NAME
+#if CHAOS_SOLVER_DEBUG_NAME
 	, const FName& DebugName
 #endif
 )
 	: Super(InSolverActor ? InSolverActor->GetWorld() : nullptr
 		, UPhysicsSettings::Get()->bTickPhysicsAsync ? UPhysicsSettings::Get()->AsyncFixedTimeStepSize : -1
-#if CHAOS_DEBUG_NAME
+#if CHAOS_SOLVER_DEBUG_NAME
 		, DebugName
 #endif
 	)
diff --git a/Engine/Source/Runtime/Engine/Private/PhysicsEngine/ImmediatePhysics/ImmediatePhysicsChaos/ImmediatePhysicsSimulation_Chaos.cpp b/Engine/Source/Runtime/Engine/Private/PhysicsEngine/ImmediatePhysics/ImmediatePhysicsChaos/ImmediatePhysicsSimulation_Chaos.cpp
--- a/Engine/Source/Runtime/Engine/Private/PhysicsEngine/ImmediatePhysics/ImmediatePhysicsChaos/ImmediatePhysicsSimulation_Chaos.cpp
+++ b/Engine/Source/Runtime/Engine/Private/PhysicsEngine/ImmediatePhysics/ImmediatePhysicsChaos/ImmediatePhysicsSimulation_Chaos.cpp
@@ -371,7 +371,7 @@ namespace ImmediatePhysics_Chaos
 
 		Implementation = MakeUnique<FImplementation>();
 
-#if CHAOS_DEBUG_NAME
+#if CHAOS_SOLVER_DEBUG_NAME
 		DebugName = TEXT("RBAN");
 #endif
 
diff --git a/Engine/Source/Runtime/Engine/Private/World.cpp b/Engine/Source/Runtime/Engine/Private/World.cpp
--- a/Engine/Source/Runtime/Engine/Private/World.cpp
+++ b/Engine/Source/Runtime/Engine/Private/World.cpp
@@ -85,6 +85,7 @@
 #include "Streaming/StreamingWorldSubsystemInterface.h"
 #include "Components/UnregisterComponentContext.h"
 #include "Misc/ScopeRWLock.h"
+#include "Chaos/ChaosDebugNameDefines.h"
 
 #if WITH_EDITOR
 	#include "DerivedDataCacheInterface.h"
@@ -6627,16 +6628,16 @@ bool UWorld::AreActorsInitialized() const
 
 void UWorld::CreatePhysicsScene(const AWorldSettings* Settings)
 {
-#if CHAOS_DEBUG_NAME
+#if CHAOS_SOLVER_DEBUG_NAME
 	TStringBuilder<128> PhysicsSceneNameBuilder;
- 	PhysicsSceneNameBuilder.Appendf(TEXT("%s | World Type [%s]"), IsNetMode(NM_DedicatedServer) ? TEXT("Server Physics") : TEXT("Client Physics"), LexToString(WorldType));
+ 	PhysicsSceneNameBuilder.Appendf(TEXT("%s | %s | World Type [%s]"), IsNetMode(NM_DedicatedServer) ? TEXT("Server") : TEXT("Client"), *GetMapName(), LexToString(WorldType));
  
 	#if UE_WITH_REMOTE_OBJECT_HANDLE
  		PhysicsSceneNameBuilder.Appendf(TEXT(" | Remote Server ID [%s]"), *FRemoteServerId::GetLocalServerId().ToString());
 	#endif
 
 	const FName PhysicsName = FName(PhysicsSceneNameBuilder);
- 
+
 	FPhysScene* NewScene = new FPhysScene(nullptr, PhysicsName);
 #else
 	FPhysScene* NewScene = new FPhysScene(nullptr);
diff --git a/Engine/Source/Runtime/Engine/Public/Physics/Experimental/PhysScene_Chaos.h b/Engine/Source/Runtime/Engine/Public/Physics/Experimental/PhysScene_Chaos.h
--- a/Engine/Source/Runtime/Engine/Public/Physics/Experimental/PhysScene_Chaos.h
+++ b/Engine/Source/Runtime/Engine/Public/Physics/Experimental/PhysScene_Chaos.h
@@ -15,6 +15,7 @@
 #include "Chaos/ContactModification.h"
 #include "Chaos/Real.h"
 #include "Chaos/AsyncInitBodyHelper.h"
+#include "Chaos/ChaosDebugNameDefines.h"
 #include "UObject/ObjectKey.h"
 
 #ifndef CHAOS_WITH_PAUSABLE_SOLVER
@@ -118,7 +119,7 @@ public:
 	using Super = FChaosScene;
 	
 	ENGINE_API FPhysScene_Chaos(AActor* InSolverActor=nullptr
-#if CHAOS_DEBUG_NAME
+#if CHAOS_SOLVER_DEBUG_NAME
 	, const FName& DebugName=NAME_None
 #endif
 );
diff --git a/Engine/Source/Runtime/Engine/Public/Physics/ImmediatePhysics/ImmediatePhysicsChaos/ImmediatePhysicsSimulation_Chaos.h b/Engine/Source/Runtime/Engine/Public/Physics/ImmediatePhysics/ImmediatePhysicsChaos/ImmediatePhysicsSimulation_Chaos.h
--- a/Engine/Source/Runtime/Engine/Public/Physics/ImmediatePhysics/ImmediatePhysicsChaos/ImmediatePhysicsSimulation_Chaos.h
+++ b/Engine/Source/Runtime/Engine/Public/Physics/ImmediatePhysics/ImmediatePhysicsChaos/ImmediatePhysicsSimulation_Chaos.h
@@ -9,6 +9,7 @@
 #include "Chaos/Collision/CollisionVisitor.h"
 #include "Engine/EngineTypes.h"
 #include "Templates/UniquePtr.h"
+#include "Chaos/ChaosDebugNameDefines.h"
 
 namespace Chaos
 {
@@ -210,7 +211,7 @@ namespace ImmediatePhysics_Chaos
 		struct FImplementation;
 		TUniquePtr<FImplementation> Implementation;
 
-#if CHAOS_DEBUG_NAME
+#if CHAOS_SOLVER_DEBUG_NAME
 	public:
 		void SetDebugName(const FName& Name)
 		{
diff --git a/Engine/Source/Runtime/Experimental/Chaos/Private/Chaos/Framework/PhysicsSolverBase.cpp b/Engine/Source/Runtime/Experimental/Chaos/Private/Chaos/Framework/PhysicsSolverBase.cpp
--- a/Engine/Source/Runtime/Experimental/Chaos/Private/Chaos/Framework/PhysicsSolverBase.cpp
+++ b/Engine/Source/Runtime/Experimental/Chaos/Private/Chaos/Framework/PhysicsSolverBase.cpp
@@ -262,7 +262,7 @@ namespace Chaos
 		}
 	}
 
-#if CHAOS_DEBUG_NAME
+#if CHAOS_SOLVER_DEBUG_NAME
 	void FPhysicsSolverBase::SetDebugName(const FName& Name)
 	{
 		DebugName = Name;
@@ -272,7 +272,7 @@ namespace Chaos
 
 	FName FPhysicsSolverBase::GetDebugName() const
 	{
-#if CHAOS_DEBUG_NAME
+#if CHAOS_SOLVER_DEBUG_NAME
 		return DebugName;
 #else
 		return NAME_None;
diff --git a/Engine/Source/Runtime/Experimental/Chaos/Private/ChaosSolversModule.cpp b/Engine/Source/Runtime/Experimental/Chaos/Private/ChaosSolversModule.cpp
--- a/Engine/Source/Runtime/Experimental/Chaos/Private/ChaosSolversModule.cpp
+++ b/Engine/Source/Runtime/Experimental/Chaos/Private/ChaosSolversModule.cpp
@@ -17,6 +17,7 @@
 #include "UObject/Class.h"
 #include "Misc/App.h"
 #include "Chaos/PhysicalMaterials.h"
+#include "Chaos/ChaosDebugNameDefines.h"
 
 TAutoConsoleVariable<int32> CVarChaosThreadEnabled(
 	TEXT("p.Chaos.DedicatedThreadEnabled"),
@@ -190,7 +191,7 @@ Chaos::FPBDRigidsSolver* FChaosSolversModule::CreateSolver(UObject* InOwner, Cha
 	TArray<FPhysicsSolverBase*>& OwnerSolverList = SolverMap.FindOrAdd(InOwner);
 	OwnerSolverList.Add(NewSolver);
 
-#if CHAOS_DEBUG_NAME
+#if CHAOS_SOLVER_DEBUG_NAME
     // Add solver number to solver name
 	const FName NewDebugName = *FString::Printf(TEXT("%s (%d)"), DebugName == NAME_None ? TEXT("Solver") : *DebugName.ToString(), AllSolvers.Num() - 1);
 	NewSolver->SetDebugName(NewDebugName);
diff --git a/Engine/Source/Runtime/Experimental/Chaos/Private/ChaosVisualDebugger/ChaosVDDataWrapperUtils.cpp b/Engine/Source/Runtime/Experimental/Chaos/Private/ChaosVisualDebugger/ChaosVDDataWrapperUtils.cpp
--- a/Engine/Source/Runtime/Experimental/Chaos/Private/ChaosVisualDebugger/ChaosVDDataWrapperUtils.cpp
+++ b/Engine/Source/Runtime/Experimental/Chaos/Private/ChaosVisualDebugger/ChaosVDDataWrapperUtils.cpp
@@ -11,7 +11,6 @@
 #include "Chaos/Character/CharacterGroundConstraintContainer.h"
 #include "Chaos/Collision/ParticlePairMidPhase.h"
 #include "Chaos/ParticleHandle.h"
-#include "ChaosVisualDebugger/ChaosVDSerializedNameTable.h"
 #include "DataWrappers/ChaosVDCollisionDataWrappers.h"
 #include "DataWrappers/ChaosVDParticleDataWrapper.h"
 #include "Math/UnitConversion.h"
@@ -97,13 +96,6 @@ FChaosVDParticleDataWrapper FChaosVDDataWrapperUtils::BuildParticleDataWrapperFr
 	WrappedParticleData.ParticleIndex = ParticleHandlePtr->UniqueIdx().Idx;
 	WrappedParticleData.Type =  static_cast<EChaosVDParticleType>(ParticleHandlePtr->Type);
 
-#if CHAOS_DEBUG_NAME
-	// Passing it as a Ptr because from here until it is serialized right after this function ends this string does not change
-	// Passing it as a sharedptr has an additional 20% cost as it has to increment the reference counter, which adds up
-	// TODO: We should switch to FName to take advantage of the new CVD Serializable name table so they can be de-duplicated, but to do so we need to change how we create our debug names to not be unique strings
-	WrappedParticleData.DebugNamePtr = ParticleHandlePtr->DebugName().IsValid() ? ParticleHandlePtr->DebugName().Get() : nullptr;
-#endif
-
 	if (const Chaos::TKinematicGeometryParticleHandleImp<Chaos::FReal, 3, true>* KinematicParticle = ParticleHandlePtr->CastToKinematicParticle())
 	{
 		WrappedParticleData.ParticleVelocities.CopyFrom(*KinematicParticle);
@@ -143,13 +135,6 @@ FChaosVDParticleDataWrapper FChaosVDDataWrapperUtils::BuildParticleDataWrapperFr
 
 	WrappedParticleData.ParticleIndex = GTParticle->UniqueIdx().Idx;
 	WrappedParticleData.Type =  static_cast<EChaosVDParticleType>(GTParticle->ObjectType());
-	
-#if CHAOS_DEBUG_NAME
-	// Passing it as a Ptr because from here until it is serialized right after this function ends this string does not change
-	// Passing it as a sharedptr has an additional 20% cost as it has to increment the reference counter, which adds up
-	// TODO: We should switch to FName to take advantage of the new CVD Serializable name table so they can be de-duplicated, but to do so we need to change how we create our debug names to not be unique strings
-	WrappedParticleData.DebugNamePtr = GTParticle->DebugName().Get();
-#endif
 
 	WrappedParticleData.ParticlePositionRotation.CopyFrom<Chaos::FGeometryParticle, FChaosVDParticlePositionRotation::EAccessorType::XR>(*GTParticle);
 
diff --git a/Engine/Source/Runtime/Experimental/Chaos/Private/ChaosVisualDebugger/ChaosVisualDebuggerTrace.cpp b/Engine/Source/Runtime/Experimental/Chaos/Private/ChaosVisualDebugger/ChaosVisualDebuggerTrace.cpp
--- a/Engine/Source/Runtime/Experimental/Chaos/Private/ChaosVisualDebugger/ChaosVisualDebuggerTrace.cpp
+++ b/Engine/Source/Runtime/Experimental/Chaos/Private/ChaosVisualDebugger/ChaosVisualDebuggerTrace.cpp
@@ -18,6 +18,7 @@
 #include "ChaosVisualDebugger/ChaosVDMemWriterReader.h"
 #include "ChaosVisualDebugger/ChaosVDSerializedNameTable.h"
 #include "Compression/OodleDataCompressionUtil.h"
+#include "Containers/StripedMap.h"
 #include "DataWrappers/ChaosVDCharacterGroundConstraintDataWrappers.h"
 #include "DataWrappers/ChaosVDCollisionDataWrappers.h"
 #include "DataWrappers/ChaosVDDebugShapeDataWrapper.h"
@@ -66,65 +67,38 @@ namespace Chaos::VisualDebugger::Cvars
 	TEXT("Oodle compression mode to use, 4 is by default which equsals to ECompressionLevel::VeryFast"));
 }
 
-/** Struct where we keep track of the geometry we are tracing */
-struct FChaosVDGeometryTraceContext
+namespace Chaos::VD
 {
-	FTransactionallySafeRWLock TracedGeometrySetLock;
-	FTransactionallySafeRWLock CachedGeometryHashesLock;
-	TSet<uint32> GeometryTracedIDs;
-
-	uint32 GetGeometryHashForImplicit(const Chaos::FImplicitObject* Implicit)
+	FRecordingSessionState::FRecordingSessionState() : CVDNameTable(MakeShared<FChaosVDSerializableNameTable>()),
+	TraceRelevancyVolume(FBox(EForceInit::ForceInitToZero)),
+	bIsTracing(false)
 	{
-		if (Implicit == nullptr)
-		{
-			return 0;
-		}
+	}
 
+	namespace Private
+	{
+		FChaosVDParticleMetadata GenerateParticleMetadata(const FGeometryParticleHandle* ParticleHandle, FRecordingSessionState& SessionState)
 		{
-			UE::TReadScopeLock ReadLock(CachedGeometryHashesLock);
-			if (uint32* FoundHashPtr = CachedGeometryHashes.Find((void*) Implicit))
+			if (!ParticleHandle)
 			{
-				return *FoundHashPtr;
+				return FChaosVDParticleMetadata();
 			}
-		}
 
-		{
-			uint32 Hash = Implicit->GetTypeHash();
-
-			UE::TWriteScopeLock WriteLock(CachedGeometryHashesLock);
-			CachedGeometryHashes.Add((void*)Implicit, Hash);
-			return Hash;
-		}
-	}
-
-	void RemoveCachedGeometryHash(const Chaos::FImplicitObject* Implicit)
-	{
-		if (Implicit == nullptr)
-		{
-			return;
+			if (SessionState.ExternalParticleMetadataGenerator.IsBound())
+			{
+				FChaosVDParticleMetadata GeneratedMetadata = SessionState.ExternalParticleMetadataGenerator.Execute(ParticleHandle->PhysicsProxy(), ParticleHandle);
+				GeneratedMetadata.MetadataID = SessionState.GenerateUniqueID();
+				return GeneratedMetadata;
+			}
+			else
+			{
+				return FChaosVDParticleMetadata();
+			}
 		}
-
-		UE::TWriteScopeLock WriteLock(CachedGeometryHashesLock);
-		CachedGeometryHashes.Remove((void*)Implicit);
 	}
+}
 
-	//TODO: Remove this when/if we have stable geometry hashes that can be serialized
-	// Right now it is too costly calculate them each time at runtime
-	TMap<void*, uint32> CachedGeometryHashes;
-};
-
-static FChaosVDGeometryTraceContext GeometryTracerObject = FChaosVDGeometryTraceContext();
-
-FDelegateHandle FChaosVisualDebuggerTrace::RecordingStartedDelegateHandle = FDelegateHandle();
-FDelegateHandle FChaosVisualDebuggerTrace::RecordingStoppedDelegateHandle = FDelegateHandle();
-FDelegateHandle FChaosVisualDebuggerTrace::RecordingFullCaptureRequestedHandle = FDelegateHandle();
-
-FRWLock FChaosVisualDebuggerTrace::DeltaRecordingStatesLock = FRWLock();
-TSet<int32> FChaosVisualDebuggerTrace::SolverIDsForDeltaRecording = TSet<int32>();
-TSet<int32> FChaosVisualDebuggerTrace::RequestedFullCaptureSolverIDs = TSet<int32>();
-TSharedRef<FChaosVDSerializableNameTable> FChaosVisualDebuggerTrace::CVDNameTable = MakeShared<FChaosVDSerializableNameTable>();
-std::atomic<bool> FChaosVisualDebuggerTrace::bIsTracing = false;
-FBox FChaosVisualDebuggerTrace::TraceRelevancyVolume = FBox(EForceInit::ForceInitToZero);
+Chaos::VD::FRecordingSessionState FChaosVisualDebuggerTrace::SessionState = Chaos::VD::FRecordingSessionState();
 
 void FChaosVisualDebuggerTrace::TraceParticle(const Chaos::FGeometryParticleHandle* ParticleHandle)
 {
@@ -163,31 +137,26 @@ void FChaosVisualDebuggerTrace::TraceParticle(Chaos::FGeometryParticleHandle* Pa
 
 	const Chaos::FImplicitObjectRef ParticleGeometry = ParticleHandle->GetGeometry();
 
-	const uint32 GeometryHash = GeometryTracerObject.GetGeometryHashForImplicit(ParticleGeometry);
+	FChaosVDParticleDataWrapper ParticleDataWrapper = FChaosVDDataWrapperUtils::BuildParticleDataWrapperFromParticle(ParticleHandle);
+	ParticleDataWrapper.GeometryHash = TraceImplicitObject(ParticleGeometry);
+	ParticleDataWrapper.MetadataId = TraceParticleMetadata(ParticleHandle);
+	ParticleDataWrapper.SolverID = ContextData.Id;
 	
-	TraceImplicitObject({ GeometryHash, ParticleGeometry });
-
+	const Chaos::FShapeInstanceArray& ShapesInstancesArray = ParticleHandle->ShapeInstances();
+	ParticleDataWrapper.CollisionDataPerShape.Reserve(ShapesInstancesArray.Num());
+	
+	for (const Chaos::FShapeInstancePtr& ShapeData : ShapesInstancesArray)
 	{
-		FChaosVDParticleDataWrapper ParticleDataWrapper = FChaosVDDataWrapperUtils::BuildParticleDataWrapperFromParticle(ParticleHandle);
-		ParticleDataWrapper.GeometryHash = GeometryHash;
-		ParticleDataWrapper.SolverID = ContextData.Id;
-		
-		const Chaos::FShapeInstanceArray& ShapesInstancesArray = ParticleHandle->ShapeInstances();
-		ParticleDataWrapper.CollisionDataPerShape.Reserve(ShapesInstancesArray.Num());
-		
-		for (const Chaos::FShapeInstancePtr& ShapeData : ShapesInstancesArray)
-		{
-			FChaosVDShapeCollisionData CVDCollisionData;
-			FChaosVDDataWrapperUtils::CopyShapeDataToWrapper(ShapeData, CVDCollisionData);
-			ParticleDataWrapper.CollisionDataPerShape.Add(MoveTemp(CVDCollisionData));
-		}
-		
-		FChaosVDScopedTLSBufferAccessor TLSDataBuffer;
+		FChaosVDShapeCollisionData CVDCollisionData;
+		FChaosVDDataWrapperUtils::CopyShapeDataToWrapper(ShapeData, CVDCollisionData);
+		ParticleDataWrapper.CollisionDataPerShape.Add(MoveTemp(CVDCollisionData));
+	}
+	
+	FChaosVDScopedTLSBufferAccessor TLSDataBuffer;
 
-		Chaos::VisualDebugger::WriteDataToBuffer(TLSDataBuffer.BufferRef, ParticleDataWrapper);
+	Chaos::VisualDebugger::WriteDataToBuffer(TLSDataBuffer.BufferRef, ParticleDataWrapper);
 
-		TraceBinaryData(TLSDataBuffer.BufferRef, FChaosVDParticleDataWrapper::WrapperTypeName);
-	}
+	TraceBinaryData(TLSDataBuffer.BufferRef, FChaosVDParticleDataWrapper::WrapperTypeName);
 }
 
 void FChaosVisualDebuggerTrace::TraceParticles(const Chaos::TGeometryParticleHandles<Chaos::FReal, 3>& ParticleHandles)
@@ -225,7 +194,7 @@ void FChaosVisualDebuggerTrace::TraceParticleDestroyed(const Chaos::FGeometryPar
 		return;
 	}
 
-	GeometryTracerObject.RemoveCachedGeometryHash(ParticleHandle->GetGeometry());
+	SessionState.TracedGeometryHashCache.RemoveCachedTraceObjectID(ParticleHandle->GetGeometry());
 	
 	const FChaosVDContext* CVDContextData = FChaosVDThreadContext::Get().GetCurrentContext(EChaosVDContextType::Solver);
 	if (!IsContextEnabledAndValid(CVDContextData))
@@ -253,7 +222,7 @@ bool FChaosVisualDebuggerTrace::IsInRelevancyVolume(const Chaos::FGeometryPartic
 		return false;
 	}
 
-	if (!TraceRelevancyVolume.IsValid)
+	if (!SessionState.TraceRelevancyVolume.IsValid)
 	{
 		return true;
 	}
@@ -270,7 +239,7 @@ bool FChaosVisualDebuggerTrace::IsInRelevancyVolume(const Chaos::FImplicitObject
 		return false;
 	}
 
-	if (!TraceRelevancyVolume.IsValid)
+	if (!SessionState.TraceRelevancyVolume.IsValid)
 	{
 		return true;
 	}
@@ -342,15 +311,15 @@ void FChaosVisualDebuggerTrace::TraceParticlesSoA(const Chaos::FPBDRigidsSOAs& P
 
 void FChaosVisualDebuggerTrace::SetupForFullCaptureIfNeeded(int32 SolverID, bool& bOutFullCaptureRequested)
 {
-	DeltaRecordingStatesLock.ReadLock();
-	bOutFullCaptureRequested = RequestedFullCaptureSolverIDs.Contains(SolverID) || !SolverIDsForDeltaRecording.Contains(SolverID);
-	DeltaRecordingStatesLock.ReadUnlock();
+	SessionState.DeltaRecordingStatesLock.ReadLock();
+	bOutFullCaptureRequested = SessionState.RequestedFullCaptureSolverIDs.Contains(SolverID) || !SessionState.SolverIDsForDeltaRecording.Contains(SolverID);
+	SessionState.DeltaRecordingStatesLock.ReadUnlock();
 
 	if (bOutFullCaptureRequested)
 	{
-		UE::TWriteScopeLock WriteLock(DeltaRecordingStatesLock);
-		SolverIDsForDeltaRecording.Remove(SolverID);
-		RequestedFullCaptureSolverIDs.Remove(SolverID);
+		UE::TWriteScopeLock WriteLock(SessionState.DeltaRecordingStatesLock);
+		SessionState.SolverIDsForDeltaRecording.Remove(SolverID);
+		SessionState.RequestedFullCaptureSolverIDs.Remove(SolverID);
 	}
 }
 
@@ -361,8 +330,8 @@ int32 FChaosVisualDebuggerTrace::GetSolverID(Chaos::FPhysicsSolverBase& Solver)
 
 bool FChaosVisualDebuggerTrace::ShouldPerformFullCapture(int32 SolverID)
 {
-	UE::TReadScopeLock ReadLock(DeltaRecordingStatesLock);
-	int32* FoundSolverID = SolverIDsForDeltaRecording.Find(SolverID);
+	UE::TReadScopeLock ReadLock(SessionState.DeltaRecordingStatesLock);
+	int32* FoundSolverID = SessionState.SolverIDsForDeltaRecording.Find(SolverID);
 
 	// If the solver ID is on the SolverIDsForDeltaRecording set, it means we should NOT perform a full capture
 	return FoundSolverID == nullptr;
@@ -662,10 +631,10 @@ void FChaosVisualDebuggerTrace::TraceSolverFrameEnd(const FChaosVDContext& Conte
 	}
 
 	{
-		UE::TWriteScopeLock WriteLock(DeltaRecordingStatesLock);
-		if (!SolverIDsForDeltaRecording.Contains(ContextData.Id))
+		UE::TWriteScopeLock WriteLock(SessionState.DeltaRecordingStatesLock);
+		if (!SessionState.SolverIDsForDeltaRecording.Contains(ContextData.Id))
 		{
-			SolverIDsForDeltaRecording.Add(ContextData.Id);
+			SessionState.SolverIDsForDeltaRecording.Add(ContextData.Id);
 		}
 	}
 
@@ -839,24 +808,79 @@ void FChaosVisualDebuggerTrace::TraceImplicitObject(FChaosVDImplicitObjectWrappe
 		return;
 	}
 
-	uint32 GeometryID = WrappedGeometryData.Hash;
+	ensureMsgf(false, TEXT("This method is deprecated and should not be used - The geometry provided will not be traced"));
+}
+
+uint32 FChaosVisualDebuggerTrace::TraceImplicitObject(Chaos::FImplicitObject* GeometryPtr)
+{
+	if (!IsTracing())
 	{
-		UE::TReadScopeLock ReadLock(GeometryTracerObject.TracedGeometrySetLock);
-		if (GeometryTracerObject.GeometryTracedIDs.Find(GeometryID))
-		{
-			return;
-		}
+		return 0;
 	}
 
+	if (!GeometryPtr)
 	{
-		UE::TWriteScopeLock WriteLock(GeometryTracerObject.TracedGeometrySetLock);
-		GeometryTracerObject.GeometryTracedIDs.Add(GeometryID);
+		return 0;
 	}
 
+	bool bGeometryAlreadyTraced = false;
+	const uint32 GeometryHash = SessionState.TracedGeometryHashCache.FindOrProduceTraceObjectID(GeometryPtr, [GeometryPtr](){ return GeometryPtr->GetTypeHash(); }, bGeometryAlreadyTraced);
+
+	/** If the geometry hash was not cached already, it means it is the first time we see this geometry and therefore we need to serialize */
+	if (bGeometryAlreadyTraced)
+	{
+		return GeometryHash;
+	}
+
+	FChaosVDImplicitObjectWrapper GeometryDataWrapper;
+	GeometryDataWrapper.Hash = GeometryHash;
+	GeometryDataWrapper.ImplicitObject = GeometryPtr;
+
 	FChaosVDScopedTLSBufferAccessor TLSDataBuffer;
-	Chaos::VisualDebugger::WriteDataToBuffer<FChaosVDImplicitObjectWrapper, Chaos::FChaosArchive>(TLSDataBuffer.BufferRef, WrappedGeometryData);
+	Chaos::VisualDebugger::WriteDataToBuffer<FChaosVDImplicitObjectWrapper, Chaos::FChaosArchive>(TLSDataBuffer.BufferRef, GeometryDataWrapper);
 
 	TraceBinaryData(TLSDataBuffer.BufferRef, FChaosVDImplicitObjectWrapper::WrapperTypeName);
+
+	return GeometryHash;
+}
+
+uint64 FChaosVisualDebuggerTrace::TraceParticleMetadata(Chaos::FGeometryParticleHandle* InParticleHandle)
+{
+	if (!IsTracing())
+	{
+		return Chaos::VD::FRecordingSessionState::InvalidUniqueID;
+	}
+
+	if (!InParticleHandle)
+	{
+		return Chaos::VD::FRecordingSessionState::InvalidUniqueID;
+	}
+
+	// We can't Generate Metadata without a valid physics proxy
+	if (!InParticleHandle->PhysicsProxy())
+	{
+		return Chaos::VD::FRecordingSessionState::InvalidUniqueID;
+	}
+
+	FChaosVDParticleMetadata OutParticleMetadata;
+	auto ParticleMetadataProducer = [InParticleHandle, &OutParticleMetadata]()
+	{
+		OutParticleMetadata = Chaos::VD::Private::GenerateParticleMetadata(InParticleHandle, SessionState);
+		return OutParticleMetadata.MetadataID;
+	};
+
+	bool bMetadataAlreadyTraced = false;
+	const uint64 MetadataID = SessionState.ParticleMetadataIDsCache.FindOrProduceTraceObjectID(InParticleHandle, ParticleMetadataProducer, bMetadataAlreadyTraced);
+
+	if (!bMetadataAlreadyTraced)
+	{
+		FChaosVDScopedTLSBufferAccessor TLSDataBuffer;
+		Chaos::VisualDebugger::WriteDataToBuffer<FChaosVDParticleMetadata>(TLSDataBuffer.BufferRef, OutParticleMetadata);
+
+		TraceBinaryData(TLSDataBuffer.BufferRef, FChaosVDParticleMetadata::WrapperTypeName);
+	}
+
+	return MetadataID;
 }
 
 void FChaosVisualDebuggerTrace::InvalidateGeometryFromCache(const Chaos::FImplicitObject* CachedGeometryToInvalidate)
@@ -866,7 +890,17 @@ void FChaosVisualDebuggerTrace::InvalidateGeometryFromCache(const Chaos::FImplic
 		return;
 	}
 
-	GeometryTracerObject.RemoveCachedGeometryHash(CachedGeometryToInvalidate);
+	SessionState.TracedGeometryHashCache.RemoveCachedTraceObjectID(CachedGeometryToInvalidate);
+}
+
+void FChaosVisualDebuggerTrace::InvalidateParticleMetadataFromCache(Chaos::FGeometryParticleHandle* InParticleToInvalidate)
+{
+	if (!InParticleToInvalidate)
+	{
+		return;
+	}
+
+	SessionState.ParticleMetadataIDsCache.RemoveCachedTraceObjectID(InParticleToInvalidate);
 }
 
 void FChaosVisualDebuggerTrace::TraceNonSolverLocation(const FVector& InLocation, FStringView DebugNameID)
@@ -941,10 +975,7 @@ void FChaosVisualDebuggerTrace::TraceSceneQueryStart(const Chaos::FImplicitObjec
 
 	if (InputGeometry)
 	{
-		const uint32 GeometryHash = GeometryTracerObject.GetGeometryHashForImplicit(InputGeometry);
-		TraceImplicitObject({ GeometryHash, const_cast<Chaos::FImplicitObject*>(InputGeometry) });
-	
-		WrappedQueryData.InputGeometryKey = GeometryHash;
+		WrappedQueryData.InputGeometryKey = TraceImplicitObject(const_cast<Chaos::FImplicitObject*>(InputGeometry));
 	}
 	
 	WrappedQueryData.ID = CVDContextData->Id;
@@ -1227,7 +1258,7 @@ void FChaosVisualDebuggerTrace::TraceDebugDrawImplicitObject(const Chaos::FImpli
 		return;
 	}
 
-	if (TraceRelevancyVolume.IsValid)
+	if (SessionState.TraceRelevancyVolume.IsValid)
 	{
 		if (Implicit->HasBoundingBox() && !IsInRelevancyVolume(Implicit->CalculateTransformedBounds(InParentTransform)))
 		{
@@ -1242,10 +1273,7 @@ void FChaosVisualDebuggerTrace::TraceDebugDrawImplicitObject(const Chaos::FImpli
 	DataWrapper.ParentTransform = InParentTransform;
 	DataWrapper.ThreadContext = IsInGameThread() ? EChaosVDParticleContext::GameThread : EChaosVDParticleContext::PhysicsThread;
 
-	const uint32 GeometryHash = GeometryTracerObject.GetGeometryHashForImplicit(Implicit);
-	TraceImplicitObject({ GeometryHash, const_cast<Chaos::FImplicitObject*>(Implicit) });
-
-	DataWrapper.ImplicitObjectHash = GeometryHash;
+	DataWrapper.ImplicitObjectHash = TraceImplicitObject(const_cast<Chaos::FImplicitObject*>(Implicit));
 
 	DataWrapper.MarkAsValid();
 
@@ -1257,8 +1285,8 @@ void FChaosVisualDebuggerTrace::TraceDebugDrawImplicitObject(const Chaos::FImpli
 
 void FChaosVisualDebuggerTrace::SetTraceRelevancyVolume(const FBox& InTraceRelevancyVolume)
 {
-	TraceRelevancyVolume = InTraceRelevancyVolume;
-	
+	SessionState.TraceRelevancyVolume = InTraceRelevancyVolume;
+
 	UE_TRACE_LOG(ChaosVDLogger, ChaosVDTraceRelevancyVolume, ChaosVDChannel)
 		<< CVD_TRACE_VECTOR_ON_EVENT(ChaosVDTraceRelevancyVolume, BoxMin, InTraceRelevancyVolume.Min)
 		<< CVD_TRACE_VECTOR_ON_EVENT(ChaosVDTraceRelevancyVolume, BoxMax, InTraceRelevancyVolume.Max);
@@ -1266,82 +1294,77 @@ void FChaosVisualDebuggerTrace::SetTraceRelevancyVolume(const FBox& InTraceRelev
 
 bool FChaosVisualDebuggerTrace::IsTracing()
 {
-	return bIsTracing;
+	return SessionState.bIsTracing;
 }
 
 void FChaosVisualDebuggerTrace::RegisterEventHandlers()
 {
 	{
-		UE::TWriteScopeLock WriteLock(DeltaRecordingStatesLock);
+		UE::TWriteScopeLock WriteLock(SessionState.DeltaRecordingStatesLock);
 
 		
-		if (!RecordingStartedDelegateHandle.IsValid())
+		if (!SessionState.RecordingStartedDelegateHandle.IsValid())
 		{
-			RecordingStartedDelegateHandle = FChaosVDRuntimeModule::RegisterRecordingStartedCallback(FChaosVDRecordingStateChangedDelegate::FDelegate::CreateStatic(&FChaosVisualDebuggerTrace::HandleRecordingStart));
+			SessionState.RecordingStartedDelegateHandle = FChaosVDRuntimeModule::RegisterRecordingStartedCallback(FChaosVDRecordingStateChangedDelegate::FDelegate::CreateStatic(&FChaosVisualDebuggerTrace::HandleRecordingStart));
 		}
 
-		if (!RecordingStoppedDelegateHandle.IsValid())
+		if (!SessionState.RecordingStoppedDelegateHandle.IsValid())
 		{
-			RecordingStoppedDelegateHandle = FChaosVDRuntimeModule::RegisterRecordingStopCallback(FChaosVDRecordingStateChangedDelegate::FDelegate::CreateStatic(&FChaosVisualDebuggerTrace::HandleRecordingStop));
+			SessionState.RecordingStoppedDelegateHandle = FChaosVDRuntimeModule::RegisterRecordingStopCallback(FChaosVDRecordingStateChangedDelegate::FDelegate::CreateStatic(&FChaosVisualDebuggerTrace::HandleRecordingStop));
 		}
 
-		if (!RecordingFullCaptureRequestedHandle.IsValid())
+		if (!SessionState.RecordingFullCaptureRequestedHandle.IsValid())
 		{
-			RecordingFullCaptureRequestedHandle = FChaosVDRuntimeModule::RegisterFullCaptureRequestedCallback(FChaosVDCaptureRequestDelegate::FDelegate::CreateStatic(&FChaosVisualDebuggerTrace::PerformFullCapture));
+			SessionState.RecordingFullCaptureRequestedHandle = FChaosVDRuntimeModule::RegisterFullCaptureRequestedCallback(FChaosVDCaptureRequestDelegate::FDelegate::CreateStatic(&FChaosVisualDebuggerTrace::PerformFullCapture));
 		}
 	}
 }
 
 void FChaosVisualDebuggerTrace::UnregisterEventHandlers()
 {
-	UE::TWriteScopeLock WriteLock(DeltaRecordingStatesLock);
-	if (RecordingStartedDelegateHandle.IsValid())
+	UE::TWriteScopeLock WriteLock(SessionState.DeltaRecordingStatesLock);
+	if (SessionState.RecordingStartedDelegateHandle.IsValid())
 	{
-		FChaosVDRuntimeModule::RemoveRecordingStartedCallback(RecordingStartedDelegateHandle);
+		FChaosVDRuntimeModule::RemoveRecordingStartedCallback(SessionState.RecordingStartedDelegateHandle);
 	}
 
-	if (RecordingStoppedDelegateHandle.IsValid())
+	if (SessionState.RecordingStoppedDelegateHandle.IsValid())
 	{
-		FChaosVDRuntimeModule::RemoveRecordingStopCallback(RecordingStoppedDelegateHandle);
+		FChaosVDRuntimeModule::RemoveRecordingStopCallback(SessionState.RecordingStoppedDelegateHandle);
 	}
 
-	if (RecordingFullCaptureRequestedHandle.IsValid())
+	if (SessionState.RecordingFullCaptureRequestedHandle.IsValid())
 	{
-		FChaosVDRuntimeModule::RemoveFullCaptureRequestedCallback(RecordingFullCaptureRequestedHandle);
+		FChaosVDRuntimeModule::RemoveFullCaptureRequestedCallback(SessionState.RecordingFullCaptureRequestedHandle);
 	}
 
-	bIsTracing = false;
+	SessionState.bIsTracing = false;
 }
 
 TSharedRef<FChaosVDSerializableNameTable>& FChaosVisualDebuggerTrace::GetNameTableInstance()
 {
-	return CVDNameTable;
+	return SessionState.CVDNameTable;
 }
 
 void FChaosVisualDebuggerTrace::Reset()
 {
-	CVDNameTable->ResetTable();
+	SessionState.CVDNameTable->ResetTable();
 
-	TraceRelevancyVolume = FBox(EForceInit::ForceInitToZero);
+	SessionState.TraceRelevancyVolume = FBox(EForceInit::ForceInitToZero);
 
 	{
-		UE::TWriteScopeLock WriteLock(DeltaRecordingStatesLock);
-		RequestedFullCaptureSolverIDs.Reset();
-		SolverIDsForDeltaRecording.Reset();
+		UE::TWriteScopeLock WriteLock(SessionState.DeltaRecordingStatesLock);
+		SessionState.RequestedFullCaptureSolverIDs.Reset();
+		SessionState.SolverIDsForDeltaRecording.Reset();
 	}
 
-	{
-		UE::TWriteScopeLock GeometryWriteLock(GeometryTracerObject.TracedGeometrySetLock);
-		GeometryTracerObject.GeometryTracedIDs.Reset();
-
-		UE::TWriteScopeLock GeometryHashWriteLock(GeometryTracerObject.CachedGeometryHashesLock);
-		GeometryTracerObject.CachedGeometryHashes.Reset();
-	}
+	SessionState.TracedGeometryHashCache.Reset();
+	SessionState.ParticleMetadataIDsCache.Reset();
 }
 
 void FChaosVisualDebuggerTrace::HandleRecordingStop()
 {
-	bIsTracing = false;
+	SessionState.bIsTracing = false;
 	Reset();
 }
 
@@ -1381,6 +1404,18 @@ void FChaosVisualDebuggerTrace::OverrideDefaultEnabledDataChannels(TConstArrayVi
 	});
 }
 
+void FChaosVisualDebuggerTrace::RegisterExternalParticleDebugNameGenerator(const Chaos::VD::FRecordingSessionState::FParticleMetaDataGeneratorDelegate& InCallback)
+{
+	if (SessionState.ExternalParticleMetadataGenerator.IsBound())
+	{
+		UE_LOG(LogChaos, Error, TEXT("An external debug name generator for particles is already registered. Ignoring new register attempt"));
+		
+		return;
+	}
+	
+	SessionState.ExternalParticleMetadataGenerator = InCallback;
+}
+
 void FChaosVisualDebuggerTrace::HandleRecordingStart()
 {
 	Reset();
@@ -1399,21 +1434,20 @@ void FChaosVisualDebuggerTrace::HandleRecordingStart()
 	UE_TRACE_LOG(ChaosVDLogger, ChaosVDUsesAutoRTFM, ChaosVDChannel)
 			<< ChaosVDUsesAutoRTFM.bUsingAutoRTFM(static_cast<bool>(UE_AUTORTFM));
 
-	bIsTracing = true;
+	SessionState.bIsTracing = true;
 }
 
 void FChaosVisualDebuggerTrace::PerformFullCapture(EChaosVDFullCaptureFlags CaptureOptions)
 {
 	if (EnumHasAnyFlags(CaptureOptions, EChaosVDFullCaptureFlags::Particles))
 	{
-		UE::TWriteScopeLock WriteLock(DeltaRecordingStatesLock);
-		RequestedFullCaptureSolverIDs.Append(SolverIDsForDeltaRecording);
+		UE::TWriteScopeLock WriteLock(SessionState.DeltaRecordingStatesLock);
+		SessionState.RequestedFullCaptureSolverIDs.Append(SessionState.SolverIDsForDeltaRecording);
 	}
 
 	if (EnumHasAnyFlags(CaptureOptions, EChaosVDFullCaptureFlags::Geometry))
 	{
-		UE::TWriteScopeLock GeometryWriteLock(GeometryTracerObject.TracedGeometrySetLock);
-		GeometryTracerObject.GeometryTracedIDs.Reset();
+		SessionState.TracedGeometryHashCache.Reset();
 	}
 }
 
diff --git a/Engine/Source/Runtime/Experimental/Chaos/Public/Chaos/ChaosDebugName.h b/Engine/Source/Runtime/Experimental/Chaos/Public/Chaos/ChaosDebugName.h
--- a/Engine/Source/Runtime/Experimental/Chaos/Public/Chaos/ChaosDebugName.h
+++ b/Engine/Source/Runtime/Experimental/Chaos/Public/Chaos/ChaosDebugName.h
@@ -1,15 +1,12 @@
 // Copyright Epic Games, Inc. All Rights Reserved.
 #pragma once
 
+#include "ChaosDebugNameDefines.h"
 #include "Containers/UnrealString.h"
 #include "Templates/SharedPointer.h"
 
 #define UE_API CHAOS_API
 
-#ifndef CHAOS_DEBUG_NAME
-#define CHAOS_DEBUG_NAME 0
-#endif
-
 namespace Chaos
 {
 	/**
diff --git a/Engine/Source/Runtime/Experimental/Chaos/Public/Chaos/ChaosDebugNameDefines.h b/Engine/Source/Runtime/Experimental/Chaos/Public/Chaos/ChaosDebugNameDefines.h
new file mode 100644
--- /dev/null
+++ b/Engine/Source/Runtime/Experimental/Chaos/Public/Chaos/ChaosDebugNameDefines.h
@@ -0,0 +1,11 @@
+// Copyright Epic Games, Inc. All Rights Reserved.
+
+#pragma once
+
+#ifndef CHAOS_DEBUG_NAME
+#define CHAOS_DEBUG_NAME 0
+#endif
+
+#ifndef CHAOS_SOLVER_DEBUG_NAME
+#define CHAOS_SOLVER_DEBUG_NAME CHAOS_DEBUG_NAME || WITH_CHAOS_VISUAL_DEBUGGER
+#endif
\ No newline at end of file
diff --git a/Engine/Source/Runtime/Experimental/Chaos/Public/Chaos/Framework/PhysicsSolverBase.h b/Engine/Source/Runtime/Experimental/Chaos/Public/Chaos/Framework/PhysicsSolverBase.h
--- a/Engine/Source/Runtime/Experimental/Chaos/Public/Chaos/Framework/PhysicsSolverBase.h
+++ b/Engine/Source/Runtime/Experimental/Chaos/Public/Chaos/Framework/PhysicsSolverBase.h
@@ -20,6 +20,7 @@
 #endif
 
 #include "ChaosDebugDraw/ChaosDDTypes.h"
+#include "Chaos/ChaosDebugNameDefines.h"
 
 class FChaosSolversModule;
 class FPhysicsReplicationAsync;
@@ -606,8 +607,8 @@ namespace Chaos
 		}
 
 		CHAOS_API FGraphEventRef AdvanceAndDispatch_External(FReal InDt);
-
-#if CHAOS_DEBUG_NAME
+ 
+#if CHAOS_SOLVER_DEBUG_NAME
 		CHAOS_API void SetDebugName(const FName& Name);
 #endif
 		CHAOS_API FName GetDebugName() const;
@@ -905,7 +906,7 @@ namespace Chaos
 		virtual void ProcessPushedData_Internal(FPushPhysicsData& PushDataArray) = 0;
 		virtual void SetExternalTimestampConsumed_Internal(const int32 Timestamp) = 0;
 
-#if CHAOS_DEBUG_NAME
+#if CHAOS_SOLVER_DEBUG_NAME
 		virtual void OnDebugNameChanged() {}
 
 		FName DebugName;
diff --git a/Engine/Source/Runtime/Experimental/Chaos/Public/Chaos/ParticleHandle.h b/Engine/Source/Runtime/Experimental/Chaos/Public/Chaos/ParticleHandle.h
--- a/Engine/Source/Runtime/Experimental/Chaos/Public/Chaos/ParticleHandle.h
+++ b/Engine/Source/Runtime/Experimental/Chaos/Public/Chaos/ParticleHandle.h
@@ -528,6 +528,8 @@ public:
 			GeometryParticles = nullptr;
 			ParticleIdx = INDEX_NONE;
 			Type = EParticleType::Unknown;
+
+			CVD_TRACE_INVALIDATE_CACHED_PARTICLE_METADATA(this);
 		}
 	}
 
diff --git a/Engine/Source/Runtime/Experimental/Chaos/Public/ChaosVisualDebugger/ChaosVDTraceMacros.h b/Engine/Source/Runtime/Experimental/Chaos/Public/ChaosVisualDebugger/ChaosVDTraceMacros.h
--- a/Engine/Source/Runtime/Experimental/Chaos/Public/ChaosVisualDebugger/ChaosVDTraceMacros.h
+++ b/Engine/Source/Runtime/Experimental/Chaos/Public/ChaosVisualDebugger/ChaosVDTraceMacros.h
@@ -130,6 +130,11 @@
 		FChaosVisualDebuggerTrace::InvalidateGeometryFromCache(ImplicitObjectPtr);
 	#endif
 
+	#ifndef CVD_TRACE_INVALIDATE_CACHED_PARTICLE_METADATA
+		#define CVD_TRACE_INVALIDATE_CACHED_PARTICLE_METADATA(ParticleHandlePtr) \
+		FChaosVisualDebuggerTrace::InvalidateParticleMetadataFromCache(ParticleHandlePtr);
+	#endif
+
 	#ifndef CVD_TRACE_SCOPED_SCENE_QUERY
 		#define CVD_TRACE_SCENE_QUERY_START(InputGeometry, GeometryOrientation, Start, End, TraceChannel, Params, ResponseParams, ObjectParams, QueryType, QueryMode, SolverID, bIsRetry) \
 		FChaosVisualDebuggerTrace::TraceSceneQueryStart(InputGeometry, GeometryOrientation, Start, End, TraceChannel, Params, ResponseParams, ObjectParams, QueryType, QueryMode, SolverID, bIsRetry);
@@ -291,6 +296,10 @@
 		#define CVD_TRACE_INVALIDATE_CACHED_GEOMETRY(ImplicitObjectPtr)
 	#endif
 
+	#ifndef CVD_TRACE_INVALIDATE_CACHED_PARTICLE_METADATA
+		#define CVD_TRACE_INVALIDATE_CACHED_PARTICLE_METADATA(ParticleHandlePtr)
+	#endif
+
 	#ifndef CVD_TRACE_SCOPED_SCENE_QUERY
 		#define CVD_TRACE_SCENE_QUERY_START(InputGeometry, GeometryOrientation, Start, End, TraceChannel, Params, ResponseParams, ObjectParams, QueryType, QueryMode, SolverID)
 	#endif
diff --git a/Engine/Source/Runtime/Experimental/Chaos/Public/ChaosVisualDebugger/ChaosVisualDebuggerTrace.h b/Engine/Source/Runtime/Experimental/Chaos/Public/ChaosVisualDebugger/ChaosVisualDebuggerTrace.h
--- a/Engine/Source/Runtime/Experimental/Chaos/Public/ChaosVisualDebugger/ChaosVisualDebuggerTrace.h
+++ b/Engine/Source/Runtime/Experimental/Chaos/Public/ChaosVisualDebugger/ChaosVisualDebuggerTrace.h
@@ -21,8 +21,10 @@
 
 #if WITH_CHAOS_VISUAL_DEBUGGER
 
+#include "Chaos/ChaosDebugNameDefines.h"
 #include "ChaosVDMemWriterReader.h"
 #include "ChaosVDRuntimeModule.h"
+#include "Containers/StripedMap.h"
 #include "DataWrappers/ChaosVDImplicitObjectDataWrapper.h"
 #include "Trace/Trace.h"
 #include "Trace/Trace.inl"
@@ -184,6 +186,104 @@ enum class EChaosVDTraceBinaryDataOptions
 };
 ENUM_CLASS_FLAGS(EChaosVDTraceBinaryDataOptions)
 
+struct FChaosVDParticleMetadata;
+
+namespace Chaos::VD
+{
+	template<typename TObjectID>
+	class FCachedObjectTraceIDByPtr
+	{
+	public:
+
+		template<typename ObjectType, typename HashProducer>
+		TObjectID FindOrProduceTraceObjectID(const ObjectType* ObjectPtr, const HashProducer& InHashProducer, bool& OutWasAlreadyCached)
+		{
+			if (ObjectPtr == nullptr)
+			{
+				return 0;
+			}
+
+			bool bHasCachedID = true;
+			TObjectID FoundID = CachedIDsByPtr.FindOrProduce(reinterpret_cast<const void*>(ObjectPtr), [&InHashProducer, &bHasCachedID]()
+			{
+				bHasCachedID = false;
+				return InHashProducer();
+			});
+
+			if (bHasCachedID)
+			{
+				OutWasAlreadyCached = true;
+				return FoundID;
+			}
+
+			return FoundID;
+		}
+
+		template<typename ObjectType>
+		void RemoveCachedTraceObjectID(const ObjectType* Implicit)
+		{
+			if (Implicit == nullptr)
+			{
+				return;
+			}
+
+			CachedIDsByPtr.Remove(reinterpret_cast<const void*>(Implicit));
+		}
+
+		void Reset()
+		{
+			CachedIDsByPtr.Reset();
+		}
+
+	private:
+
+		static constexpr int32 StripCount = 32;
+		TStripedMap<StripCount, const void*, TObjectID> CachedIDsByPtr;
+	};
+
+	struct FRecordingSessionState
+	{
+		FRecordingSessionState();
+
+		DECLARE_DELEGATE_RetVal_TwoParams(FChaosVDParticleMetadata, FParticleMetaDataGeneratorDelegate, const IPhysicsProxyBase*, const Chaos::FGeometryParticleHandle* ParticleHandle)
+		FParticleMetaDataGeneratorDelegate ExternalParticleMetadataGenerator;
+
+		FDelegateHandle RecordingStartedDelegateHandle;
+		FDelegateHandle RecordingStoppedDelegateHandle;
+		FDelegateHandle RecordingFullCaptureRequestedHandle;
+
+		TSet<int32> SolverIDsForDeltaRecording;
+		TSet<int32> RequestedFullCaptureSolverIDs;
+		TSharedRef<FChaosVDSerializableNameTable> CVDNameTable;
+
+		FBox TraceRelevancyVolume;
+
+		std::atomic<bool> bIsTracing;
+
+		FCachedObjectTraceIDByPtr<uint32> TracedGeometryHashCache;
+
+		FCachedObjectTraceIDByPtr<uint64> ParticleMetadataIDsCache;
+
+		FRWLock DeltaRecordingStatesLock;
+
+		uint64 GenerateUniqueID()
+		{
+			uint64 NewID = 0;
+			UE_AUTORTFM_OPEN
+			{
+				NewID = InternalGenericIDCounter++;
+			};
+
+			return NewID;
+		}
+
+		static constexpr uint64 InvalidUniqueID = 0; 
+
+		private:
+		std::atomic<uint64> InternalGenericIDCounter;
+	};
+}
+
 /** Class containing  all the Tracing logic to record data for the Chaos Visual Debugger tool */
 class FChaosVisualDebuggerTrace
 {
@@ -268,14 +368,37 @@ public:
 	 * The trace event is not tied to a particular Solver Frame/Step
 	 *  @param WrappedGeometryData Wrapper containing a ptr to the implicit and its ID
 	 */
+	UE_DEPRECATED(5.7, "Use the TraceImplicitObject version that takes a ptr to the implicit object itself")
 	static CHAOS_API void TraceImplicitObject(FChaosVDImplicitObjectWrapper WrappedGeometryData);
 
+	/**
+	 * Serializes a implicit object
+	 * The trace event is not tied to a particular Solver Frame/Step
+	 * @param GeometryPtr Geometry instance to trace
+	 * @return Hash of the traced geometry that will be used when the recording is loaded to match it to the correct piece of data
+	 */
+	static CHAOS_API uint32 TraceImplicitObject(Chaos::FImplicitObject* GeometryPtr);
+
+	/**
+	 * Gathers and serializes the debug name and other metadata for the provider particle once. Any subsequent calls will only return the existing hash for the debug name
+	 * The trace event is not tied to a particular Solver Frame/Step
+	 * @param InParticleHandle Handle to the particle from which we want to trace its debug name
+	 * @return Id of the traced particle metadata that will be used when the recording is loaded to match it to the correct piece of data
+	 */
+	static CHAOS_API uint64 TraceParticleMetadata(Chaos::FGeometryParticleHandle* InParticleHandle);
+
 	/**
 	 * Removes an implicit object from the serialized geometry IDs cache, to ensure we re-serialize it with any new changes
 	 *  @param CachedGeometryToInvalidate Ptr to the Geometry we want to invalidate from the cache
 	 */
 	static CHAOS_API void InvalidateGeometryFromCache(const Chaos::FImplicitObject* CachedGeometryToInvalidate);
 
+	/**
+	 * Removes the associated metadata id for a particle from the traced data IDs cache, to ensure we re-serialize it with any new changes
+	 *  @param InParticleToInvalidate Ptr to the particle we want to invalidate from the cache
+	 */
+	static CHAOS_API void InvalidateParticleMetadataFromCache(Chaos::FGeometryParticleHandle* InParticleToInvalidate);
+
 	/** Records a visit step of a scene query. This needs to me called within the scope of an inflight scene query
 	 * @param InputGeometry : Geometry used as input to perform the query we are recording
 	 * @param GeometryOrientation : Orientation of the input geometry
@@ -361,10 +484,9 @@ public:
 	 */
 	static FBox GetTraceRelevancyVolume()
 	{
-		return TraceRelevancyVolume;
+		return SessionState.TraceRelevancyVolume;
 	}
 
-
 	/**
 	 * Evaluates the provided bounds to see if it is withing the current relevancy volume, if any.
 	 * @param InChaosBounds Bounds to evaluate
@@ -372,13 +494,13 @@ public:
 	 */
 	static bool IsInRelevancyVolume(const Chaos::FAABB3& InChaosBounds)
 	{
-		if (!TraceRelevancyVolume.IsValid)
+		if (!SessionState.TraceRelevancyVolume.IsValid)
 		{
 			// An invalid volume means we want to trace everything
 			return true;
 		}
 
-		return TraceRelevancyVolume.Intersect(FBox(InChaosBounds.Min(), InChaosBounds.Max()));
+		return SessionState.TraceRelevancyVolume.Intersect(FBox(InChaosBounds.Min(), InChaosBounds.Max()));
 	}
 
 	/**
@@ -388,13 +510,13 @@ public:
 	 */
 	static bool IsInRelevancyVolume(const FBox& InBounds)
 	{
-		if (!TraceRelevancyVolume.IsValid)
+		if (!SessionState.TraceRelevancyVolume.IsValid)
 		{
 			// An invalid volume means we want to trace everything
 			return true;
 		}
 
-		return TraceRelevancyVolume.Intersect(InBounds);
+		return SessionState.TraceRelevancyVolume.Intersect(InBounds);
 	}
 
 	/**
@@ -405,13 +527,13 @@ public:
 	 */
 	static bool IsInRelevancyVolume(const FVector& Center, double Radius)
 	{
-		if (!TraceRelevancyVolume.IsValid)
+		if (!SessionState.TraceRelevancyVolume.IsValid)
 		{
 			// An invalid volume means we want to trace everything
 			return true;
 		}
 	
-		return FMath::SphereAABBIntersection(Center, Radius * Radius, TraceRelevancyVolume);
+		return FMath::SphereAABBIntersection(Center, Radius * Radius, SessionState.TraceRelevancyVolume);
 	}
 
 	/**
@@ -421,13 +543,13 @@ public:
 	 */
 	static bool IsInRelevancyVolume(const Chaos::FVec3& InPosition)
 	{
-		if (!TraceRelevancyVolume.IsValid)
+		if (!SessionState.TraceRelevancyVolume.IsValid)
 		{
 			// An invalid volume means we want to trace everything
 			return true;
 		}
 
-		return FMath::PointBoxIntersection(InPosition, TraceRelevancyVolume);
+		return FMath::PointBoxIntersection(InPosition, SessionState.TraceRelevancyVolume);
 	}
 
 	/**
@@ -438,13 +560,13 @@ public:
 	 */
 	static bool IsInRelevancyVolume(const Chaos::FVec3& InStartPoint, const Chaos::FVec3& InEndPoint)
 	{
-		if (!TraceRelevancyVolume.IsValid)
+		if (!SessionState.TraceRelevancyVolume.IsValid)
 		{
 			// An invalid volume means we want to trace everything
 			return true;
 		}
 
-		return TraceRelevancyVolume.Intersect(FBox(FVector::Min(InStartPoint, InEndPoint), FVector::Max(InStartPoint, InEndPoint)));
+		return SessionState.TraceRelevancyVolume.Intersect(FBox(FVector::Min(InStartPoint, InEndPoint), FVector::Max(InStartPoint, InEndPoint)));
 	}
 
 	static bool IsInRelevancyVolume(const Chaos::FGeometryParticleHandle* InParticleHandle);
@@ -496,6 +618,8 @@ public:
 	 */
 	static CHAOS_API void OverrideDefaultEnabledDataChannels(TConstArrayView<FString> EnabledDataChannelsOverrideList);
 
+	static CHAOS_API void RegisterExternalParticleDebugNameGenerator(const Chaos::VD::FRecordingSessionState::FParticleMetaDataGeneratorDelegate& InCallback);
+
 private:
 
 	/**
@@ -550,19 +674,7 @@ private:
 	/** Setups the current Solver frame for a full capture if needed */
 	static void SetupForFullCaptureIfNeeded(int32 SolverID, bool& bOutFullCaptureRequested);
 
-	static FDelegateHandle RecordingStartedDelegateHandle;
-	static FDelegateHandle RecordingStoppedDelegateHandle;
-	static FDelegateHandle RecordingFullCaptureRequestedHandle;
-
-	static TSet<int32> SolverIDsForDeltaRecording;
-	static TSet<int32> RequestedFullCaptureSolverIDs;
-	static TSharedRef<FChaosVDSerializableNameTable> CVDNameTable;
-
-	static FBox TraceRelevancyVolume;
-
-	static std::atomic<bool> bIsTracing;
-
-	static FRWLock DeltaRecordingStatesLock;
+	static Chaos::VD::FRecordingSessionState SessionState;
 
 	friend class FChaosEngineModule;
 };
@@ -614,7 +726,7 @@ void FChaosVisualDebuggerTrace::GetCVDContext(T& ObjectWithContext, FChaosVDCont
 template <typename T>
 FString FChaosVisualDebuggerTrace::GetDebugName(T& ObjectWithDebugName)
 {
-#if CHAOS_DEBUG_NAME
+#if CHAOS_SOLVER_DEBUG_NAME
 	return ObjectWithDebugName.GetDebugName().ToString();
 #else
 	return FString("COMPILED OUT");
diff --git a/Engine/Source/Runtime/Experimental/ChaosSolverEngine/Private/Chaos/ChaosVDEngineEditorBridge.cpp b/Engine/Source/Runtime/Experimental/ChaosSolverEngine/Private/Chaos/ChaosVDEngineEditorBridge.cpp
--- a/Engine/Source/Runtime/Experimental/ChaosSolverEngine/Private/Chaos/ChaosVDEngineEditorBridge.cpp
+++ b/Engine/Source/Runtime/Experimental/ChaosSolverEngine/Private/Chaos/ChaosVDEngineEditorBridge.cpp
@@ -9,9 +9,17 @@
 #include "ChaosVDRuntimeModule.h"
 #include "ChaosVisualDebugger/ChaosVDTraceMacros.h"
 #include "ChaosVisualDebugger/ChaosVisualDebuggerTrace.h"
+#include "Components/PrimitiveComponent.h"
+#include "Components/SkeletalMeshComponent.h"
 #include "DataWrappers/ChaosVDCollisionDataWrappers.h"
+#include "DataWrappers/ChaosVDParticleDataWrapper.h"
 #include "Engine/Engine.h"
+#include "PhysicsEngine/BodySetup.h"
+#include "PhysicsProxy/SingleParticlePhysicsProxy.h"
+#include "Runtime/Experimental/Chaos/Private/Chaos/PhysicsObjectInternal.h"
 #include "Serialization/MemoryWriter.h"
+#include "UObject/Package.h"
+
 
 void FChaosVDEngineEditorBridge::AddOnScreenRecordingMessage()
 {
@@ -149,6 +157,84 @@ bool FChaosVDEngineEditorBridge::BroadcastSessionStatus(float DeltaTime)
 	return true;
 }
 
+FChaosVDParticleMetadata FChaosVDEngineEditorBridge::GenerateParticleMetadata(const IPhysicsProxyBase* ParticleProxy, const Chaos::FGeometryParticleHandle* ParticleHandle)
+{
+	FChaosVDParticleMetadata Metadata;
+
+	if (!ParticleProxy)
+	{
+		return Metadata;
+	}
+
+	// This method is expected to be called from worker threads as we discover new particles that need to be traced.
+	// As we need to access their uobject owners to generate the name, we need to make sure these are not GC'd or wait until GC is done.
+	FGCScopeGuard GCLock;
+
+	if (UObject* ProxyOwner = ParticleProxy->GetOwner())
+	{
+		if (UPrimitiveComponent* OwnerAsPrimitiveComponent = Cast<UPrimitiveComponent>(ProxyOwner))
+		{
+			Metadata.ComponentName = ProxyOwner->GetFName();
+			Metadata.MapAssetPath = FTopLevelAssetPath(ProxyOwner->GetWorld());
+			Metadata.OwnerName = GetFNameSafe(OwnerAsPrimitiveComponent->GetOwner());
+
+			if (AActor* OwningActor = OwnerAsPrimitiveComponent->GetOwner())
+			{
+				Metadata.OwnerName = OwningActor->GetFName();
+				FSoftObjectPath ObjectPath = FSoftObjectPath(OwningActor->GetClass());
+				Metadata.OwnerAssetPath = ObjectPath.GetAssetPath();
+			}
+		}
+		else
+		{
+			Metadata.OwnerName = ProxyOwner->GetFName();
+		}
+
+		if (!Metadata.OwnerAssetPath.IsValid())
+		{
+			FSoftObjectPath ObjectPath = FSoftObjectPath(ProxyOwner);
+			Metadata.OwnerAssetPath = ObjectPath.GetAssetPath();
+		}
+	}
+
+	switch (ParticleProxy->GetType())
+	{
+	case EPhysicsProxyType::SingleParticleProxy:
+		{
+			const Chaos::FSingleParticlePhysicsProxy* AsSingleParticleProxy = static_cast<const Chaos::FSingleParticlePhysicsProxy*>(ParticleProxy);
+			if (Chaos::FPhysicsObjectHandle PhysicsObjectHandle = AsSingleParticleProxy->GetPhysicsObject())
+			{
+				// Currently only bodies corresponding to the skeletal mesh bones bodies set this value
+				Metadata.BoneName = PhysicsObjectHandle->GetBodyName();
+			}
+
+			break;
+		}
+	case EPhysicsProxyType::GeometryCollectionType:
+		{
+			if (ParticleHandle)
+			{
+				const FGeometryCollectionPhysicsProxy* AsGeometryCollectionProxy = static_cast<const FGeometryCollectionPhysicsProxy*>(ParticleProxy);
+				Metadata.Index = AsGeometryCollectionProxy->GetTransformGroupIndexFromHandle(ParticleHandle->CastToRigidParticle());
+			}
+			break;
+		}
+	case EPhysicsProxyType::ClusterUnionProxy: // These follow the same rules as Single particle proxies
+	case EPhysicsProxyType::NoneType:
+	case EPhysicsProxyType::StaticMeshType:
+	case EPhysicsProxyType::FieldType:
+	case EPhysicsProxyType::JointConstraintType:
+	case EPhysicsProxyType::SuspensionConstraintType:
+	case EPhysicsProxyType::CharacterGroundConstraintType:
+	case EPhysicsProxyType::Count:
+	default:
+		ensureMsgf(false, TEXT("Unsupported Proxy type"));
+		break;
+	}
+
+	return Metadata;
+}
+
 void FChaosVDEngineEditorBridge::Initialize()
 {
 	FChaosVDRuntimeModule& Module = FChaosVDRuntimeModule::Get();
@@ -162,6 +248,8 @@ void FChaosVDEngineEditorBridge::Initialize()
 	RecordingStoppedHandle = FChaosVDRuntimeModule::Get().RegisterRecordingStopCallback(FChaosVDRecordingStateChangedDelegate::FDelegate::CreateRaw(this, &FChaosVDEngineEditorBridge::HandleCVDRecordingStopped));
 	RecordingStartFailedHandle = FChaosVDRuntimeModule::Get().RegisterRecordingStartFailedCallback(FChaosVDRecordingStartFailedDelegate::FDelegate::CreateRaw(this, &FChaosVDEngineEditorBridge::HandleCVDRecordingStartFailed));
 
+	FChaosVisualDebuggerTrace::RegisterExternalParticleDebugNameGenerator(Chaos::VD::FRecordingSessionState::FParticleMetaDataGeneratorDelegate::CreateStatic(&FChaosVDEngineEditorBridge::GenerateParticleMetadata));
+
 	Chaos::VisualDebugger::FChaosVDDataChannelsManager::Get().OnChannelStateChanged().AddRaw(this, &FChaosVDEngineEditorBridge::HandleDataChannelChanged);
 
 #if WITH_EDITOR
diff --git a/Engine/Source/Runtime/Experimental/ChaosSolverEngine/Public/Chaos/ChaosVDEngineEditorBridge.h b/Engine/Source/Runtime/Experimental/ChaosSolverEngine/Public/Chaos/ChaosVDEngineEditorBridge.h
--- a/Engine/Source/Runtime/Experimental/ChaosSolverEngine/Public/Chaos/ChaosVDEngineEditorBridge.h
+++ b/Engine/Source/Runtime/Experimental/ChaosSolverEngine/Public/Chaos/ChaosVDEngineEditorBridge.h
@@ -5,9 +5,13 @@
 
 #if WITH_CHAOS_VISUAL_DEBUGGER
 
+#include "Chaos/ParticleHandleFwd.h"
+#include "DataWrappers/ChaosVDParticleDataWrapper.h"
 #include "Delegates/IDelegateInstance.h"
 #include "HAL/Platform.h"
 
+class UBodySetup;
+class IPhysicsProxyBase;
 class UGameInstance;
 
 namespace Chaos
@@ -58,6 +62,8 @@ private:
 	
 	bool BroadcastSessionStatus(float DeltaTime);
 
+	static FChaosVDParticleMetadata GenerateParticleMetadata(const IPhysicsProxyBase* ParticleProxy, const Chaos::FGeometryParticleHandle* ParticleHandle);
+
 	FDelegateHandle RecordingStartedHandle;
 	FDelegateHandle PostRecordingStartedHandle;
 	FDelegateHandle RecordingStoppedHandle;
diff --git a/Engine/Source/Runtime/Experimental/ChaosVisualDebugger/Private/DataWrappers/ChaosVDParticleDataWrapper.cpp b/Engine/Source/Runtime/Experimental/ChaosVisualDebugger/Private/DataWrappers/ChaosVDParticleDataWrapper.cpp
--- a/Engine/Source/Runtime/Experimental/ChaosVisualDebugger/Private/DataWrappers/ChaosVDParticleDataWrapper.cpp
+++ b/Engine/Source/Runtime/Experimental/ChaosVisualDebugger/Private/DataWrappers/ChaosVDParticleDataWrapper.cpp
@@ -4,6 +4,7 @@
 
 #include "UObject/FortniteMainBranchObjectVersion.h"
 #include "UObject/FortniteSeasonBranchObjectVersion.h"
+#include "Misc/StringBuilder.h"
 
 // @note: Tracing an scene with 1000 particles moving, Manually serializing the structs is ~20% faster
 // than normal UStruct serialization in unversioned mode. As we do this at runtime when tracing in development builds, this is important.
@@ -11,6 +12,40 @@
 
 #include UE_INLINE_GENERATED_CPP_BY_NAME(ChaosVDParticleDataWrapper)
 
+FStringView FChaosVDParticleMetadata::WrapperTypeName = TEXT("FChaosVDParticleMetadata");
+
+bool FChaosVDParticleMetadata::Serialize(FArchive& Ar)
+{
+	Ar << OwnerName;
+	Ar << ComponentName;
+	Ar << BoneName;
+	Ar << Index;
+	Ar << MapAssetPath;
+	Ar << OwnerAssetPath;
+	Ar << MetadataID;
+
+	return !Ar.IsError();
+}
+
+FString FChaosVDParticleMetadata::ToString() const
+{
+	TStringBuilder<1024> NameBuilder;
+
+	NameBuilder.Appendf(TEXT("%s | %s "), *OwnerName.ToString(), *ComponentName.ToString());
+
+	if (!BoneName.IsNone())
+	{
+		NameBuilder.Appendf(TEXT("| Bone Name: %s "), *BoneName.ToString());
+	}
+
+	if (Index != INDEX_NONE)
+	{
+		NameBuilder.Appendf(TEXT("| Index: %d "), Index);
+	}
+
+	return NameBuilder.ToString();
+}
+
 bool FChaosVDFRigidParticleControlFlags::Serialize(FArchive& Ar)
 {
 	Ar.UsingCustomVersion(FFortniteSeasonBranchObjectVersion::GUID);
@@ -223,6 +258,54 @@ bool FChaosVDVSmooth::Serialize(FArchive& Ar)
 	return !Ar.IsError();
 }
 
+bool FChaosVDParticleDataWrapper::HasLegacyDebugName() const
+{
+#if WITH_EDITOR
+	PRAGMA_DISABLE_DEPRECATION_WARNINGS
+	return bHasDebugName;
+	PRAGMA_ENABLE_DEPRECATION_WARNINGS
+#else
+	return false;
+#endif
+}
+
+FString FChaosVDParticleDataWrapper::GetDebugName() const
+{
+	static FString DefaultDebugName(TEXT("UnnamedParticle"));
+
+	if (ParticleMetadataInstance.IsValid())
+	{
+		return *ParticleMetadataInstance->ToString();
+	}
+	else if (HasLegacyDebugName())
+	{
+#if WITH_EDITOR
+		PRAGMA_DISABLE_DEPRECATION_WARNINGS
+		return DebugName;
+		PRAGMA_ENABLE_DEPRECATION_WARNINGS
+#else
+		return DefaultDebugName;
+#endif
+	}
+
+	return DefaultDebugName;
+}
+
+void FChaosVDParticleDataWrapper::SetMetadataInstance(const TSharedPtr<FChaosVDParticleMetadata>& InMetadata)
+{
+	if (InMetadata.IsValid())
+	{
+		check(MetadataId == InMetadata->MetadataID);
+	}
+	
+	ParticleMetadataInstance = InMetadata;
+}
+
+const TSharedPtr<FChaosVDParticleMetadata>& FChaosVDParticleDataWrapper::GetMetadataInstance() const
+{
+	return ParticleMetadataInstance;
+}
+
 bool FChaosVDParticleDataWrapper::Serialize(FArchive& Ar)
 {
 	Ar.UsingCustomVersion(FFortniteMainBranchObjectVersion::GUID);
@@ -236,25 +319,37 @@ bool FChaosVDParticleDataWrapper::Serialize(FArchive& Ar)
 	Ar << Type;
 	Ar << GeometryHash;
 
-	bHasDebugName = DebugNamePtr != nullptr;
-	Ar << bHasDebugName;
-
-	if (bHasDebugName)
+	if (Ar.CustomVer(FFortniteMainBranchObjectVersion::GUID) < FFortniteMainBranchObjectVersion::DeduplicatedDebugNameSerializationInCVD)
 	{
-		if (Ar.IsLoading())
+#if WITH_EDITOR
+		PRAGMA_DISABLE_DEPRECATION_WARNINGS
+		Ar << bHasDebugName;
+		if (bHasDebugName)
 		{
-			Ar << DebugName;
+			if (Ar.IsLoading())
+			{
+				Ar << DebugName;
+			}
 		}
-		else
+		PRAGMA_ENABLE_DEPRECATION_WARNINGS
+#else
+		bool bDummyHasDebugName;
+		Ar << bDummyHasDebugName;
+		if (bDummyHasDebugName)
 		{
-			// Keep the static analyser happy, we already checked this above
-			checkSlow(DebugNamePtr != nullptr);
-
-			FString& DebugNameRef = *DebugNamePtr;
-			Ar << DebugNameRef;
+			if (Ar.IsLoading())
+			{
+				FString DummyString;
+				Ar << DummyString;
+			}
 		}
+#endif // WITH_EDITOR
 	}
-
+	else
+	{
+		Ar << MetadataId;
+	}
+	
 	Ar << ParticleIndex;
 	Ar << SolverID;
 
diff --git a/Engine/Source/Runtime/Experimental/ChaosVisualDebugger/Public/DataWrappers/ChaosVDParticleDataWrapper.h b/Engine/Source/Runtime/Experimental/ChaosVisualDebugger/Public/DataWrappers/ChaosVDParticleDataWrapper.h
--- a/Engine/Source/Runtime/Experimental/ChaosVisualDebugger/Public/DataWrappers/ChaosVDParticleDataWrapper.h
+++ b/Engine/Source/Runtime/Experimental/ChaosVisualDebugger/Public/DataWrappers/ChaosVDParticleDataWrapper.h
@@ -5,6 +5,7 @@
 #include "ChaosVDCollisionDataWrappers.h"
 #include "UObject/ObjectMacros.h"
 #include "HAL/Platform.h"
+#include "StructUtils/InstancedStruct.h"
 
 #include "ChaosVDParticleDataWrapper.generated.h"
 
@@ -40,6 +41,35 @@ enum class EChaosVDObjectStateType: int8
 	Count
 };
 
+USTRUCT()
+struct FChaosVDParticleMetadata
+{
+	GENERATED_BODY()
+
+	CHAOSVDRUNTIME_API static FStringView WrapperTypeName;
+
+	UPROPERTY(VisibleAnywhere, Category=ParticleName)
+	FName OwnerName;
+	UPROPERTY(VisibleAnywhere, Category=ParticleName)
+	FName ComponentName;
+	UPROPERTY(VisibleAnywhere, Category=ParticleName)
+	FName BoneName;
+	UPROPERTY(VisibleAnywhere, Category=ParticleName)
+	int32 Index = INDEX_NONE;
+
+	UPROPERTY(VisibleAnywhere, Category=ParticleAssetData)
+	FTopLevelAssetPath MapAssetPath;
+	UPROPERTY(VisibleAnywhere, Category=ParticleAssetData)
+	FTopLevelAssetPath OwnerAssetPath;
+
+	UPROPERTY()
+	uint64 MetadataID = 0;
+
+	CHAOSVDRUNTIME_API bool Serialize(FArchive& Ar);
+
+	CHAOSVDRUNTIME_API FString ToString() const;
+};
+
 /** Base struct that declares the interface to be used for any ParticleData Viewer */
 USTRUCT()
 struct FChaosVDWrapperDataBase
@@ -715,16 +745,23 @@ CVD_IMPLEMENT_SERIALIZER(FChaosVDVSmooth)
 USTRUCT(DisplayName="Particle Data")
 struct FChaosVDParticleDataWrapper : public FChaosVDWrapperDataBase
 {
+PRAGMA_DISABLE_DEPRECATION_WARNINGS
+
+	FChaosVDParticleDataWrapper(const FChaosVDParticleDataWrapper& Other) = default;
+	FChaosVDParticleDataWrapper(FChaosVDParticleDataWrapper&& Other) noexcept = default;
+	FChaosVDParticleDataWrapper& operator=(const FChaosVDParticleDataWrapper& Other) = default;
+	FChaosVDParticleDataWrapper& operator=(FChaosVDParticleDataWrapper&& Other) noexcept = default;
 	virtual ~FChaosVDParticleDataWrapper() override = default;
+	FChaosVDParticleDataWrapper()
+	{
+	}
+
+PRAGMA_ENABLE_DEPRECATION_WARNINGS
 
 	inline static FStringView WrapperTypeName = TEXT("FChaosVDParticleDataWrapper");
 
 	GENERATED_BODY()
 
-	FChaosVDParticleDataWrapper()
-	{
-	}
-
 	UPROPERTY()
 	int32 DirtyFlagsBits = 0;
 
@@ -734,11 +771,23 @@ struct FChaosVDParticleDataWrapper : public FChaosVDWrapperDataBase
 	UPROPERTY(VisibleAnywhere, Category= "General")
 	uint32 GeometryHash = 0;
 
-	UPROPERTY(VisibleAnywhere, Category= "General")
+#if WITH_EDITORONLY_DATA
+	UE_DEPRECATED(5.7, "Please Use the GetDebugName accessor. This property is no longer in use")
+	/** Shows the serialized debug name of the particle. This is only valid in CVD recordings from UE 5.6 and earlier.
+	 * Particles Debug name now is part of the particle metadata
+	* */
+	UPROPERTY(VisibleAnywhere, Category= "General", DisplayName="Legacy Debug Name")
 	FString DebugName;
 
-	/** Hash used during CVD Playback to compare names */
-	uint64 DebugNameHash = 0;
+	/** Serialized boolean used as a flag to know if a particle has the legacy debug name serialized as part of the particle data.
+	 * If false, it means the particle either had no debug name at the time of being traced or it is using the new particle metadata system
+	 */
+	UE_DEPRECATED(5.7, "Please Use the HasLegacyDebugName accessor. This property is no longer in use")
+	bool bHasDebugName = false;
+#endif
+
+	UPROPERTY()
+	uint64 MetadataId = 0;
 
 	UPROPERTY(VisibleAnywhere, Category= "General")
 	int32 ParticleIndex = INDEX_NONE;
@@ -779,13 +828,21 @@ struct FChaosVDParticleDataWrapper : public FChaosVDWrapperDataBase
 	UPROPERTY()
 	TArray<FChaosVDShapeCollisionData> CollisionDataPerShape;
 
-	/** Only used during recording */
-	FString* DebugNamePtr = nullptr;
-	bool bHasDebugName = false;
+	CHAOSVDRUNTIME_API bool HasLegacyDebugName() const;
+	CHAOSVDRUNTIME_API FString GetDebugName() const;
+	CHAOSVDRUNTIME_API void SetMetadataInstance(const TSharedPtr<FChaosVDParticleMetadata>& InMetadataInstance);
+	CHAOSVDRUNTIME_API const TSharedPtr<FChaosVDParticleMetadata>& GetMetadataInstance() const;
 
-	virtual bool HasValidData() const override { return bHasValidData; }
+	virtual bool HasValidData() const override
+	{
+		return bHasValidData;
+	}
 
 	CHAOSVDRUNTIME_API bool Serialize(FArchive& Ar);
+
+private:
+	/** Debug Name instance */
+	TSharedPtr<FChaosVDParticleMetadata> ParticleMetadataInstance;
 };
 
 CVD_IMPLEMENT_SERIALIZER(FChaosVDParticleDataWrapper)
