Removed historical vim fold comments

#rnx

#ushell-unshelve of 44547637

diff --git a/Engine/Source/Runtime/Experimental/IoStore/HttpClient/Private/Activity.inl b/Engine/Source/Runtime/Experimental/IoStore/HttpClient/Private/Activity.inl
--- a/Engine/Source/Runtime/Experimental/IoStore/HttpClient/Private/Activity.inl
+++ b/Engine/Source/Runtime/Experimental/IoStore/HttpClient/Private/Activity.inl
@@ -725,6 +725,4 @@ static void Trace(const struct FActivity* Activity, ETrace Action, uint32 Param)
 	}
 }
 
-// }}}
-
 } // namespace UE::IoStore::HTTP
diff --git a/Engine/Source/Runtime/Experimental/IoStore/HttpClient/Private/ConnectionPool.inl b/Engine/Source/Runtime/Experimental/IoStore/HttpClient/Private/ConnectionPool.inl
--- a/Engine/Source/Runtime/Experimental/IoStore/HttpClient/Private/ConnectionPool.inl
+++ b/Engine/Source/Runtime/Experimental/IoStore/HttpClient/Private/ConnectionPool.inl
@@ -299,7 +299,5 @@ bool FConnectionPool::IsValidHostUrl(FAnsiStringView Url)
 	return ParseUrl(Url, Tmp) >= 0;
 }
 
-// }}}
-
 } // namespace UE::IoStore::HTTP
 
diff --git a/Engine/Source/Runtime/Experimental/IoStore/HttpClient/Private/Socks.inl b/Engine/Source/Runtime/Experimental/IoStore/HttpClient/Private/Socks.inl
--- a/Engine/Source/Runtime/Experimental/IoStore/HttpClient/Private/Socks.inl
+++ b/Engine/Source/Runtime/Experimental/IoStore/HttpClient/Private/Socks.inl
@@ -181,6 +181,4 @@ static FOutcome MaybeConnectSocks(FSocket& Socket, uint32 IpAddress, uint32 Port
 #endif // UE_BUILD_SHIPPING
 }
 
-// }}}
-
 } // namespace UE::IoStore::HTTP
diff --git a/Engine/Source/Runtime/Experimental/IoStore/HttpClient/Private/Test.inl b/Engine/Source/Runtime/Experimental/IoStore/HttpClient/Private/Test.inl
--- a/Engine/Source/Runtime/Experimental/IoStore/HttpClient/Private/Test.inl
+++ b/Engine/Source/Runtime/Experimental/IoStore/HttpClient/Private/Test.inl
@@ -983,6 +983,4 @@ IOSTOREHTTPCLIENT_API void IasHttpTest(const ANSICHAR* TestHost="localhost", uin
 
 #endif // !SHIP|TEST
 
-// }}}
-
 } // namespace UE::IoStore::HTTP
