Broader ifdef condition for when tests are compiled in

#rnx

#ushell-unshelve of 44547637

diff --git a/Engine/Source/Runtime/Experimental/IoStore/HttpClient/Private/Test.inl b/Engine/Source/Runtime/Experimental/IoStore/HttpClient/Private/Test.inl
--- a/Engine/Source/Runtime/Experimental/IoStore/HttpClient/Private/Test.inl
+++ b/Engine/Source/Runtime/Experimental/IoStore/HttpClient/Private/Test.inl
@@ -10,7 +10,7 @@
 namespace UE::IoStore::HTTP
 {
 
-#if !(UE_BUILD_SHIPPING|UE_BUILD_TEST) && !defined(NO_UE_INCLUDES)
+#if !(UE_BUILD_SHIPPING|UE_BUILD_TEST)
 
 ////////////////////////////////////////////////////////////////////////////////
 static void MiscTest()
