Sped tests up a little bit

#rnx

#ushell-unshelve of 44547637

diff --git a/Engine/Source/Runtime/Experimental/IoStore/HttpClient/Private/Test.inl b/Engine/Source/Runtime/Experimental/IoStore/HttpClient/Private/Test.inl
--- a/Engine/Source/Runtime/Experimental/IoStore/HttpClient/Private/Test.inl
+++ b/Engine/Source/Runtime/Experimental/IoStore/HttpClient/Private/Test.inl
@@ -379,8 +379,10 @@ static void ChunkedTest(const ANSICHAR* TestHost)
 {
 	FEventLoop Loop;
 
-	auto WaitForLoopIdle = [&] {
-		for (; Loop.Tick(-1); FPlatformProcess::SleepNoStats(0.02f));
+	auto WaitForLoopIdle = [&]
+	{
+		while (Loop.Tick(0))
+			;
 	};
 
 	TAnsiStringBuilder<64> Url;
