Increase length score size to surface a little more information in traces

#rnx

#ushell-unshelve of 44547637

diff --git a/Engine/Source/Runtime/Experimental/IoStore/HttpClient/Private/Activity.inl b/Engine/Source/Runtime/Experimental/IoStore/HttpClient/Private/Activity.inl
--- a/Engine/Source/Runtime/Experimental/IoStore/HttpClient/Private/Activity.inl
+++ b/Engine/Source/Runtime/Experimental/IoStore/HttpClient/Private/Activity.inl
@@ -527,8 +527,7 @@ struct alignas(16) FActivity
 	uint8				NoContent : 1;
 	uint8				bFollow30x : 1;
 	uint8				bAllowChunked : 1;
-	uint8				LengthScore : 3;
-	uint8				_Unused : 1;
+	uint8				LengthScore : 4;
 	uint32				StateParam = 0;
 #if IAS_HTTP_WITH_PERF
 	FStopwatch			Stopwatch;
@@ -648,9 +647,8 @@ static void Activity_SetScore(FActivity* Activity, uint32 ContentSizeEst)
 	}
 
 	uint32 ContentEstKiB = (ContentSizeEst + 1023) >> 10;
-	ContentEstKiB |= 2;
 	uint32 Pow2 = FMath::FloorLog2(uint32(ContentEstKiB));
-	Pow2 = FMath::Min(Pow2, 7u);
+	Pow2 = FMath::Min(Pow2, 15u);
 	Activity->LengthScore = uint8(Pow2);
 }
 
@@ -661,17 +659,18 @@ static void Trace(const struct FActivity* Activity, ETrace Action, uint32 Param)
 {
 	if (Action == ETrace::ActivityCreate)
 	{
-		static uint32 ActScopes[8] = {};
+		static uint32 ActScopes[15] = {};
 		if (ActScopes[0] == 0)
 		{
 			ActScopes[0] = LaneTrace_NewScope("Iax/Activity");
-			ActScopes[1] = LaneTrace_NewScope("Iax/Activity_2");
-			ActScopes[2] = LaneTrace_NewScope("Iax/Activity_4");
-			ActScopes[3] = LaneTrace_NewScope("Iax/Activity_8");
-			ActScopes[4] = LaneTrace_NewScope("Iax/Activity_16");
-			ActScopes[5] = LaneTrace_NewScope("Iax/Activity_32");
-			ActScopes[6] = LaneTrace_NewScope("Iax/Activity_128");
-			ActScopes[7] = LaneTrace_NewScope("Iax/Activity_256");
+
+			TAnsiStringBuilder<32> Builder;
+			for (int32 i = 1, n = UE_ARRAY_COUNT(ActScopes); i < n; ++i)
+			{
+				Builder << "Iax/Activity_";
+				Builder << (1ull << (i - 1));
+				ActScopes[i] = LaneTrace_NewScope(Builder);
+			}
 		}
 
 		FLaneTrace* Lane = LaneEstate_Build(GActivityTraceEstate, Activity);
