Renamed a few types to aid communicating intent

#rnx

diff --git a/Engine/Source/Runtime/Experimental/IoStore/HttpClient/Private/Loop.inl b/Engine/Source/Runtime/Experimental/IoStore/HttpClient/Private/Loop.inl
--- a/Engine/Source/Runtime/Experimental/IoStore/HttpClient/Private/Loop.inl
+++ b/Engine/Source/Runtime/Experimental/IoStore/HttpClient/Private/Loop.inl
@@ -497,10 +497,16 @@ static void DoFail(FActivity* Activity)
  *			Act.Recv
  */
 
+////////////////////////////////////////////////////////////////////////////////
+using FActivityNode = FActivity;
+using FActivityList = FActivityNode*;
+
+
+
 ////////////////////////////////////////////////////////////////////////////////
 struct FTickState
 {
-	FActivity*					DoneList;
+	FActivityList				DoneList;
 	uint64						Cancels;
 	int32&						RecvAllowance;
 	int32						PollTimeoutMs;
@@ -518,13 +524,13 @@ public:
 						FWorkQueue() = default;
 						~FWorkQueue();
 	bool				HasWork() const { return List != nullptr; }
-	void				AddActivity(FActivity* Activity);
-	FActivity*			PopActivity();
+	void				AddActivity(FActivityNode* Activity);
+	FActivityNode*		PopActivity();
 	void				TickCancels(FTickState& State);
 
 private:
-	FActivity*			List = nullptr;
-	FActivity*			ListTail = nullptr;
+	FActivityList		List = nullptr;
+	FActivityNode*		ListTail = nullptr;
 	uint64				ActiveSlots = 0;
 
 	UE_NONCOPYABLE(FWorkQueue);
@@ -538,7 +544,7 @@ FWorkQueue::~FWorkQueue()
 }
 
 ////////////////////////////////////////////////////////////////////////////////
-void FWorkQueue::AddActivity(FActivity* Activity)
+void FWorkQueue::AddActivity(FActivityNode* Activity)
 {
 	// We use a tail pointer here to maintain order that requests were made
 
@@ -555,14 +561,14 @@ void FWorkQueue::AddActivity(FActivity* Activity)
 }
 
 ////////////////////////////////////////////////////////////////////////////////
-FActivity* FWorkQueue::PopActivity()
+FActivityNode* FWorkQueue::PopActivity()
 {
 	if (List == nullptr)
 	{
 		return nullptr;
 	}
 
-	FActivity* Activity = List;
+	FActivityNode* Activity = List;
 	if ((List = List->Next) == nullptr)
 	{
 		ListTail = nullptr;
@@ -587,11 +593,11 @@ void FWorkQueue::TickCancels(FTickState& State)
 	// activity list is singular.
 
 	check(List != nullptr);
-	FActivity* Activity = List;
+	FActivityNode* Activity = List;
 	List = ListTail = nullptr;
 	ActiveSlots = 0;
 
-	for (FActivity* Next; Activity != nullptr; Activity = Next)
+	for (FActivityNode* Next; Activity != nullptr; Activity = Next)
 	{
 		Next = Activity->Next;
 
@@ -629,8 +635,8 @@ private:
 	void				Negotiate(FTickState& State);
 	void				RecvInternal(FTickState& State);
 	void				SendInternal(FTickState& State);
-	FActivity*			Send = nullptr;
-	FActivity*			Recv = nullptr;
+	FActivityList		Send = nullptr;
+	FActivityList		Recv = nullptr;
 	FHttpPeer			Peer;
 	uint32				LastUseMs = 0;
 	uint8				IsKeepAlive = 0;
@@ -673,7 +679,7 @@ void FPeerGroup::Fail(FTickState& State, const char* Reason)
 	}
 
 	// Failure is quite terminal and we need to abort everything
-	for (FActivity* Activity = Recv; Activity != nullptr;)
+	for (FActivityNode* Activity = Recv; Activity != nullptr;)
 	{
 		if (Activity->State != FActivity::EState::Failed)
 		{
@@ -682,7 +688,7 @@ void FPeerGroup::Fail(FTickState& State, const char* Reason)
 
 		DoFail(Activity);
 
-		FActivity* Next = Activity->Next;
+		FActivityNode* Next = Activity->Next;
 		Activity->Next = State.DoneList;
 		State.DoneList = Activity;
 		Activity = Next;
@@ -726,13 +732,13 @@ void FPeerGroup::RecvInternal(FTickState& State)
 	check(Recv != nullptr);
 
 	// Another helper lambda
-	auto IsReceiving = [] (const FActivity* Act)
+	auto IsReceiving = [] (const FActivityNode* Act)
 	{
 		using EState = FActivity::EState;
 		return (Act->State >= EState::RecvMessage) & (Act->State < EState::RecvDone);
 	};
 
-	FActivity* Activity = Recv;
+	FActivityNode* Activity = Recv;
 	check(IsReceiving(Activity));
 
 	FOutcome Outcome = DoRecv(Activity, Peer, State.RecvAllowance);
@@ -754,7 +760,7 @@ void FPeerGroup::RecvInternal(FTickState& State)
 		uint32 Remaining = Activity_RemainingKiB(Activity);
 		if (Remaining < uint32(GRecvWorkThresholdKiB))
 		{
-			if (FActivity* Next = State.Work->PopActivity(); Next != nullptr)
+			if (FActivityNode* Next = State.Work->PopActivity(); Next != nullptr)
 			{
 				Trace(Activity, ETrace::StartWork);
 	
@@ -811,7 +817,7 @@ void FPeerGroup::SendInternal(FTickState& State)
 	check(IsKeepAlive == 1);
 	check(Send != nullptr);
 
-	FActivity* Activity = Send;
+	FActivityNode* Activity = Send;
 
 	FOutcome Outcome = DoSend(Activity, Peer);
 
@@ -871,7 +877,7 @@ void FPeerGroup::TickSend(FTickState& State, FHost& Host, FPoller& Poller)
 	}
 
 	// Failing will try and recover work which we don't want to happen yet
-	FActivity* Pending = State.Work->PopActivity();
+	FActivityNode* Pending = State.Work->PopActivity();
 	check(Pending != nullptr);
 
 	// Close idle sockets
@@ -946,7 +952,7 @@ public:
 	bool						IsBusy() const	{ return BusyCount != 0; }
 	const FHost&				GetHost() const	{ return Host; }
 	void						Tick(FTickState& State);
-	void						AddActivity(FActivity* Activity);
+	void						AddActivity(FActivityNode* Activity);
 
 private:
 	int32						Wait(const FTickState& State);
@@ -1084,7 +1090,7 @@ void FHostGroup::Tick(FTickState& State)
 }
 
 ////////////////////////////////////////////////////////////////////////////////
-void FHostGroup::AddActivity(FActivity* Activity)
+void FHostGroup::AddActivity(FActivityNode* Activity)
 {
 	Work.AddActivity(Activity);
 }
@@ -1106,8 +1112,8 @@ public:
 	void					Throttle(uint32 KiBPerSec);
 	void					SetFailTimeout(int32 TimeoutMs);
 	void					Cancel(FTicket Ticket);
-	FRequest				Request(FAnsiStringView Method, FAnsiStringView Path, FActivity* Activity);
-	FTicket					Send(FActivity* Activity);
+	FRequest				Request(FAnsiStringView Method, FAnsiStringView Path, FActivityNode* Activity);
+	FTicket					Send(FActivityNode* Activity);
 
 private:
 	void					ReceiveWork();
@@ -1115,7 +1121,7 @@ private:
 	std::atomic<uint64>		FreeSlots		= ~0ull;
 	std::atomic<uint64>		Cancels			= 0;
 	uint64					PrevFreeSlots	= ~0ull;
-	FActivity*				Pending			= nullptr;
+	FActivityList			Pending			= nullptr;
 	FThrottler				Throttler;
 	TArray<FHostGroup>		Groups;
 	int32					FailTimeoutMs	= GIdleMs;
@@ -1132,7 +1138,7 @@ FEventLoop::FImpl::~FImpl()
 FRequest FEventLoop::FImpl::Request(
 	FAnsiStringView Method,
 	FAnsiStringView Path,
-	FActivity* Activity)
+	FActivityNode* Activity)
 {
 	Trace(Activity, ETrace::ActivityCreate, 0);
 	Activity_ChangeState(Activity, FActivity::EState::Build);
@@ -1158,7 +1164,7 @@ FRequest FEventLoop::FImpl::Request(
 }
 
 ////////////////////////////////////////////////////////////////////////////////
-FTicket FEventLoop::FImpl::Send(FActivity* Activity)
+FTicket FEventLoop::FImpl::Send(FActivityNode* Activity)
 {
 	Trace(Activity, ETrace::RequestBegin);
 
@@ -1237,15 +1243,15 @@ void FEventLoop::FImpl::ReceiveWork()
 	PrevFreeSlots = FreeSlotsLoad;
 
 	// Fetch the pending activities from out in the wild
-	FActivity* Activity = nullptr;
+	FActivityNode* Activity = nullptr;
 	{
 		FScopeLock _(&Lock);
 		Swap(Activity, Pending);
 	}
 
 	// Pending is in the reverse of the order that requests were made
-	FActivity* Reverse = nullptr;
-	for (FActivity* Next; Activity != nullptr; Activity = Next)
+	FActivityNode* Reverse = nullptr;
+	for (FActivityNode* Next; Activity != nullptr; Activity = Next)
 	{
 		Next = Activity->Next;
 		Activity->Next = Reverse;
@@ -1254,7 +1260,7 @@ void FEventLoop::FImpl::ReceiveWork()
 	Activity = Reverse;
 
 	// Group activities by their host.
-	for (FActivity* Next; Activity != nullptr; Activity = Next)
+	for (FActivityNode* Next; Activity != nullptr; Activity = Next)
 	{
 		Next = Activity->Next;
 		Activity->Next = nullptr;
@@ -1348,9 +1354,9 @@ uint32 FEventLoop::FImpl::Tick(int32 PollTimeoutMs)
 	Throttler.ReturnUnused(RecvAllowance);
 
 	uint64 ReturnedSlots = 0;
-	for (FActivity* Activity = TickState.DoneList; Activity != nullptr;)
+	for (FActivityNode* Activity = TickState.DoneList; Activity != nullptr;)
 	{
-		FActivity* Next = Activity->Next;
+		FActivityNode* Next = Activity->Next;
 		ReturnedSlots |= (1ull << Activity->Slot);
 		Activity_Free(Activity);
 		--BusyCount;
@@ -1428,7 +1434,7 @@ FRequest FEventLoop::Request(
 	uint32 BufferSize = Params->BufferSize;
 	BufferSize = (BufferSize >= 128) ? BufferSize : 128;
 	BufferSize += sizeof(FHost) + HostName.Len();
-	FActivity* Activity = Activity_Alloc(BufferSize);
+	FActivityNode* Activity = Activity_Alloc(BufferSize);
 
 	FBuffer& Buffer = Activity->Buffer;
 
@@ -1469,7 +1475,7 @@ FRequest FEventLoop::Request(
 
 	uint32 BufferSize = Params->BufferSize;
 	BufferSize = (BufferSize >= 128) ? BufferSize : 128;
-	FActivity* Activity = Activity_Alloc(BufferSize);
+	FActivityNode* Activity = Activity_Alloc(BufferSize);
 
 	Activity->Host = Pool.Ptr;
 	Activity->IsKeepAlive = 1;
@@ -1561,7 +1567,7 @@ bool FEventLoop::Redirect(const FTicketStatus& Status, FTicketSink& OuterSink)
 ////////////////////////////////////////////////////////////////////////////////
 FTicket FEventLoop::Send(FRequest&& Request, FTicketSink Sink, UPTRINT SinkParam)
 {
-	FActivity* Activity = nullptr;
+	FActivityNode* Activity = nullptr;
 	Swap(Activity, Request.Ptr);
 	Activity->SinkParam = SinkParam;
 	Activity->Sink = MoveTemp(Sink);
