Fixed off-by-one and a missing reset()

#rnx

diff --git a/Engine/Source/Runtime/Experimental/IoStore/HttpClient/Private/Activity.inl b/Engine/Source/Runtime/Experimental/IoStore/HttpClient/Private/Activity.inl
--- a/Engine/Source/Runtime/Experimental/IoStore/HttpClient/Private/Activity.inl
+++ b/Engine/Source/Runtime/Experimental/IoStore/HttpClient/Private/Activity.inl
@@ -659,7 +659,7 @@ static void Trace(const struct FActivity* Activity, ETrace Action, uint32 Param)
 {
 	if (Action == ETrace::ActivityCreate)
 	{
-		static uint32 ActScopes[15] = {};
+		static uint32 ActScopes[16] = {};
 		if (ActScopes[0] == 0)
 		{
 			ActScopes[0] = LaneTrace_NewScope("Iax/Activity");
@@ -667,6 +667,7 @@ static void Trace(const struct FActivity* Activity, ETrace Action, uint32 Param)
 			TAnsiStringBuilder<32> Builder;
 			for (int32 i = 1, n = UE_ARRAY_COUNT(ActScopes); i < n; ++i)
 			{
+				Builder.Reset();
 				Builder << "Iax/Activity_";
 				Builder << (1ull << (i - 1));
 				ActScopes[i] = LaneTrace_NewScope(Builder);
