Added an IsPooled() accessor to carry a little context around.

#rnx

diff --git a/Engine/Source/Runtime/Experimental/IoStore/HttpClient/Private/ConnectionPool.inl b/Engine/Source/Runtime/Experimental/IoStore/HttpClient/Private/ConnectionPool.inl
--- a/Engine/Source/Runtime/Experimental/IoStore/HttpClient/Private/ConnectionPool.inl
+++ b/Engine/Source/Runtime/Experimental/IoStore/HttpClient/Private/ConnectionPool.inl
@@ -16,7 +16,8 @@ public:
 	{
 		const ANSICHAR* HostName;
 		uint32			Port = 0;
-		uint32			MaxConnections = 1;
+		uint16			MaxConnections = 1;
+		bool			bPooled = false;
 		FCertRootsRef	VerifyCert;
 	};
 
@@ -28,6 +29,7 @@ public:
 	FOutcome		ResolveHostName();
 	FCertRootsRef	GetVerifyCert() const		{ return VerifyCert; }
 	uint32			GetMaxConnections() const	{ return MaxConnections; }
+	bool			IsPooled() const			{ return bPooled; }
 	uint32			GetIpAddress() const		{ return IpAddresses[0]; }
 	FAnsiStringView	GetHostName() const			{ return HostName; }
 	uint32			GetPort() const				{ return Port; }
@@ -40,6 +42,7 @@ private:
 	int16			RecvBufKb = -1;
 	uint16			Port;
 	uint8			MaxConnections;
+	bool			bPooled;
 };
 
 ////////////////////////////////////////////////////////////////////////////////
@@ -48,6 +51,7 @@ FHost::FHost(const FParams& Params)
 , HostName(Params.HostName)
 , Port(uint16(Params.Port))
 , MaxConnections(uint8(Params.MaxConnections))
+, bPooled(Params.bPooled)
 {
 	check(MaxConnections && MaxConnections == Params.MaxConnections);
 
@@ -248,6 +252,7 @@ FConnectionPool::FConnectionPool(const FParams& Params)
 		.HostName		= HostDest,
 		.Port			= Params.Port,
 		.MaxConnections	= Params.ConnectionCount,
+		.bPooled		= true,
 		.VerifyCert		= Params.VerifyCert,
 	});
 	Internal->SetBufferSize(FHost::EDirection::Send, Params.SendBufSize);
