Member-init just in case.

#rnx

diff --git a/Engine/Source/Runtime/Experimental/IoStore/HttpClient/Private/ConnectionPool.inl b/Engine/Source/Runtime/Experimental/IoStore/HttpClient/Private/ConnectionPool.inl
--- a/Engine/Source/Runtime/Experimental/IoStore/HttpClient/Private/ConnectionPool.inl
+++ b/Engine/Source/Runtime/Experimental/IoStore/HttpClient/Private/ConnectionPool.inl
@@ -18,7 +18,7 @@ public:
 		uint32			Port = 0;
 		uint16			MaxConnections = 1;
 		bool			bPooled = false;
-		FCertRootsRef	VerifyCert;
+		FCertRootsRef	VerifyCert = {};
 	};
 
 					FHost(const FParams& Params);
