Always consider thie true unless told otherwise.

#rnx

diff --git a/Engine/Source/Runtime/Experimental/IoStore/HttpClient/Private/Activity.inl b/Engine/Source/Runtime/Experimental/IoStore/HttpClient/Private/Activity.inl
--- a/Engine/Source/Runtime/Experimental/IoStore/HttpClient/Private/Activity.inl
+++ b/Engine/Source/Runtime/Experimental/IoStore/HttpClient/Private/Activity.inl
@@ -137,8 +137,7 @@ private:
 ////////////////////////////////////////////////////////////////////////////////
 bool FResponseCloak::IsKeepAlive() const
 {
-	check(MessageRight > 0);
-	return bKeepAlive;
+	return (MessageRight > 0) ? bKeepAlive : true;
 }
 
 ////////////////////////////////////////////////////////////////////////////////
