Follow redirects all the way.

#rnx

diff --git a/Engine/Source/Runtime/Experimental/IoStore/HttpClient/Private/Loop.inl b/Engine/Source/Runtime/Experimental/IoStore/HttpClient/Private/Loop.inl
--- a/Engine/Source/Runtime/Experimental/IoStore/HttpClient/Private/Loop.inl
+++ b/Engine/Source/Runtime/Experimental/IoStore/HttpClient/Private/Loop.inl
@@ -1510,14 +1510,18 @@ bool FEventLoop::Redirect(const FTicketStatus& Status, FTicketSink& OuterSink)
 		return false;
 	}
 
-	check(Response.GetContentLength() == 0); // should we ever hit this, we'll fix it
-
 	const auto& Activity = (FActivity&)Response; // todo: yuk
 
+	check(!Activity.Buffer.HasBody() || Response.GetContentLength() == 0); // should we ever hit this, we'll fix it
+
 	// Original method should remain unchanged
 	FAnsiStringView Method = Activity.Buffer.GetMethod();
 	check(!Method.IsEmpty());
 
+	FRequestParams RequestParams = {
+		.bAutoRedirect = true,
+	};
+
 	FRequest ForwardRequest;
 	if (!Location.StartsWith("http://") && !Location.StartsWith("https://"))
 	{
@@ -1535,14 +1539,12 @@ bool FEventLoop::Redirect(const FTicketStatus& Status, FTicketSink& OuterSink)
 		Url << ":" << Host.GetPort();
 		Url << Location;
 
-		FRequestParams RequestParams = {
-			.VerifyCert = Host.GetVerifyCert(),
-		};
+		RequestParams.VerifyCert = Host.GetVerifyCert();
 		new (&ForwardRequest) FRequest(Request(Method, Url, &RequestParams));
 	}
 	else
 	{
-		new (&ForwardRequest) FRequest(Request(Method, Location));
+		new (&ForwardRequest) FRequest(Request(Method, Location, &RequestParams));
 	}
 
 	// Transfer original request headers
diff --git a/Engine/Source/Runtime/Experimental/IoStore/HttpClient/Private/Test.inl b/Engine/Source/Runtime/Experimental/IoStore/HttpClient/Private/Test.inl
--- a/Engine/Source/Runtime/Experimental/IoStore/HttpClient/Private/Test.inl
+++ b/Engine/Source/Runtime/Experimental/IoStore/HttpClient/Private/Test.inl
@@ -254,7 +254,8 @@ static void TlsTest()
 	static const ANSICHAR* Url = "https://httpbin.org/get";
 
 	{
-		FRequest Request = Loop.Get(Url);
+		FEventLoop::FRequestParams RequestParams = { .bAutoRedirect = true };
+		FRequest Request = Loop.Request("HEAD", Url, &RequestParams);
 		Loop.Send(MoveTemp(Request), OkSink);
 		WaitForLoopIdle();
 	}
