Use our own properties for tests.

#rnx

diff --git a/Engine/Source/Runtime/Experimental/IoStore/HttpClient/Private/Test.inl b/Engine/Source/Runtime/Experimental/IoStore/HttpClient/Private/Test.inl
--- a/Engine/Source/Runtime/Experimental/IoStore/HttpClient/Private/Test.inl
+++ b/Engine/Source/Runtime/Experimental/IoStore/HttpClient/Private/Test.inl
@@ -122,7 +122,7 @@ static void MiscTest()
 
 	for (int32 i : { 0x10, 0x20, 0x40, 0x7f, 0xff })
 	{
-		char Url[] = "http://stockholm.patchercache.epicgames.net:123";
+		char Url[] = "http://sweden.denmark.finland.norway.iceland:123";
 		char Buffer[512];
 		std::memset(Buffer, i, sizeof(Buffer));
 		std::memcpy(Buffer, Url, sizeof(Url) - 1);
@@ -251,7 +251,7 @@ static void TlsTest()
 		check(Status.GetId() == FTicketStatus::EId::Error);
 	};
 
-	static const ANSICHAR* Url = "https://httpbin.org/get";
+	static const ANSICHAR* Url = "http://epicgames.com";
 
 	{
 		FEventLoop::FRequestParams RequestParams = { .bAutoRedirect = true };
@@ -650,7 +650,7 @@ static void HttpTest(const ANSICHAR* TestHost, FCertRootsRef VerifyCert)
 
 		Loop.Send(Loop.Get(BuildUrl("/data"), ReqParams).Accept(EMimeType::Json), HashSink),
 		Loop.Send(MoveTemp(Request), HashSink),
-		Loop.Send(Loop.Get("http://httpbin.org/get"), NoErrorSink),
+		Loop.Send(Loop.Get("http://epicgames.com"), NoErrorSink),
 
 		WaitForLoopIdle();
 	}
