Introduce some encapsulation and contract between frontend and activities.

#rnx

diff --git a/Engine/Source/Runtime/Experimental/IoStore/HttpClient/Private/Activity.inl b/Engine/Source/Runtime/Experimental/IoStore/HttpClient/Private/Activity.inl
--- a/Engine/Source/Runtime/Experimental/IoStore/HttpClient/Private/Activity.inl
+++ b/Engine/Source/Runtime/Experimental/IoStore/HttpClient/Private/Activity.inl
@@ -504,6 +504,20 @@ static FLaneEstate* GActivityTraceEstate = LaneEstate_New({
 ////////////////////////////////////////////////////////////////////////////////
 struct alignas(16) FActivity
 {
+	void				SetSink(FTicketSink&& InSink, UPTRINT Param);
+	void				SetDestination(FIoBuffer* InDest);
+	FHost*				GetHost() const;
+	uint32				GetIndex() const;
+	uint64				GetTicket() const;
+	UPTRINT				GetSinkParam() const;
+	FIoBuffer&			GetContent() const;
+	const char*			GetErrorReason() const;
+	uint32				GetErrorCode() const;
+
+#if IAS_HTTP_WITH_PERF
+	const FStopwatch&	GetStopwatch() const;
+#endif
+
 	enum class EState : uint8
 	{
 		None,
@@ -541,6 +555,75 @@ struct alignas(16) FActivity
 	FBodyCloak			Buffer;
 };
 
+////////////////////////////////////////////////////////////////////////////////
+void FActivity::SetSink(FTicketSink&& InSink, UPTRINT Param)
+{
+	Sink = MoveTemp(InSink);
+	SinkParam = Param;
+}
+
+////////////////////////////////////////////////////////////////////////////////
+void FActivity::SetDestination(FIoBuffer* InDest)
+{
+	Dest = InDest;
+}
+
+////////////////////////////////////////////////////////////////////////////////
+FHost* FActivity::GetHost() const
+{
+	check(State <= EState::RecvMessage);
+	return Host;
+}
+
+////////////////////////////////////////////////////////////////////////////////
+uint32 FActivity::GetIndex() const
+{
+	return Slot;
+}
+
+////////////////////////////////////////////////////////////////////////////////
+uint64 FActivity::GetTicket() const
+{
+	return 1ull << Slot;
+}
+
+////////////////////////////////////////////////////////////////////////////////
+UPTRINT FActivity::GetSinkParam() const
+{
+	return SinkParam;
+}
+
+////////////////////////////////////////////////////////////////////////////////
+FIoBuffer& FActivity::GetContent() const
+{
+	check(State == EState::RecvContent || State == EState::RecvStream);
+	return *Dest;
+}
+
+////////////////////////////////////////////////////////////////////////////////
+const char* FActivity::GetErrorReason() const
+{
+	check(State == EState::Failed);
+	return ErrorReason;
+}
+
+////////////////////////////////////////////////////////////////////////////////
+uint32 FActivity::GetErrorCode() const
+{
+	check(State == EState::Failed);
+	return uint32(StateParam);
+}
+
+////////////////////////////////////////////////////////////////////////////////
+#if IAS_HTTP_WITH_PERF
+const FStopwatch& FActivity::GetStopwatch() const
+{
+	return Stopwatch;
+}
+#endif // IAS_HTTP_WITH_PERF
+
+
+
 ////////////////////////////////////////////////////////////////////////////////
 static void Activity_ChangeState(FActivity* Activity, FActivity::EState InState, uint32 Param=0)
 {
diff --git a/Engine/Source/Runtime/Experimental/IoStore/HttpClient/Private/Api.inl b/Engine/Source/Runtime/Experimental/IoStore/HttpClient/Private/Api.inl
--- a/Engine/Source/Runtime/Experimental/IoStore/HttpClient/Private/Api.inl
+++ b/Engine/Source/Runtime/Experimental/IoStore/HttpClient/Private/Api.inl
@@ -150,7 +150,7 @@ void FResponse::ReadHeaders(FHeaderSink Sink) const
 void FResponse::SetDestination(FIoBuffer* Buffer)
 {
 	auto* Activity = (FActivity*)this;
-	Activity->Dest = Buffer;
+	Activity->SetDestination(Buffer);
 }
 
 
@@ -177,21 +177,21 @@ FTicketStatus::EId FTicketStatus::GetId() const
 UPTRINT FTicketStatus::GetParam() const
 {
 	const auto* Activity = (FActivity*)this;
-	return Activity->SinkParam;
+	return Activity->GetSinkParam();
 }
 
 ////////////////////////////////////////////////////////////////////////////////
 FTicket FTicketStatus::GetTicket() const
 {
 	const auto* Activity = (FActivity*)this;
-	return 1ull << Activity->Slot;
+	return Activity->GetTicket();
 }
 
 ////////////////////////////////////////////////////////////////////////////////
 uint32 FTicketStatus::GetIndex() const
 {
 	const auto* Activity = (FActivity*)this;
-	return Activity->Slot;
+	return Activity->GetIndex();
 }
 
 ////////////////////////////////////////////////////////////////////////////////
@@ -223,7 +223,7 @@ const FIoBuffer& FTicketStatus::GetContent() const
 {
 	check(GetId() == EId::Content);
 	const auto* Activity = (FActivity*)this;
-	return *(Activity->Dest);
+	return Activity->GetContent();
 }
 
 ////////////////////////////////////////////////////////////////////////////////
@@ -231,7 +231,7 @@ FTicketStatus::FError FTicketStatus::GetError() const
 {
 	check(GetId() == EId::Error);
 	const auto* Activity = (FActivity*)this;
-	return FError({Activity->ErrorReason, Activity->StateParam});
+	return { Activity->GetErrorReason(), Activity->GetErrorCode() };
 }
 
 
@@ -254,7 +254,7 @@ FTicketPerf::FSample FTicketPerf::GetSample() const
 	auto Clamp = [] (auto Value) { return uint16(FMath::Min(uint32(Value), 0xffffu)); };
 	auto ToMs = [&] (uint64 Value) { return Clamp((Value * 1000ull) / Freq); };
 
-	const FStopwatch& Stopwatch = Activity.Stopwatch;
+	const FStopwatch& Stopwatch = Activity.GetStopwatch();
 	FSample Sample = {
 		ToMs(Stopwatch.GetInterval(0)),
 		ToMs(Stopwatch.GetInterval(1)),
diff --git a/Engine/Source/Runtime/Experimental/IoStore/HttpClient/Private/Loop.inl b/Engine/Source/Runtime/Experimental/IoStore/HttpClient/Private/Loop.inl
--- a/Engine/Source/Runtime/Experimental/IoStore/HttpClient/Private/Loop.inl
+++ b/Engine/Source/Runtime/Experimental/IoStore/HttpClient/Private/Loop.inl
@@ -104,11 +104,8 @@ static FOutcome DoRecvMessage(FActivity* Activity, FHttpPeer& Peer)
 	check(Outcome.IsOk());
 
 	bool bChunked = Response.IsChunked();
-	bool IsKeepAlive = Response.IsKeepAlive();
 	int64 ContentLength = Response.GetContentLength();
 
-	Activity->IsKeepAlive &= IsKeepAlive;
-
 	// Validate that the server's told us how and how much it will transmit
 	if (bChunked)
 	{
@@ -750,7 +747,7 @@ void FPeerGroup::RecvInternal(FTickState& State)
 		return;
 	}
 
-	IsKeepAlive &= Activity->IsKeepAlive;
+	IsKeepAlive &= uint8(Activity->Buffer.IsKeepAlive() == true);
 	LastUseMs = State.NowMs;
 	bWaiting |= Outcome.IsWaiting();
 
@@ -1148,12 +1145,14 @@ FRequest FEventLoop::FImpl::Request(
 		Path = "/";
 	}
 
+	const FHost* Host = Activity->GetHost();
+
 	auto& Buffer = Activity->Buffer;
 	Buffer.Begin(Method, Path);
-	Buffer.AddHeader("Host", Activity->Host->GetHostName());
+	Buffer.AddHeader("Host", Host->GetHostName());
 
 	// HTTP/1.1 is persistent by default thus "Connection" header isn't required
-	if (!Activity->IsKeepAlive)
+	if (!Host->IsPooled())
 	{
 		Buffer.AddHeader("Connection", "close");
 	}
@@ -1265,7 +1264,7 @@ void FEventLoop::FImpl::ReceiveWork()
 		Next = Activity->Next;
 		Activity->Next = nullptr;
 
-		FHost& Host = *(Activity->Host);
+		FHost& Host = *(Activity->GetHost());
 		auto Pred = [&Host] (const FHostGroup& Lhs) { return &Lhs.GetHost() == &Host; };
 		FHostGroup* Group = Groups.FindByPredicate(Pred);
 		if (Group == nullptr)
@@ -1530,7 +1529,7 @@ bool FEventLoop::Redirect(const FTicketStatus& Status, FTicketSink& OuterSink)
 			return false;
 		}
 
-		FHost& Host = *(Activity.Host);
+		const FHost& Host = *(Activity.GetHost());
 
 		TAnsiStringBuilder<256> Url;
 		Url << ((Host.GetVerifyCert() != ECertRootsRefType::None) ? "https" : "http");
@@ -1571,15 +1570,13 @@ FTicket FEventLoop::Send(FRequest&& Request, FTicketSink Sink, UPTRINT SinkParam
 {
 	FActivityNode* Activity = nullptr;
 	Swap(Activity, Request.Ptr);
-	Activity->SinkParam = SinkParam;
-	Activity->Sink = MoveTemp(Sink);
 
 	// Intercept sink calls to catch 30x status codes and follow them
 	if (Activity->bFollow30x)
 	{
 		auto RedirectSink = [
 				this,
-				OuterSink=MoveTemp(Activity->Sink)
+				OuterSink=MoveTemp(Sink)
 			] (const FTicketStatus& Status) mutable
 		{
 			if (Status.GetId() == FTicketStatus::EId::Response)
@@ -1595,9 +1592,10 @@ FTicket FEventLoop::Send(FRequest&& Request, FTicketSink Sink, UPTRINT SinkParam
 				return OuterSink(Status);
 			}
 		};
-		Activity->Sink = RedirectSink;
+		Sink = MoveTemp(RedirectSink);
 	}
 
+	Activity->SetSink(MoveTemp(Sink), SinkParam);
 	return Impl->Send(Activity);
 }
 
