Consolidate activity creation/destruction.

#rnx

diff --git a/Engine/Source/Runtime/Experimental/IoStore/HttpClient/Internal/IO/Http/Client.h b/Engine/Source/Runtime/Experimental/IoStore/HttpClient/Internal/IO/Http/Client.h
--- a/Engine/Source/Runtime/Experimental/IoStore/HttpClient/Internal/IO/Http/Client.h
+++ b/Engine/Source/Runtime/Experimental/IoStore/HttpClient/Internal/IO/Http/Client.h
@@ -47,7 +47,7 @@ enum class EStatusCodeClass
 ////////////////////////////////////////////////////////////////////////////////
 using	FCertRootsRef	= UPTRINT;
 using	FTicket			= uint64;
-struct	FActivity;
+struct	FActivityNode;
 
 ////////////////////////////////////////////////////////////////////////////////
 class FCertRoots
@@ -127,7 +127,7 @@ public:
 private:
 	friend				class FEventLoop;
 						FRequest() = default;
-	FActivity*			Ptr = nullptr;
+	FActivityNode*		Ptr = nullptr;
 
 private:
 						FRequest(const FRequest&) = delete;
diff --git a/Engine/Source/Runtime/Experimental/IoStore/HttpClient/Private/Activity.inl b/Engine/Source/Runtime/Experimental/IoStore/HttpClient/Private/Activity.inl
--- a/Engine/Source/Runtime/Experimental/IoStore/HttpClient/Private/Activity.inl
+++ b/Engine/Source/Runtime/Experimental/IoStore/HttpClient/Private/Activity.inl
@@ -502,8 +502,22 @@ static FLaneEstate* GActivityTraceEstate = LaneEstate_New({
 
 
 ////////////////////////////////////////////////////////////////////////////////
-struct alignas(16) FActivity
+struct FActivity
 {
+	struct FParams
+	{
+		FHost*			Host = nullptr;
+		char*			Buffer = nullptr;
+		uint16			BufferSize = 0;
+		bool			bIsKeepAlive = false;
+		bool			bFollow30x = false;
+		bool			bAllowChunked = true;
+		uint32			ContentSizeEst = 0;
+	};
+
+						FActivity() = default;
+						FActivity(const FParams& Params);
+						~FActivity() = default;
 	void				SetSink(FTicketSink&& InSink, UPTRINT Param);
 	void				SetDestination(FIoBuffer* InDest);
 	FHost*				GetHost() const;
@@ -533,7 +547,6 @@ struct alignas(16) FActivity
 		_Num,
 	};
 
-	FActivity*			Next = nullptr;
 	int8				Slot = -1;
 	EState				State = EState::None;
 	uint8				IsKeepAlive : 1;
@@ -555,6 +568,21 @@ struct alignas(16) FActivity
 	FBodyCloak			Buffer;
 };
 
+////////////////////////////////////////////////////////////////////////////////
+FActivity::FActivity(const FParams& Params)
+: IsKeepAlive(Params.bIsKeepAlive)
+, bFollow30x(Params.bFollow30x)
+, bAllowChunked(Params.bAllowChunked)
+, Host(Params.Host)
+, Buffer(Params.Buffer, Params.BufferSize)
+{
+	check(Host != nullptr);
+	uint32 ContentEstKiB = (Params.ContentSizeEst + 1023) >> 10;
+	uint32 Pow2 = FMath::FloorLog2(uint32(ContentEstKiB));
+	Pow2 = FMath::Min(Pow2, 15u);
+	LengthScore = uint8(Pow2);
+}
+
 ////////////////////////////////////////////////////////////////////////////////
 void FActivity::SetSink(FTicketSink&& InSink, UPTRINT Param)
 {
diff --git a/Engine/Source/Runtime/Experimental/IoStore/HttpClient/Private/Api.inl b/Engine/Source/Runtime/Experimental/IoStore/HttpClient/Private/Api.inl
--- a/Engine/Source/Runtime/Experimental/IoStore/HttpClient/Private/Api.inl
+++ b/Engine/Source/Runtime/Experimental/IoStore/HttpClient/Private/Api.inl
@@ -15,7 +15,7 @@ FRequest::~FRequest()
 		return;
 	}
 
-	Activity_Free(Ptr);
+	FActivityNode::Destroy(Ptr);
 }
 
 ////////////////////////////////////////////////////////////////////////////////
diff --git a/Engine/Source/Runtime/Experimental/IoStore/HttpClient/Private/Loop.inl b/Engine/Source/Runtime/Experimental/IoStore/HttpClient/Private/Loop.inl
--- a/Engine/Source/Runtime/Experimental/IoStore/HttpClient/Private/Loop.inl
+++ b/Engine/Source/Runtime/Experimental/IoStore/HttpClient/Private/Loop.inl
@@ -495,8 +495,118 @@ static void DoFail(FActivity* Activity)
  */
 
 ////////////////////////////////////////////////////////////////////////////////
-using FActivityNode = FActivity;
-using FActivityList = FActivityNode*;
+using FActivityList = struct FActivityNode*;
+
+struct FActivityNode
+	: public FActivity
+{
+	static FActivityNode*	Create(FParams& Params, FAnsiStringView* Url=nullptr, FCertRootsRef VerifyCert={});
+	static void				Destroy(FActivityNode* Activity);
+	FActivityNode*			Next = nullptr;
+
+private:
+	using					FActivity::FActivity;
+							~FActivityNode();
+};
+
+////////////////////////////////////////////////////////////////////////////////
+FActivityNode* FActivityNode::Create(FParams& Params, FAnsiStringView* InOutUrl, FCertRootsRef VerifyCert)
+{
+	static_assert(alignof(FActivityNode) <= 16);
+	static_assert(alignof(FHost) <= alignof(FActivityNode));
+
+	if (InOutUrl == nullptr)
+	{
+		uint32 BufferSize = Params.BufferSize;
+		BufferSize = (BufferSize >= 128) ? BufferSize : 128;
+		BufferSize = (BufferSize + 15) & ~15;
+
+		uint32 AllocSize = sizeof(FActivityNode) + BufferSize;
+		auto* Ptr = (FActivityNode*)FMemory::Malloc(AllocSize, alignof(FActivityNode));
+
+		check(Params.Host != nullptr);
+		Params.Buffer = (char*)(Ptr + 1);
+		Params.BufferSize = uint16(BufferSize);
+		Params.bIsKeepAlive = true;
+		return new (Ptr) FActivityNode(Params);
+	}
+
+	FAnsiStringView& Url = *InOutUrl;
+
+	// Parse the URL into its components
+	FUrlOffsets UrlOffsets;
+	if (ParseUrl(Url, UrlOffsets) < 0)
+	{
+		return nullptr;
+	}
+
+	FAnsiStringView HostName = UrlOffsets.HostName.Get(Url);
+
+	uint32 Port = 0;
+	if (UrlOffsets.Port)
+	{
+		FAnsiStringView PortView = UrlOffsets.Port.Get(Url);
+		Port = uint32(CrudeToInt(PortView));
+	}
+
+	FAnsiStringView Path;
+	if (UrlOffsets.Path > 0)
+	{
+		Path = Url.Mid(UrlOffsets.Path);
+	}
+	*InOutUrl = Path;
+
+	uint32 BufferSize = Params.BufferSize;
+	BufferSize = (BufferSize >= 128) ? BufferSize : 128;
+	BufferSize += sizeof(FHost);
+	BufferSize += HostName.Len() + 1;
+	BufferSize = (BufferSize + 15) & ~15;
+
+	uint32 AllocSize = sizeof(FActivityNode) + BufferSize;
+	auto* Ptr = (FActivityNode*)FMemory::Malloc(AllocSize, alignof(FActivityNode));
+
+	// Create an emphemeral host
+	if (UrlOffsets.SchemeLength == 5)
+	{
+		if (VerifyCert == ECertRootsRefType::None)
+		{
+			VerifyCert = FCertRoots::Default();
+		}
+		check(VerifyCert != ECertRootsRefType::None);
+	}
+	else
+	{
+		VerifyCert = FCertRoots::NoTls();
+	}
+
+	FHost* Host = (FHost*)(Ptr + 1);
+
+	char* HostNamePtr = (char*)(Host + 1);
+	uint32 HostNameLength = HostName.Len();
+	memcpy(HostNamePtr, HostName.GetData(), HostNameLength);
+	HostNamePtr[HostNameLength] = '\0';
+
+	new (Host) FHost({
+		.HostName	= HostNamePtr,
+		.Port		= Port,
+		.VerifyCert = VerifyCert,
+	});
+
+	HostNameLength = (HostNameLength + 8) & ~7;
+
+	check(Params.Host == nullptr);
+	Params.Host = Host;
+	Params.Buffer = HostNamePtr + HostNameLength;
+	Params.BufferSize = uint16(AllocSize - ptrdiff_t(Params.Buffer - (char*)Ptr));
+	Params.bIsKeepAlive = false;
+	return new (Ptr) (FActivityNode)(Params);
+}
+
+////////////////////////////////////////////////////////////////////////////////
+void FActivityNode::Destroy(FActivityNode* Ptr)
+{
+	FMemory::Free(Ptr);
+}
 
 
 
@@ -1357,7 +1467,9 @@ uint32 FEventLoop::FImpl::Tick(int32 PollTimeoutMs)
 	{
 		FActivityNode* Next = Activity->Next;
 		ReturnedSlots |= (1ull << Activity->Slot);
-		Activity_Free(Activity);
+
+		FActivityNode::Destroy(Activity);
+
 		--BusyCount;
 		Activity = Next;
 	}
@@ -1395,69 +1507,21 @@ FRequest FEventLoop::Request(
 	FAnsiStringView Url,
 	const FRequestParams* Params)
 {
-	// Parse the URL into its components
-	FUrlOffsets UrlOffsets;
-	if (ParseUrl(Url, UrlOffsets) < 0)
-	{
-		return FRequest();
-	}
-
-	FAnsiStringView HostName = UrlOffsets.HostName.Get(Url);
-
-	uint32 Port = 0;
-	if (UrlOffsets.Port)
-	{
-		FAnsiStringView PortView = UrlOffsets.Port.Get(Url);
-		Port = uint32(CrudeToInt(PortView));
-	}
-
-	FAnsiStringView Path;
-	if (UrlOffsets.Path > 0)
-	{
-		Path = Url.Mid(UrlOffsets.Path);
-	}
-
-	// Create an activity and an emphemeral host
 	Params = (Params != nullptr) ? Params : &GDefaultParams;
 
-	FCertRootsRef VerifyCert = FCertRoots::NoTls();
-	if (UrlOffsets.SchemeLength == 5)
+	FActivityNode::FParams ActivityParams;
+	ActivityParams.BufferSize = Params->BufferSize;
+	ActivityParams.bFollow30x = (Params->bAutoRedirect == true);
+	ActivityParams.bAllowChunked = (Params->bAllowChunked == true);
+	ActivityParams.ContentSizeEst = Params->ContentSizeEst;
+
+	FActivityNode* Activity = FActivityNode::Create(ActivityParams, &Url, Params->VerifyCert);
+	if (Activity == nullptr)
 	{
-		if (VerifyCert = Params->VerifyCert; VerifyCert == ECertRootsRefType::None)
-		{
-			VerifyCert = FCertRoots::Default();
-		}
-		check(VerifyCert != ECertRootsRefType::None);
+		return FRequest();
 	}
 
-	uint32 BufferSize = Params->BufferSize;
-	BufferSize = (BufferSize >= 128) ? BufferSize : 128;
-	BufferSize += sizeof(FHost) + HostName.Len();
-	FActivityNode* Activity = Activity_Alloc(BufferSize);
-
-	FBuffer& Buffer = Activity->Buffer;
-
-	FHost* Host = Activity->Host = Buffer.Alloc<FHost>();
-	Activity->IsKeepAlive = 0;
-	Activity->bFollow30x = (Params->bAutoRedirect == true);
-	Activity->bAllowChunked = (Params->bAllowChunked == true);
-
-	uint32 HostNameLength = HostName.Len();
-	char* HostNamePtr = Buffer.Alloc<char>(HostNameLength + 1);
-
-	Buffer.Fix();
-
-	Activity_SetScore(Activity, Params->ContentSizeEst);
-
-	memcpy(HostNamePtr, HostName.GetData(), HostNameLength);
-	HostNamePtr[HostNameLength] = '\0';
-	new (Host) FHost({
-		.HostName	= HostNamePtr,
-		.Port		= Port,
-		.VerifyCert = VerifyCert,
-	});
-
-	return Impl->Request(Method, Path, Activity);
+	return Impl->Request(Method, Url, Activity);
 }
 
 ////////////////////////////////////////////////////////////////////////////////
@@ -1472,18 +1536,13 @@ FRequest FEventLoop::Request(
 
 	Params = (Params != nullptr) ? Params : &GDefaultParams;
 
-	uint32 BufferSize = Params->BufferSize;
-	BufferSize = (BufferSize >= 128) ? BufferSize : 128;
-	FActivityNode* Activity = Activity_Alloc(BufferSize);
-
-	Activity->Host = Pool.Ptr;
-	Activity->IsKeepAlive = 1;
-	Activity->bFollow30x = (Params->bAutoRedirect == true);
-	Activity->bAllowChunked = (Params->bAllowChunked == true);
-	Activity->LengthScore = 0;
-	Activity->NoContent = 0;
-
-	Activity_SetScore(Activity, Params->ContentSizeEst);
+	FActivityNode::FParams ActivityParams;
+	ActivityParams.Host = Pool.Ptr;
+	ActivityParams.bFollow30x = (Params->bAutoRedirect == true);
+	ActivityParams.bAllowChunked = (Params->bAllowChunked == true);
+	ActivityParams.ContentSizeEst = Params->ContentSizeEst;
+	FActivityNode* Activity = FActivityNode::Create(ActivityParams);
+	check(Activity != nullptr);
 
 	return Impl->Request(Method, Path, Activity);
 }
