Remove unused alloc/free functions

#rnx

diff --git a/Engine/Source/Runtime/Experimental/IoStore/HttpClient/Private/Activity.inl b/Engine/Source/Runtime/Experimental/IoStore/HttpClient/Private/Activity.inl
--- a/Engine/Source/Runtime/Experimental/IoStore/HttpClient/Private/Activity.inl
+++ b/Engine/Source/Runtime/Experimental/IoStore/HttpClient/Private/Activity.inl
@@ -705,32 +705,6 @@ static void Activity_CallSink(FActivity* Activity)
 	Activity->Sink(SinkArg);
 }
 
-////////////////////////////////////////////////////////////////////////////////
-static FActivity* Activity_Alloc(uint32 BufferSize)
-{
-	BufferSize = (BufferSize + 15) & ~15;
-
-	uint32 Size = BufferSize + sizeof(FActivity);
-	auto* Activity = (FActivity*)FMemory::Malloc(Size, alignof(FActivity));
-
-	new (Activity) FActivity();
-
-	auto* Scratch = (char*)(Activity + 1);
-	uint32 ScratchSize = BufferSize;
-	Activity->Buffer = decltype(Activity->Buffer)(Scratch, ScratchSize);
-
-	return Activity;
-}
-
-////////////////////////////////////////////////////////////////////////////////
-static void Activity_Free(FActivity* Activity)
-{
-	Trace(Activity, ETrace::ActivityDestroy, 0);
-
-	Activity->~FActivity();
-	FMemory::Free(Activity);
-}
-
 ////////////////////////////////////////////////////////////////////////////////
 static void Activity_SetError(FActivity* Activity, const char* Reason, int32 Code=-1)
 {
@@ -747,21 +721,6 @@ static void Activity_SetError(FActivity* Activity, const FOutcome& Outcome)
 	Activity_SetError(Activity, Outcome.GetMessage().GetData(), Outcome.GetErrorCode());
 }
 
-////////////////////////////////////////////////////////////////////////////////
-static void Activity_SetScore(FActivity* Activity, uint32 ContentSizeEst)
-{
-	if (ContentSizeEst == 0)
-	{
-		Activity->LengthScore = 0;
-		return;
-	}
-
-	uint32 ContentEstKiB = (ContentSizeEst + 1023) >> 10;
-	uint32 Pow2 = FMath::FloorLog2(uint32(ContentEstKiB));
-	Pow2 = FMath::Min(Pow2, 15u);
-	Activity->LengthScore = uint8(Pow2);
-}
-
 
 
 ////////////////////////////////////////////////////////////////////////////////
