Added a few checks to assert assumptions

#rnx

diff --git a/Engine/Source/Runtime/Experimental/IoStore/HttpClient/Private/Activity.inl b/Engine/Source/Runtime/Experimental/IoStore/HttpClient/Private/Activity.inl
--- a/Engine/Source/Runtime/Experimental/IoStore/HttpClient/Private/Activity.inl
+++ b/Engine/Source/Runtime/Experimental/IoStore/HttpClient/Private/Activity.inl
@@ -55,6 +55,7 @@ void FRequestCloak::Begin(FAnsiStringView Method, FAnsiStringView Path)
 ////////////////////////////////////////////////////////////////////////////////
 void FRequestCloak::AddHeader(FAnsiStringView Key, FAnsiStringView Value)
 {
+	check(HeaderRight < 0);
 	*this << Key << ":" << Value << "\r\n";
 }
 
@@ -515,7 +516,6 @@ struct FActivity
 		uint32			ContentSizeEst = 0;
 	};
 
-						FActivity() = default;
 						FActivity(const FParams& Params);
 						~FActivity() = default;
 	void				SetSink(FTicketSink&& InSink, UPTRINT Param);
diff --git a/Engine/Source/Runtime/Experimental/IoStore/HttpClient/Private/Loop.inl b/Engine/Source/Runtime/Experimental/IoStore/HttpClient/Private/Loop.inl
--- a/Engine/Source/Runtime/Experimental/IoStore/HttpClient/Private/Loop.inl
+++ b/Engine/Source/Runtime/Experimental/IoStore/HttpClient/Private/Loop.inl
@@ -525,6 +525,8 @@ FActivityNode* FActivityNode::Create(FParams& Params, FAnsiStringView* InOutUrl,
 		auto* Ptr = (FActivityNode*)FMemory::Malloc(AllocSize, alignof(FActivityNode));
 
 		check(Params.Host != nullptr);
+		check(Params.Host->IsPooled());
+
 		Params.Buffer = (char*)(Ptr + 1);
 		Params.BufferSize = uint16(BufferSize);
 		Params.bIsKeepAlive = true;
@@ -871,7 +873,9 @@ void FPeerGroup::RecvInternal(FTickState& State)
 			{
 				Trace(Activity, ETrace::StartWork);
 	
+				check(Next->GetHost()->IsPooled());
 				check(Send == nullptr);
+
 				Send = Next;
 				SendInternal(State);
 
