Improve logging after an IAS hostgroup has finished trying to connect.

#rnx

- Add if the connection was a success or failure to the logging output.

diff --git a/Engine/Source/Runtime/Experimental/IoStore/OnDemand/Private/OnDemandIoDispatcherBackend.cpp b/Engine/Source/Runtime/Experimental/IoStore/OnDemand/Private/OnDemandIoDispatcherBackend.cpp
--- a/Engine/Source/Runtime/Experimental/IoStore/OnDemand/Private/OnDemandIoDispatcherBackend.cpp
+++ b/Engine/Source/Runtime/Experimental/IoStore/OnDemand/Private/OnDemandIoDispatcherBackend.cpp
@@ -1214,7 +1214,8 @@ bool FOnDemandIoBackend::CreateDefaultHostGroup()
 		return false;
 	}
 
-	if (ConnectionTest(DefaultHostGroup.GetPrimaryHostUrl(), DefaultHostGroup.GetTestPath(), GIasHttpTimeOutMs))
+	const bool bResult = ConnectionTest(DefaultHostGroup.GetPrimaryHostUrl(), DefaultHostGroup.GetTestPath(), GIasHttpTimeOutMs);
+	if (bResult)
 	{
 		Stats.OnHttpConnected();
 	}
@@ -1224,7 +1225,7 @@ bool FOnDemandIoBackend::CreateDefaultHostGroup()
 	}
 
 	const double InitTime = FPlatformTime::Seconds() - InitStart;
-	UE_LOG(LogIas, Display, TEXT("HostGroup init took %.3f seconds"), InitTime);
+	UE_LOG(LogIas, Display, TEXT("HostGroup init took %.3f seconds (%s)"), InitTime , bResult ? TEXT("Succeeded") : TEXT("Failed"));
 
 	return true;
 }
