Fix issue where IAS could be prevented from ever connecting to a valid endpoint

#rnx

- Add a new delegate FHostGroupManager::OnHostGroupDisconncted that is fired when ever one of the registered hostgroups is disconnected.
-- Cleaned up the logging when hostgroups try to reconnect:
--- The parameters in some places were the wrong way around.
--- We did not log when trying to reconnect to the original primary host.
--- We did not log the result if the primary host could not be reconnected too.
- FOnDemandHttpThread now registers itself to the new delegate and will attempt to wake the thread up when the delegate is fired. This will allow the thread to tick and start trying to reconnect to any disconnected hostgroups.
-- We know that FHostGroupManager outlives FOnDemandHttpThread so we can just unregister in the destructor without checking.
- Changed FOnDemandHttpThread::FOnTickIdle to thread safe as there is a very very small chance of concurrent access.

diff --git a/Engine/Source/Runtime/Experimental/IoStore/OnDemand/Private/IasHostGroup.cpp b/Engine/Source/Runtime/Experimental/IoStore/OnDemand/Private/IasHostGroup.cpp
--- a/Engine/Source/Runtime/Experimental/IoStore/OnDemand/Private/IasHostGroup.cpp
+++ b/Engine/Source/Runtime/Experimental/IoStore/OnDemand/Private/IasHostGroup.cpp
@@ -212,6 +212,8 @@ void FIASHostGroup::Disconnect()
 	Impl->HttpErrorHistory.Reset();
 
 	SetPrimaryHost(INDEX_NONE);
+
+	FHostGroupManager::Get().OnHostGroupDisconncted().Broadcast();
 }
 
 FIASHostGroup::EReconnectionResult FIASHostGroup::AttemptReconnection(uint32 TimeoutMs, std::atomic_bool& CancellationToken)
@@ -224,7 +226,7 @@ FIASHostGroup::EReconnectionResult FIASHostGroup::AttemptReconnection(uint32 Tim
 		{
 			Connect(Idx);
 
-			UE_LOGFMT(LogIas, Log, "[{HostName}] Successfully reconnected to '{Url}'", GetPrimaryHostUrl(), GetName());
+			UE_LOGFMT(LogIas, Log, "[{HostName}] Successfully reconnected to '{Url}'", GetName(), GetPrimaryHostUrl());
 			return EReconnectionResult::Reconnected;
 		}
 		else
@@ -234,10 +236,15 @@ FIASHostGroup::EReconnectionResult FIASHostGroup::AttemptReconnection(uint32 Tim
 	}
 	else if (GetPrimaryHostIndex() != 0)
 	{
+		UE_LOGFMT(LogIas, Log, "[{HostName}] Trying to reconnect to primary endpoint...", GetName()); 
 		if (const int32 Idx = ConnectionTest(GetHostUrls().Left(1), Impl->TestPath, TimeoutMs, CancellationToken); Idx != INDEX_NONE)
 		{
 			SetPrimaryHost(Idx);
-			UE_LOGFMT(LogIas, Log, "[{HostName}] Reconnected to primary host '{Url}'", GetPrimaryHostUrl(), GetName());
+			UE_LOGFMT(LogIas, Log, "[{HostName}] Reconnected to primary host '{Url}'", GetName(), GetPrimaryHostUrl());
+		}
+		else
+		{
+			UE_LOGFMT(LogIas, Log, "[{HostName}] Failed to reconnect to primary host, connection remains '{Url}'", GetName(), GetPrimaryHostUrl());
 		}
 	}
 
diff --git a/Engine/Source/Runtime/Experimental/IoStore/OnDemand/Private/IasHostGroup.h b/Engine/Source/Runtime/Experimental/IoStore/OnDemand/Private/IasHostGroup.h
--- a/Engine/Source/Runtime/Experimental/IoStore/OnDemand/Private/IasHostGroup.h
+++ b/Engine/Source/Runtime/Experimental/IoStore/OnDemand/Private/IasHostGroup.h
@@ -68,10 +68,11 @@ private:
 
 class FHostGroupManager final
 {
-	FHostGroupManager() =default;
+	FHostGroupManager() = default;
 	~FHostGroupManager() = default;
 
 public:
+	DECLARE_TS_MULTICAST_DELEGATE(FOnHostGroupDisconnectedDelegate);
 
 	static FHostGroupManager& Get(); 
 
@@ -89,10 +90,17 @@ public:
 	void DisconnectAll(); 
 	uint32 GetNumDisconnctedHosts() const;
 
+	FOnHostGroupDisconnectedDelegate& OnHostGroupDisconncted()
+	{
+		return HostGroupDisconnectedDelegate;
+	}
+
 private:
 
 	TArray<FIASHostGroup> HostGroups;
 
+	FOnHostGroupDisconnectedDelegate HostGroupDisconnectedDelegate;
+
 	mutable UE::FMutex Mutex;
 };
 
diff --git a/Engine/Source/Runtime/Experimental/IoStore/OnDemand/Private/OnDemandHttpThread.cpp b/Engine/Source/Runtime/Experimental/IoStore/OnDemand/Private/OnDemandHttpThread.cpp
--- a/Engine/Source/Runtime/Experimental/IoStore/OnDemand/Private/OnDemandHttpThread.cpp
+++ b/Engine/Source/Runtime/Experimental/IoStore/OnDemand/Private/OnDemandHttpThread.cpp
@@ -241,10 +241,21 @@ FOnDemandHttpThread::FOnDemandHttpThread()
 		bRecreateHttpClient = true;
 		TickThreadEvent->Trigger();
 	});
+
+	OnHostGroupDisconnectedHandle = FHostGroupManager::Get().OnHostGroupDisconncted().AddLambda([this]()
+	{
+		TickThreadEvent->Trigger();
+	});
 }
 
 FOnDemandHttpThread::~FOnDemandHttpThread()
 {
+	if (OnHostGroupDisconnectedHandle.IsValid())
+	{
+		FHostGroupManager::Get().OnHostGroupDisconncted().Remove(OnHostGroupDisconnectedHandle);
+		OnHostGroupDisconnectedHandle.Reset();
+	}
+
 	if (OnRecreateHttpClientHandle.IsValid())
 	{
 		OnRecreateHttpClient.Remove(OnRecreateHttpClientHandle);
diff --git a/Engine/Source/Runtime/Experimental/IoStore/OnDemand/Private/OnDemandHttpThread.h b/Engine/Source/Runtime/Experimental/IoStore/OnDemand/Private/OnDemandHttpThread.h
--- a/Engine/Source/Runtime/Experimental/IoStore/OnDemand/Private/OnDemandHttpThread.h
+++ b/Engine/Source/Runtime/Experimental/IoStore/OnDemand/Private/OnDemandHttpThread.h
@@ -51,7 +51,7 @@ class FOnDemandHttpThread final
 public:
 	using FCompletionCallback = TUniqueFunction<void(uint32 /*StatusCode*/ , FStringView/*ErrorReason*/, FIoBuffer&& /*Data*/)>;
 	using FRequestHandle = void*;
-	DECLARE_MULTICAST_DELEGATE(FOnTickIdle);
+	DECLARE_TS_MULTICAST_DELEGATE(FOnTickIdle);
 
 	FOnDemandHttpThread();
 	virtual ~FOnDemandHttpThread();
@@ -123,6 +123,8 @@ private:
 	TUniquePtr<FMultiEndpointHttpClient> HttpClient;
 
 	FDelegateHandle OnRecreateHttpClientHandle;
+	FDelegateHandle OnHostGroupDisconnectedHandle;
+
 	UE::FMutex RecreateMutex;
 
 	TThreadSafeIntrusiveQueue<FHttpRequest> HttpRequests;
