Cleaner activity by lifting loop-related members out.

#rnx

diff --git a/Engine/Source/Runtime/Experimental/IoStore/HttpClient/Private/Activity.inl b/Engine/Source/Runtime/Experimental/IoStore/HttpClient/Private/Activity.inl
--- a/Engine/Source/Runtime/Experimental/IoStore/HttpClient/Private/Activity.inl
+++ b/Engine/Source/Runtime/Experimental/IoStore/HttpClient/Private/Activity.inl
@@ -509,11 +509,10 @@ struct FActivity
 	{
 		FHost*			Host = nullptr;
 		char*			Buffer = nullptr;
+		uint32			ContentSizeEst = 0;
 		uint16			BufferSize = 0;
 		bool			bIsKeepAlive = false;
-		bool			bFollow30x = false;
 		bool			bAllowChunked = true;
-		uint32			ContentSizeEst = 0;
 	};
 
 						FActivity(const FParams& Params);
@@ -521,8 +520,6 @@ struct FActivity
 	void				SetSink(FTicketSink&& InSink, UPTRINT Param);
 	void				SetDestination(FIoBuffer* InDest);
 	FHost*				GetHost() const;
-	uint32				GetIndex() const;
-	uint64				GetTicket() const;
 	UPTRINT				GetSinkParam() const;
 	FIoBuffer&			GetContent() const;
 	const char*			GetErrorReason() const;
@@ -547,12 +544,11 @@ struct FActivity
 		_Num,
 	};
 
-	int8				Slot = -1;
 	EState				State = EState::None;
 	uint8				IsKeepAlive : 1;
 	uint8				NoContent : 1;
-	uint8				bFollow30x : 1;
 	uint8				bAllowChunked : 1;
+	uint8				_Unused : 1;
 	uint8				LengthScore : 4;
 	uint32				StateParam = 0;
 #if IAS_HTTP_WITH_PERF
@@ -571,7 +567,6 @@ struct FActivity
 ////////////////////////////////////////////////////////////////////////////////
 FActivity::FActivity(const FParams& Params)
 : IsKeepAlive(Params.bIsKeepAlive)
-, bFollow30x(Params.bFollow30x)
 , bAllowChunked(Params.bAllowChunked)
 , Host(Params.Host)
 , Buffer(Params.Buffer, Params.BufferSize)
@@ -603,18 +598,6 @@ FHost* FActivity::GetHost() const
 	return Host;
 }
 
-////////////////////////////////////////////////////////////////////////////////
-uint32 FActivity::GetIndex() const
-{
-	return Slot;
-}
-
-////////////////////////////////////////////////////////////////////////////////
-uint64 FActivity::GetTicket() const
-{
-	return 1ull << Slot;
-}
-
 ////////////////////////////////////////////////////////////////////////////////
 UPTRINT FActivity::GetSinkParam() const
 {
diff --git a/Engine/Source/Runtime/Experimental/IoStore/HttpClient/Private/Api.inl b/Engine/Source/Runtime/Experimental/IoStore/HttpClient/Private/Api.inl
--- a/Engine/Source/Runtime/Experimental/IoStore/HttpClient/Private/Api.inl
+++ b/Engine/Source/Runtime/Experimental/IoStore/HttpClient/Private/Api.inl
@@ -183,15 +183,15 @@ UPTRINT FTicketStatus::GetParam() const
 ////////////////////////////////////////////////////////////////////////////////
 FTicket FTicketStatus::GetTicket() const
 {
-	const auto* Activity = (FActivity*)this;
-	return Activity->GetTicket();
+	const auto* Activity = (FActivityNode*)this;
+	return 1ull << Activity->Slot;
 }
 
 ////////////////////////////////////////////////////////////////////////////////
 uint32 FTicketStatus::GetIndex() const
 {
-	const auto* Activity = (FActivity*)this;
-	return Activity->GetIndex();
+	const auto* Activity = (FActivityNode*)this;
+	return Activity->Slot;
 }
 
 ////////////////////////////////////////////////////////////////////////////////
diff --git a/Engine/Source/Runtime/Experimental/IoStore/HttpClient/Private/Loop.inl b/Engine/Source/Runtime/Experimental/IoStore/HttpClient/Private/Loop.inl
--- a/Engine/Source/Runtime/Experimental/IoStore/HttpClient/Private/Loop.inl
+++ b/Engine/Source/Runtime/Experimental/IoStore/HttpClient/Private/Loop.inl
@@ -500,15 +500,30 @@ using FActivityList = struct FActivityNode*;
 struct FActivityNode
 	: public FActivity
 {
+	struct FParams
+		: public FActivity::FParams
+	{
+		bool				bFollow30x = false;
+	};
+
 	static FActivityNode*	Create(FParams& Params, FAnsiStringView* Url=nullptr, FCertRootsRef VerifyCert={});
 	static void				Destroy(FActivityNode* Activity);
 	FActivityNode*			Next = nullptr;
+	int8					Slot = -1;
+	bool					bFollow30x;
 
 private:
-	using					FActivity::FActivity;
+							FActivityNode(const FParams& Params);
 							~FActivityNode();
 };
 
+////////////////////////////////////////////////////////////////////////////////
+FActivityNode::FActivityNode(const FParams& Params)
+: FActivity(Params)
+, bFollow30x(Params.bFollow30x)
+{
+}
+
 ////////////////////////////////////////////////////////////////////////////////
 FActivityNode* FActivityNode::Create(FParams& Params, FAnsiStringView* InOutUrl, FCertRootsRef VerifyCert)
 {
