Prevent potential stalls on the GameThread when trying to reconnect IAX CDN

#rnx

- If a hostgroup needs to try to reconnect then the mutex lock around FHostGroupManager::HostGroups can be kept for a fairly long time (the latency test will time out after about 4 seconds which is the worst case scenario)
- The debug display for IAS can also try to access FHostGroupManager::HostGroups when trying to output debug info about each hostgroup. This can potentially stall for some time.
- As a quick fix, we can swap out the UE::Mutex for a FRWLock allowing both the reconnect and debug rendering to access the array at the same time while being protected from the array being modified.

diff --git a/Engine/Source/Runtime/Experimental/IoStore/OnDemand/Private/IasHostGroup.cpp b/Engine/Source/Runtime/Experimental/IoStore/OnDemand/Private/IasHostGroup.cpp
--- a/Engine/Source/Runtime/Experimental/IoStore/OnDemand/Private/IasHostGroup.cpp
+++ b/Engine/Source/Runtime/Experimental/IoStore/OnDemand/Private/IasHostGroup.cpp
@@ -3,12 +3,12 @@
 #include "IasHostGroup.h"
 
 #include "Algo/Find.h"
-#include "Async/UniqueLock.h"
 #include "HAL/IConsoleManager.h"
 #include "IO/IoStoreOnDemand.h"
 #include "IO/OnDemandHostGroup.h"
 #include "LatencyTesting.h"
 #include "Logging/StructuredLog.h"
+#include "Misc/ScopeRWLock.h"
 #include "Statistics.h"
 
 namespace UE::IoStore
@@ -314,7 +314,7 @@ TIoStatusOr<FIASHostGroup> FHostGroupManager::Register(FName Name, FAnsiStringVi
 	FIASHostGroup HostGroup(Name, TestPath);
 
 	{
-		UE::TUniqueLock _(Mutex);
+		UE::TWriteScopeLock _(Mutex);
 		HostGroups.Add(HostGroup);
 	}
 
@@ -327,7 +327,7 @@ TIoStatusOr<FIASHostGroup> FHostGroupManager::Register(FName Name, TConstArrayVi
 
 	if (Result.IsOk())
 	{
-		UE::TUniqueLock _(Mutex);
+		UE::TWriteScopeLock _(Mutex);
 		HostGroups.Add(Result.ValueOrDie());
 	}
 
@@ -336,7 +336,7 @@ TIoStatusOr<FIASHostGroup> FHostGroupManager::Register(FName Name, TConstArrayVi
 
 FIASHostGroup FHostGroupManager::Find(FName Name)
 {
-	UE::TUniqueLock _(Mutex);
+	UE::TReadScopeLock _(Mutex);
 
 	FIASHostGroup* FoundHostGroup = Algo::FindByPredicate(HostGroups, [Name](const FIASHostGroup& HostGroup)->bool
 	{
@@ -355,7 +355,7 @@ FIASHostGroup FHostGroupManager::Find(FName Name)
 
 void FHostGroupManager::ForEachHostGroup(TFunctionRef<void(const FIASHostGroup&)> Callback) const
 {
-	UE::TUniqueLock _(Mutex);
+	UE::TReadScopeLock _(Mutex);
 
 	for (const FIASHostGroup& Host : HostGroups)
 	{
@@ -367,7 +367,7 @@ void FHostGroupManager::Tick(uint32 TimeoutMs, std::atomic_bool& CancellationTok
 {
 	TRACE_CPUPROFILER_EVENT_SCOPE(FHostGroupManager::Tick);
 
-	UE::TUniqueLock _(Mutex);
+	UE::TReadScopeLock _(Mutex);
 
 	for (FIASHostGroup& HostGroup : HostGroups)
 	{
@@ -380,7 +380,7 @@ void FHostGroupManager::Tick(uint32 TimeoutMs, std::atomic_bool& CancellationTok
 
 void FHostGroupManager::DisconnectAll()
 {
-	UE::TUniqueLock _(Mutex);
+	UE::TReadScopeLock _(Mutex);
 	for (FIASHostGroup& HostGroup : HostGroups)
 	{
 		HostGroup.Disconnect();
@@ -391,7 +391,7 @@ uint32 FHostGroupManager::GetNumDisconnctedHosts() const
 {
 	uint32 NumDisconnectedHosts = 0;
 
-	UE::TUniqueLock _(Mutex);
+	UE::TReadScopeLock _(Mutex);
 	for (const FIASHostGroup& HostGroup : HostGroups)
 	{
 		NumDisconnectedHosts += HostGroup.IsConnected() == false ? 1 : 0;
diff --git a/Engine/Source/Runtime/Experimental/IoStore/OnDemand/Private/IasHostGroup.h b/Engine/Source/Runtime/Experimental/IoStore/OnDemand/Private/IasHostGroup.h
--- a/Engine/Source/Runtime/Experimental/IoStore/OnDemand/Private/IasHostGroup.h
+++ b/Engine/Source/Runtime/Experimental/IoStore/OnDemand/Private/IasHostGroup.h
@@ -2,7 +2,7 @@
 
 #pragma once
 
-#include "Async/Mutex.h"
+#include "HAL/CriticalSection.h"
 #include "IO/IoStatus.h"
 #include "OnDemandBackendStatus.h"
 #include "Templates/Function.h"
@@ -101,7 +101,7 @@ private:
 
 	FOnHostGroupDisconnectedDelegate HostGroupDisconnectedDelegate;
 
-	mutable UE::FMutex Mutex;
+	mutable FRWLock Mutex;
 };
 
 } //namespace UE::IoStore
