Iris
* Prevent crash when an ensure triggers when setting an invalid EReplicatedObjectState::PendingDestroy state

#rnx

diff --git a/Engine/Source/Runtime/Experimental/Iris/Core/Private/Iris/ReplicationSystem/NetRefHandleManager.cpp b/Engine/Source/Runtime/Experimental/Iris/Core/Private/Iris/ReplicationSystem/NetRefHandleManager.cpp
--- a/Engine/Source/Runtime/Experimental/Iris/Core/Private/Iris/ReplicationSystem/NetRefHandleManager.cpp
+++ b/Engine/Source/Runtime/Experimental/Iris/Core/Private/Iris/ReplicationSystem/NetRefHandleManager.cpp
@@ -1153,6 +1153,13 @@ FString FNetRefHandleManager::PrintObjectFromIndex(FInternalNetRefIndex ObjectIn
 {
 	if (ObjectIndex != InvalidInternalIndex)
 	{
+		const bool bIsIndexAssigned = AssignedInternalIndices.GetBit(ObjectIndex);
+		if (!bIsIndexAssigned)
+		{
+			// If the index isn't bound to any object (possibly invalid or obsolete)
+			return FString::Printf(TEXT("UnassignedObject (InternalIndex: %u)"), ObjectIndex);
+		}
+
 		const FNetRefHandle NetRefHandle = GetNetRefHandleFromInternalIndex(ObjectIndex);
 		const FReplicatedObjectData& ObjectData = GetReplicatedObjectDataNoCheck(ObjectIndex);
 
diff --git a/Engine/Source/Runtime/Experimental/Iris/Core/Private/Iris/ReplicationSystem/ReplicationWriter.cpp b/Engine/Source/Runtime/Experimental/Iris/Core/Private/Iris/ReplicationSystem/ReplicationWriter.cpp
--- a/Engine/Source/Runtime/Experimental/Iris/Core/Private/Iris/ReplicationSystem/ReplicationWriter.cpp
+++ b/Engine/Source/Runtime/Experimental/Iris/Core/Private/Iris/ReplicationSystem/ReplicationWriter.cpp
@@ -385,7 +385,8 @@ void FReplicationWriter::SetState(uint32 InternalIndex, EReplicatedObjectState N
 		break;
 		case EReplicatedObjectState::PendingDestroy:
 		{
-			ensureMsgf(CurrentState != EReplicatedObjectState::Invalid, TEXT("Trying to set state %s when state is %s. Object: %s IsDestructionInfo: %" UINT64_FMT " IsSubObject: %" UINT64_FMT), LexToString(NewState), LexToString(CurrentState), ToCStr(NetRefHandleManager->PrintObjectFromIndex(InternalIndex)), Info.IsDestructionInfo, Info.IsSubObject);
+			ensureMsgf(CurrentState != EReplicatedObjectState::Invalid, TEXT("Trying to set state %s when state is %s. (InternalIndex: %u) IsDestructionInfo: %" UINT64_FMT " IsSubObject: %" UINT64_FMT), 
+					   LexToString(NewState), LexToString(CurrentState), InternalIndex, Info.IsDestructionInfo, Info.IsSubObject);
 		}
 		break;
 		case EReplicatedObjectState::WaitOnDestroyConfirmation:
