[ChaosVD] Removing inlining of data wrapper strings.

diff --git a/Engine/Plugins/ChaosVD/Source/ChaosVD/Private/Widgets/SChaosVDMainTab.cpp b/Engine/Plugins/ChaosVD/Source/ChaosVD/Private/Widgets/SChaosVDMainTab.cpp
--- a/Engine/Plugins/ChaosVD/Source/ChaosVD/Private/Widgets/SChaosVDMainTab.cpp
+++ b/Engine/Plugins/ChaosVD/Source/ChaosVD/Private/Widgets/SChaosVDMainTab.cpp
@@ -77,6 +77,8 @@
 
 #define LOCTEXT_NAMESPACE "ChaosVisualDebugger"
 
+const FName SChaosVDMainTab::MainToolBarName = FName("ChaosVD.MainToolBar");
+
 void SChaosVDMainTab::Construct(const FArguments& InArgs, TSharedPtr<FChaosVDEngine> InChaosVDEngine)
 {
 	ChaosVDEngine = InChaosVDEngine;
diff --git a/Engine/Plugins/ChaosVD/Source/ChaosVD/Public/Widgets/SChaosVDMainTab.h b/Engine/Plugins/ChaosVD/Source/ChaosVD/Public/Widgets/SChaosVDMainTab.h
--- a/Engine/Plugins/ChaosVD/Source/ChaosVD/Public/Widgets/SChaosVDMainTab.h
+++ b/Engine/Plugins/ChaosVD/Source/ChaosVD/Public/Widgets/SChaosVDMainTab.h
@@ -239,7 +239,7 @@ private:
 	/** Commandlist used for any UI actions that need to be processed globally regardless on which specific widget we are in */
 	TSharedPtr<FUICommandList> GlobalCommandList;
 
-	static inline const FName MainToolBarName = FName("ChaosVD.MainToolBar");
+	static const FName MainToolBarName;
 };
 
 template <typename TabType>
diff --git a/Engine/Plugins/Experimental/Mover/Source/MoverCVDData/Private/MoverCVDDataWrappers.cpp b/Engine/Plugins/Experimental/Mover/Source/MoverCVDData/Private/MoverCVDDataWrappers.cpp
--- a/Engine/Plugins/Experimental/Mover/Source/MoverCVDData/Private/MoverCVDDataWrappers.cpp
+++ b/Engine/Plugins/Experimental/Mover/Source/MoverCVDData/Private/MoverCVDDataWrappers.cpp
@@ -4,6 +4,8 @@
 
 #include UE_INLINE_GENERATED_CPP_BY_NAME(MoverCVDDataWrappers)
 
+FStringView FMoverCVDSimDataWrapper::WrapperTypeName = TEXT("FMoverCVDSimDataWrapper");
+
 bool FMoverCVDSimDataWrapper::Serialize(FArchive& Ar)
 {
 	Ar << bHasValidData;
diff --git a/Engine/Plugins/Experimental/Mover/Source/MoverCVDData/Public/MoverCVDDataWrappers.h b/Engine/Plugins/Experimental/Mover/Source/MoverCVDData/Public/MoverCVDDataWrappers.h
--- a/Engine/Plugins/Experimental/Mover/Source/MoverCVDData/Public/MoverCVDDataWrappers.h
+++ b/Engine/Plugins/Experimental/Mover/Source/MoverCVDData/Public/MoverCVDDataWrappers.h
@@ -13,7 +13,7 @@ struct FMoverCVDSimDataWrapper : public FChaosVDWrapperDataBase
 {
 	GENERATED_BODY()
 	
-	inline static FStringView WrapperTypeName = TEXT("FMoverCVDSimDataWrapper");
+	MOVERCVDDATA_API static FStringView WrapperTypeName;
 
 	UPROPERTY(VisibleAnywhere, Category="Mover Info")
 	int32 SolverID = INDEX_NONE;
diff --git a/Engine/Source/Runtime/Experimental/Chaos/Private/ChaosVisualDebugger/ChaosVDMemWriterReader.cpp b/Engine/Source/Runtime/Experimental/Chaos/Private/ChaosVisualDebugger/ChaosVDMemWriterReader.cpp
--- a/Engine/Source/Runtime/Experimental/Chaos/Private/ChaosVisualDebugger/ChaosVDMemWriterReader.cpp
+++ b/Engine/Source/Runtime/Experimental/Chaos/Private/ChaosVisualDebugger/ChaosVDMemWriterReader.cpp
@@ -4,6 +4,8 @@
 
 namespace Chaos::VisualDebugger
 {
+	FStringView FChaosVDArchiveHeader::WrapperTypeName = TEXT("FChaosVDArchiveHeader");
+
 	bool FChaosVDArchiveHeader::Serialize(FArchive& Ar)
 	{
 		Ar << EngineVersion;
diff --git a/Engine/Source/Runtime/Experimental/Chaos/Public/ChaosVisualDebugger/ChaosVDMemWriterReader.h b/Engine/Source/Runtime/Experimental/Chaos/Public/ChaosVisualDebugger/ChaosVDMemWriterReader.h
--- a/Engine/Source/Runtime/Experimental/Chaos/Public/ChaosVisualDebugger/ChaosVDMemWriterReader.h
+++ b/Engine/Source/Runtime/Experimental/Chaos/Public/ChaosVisualDebugger/ChaosVDMemWriterReader.h
@@ -16,7 +16,7 @@ namespace Chaos::VisualDebugger
 {
 	struct FChaosVDArchiveHeader
 	{
-		inline static FStringView WrapperTypeName = TEXT("FChaosVDArchiveHeader");
+		CHAOS_API static FStringView WrapperTypeName;
 
 		FChaosVDArchiveHeader()
 		{
diff --git a/Engine/Source/Runtime/Experimental/ChaosVisualDebugger/Private/DataWrappers/ChaosVDAccelerationStructureDataWrappers.cpp b/Engine/Source/Runtime/Experimental/ChaosVisualDebugger/Private/DataWrappers/ChaosVDAccelerationStructureDataWrappers.cpp
--- a/Engine/Source/Runtime/Experimental/ChaosVisualDebugger/Private/DataWrappers/ChaosVDAccelerationStructureDataWrappers.cpp
+++ b/Engine/Source/Runtime/Experimental/ChaosVisualDebugger/Private/DataWrappers/ChaosVDAccelerationStructureDataWrappers.cpp
@@ -8,6 +8,7 @@
 
 #include UE_INLINE_GENERATED_CPP_BY_NAME(ChaosVDAccelerationStructureDataWrappers)
 
+FStringView FChaosVDAABBTreeDataWrapper::WrapperTypeName = TEXT("FChaosVDAABBTreeDataWrapper");
 
 bool FChaosVDBVCellElementDataWrapper::Serialize(FArchive& Ar)
 {
diff --git a/Engine/Source/Runtime/Experimental/ChaosVisualDebugger/Private/DataWrappers/ChaosVDCharacterGroundConstraintDataWrappers.cpp b/Engine/Source/Runtime/Experimental/ChaosVisualDebugger/Private/DataWrappers/ChaosVDCharacterGroundConstraintDataWrappers.cpp
--- a/Engine/Source/Runtime/Experimental/ChaosVisualDebugger/Private/DataWrappers/ChaosVDCharacterGroundConstraintDataWrappers.cpp
+++ b/Engine/Source/Runtime/Experimental/ChaosVisualDebugger/Private/DataWrappers/ChaosVDCharacterGroundConstraintDataWrappers.cpp
@@ -6,6 +6,8 @@
 
 #include UE_INLINE_GENERATED_CPP_BY_NAME(ChaosVDCharacterGroundConstraintDataWrappers)
 
+FStringView FChaosVDCharacterGroundConstraint::WrapperTypeName = TEXT("FChaosVDCharacterGroundConstraint");
+
 bool FChaosVDCharacterGroundConstraintStateDataWrapper::Serialize(FArchive& Ar)
 {
 	Ar << bHasValidData;
diff --git a/Engine/Source/Runtime/Experimental/ChaosVisualDebugger/Private/DataWrappers/ChaosVDCollisionDataWrappers.cpp b/Engine/Source/Runtime/Experimental/ChaosVisualDebugger/Private/DataWrappers/ChaosVDCollisionDataWrappers.cpp
--- a/Engine/Source/Runtime/Experimental/ChaosVisualDebugger/Private/DataWrappers/ChaosVDCollisionDataWrappers.cpp
+++ b/Engine/Source/Runtime/Experimental/ChaosVisualDebugger/Private/DataWrappers/ChaosVDCollisionDataWrappers.cpp
@@ -6,6 +6,9 @@
 
 #include UE_INLINE_GENERATED_CPP_BY_NAME(ChaosVDCollisionDataWrappers)
 
+FStringView FChaosVDConstraint::WrapperTypeName = TEXT("FChaosVDConstraint");
+FStringView FChaosVDParticlePairMidPhase::WrapperTypeName = TEXT("FChaosVDParticlePairMidPhase");
+FStringView FChaosVDCollisionChannelsInfoContainer::WrapperTypeName = TEXT("FChaosVDCollisionChannelsInfoContainer");
 
 bool FChaosVDContactPoint::Serialize(FArchive& Ar)
 {
diff --git a/Engine/Source/Runtime/Experimental/ChaosVisualDebugger/Private/DataWrappers/ChaosVDDebugShapeDataWrapper.cpp b/Engine/Source/Runtime/Experimental/ChaosVisualDebugger/Private/DataWrappers/ChaosVDDebugShapeDataWrapper.cpp
--- a/Engine/Source/Runtime/Experimental/ChaosVisualDebugger/Private/DataWrappers/ChaosVDDebugShapeDataWrapper.cpp
+++ b/Engine/Source/Runtime/Experimental/ChaosVisualDebugger/Private/DataWrappers/ChaosVDDebugShapeDataWrapper.cpp
@@ -6,6 +6,11 @@
 
 #include UE_INLINE_GENERATED_CPP_BY_NAME(ChaosVDDebugShapeDataWrapper)
 
+FStringView FChaosVDDebugDrawSphereDataWrapper::WrapperTypeName = TEXT("FChaosVDDebugDrawSphereDataWrapper");
+FStringView FChaosVDDebugDrawLineDataWrapper::WrapperTypeName = TEXT("FChaosVDDebugDrawLineDataWrapper");
+FStringView FChaosVDDebugDrawImplicitObjectDataWrapper::WrapperTypeName = TEXT("FChaosVDDebugDrawImplicitObjectDataWrapper");
+FStringView FChaosVDDebugDrawBoxDataWrapper::WrapperTypeName = TEXT("FChaosVDDebugDrawBoxDataWrapper");
+
 void FChaosVDDebugDrawShapeBase::SerializeBase_Internal(FArchive& Ar)
 {
 	Ar << SolverID;
diff --git a/Engine/Source/Runtime/Experimental/ChaosVisualDebugger/Private/DataWrappers/ChaosVDJointDataWrappers.cpp b/Engine/Source/Runtime/Experimental/ChaosVisualDebugger/Private/DataWrappers/ChaosVDJointDataWrappers.cpp
--- a/Engine/Source/Runtime/Experimental/ChaosVisualDebugger/Private/DataWrappers/ChaosVDJointDataWrappers.cpp
+++ b/Engine/Source/Runtime/Experimental/ChaosVisualDebugger/Private/DataWrappers/ChaosVDJointDataWrappers.cpp
@@ -7,6 +7,7 @@
 
 #include UE_INLINE_GENERATED_CPP_BY_NAME(ChaosVDJointDataWrappers)
 
+FStringView FChaosVDJointConstraint::WrapperTypeName = TEXT("FChaosVDJointConstraint");
 
 bool FChaosVDJointStateDataWrapper::Serialize(FArchive& Ar)
 {
diff --git a/Engine/Source/Runtime/Experimental/ChaosVisualDebugger/Private/DataWrappers/ChaosVDParticleDataWrapper.cpp b/Engine/Source/Runtime/Experimental/ChaosVisualDebugger/Private/DataWrappers/ChaosVDParticleDataWrapper.cpp
--- a/Engine/Source/Runtime/Experimental/ChaosVisualDebugger/Private/DataWrappers/ChaosVDParticleDataWrapper.cpp
+++ b/Engine/Source/Runtime/Experimental/ChaosVisualDebugger/Private/DataWrappers/ChaosVDParticleDataWrapper.cpp
@@ -12,6 +12,7 @@
 
 #include UE_INLINE_GENERATED_CPP_BY_NAME(ChaosVDParticleDataWrapper)
 
+FStringView FChaosVDParticleDataWrapper::WrapperTypeName = TEXT("FChaosVDParticleDataWrapper");
 FStringView FChaosVDParticleMetadata::WrapperTypeName = TEXT("FChaosVDParticleMetadata");
 
 bool FChaosVDParticleMetadata::Serialize(FArchive& Ar)
diff --git a/Engine/Source/Runtime/Experimental/ChaosVisualDebugger/Private/DataWrappers/ChaosVDQueryDataWrappers.cpp b/Engine/Source/Runtime/Experimental/ChaosVisualDebugger/Private/DataWrappers/ChaosVDQueryDataWrappers.cpp
--- a/Engine/Source/Runtime/Experimental/ChaosVisualDebugger/Private/DataWrappers/ChaosVDQueryDataWrappers.cpp
+++ b/Engine/Source/Runtime/Experimental/ChaosVisualDebugger/Private/DataWrappers/ChaosVDQueryDataWrappers.cpp
@@ -7,6 +7,8 @@
 
 #include UE_INLINE_GENERATED_CPP_BY_NAME(ChaosVDQueryDataWrappers)
 
+FStringView FChaosVDQueryVisitStep::WrapperTypeName = TEXT("FChaosVDQueryVisitStep");
+FStringView FChaosVDQueryDataWrapper::WrapperTypeName = TEXT("FChaosVDQueryDataWrapper");
 
 bool FChaosVDCollisionResponseParams::Serialize(FArchive& Ar)
 {
diff --git a/Engine/Source/Runtime/Experimental/ChaosVisualDebugger/Public/DataWrappers/ChaosVDAccelerationStructureDataWrappers.h b/Engine/Source/Runtime/Experimental/ChaosVisualDebugger/Public/DataWrappers/ChaosVDAccelerationStructureDataWrappers.h
--- a/Engine/Source/Runtime/Experimental/ChaosVisualDebugger/Public/DataWrappers/ChaosVDAccelerationStructureDataWrappers.h
+++ b/Engine/Source/Runtime/Experimental/ChaosVisualDebugger/Public/DataWrappers/ChaosVDAccelerationStructureDataWrappers.h
@@ -160,7 +160,7 @@ struct FChaosVDAABBTreeDataWrapper : public FChaosVDAccelerationStructureBase
 {
 	GENERATED_BODY()
 
-	inline static FStringView WrapperTypeName = TEXT("FChaosVDAABBTreeDataWrapper");
+	CHAOSVDRUNTIME_API static FStringView WrapperTypeName;
 
 	UPROPERTY(VisibleAnywhere, Category="Tree Data")
 	int32 RootNodeIndex = INDEX_NONE;
diff --git a/Engine/Source/Runtime/Experimental/ChaosVisualDebugger/Public/DataWrappers/ChaosVDCharacterGroundConstraintDataWrappers.h b/Engine/Source/Runtime/Experimental/ChaosVisualDebugger/Public/DataWrappers/ChaosVDCharacterGroundConstraintDataWrappers.h
--- a/Engine/Source/Runtime/Experimental/ChaosVisualDebugger/Public/DataWrappers/ChaosVDCharacterGroundConstraintDataWrappers.h
+++ b/Engine/Source/Runtime/Experimental/ChaosVisualDebugger/Public/DataWrappers/ChaosVDCharacterGroundConstraintDataWrappers.h
@@ -120,7 +120,7 @@ struct FChaosVDCharacterGroundConstraint : public FChaosVDConstraintDataWrapperB
 	GENERATED_BODY()
 public:
 
-	inline static FStringView WrapperTypeName = TEXT("FChaosVDCharacterGroundConstraint");
+	CHAOSVDRUNTIME_API static FStringView WrapperTypeName;
 
 	CHAOSVDRUNTIME_API bool Serialize(FArchive& Ar);
 
diff --git a/Engine/Source/Runtime/Experimental/ChaosVisualDebugger/Public/DataWrappers/ChaosVDCollisionDataWrappers.h b/Engine/Source/Runtime/Experimental/ChaosVisualDebugger/Public/DataWrappers/ChaosVDCollisionDataWrappers.h
--- a/Engine/Source/Runtime/Experimental/ChaosVisualDebugger/Public/DataWrappers/ChaosVDCollisionDataWrappers.h
+++ b/Engine/Source/Runtime/Experimental/ChaosVisualDebugger/Public/DataWrappers/ChaosVDCollisionDataWrappers.h
@@ -210,7 +210,7 @@ struct FChaosVDConstraint
 {
 	GENERATED_BODY()
 
-	inline static FStringView WrapperTypeName = TEXT("FChaosVDConstraint");
+	CHAOSVDRUNTIME_API static FStringView WrapperTypeName;
 
 	UPROPERTY()
 	uint8 bIsCurrent:1 = false;
@@ -343,7 +343,7 @@ struct FChaosVDParticlePairMidPhase
 {
 	GENERATED_BODY()
 
-	inline static FStringView WrapperTypeName = TEXT("FChaosVDParticlePairMidPhase");
+	CHAOSVDRUNTIME_API static FStringView WrapperTypeName;
 
 	UPROPERTY()
 	int32 SolverID = INDEX_NONE;
@@ -497,7 +497,7 @@ struct FChaosVDCollisionChannelsInfoContainer
 {
 	GENERATED_BODY()
 
-	inline static FStringView WrapperTypeName = TEXT("FChaosVDCollisionChannelsInfoContainer");
+	CHAOSVDRUNTIME_API static FStringView WrapperTypeName;
 
 	UPROPERTY(VisibleAnywhere, Category=CollisionData)
 	FChaosVDCollisionChannelInfo CustomChannelsNames[32] = {};
diff --git a/Engine/Source/Runtime/Experimental/ChaosVisualDebugger/Public/DataWrappers/ChaosVDDebugShapeDataWrapper.h b/Engine/Source/Runtime/Experimental/ChaosVisualDebugger/Public/DataWrappers/ChaosVDDebugShapeDataWrapper.h
--- a/Engine/Source/Runtime/Experimental/ChaosVisualDebugger/Public/DataWrappers/ChaosVDDebugShapeDataWrapper.h
+++ b/Engine/Source/Runtime/Experimental/ChaosVisualDebugger/Public/DataWrappers/ChaosVDDebugShapeDataWrapper.h
@@ -51,7 +51,7 @@ struct FChaosVDDebugDrawBoxDataWrapper : public FChaosVDDebugDrawShapeBase
 {
 	GENERATED_BODY()
 
-	inline static FStringView WrapperTypeName = TEXT("FChaosVDDDebugDrawBoxDataWrapper");
+	CHAOSVDRUNTIME_API static FStringView WrapperTypeName;
 
 	UPROPERTY()
 	FBox Box = FBox(ForceInitToZero);
@@ -66,7 +66,7 @@ struct FChaosVDDebugDrawSphereDataWrapper : public FChaosVDDebugDrawShapeBase
 {
 	GENERATED_BODY()
 	
-	inline static FStringView WrapperTypeName = TEXT("FChaosVDDebugDrawSphereDataWrapper");
+	CHAOSVDRUNTIME_API static FStringView WrapperTypeName;
 
 	UPROPERTY()
 	FVector Origin = FVector::ZeroVector;
@@ -84,7 +84,7 @@ struct FChaosVDDebugDrawLineDataWrapper : public FChaosVDDebugDrawShapeBase
 {
 	GENERATED_BODY()
 	
-	inline static FStringView WrapperTypeName = TEXT("FChaosVDDebugDrawLineDataWrapper");
+	CHAOSVDRUNTIME_API static FStringView WrapperTypeName;
 
 	UPROPERTY()
 	FVector StartLocation = FVector::ZeroVector;
@@ -105,7 +105,7 @@ struct FChaosVDDebugDrawImplicitObjectDataWrapper : public FChaosVDDebugDrawShap
 {
 	GENERATED_BODY()
 	
-	inline static FStringView WrapperTypeName = TEXT("FChaosVDDebugDrawImplicitObjectDataWrapper");
+	CHAOSVDRUNTIME_API static FStringView WrapperTypeName;
 
 	uint32 ImplicitObjectHash = 0;
 
diff --git a/Engine/Source/Runtime/Experimental/ChaosVisualDebugger/Public/DataWrappers/ChaosVDJointDataWrappers.h b/Engine/Source/Runtime/Experimental/ChaosVisualDebugger/Public/DataWrappers/ChaosVDJointDataWrappers.h
--- a/Engine/Source/Runtime/Experimental/ChaosVisualDebugger/Public/DataWrappers/ChaosVDJointDataWrappers.h
+++ b/Engine/Source/Runtime/Experimental/ChaosVisualDebugger/Public/DataWrappers/ChaosVDJointDataWrappers.h
@@ -405,7 +405,7 @@ struct FChaosVDJointConstraint : public FChaosVDConstraintDataWrapperBase
 {
 	GENERATED_BODY()
 
-	inline static FStringView WrapperTypeName = TEXT("FChaosVDJointConstraint");
+	CHAOSVDRUNTIME_API static FStringView WrapperTypeName;
 
 	CHAOSVDRUNTIME_API bool Serialize(FArchive& Ar);
 
diff --git a/Engine/Source/Runtime/Experimental/ChaosVisualDebugger/Public/DataWrappers/ChaosVDParticleDataWrapper.h b/Engine/Source/Runtime/Experimental/ChaosVisualDebugger/Public/DataWrappers/ChaosVDParticleDataWrapper.h
--- a/Engine/Source/Runtime/Experimental/ChaosVisualDebugger/Public/DataWrappers/ChaosVDParticleDataWrapper.h
+++ b/Engine/Source/Runtime/Experimental/ChaosVisualDebugger/Public/DataWrappers/ChaosVDParticleDataWrapper.h
@@ -758,7 +758,7 @@ PRAGMA_DISABLE_DEPRECATION_WARNINGS
 
 PRAGMA_ENABLE_DEPRECATION_WARNINGS
 
-	inline static FStringView WrapperTypeName = TEXT("FChaosVDParticleDataWrapper");
+	CHAOSVDRUNTIME_API static FStringView WrapperTypeName;
 
 	GENERATED_BODY()
 
diff --git a/Engine/Source/Runtime/Experimental/ChaosVisualDebugger/Public/DataWrappers/ChaosVDQueryDataWrappers.h b/Engine/Source/Runtime/Experimental/ChaosVisualDebugger/Public/DataWrappers/ChaosVDQueryDataWrappers.h
--- a/Engine/Source/Runtime/Experimental/ChaosVisualDebugger/Public/DataWrappers/ChaosVDQueryDataWrappers.h
+++ b/Engine/Source/Runtime/Experimental/ChaosVisualDebugger/Public/DataWrappers/ChaosVDQueryDataWrappers.h
@@ -365,8 +365,8 @@ USTRUCT()
 struct FChaosVDQueryVisitStep : public FChaosVDWrapperDataBase
 {
 	GENERATED_BODY()
-	
-	inline static FStringView WrapperTypeName = TEXT("FChaosVDQueryVisitStep");
+
+	CHAOSVDRUNTIME_API static FStringView WrapperTypeName;
 
 	CHAOSVDRUNTIME_API bool Serialize(FArchive& Ar);
 
@@ -422,7 +422,7 @@ struct FChaosVDQueryDataWrapper
 {
 	GENERATED_BODY()
 
-	inline static FStringView WrapperTypeName = TEXT("FChaosVDQueryDataWrapper");
+	CHAOSVDRUNTIME_API static FStringView WrapperTypeName;
 
 	CHAOSVDRUNTIME_API bool Serialize(FArchive& Ar);
 
