Removed two state-based checks that are now covered via other means

#rnx

diff --git a/Engine/Source/Runtime/Experimental/IoStore/HttpClient/Private/Api.inl b/Engine/Source/Runtime/Experimental/IoStore/HttpClient/Private/Api.inl
--- a/Engine/Source/Runtime/Experimental/IoStore/HttpClient/Private/Api.inl
+++ b/Engine/Source/Runtime/Experimental/IoStore/HttpClient/Private/Api.inl
@@ -51,7 +51,6 @@ FRequest&& FRequest::Accept(FAnsiStringView MimeType)
 ////////////////////////////////////////////////////////////////////////////////
 FRequest&& FRequest::Header(FAnsiStringView Key, FAnsiStringView Value)
 {
-	checkSlow(Ptr->State == FActivity::EState::Build);
 	Ptr->Buffer.AddHeader(Key, Value);
 	return MoveTemp(*this);
 }
diff --git a/Engine/Source/Runtime/Experimental/IoStore/HttpClient/Private/Loop.inl b/Engine/Source/Runtime/Experimental/IoStore/HttpClient/Private/Loop.inl
--- a/Engine/Source/Runtime/Experimental/IoStore/HttpClient/Private/Loop.inl
+++ b/Engine/Source/Runtime/Experimental/IoStore/HttpClient/Private/Loop.inl
@@ -1625,7 +1625,6 @@ bool FEventLoop::Redirect(const FTicketStatus& Status, FTicketSink& OuterSink)
 	}
 
 	// Transfer original request headers
-	check(Activity.State == FActivity::EState::RecvMessage);
 	Activity.Buffer.EnumerateHeaders([&ForwardRequest] (FAnsiStringView Name, FAnsiStringView Value)
 	{
 		if (Name != "Host" && Name != "Connection")
