Another step forward encapsulating activities from frontend

#rnx

diff --git a/Engine/Source/Runtime/Experimental/IoStore/HttpClient/Private/Activity.inl b/Engine/Source/Runtime/Experimental/IoStore/HttpClient/Private/Activity.inl
--- a/Engine/Source/Runtime/Experimental/IoStore/HttpClient/Private/Activity.inl
+++ b/Engine/Source/Runtime/Experimental/IoStore/HttpClient/Private/Activity.inl
@@ -522,6 +522,7 @@ struct FActivity
 	FHost*				GetHost() const;
 	UPTRINT				GetSinkParam() const;
 	FIoBuffer&			GetContent() const;
+	uint32				GetRemainingKiB() const;
 	const char*			GetErrorReason() const;
 	uint32				GetErrorCode() const;
 
@@ -611,6 +612,17 @@ FIoBuffer& FActivity::GetContent() const
 	return *Dest;
 }
 
+////////////////////////////////////////////////////////////////////////////////
+uint32 FActivity::GetRemainingKiB() const
+{
+	if (State <= FActivity::EState::RecvStream)  return MAX_uint32;
+	if (State >  FActivity::EState::RecvContent) return 0;
+
+	uint32 ContentLength = uint32(Buffer.GetContentLength());
+	check(StateParam <= ContentLength);
+	return (ContentLength - StateParam) >> 10;
+}
+
 ////////////////////////////////////////////////////////////////////////////////
 const char* FActivity::GetErrorReason() const
 {
@@ -666,17 +678,6 @@ static int32 Activity_Rewind(FActivity* Activity)
 	return -1;
 }
 
-////////////////////////////////////////////////////////////////////////////////
-static uint32 Activity_RemainingKiB(FActivity* Activity)
-{
-	if (Activity->State <= FActivity::EState::RecvStream)  return MAX_uint32;
-	if (Activity->State >  FActivity::EState::RecvContent) return 0;
-
-	uint32 ContentLength = uint32(Activity->Buffer.GetContentLength());
-	check(Activity->StateParam <= ContentLength);
-	return (ContentLength - Activity->StateParam) >> 10;
-}
-
 ////////////////////////////////////////////////////////////////////////////////
 static void Activity_CallSink(FActivity* Activity)
 {
diff --git a/Engine/Source/Runtime/Experimental/IoStore/HttpClient/Private/Loop.inl b/Engine/Source/Runtime/Experimental/IoStore/HttpClient/Private/Loop.inl
--- a/Engine/Source/Runtime/Experimental/IoStore/HttpClient/Private/Loop.inl
+++ b/Engine/Source/Runtime/Experimental/IoStore/HttpClient/Private/Loop.inl
@@ -881,7 +881,7 @@ void FPeerGroup::RecvInternal(FTickState& State)
 	// If we've only a small amount left to receive we can start more work
 	if (IsKeepAlive & (Recv->Next == nullptr) & (Send == nullptr))
 	{
-		uint32 Remaining = Activity_RemainingKiB(Activity);
+		uint32 Remaining = Activity->GetRemainingKiB();
 		if (Remaining < uint32(GRecvWorkThresholdKiB))
 		{
 			if (FActivityNode* Next = State.Work->PopActivity(); Next != nullptr)
