Step one of moving the chunked part out of loose functions in loop.

#rnx

diff --git a/Engine/Source/Runtime/Experimental/IoStore/HttpClient/Private/Activity.inl b/Engine/Source/Runtime/Experimental/IoStore/HttpClient/Private/Activity.inl
--- a/Engine/Source/Runtime/Experimental/IoStore/HttpClient/Private/Activity.inl
+++ b/Engine/Source/Runtime/Experimental/IoStore/HttpClient/Private/Activity.inl
@@ -362,16 +362,28 @@ class FBodyCloak
 public:
 	using FResponseCloak::FResponseCloak;
 
-	FOutcome	TryRecv(FMutableMemoryView Dest, FHttpPeer& Peer);
+	FOutcome		TryRecv(FMutableMemoryView Dest, FHttpPeer& Peer);
 
 private:
-	enum class EState : uint8 { Init, Recv, Done };
-
-	FOutcome	Init(FMutableMemoryView Dest, FHttpPeer& Peer);
-	FOutcome	Recv(FMutableMemoryView Dest, FHttpPeer& Peer);
-	FMemoryView	Overspill;
-	int64		Remaining = 0;
-	EState		State = EState::Init;
+	enum class EState : uint8 { Init, Recv, Prologue, Chunk, Epilogue, Done };
+
+	FOutcome			Gather(FMutableMemoryView Dest, FHttpPeer& Peer);
+	FOutcome			Init(FMutableMemoryView Dest, FHttpPeer& Peer);
+	FOutcome			Recv(FMutableMemoryView Dest, FHttpPeer& Peer);
+	FOutcome			Prologue(FMutableMemoryView Dest, FHttpPeer& Peer);
+	FOutcome			Chunk(FMutableMemoryView Dest, FHttpPeer& Peer);
+	FOutcome			Epilogue(FMutableMemoryView Dest, FHttpPeer& Peer);
+	FMemoryView			Overspill;
+	FMutableMemoryView	Scratch;
+	int64				Remaining = 0;
+	EState				State = EState::Init;
+
+	enum
+	{
+		HeaderBufSize	= 32,
+		CrLfLength		= 2,
+		EndOfXfer		= -1,
+	};
 };
 
 ////////////////////////////////////////////////////////////////////////////////
@@ -381,9 +393,10 @@ FOutcome FBodyCloak::TryRecv(FMutableMemoryView Dest, FHttpPeer& Peer)
 
 	switch (State)
 	{
-	case EState::Init:		Outcome = Init(Dest, Peer); break;
-	case EState::Recv:		Outcome = Recv(Dest, Peer); break;
-	case EState::Done:		Outcome = FOutcome::Error("Transaction is complete"); break;
+	case EState::Init:	Outcome = Init(Dest, Peer);		break;
+	case EState::Recv:	Outcome = Recv(Dest, Peer);		break;
+	default:			Outcome = Gather(Dest, Peer);	break;
+	case EState::Done:	Outcome = FOutcome::Error("Transaction is complete"); break;
 	}
 
 	if (Outcome.IsOk())
@@ -394,6 +407,43 @@ FOutcome FBodyCloak::TryRecv(FMutableMemoryView Dest, FHttpPeer& Peer)
 	return Outcome;
 }
 
+////////////////////////////////////////////////////////////////////////////////
+FOutcome FBodyCloak::Gather(FMutableMemoryView Dest, FHttpPeer& Peer)
+{
+	FOutcome Outcome = FOutcome::None();
+
+	bool bDone = false;
+	int32 Size = 0;
+	do
+	{
+		switch (State)
+		{
+		case EState::Prologue:	Outcome = Prologue(Dest, Peer); break;
+		case EState::Chunk:		Outcome = Chunk(Dest, Peer);	break;
+		case EState::Epilogue:	Outcome = Epilogue(Dest, Peer);	break;
+		default:				Outcome = FOutcome::Error("Unexpected try-chunked state"); break;
+		}
+
+		if (Outcome.IsError())
+		{
+			return Outcome;
+		}
+
+		bDone = Outcome.IsOk();
+		int32 Result = Outcome.GetResult();
+		if (Result == 0)
+		{
+			break;
+		}
+
+		Size += Result;
+		Dest = Dest.Mid(Result);
+	}
+	while (!Dest.IsEmpty());
+
+	return bDone ? FOutcome::Ok(Size) : FOutcome::Waiting(Size);
+}
+
 ////////////////////////////////////////////////////////////////////////////////
 FOutcome FBodyCloak::Init(FMutableMemoryView Dest, FHttpPeer& Peer)
 {
@@ -404,11 +454,28 @@ FOutcome FBodyCloak::Init(FMutableMemoryView Dest, FHttpPeer& Peer)
 	{
 		Overspill = { (const uint8*)MessageEnd, AlreadyReceived };
 	}
+
+	if (AlreadyReceived < HeaderBufSize)
+	{
+		GetMutableFree(HeaderBufSize, HeaderBufSize);
+		MessageEnd = TEMP_GetMessageRight();
+	}
+
+	Scratch = { (uint8*)MessageEnd, HeaderBufSize };
 	
-	Remaining = GetContentLength();
-	check(Remaining > 0);
-	State = EState::Recv;
-	return Recv(Dest, Peer);
+	if (IsChunked())
+	{
+		check(Remaining == 0);
+		State = EState::Prologue;
+	}
+	else
+	{
+		Remaining = GetContentLength();
+		check(Remaining > 0);
+		State = EState::Recv;
+	}
+
+	return TryRecv(Dest, Peer);
 }
 
 ////////////////////////////////////////////////////////////////////////////////
@@ -448,6 +515,163 @@ FOutcome FBodyCloak::Recv(FMutableMemoryView Dest, FHttpPeer& Peer)
 	return Remaining ? FOutcome::Waiting(Result) : FOutcome::Ok(Result);
 }
 
+////////////////////////////////////////////////////////////////////////////////
+FOutcome FBodyCloak::Prologue(FMutableMemoryView Dest, FHttpPeer& Peer)
+{
+	auto FindHeader = [] (const char* Cursor, int32 Size)
+	{
+		// Base-16 size
+		int32 HexDigLen = 0;
+		for (; HexDigLen < Size; ++HexDigLen)
+		{
+			uint32 c = Cursor[HexDigLen];
+			if ((c - '0' >= 10) & (c - 'a' >= 6) & (c - 'A' >= 6))
+			{
+				break;
+			}
+			continue;
+		}
+
+		int32 Delta = Size - HexDigLen;
+
+		if (HexDigLen == 0)					return FOutcome::Error("Invalid chunk header");
+		if (Delta < 1)						return FOutcome::Waiting();
+		if (Cursor[HexDigLen] != '\r')		return FOutcome::Error("Chunk extensions are not supported (ERREXT)");
+		if (Delta < 2)						return FOutcome::Waiting();
+		if (Cursor[HexDigLen + 1] != '\n')	return FOutcome::Error("Invalid chunk CRLF");
+
+		return FOutcome::Ok(HexDigLen + 2);
+	};
+
+	const char* Header = nullptr;
+	FOutcome Outcome = FOutcome::None();
+	int32 HeaderLength = 0;
+
+	uint32 OverspillSize = uint32(Overspill.GetSize());
+	if (OverspillSize != 0)
+	{
+		// We try all the overspill first - we don't know how much we're expecting
+		// and overspill could contain both header and all data (thus further recvs
+		// could falsely hang us up).
+		Header = (char*)Overspill.GetData();
+
+		Outcome = FindHeader(Header, OverspillSize);
+		if (Outcome.IsError())
+		{
+			return Outcome;
+		}
+
+		HeaderLength = Outcome.GetResult();
+		Overspill = Overspill.Mid(HeaderLength);
+	}
+
+	if (HeaderLength == 0)
+	{
+		char* Cursor = (char*)Scratch.GetData();
+
+		if (OverspillSize != 0)
+		{
+			check(OverspillSize <= HeaderBufSize);
+			std::memmove(Cursor, Overspill.GetData(), OverspillSize);
+			Cursor += OverspillSize;
+		}
+
+		Outcome = Peer.Recv(Cursor, HeaderBufSize);
+		if (!Outcome.IsOk())
+		{
+			return Outcome;
+		}
+
+		Cursor -= OverspillSize;
+		OverspillSize += Outcome.GetResult();
+		Overspill = { (uint8*)Cursor, OverspillSize };
+
+		Outcome = FindHeader(Cursor, OverspillSize);
+		if (!Outcome.IsOk())
+		{
+			return Outcome;
+		}
+
+		Header = Cursor;
+		HeaderLength = Outcome.GetResult();
+		Overspill = Overspill.Mid(HeaderLength);
+	}
+
+	// Read chunk size
+	Header += Remaining;
+	int64 ChunkSize = CrudeToInt<16>(FAnsiStringView(Header, HeaderLength));
+	if (ChunkSize < 0)			return FOutcome::Error("Unparsable chunk size");
+	if (ChunkSize > 16 << 20)	return FOutcome::Error("Unacceptable chunk size");
+
+	if (ChunkSize == 0)
+	{
+		Remaining = EndOfXfer;
+		::memset(Scratch.GetData(), 0, CrLfLength);
+		State = EState::Epilogue;
+		return Epilogue(Dest, Peer);
+	}
+
+	Remaining = ChunkSize;
+	State = EState::Chunk;
+	return Chunk(Dest, Peer);
+}
+
+////////////////////////////////////////////////////////////////////////////////
+FOutcome FBodyCloak::Chunk(FMutableMemoryView Dest, FHttpPeer& Peer)
+{
+	check(Remaining != 0);
+
+	Dest = Dest.Left(Remaining);
+	FOutcome Outcome = Recv(Dest, Peer);
+	if (!Outcome.IsOk())
+	{
+		return Outcome;
+	}
+
+	check(Remaining == 0);
+	::memset(Scratch.GetData(), 0, CrLfLength);
+	State = EState::Epilogue;
+
+	int32 Result = Outcome.GetResult();
+	return FOutcome::Waiting(Result);
+}
+
+////////////////////////////////////////////////////////////////////////////////
+FOutcome FBodyCloak::Epilogue(FMutableMemoryView Dest, FHttpPeer& Peer)
+{
+	auto* Cursor = (uint8*)(Scratch.GetData());
+
+	FMutableMemoryView CrLfBuffer(Cursor, CrLfLength);
+	if (Cursor[0] != 0)
+	{
+		CrLfBuffer = { Cursor + 1, 1 };
+	}
+
+	check(Remaining <= 0);
+
+	int64 RemainingCache = Remaining;
+	Remaining = CrLfBuffer.GetSize();
+	FOutcome Outcome = Recv(CrLfBuffer, Peer);
+	if (!Outcome.IsOk())
+	{
+		Remaining = RemainingCache;
+		return Outcome;
+	}
+
+	if (Cursor[0] != '\r' || Cursor[1] != '\n')
+	{
+		return FOutcome::Error("Trailing headers are not supported (ERRTRAIL)");
+	}
+
+	if (RemainingCache > EndOfXfer)
+	{
+		State = EState::Prologue;
+		return Prologue(Dest, Peer);
+	}
+
+	return FOutcome::Ok();
+}
+
 
 
 #if IAS_HTTP_WITH_PERF
