Removed old chunk impl and swapped to the new one.

#rnx

diff --git a/Engine/Source/Runtime/Experimental/IoStore/HttpClient/Private/Loop.inl b/Engine/Source/Runtime/Experimental/IoStore/HttpClient/Private/Loop.inl
--- a/Engine/Source/Runtime/Experimental/IoStore/HttpClient/Private/Loop.inl
+++ b/Engine/Source/Runtime/Experimental/IoStore/HttpClient/Private/Loop.inl
@@ -40,43 +40,6 @@ static FOutcome DoSend(FActivity* Activity, FHttpPeer& Peer)
 	return FOutcome::Ok();
 }
 
-////////////////////////////////////////////////////////////////////////////////
-static FOutcome DoRecvPeer(FActivity* Activity, FHttpPeer& Peer, int32& MaxRecvSize, int32 Size)
-{
-	Size = FMath::Min(Size, MaxRecvSize);
-	check(Size >= 0);
-	if (Size == 0)
-	{
-		return FOutcome::Waiting();
-	}
-
-
-	FMutableMemoryView DestView = Activity->Dest->GetMutableView();
-	char* Cursor = (char*)(DestView.GetData()) + Activity->StateParam;
-	check(Size + Activity->StateParam <= DestView.GetSize());
-
-	FOutcome Outcome = Peer.Recv(Cursor, Size);
-
-	if (Outcome.IsWaiting())
-	{
-		return Outcome;
-	}
-
-	if (Outcome.IsError())
-	{
-		Activity_SetError(Activity, Outcome);
-		return Outcome;
-	}
-
-	check(Outcome.IsOk());
-
-	int32 Result = Outcome.GetResult();
-	Activity->StateParam += Result;
-	MaxRecvSize = FMath::Max(MaxRecvSize - Result, 0);
-
-	return Outcome;
-}
-
 ////////////////////////////////////////////////////////////////////////////////
 static FOutcome DoRecvMessage(FActivity* Activity, FHttpPeer& Peer)
 {
@@ -158,37 +121,8 @@ static FOutcome DoRecvMessage(FActivity* Activity, FHttpPeer& Peer)
 
 	// We're all set to go and get content
 	check(Activity->Dest != nullptr);
-
-	if (!bChunked)
-	{
-		Activity_ChangeState(Activity, FActivity::EState::RecvContent);
-		return FOutcome::Ok();
-	}
-
-	// Perhaps we have some of the content already?
-	const char* MessageRight = Response.TEMP_GetMessageRight();
-	const char* BufferRight = Response.GetData() + Response.GetSize();
-	uint32 AlreadyReceived = uint32(ptrdiff_t(BufferRight - MessageRight));
-
-	Activity_ChangeState(Activity, FActivity::EState::RecvStream, AlreadyReceived);
-
-	// Copy any of the content we may have already received.
-	if (AlreadyReceived == 0)
-	{
-		return FOutcome::Ok();
-	}
-
-	// This ordinarily doesn't happen due to the way higher levels pipeline
-	// requests. It can however occur with chunked transfers.
-	if (AlreadyReceived > Dest.GetSize())
-	{
-		Dest = FIoBuffer(AlreadyReceived);
-	}
-
-	FMutableMemoryView DestView = Dest.GetMutableView();
-	const char* Cursor = BufferRight - AlreadyReceived;
-	::memcpy(DestView.GetData(), Cursor, AlreadyReceived);
-
+	auto NextState = bChunked ? FActivity::EState::RecvStream : FActivity::EState::RecvContent;
+	Activity_ChangeState(Activity, NextState);
 	return FOutcome::Ok();
 }
 
@@ -228,202 +162,50 @@ static FOutcome DoRecvContent(FActivity* Activity, FHttpPeer& Peer, int32& MaxRe
 ////////////////////////////////////////////////////////////////////////////////
 static FOutcome DoRecvStream(FActivity* Activity, FHttpPeer& Peer, int32& MaxRecvSize)
 {
-	auto RaiseCrLfError = [Activity] ()
-	{
-		if (Activity->NoContent)
-		{
-			Activity_SetError(Activity, "Trailing headers are not supported (ERRTRAIL)");
-			return FOutcome::Error(Activity->ErrorReason);
-		}
-
-		Activity_SetError(Activity, "Expected CRLF chunk terminal");
-		return FOutcome::Error(Activity->ErrorReason);
-	};
-
-	auto SinkData = [Activity] (FMemoryView View)
-	{
-		if (View.GetSize() == 0)
-		{
-			return;
-		}
-
-		// Temporarily clamp IoBuffer so if the sink does GetView/GetSize() it
-		// represents actual content and not the underlying working buffer.
-		FIoBuffer& Dest = *(Activity->Dest);
-		FIoBuffer Slice(View, Dest);
-		Swap(Dest, Slice);
-		
-		Activity_CallSink(Activity);
+	TRACE_CPUPROFILER_EVENT_SCOPE(IasHttp::DoRecvStream);
 
-		Swap(Dest, Slice);
-	};
+	FMutableMemoryView Dest = Activity->Dest->GetMutableView();
+	Dest = Dest.Left(MaxRecvSize);
+	MaxRecvSize -= int32(Dest.GetSize());
 
-	auto Done = [Activity] ()
+	auto& Body = Activity->Buffer;
+	FOutcome Outcome = Body.TryRecv(Dest, Peer);
+	if (Outcome.IsError())
 	{
-#if IAS_HTTP_WITH_PERF
-			Activity->Stopwatch.RecvEnd();
-#endif
-			*(Activity->Dest) = FIoBuffer();
-			Activity_ChangeState(Activity, FActivity::EState::RecvDone);
-			return FOutcome::Ok();
-	};
-
-	enum { CrLfLength = 2 };
+		Activity_SetError(Activity, Outcome);
+		return Outcome;
+	}
 
-	int32 Size = int32(Activity->StateParam);
+	int32 Result = Outcome.GetResult();
 
-	// Trailing chunk data.
-	while (Size < 0)
+	if (Result != 0)
 	{
-		Activity->StateParam = 0;
-		Size = -Size;
-
-		int32 RefillSize = FMath::Min<int32>(Size, int32(Activity->Dest->GetSize()));
-		FOutcome Outcome = DoRecvPeer(Activity, Peer, MaxRecvSize, RefillSize);
-		if (!Outcome.IsOk())
-		{
-			Activity->StateParam = 0 - Size;
-			return Outcome;
-		}
-
-		int32 Result = Outcome.GetResult();
-		check(Result > 0);
-
-		FMemoryView View = Activity->Dest->GetView();
-		if (Size > CrLfLength)
-		{
-			int32 SinkSize = Size - CrLfLength;
-			SinkSize = FMath::Min(Result, SinkSize);
-			SinkData(View.Left(SinkSize));
-			View = View.Mid(SinkSize);
-			Size -= SinkSize;
-			Result -= SinkSize;
-		}
-
-		const char* Cursor = (char*)View.GetData();
-		int32 CrLfError = 0;
-		if (int32 n = CrLfLength; Size == n && Result >= n)
-		{
-			CrLfError |= (Cursor[0] != '\r');
-			--Size; --Result;
-			++Cursor;
-		}
-		if (int32 n = CrLfLength - 1; Size == n && Result >= n)
-		{
-			CrLfError |= (Cursor[0] != '\n');
-			--Size; --Result;
-		}
-		if (CrLfError)
-		{
-			return RaiseCrLfError();
-		}
+		// Temporarily clamp IoBuffer so if the sink does GetView/GetSize() it
+		// represents actual content and not the underlying working buffer.
+		FIoBuffer& Outer = *(Activity->Dest);
 
-		Size = Result - Size;
-		Activity->StateParam = Size;
-		check(Size <= 0);
+		FMemoryView View = Outer.GetView();
+		View = View.Left(Result);
+		FIoBuffer Slice(View, Outer);
 
-		// Have we found the trailer-section that follows last-chunk?
-		if (Size == 0 && Activity->NoContent)
-		{
-			return Done();
-		}
+		Swap(Outer, Slice);
+		Activity_CallSink(Activity);
+		Swap(Outer, Slice);
 	}
 
-	// Peel off chunks
-	for (FMemoryView View = Activity->Dest->GetView(); Size > 0;)
+	if (!Outcome.IsOk())
 	{
-		const char* Cursor = (char*)(View.GetData());
-
-		// Isolate chunk size
- 		int32 ChunkSize = -1;
-		uint32 HeaderLength = 0;
-		for (; HeaderLength < uint32(Size - 1); ++HeaderLength)
-		{
-			// Detect CRLF.
-			if (Cursor[HeaderLength + 1] != '\n')
-			{
-				continue;
-			}
-
-			++HeaderLength;
-			if (Cursor[HeaderLength - 1] != '\r')
-			{
-				continue;
-			}
-			++HeaderLength;
-
-			ChunkSize = int32(CrudeToInt<16>(FAnsiStringView(Cursor, (int32)View.GetSize())));
-
-			if (ChunkSize < 0)
-			{
-				Activity_SetError(Activity, "Unparsable chunk size");
-				return FOutcome::Error(Activity->ErrorReason);
-			}
-
-			break;
-		}
-
-		// Maybe we were not able to find a CRLF terminator and need more data
-		if (ChunkSize < 0)
-		{
-			FMutableMemoryView WriteView = Activity->Dest->GetMutableView();
-			std::memmove(WriteView.GetData(), Cursor, Size);
-			Activity->StateParam = Size;
-			break;
-		}
-
-		check(ChunkSize >= 0);
-		Size -= HeaderLength;
-
-		// Dispatch as much data as we can.
-		uint32 SinkSize = FMath::Min<uint32>(ChunkSize, uint32(Size));
-		SinkData(View.Mid(HeaderLength, SinkSize));
-		View = View.Mid(HeaderLength + SinkSize);
-
-		Activity->StateParam = (Size -= ChunkSize);
-		Activity->NoContent = (ChunkSize == 0);
-
-		// A CRLF follows a chunk's data
-		Cursor = (char*)View.GetData();
-		int32 CrLfError = 0;
-		CrLfError |= (Size >= (CrLfLength - 1)) && Cursor[0] != '\r';
-		CrLfError |= (Size >= (CrLfLength - 0)) && Cursor[1] != '\n';
-		if (CrLfError != 0)
-		{
-			return RaiseCrLfError();
-		}
-
-		// Can we do CRLF now?
-		if (Size >= CrLfLength)
-		{
-			// Have we found the trailer-section that follows last-chunk?
-			if (Activity->NoContent)
-			{
-				return Done();
-			}
-
-			Activity->StateParam = (Size -= CrLfLength);
-			View = View.Mid(CrLfLength);
-			continue;
-		}
-
-		Activity->StateParam -= CrLfLength;
-		check(int32(Activity->StateParam) < 0);
-		break;
+		check(Outcome.IsWaiting());
+		return Outcome;
 	}
 
-	// Refill
-	if (int32(Activity->StateParam) >= 0)
-	{
-		uint32 RefillSize = uint32(Activity->Dest->GetSize()) - Activity->StateParam;
-		FOutcome Outcome = DoRecvPeer(Activity, Peer, MaxRecvSize, RefillSize);
-		if (!Outcome.IsOk())
-		{
-			return Outcome;
-		}
-	}
+#if IAS_HTTP_WITH_PERF
+	Activity->Stopwatch.RecvEnd();
+#endif
 
-	return DoRecvStream(Activity, Peer, MaxRecvSize);
+	*(Activity->Dest) = FIoBuffer();
+	Activity_ChangeState(Activity, FActivity::EState::RecvDone);
+	return FOutcome::Ok();
 }
 
 ////////////////////////////////////////////////////////////////////////////////
