Updated tests

#rnx

diff --git a/Engine/Source/Runtime/Experimental/IoStore/HttpClient/Private/Test.inl b/Engine/Source/Runtime/Experimental/IoStore/HttpClient/Private/Test.inl
--- a/Engine/Source/Runtime/Experimental/IoStore/HttpClient/Private/Test.inl
+++ b/Engine/Source/Runtime/Experimental/IoStore/HttpClient/Private/Test.inl
@@ -455,7 +455,11 @@ static void ChunkedTest(const ANSICHAR* TestHost)
 	};
 
 	// General soak test
-	for (FAnsiStringView UrlSuffix : { "", "/ext" })
+	FAnsiStringView UrlSuffices[] = {
+		"",
+	//	"/ext", /* chunk-ext support was removed */
+	};
+	for (FAnsiStringView UrlSuffix : UrlSuffices)
 	{
 		for (uint32 Mixer : { 1, 2, 3, 17, 71, 4931, 0xa9e })
 		{
@@ -488,6 +492,7 @@ static void ChunkedTest(const ANSICHAR* TestHost)
 		FAnsiStringView Reason = Status.GetError().Reason;
 		ErrorMarks |= Reason.Contains("ERRTRAIL",   ESearchCase::CaseSensitive) ? 1 : 0;
 		ErrorMarks |= Reason.Contains("ERRNOCHUNK", ESearchCase::CaseSensitive) ? 2 : 0;
+		ErrorMarks |= Reason.Contains("ERREXT", 	ESearchCase::CaseSensitive) ? 4 : 0;
 	};
 	ErrorMarks = 0;
 	BuildUrl(16 << 10, "/trailer");
@@ -495,6 +500,15 @@ static void ChunkedTest(const ANSICHAR* TestHost)
 	WaitForLoopIdle();
 	check(ErrorMarks == 1);
 
+	// Chunk extensions are not supported
+	{
+		ErrorMarks = 0;
+		BuildUrl(493, "/ext");
+		Loop.Send(Loop.Get(Url), ExpectError);
+		WaitForLoopIdle();
+		check(ErrorMarks == 4);
+	}
+
 	// Disabling of chunked transfers 
 	{
 		ErrorMarks = 0;
