Hid a now internally-used method

#rnx

diff --git a/Engine/Source/Runtime/Experimental/IoStore/HttpClient/Private/Activity.inl b/Engine/Source/Runtime/Experimental/IoStore/HttpClient/Private/Activity.inl
--- a/Engine/Source/Runtime/Experimental/IoStore/HttpClient/Private/Activity.inl
+++ b/Engine/Source/Runtime/Experimental/IoStore/HttpClient/Private/Activity.inl
@@ -120,7 +120,9 @@ public:
 	bool			IsChunked() const;
 	bool			HasBody() const;
 	void			ReadHeaders(FResponse::FHeaderSink Sink) const;
-	const char*		TEMP_GetMessageRight() const;
+
+protected:
+	const char*		GetMessageRight() const;
 
 private:
 	FOutcome		Parse();
@@ -187,7 +189,7 @@ void FResponseCloak::ReadHeaders(FResponse::FHeaderSink Sink) const
 }
 
 ////////////////////////////////////////////////////////////////////////////////
-const char* FResponseCloak::TEMP_GetMessageRight() const
+const char* FResponseCloak::GetMessageRight() const
 {
 	return GetData() + MessageRight;
 }
@@ -448,7 +450,7 @@ FOutcome FBodyCloak::Gather(FMutableMemoryView Dest, FHttpPeer& Peer)
 FOutcome FBodyCloak::Init(FMutableMemoryView Dest, FHttpPeer& Peer)
 {
 	const char* DataEnd = GetData() + GetSize();
-	const char* MessageEnd = TEMP_GetMessageRight();
+	const char* MessageEnd = GetMessageRight();
 	uint32 AlreadyReceived = uint32(ptrdiff_t(DataEnd - MessageEnd));
 	if (AlreadyReceived)
 	{
@@ -458,7 +460,7 @@ FOutcome FBodyCloak::Init(FMutableMemoryView Dest, FHttpPeer& Peer)
 	if (AlreadyReceived < HeaderBufSize)
 	{
 		GetMutableFree(HeaderBufSize, HeaderBufSize);
-		MessageEnd = TEMP_GetMessageRight();
+		MessageEnd = GetMessageRight();
 	}
 
 	Scratch = { (uint8*)MessageEnd, HeaderBufSize };
