Deleted Fix() method that is now unused.

#rnx

diff --git a/Engine/Source/Runtime/Experimental/IoStore/HttpClient/Private/Misc.inl b/Engine/Source/Runtime/Experimental/IoStore/HttpClient/Private/Misc.inl
--- a/Engine/Source/Runtime/Experimental/IoStore/HttpClient/Private/Misc.inl
+++ b/Engine/Source/Runtime/Experimental/IoStore/HttpClient/Private/Misc.inl
@@ -444,7 +444,6 @@ public:
 								FBuffer(char* InData, uint32 InMax);
 								~FBuffer();
 	FBuffer&					operator = (FBuffer&& Rhs);
-	void						Fix();
 	void						Resize(uint32 Size);
 	const char*					GetData() const;
 	uint32						GetSize() const;
@@ -500,15 +499,6 @@ FBuffer& FBuffer::operator = (FBuffer&& Rhs)
 	return *this;
 }
 
-////////////////////////////////////////////////////////////////////////////////
-void FBuffer::Fix()
-{
-	check(Inline);
-	Data += Used;
-	Max -= Used;
-	Used = 0;
-}
-
 ////////////////////////////////////////////////////////////////////////////////
 void FBuffer::Resize(uint32 Size)
 {
