Removed NoContent flag that is no longer referenced.

#rnx

diff --git a/Engine/Source/Runtime/Experimental/IoStore/HttpClient/Private/Activity.inl b/Engine/Source/Runtime/Experimental/IoStore/HttpClient/Private/Activity.inl
--- a/Engine/Source/Runtime/Experimental/IoStore/HttpClient/Private/Activity.inl
+++ b/Engine/Source/Runtime/Experimental/IoStore/HttpClient/Private/Activity.inl
@@ -773,9 +773,8 @@ struct FActivity
 
 	EState				State = EState::None;
 	uint8				IsKeepAlive : 1;
-	uint8				NoContent : 1;
 	uint8				bAllowChunked : 1;
-	uint8				_Unused : 1;
+	uint8				_Unused : 2;
 	uint8				LengthScore : 4;
 	uint32				StateParam = 0;
 #if IAS_HTTP_WITH_PERF
