More frontend/activity encapsulation.

#rnx

diff --git a/Engine/Source/Runtime/Experimental/IoStore/HttpClient/Private/Activity.inl b/Engine/Source/Runtime/Experimental/IoStore/HttpClient/Private/Activity.inl
--- a/Engine/Source/Runtime/Experimental/IoStore/HttpClient/Private/Activity.inl
+++ b/Engine/Source/Runtime/Experimental/IoStore/HttpClient/Private/Activity.inl
@@ -745,6 +745,9 @@ struct FActivity
 						~FActivity() = default;
 	void				SetSink(FTicketSink&& InSink, UPTRINT Param);
 	void				SetDestination(FIoBuffer* InDest);
+	void				Cancel();
+	void				Done();
+	void				Fail(const FOutcome& Outcome);
 	FHost*				GetHost() const;
 	UPTRINT				GetSinkParam() const;
 	FIoBuffer&			GetContent() const;
@@ -771,6 +774,9 @@ struct FActivity
 		_Num,
 	};
 
+	void				CallSink();
+	void				ChangeState(EState InState, uint32 Param=0);
+
 	EState				State = EState::None;
 	uint8				IsKeepAlive : 1;
 	uint8				bAllowChunked : 1;
@@ -817,6 +823,65 @@ void FActivity::SetDestination(FIoBuffer* InDest)
 	Dest = InDest;
 }
 
+////////////////////////////////////////////////////////////////////////////////
+void FActivity::Cancel()
+{
+	if (State >= FActivity::EState::Completed)
+	{
+		return;
+	}
+
+	ChangeState(FActivity::EState::Cancelled);
+	CallSink();
+}
+
+////////////////////////////////////////////////////////////////////////////////
+void FActivity::Done()
+{
+	ChangeState(EState::RecvDone);
+	CallSink();
+	ChangeState(EState::Completed);
+}
+
+////////////////////////////////////////////////////////////////////////////////
+void FActivity::Fail(const FOutcome& Outcome)
+{
+	if (State == EState::Failed)
+	{
+		return;
+	}
+
+	ChangeState(EState::Failed);
+
+	IsKeepAlive = 0;
+	ErrorReason = Outcome.GetMessage().GetData();
+	StateParam = uint32(Outcome.GetErrorCode());
+
+	CallSink();
+}
+
+////////////////////////////////////////////////////////////////////////////////
+void FActivity::ChangeState(EState InState, uint32 Param)
+{
+	check(State != InState);
+
+	Trace(this, ETrace::StateChange, uint32(InState));
+
+	State = InState;
+	StateParam = Param;
+}
+
+////////////////////////////////////////////////////////////////////////////////
+void FActivity::CallSink()
+{
+	static uint32 Scope = LaneTrace_NewScope("Iax/Sink");
+	FLaneTrace* Lane = LaneEstate_Lookup(GActivityTraceEstate, this);
+	FLaneTraceScope _(Lane, Scope);
+
+	FTicketStatus& SinkArg = *(FTicketStatus*)this;
+	Sink(SinkArg);
+}
+
 ////////////////////////////////////////////////////////////////////////////////
 FHost* FActivity::GetHost() const
 {
@@ -872,16 +937,6 @@ const FStopwatch& FActivity::GetStopwatch() const
 
 
 
-////////////////////////////////////////////////////////////////////////////////
-static void Activity_ChangeState(FActivity* Activity, FActivity::EState InState, uint32 Param=0)
-{
-	Trace(Activity, ETrace::StateChange, uint32(InState));
-
-	check(Activity->State != InState);
-	Activity->State = InState;
-	Activity->StateParam = Param;
-}
-
 ////////////////////////////////////////////////////////////////////////////////
 static int32 Activity_Rewind(FActivity* Activity)
 {
@@ -896,40 +951,13 @@ static int32 Activity_Rewind(FActivity* Activity)
 	if (Activity->State == EState::RecvMessage)
 	{
 		Activity->Buffer.Resize(Activity->StateParam);
-		Activity_ChangeState(Activity, EState::Send);
+		Activity->ChangeState(FActivity::EState::Send);
 		return 1;
 	}
 
 	return -1;
 }
 
-////////////////////////////////////////////////////////////////////////////////
-static void Activity_CallSink(FActivity* Activity)
-{
-	static uint32 Scope = LaneTrace_NewScope("Iax/Sink");
-	FLaneTrace* Lane = LaneEstate_Lookup(GActivityTraceEstate, Activity);
-	FLaneTraceScope _(Lane, Scope);
-
-	FTicketStatus& SinkArg = *(FTicketStatus*)Activity;
-	Activity->Sink(SinkArg);
-}
-
-////////////////////////////////////////////////////////////////////////////////
-static void Activity_SetError(FActivity* Activity, const char* Reason, int32 Code=-1)
-{
-	Activity->IsKeepAlive = 0;
-	Activity->ErrorReason = Reason;
-
-	Code = (Code < 0) ? LastSocketResult() : Code;
-	Activity_ChangeState(Activity, FActivity::EState::Failed, Code);
-}
-
-////////////////////////////////////////////////////////////////////////////////
-static void Activity_SetError(FActivity* Activity, const FOutcome& Outcome)
-{
-	Activity_SetError(Activity, Outcome.GetMessage().GetData(), Outcome.GetErrorCode());
-}
-
 
 
 ////////////////////////////////////////////////////////////////////////////////
diff --git a/Engine/Source/Runtime/Experimental/IoStore/HttpClient/Private/Loop.inl b/Engine/Source/Runtime/Experimental/IoStore/HttpClient/Private/Loop.inl
--- a/Engine/Source/Runtime/Experimental/IoStore/HttpClient/Private/Loop.inl
+++ b/Engine/Source/Runtime/Experimental/IoStore/HttpClient/Private/Loop.inl
@@ -21,7 +21,7 @@ static FOutcome DoSend(FActivity* Activity, FHttpPeer& Peer)
 	FOutcome Outcome = Buffer.TrySendRequest(Peer);
 	if (Outcome.IsError())
 	{
-		Activity_SetError(Activity, Outcome);
+		Activity->Fail(Outcome);
 		return Outcome;
 	}
 
@@ -36,7 +36,7 @@ static FOutcome DoSend(FActivity* Activity, FHttpPeer& Peer)
 	Activity->Stopwatch.SendEnd();
 #endif
 
-	Activity_ChangeState(Activity, FActivity::EState::RecvMessage);
+	Activity->ChangeState(FActivity::EState::RecvMessage);
 	return FOutcome::Ok();
 }
 
@@ -55,7 +55,7 @@ static FOutcome DoRecvMessage(FActivity* Activity, FHttpPeer& Peer)
 	FOutcome Outcome = Response.TryRecvResponse(Peer);
 	if (Outcome.IsError())
 	{
-		Activity_SetError(Activity, Outcome);
+		Activity->Fail(Outcome);
 		return Outcome;
 	}
 
@@ -74,14 +74,15 @@ static FOutcome DoRecvMessage(FActivity* Activity, FHttpPeer& Peer)
 	{
 		if (Activity->bAllowChunked == 0)
 		{
-			Activity_SetError(Activity, "Chunked transfer encoding disabled (ERRNOCHUNK)");
-			return FOutcome::Error(Activity->ErrorReason);
+			Outcome = FOutcome::Error("Chunked transfer encoding disabled (ERRNOCHUNK)");
+			Activity->Fail(Outcome);
+			return Outcome;
 		}
 	}
 
 	// Call out to the sink to get a content destination
 	FIoBuffer* PriorDest = Activity->Dest; // to retain unioned Host ptr (redirect uses it in sink)
-	Activity_CallSink(Activity);
+	Activity->CallSink();
 
 	// Check the user gave us a destination for content
 	FIoBuffer& Dest = *(Activity->Dest);
@@ -89,8 +90,9 @@ static FOutcome DoRecvMessage(FActivity* Activity, FHttpPeer& Peer)
 	{
 		if (&Dest == PriorDest)
 		{
-			Activity_SetError(Activity, "User did not provide a destination buffer");
-			return FOutcome::Error(Activity->ErrorReason);
+			Outcome = FOutcome::Error("User did not provide a destination buffer");
+			Activity->Fail(Outcome);
+			return Outcome;
 		}
 
 		// The user seems to have forgotten something. Let's help them along
@@ -103,8 +105,9 @@ static FOutcome DoRecvMessage(FActivity* Activity, FHttpPeer& Peer)
 		else if (!bChunked && DestSize < ContentLength)
 		{
 			// todo: support piece-wise transfer of content (a la chunked).
-			Activity_SetError(Activity, "Destination buffer too small");
-			return FOutcome::Error(Activity->ErrorReason);
+			Outcome = FOutcome::Error("Destination buffer too small");
+			Activity->Fail(Outcome);
+			return Outcome;
 		}
 		else if (enum { MinStreamBuf = 256 }; bChunked && DestSize < MinStreamBuf)
 		{
@@ -115,14 +118,14 @@ static FOutcome DoRecvMessage(FActivity* Activity, FHttpPeer& Peer)
 	// HEAD methods
 	if (!Response.HasBody())
 	{
-		Activity_ChangeState(Activity, FActivity::EState::RecvDone);
+		Activity->ChangeState(FActivity::EState::RecvDone);
 		return FOutcome::Ok();
 	}
 
 	// We're all set to go and get content
 	check(Activity->Dest != nullptr);
 	auto NextState = bChunked ? FActivity::EState::RecvStream : FActivity::EState::RecvContent;
-	Activity_ChangeState(Activity, NextState);
+	Activity->ChangeState(NextState);
 	return FOutcome::Ok();
 }
 
@@ -139,7 +142,7 @@ static FOutcome DoRecvContent(FActivity* Activity, FHttpPeer& Peer, int32& MaxRe
 	FOutcome Outcome = Body.TryRecv(Dest, Peer);
 	if (Outcome.IsError())
 	{
-		Activity_SetError(Activity, Outcome);
+		Activity->Fail(Outcome);
 		return Outcome;
 	}
 
@@ -155,7 +158,7 @@ static FOutcome DoRecvContent(FActivity* Activity, FHttpPeer& Peer, int32& MaxRe
 	Activity->Stopwatch.RecvEnd();
 #endif
 
-	Activity_ChangeState(Activity, FActivity::EState::RecvDone);
+	Activity->Done();
 	return FOutcome::Ok();
 }
 
@@ -172,7 +175,7 @@ static FOutcome DoRecvStream(FActivity* Activity, FHttpPeer& Peer, int32& MaxRec
 	FOutcome Outcome = Body.TryRecv(Dest, Peer);
 	if (Outcome.IsError())
 	{
-		Activity_SetError(Activity, Outcome);
+		Activity->Fail(Outcome);
 		return Outcome;
 	}
 
@@ -189,7 +192,7 @@ static FOutcome DoRecvStream(FActivity* Activity, FHttpPeer& Peer, int32& MaxRec
 		FIoBuffer Slice(View, Outer);
 
 		Swap(Outer, Slice);
-		Activity_CallSink(Activity);
+		Activity->CallSink();
 		Swap(Outer, Slice);
 	}
 
@@ -204,7 +207,8 @@ static FOutcome DoRecvStream(FActivity* Activity, FHttpPeer& Peer, int32& MaxRec
 #endif
 
 	*(Activity->Dest) = FIoBuffer();
-	Activity_ChangeState(Activity, FActivity::EState::RecvDone);
+
+	Activity->Done();
 	return FOutcome::Ok();
 }
 
@@ -224,37 +228,6 @@ static FOutcome DoRecv(FActivity* Activity, FHttpPeer& Peer, int32& MaxRecvSize)
 	return FOutcome::Error("unreachable");
 }
 
-////////////////////////////////////////////////////////////////////////////////
-static void DoRecvDone(FActivity* Activity)
-{
-	// Notify the user we've received everything
-	Activity_CallSink(Activity);
-
-	Activity_ChangeState(Activity, FActivity::EState::Completed);
-}
-
-////////////////////////////////////////////////////////////////////////////////
-static void DoCancel(FActivity* Activity)
-{
-	if (Activity->State >= FActivity::EState::Completed)
-	{
-		return;
-	}
-
-	Activity_ChangeState(Activity, FActivity::EState::Cancelled);
-
-	Activity_CallSink(Activity);
-}
-
-////////////////////////////////////////////////////////////////////////////////
-static void DoFail(FActivity* Activity)
-{
-	check(Activity->State == FActivity::EState::Failed);
-
-	// Notify the user we've received everything
-	Activity_CallSink(Activity);
-}
-
 
 
 // {{{1 work-queue .............................................................
@@ -514,7 +487,7 @@ void FWorkQueue::TickCancels(FTickState& State)
 			continue;
 		}
 
-		DoCancel(Activity);
+		Activity->Cancel();
 
 		Activity->Next = State.DoneList;
 		State.DoneList = Activity;
@@ -587,12 +560,8 @@ void FPeerGroup::Fail(FTickState& State, const char* Reason)
 	// Failure is quite terminal and we need to abort everything
 	for (FActivityNode* Activity = Recv; Activity != nullptr;)
 	{
-		if (Activity->State != FActivity::EState::Failed)
-		{
-			Activity_SetError(Activity, Reason);
-		}
-
-		DoFail(Activity);
+		FOutcome Outcome = FOutcome::Error(Reason);
+		Activity->Fail(Outcome);
 
 		FActivityNode* Next = Activity->Next;
 		Activity->Next = State.DoneList;
@@ -697,8 +666,6 @@ void FPeerGroup::RecvInternal(FTickState& State)
 		return;
 	}
 
-	DoRecvDone(Activity);
-
 	Recv = Activity->Next;
 	Activity->Next = State.DoneList;
 	State.DoneList = Activity;
@@ -1049,7 +1016,6 @@ FRequest FEventLoop::FImpl::Request(
 	FActivityNode* Activity)
 {
 	Trace(Activity, ETrace::ActivityCreate, 0);
-	Activity_ChangeState(Activity, FActivity::EState::Build);
 
 	if (Path.Len() == 0)
 	{
@@ -1079,7 +1045,7 @@ FTicket FEventLoop::FImpl::Send(FActivityNode* Activity)
 	Trace(Activity, ETrace::RequestBegin);
 
 	Activity->Buffer.End();
-	Activity_ChangeState(Activity, FActivity::EState::Send);
+	Activity->ChangeState(FActivity::EState::Send);
 
 	uint64 Slot;
 	{
