Tag "Ok" outcomes to communicate progress.

#rnx

diff --git a/Engine/Source/Runtime/Experimental/IoStore/HttpClient/Private/Loop.inl b/Engine/Source/Runtime/Experimental/IoStore/HttpClient/Private/Loop.inl
--- a/Engine/Source/Runtime/Experimental/IoStore/HttpClient/Private/Loop.inl
+++ b/Engine/Source/Runtime/Experimental/IoStore/HttpClient/Private/Loop.inl
@@ -7,6 +7,14 @@ namespace UE::IoStore::HTTP
 
 // {{{1 event-loop-int .........................................................
 
+////////////////////////////////////////////////////////////////////////////////
+enum EStage : int32
+{
+	Send		= -1,
+	Response	= -2,
+	Content		= -3,
+};
+
 ////////////////////////////////////////////////////////////////////////////////
 static FOutcome DoSend(FActivity* Activity, FHttpPeer& Peer)
 {
@@ -37,7 +45,7 @@ static FOutcome DoSend(FActivity* Activity, FHttpPeer& Peer)
 #endif
 
 	Activity->ChangeState(FActivity::EState::RecvMessage);
-	return FOutcome::Ok();
+	return FOutcome::Ok(EStage::Send);
 }
 
 ////////////////////////////////////////////////////////////////////////////////
@@ -119,14 +127,14 @@ static FOutcome DoRecvMessage(FActivity* Activity, FHttpPeer& Peer)
 	if (!Response.HasBody())
 	{
 		Activity->ChangeState(FActivity::EState::RecvDone);
-		return FOutcome::Ok();
+		return FOutcome::Ok(EStage::Content);
 	}
 
 	// We're all set to go and get content
 	check(Activity->Dest != nullptr);
 	auto NextState = bChunked ? FActivity::EState::RecvStream : FActivity::EState::RecvContent;
 	Activity->ChangeState(NextState);
-	return FOutcome::Ok();
+	return FOutcome::Ok(EStage::Response);
 }
 
 ////////////////////////////////////////////////////////////////////////////////
@@ -159,7 +167,7 @@ static FOutcome DoRecvContent(FActivity* Activity, FHttpPeer& Peer, int32& MaxRe
 #endif
 
 	Activity->Done();
-	return FOutcome::Ok();
+	return FOutcome::Ok(EStage::Content);
 }
 
 ////////////////////////////////////////////////////////////////////////////////
@@ -209,7 +217,7 @@ static FOutcome DoRecvStream(FActivity* Activity, FHttpPeer& Peer, int32& MaxRec
 	*(Activity->Dest) = FIoBuffer();
 
 	Activity->Done();
-	return FOutcome::Ok();
+	return FOutcome::Ok(EStage::Content);
 }
 
 ////////////////////////////////////////////////////////////////////////////////
@@ -606,15 +614,7 @@ void FPeerGroup::RecvInternal(FTickState& State)
 	check(bNegotiating == false);
 	check(Recv != nullptr);
 
-	// Another helper lambda
-	auto IsReceiving = [] (const FActivityNode* Act)
-	{
-		using EState = FActivity::EState;
-		return (Act->State >= EState::RecvMessage) & (Act->State < EState::RecvDone);
-	};
-
 	FActivityNode* Activity = Recv;
-	check(IsReceiving(Activity));
 
 	FOutcome Outcome = DoRecv(Activity, Peer, State.RecvAllowance);
 
@@ -659,13 +659,15 @@ void FPeerGroup::RecvInternal(FTickState& State)
 		return;
 	}
 
-	// If we're still in a receiving state we will just try again otherwise it
-	// is finished and we will let DoneList recipient finish it off.
-	if (IsReceiving(Activity))
+	// If an okay result is -1, then message recv is done and content is next
+	// and it is likely that we have it.
+	if (Outcome.GetResult() == EStage::Response)
 	{
-		return;
+		return RecvInternal(State);
 	}
 
+	check(Outcome.GetResult() == EStage::Content);
+
 	Recv = Activity->Next;
 	Activity->Next = State.DoneList;
 	State.DoneList = Activity;
@@ -708,6 +710,7 @@ void FPeerGroup::SendInternal(FTickState& State)
 		return;
 	}
 
+	check(Outcome.GetResult() == EStage::Send);
 	Send = nullptr;
 
 	// Pass along this send to be received
