Added more fleshed out tests for Chaos  sweeps, mostly to compare perf.
Also fixed a small bug in gjk with a deep initial overlap. The margin was not appropriately accounted for after EPA.

diff --git a/Engine/Source/Runtime/Experimental/Chaos/Public/Chaos/GJK.h b/Engine/Source/Runtime/Experimental/Chaos/Public/Chaos/GJK.h
--- a/Engine/Source/Runtime/Experimental/Chaos/Public/Chaos/GJK.h
+++ b/Engine/Source/Runtime/Experimental/Chaos/Public/Chaos/GJK.h
@@ -1826,7 +1826,7 @@ namespace Chaos
 						OutNormal = MTD;
 						VectorStoreFloat1(Penetration, &OutTime);
 						OutTime = -OutTime - (A.Margin + B.Margin);
-						OutPosition = ClosestA;
+						OutPosition = VectorMultiplyAdd(OutNormal, MarginASimd, ClosestA);
 					}
 					else if (EPAResult == EEPAResult::BadInitialSimplex)
 					{
