Added a vectorized GetRelativeTransformNoScale and switched SweepQuery to use the NoScale variant.

diff --git a/Engine/Source/Runtime/Experimental/Chaos/Public/Chaos/GeometryQueries.h b/Engine/Source/Runtime/Experimental/Chaos/Public/Chaos/GeometryQueries.h
--- a/Engine/Source/Runtime/Experimental/Chaos/Public/Chaos/GeometryQueries.h
+++ b/Engine/Source/Runtime/Experimental/Chaos/Public/Chaos/GeometryQueries.h
@@ -217,7 +217,7 @@ namespace Chaos
 		FVec3 LocalPosition(-TNumericLimits<FReal>::Max()); // Make it obvious when things go wrong
 		FVec3 LocalNormal(0);
 
-		const FRigidTransform3 BToATM = BTM.GetRelativeTransform(ATM);
+		const FRigidTransform3 BToATM = BTM.GetRelativeTransformNoScale(ATM);
 		ensure(FMath::IsNearlyEqual(Dir.SizeSquared(), (FReal)1, (FReal)UE_KINDA_SMALL_NUMBER)); // Added to help determine cause of this ensure firing in GJKRaycast2.
 		const FVec3 LocalDir = ATM.InverseTransformVectorNoScale(Dir);
 
diff --git a/Engine/Source/Runtime/Experimental/ChaosCore/Public/Chaos/Transform.h b/Engine/Source/Runtime/Experimental/ChaosCore/Public/Chaos/Transform.h
--- a/Engine/Source/Runtime/Experimental/ChaosCore/Public/Chaos/Transform.h
+++ b/Engine/Source/Runtime/Experimental/ChaosCore/Public/Chaos/Transform.h
@@ -159,6 +159,8 @@ namespace Chaos
 			: BaseTransform(Matrix) {}
 		TRigidTransform(const FMatrix44d& Matrix)
 			: BaseTransform(FMatrix44f(Matrix)) {}
+		TRigidTransform(TransformVectorRegister InRotation, TransformVectorRegister InTranslation)
+			: BaseTransform(InRotation, InTranslation, GlobalVectorConstants::Float1110) {}
 		TRigidTransform(const BaseTransform& Transform)
 			: BaseTransform(Transform) {}
 		template<typename OtherType>
@@ -190,11 +192,31 @@ namespace Chaos
 		// Get the transform which maps from Other to This, ignoring the scale on both.
 		TRigidTransform<FRealSingle, 3> GetRelativeTransformNoScale(const TRigidTransform<FRealSingle, 3>& Other) const
 		{
+#if ENABLE_VECTORIZED_TRANSFORM
+			// A * B(-1) = VQS(B)(-1) (VQS (A))
+			// 
+			// Rotation = Q(B)(-1) * Q(A)
+			// Translation = [Q(B)(-1)*(T(A)-T(B))*Q(B)]
+			// where A = this, B = Other
+
+			//VQTranslation = (  ( T(A).X - T(B).X ),  ( T(A).Y - T(B).Y ), ( T(A).Z - T(B).Z), 0.f );
+			TransformVectorRegister VQTranslation = VectorSet_W0(VectorSubtract(Translation, Other.Translation));
+
+			// Inverse RotatedTranslation
+			TransformVectorRegister VInverseRot = VectorQuaternionInverse(Other.Rotation);
+			TransformVectorRegister VTranslation = VectorQuaternionRotateVector(VInverseRot, VQTranslation);
+
+			// Rotation = Q(B)(-1) * Q(A)	
+			TransformVectorRegister VRotation = VectorQuaternionMultiply2(VInverseRot, Rotation);
+
+			return TRigidTransform<FRealSingle, 3>(VRotation, VTranslation);
+#else
 			// @todo(chaos): optimize
 			TRotation<FRealSingle, 3> OtherInverse = Other.GetRotation().Inverse();
 			return TRigidTransform<FRealSingle, 3>(
 				(OtherInverse * (GetTranslation() - Other.GetTranslation())),
 				OtherInverse * GetRotation());
+#endif
 		}
 
 		TVector<FRealSingle, 3> TransformNormalNoScale(const TVector<FRealSingle, 3>& Normal) const
@@ -288,6 +310,8 @@ namespace Chaos
 			: BaseTransform(Matrix) {}
 		TRigidTransform(const FMatrix44f& Matrix)
 			: BaseTransform(FMatrix44d(Matrix)) {}
+		TRigidTransform(TransformVectorRegister InRotation, TransformVectorRegister InTranslation)
+			: BaseTransform(InRotation, InTranslation, GlobalVectorConstants::Double1110) {}
 		TRigidTransform(const BaseTransform& Transform)
 			: BaseTransform(Transform) {}
 		template<typename OtherType>
@@ -319,11 +343,31 @@ namespace Chaos
 		// Get the transform which maps from Other to This, ignoring the scale on both.
 		TRigidTransform<FRealDouble, 3> GetRelativeTransformNoScale(const TRigidTransform<FRealDouble, 3>& Other) const
 		{
+#if ENABLE_VECTORIZED_TRANSFORM
+			// A * B(-1) = VQS(B)(-1) (VQS (A))
+			// 
+			// Rotation = Q(B)(-1) * Q(A)
+			// Translation = [Q(B)(-1)*(T(A)-T(B))*Q(B)]
+			// where A = this, B = Other
+
+			//VQTranslation = (  ( T(A).X - T(B).X ),  ( T(A).Y - T(B).Y ), ( T(A).Z - T(B).Z), 0.f );
+			TransformVectorRegister VQTranslation = VectorSet_W0(VectorSubtract(Translation, Other.Translation));
+
+			// Inverse RotatedTranslation
+			TransformVectorRegister VInverseRot = VectorQuaternionInverse(Other.Rotation);
+			TransformVectorRegister VTranslation = VectorQuaternionRotateVector(VInverseRot, VQTranslation);
+
+			// Rotation = Q(B)(-1) * Q(A)	
+			TransformVectorRegister VRotation = VectorQuaternionMultiply2(VInverseRot, Rotation);
+
+			return TRigidTransform<FRealDouble, 3>(VRotation, VTranslation);
+#else
 			// @todo(chaos): optimize
 			TRotation<FRealDouble, 3> OtherInverse = Other.GetRotation().Inverse();
 			return TRigidTransform<FRealDouble, 3>(
 				(OtherInverse * (GetTranslation() - Other.GetTranslation())),
 				OtherInverse * GetRotation());
+#endif
 		}
 
 		TVector<FRealDouble, 3> TransformNormalNoScale(const TVector<FRealDouble, 3>& Normal) const
