bulk-append point geometry to the geometry collection facade when creating a point visualization, to avoid the overhead of repeatedly growing the underlying managed arrays when appending large number of points

diff --git a/Engine/Source/Runtime/Experimental/Chaos/Private/GeometryCollection/Facades/CollectionRenderingFacade.cpp b/Engine/Source/Runtime/Experimental/Chaos/Private/GeometryCollection/Facades/CollectionRenderingFacade.cpp
--- a/Engine/Source/Runtime/Experimental/Chaos/Private/GeometryCollection/Facades/CollectionRenderingFacade.cpp
+++ b/Engine/Source/Runtime/Experimental/Chaos/Private/GeometryCollection/Facades/CollectionRenderingFacade.cpp
@@ -125,14 +125,8 @@ namespace GeometryCollection::Facades
 		return 0;
 	}
 
-	void FRenderingFacade::AddPoint(const FVector3f& InPoint)
+	static void AddPointGeometryHelper(const FVector3f& InPoint, TArrayView<FVector3f> Vertices, int32 VertexStart, TArrayView<FIntVector> Tris)
 	{
-		constexpr int32 NumVertices = 15;
-		constexpr int32 NumTriangles = 9;
-
-		TArray<FVector3f> Vertices; Vertices.AddUninitialized(NumVertices);
-		TArray<FIntVector> Tris; Tris.AddUninitialized(NumTriangles);
-
 		constexpr float Extension = 1.f;
 
 		// Add vertices
@@ -172,22 +166,56 @@ namespace GeometryCollection::Facades
 		Vertices[13] = FVector3f(InPoint) + FVector3f(0.f, 0.f, -SmallExtension);
 		Vertices[14] = FVector3f(InPoint) + FVector3f(0.f, 0.f, SmallExtension);
 
-		AddTriangle(Chaos::FTriangle(Vertices[0], Vertices[1], Vertices[2]));
-		AddTriangle(Chaos::FTriangle(Vertices[3], Vertices[4], Vertices[5]));
-		AddTriangle(Chaos::FTriangle(Vertices[6], Vertices[7], Vertices[8]));
-		AddTriangle(Chaos::FTriangle(Vertices[13], Vertices[9], Vertices[14]));
-		AddTriangle(Chaos::FTriangle(Vertices[14], Vertices[10], Vertices[13]));
-		AddTriangle(Chaos::FTriangle(Vertices[9], Vertices[11], Vertices[10]));
-		AddTriangle(Chaos::FTriangle(Vertices[10], Vertices[12], Vertices[9]));
-		AddTriangle(Chaos::FTriangle(Vertices[13], Vertices[11], Vertices[12]));
-		AddTriangle(Chaos::FTriangle(Vertices[12], Vertices[11], Vertices[14]));
+		Tris[0] = FIntVector( VertexStart + 0,  VertexStart + 1,  VertexStart + 2  );
+		Tris[1] = FIntVector( VertexStart + 3,  VertexStart + 4,  VertexStart + 5  );
+		Tris[2] = FIntVector( VertexStart + 6,  VertexStart + 7,  VertexStart + 8  );
+		Tris[3] = FIntVector( VertexStart + 13, VertexStart + 9,  VertexStart + 14 );
+		Tris[4] = FIntVector( VertexStart + 14, VertexStart + 10, VertexStart + 13 );
+		Tris[5] = FIntVector( VertexStart + 9,  VertexStart + 11, VertexStart + 10 );
+		Tris[6] = FIntVector( VertexStart + 10, VertexStart + 12, VertexStart + 9  );
+		Tris[7] = FIntVector( VertexStart + 13, VertexStart + 11, VertexStart + 12 );
+		Tris[8] = FIntVector( VertexStart + 12, VertexStart + 11, VertexStart + 14 );
+	}
+
+	void FRenderingFacade::AddPoint(const FVector3f& InPoint)
+	{
+		check(!IsConst());
+		if (IsValid())
+		{
+			constexpr int32 NumVertices = 15;
+			constexpr int32 NumTriangles = 9;
+
+			TManagedArray<FVector3f>& Vertices = VertexAttribute.Modify();
+			TManagedArray<FIntVector>& Indices = IndicesAttribute.Modify();
+
+			int32 IndicesStart = IndicesAttribute.AddElements(NumTriangles);
+			int32 VertexStart = VertexAttribute.AddElements(NumVertices);
+			TArrayView<FIntVector> TriangleView(Indices.GetData() + IndicesStart, NumTriangles);
+			TArrayView<FVector3f> VertexView(Vertices.GetData() + VertexStart, NumVertices);
+			AddPointGeometryHelper(InPoint, VertexView, VertexStart, TriangleView);
+		}
 	}
 
 	void FRenderingFacade::AddPoints(TArray<FVector3f>&& InPoints)
 	{
-		for (int32 Idx = 0; Idx < InPoints.Num(); ++Idx)
+		check(!IsConst());
+		if (IsValid() && !InPoints.IsEmpty())
 		{
-			AddPoint(InPoints[Idx]);
+			constexpr int32 NumVerticesPerPoint = 15;
+			constexpr int32 NumTrianglesPerPoint = 9;
+			TManagedArray<FVector3f>& Vertices = VertexAttribute.Modify();
+			TManagedArray<FIntVector>& Indices = IndicesAttribute.Modify();
+
+			int32 IndicesStart = IndicesAttribute.AddElements(NumTrianglesPerPoint * InPoints.Num());
+			int32 VertexStart = VertexAttribute.AddElements(NumVerticesPerPoint * InPoints.Num());
+
+			for (int32 Idx = 0; Idx < InPoints.Num(); ++Idx)
+			{
+				TArrayView<FIntVector> TriangleView(Indices.GetData() + IndicesStart + Idx * NumTrianglesPerPoint, NumTrianglesPerPoint);
+				int32 LocalVertexStart = VertexStart + Idx * NumVerticesPerPoint;
+				TArrayView<FVector3f> VertexView(Vertices.GetData() + LocalVertexStart, NumVerticesPerPoint);
+				AddPointGeometryHelper(InPoints[Idx], VertexView, LocalVertexStart, TriangleView);
+			}
 		}
 	}
 
