Placate compiler

#rnx

diff --git a/Engine/Source/Runtime/Experimental/IoStore/HttpClient/Private/ConnectionPool.inl b/Engine/Source/Runtime/Experimental/IoStore/HttpClient/Private/ConnectionPool.inl
--- a/Engine/Source/Runtime/Experimental/IoStore/HttpClient/Private/ConnectionPool.inl
+++ b/Engine/Source/Runtime/Experimental/IoStore/HttpClient/Private/ConnectionPool.inl
@@ -234,7 +234,7 @@ int32 FConnectionPool::FParams::SetHostFromUrl(FAnsiStringView Url)
 ////////////////////////////////////////////////////////////////////////////////
 FConnectionPool::FConnectionPool(const FParams& Params)
 {
-	check(Params.ConnectionCount - 1 <= 63u);
+	check(Params.ConnectionCount - 1u <= 63u);
 	check(Params.Port <= 0xffffu);
 
 	// Alloc a new internal object
diff --git a/Engine/Source/Runtime/Experimental/IoStore/HttpClient/Private/Misc.inl b/Engine/Source/Runtime/Experimental/IoStore/HttpClient/Private/Misc.inl
--- a/Engine/Source/Runtime/Experimental/IoStore/HttpClient/Private/Misc.inl
+++ b/Engine/Source/Runtime/Experimental/IoStore/HttpClient/Private/Misc.inl
@@ -79,7 +79,7 @@ public:
 	bool			IsError() const		{ return Message < 0x8000'0000'0000; }
 	bool			IsWaiting() const	{ return Tag == WaitTag; }
 	bool			IsOk() const		{ return Tag == OkTag; }
-	FAnsiStringView GetMessage() const	{ check(IsError()); return (const char*)(Message); }
+	FAnsiStringView GetMessage() const	{ check(IsError()); return (const char*)UPTRINT(Message); }
 	int32			GetErrorCode() const{ check(IsError()); return int32(Code); }
 	uint32			GetResult() const	{ check(!IsError()); return Result; }
 
@@ -575,7 +575,7 @@ void FBuffer::AdvanceUsed(uint32 Delta)
 ////////////////////////////////////////////////////////////////////////////////
 void FBuffer::Extend(uint32 AtLeast, uint32 PageSize)
 {
-	checkSlow((PageSize - 1 & PageSize) == 0);
+	checkSlow(((PageSize - 1) & PageSize) == 0);
 
 	--PageSize;
 	Max = (AtLeast + PageSize) & ~PageSize;
