Removed rewinding.

#rnx

diff --git a/Engine/Source/Runtime/Experimental/IoStore/HttpClient/Private/Activity.inl b/Engine/Source/Runtime/Experimental/IoStore/HttpClient/Private/Activity.inl
--- a/Engine/Source/Runtime/Experimental/IoStore/HttpClient/Private/Activity.inl
+++ b/Engine/Source/Runtime/Experimental/IoStore/HttpClient/Private/Activity.inl
@@ -937,29 +937,6 @@ const FStopwatch& FActivity::GetStopwatch() const
 
 
 
-////////////////////////////////////////////////////////////////////////////////
-static int32 Activity_Rewind(FActivity* Activity)
-{
-	using EState = FActivity::EState;
-
-	if (Activity->State == EState::Send)
-	{
-		Activity->StateParam = 0;
-		return 0;
-	}
-
-	if (Activity->State == EState::RecvMessage)
-	{
-		Activity->Buffer.Resize(Activity->StateParam);
-		Activity->ChangeState(FActivity::EState::Send);
-		return 1;
-	}
-
-	return -1;
-}
-
-
-
 ////////////////////////////////////////////////////////////////////////////////
 static void Trace(const struct FActivity* Activity, ETrace Action, uint32 Param)
 {
diff --git a/Engine/Source/Runtime/Experimental/IoStore/HttpClient/Private/Loop.inl b/Engine/Source/Runtime/Experimental/IoStore/HttpClient/Private/Loop.inl
--- a/Engine/Source/Runtime/Experimental/IoStore/HttpClient/Private/Loop.inl
+++ b/Engine/Source/Runtime/Experimental/IoStore/HttpClient/Private/Loop.inl
@@ -672,19 +672,11 @@ void FPeerGroup::RecvInternal(FTickState& State)
 	Activity->Next = State.DoneList;
 	State.DoneList = Activity;
 
-	// If the server wants to close the socket we need to rewind the send
-	if (IsKeepAlive != 0)
+	// If the server wants the connection closed then we can drop our peer
+	if (IsKeepAlive == 0)
 	{
-		return;
-	}
-
-	if (Send != nullptr && Activity_Rewind(Send) < 0)
-	{
-		Fail(State, "Unable to rewind on keep-alive close");
-		return;
+		Peer = FHttpPeer();
 	}
-
-	Peer = FHttpPeer();
 }
 
 ////////////////////////////////////////////////////////////////////////////////
