Moved legacy free functions into a single activity tick

#rnx

diff --git a/Engine/Source/Runtime/Experimental/IoStore/HttpClient/Private/Activity.inl b/Engine/Source/Runtime/Experimental/IoStore/HttpClient/Private/Activity.inl
--- a/Engine/Source/Runtime/Experimental/IoStore/HttpClient/Private/Activity.inl
+++ b/Engine/Source/Runtime/Experimental/IoStore/HttpClient/Private/Activity.inl
@@ -741,8 +741,16 @@ struct FActivity
 		bool			bAllowChunked = true;
 	};
 
+	enum EStage : uint32
+	{
+		Request		= 0xffff'ffff,
+		Response	= 0xffff'fffe,
+		Content		= 0xffff'fffd,
+	};
+
 						FActivity(const FParams& Params);
 						~FActivity() = default;
+	FOutcome			Tick(FHttpPeer& Peer, int32* MaxRecvSize=nullptr);
 	void				SetSink(FTicketSink&& InSink, UPTRINT Param);
 	void				SetDestination(FIoBuffer* InDest);
 	void				Cancel();
@@ -776,6 +784,11 @@ struct FActivity
 
 	void				CallSink();
 	void				ChangeState(EState InState, uint32 Param=0);
+	FOutcome			Send(FHttpPeer& Peer);
+	FOutcome			RecvMessage(FHttpPeer& Peer);
+	FOutcome			RecvContent(FHttpPeer& Peer, int32& MaxRecvSize);
+	FOutcome			RecvStream(FHttpPeer& Peer, int32& MaxRecvSize);
+	FOutcome			Recv(FHttpPeer& Peer, int32& MaxRecvSize);
 
 	EState				State = EState::None;
 	uint8				IsKeepAlive : 1;
@@ -935,6 +948,229 @@ const FStopwatch& FActivity::GetStopwatch() const
 }
 #endif // IAS_HTTP_WITH_PERF
 
+////////////////////////////////////////////////////////////////////////////////
+FOutcome FActivity::Send(FHttpPeer& Peer)
+{
+#if IAS_HTTP_WITH_PERF
+	Stopwatch.SendStart();
+#endif
+
+	TRACE_CPUPROFILER_EVENT_SCOPE(IasHttp::DoSend);
+
+	FOutcome Outcome = Buffer.TrySendRequest(Peer);
+	if (Outcome.IsError())
+	{
+		Fail(Outcome);
+		return Outcome;
+	}
+
+	if (Outcome.IsWaiting())
+	{
+		return Outcome;
+	}
+
+	check(Outcome.IsOk());
+
+#if IAS_HTTP_WITH_PERF
+	Stopwatch.SendEnd();
+#endif
+
+	ChangeState(FActivity::EState::RecvMessage);
+	return FOutcome::Ok(EStage::Request);
+}
+
+////////////////////////////////////////////////////////////////////////////////
+FOutcome FActivity::RecvMessage(FHttpPeer& Peer)
+{
+	TRACE_CPUPROFILER_EVENT_SCOPE(IasHttp::DoRecvMessage);
+	Trace(this, ETrace::StateChange, uint32(State));
+
+#if IAS_HTTP_WITH_PERF
+	Stopwatch.RecvStart();
+#endif
+
+	FOutcome Outcome = Buffer.TryRecvResponse(Peer);
+	if (Outcome.IsError())
+	{
+		Fail(Outcome);
+		return Outcome;
+	}
+
+	if (Outcome.IsWaiting())
+	{
+		return Outcome;
+	}
+
+	check(Outcome.IsOk());
+
+	bool bChunked = Buffer.IsChunked();
+	int64 ContentLength = Buffer.GetContentLength();
+
+	// Validate that the server's told us how and how much it will transmit
+	if (bChunked)
+	{
+		if (bAllowChunked == 0)
+		{
+			Outcome = FOutcome::Error("Chunked transfer encoding disabled (ERRNOCHUNK)");
+			Fail(Outcome);
+			return Outcome;
+		}
+	}
+
+	// Call out to the sink to get a content destination
+	FIoBuffer* PriorDest = Dest; // to retain unioned Host ptr (redirect uses it in sink)
+	CallSink();
+
+	// Check the user gave us a destination for content
+	if (Buffer.HasBody())
+	{
+		if (Dest == PriorDest)
+		{
+			Outcome = FOutcome::Error("User did not provide a destination buffer");
+			Fail(Outcome);
+			return Outcome;
+		}
+
+		// The user seems to have forgotten something. Let's help them along
+		if (int32 DestSize = int32(Dest->GetSize()); DestSize == 0)
+		{
+			static const uint32 DefaultChunkSize = 4 << 10;
+			uint32 Size = bChunked ? DefaultChunkSize : uint32(ContentLength);
+			*Dest = FIoBuffer(Size);
+		}
+		else if (!bChunked && DestSize < ContentLength)
+		{
+			// todo: support piece-wise transfer of content (a la chunked).
+			Outcome = FOutcome::Error("Destination buffer too small");
+			Fail(Outcome);
+			return Outcome;
+		}
+		else if (enum { MinStreamBuf = 256 }; bChunked && DestSize < MinStreamBuf)
+		{
+			*Dest = FIoBuffer(MinStreamBuf);
+		}
+	}
+
+	// HEAD methods
+	if (!Buffer.HasBody())
+	{
+		ChangeState(FActivity::EState::RecvDone);
+		return FOutcome::Ok(EStage::Content);
+	}
+
+	// We're all set to go and get content
+	check(Dest != nullptr);
+	auto NextState = bChunked ? FActivity::EState::RecvStream : FActivity::EState::RecvContent;
+	ChangeState(NextState);
+	return FOutcome::Ok(EStage::Response);
+}
+
+////////////////////////////////////////////////////////////////////////////////
+FOutcome FActivity::RecvContent(FHttpPeer& Peer, int32& MaxRecvSize)
+{
+	TRACE_CPUPROFILER_EVENT_SCOPE(IasHttp::DoRecvContent);
+
+	FMutableMemoryView View = Dest->GetMutableView();
+	View = View.Right(View.GetSize() - StateParam).Left(MaxRecvSize);
+	MaxRecvSize -= int32(View.GetSize());
+
+	FOutcome Outcome = Buffer.TryRecv(View, Peer);
+	if (Outcome.IsError())
+	{
+		Fail(Outcome);
+		return Outcome;
+	}
+
+	if (Outcome.IsWaiting())
+	{
+		StateParam += Outcome.GetResult();
+		return Outcome;
+	}
+
+	check(StateParam + Outcome.GetResult() == Buffer.GetContentLength());
+
+#if IAS_HTTP_WITH_PERF
+	Stopwatch.RecvEnd();
+#endif
+
+	Done();
+	return FOutcome::Ok(EStage::Content);
+}
+
+////////////////////////////////////////////////////////////////////////////////
+FOutcome FActivity::RecvStream(FHttpPeer& Peer, int32& MaxRecvSize)
+{
+	TRACE_CPUPROFILER_EVENT_SCOPE(IasHttp::DoRecvStream);
+
+	FMutableMemoryView View = Dest->GetMutableView();
+	View = View.Left(MaxRecvSize);
+	MaxRecvSize -= int32(View.GetSize());
+
+	FOutcome Outcome = Buffer.TryRecv(View, Peer);
+	if (Outcome.IsError())
+	{
+		Fail(Outcome);
+		return Outcome;
+	}
+
+	int32 Result = Outcome.GetResult();
+
+	if (Result != 0)
+	{
+		// Temporarily clamp IoBuffer so if the sink does GetView/GetSize() it
+		// represents actual content and not the underlying working buffer.
+		FIoBuffer& Outer = *Dest;
+
+		FMemoryView SliceView = Outer.GetView();
+		SliceView = SliceView.Left(Result);
+		FIoBuffer Slice(SliceView, Outer);
+
+		Swap(Outer, Slice);
+		CallSink();
+		Swap(Outer, Slice);
+	}
+
+	if (!Outcome.IsOk())
+	{
+		check(Outcome.IsWaiting());
+		return Outcome;
+	}
+
+#if IAS_HTTP_WITH_PERF
+	Stopwatch.RecvEnd();
+#endif
+
+	*Dest = FIoBuffer();
+
+	Done();
+	return FOutcome::Ok(EStage::Content);
+}
+
+////////////////////////////////////////////////////////////////////////////////
+FOutcome FActivity::Recv(FHttpPeer& Peer, int32& MaxRecvSize)
+{
+	check(State >= EState::RecvMessage && State < EState::RecvDone);
+
+	if (State == EState::RecvMessage)	return RecvMessage(Peer);
+	if (State == EState::RecvContent)	return RecvContent(Peer, MaxRecvSize);
+	if (State == EState::RecvStream)	return RecvStream(Peer, MaxRecvSize); //-V547
+	
+	check(false); // it is not expected that we'll get here
+	return FOutcome::Error("unreachable");
+}
+
+////////////////////////////////////////////////////////////////////////////////
+FOutcome FActivity::Tick(FHttpPeer& Peer, int32* MaxRecvSize)
+{
+	if (State == FActivity::EState::Send)
+	{
+		return Send(Peer);
+	}
+
+	check(MaxRecvSize != nullptr);
+	return Recv(Peer, *MaxRecvSize);
+}
+
 
 
 ////////////////////////////////////////////////////////////////////////////////
diff --git a/Engine/Source/Runtime/Experimental/IoStore/HttpClient/Private/Loop.inl b/Engine/Source/Runtime/Experimental/IoStore/HttpClient/Private/Loop.inl
--- a/Engine/Source/Runtime/Experimental/IoStore/HttpClient/Private/Loop.inl
+++ b/Engine/Source/Runtime/Experimental/IoStore/HttpClient/Private/Loop.inl
@@ -5,239 +5,6 @@
 namespace UE::IoStore::HTTP
 {
 
-// {{{1 event-loop-int .........................................................
-
-////////////////////////////////////////////////////////////////////////////////
-enum EStage : int32
-{
-	Send		= -1,
-	Response	= -2,
-	Content		= -3,
-};
-
-////////////////////////////////////////////////////////////////////////////////
-static FOutcome DoSend(FActivity* Activity, FHttpPeer& Peer)
-{
-#if IAS_HTTP_WITH_PERF
-	Activity->Stopwatch.SendStart();
-#endif
-
-	TRACE_CPUPROFILER_EVENT_SCOPE(IasHttp::DoSend);
-
-	auto& Buffer = Activity->Buffer;
-
-	FOutcome Outcome = Buffer.TrySendRequest(Peer);
-	if (Outcome.IsError())
-	{
-		Activity->Fail(Outcome);
-		return Outcome;
-	}
-
-	if (Outcome.IsWaiting())
-	{
-		return Outcome;
-	}
-
-	check(Outcome.IsOk());
-
-#if IAS_HTTP_WITH_PERF
-	Activity->Stopwatch.SendEnd();
-#endif
-
-	Activity->ChangeState(FActivity::EState::RecvMessage);
-	return FOutcome::Ok(EStage::Send);
-}
-
-////////////////////////////////////////////////////////////////////////////////
-static FOutcome DoRecvMessage(FActivity* Activity, FHttpPeer& Peer)
-{
-	TRACE_CPUPROFILER_EVENT_SCOPE(IasHttp::DoRecvMessage);
-	Trace(Activity, ETrace::StateChange, uint32(Activity->State));
-
-#if IAS_HTTP_WITH_PERF
-	Activity->Stopwatch.RecvStart();
-#endif
-
-	auto& Response = Activity->Buffer;
-
-	FOutcome Outcome = Response.TryRecvResponse(Peer);
-	if (Outcome.IsError())
-	{
-		Activity->Fail(Outcome);
-		return Outcome;
-	}
-
-	if (Outcome.IsWaiting())
-	{
-		return Outcome;
-	}
-
-	check(Outcome.IsOk());
-
-	bool bChunked = Response.IsChunked();
-	int64 ContentLength = Response.GetContentLength();
-
-	// Validate that the server's told us how and how much it will transmit
-	if (bChunked)
-	{
-		if (Activity->bAllowChunked == 0)
-		{
-			Outcome = FOutcome::Error("Chunked transfer encoding disabled (ERRNOCHUNK)");
-			Activity->Fail(Outcome);
-			return Outcome;
-		}
-	}
-
-	// Call out to the sink to get a content destination
-	FIoBuffer* PriorDest = Activity->Dest; // to retain unioned Host ptr (redirect uses it in sink)
-	Activity->CallSink();
-
-	// Check the user gave us a destination for content
-	FIoBuffer& Dest = *(Activity->Dest);
-	if (Response.HasBody())
-	{
-		if (&Dest == PriorDest)
-		{
-			Outcome = FOutcome::Error("User did not provide a destination buffer");
-			Activity->Fail(Outcome);
-			return Outcome;
-		}
-
-		// The user seems to have forgotten something. Let's help them along
-		if (int32 DestSize = int32(Dest.GetSize()); DestSize == 0)
-		{
-			static const uint32 DefaultChunkSize = 4 << 10;
-			uint32 Size = bChunked ? DefaultChunkSize : uint32(ContentLength);
-			Dest = FIoBuffer(Size);
-		}
-		else if (!bChunked && DestSize < ContentLength)
-		{
-			// todo: support piece-wise transfer of content (a la chunked).
-			Outcome = FOutcome::Error("Destination buffer too small");
-			Activity->Fail(Outcome);
-			return Outcome;
-		}
-		else if (enum { MinStreamBuf = 256 }; bChunked && DestSize < MinStreamBuf)
-		{
-			Dest = FIoBuffer(MinStreamBuf);
-		}
-	}
-
-	// HEAD methods
-	if (!Response.HasBody())
-	{
-		Activity->ChangeState(FActivity::EState::RecvDone);
-		return FOutcome::Ok(EStage::Content);
-	}
-
-	// We're all set to go and get content
-	check(Activity->Dest != nullptr);
-	auto NextState = bChunked ? FActivity::EState::RecvStream : FActivity::EState::RecvContent;
-	Activity->ChangeState(NextState);
-	return FOutcome::Ok(EStage::Response);
-}
-
-////////////////////////////////////////////////////////////////////////////////
-static FOutcome DoRecvContent(FActivity* Activity, FHttpPeer& Peer, int32& MaxRecvSize)
-{
-	TRACE_CPUPROFILER_EVENT_SCOPE(IasHttp::DoRecvContent);
-
-	FMutableMemoryView Dest = Activity->Dest->GetMutableView();
-	Dest = Dest.Right(Dest.GetSize() - Activity->StateParam).Left(MaxRecvSize);
-	MaxRecvSize -= int32(Dest.GetSize());
-
-	auto& Body = Activity->Buffer;
-	FOutcome Outcome = Body.TryRecv(Dest, Peer);
-	if (Outcome.IsError())
-	{
-		Activity->Fail(Outcome);
-		return Outcome;
-	}
-
-	if (Outcome.IsWaiting())
-	{
-		Activity->StateParam += Outcome.GetResult();
-		return Outcome;
-	}
-
-	check(Activity->StateParam + Outcome.GetResult() == Body.GetContentLength());
-
-#if IAS_HTTP_WITH_PERF
-	Activity->Stopwatch.RecvEnd();
-#endif
-
-	Activity->Done();
-	return FOutcome::Ok(EStage::Content);
-}
-
-////////////////////////////////////////////////////////////////////////////////
-static FOutcome DoRecvStream(FActivity* Activity, FHttpPeer& Peer, int32& MaxRecvSize)
-{
-	TRACE_CPUPROFILER_EVENT_SCOPE(IasHttp::DoRecvStream);
-
-	FMutableMemoryView Dest = Activity->Dest->GetMutableView();
-	Dest = Dest.Left(MaxRecvSize);
-	MaxRecvSize -= int32(Dest.GetSize());
-
-	auto& Body = Activity->Buffer;
-	FOutcome Outcome = Body.TryRecv(Dest, Peer);
-	if (Outcome.IsError())
-	{
-		Activity->Fail(Outcome);
-		return Outcome;
-	}
-
-	int32 Result = Outcome.GetResult();
-
-	if (Result != 0)
-	{
-		// Temporarily clamp IoBuffer so if the sink does GetView/GetSize() it
-		// represents actual content and not the underlying working buffer.
-		FIoBuffer& Outer = *(Activity->Dest);
-
-		FMemoryView View = Outer.GetView();
-		View = View.Left(Result);
-		FIoBuffer Slice(View, Outer);
-
-		Swap(Outer, Slice);
-		Activity->CallSink();
-		Swap(Outer, Slice);
-	}
-
-	if (!Outcome.IsOk())
-	{
-		check(Outcome.IsWaiting());
-		return Outcome;
-	}
-
-#if IAS_HTTP_WITH_PERF
-	Activity->Stopwatch.RecvEnd();
-#endif
-
-	*(Activity->Dest) = FIoBuffer();
-
-	Activity->Done();
-	return FOutcome::Ok(EStage::Content);
-}
-
-////////////////////////////////////////////////////////////////////////////////
-static FOutcome DoRecv(FActivity* Activity, FHttpPeer& Peer, int32& MaxRecvSize)
-{
-	using EState = FActivity::EState;
-
-	EState State = Activity->State; 
-	check(State >= EState::RecvMessage && State < EState::RecvDone);
-
-	if (State == EState::RecvMessage)	return DoRecvMessage(Activity, Peer);
-	if (State == EState::RecvContent)	return DoRecvContent(Activity, Peer, MaxRecvSize);
-	if (State == EState::RecvStream)	return DoRecvStream(Activity, Peer, MaxRecvSize); //-V547
-	
-	check(false); // it is not expected that we'll get here
-	return FOutcome::Error("unreachable");
-}
-
-
-
 // {{{1 work-queue .............................................................
 
 /*
@@ -616,7 +383,7 @@ void FPeerGroup::RecvInternal(FTickState& State)
 
 	FActivityNode* Activity = Recv;
 
-	FOutcome Outcome = DoRecv(Activity, Peer, State.RecvAllowance);
+	FOutcome Outcome = Activity->Tick(Peer, &State.RecvAllowance);
 
 	// Any sort of error here is unrecoverable
 	if (Outcome.IsError())
@@ -661,12 +428,12 @@ void FPeerGroup::RecvInternal(FTickState& State)
 
 	// If an okay result is -1, then message recv is done and content is next
 	// and it is likely that we have it.
-	if (Outcome.GetResult() == EStage::Response)
+	if (Outcome.GetResult() == FActivity::EStage::Response)
 	{
 		return RecvInternal(State);
 	}
 
-	check(Outcome.GetResult() == EStage::Content);
+	check(Outcome.GetResult() == FActivity::EStage::Content);
 
 	Recv = Activity->Next;
 	Activity->Next = State.DoneList;
@@ -688,7 +455,7 @@ void FPeerGroup::SendInternal(FTickState& State)
 
 	FActivityNode* Activity = Send;
 
-	FOutcome Outcome = DoSend(Activity, Peer);
+	FOutcome Outcome = Activity->Tick(Peer);
 
 	if (Outcome.IsWaiting())
 	{
@@ -702,7 +469,7 @@ void FPeerGroup::SendInternal(FTickState& State)
 		return;
 	}
 
-	check(Outcome.GetResult() == EStage::Send);
+	check(Outcome.GetResult() == FActivity::EStage::Request);
 	Send = nullptr;
 
 	// Pass along this send to be received
