Variable rename.

#rnx

diff --git a/Engine/Source/Runtime/Experimental/IoStore/HttpClient/Private/Activity.inl b/Engine/Source/Runtime/Experimental/IoStore/HttpClient/Private/Activity.inl
--- a/Engine/Source/Runtime/Experimental/IoStore/HttpClient/Private/Activity.inl
+++ b/Engine/Source/Runtime/Experimental/IoStore/HttpClient/Private/Activity.inl
@@ -806,7 +806,7 @@ struct FActivity
 	};
 	UPTRINT				SinkParam;
 	FTicketSink			Sink;
-	FBodyCloak			Buffer;
+	FBodyCloak			Transaction;
 };
 
 ////////////////////////////////////////////////////////////////////////////////
@@ -814,7 +814,7 @@ FActivity::FActivity(const FParams& Params)
 : IsKeepAlive(Params.bIsKeepAlive)
 , bAllowChunked(Params.bAllowChunked)
 , Host(Params.Host)
-, Buffer(Params.Buffer, Params.BufferSize)
+, Transaction(Params.Buffer, Params.BufferSize)
 {
 	check(Host != nullptr);
 	uint32 ContentEstKiB = (Params.ContentSizeEst + 1023) >> 10;
@@ -921,7 +921,7 @@ uint32 FActivity::GetRemainingKiB() const
 	if (State <= FActivity::EState::RecvStream)  return MAX_uint32;
 	if (State >  FActivity::EState::RecvContent) return 0;
 
-	uint32 ContentLength = uint32(Buffer.GetContentLength());
+	uint32 ContentLength = uint32(Transaction.GetContentLength());
 	check(StateParam <= ContentLength);
 	return (ContentLength - StateParam) >> 10;
 }
@@ -957,7 +957,7 @@ FOutcome FActivity::Send(FHttpPeer& Peer)
 
 	TRACE_CPUPROFILER_EVENT_SCOPE(IasHttp::DoSend);
 
-	FOutcome Outcome = Buffer.TrySendRequest(Peer);
+	FOutcome Outcome = Transaction.TrySendRequest(Peer);
 	if (Outcome.IsError())
 	{
 		Fail(Outcome);
@@ -989,7 +989,7 @@ FOutcome FActivity::RecvMessage(FHttpPeer& Peer)
 	Stopwatch.RecvStart();
 #endif
 
-	FOutcome Outcome = Buffer.TryRecvResponse(Peer);
+	FOutcome Outcome = Transaction.TryRecvResponse(Peer);
 	if (Outcome.IsError())
 	{
 		Fail(Outcome);
@@ -1003,8 +1003,8 @@ FOutcome FActivity::RecvMessage(FHttpPeer& Peer)
 
 	check(Outcome.IsOk());
 
-	bool bChunked = Buffer.IsChunked();
-	int64 ContentLength = Buffer.GetContentLength();
+	bool bChunked = Transaction.IsChunked();
+	int64 ContentLength = Transaction.GetContentLength();
 
 	// Validate that the server's told us how and how much it will transmit
 	if (bChunked)
@@ -1022,7 +1022,7 @@ FOutcome FActivity::RecvMessage(FHttpPeer& Peer)
 	CallSink();
 
 	// Check the user gave us a destination for content
-	if (Buffer.HasBody())
+	if (Transaction.HasBody())
 	{
 		if (Dest == PriorDest)
 		{
@@ -1052,7 +1052,7 @@ FOutcome FActivity::RecvMessage(FHttpPeer& Peer)
 	}
 
 	// HEAD methods
-	if (!Buffer.HasBody())
+	if (!Transaction.HasBody())
 	{
 		ChangeState(FActivity::EState::RecvDone);
 		return FOutcome::Ok(EStage::Content);
@@ -1074,7 +1074,7 @@ FOutcome FActivity::RecvContent(FHttpPeer& Peer, int32& MaxRecvSize)
 	View = View.Right(View.GetSize() - StateParam).Left(MaxRecvSize);
 	MaxRecvSize -= int32(View.GetSize());
 
-	FOutcome Outcome = Buffer.TryRecv(View, Peer);
+	FOutcome Outcome = Transaction.TryRecv(View, Peer);
 	if (Outcome.IsError())
 	{
 		Fail(Outcome);
@@ -1087,7 +1087,7 @@ FOutcome FActivity::RecvContent(FHttpPeer& Peer, int32& MaxRecvSize)
 		return Outcome;
 	}
 
-	check(StateParam + Outcome.GetResult() == Buffer.GetContentLength());
+	check(StateParam + Outcome.GetResult() == Transaction.GetContentLength());
 
 #if IAS_HTTP_WITH_PERF
 	Stopwatch.RecvEnd();
@@ -1106,7 +1106,7 @@ FOutcome FActivity::RecvStream(FHttpPeer& Peer, int32& MaxRecvSize)
 	View = View.Left(MaxRecvSize);
 	MaxRecvSize -= int32(View.GetSize());
 
-	FOutcome Outcome = Buffer.TryRecv(View, Peer);
+	FOutcome Outcome = Transaction.TryRecv(View, Peer);
 	if (Outcome.IsError())
 	{
 		Fail(Outcome);
diff --git a/Engine/Source/Runtime/Experimental/IoStore/HttpClient/Private/Api.inl b/Engine/Source/Runtime/Experimental/IoStore/HttpClient/Private/Api.inl
--- a/Engine/Source/Runtime/Experimental/IoStore/HttpClient/Private/Api.inl
+++ b/Engine/Source/Runtime/Experimental/IoStore/HttpClient/Private/Api.inl
@@ -51,7 +51,7 @@ FRequest&& FRequest::Accept(FAnsiStringView MimeType)
 ////////////////////////////////////////////////////////////////////////////////
 FRequest&& FRequest::Header(FAnsiStringView Key, FAnsiStringView Value)
 {
-	Ptr->Buffer.AddHeader(Key, Value);
+	Ptr->Transaction.AddHeader(Key, Value);
 	return MoveTemp(*this);
 }
 
@@ -75,21 +75,21 @@ EStatusCodeClass FResponse::GetStatus() const
 uint32 FResponse::GetStatusCode() const
 {
 	const auto* Activity = (const FActivity*)this;
-	return Activity->Buffer.GetStatusCode();
+	return Activity->Transaction.GetStatusCode();
 }
 
 ////////////////////////////////////////////////////////////////////////////////
 FAnsiStringView FResponse::GetStatusMessage() const
 {
 	const auto* Activity = (const FActivity*)this;
-	return Activity->Buffer.GetStatusMessage();
+	return Activity->Transaction.GetStatusMessage();
 }
 
 ////////////////////////////////////////////////////////////////////////////////
 int64 FResponse::GetContentLength() const
 {
 	const auto* Activity = (const FActivity*)this;
-	return Activity->Buffer.GetContentLength();
+	return Activity->Transaction.GetContentLength();
 }
 
 ////////////////////////////////////////////////////////////////////////////////
@@ -142,7 +142,7 @@ FAnsiStringView FResponse::GetHeader(FAnsiStringView Name) const
 void FResponse::ReadHeaders(FHeaderSink Sink) const
 {
 	const auto* Activity = (const FActivity*)this;
-	Activity->Buffer.ReadHeaders(Sink);
+	Activity->Transaction.ReadHeaders(Sink);
 }
 
 ////////////////////////////////////////////////////////////////////////////////
@@ -206,7 +206,7 @@ uint32 FTicketStatus::GetContentLength() const
 {
 	check(GetId() <= EId::Content);
 	const auto* Activity = (FActivity*)this;
-	return uint32(Activity->Buffer.GetContentLength());
+	return uint32(Activity->Transaction.GetContentLength());
 }
 
 ////////////////////////////////////////////////////////////////////////////////
@@ -263,7 +263,7 @@ FTicketPerf::FSample FTicketPerf::GetSample() const
 	uint32 Bps = ~0u;
 	if (Sample.RecvMs)
 	{
-		Bps = uint32((Activity.Buffer.GetContentLength() * 1000ull) / Sample.RecvMs);
+		Bps = uint32((Activity.Transaction.GetContentLength() * 1000ull) / Sample.RecvMs);
 	}
 	Sample.RecvKiBps = Clamp(Bps >> 10);
 
diff --git a/Engine/Source/Runtime/Experimental/IoStore/HttpClient/Private/Loop.inl b/Engine/Source/Runtime/Experimental/IoStore/HttpClient/Private/Loop.inl
--- a/Engine/Source/Runtime/Experimental/IoStore/HttpClient/Private/Loop.inl
+++ b/Engine/Source/Runtime/Experimental/IoStore/HttpClient/Private/Loop.inl
@@ -392,7 +392,7 @@ void FPeerGroup::RecvInternal(FTickState& State)
 		return;
 	}
 
-	IsKeepAlive &= uint8(Activity->Buffer.IsKeepAlive() == true);
+	IsKeepAlive &= uint8(Activity->Transaction.IsKeepAlive() == true);
 	LastUseMs = State.NowMs;
 	bWaiting |= Outcome.IsWaiting();
 
@@ -786,14 +786,14 @@ FRequest FEventLoop::FImpl::Request(
 
 	const FHost* Host = Activity->GetHost();
 
-	auto& Buffer = Activity->Buffer;
-	Buffer.Begin(Method, Path);
-	Buffer.AddHeader("Host", Host->GetHostName());
+	auto& Transaction = Activity->Transaction;
+	Transaction.Begin(Method, Path);
+	Transaction.AddHeader("Host", Host->GetHostName());
 
 	// HTTP/1.1 is persistent by default thus "Connection" header isn't required
 	if (!Host->IsPooled())
 	{
-		Buffer.AddHeader("Connection", "close");
+		Transaction.AddHeader("Connection", "close");
 	}
 
 	FRequest Ret;
@@ -806,7 +806,7 @@ FTicket FEventLoop::FImpl::Send(FActivityNode* Activity)
 {
 	Trace(Activity, ETrace::RequestBegin);
 
-	Activity->Buffer.End();
+	Activity->Transaction.End();
 	Activity->ChangeState(FActivity::EState::Send);
 
 	uint64 Slot;
@@ -1099,10 +1099,10 @@ bool FEventLoop::Redirect(const FTicketStatus& Status, FTicketSink& OuterSink)
 
 	const auto& Activity = (FActivity&)Response; // todo: yuk
 
-	check(!Activity.Buffer.HasBody() || Response.GetContentLength() == 0); // should we ever hit this, we'll fix it
+	check(!Activity.Transaction.HasBody() || Response.GetContentLength() == 0); // should we ever hit this, we'll fix it
 
 	// Original method should remain unchanged
-	FAnsiStringView Method = Activity.Buffer.GetMethod();
+	FAnsiStringView Method = Activity.Transaction.GetMethod();
 	check(!Method.IsEmpty());
 
 	FRequestParams RequestParams = {
@@ -1135,7 +1135,7 @@ bool FEventLoop::Redirect(const FTicketStatus& Status, FTicketSink& OuterSink)
 	}
 
 	// Transfer original request headers
-	Activity.Buffer.EnumerateHeaders([&ForwardRequest] (FAnsiStringView Name, FAnsiStringView Value)
+	Activity.Transaction.EnumerateHeaders([&ForwardRequest] (FAnsiStringView Name, FAnsiStringView Value)
 	{
 		if (Name != "Host" && Name != "Connection")
 		{
