Post-step waking based on collision solver impulse

Post-step waking is executed post-solve right after putting non-moving particles asleep (partial island sleep).
All awake-sleeping particle pairs connected by constraints are traversed.
If the velocity-based dynamics (VBD) collision solver has just applied an impulse above a user-defined threshold, we wake up the sleeping particle for the next tick.
Only one layer of sleeping particles can wake up per tick (i.e. propagation depth = 1 through the constraint graph).

The solver impulse is queried through a virtual function and is based on the assumption that the VBD solver applies significant impulses in case of impacts only, not in case of sustained contact.
Heavy impacts and collisions will therefore wake up the sleep particles, but stable stacks and piles should remain asleep.

diff --git a/Engine/Source/Programs/HeadlessChaos/Private/HeadlessChaosTestSleeping.cpp b/Engine/Source/Programs/HeadlessChaos/Private/HeadlessChaosTestSleeping.cpp
--- a/Engine/Source/Programs/HeadlessChaos/Private/HeadlessChaosTestSleeping.cpp
+++ b/Engine/Source/Programs/HeadlessChaos/Private/HeadlessChaosTestSleeping.cpp
@@ -895,5 +895,69 @@ namespace ChaosTest
 		CVarMomentumPropagation->Set(MomentumPropagationDefault);
 	}
 
+	TEST_P(SleepingTests, PostStepWaking)
+	{
+		const int32 NumOfStackedObjects = 5;
+		const FReal Dimension = 200;
+		const bool bPartialSleepEnabled = true;
+		const FRealSingle PostStepWakeThreshold = GetParam() ? 10.0f : 0.0f;
+
+		IConsoleVariable* CVarPostStepWakeThreshold = IConsoleManager::Get().FindConsoleVariable(TEXT("p.Chaos.Solver.Sleep.PostStepWakeThreshold"), false);
+		check(CVarPostStepWakeThreshold && CVarPostStepWakeThreshold->IsVariableFloat());
+		const FRealSingle PostStepWakeThresholdDefault = CVarPostStepWakeThreshold->GetFloat();
+		CVarPostStepWakeThreshold->Set(PostStepWakeThreshold);
+
+		// Create a stack of cubes
+		FSleepingTest Test(bPartialSleepEnabled);
+		Test.MakeQueueOfCubes(NumOfStackedObjects, /*Size=*/Dimension);
+
+		// Simulate until all particles are sleeping
+		Test.AdvanceUntilSleeping();
+
+		// Check that everything is sleeping
+		for (FRigidBodyHandle_External* ParticleHandle : Test.ParticleHandles)
+		{
+			EXPECT_TRUE(ParticleHandle->ObjectState() == EObjectStateType::Sleeping);
+		}
+
+		// Create another particle which will collide with the queue
+		FRigidBodyHandle_External* CollidingCube = Test.MakeCube(FVec3(0, -200, 100), 200);
+
+		for (int32 Step = 0; Step < 10; ++Step)
+		{
+			CollidingCube->SetLinearImpulse(FVec3(0, 1000, 0), true);
+			Test.Advance();
+
+			//// Debug print of all particle positions
+			//for (int32 Id = 0; Id < Test.ParticleHandles.Num(); ++Id)
+			//{
+			//	FRigidBodyHandle_External* ParticleHandle = Test.ParticleHandles[Id];
+			//	printf("Particle Id: %d, State: %d, Position: %f, %f, %f\n", Id,
+			//		ParticleHandle->ObjectState(), ParticleHandle->GetX().X, ParticleHandle->GetX().Y, ParticleHandle->GetX().Z);
+			//}
+		}
+
+		if (GetParam()) // post-step waking enabled
+		{
+			// Check that everything is awake
+			for (FRigidBodyHandle_External* ParticleHandle : Test.ParticleHandles)
+			{
+				EXPECT_TRUE(ParticleHandle->ObjectState() == EObjectStateType::Dynamic);
+			}
+		}
+		else // post-step waking disabled
+		{
+			// Only the colliding cube and first particle in the queue wake up
+			EXPECT_TRUE(Test.ParticleHandles[5]->ObjectState() == EObjectStateType::Dynamic);
+			EXPECT_TRUE(Test.ParticleHandles[0]->ObjectState() == EObjectStateType::Dynamic);
+			EXPECT_TRUE(Test.ParticleHandles[1]->ObjectState() == EObjectStateType::Sleeping);
+			EXPECT_TRUE(Test.ParticleHandles[2]->ObjectState() == EObjectStateType::Sleeping);
+			EXPECT_TRUE(Test.ParticleHandles[3]->ObjectState() == EObjectStateType::Sleeping);
+			EXPECT_TRUE(Test.ParticleHandles[4]->ObjectState() == EObjectStateType::Sleeping);
+		}
+
+		CVarPostStepWakeThreshold->Set(PostStepWakeThresholdDefault);
+	}
+
 	INSTANTIATE_TEST_SUITE_P(, SleepingTests, testing::Bool());
 }
\ No newline at end of file
diff --git a/Engine/Source/Runtime/Experimental/Chaos/Private/Chaos/ChaosDebugDraw.cpp b/Engine/Source/Runtime/Experimental/Chaos/Private/Chaos/ChaosDebugDraw.cpp
--- a/Engine/Source/Runtime/Experimental/Chaos/Private/Chaos/ChaosDebugDraw.cpp
+++ b/Engine/Source/Runtime/Experimental/Chaos/Private/Chaos/ChaosDebugDraw.cpp
@@ -1910,7 +1910,8 @@ namespace Chaos
 					DrawDebugDirectionalArrow(Xb, Xb + Settings.DrawScale * Settings.ConstraintAxisLen * SpaceTransform.TransformVector(Rb.GetAxis(2)), Settings.DrawScale * Settings.ArrowSize, Y, false, UE_KINDA_SMALL_NUMBER, uint8(Settings.DrawPriority), Settings.LineThickness);
 				}
 
-				// NOTE: GetLinearImpulse is the positional impulse (pushout)
+				// NOTE: GetLinearImpulse is the velocity impulse
+				// (see PBDJointContainerSolver.cpp: const FVec3 LinearImpulse = -Solver.GetNetLinearImpulse() / Dt;)
 				if ((Settings.PushOutScale > 0) && !ConstraintHandle->GetLinearImpulse().IsNearlyZero())
 				{
 					FColor PushOutImpusleColor = FColor(0, 250, 250);
diff --git a/Engine/Source/Runtime/Experimental/Chaos/Private/Chaos/Collision/PBDCollisionConstraint.cpp b/Engine/Source/Runtime/Experimental/Chaos/Private/Chaos/Collision/PBDCollisionConstraint.cpp
--- a/Engine/Source/Runtime/Experimental/Chaos/Private/Chaos/Collision/PBDCollisionConstraint.cpp
+++ b/Engine/Source/Runtime/Experimental/Chaos/Private/Chaos/Collision/PBDCollisionConstraint.cpp
@@ -544,6 +544,16 @@ namespace Chaos
 		ContainerCookie.MidPhase->SetIsSleeping(bInIsSleeping, ConcreteContainer()->GetConstraintAllocator().GetCurrentEpoch());
 	}
 
+	Chaos::FVec3f FPBDCollisionConstraint::GetNetLinearImpulse() const
+	{
+		FVec3f TotalNetImpulse = FVec3f(0, 0, 0);
+		for (const FManifoldPointResult& ManifoldPoint : ManifoldPointResults)
+		{
+			TotalNetImpulse += ManifoldPoint.NetImpulse;
+		}
+		return TotalNetImpulse;
+	}
+
 	void FPBDCollisionConstraint::UpdateParticleTransform(FGeometryParticleHandle* InParticle)
 	{
 		// We need to update the friction anchors based on how much the user moved the particle, but
diff --git a/Engine/Source/Runtime/Experimental/Chaos/Private/Chaos/Island/IslandManager.cpp b/Engine/Source/Runtime/Experimental/Chaos/Private/Chaos/Island/IslandManager.cpp
--- a/Engine/Source/Runtime/Experimental/Chaos/Private/Chaos/Island/IslandManager.cpp
+++ b/Engine/Source/Runtime/Experimental/Chaos/Private/Chaos/Island/IslandManager.cpp
@@ -66,9 +66,15 @@ namespace Chaos::CVars
 	/** Cvar to enable/disable triggering wake up for newly detected constraints with moving particles only (partial island sleeping only) */
 	bool bChaosSolverWakeOnMovingOnly = true;
 	FAutoConsoleVariableRef CVarChaosSolverWakeOnMovingOnly(TEXT("p.Chaos.Solver.Sleep.WakeOnMovingOnly"), bChaosSolverWakeOnMovingOnly, TEXT("[EXPERIMENTAL] Only newly detected constraints connected to a moving particle will trigger wake-up event (partial island sleep only)."));
+	
+	/** Cvar to adjust post-step waking triggered by velocity solver impulses.
+	* The solver impulse is converted to a change in velocity by division by the particle mass.
+	* Ideal setting = 10.0f, i.e. change in velocity > 10 will trigger wake of the sleeping particle. */
+	FRealSingle ChaosSolverPostStepWakeThreshold = 0.0f;
+	FAutoConsoleVariableRef CVarChaosSolverPostStepWakeThreshold(TEXT("p.Chaos.Solver.Sleep.PostStepWakeThreshold"), ChaosSolverPostStepWakeThreshold, TEXT("[EXPERIMENTAL] Post-solver waking of sleeping particles: Minimum squared change of velocity induced by the velocity-based collision solver (VBD) or the joint solver (partial island sleeping only, 0 to disable).[def:10]"));
 
 	/** Cvar to enable/disable momentum propagation through the constraint graph (partial island sleeping only) */
-	FRealSingle bChaosSolverMomentumPropagation = 0.25f;
+	FRealSingle bChaosSolverMomentumPropagation = 0.0f;
 	FAutoConsoleVariableRef CVarChaosSolverMomentumPropagation(TEXT("p.Chaos.Solver.Sleep.MomentumPropagation"), bChaosSolverMomentumPropagation, TEXT("[EXPERIMENTAL] Ratio of the impact momentum transferred from one sleeping particle to another through a constraint during partial wake event (partial island sleeping only, 0 to disable, 1 to propagate with dissipation).[def:0.25]"));
 
 	/** Cvar to override the sleep counter threshold if necessary */
@@ -2536,7 +2542,16 @@ namespace Chaos::Private
 		SleepNonMovingParticles(Island, Dt);
 
 		// Deactivate all adjacent constraints between (a) 2 sleeping particles or (b) a sleeping and a non-dynamic particle
+		// Processing non-sleeping constraints is sufficient because no particles were woken up since the last call to @SyncConstraintSleepStateToParticles
 		PropagateParticleSleepToConstraints(Island);
+
+		if (CVars::ChaosSolverPostStepWakeThreshold > 0.0f)
+		{
+			// Process all constraints connecting active-sleeping pairs.
+			// Wake sleeping body if the constraint solver applies a significant impulse.
+			// Only one layer of sleeping particles can be woken up per tick (i.e. no propagation through layers).
+			PostStepPartialWake(Island);
+		}
 	}
 
 	void FPBDIslandManager::SleepNonMovingParticles(FPBDIsland* Island, const FRealSingle Dt)
@@ -2664,6 +2679,96 @@ namespace Chaos::Private
 		}
 	}
 
+	void FPBDIslandManager::PostStepPartialWake(FPBDIsland* Island)
+	{
+		const auto IsSleeping = [](const FPBDIslandParticle* Node) -> bool
+			{
+				if (Node == nullptr)
+				{
+					return false;
+				}
+				const FPBDRigidParticleHandle* Rigid = Node->GetParticle()->CastToRigidParticle();
+				if (Rigid == nullptr)
+				{
+					return false;
+				}
+
+				// NOTE: Make sure to use the particle flag here in order not to consider particles
+				// newly woken up by this method when checking active-sleeping pairs
+				return Rigid->IsSleeping();
+			};
+
+		const auto CheckWakeActiveSleepingPair = [](FPBDIslandParticle* Node, FPBDIslandConstraint* IslandEdge, TArray<FGeometryParticleHandle*>& WakeUpQueue)
+			{
+				if (FPBDRigidParticleHandle* Rigid = Node->GetParticle()->CastToRigidParticle())
+				{
+					FConstraintHandle* Constraint = IslandEdge->GetConstraint();
+
+					// NetImpulse = Mass * change in velocity
+					const FVec3f VelocitySolverImpulse = Constraint->GetNetLinearImpulse();
+					const FVec3f DiffVelocity = VelocitySolverImpulse / (FRealSingle)Rigid->M();
+
+					if (DiffVelocity.SizeSquared() > FMath::Square(CVars::ChaosSolverPostStepWakeThreshold))
+					{
+						Node->Flags.bIsSleeping = false;
+						WakeUpQueue.Add(Node->GetParticle());
+						// Set the sleep counter so that the particle can sleep again after 1 step.
+						Rigid->SetSleepCounter(int8(Node->SleepCounterThreshold));
+
+						// If a particle is active, we know for sure that all connected constraints need to be active too.
+						// Only constraints between sleeping particles and kinematics can sleep.
+						for (FPBDIslandConstraint* Edge : Node->Edges)
+						{
+							if (Edge->Flags.bIsSleeping)
+							{
+								Edge->Flags.bIsSleeping = false;
+								if (FConstraintHandle* ConnectedContraint = Edge->GetConstraint())
+								{
+									ConnectedContraint->SetIsSleeping(false);
+								}
+							}
+						}
+					}
+				}
+			};
+
+		TArray<FGeometryParticleHandle*> WakeUpQueue;
+		for (TArray<FPBDIslandConstraint*>& IslandEdges : Island->ContainerEdges)
+		{
+			for (FPBDIslandConstraint* IslandEdge : IslandEdges)
+			{
+				// Skip already sleeping constraints
+				if (!IslandEdge->Flags.bIsSleeping)
+				{
+					FPBDIslandParticle* Node0 = IslandEdge->Nodes[0];
+					FPBDIslandParticle* Node1 = IslandEdge->Nodes[1];
+					FConstraintHandle* Constraint = IslandEdge->GetConstraint();
+
+					if (Node0 != nullptr && Node1 != nullptr)
+					{
+						const bool bIsSleeping0 = IsSleeping(Node0);
+						const bool bIsSleeping1 = IsSleeping(Node1);
+						const bool bIsAwake0 = Node0->Flags.bIsDynamic && !bIsSleeping0;
+						const bool bIsAwake1 = Node1->Flags.bIsDynamic && !bIsSleeping1;
+
+						// Reactivate the body if:
+						// (a) One particle is sleeping and the other is awake AND
+						// (b) The constraint between them has a significant net impulse.
+						if (bIsAwake0 && bIsSleeping1)
+						{
+							CheckWakeActiveSleepingPair(Node1, IslandEdge, WakeUpQueue);
+						}
+						else if (bIsAwake1 && bIsSleeping0)
+						{
+							CheckWakeActiveSleepingPair(Node0, IslandEdge, WakeUpQueue);
+						}
+					}
+				}
+			}
+		}
+		WakeUpParticles(WakeUpQueue, /*bResetSleepCounter=*/false);
+	}
+
 	void FPBDIslandManager::ProcessPartialIslandWake(FPBDIsland* Island)
 	{
 		TArray<FGeometryParticleHandle*> ParticlesToWakeUp; // Particles which will be woken up after processing all partial wake events.
diff --git a/Engine/Source/Runtime/Experimental/Chaos/Public/Chaos/Collision/PBDCollisionConstraint.h b/Engine/Source/Runtime/Experimental/Chaos/Public/Chaos/Collision/PBDCollisionConstraint.h
--- a/Engine/Source/Runtime/Experimental/Chaos/Public/Chaos/Collision/PBDCollisionConstraint.h
+++ b/Engine/Source/Runtime/Experimental/Chaos/Public/Chaos/Collision/PBDCollisionConstraint.h
@@ -420,6 +420,9 @@ namespace Chaos
 		CHAOS_API virtual bool IsSleeping() const override final;
 		CHAOS_API virtual void SetIsSleeping(const bool bInIsSleeping) override final;
 
+		// The total impulse for a collision constraint is computed summing over the net impulse of all manifold points.
+		CHAOS_API virtual FVec3f GetNetLinearImpulse() const override final;
+
 		// Get the world-space normal of the closest manifold point
 		// @todo(chaos): remove (used by legacy RBAN collision solver)
 		CHAOS_API FVec3 CalculateWorldContactNormal() const;
diff --git a/Engine/Source/Runtime/Experimental/Chaos/Public/Chaos/ConstraintHandle.h b/Engine/Source/Runtime/Experimental/Chaos/Public/Chaos/ConstraintHandle.h
--- a/Engine/Source/Runtime/Experimental/Chaos/Public/Chaos/ConstraintHandle.h
+++ b/Engine/Source/Runtime/Experimental/Chaos/Public/Chaos/ConstraintHandle.h
@@ -168,6 +168,9 @@ namespace Chaos
 
 		virtual bool IsSleeping() const { return false; }
 		virtual void SetIsSleeping(const bool bInIsSleeping) {}
+
+		// Total net linear impulse applied by the velocity-based dynamics (VBD) solver (velocity iterations only).
+		virtual FVec3f GetNetLinearImpulse() const { return FVec3f(0); }
 		
 		virtual bool WasAwakened() const { return false; }
 		virtual void SetWasAwakened(const bool bInWasAwakened) {}
diff --git a/Engine/Source/Runtime/Experimental/Chaos/Public/Chaos/Island/IslandManager.h b/Engine/Source/Runtime/Experimental/Chaos/Public/Chaos/Island/IslandManager.h
--- a/Engine/Source/Runtime/Experimental/Chaos/Public/Chaos/Island/IslandManager.h
+++ b/Engine/Source/Runtime/Experimental/Chaos/Public/Chaos/Island/IslandManager.h
@@ -614,6 +614,7 @@ namespace Chaos::Private
 		CHAOS_API void ProcessPartialIslandSleep(FPBDIsland* Island, const FRealSingle Dt);
 		CHAOS_API void SleepNonMovingParticles(FPBDIsland* Island, const FRealSingle Dt);
 		CHAOS_API void PropagateParticleSleepToConstraints(FPBDIsland* Island);
+		CHAOS_API void PostStepPartialWake(FPBDIsland* Island);
 
 		// Partial Island Waking
 		CHAOS_API void ProcessPartialIslandWake(FPBDIsland* Island);
