Iris
* Tranform invalid template check into a log error temporarely

#rnx

diff --git a/Engine/Source/Runtime/Experimental/Iris/Core/Private/Iris/ReplicationSystem/ReplicationProtocolManager.cpp b/Engine/Source/Runtime/Experimental/Iris/Core/Private/Iris/ReplicationSystem/ReplicationProtocolManager.cpp
--- a/Engine/Source/Runtime/Experimental/Iris/Core/Private/Iris/ReplicationSystem/ReplicationProtocolManager.cpp
+++ b/Engine/Source/Runtime/Experimental/Iris/Core/Private/Iris/ReplicationSystem/ReplicationProtocolManager.cpp
@@ -274,7 +274,13 @@ const FReplicationProtocol* FReplicationProtocolManager::CreateReplicationProtoc
 	}
 
 	// Template key should be valid unless the protocol doesn't use one.
-	check(Params.bHasTemplateKey == false || IsValid(Params.TemplateKey));
+	// This should be turned back into a check once the source of the issue is found
+	//check(Params.bHasTemplateKey == false || IsValid(Params.TemplateKey));
+	if (Params.bHasTemplateKey && !IsValid(Params.TemplateKey))
+	{
+		UE_LOG(LogIris, Error, TEXT("Cannot create replication protocol %s due to invalid template"), DebugName);
+		return nullptr;
+	}
 
 	// Protocol should not exist already
 	check(GetReplicationProtocol(ProtocolId, Params.TemplateKey) == nullptr);
