DataflowNodeFactory refactor RegisterNodeFromType to minimize per-node-type code.

Move as much logic as possible into a shared static registration function and be careful to minimize
any creation of temporaries (esp. with complex ctors/dtors) in the individual registration instances.

diff --git a/Engine/Plugins/Experimental/Dataflow/Source/DataflowEditor/Private/Dataflow/DataflowAssetEditUtils.cpp b/Engine/Plugins/Experimental/Dataflow/Source/DataflowEditor/Private/Dataflow/DataflowAssetEditUtils.cpp
--- a/Engine/Plugins/Experimental/Dataflow/Source/DataflowEditor/Private/Dataflow/DataflowAssetEditUtils.cpp
+++ b/Engine/Plugins/Experimental/Dataflow/Source/DataflowEditor/Private/Dataflow/DataflowAssetEditUtils.cpp
@@ -125,8 +125,11 @@ namespace UE::Dataflow
 				{
 					.Guid = FGuid::NewGuid(),
 					.Type = NodeTypeName,
-					.Name = GenerateUniqueObjectName(Dataflow, NodeName),
-					.OwningObject = Dataflow,
+					.NodeParams =
+					{
+						.Name = GenerateUniqueObjectName(Dataflow, NodeName),
+						.OwningObject = Dataflow,
+					},
 				};
 
 				return Factory->NewNodeFromRegisteredType(*Dataflow->GetDataflow(), Parameters);
@@ -1147,4 +1150,4 @@ namespace UE::Dataflow
 
 }
 
-#undef LOCTEXT_NAMESPACE
\ No newline at end of file
+#undef LOCTEXT_NAMESPACE
diff --git a/Engine/Source/Runtime/Experimental/Dataflow/Core/Private/Dataflow/DataflowGraph.cpp b/Engine/Source/Runtime/Experimental/Dataflow/Core/Private/Dataflow/DataflowGraph.cpp
--- a/Engine/Source/Runtime/Experimental/Dataflow/Core/Private/Dataflow/DataflowGraph.cpp
+++ b/Engine/Source/Runtime/Experimental/Dataflow/Core/Private/Dataflow/DataflowGraph.cpp
@@ -442,7 +442,7 @@ namespace UE::Dataflow
 			FName ArNodeName;
 			Ar << ArGuid << ArType << ArNodeName;
 
-			TSharedPtr<FDataflowNode> Node = FNodeFactory::GetInstance()->NewNodeFromRegisteredType(*InGraph, { ArGuid, ArType, ArNodeName, OwningObject });
+			TSharedPtr<FDataflowNode> Node = FNodeFactory::GetInstance()->NewNodeFromRegisteredType(*InGraph, { ArGuid, ArType, { ArNodeName, OwningObject } });
 			DATAFLOW_OPTIONAL_BLOCK_READ_BEGIN(Node != nullptr)
 			{
 				ensure(!NodeGuidMap.Contains(ArGuid));
diff --git a/Engine/Source/Runtime/Experimental/Dataflow/Core/Private/Dataflow/DataflowNodeFactory.cpp b/Engine/Source/Runtime/Experimental/Dataflow/Core/Private/Dataflow/DataflowNodeFactory.cpp
--- a/Engine/Source/Runtime/Experimental/Dataflow/Core/Private/Dataflow/DataflowNodeFactory.cpp
+++ b/Engine/Source/Runtime/Experimental/Dataflow/Core/Private/Dataflow/DataflowNodeFactory.cpp
@@ -18,6 +18,61 @@ namespace UE::Dataflow
 {
 	FNodeFactory* FNodeFactory::Instance = nullptr;
 
+	FNodeFactory::FNodeFactory()
+	{
+	}
+
+	FNodeFactory::~FNodeFactory()
+	{
+		delete Instance; 
+	}
+
+	FNodeFactory* FNodeFactory::GetInstance()
+	{
+		if (!Instance)
+		{
+			Instance = new FNodeFactory();
+			Instance->RegisterDefaultNodes();
+		}
+		return Instance;
+	}
+
+	// This part is common for all RegisterNode calls, so share it!
+	void FNodeFactory::RegisterNodeStaticInternal(
+		UScriptStruct* StaticStruct, FName StaticType, const char* StaticDisplay, const char* StaticCategory,
+		const FString& StaticTags, FRawNewNodeFunction CreationFunction
+	)
+	{
+		FNewNodeParameters CreateParams
+		{
+			.Guid = FGuid::NewGuid(),
+			.Type = FName{},
+			.NodeParams =
+			{
+				.Name = StaticType,
+				.OwningObject = GetTransientPackage()
+			}
+		};
+
+		// Create a node and take over ownership temporarily...
+		FDataflowNode* DefaultNode = CreationFunction(CreateParams).Release();
+
+		FFactoryParameters FactoryParameters
+		{
+			.TypeName		= StaticType,
+			.DisplayName	= StaticDisplay,
+			.Category		= StaticCategory,
+			.Tags			= StaticTags,
+			.ToolTip		= GetToolTipFromStruct(StaticStruct, StaticType, StaticDisplay),
+			.bIsDeprecated	= IsNodeDeprecated(StaticStruct),
+			.bIsExperimental= IsNodeExperimental(StaticStruct),
+			.NodeVersion	= GetVersionFromTypeName(StaticType),
+			.DefaultNodeObject = TSharedPtr<FDataflowNode>(DefaultNode), // ...then transfer ownership to the SharedPtr
+		};
+
+		GetInstance()->RegisterNode(FactoryParameters, CreationFunction);
+	}
+
 	void FNodeFactory::RegisterDefaultNodes()
 	{
 		// anytypes need to be at the top as the connection registration in the nodes are using it
diff --git a/Engine/Source/Runtime/Experimental/Dataflow/Core/Public/Dataflow/DataflowNode.h b/Engine/Source/Runtime/Experimental/Dataflow/Core/Public/Dataflow/DataflowNode.h
--- a/Engine/Source/Runtime/Experimental/Dataflow/Core/Public/Dataflow/DataflowNode.h
+++ b/Engine/Source/Runtime/Experimental/Dataflow/Core/Public/Dataflow/DataflowNode.h
@@ -980,8 +980,8 @@ private:
 #define DATAFLOW_NODE_DEFINE_INTERNAL(TYPE, DISPLAY_NAME, CATEGORY, TAGS)			\
 public:																				\
 	static FName StaticType() {return #TYPE;}										\
-	static FName StaticDisplay() {return DISPLAY_NAME;}								\
-	static FName StaticCategory() {return CATEGORY;}								\
+	static const char* StaticDisplay() {return DISPLAY_NAME;}						\
+	static const char* StaticCategory() {return CATEGORY;}							\
 	static FString StaticTags() {return TAGS;}										\
 	static FString StaticToolTip() {return FString("Create a dataflow node.");}		\
 	virtual FName GetType() const { return #TYPE; }									\
diff --git a/Engine/Source/Runtime/Experimental/Dataflow/Core/Public/Dataflow/DataflowNodeFactory.h b/Engine/Source/Runtime/Experimental/Dataflow/Core/Public/Dataflow/DataflowNodeFactory.h
--- a/Engine/Source/Runtime/Experimental/Dataflow/Core/Public/Dataflow/DataflowNodeFactory.h
+++ b/Engine/Source/Runtime/Experimental/Dataflow/Core/Public/Dataflow/DataflowNodeFactory.h
@@ -15,8 +15,7 @@ namespace UE::Dataflow
 	{
 		FGuid Guid;
 		FName Type;
-		FName Name;
-		UObject* OwningObject = nullptr;
+		FNodeParameters NodeParams;
 	};
 
 	struct FFactoryParameters
@@ -61,6 +60,8 @@ namespace UE::Dataflow
 		typedef TFunction<TUniquePtr<FDataflowNode> (const FNewNodeParameters&)> FNewNodeFunction;
 
 	private:
+		typedef TUniquePtr<FDataflowNode> FRawNewNodeFunction(const FNewNodeParameters&);
+
 		// All Maps indexed by TypeName
 		TMap<FName, FNewNodeFunction > ClassMap;			// [TypeName] -> NewNodeFunction
 		TMap<FName, FFactoryParameters > ParametersMap;		// [TypeName] -> Parameters
@@ -69,53 +70,33 @@ namespace UE::Dataflow
 		TMap<FName, FName> GetterNodesByAssetType;	// list of getter nodes that relate to a specific asset type
 
 		DATAFLOWCORE_API static FNodeFactory* Instance;
-		FNodeFactory() {};
+		DATAFLOWCORE_API FNodeFactory();
+
+		static DATAFLOWCORE_API void RegisterNodeStaticInternal(
+			UScriptStruct* StaticStruct, FName StaticType, const char* StaticDisplay, const char* StaticCategory,
+			const FString& StaticTags, FRawNewNodeFunction CreationFunction
+		);
 
 	public:
-		~FNodeFactory() { delete Instance; }
+		DATAFLOWCORE_API ~FNodeFactory();
 
-		static FNodeFactory* GetInstance()
-		{
-			if (!Instance)
-			{
-				Instance = new FNodeFactory();
-				Instance->RegisterDefaultNodes();
-			}
-			return Instance;
-		}
+		static DATAFLOWCORE_API FNodeFactory* GetInstance();
 
 		template<typename T>
 		static void RegisterNodeFromType()
 		{
-			UE::Dataflow::FNodeParameters CreateParams
-			{
-				.Name			= T::StaticType(),
-				.OwningObject	= GetTransientPackage(),
-			};
-
-			FFactoryParameters FactoryParameters
-			{
-				.TypeName		= T::StaticType(),
-				.DisplayName	= T::StaticDisplay(),
-				.Category		= T::StaticCategory(),
-				.Tags		= T::StaticTags(),
-				.ToolTip		= GetToolTipFromStruct(T::StaticStruct(), T::StaticType(), T::StaticDisplay()),
-				.bIsDeprecated	= IsNodeDeprecated(T::StaticStruct()),
-				.bIsExperimental= IsNodeExperimental(T::StaticStruct()),
-				.NodeVersion	= GetVersionFromTypeName(T::StaticType()),
-				.DefaultNodeObject = MakeShared<T>(CreateParams),
-			};
-
 			auto CreationFunction = [](const UE::Dataflow::FNewNodeParameters& InParam)
 			{
-				const FNodeParameters NodeParameters = { InParam.Name, InParam.OwningObject };
-				TUniquePtr<T> Val = MakeUnique<T>(NodeParameters, InParam.Guid);
+				TUniquePtr<FDataflowNode> Val = MakeUnique<T>(InParam.NodeParams, InParam.Guid);
 				Val->ValidateProperties();
 				Val->ValidateConnections();
 				return Val;
 			};
 
-			GetInstance()->RegisterNode(FactoryParameters, CreationFunction);
+			RegisterNodeStaticInternal(
+				T::StaticStruct(), T::StaticType(), T::StaticDisplay(), T::StaticCategory(),
+				T::StaticTags(), CreationFunction
+			);
 		}
 
 		DATAFLOWCORE_API void RegisterGetterNodeForAssetType(FName AssetTypeName, FName NodeTypeName);
@@ -151,7 +132,7 @@ namespace UE::Dataflow
 
 		template<class T> TSharedPtr<T> NewNode(FGraph& Graph, const FNewNodeParameters& Param)
 		{
-			return Graph.AddNode(new T(Param.Name, Param.Guid));
+			return Graph.AddNode(new T(Param.NodeParams.Name, Param.Guid));
 		}
 
 		DATAFLOWCORE_API void RegisterDefaultNodes();
