Removed redundant "None" state.

#rnx

diff --git a/Engine/Source/Runtime/Experimental/IoStore/HttpClient/Private/Activity.inl b/Engine/Source/Runtime/Experimental/IoStore/HttpClient/Private/Activity.inl
--- a/Engine/Source/Runtime/Experimental/IoStore/HttpClient/Private/Activity.inl
+++ b/Engine/Source/Runtime/Experimental/IoStore/HttpClient/Private/Activity.inl
@@ -769,7 +769,6 @@ struct FActivity
 
 	enum class EState : uint8
 	{
-		None,
 		Build,
 		Send,
 		RecvMessage,
@@ -790,7 +789,7 @@ struct FActivity
 	FOutcome			RecvStream(FHttpPeer& Peer, int32& MaxRecvSize);
 	FOutcome			Recv(FHttpPeer& Peer, int32& MaxRecvSize);
 
-	EState				State = EState::None;
+	EState				State = EState::Build;
 	uint8				IsKeepAlive : 1;
 	uint8				bAllowChunked : 1;
 	uint8				_Unused : 2;
@@ -1209,7 +1208,6 @@ static void Trace(const struct FActivity* Activity, ETrace Action, uint32 Param)
 	if (Action == ETrace::StateChange)
 	{
 		static constexpr FAnsiStringView StateNames[] = {
-			"Iax/None",
 			"Iax/Build",
 			"Iax/WaitForSocket",
 			"Iax/WaitResponse",
