An extra check to surface assumptions and user error

#rnx

diff --git a/Engine/Source/Runtime/Experimental/IoStore/HttpClient/Private/Loop.inl b/Engine/Source/Runtime/Experimental/IoStore/HttpClient/Private/Loop.inl
--- a/Engine/Source/Runtime/Experimental/IoStore/HttpClient/Private/Loop.inl
+++ b/Engine/Source/Runtime/Experimental/IoStore/HttpClient/Private/Loop.inl
@@ -783,6 +783,7 @@ FRequest FEventLoop::FImpl::Request(
 	{
 		Path = "/";
 	}
+	check(Path[0] == '/');
 
 	const FHost* Host = Activity->GetHost();
 
