Do not assume added data is right at the beginning.

#rnx

diff --git a/Engine/Source/Runtime/Experimental/IoStore/HttpClient/Private/Activity.inl b/Engine/Source/Runtime/Experimental/IoStore/HttpClient/Private/Activity.inl
--- a/Engine/Source/Runtime/Experimental/IoStore/HttpClient/Private/Activity.inl
+++ b/Engine/Source/Runtime/Experimental/IoStore/HttpClient/Private/Activity.inl
@@ -47,6 +47,7 @@ void FRequestCloak::EnumerateHeaders(FResponse::FHeaderSink Sink) const
 void FRequestCloak::Begin(FAnsiStringView Method, FAnsiStringView Path)
 {
 	check(MethodLength < 0);
+	AlreadySent = uint16(GetSize());
 	*this << Method << " " << Path << " HTTP/1.1" "\r\n";
 	MethodLength = int8(Method.Len());
 	HeaderLeft = uint16(GetSize());
