Deleted unused member

#rnx

diff --git a/Engine/Source/Runtime/Experimental/IoStore/HttpClient/Private/Activity.inl b/Engine/Source/Runtime/Experimental/IoStore/HttpClient/Private/Activity.inl
--- a/Engine/Source/Runtime/Experimental/IoStore/HttpClient/Private/Activity.inl
+++ b/Engine/Source/Runtime/Experimental/IoStore/HttpClient/Private/Activity.inl
@@ -791,9 +791,8 @@ struct FActivity
 	FOutcome			Recv(FHttpPeer& Peer, int32& MaxRecvSize);
 
 	EState				State = EState::Build;
-	uint8				IsKeepAlive : 1;
 	uint8				bAllowChunked : 1;
-	uint8				_Unused : 2;
+	uint8				_Unused : 3;
 	uint8				LengthScore : 4;
 	uint32				StateParam = 0;
 #if IAS_HTTP_WITH_PERF
@@ -811,8 +810,7 @@ struct FActivity
 
 ////////////////////////////////////////////////////////////////////////////////
 FActivity::FActivity(const FParams& Params)
-: IsKeepAlive(Params.bIsKeepAlive)
-, bAllowChunked(Params.bAllowChunked)
+: bAllowChunked(Params.bAllowChunked)
 , Host(Params.Host)
 , Transaction(Params.Buffer, Params.BufferSize)
 {
@@ -866,7 +864,6 @@ void FActivity::Fail(const FOutcome& Outcome)
 
 	ChangeState(EState::Failed);
 
-	IsKeepAlive = 0;
 	ErrorReason = Outcome.GetMessage().GetData();
 	StateParam = uint32(Outcome.GetErrorCode());
 
