Few needless scope qualifications crept in while refactoring.

#rnx

diff --git a/Engine/Source/Runtime/Experimental/IoStore/HttpClient/Private/Activity.inl b/Engine/Source/Runtime/Experimental/IoStore/HttpClient/Private/Activity.inl
--- a/Engine/Source/Runtime/Experimental/IoStore/HttpClient/Private/Activity.inl
+++ b/Engine/Source/Runtime/Experimental/IoStore/HttpClient/Private/Activity.inl
@@ -837,12 +837,12 @@ void FActivity::SetDestination(FIoBuffer* InDest)
 ////////////////////////////////////////////////////////////////////////////////
 void FActivity::Cancel()
 {
-	if (State >= FActivity::EState::Completed)
+	if (State >= EState::Completed)
 	{
 		return;
 	}
 
-	ChangeState(FActivity::EState::Cancelled);
+	ChangeState(EState::Cancelled);
 	CallSink();
 }
 
@@ -915,8 +915,8 @@ FIoBuffer& FActivity::GetContent() const
 ////////////////////////////////////////////////////////////////////////////////
 uint32 FActivity::GetRemainingKiB() const
 {
-	if (State <= FActivity::EState::RecvStream)  return MAX_uint32;
-	if (State >  FActivity::EState::RecvContent) return 0;
+	if (State <= EState::RecvStream)  return MAX_uint32;
+	if (State >  EState::RecvContent) return 0;
 
 	uint32 ContentLength = uint32(Transaction.GetContentLength());
 	check(StateParam <= ContentLength);
@@ -972,7 +972,7 @@ FOutcome FActivity::Send(FHttpPeer& Peer)
 	Stopwatch.SendEnd();
 #endif
 
-	ChangeState(FActivity::EState::RecvMessage);
+	ChangeState(EState::RecvMessage);
 	return FOutcome::Ok(EStage::Request);
 }
 
@@ -1051,13 +1051,13 @@ FOutcome FActivity::RecvMessage(FHttpPeer& Peer)
 	// HEAD methods
 	if (!Transaction.HasBody())
 	{
-		ChangeState(FActivity::EState::RecvDone);
+		ChangeState(EState::RecvDone);
 		return FOutcome::Ok(EStage::Content);
 	}
 
 	// We're all set to go and get content
 	check(Dest != nullptr);
-	auto NextState = bChunked ? FActivity::EState::RecvStream : FActivity::EState::RecvContent;
+	auto NextState = bChunked ? EState::RecvStream : EState::RecvContent;
 	ChangeState(NextState);
 	return FOutcome::Ok(EStage::Response);
 }
@@ -1159,7 +1159,7 @@ FOutcome FActivity::Recv(FHttpPeer& Peer, int32& MaxRecvSize)
 ////////////////////////////////////////////////////////////////////////////////
 FOutcome FActivity::Tick(FHttpPeer& Peer, int32* MaxRecvSize)
 {
-	if (State == FActivity::EState::Send)
+	if (State == EState::Send)
 	{
 		return Send(Peer);
 	}
