Indirect calls rather than doing directly to the exposed member.

#rnx

diff --git a/Engine/Source/Runtime/Experimental/IoStore/HttpClient/Private/Activity.inl b/Engine/Source/Runtime/Experimental/IoStore/HttpClient/Private/Activity.inl
--- a/Engine/Source/Runtime/Experimental/IoStore/HttpClient/Private/Activity.inl
+++ b/Engine/Source/Runtime/Experimental/IoStore/HttpClient/Private/Activity.inl
@@ -751,6 +751,9 @@ struct FActivity
 
 						FActivity(const FParams& Params);
 						~FActivity() = default;
+	void				Begin(FAnsiStringView Method, FAnsiStringView Path);
+	void				AddHeader(FAnsiStringView Key, FAnsiStringView Value);
+	void				End();
 	FOutcome			Tick(FHttpPeer& Peer, int32* MaxRecvSize=nullptr);
 	void				SetSink(FTicketSink&& InSink, UPTRINT Param);
 	void				SetDestination(FIoBuffer* InDest);
@@ -821,6 +824,25 @@ FActivity::FActivity(const FParams& Params)
 	LengthScore = uint8(Pow2);
 }
 
+////////////////////////////////////////////////////////////////////////////////
+void FActivity::Begin(FAnsiStringView Method, FAnsiStringView Path)
+{
+	Transaction.Begin(Method, Path);
+}
+
+////////////////////////////////////////////////////////////////////////////////
+void FActivity::AddHeader(FAnsiStringView Key, FAnsiStringView Value)
+{
+	Transaction.AddHeader(Key, Value);
+}
+
+////////////////////////////////////////////////////////////////////////////////
+void FActivity::End()
+{
+	Transaction.End();
+	ChangeState(EState::Send);
+}
+
 ////////////////////////////////////////////////////////////////////////////////
 void FActivity::SetSink(FTicketSink&& InSink, UPTRINT Param)
 {
diff --git a/Engine/Source/Runtime/Experimental/IoStore/HttpClient/Private/Api.inl b/Engine/Source/Runtime/Experimental/IoStore/HttpClient/Private/Api.inl
--- a/Engine/Source/Runtime/Experimental/IoStore/HttpClient/Private/Api.inl
+++ b/Engine/Source/Runtime/Experimental/IoStore/HttpClient/Private/Api.inl
@@ -51,7 +51,7 @@ FRequest&& FRequest::Accept(FAnsiStringView MimeType)
 ////////////////////////////////////////////////////////////////////////////////
 FRequest&& FRequest::Header(FAnsiStringView Key, FAnsiStringView Value)
 {
-	Ptr->Transaction.AddHeader(Key, Value);
+	Ptr->AddHeader(Key, Value);
 	return MoveTemp(*this);
 }
 
diff --git a/Engine/Source/Runtime/Experimental/IoStore/HttpClient/Private/Loop.inl b/Engine/Source/Runtime/Experimental/IoStore/HttpClient/Private/Loop.inl
--- a/Engine/Source/Runtime/Experimental/IoStore/HttpClient/Private/Loop.inl
+++ b/Engine/Source/Runtime/Experimental/IoStore/HttpClient/Private/Loop.inl
@@ -787,14 +787,13 @@ FRequest FEventLoop::FImpl::Request(
 
 	const FHost* Host = Activity->GetHost();
 
-	auto& Transaction = Activity->Transaction;
-	Transaction.Begin(Method, Path);
-	Transaction.AddHeader("Host", Host->GetHostName());
+	Activity->Begin(Method, Path);
+	Activity->AddHeader("Host", Host->GetHostName());
 
 	// HTTP/1.1 is persistent by default thus "Connection" header isn't required
 	if (!Host->IsPooled())
 	{
-		Transaction.AddHeader("Connection", "close");
+		Activity->AddHeader("Connection", "close");
 	}
 
 	FRequest Ret;
@@ -807,8 +806,7 @@ FTicket FEventLoop::FImpl::Send(FActivityNode* Activity)
 {
 	Trace(Activity, ETrace::RequestBegin);
 
-	Activity->Transaction.End();
-	Activity->ChangeState(FActivity::EState::Send);
+	Activity->End();
 
 	uint64 Slot;
 	{
