Copy activity properties instead of immediate use and discard.

#rnx

diff --git a/Engine/Source/Runtime/Experimental/IoStore/HttpClient/Private/Activity.inl b/Engine/Source/Runtime/Experimental/IoStore/HttpClient/Private/Activity.inl
--- a/Engine/Source/Runtime/Experimental/IoStore/HttpClient/Private/Activity.inl
+++ b/Engine/Source/Runtime/Experimental/IoStore/HttpClient/Private/Activity.inl
@@ -734,6 +734,8 @@ struct FActivity
 {
 	struct FParams
 	{
+		FAnsiStringView	Method;
+		FAnsiStringView	Path;
 		FHost*			Host = nullptr;
 		char*			Buffer = nullptr;
 		uint32			ContentSizeEst = 0;
@@ -751,7 +753,7 @@ struct FActivity
 
 						FActivity(const FParams& Params);
 						~FActivity() = default;
-	void				Begin(FAnsiStringView Method, FAnsiStringView Path);
+	void				Begin();
 	void				AddHeader(FAnsiStringView Key, FAnsiStringView Value);
 	void				End();
 	FOutcome			Tick(FHttpPeer& Peer, int32* MaxRecvSize=nullptr);
@@ -797,6 +799,9 @@ struct FActivity
 	uint8				bAllowChunked : 1;
 	uint8				_Unused : 3;
 	uint8				LengthScore : 4;
+	uint8				MethodOffset;
+	uint8				PathOffset;
+	uint16				PathLength;
 	uint32				StateParam = 0;
 #if IAS_HTTP_WITH_PERF
 	FStopwatch			Stopwatch;
@@ -818,6 +823,26 @@ FActivity::FActivity(const FParams& Params)
 , Transaction(Params.Buffer, Params.BufferSize)
 {
 	check(Host != nullptr);
+
+	// Make a copy of data.
+	FAnsiStringView Path = Params.Path.IsEmpty() ? "/" : Params.Path;
+	check(Path[0] == '/');
+
+	auto Copy = [this] (FAnsiStringView Value)
+	{
+		uint32 Length = uint32(Value.Len());
+		char* Ptr = Transaction.Alloc<char>(Length);
+		std::memcpy(Ptr, Value.GetData(), Length);
+		uint32 Ret = uint32(ptrdiff_t(Ptr - Transaction.GetData()));
+		check(Ret <= 255);
+		return uint8(Ret);
+	};
+
+	MethodOffset = Copy(Params.Method);
+	PathOffset = Copy(Path);
+	PathLength = uint16(Path.Len());
+
+	// Calculate a length score.
 	uint32 ContentEstKiB = (Params.ContentSizeEst + 1023) >> 10;
 	uint32 Pow2 = FMath::FloorLog2(uint32(ContentEstKiB));
 	Pow2 = FMath::Min(Pow2, 15u);
@@ -825,8 +850,10 @@ FActivity::FActivity(const FParams& Params)
 }
 
 ////////////////////////////////////////////////////////////////////////////////
-void FActivity::Begin(FAnsiStringView Method, FAnsiStringView Path)
+void FActivity::Begin()
 {
+	FAnsiStringView Method(Transaction.GetData() + MethodOffset, PathOffset - MethodOffset);
+	FAnsiStringView Path(Transaction.GetData() + PathOffset, PathLength);
 	Transaction.Begin(Method, Path);
 }
 
diff --git a/Engine/Source/Runtime/Experimental/IoStore/HttpClient/Private/Loop.inl b/Engine/Source/Runtime/Experimental/IoStore/HttpClient/Private/Loop.inl
--- a/Engine/Source/Runtime/Experimental/IoStore/HttpClient/Private/Loop.inl
+++ b/Engine/Source/Runtime/Experimental/IoStore/HttpClient/Private/Loop.inl
@@ -36,7 +36,7 @@ struct FActivityNode
 		bool				bFollow30x = false;
 	};
 
-	static FActivityNode*	Create(FParams& Params, FAnsiStringView* Url=nullptr, FCertRootsRef VerifyCert={});
+	static FActivityNode*	Create(FParams& Params, FAnsiStringView Url={}, FCertRootsRef VerifyCert={});
 	static void				Destroy(FActivityNode* Activity);
 	FActivityNode*			Next = nullptr;
 	int8					Slot = -1;
@@ -55,12 +55,14 @@ FActivityNode::FActivityNode(const FParams& Params)
 }
 
 ////////////////////////////////////////////////////////////////////////////////
-FActivityNode* FActivityNode::Create(FParams& Params, FAnsiStringView* InOutUrl, FCertRootsRef VerifyCert)
+FActivityNode* FActivityNode::Create(FParams& Params, FAnsiStringView Url, FCertRootsRef VerifyCert)
 {
 	static_assert(alignof(FActivityNode) <= 16);
 	static_assert(alignof(FHost) <= alignof(FActivityNode));
 
-	if (InOutUrl == nullptr)
+	check(!Params.Method.IsEmpty());
+
+	if (Url.IsEmpty())
 	{
 		uint32 BufferSize = Params.BufferSize;
 		BufferSize = (BufferSize >= 128) ? BufferSize : 128;
@@ -78,8 +80,6 @@ FActivityNode* FActivityNode::Create(FParams& Params, FAnsiStringView* InOutUrl,
 		return new (Ptr) FActivityNode(Params);
 	}
 
-	FAnsiStringView& Url = *InOutUrl;
-
 	// Parse the URL into its components
 	FUrlOffsets UrlOffsets;
 	if (ParseUrl(Url, UrlOffsets) < 0)
@@ -101,7 +101,6 @@ FActivityNode* FActivityNode::Create(FParams& Params, FAnsiStringView* InOutUrl,
 	{
 		Path = Url.Mid(UrlOffsets.Path);
 	}
-	*InOutUrl = Path;
 
 	uint32 BufferSize = Params.BufferSize;
 	BufferSize = (BufferSize >= 128) ? BufferSize : 128;
@@ -142,6 +141,7 @@ FActivityNode* FActivityNode::Create(FParams& Params, FAnsiStringView* InOutUrl,
 	HostNameLength = (HostNameLength + 8) & ~7;
 
 	check(Params.Host == nullptr);
+	Params.Path = Path;
 	Params.Host = Host;
 	Params.Buffer = HostNamePtr + HostNameLength;
 	Params.BufferSize = uint16(AllocSize - ptrdiff_t(Params.Buffer - (char*)Ptr));
@@ -749,7 +749,7 @@ public:
 	void					Throttle(uint32 KiBPerSec);
 	void					SetFailTimeout(int32 TimeoutMs);
 	void					Cancel(FTicket Ticket);
-	FRequest				Request(FAnsiStringView Method, FAnsiStringView Path, FActivityNode* Activity);
+	FRequest				Request(FActivityNode* Activity);
 	FTicket					Send(FActivityNode* Activity);
 
 private:
@@ -772,22 +772,13 @@ FEventLoop::FImpl::~FImpl()
 }
 
 ////////////////////////////////////////////////////////////////////////////////
-FRequest FEventLoop::FImpl::Request(
-	FAnsiStringView Method,
-	FAnsiStringView Path,
-	FActivityNode* Activity)
+FRequest FEventLoop::FImpl::Request(FActivityNode* Activity)
 {
 	Trace(Activity, ETrace::ActivityCreate, 0);
 
-	if (Path.Len() == 0)
-	{
-		Path = "/";
-	}
-	check(Path[0] == '/');
-
 	const FHost* Host = Activity->GetHost();
 
-	Activity->Begin(Method, Path);
+	Activity->Begin();
 	Activity->AddHeader("Host", Host->GetHostName());
 
 	// HTTP/1.1 is persistent by default thus "Connection" header isn't required
@@ -1038,18 +1029,19 @@ FRequest FEventLoop::Request(
 	Params = (Params != nullptr) ? Params : &GDefaultParams;
 
 	FActivityNode::FParams ActivityParams;
+	ActivityParams.Method = Method;
 	ActivityParams.BufferSize = Params->BufferSize;
 	ActivityParams.bFollow30x = (Params->bAutoRedirect == true);
 	ActivityParams.bAllowChunked = (Params->bAllowChunked == true);
 	ActivityParams.ContentSizeEst = Params->ContentSizeEst;
 
-	FActivityNode* Activity = FActivityNode::Create(ActivityParams, &Url, Params->VerifyCert);
+	FActivityNode* Activity = FActivityNode::Create(ActivityParams, Url, Params->VerifyCert);
 	if (Activity == nullptr)
 	{
 		return FRequest();
 	}
 
-	return Impl->Request(Method, Url, Activity);
+	return Impl->Request(Activity);
 }
 
 ////////////////////////////////////////////////////////////////////////////////
@@ -1065,6 +1057,8 @@ FRequest FEventLoop::Request(
 	Params = (Params != nullptr) ? Params : &GDefaultParams;
 
 	FActivityNode::FParams ActivityParams;
+	ActivityParams.Method = Method;
+	ActivityParams.Path = Path;
 	ActivityParams.Host = Pool.Ptr;
 	ActivityParams.bFollow30x = (Params->bAutoRedirect == true);
 	ActivityParams.bAllowChunked = (Params->bAllowChunked == true);
@@ -1072,7 +1066,7 @@ FRequest FEventLoop::Request(
 	FActivityNode* Activity = FActivityNode::Create(ActivityParams);
 	check(Activity != nullptr);
 
-	return Impl->Request(Method, Path, Activity);
+	return Impl->Request(Activity);
 }
 
 ////////////////////////////////////////////////////////////////////////////////
