Track bias directly until some dependent methods can be moved elsewhere.

#rnx

diff --git a/Engine/Source/Runtime/Experimental/IoStore/HttpClient/Private/Activity.inl b/Engine/Source/Runtime/Experimental/IoStore/HttpClient/Private/Activity.inl
--- a/Engine/Source/Runtime/Experimental/IoStore/HttpClient/Private/Activity.inl
+++ b/Engine/Source/Runtime/Experimental/IoStore/HttpClient/Private/Activity.inl
@@ -21,6 +21,7 @@ public:
 
 private:
 	FRequestCloak&	operator << (FAnsiStringView Value);
+	uint16			Bias = 0;
 	uint16			HeaderLeft;
 	uint16			AlreadySent = 0;
 	int16			HeaderRight = -1;
@@ -31,7 +32,7 @@ private:
 FAnsiStringView FRequestCloak::GetMethod() const
 {
 	check(MethodLength > 0);
-	return FAnsiStringView(GetData(), MethodLength);
+	return FAnsiStringView(GetData() + Bias, MethodLength);
 }
 
 ////////////////////////////////////////////////////////////////////////////////
@@ -47,7 +48,7 @@ void FRequestCloak::EnumerateHeaders(FResponse::FHeaderSink Sink) const
 void FRequestCloak::Begin(FAnsiStringView Method, FAnsiStringView Path)
 {
 	check(MethodLength < 0);
-	AlreadySent = uint16(GetSize());
+	AlreadySent = Bias = uint16(GetSize());
 	*this << Method << " " << Path << " HTTP/1.1" "\r\n";
 	MethodLength = int8(Method.Len());
 	HeaderLeft = uint16(GetSize());
