Alignment fix

#rnx

diff --git a/Engine/Source/Runtime/Experimental/IoStore/HttpClient/Private/Misc.inl b/Engine/Source/Runtime/Experimental/IoStore/HttpClient/Private/Misc.inl
--- a/Engine/Source/Runtime/Experimental/IoStore/HttpClient/Private/Misc.inl
+++ b/Engine/Source/Runtime/Experimental/IoStore/HttpClient/Private/Misc.inl
@@ -534,7 +534,7 @@ uint32 FBuffer::GetCapacity() const
 template <typename T>
 T* FBuffer::Alloc(uint32 Count)
 {
-	uint32 AlignBias = uint32(Data) & (alignof(T) - 1);
+	uint32 AlignBias = uint32(Data + Used) & (alignof(T) - 1);
 	if (AlignBias)
 	{
 		AlignBias = alignof(T) - AlignBias;
