Add more testable data earlier for slightly faster iteration.

#rnx

diff --git a/Engine/Source/Runtime/Experimental/IoStore/HttpClient/Private/Test.inl b/Engine/Source/Runtime/Experimental/IoStore/HttpClient/Private/Test.inl
--- a/Engine/Source/Runtime/Experimental/IoStore/HttpClient/Private/Test.inl
+++ b/Engine/Source/Runtime/Experimental/IoStore/HttpClient/Private/Test.inl
@@ -968,6 +968,8 @@ IOSTOREHTTPCLIENT_API void IasHttpTest(const ANSICHAR* TestHost="localhost", uin
 
 		FEventLoop Loop;
 		FRequest Request = Loop.Get(CaUrl);
+		Request.Header("X-TestServer-M", "Girders");
+		Request.Header("X-TestServer-R", "RedRigs");
 		Loop.Send(MoveTemp(Request), [&] (const FTicketStatus& Status) {
 			check(Status.GetId() != FTicketStatus::EId::Error);
 
