Iris
* Renable check when detecting an invalid template pointer

#rnx

diff --git a/Engine/Source/Runtime/Experimental/Iris/Core/Private/Iris/ReplicationSystem/ReplicationProtocolManager.cpp b/Engine/Source/Runtime/Experimental/Iris/Core/Private/Iris/ReplicationSystem/ReplicationProtocolManager.cpp
--- a/Engine/Source/Runtime/Experimental/Iris/Core/Private/Iris/ReplicationSystem/ReplicationProtocolManager.cpp
+++ b/Engine/Source/Runtime/Experimental/Iris/Core/Private/Iris/ReplicationSystem/ReplicationProtocolManager.cpp
@@ -274,11 +274,10 @@ const FReplicationProtocol* FReplicationProtocolManager::CreateReplicationProtoc
 	}
 
 	// Template key should be valid unless the protocol doesn't use one.
-	// This should be turned back into a check once the source of the issue is found
-	//check(Params.bHasTemplateKey == false || IsValid(Params.TemplateKey));
 	if (Params.bHasTemplateKey && !IsValid(Params.TemplateKey))
 	{
 		UE_LOG(LogIris, Error, TEXT("Cannot create replication protocol %s due to invalid template"), DebugName);
+		check(Params.bHasTemplateKey == false || IsValid(Params.TemplateKey));
 		return nullptr;
 	}
 
