Created a SameSpace version of GJkRaycast.
Updated a few sweeps to pre-transform one of the objects and then use GjkRaycastSameSpace.

diff --git a/Engine/Source/Runtime/Experimental/Chaos/Private/Chaos/GeometryQueries.cpp b/Engine/Source/Runtime/Experimental/Chaos/Private/Chaos/GeometryQueries.cpp
--- a/Engine/Source/Runtime/Experimental/Chaos/Private/Chaos/GeometryQueries.cpp
+++ b/Engine/Source/Runtime/Experimental/Chaos/Private/Chaos/GeometryQueries.cpp
@@ -4,6 +4,7 @@
 #include "Chaos/Box.h"
 #include "Chaos/Capsule.h"
 #include "Chaos/CastingUtilities.h"
+#include "Chaos/GJKShape.h"
 #include "Chaos/ImplicitObject.h"
 #include "Chaos/Sphere.h"
 #include "Chaos/Sweeps.h"
@@ -33,6 +34,39 @@ namespace Chaos
 		}
 	}
 
+	FGJKSphereSingleSIMD ComputeLocalSphere(const FSphere& Sphere, const FRigidTransform3& SphereTM, const FRigidTransform3& OtherTM)
+	{
+		const FVec3 SphereWorldLocation = FVec3(Sphere.GetCenterf()) + SphereTM.GetLocation();
+		const FVec3 SphereLocation = OtherTM.InverseTransformPositionNoScale(SphereWorldLocation);
+		return FGJKSphereSingleSIMD(SphereLocation, Sphere.GetRadiusf());
+	}
+
+	FGJKCapsuleSingleSIMD ComputeLocalCapsule(const FCapsule& Capsule, const FRigidTransform3& CapsuleTM, const FRigidTransform3& OtherTM)
+	{
+		const FRigidTransform3 CapsuleToBoxTM = CapsuleTM.GetRelativeTransformNoScale(OtherTM);
+		const FVec3 CapsuleLocalAxis = CapsuleToBoxTM.TransformVectorNoScale(FVec3(Capsule.GetAxisf()));
+		const FVec3 CapsuleLocalPointA = CapsuleToBoxTM.TransformPositionNoScale(FVec3(Capsule.GetX1f()));
+		const FVec3 CapsuleLocalPointB = CapsuleLocalPointA + CapsuleLocalAxis * Capsule.GetHeightf();
+		return FGJKCapsuleSingleSIMD(CapsuleLocalPointA, CapsuleLocalPointB, CapsuleLocalAxis, Capsule.GetRadiusf());
+	}
+
+	template <typename SweptShapeType, typename TestShapeType, typename TestObjectType>
+	bool LocalSweep(const SweptShapeType& SweptShape, const TestShapeType& TestShape, const TestObjectType& TestObject, const FRigidTransform3& LocalToWorldTM,
+		const FVec3& WorldSweepDir, const FReal Length, const FReal Thickness, const bool bComputeMTD,
+		FReal& OutTime, FVec3& OutPosition, FVec3& OutNormal, int32& OutFaceIndex, FVec3& OutFaceNormal)
+	{
+		const FVec3 LocalSweepDir = LocalToWorldTM.InverseTransformVectorNoScale(WorldSweepDir);
+		// Note: The swept shape was fully transformed into the local space of the test object. 
+		// This means the "start" is zero as we've offset the shape. This is typically a little faster than computing the "origin centered local object".
+		const FVec3 Start = FVec3::ZeroVector;
+		// Note: Technically any initial dir is valid. The unoptimized version uses A - B which is a little more expensive to calculate. The sweep dir should be a good initial guess as well.
+		const FVec3 InitialDir = LocalSweepDir;
+
+		const bool bResult = GJKRaycast2SameSpace<FReal>(TestShape, SweptShape, Start, LocalSweepDir, Length, OutTime, OutPosition, OutNormal, Thickness, bComputeMTD, InitialDir, Thickness);
+		TransformSweepResultsToWorld(bResult, OutTime, bComputeMTD, TestObject, LocalToWorldTM, LocalSweepDir, OutPosition, OutNormal, OutFaceIndex, OutPosition, OutNormal, OutFaceNormal);
+		return bResult;
+	}
+
 	bool SweepSphereVsSphere(const FSphere& SweptSphere, const FRigidTransform3& SweptSphereTM, const FSphere& TestSphere, const FRigidTransform3& TestSphereTM,
 		const FVec3& SweepDir, const FReal Length, const FReal Thickness, const bool bComputeMTD,
 		FReal& OutTime, FVec3& OutPosition, FVec3& OutNormal, int32& OutFaceIndex, FVec3& OutFaceNormal)
@@ -50,6 +84,12 @@ namespace Chaos
 		return bResult;
 	}
 
+	bool SweepSphereVsBox(const FSphere& SweptSphere, const FRigidTransform3& SweptSphereTM, const TBox<FReal, 3>& TestBox, const FRigidTransform3& TestBoxTM, const FVec3& SweepDir, const FReal Length, const FReal Thickness, const bool bComputeMTD, FReal& OutTime, FVec3& OutPosition, FVec3& OutNormal, int32& OutFaceIndex, FVec3& OutFaceNormal)
+	{
+		const FGJKSphereSingleSIMD LocalSweptSphere = ComputeLocalSphere(SweptSphere, SweptSphereTM, TestBoxTM);
+		return LocalSweep(LocalSweptSphere, TestBox, TestBox, TestBoxTM, SweepDir, Length, Thickness, bComputeMTD, OutTime, OutPosition, OutNormal, OutFaceIndex, OutFaceNormal);
+	}
+
 	bool SweepSphereVsCapsule(const FSphere& SweptSphere, const FRigidTransform3& SweptSphereTM, const FCapsule& TestCapsule, const FRigidTransform3& TestCapsuleTM,
 		const FVec3& SweepDir, const FReal Length, const FReal Thickness, const bool bComputeMTD,
 		FReal& OutTime, FVec3& OutPosition, FVec3& OutNormal, int32& OutFaceIndex, FVec3& OutFaceNormal)
@@ -70,6 +110,22 @@ namespace Chaos
 		return bResult;
 	}
 
+	bool SweepBoxVsSphere(const TBox<FReal, 3>& SweptBox, const FRigidTransform3& SweptBoxTM, const FSphere& TestSphere, const FRigidTransform3& TestSphereTM,
+		const FVec3& SweepDir, const FReal Length, const FReal Thickness, const bool bComputeMTD,
+		FReal& OutTime, FVec3& OutPosition, FVec3& OutNormal, int32& OutFaceIndex, FVec3& OutFaceNormal)
+	{
+		const FGJKSphereSingleSIMD LocalTestSphere = ComputeLocalSphere(TestSphere, TestSphereTM, SweptBoxTM);
+		return LocalSweep(SweptBox, LocalTestSphere, TestSphere, SweptBoxTM, SweepDir, Length, Thickness, bComputeMTD, OutTime, OutPosition, OutNormal, OutFaceIndex, OutFaceNormal);
+	}
+
+	bool SweepBoxVsCapsule(const TBox<FReal, 3>& SweptBox, const FRigidTransform3& SweptBoxTM, const FCapsule& TestCapsule, const FRigidTransform3& TestCapsuleTM,
+		const FVec3& SweepDir, const FReal Length, const FReal Thickness, const bool bComputeMTD,
+		FReal& OutTime, FVec3& OutPosition, FVec3& OutNormal, int32& OutFaceIndex, FVec3& OutFaceNormal)
+	{
+		const FGJKCapsuleSingleSIMD LocalTestCapsule = ComputeLocalCapsule(TestCapsule, TestCapsuleTM, SweptBoxTM);
+		return LocalSweep(SweptBox, LocalTestCapsule, TestCapsule, SweptBoxTM, SweepDir, Length, Thickness, bComputeMTD, OutTime, OutPosition, OutNormal, OutFaceIndex, OutFaceNormal);
+	}
+
 	bool SweepCapsuleVsSphere(const FCapsule& SweptCapsule, const FRigidTransform3& SweptCapsuleTM, const FSphere& TestSphere, const FRigidTransform3& TestSphereTM,
 		const FVec3& SweepDir, const FReal Length, const FReal Thickness, const bool bComputeMTD,
 		FReal& OutTime, FVec3& OutPosition, FVec3& OutNormal, int32& OutFaceIndex, FVec3& OutFaceNormal)
@@ -91,4 +147,10 @@ namespace Chaos
 		}
 		return bResult;
 	}
+
+	bool SweepCapsuleVsBox(const FCapsule& SweptCapsule, const FRigidTransform3& SweptCapsuleTM, const TBox<FReal, 3>& TestBox, const FRigidTransform3& TestBoxTM, const FVec3& SweepDir, const FReal Length, const FReal Thickness, const bool bComputeMTD, FReal& OutTime, FVec3& OutPosition, FVec3& OutNormal, int32& OutFaceIndex, FVec3& OutFaceNormal)
+	{
+		const FGJKCapsuleSingleSIMD LocalSweptCapsule = ComputeLocalCapsule(SweptCapsule, SweptCapsuleTM, TestBoxTM);
+		return LocalSweep(LocalSweptCapsule, TestBox, TestBox, TestBoxTM, SweepDir, Length, Thickness, bComputeMTD, OutTime, OutPosition, OutNormal, OutFaceIndex, OutFaceNormal);
+	}
 }
diff --git a/Engine/Source/Runtime/Experimental/Chaos/Public/Chaos/GJK.h b/Engine/Source/Runtime/Experimental/Chaos/Public/Chaos/GJK.h
--- a/Engine/Source/Runtime/Experimental/Chaos/Public/Chaos/GJK.h
+++ b/Engine/Source/Runtime/Experimental/Chaos/Public/Chaos/GJK.h
@@ -77,6 +77,64 @@ namespace Chaos
 		}
 	};
 
+	enum EGJKTestSpace
+	{
+		SameSpace,
+		ALocalSpace,
+	};
+
+	template <typename T, EGJKTestSpace TestSpace>
+	struct TGJKSupportHelperSIMD
+	{
+		static void Support(const FGeomGJKHelperSIMD& RESTRICT A, const FGeomGJKHelperSIMD& RESTRICT B, const T& SearchDir, const T& AToBRotation, const T& BToARotation, T& OutSupportA, T& OutSupportB)
+		{
+			OutSupportA = A.SupportFunction(TMakeVectorRegisterFloatFromDouble<T>(SearchDir));
+			OutSupportB = B.SupportFunction(TMakeVectorRegisterFloatFromDouble<T>(AToBRotation), TMakeVectorRegisterFloatFromDouble<T>(BToARotation), TMakeVectorRegisterFloatFromDouble<T>(VectorNegate(SearchDir)));
+		}
+	};
+
+	template <typename T>
+	struct TGJKSupportHelperSIMD<T, EGJKTestSpace::SameSpace>
+	{
+		static void Support(const FGeomGJKHelperSIMD& RESTRICT A, const FGeomGJKHelperSIMD& RESTRICT B, const T& SearchDir, const T& BToARotation, const T& AToBRotation, T& OutSupportA, T& OutSupportB)
+		{
+			OutSupportA = A.SupportFunction(TMakeVectorRegisterFloatFromDouble<T>(SearchDir));
+			OutSupportB = B.SupportFunction(TMakeVectorRegisterFloatFromDouble<T>(VectorNegate(SearchDir)));
+		}
+	};
+
+	template <EGJKTestSpace TestSpace>
+	struct TSupportBAtOriginHelperSIMD
+	{
+		const FGeomGJKHelperSIMD& B;
+		const VectorRegister4Float AToBRotation;
+		const VectorRegister4Float BToARotation;
+		const VectorRegister4Float StartPoint;
+
+		TSupportBAtOriginHelperSIMD(const FGeomGJKHelperSIMD& InB, const VectorRegister4Float& InAToBRotation, const VectorRegister4Float InBToARotation, const VectorRegister4Float& InStartPoint)
+			: B(InB)
+			, AToBRotation(InAToBRotation)
+			, BToARotation(InBToARotation)
+			, StartPoint(InStartPoint)
+		{
+		}
+
+		VectorRegister4Float operator()(const VectorRegister4Float& V) const { return VectorAdd(B.SupportFunction(AToBRotation, BToARotation, V), StartPoint); }
+	};
+
+	template <>
+	struct TSupportBAtOriginHelperSIMD<EGJKTestSpace::SameSpace>
+	{
+		const FGeomGJKHelperSIMD& B;
+		const VectorRegister4Float StartPoint;
+		TSupportBAtOriginHelperSIMD(const FGeomGJKHelperSIMD& InB, const VectorRegister4Float& InAToBRotation, const VectorRegister4Float InBToARotation, const VectorRegister4Float& InStartPoint)
+			: B(InB)
+			, StartPoint(InStartPoint)
+		{
+		}
+		VectorRegister4Float operator()(const VectorRegister4Float& V) const { return VectorAdd(B.SupportFunction(V), StartPoint); }
+	};
+
 	// Check the GJK iteration count against a limit to prevent infinite loops. We should never really hit the limit but 
 	// we are seeing it happen in some cases and more often on some platforms than others so for now we have a warning when it hits.
 	// @todo(chaos): track this issue down (see UE-156361)
@@ -1510,9 +1568,6 @@ namespace Chaos
 
 		return true;
 	}
-	
-
-
 	/** Sweeps one geometry against the other
 	 @A The first geometry
 	 @B The second geometry
@@ -1528,7 +1583,7 @@ namespace Chaos
 	 @return True if the geometries overlap during the sweep, False otherwise
 	 @note If A overlaps B at the start of the ray ("initial overlap" condition) then this function returns true, and sets OutTime = 0, but does not set any other output variables unless bComputeMTD is true.
 	 */
-	template<typename T>
+	template<typename T, EGJKTestSpace TestSpace = EGJKTestSpace::ALocalSpace>
 	bool GJKRaycast2ImplSimd(const FGeomGJKHelperSIMD& RESTRICT A, const FGeomGJKHelperSIMD& RESTRICT B, const T& BToARotation, const T& StartPoint, const T& RayDir,
 		FRealSingle RayLength, FRealSingle& OutTime, T& OutPosition, T& OutNormal, bool bComputeMTD, const T& InitialDir)
 	{
@@ -1573,10 +1628,9 @@ namespace Chaos
 
 		T AToBRotation = VectorQuaternionInverse(BToARotation);
 
-		T SupportA = A.SupportFunction(TMakeVectorRegisterFloatFromDouble<T>(InitialDir));
+		T SupportA, SupportB;
+		TGJKSupportHelperSIMD<T, TestSpace>::Support(A, B, InitialDir, AToBRotation, BToARotation, SupportA, SupportB);
 		As[0] = SupportA;
-
-		T SupportB = B.SupportFunction(TMakeVectorRegisterFloatFromDouble<T>(AToBRotation), TMakeVectorRegisterFloatFromDouble<T>(BToARotation), TMakeVectorRegisterFloatFromDouble<T>(VectorNegate(InitialDir)));
 		Bs[0] = SupportB;
 
 		T Lambda = TVectorZero<T>();
@@ -1612,8 +1666,7 @@ namespace Chaos
 
 			V = VectorNormalizeAccurate(V);
 
-			SupportA = A.SupportFunction(TMakeVectorRegisterFloatFromDouble<T>(V));
-			SupportB = B.SupportFunction(TMakeVectorRegisterFloatFromDouble<T>(AToBRotation), TMakeVectorRegisterFloatFromDouble<T>(BToARotation), TMakeVectorRegisterFloatFromDouble<T>(VectorNegate(V)));
+			TGJKSupportHelperSIMD<T, TestSpace>::Support(A, B, V, AToBRotation, BToARotation, SupportA, SupportB);
 			T P = VectorSubtract(SupportA, SupportB);
 			T W = VectorSubtract(X, P);
 
@@ -1807,17 +1860,7 @@ namespace Chaos
 						VertsB.Add(TMakeVectorRegisterFloatFromDouble<T>(VectorAdd(Bs[i], X)));
 					}
 
-					struct SupportBAtOriginHelper
-					{
-						const FGeomGJKHelperSIMD& B;
-						const VectorRegister4Float AToBRotation;
-						const VectorRegister4Float BToARotation;
-						const VectorRegister4Float StartPoint;
-
-						VectorRegister4Float operator()(const VectorRegister4Float& V) const { return VectorAdd(B.SupportFunction(AToBRotation, BToARotation, V), StartPoint); }
-					};
-
-					SupportBAtOriginHelper SupportBAtOrigin = { B, TMakeVectorRegisterFloatFromDouble<T>(AToBRotation), TMakeVectorRegisterFloatFromDouble<T>(BToARotation), TMakeVectorRegisterFloatFromDouble<T>(StartPoint) };
+					TSupportBAtOriginHelperSIMD<TestSpace> SupportBAtOrigin(B, TMakeVectorRegisterFloatFromDouble<T>(AToBRotation), TMakeVectorRegisterFloatFromDouble<T>(BToARotation), TMakeVectorRegisterFloatFromDouble<T>(StartPoint));
 					VectorRegister4Float Penetration;
 					VectorRegister4Float MTD, ClosestA, ClosestBInA;
 					const EEPAResult EPAResult = VectorEPA(VertsA, VertsB, A, SupportBAtOrigin, Penetration, MTD, ClosestA, ClosestBInA);
@@ -1913,6 +1956,42 @@ namespace Chaos
 		return GJKRaycast2Impl(A, B, StartTM, RayDir, RayLength, OutTime, OutPosition, OutNormal, GivenThicknessA, bComputeMTD, InitialDir, GivenThicknessB);
 	}
 
+	inline bool GJKRaycast2SameSpaceImpl(const FGeomGJKHelperSIMD& A, const FGeomGJKHelperSIMD& B, const TVector<FReal, 3>& Start, const TVector<FReal, 3>& RayDir, const FReal RayLength,
+		FReal& OutTime, TVector<FReal, 3>& OutPosition, TVector<FReal, 3>& OutNormal, const FReal GivenThicknessA, bool bComputeMTD, const TVector<FReal, 3>& InitialDir, const FReal GivenThicknessB)
+	{
+		const VectorRegister4Float UnusedRotation = GlobalVectorConstants::Float0001;
+		const VectorRegister4Float Translation = MakeVectorRegisterFloatFromDouble(MakeVectorRegister(Start.X, Start.Y, Start.Z, 0.0));
+
+		const VectorRegister4Float InitialDirSimd = MakeVectorRegisterFloatFromDouble(MakeVectorRegister(InitialDir[0], InitialDir[1], InitialDir[2], 0.0));
+		const VectorRegister4Float RayDirSimd = MakeVectorRegisterFloatFromDouble(MakeVectorRegister(RayDir[0], RayDir[1], RayDir[2], 0.0));
+
+		FRealSingle OutTimeFloat = 0.0f;
+		VectorRegister4Float OutPositionSimd, OutNormalSimd;
+		const bool Result = GJKRaycast2ImplSimd<VectorRegister4Float, EGJKTestSpace::SameSpace>(A, B, UnusedRotation, Translation, RayDirSimd, static_cast<FRealSingle>(RayLength), OutTimeFloat, OutPositionSimd, OutNormalSimd, bComputeMTD, InitialDirSimd);
+
+		OutTime = static_cast<double>(OutTimeFloat);
+
+		alignas(16) FRealSingle OutFloat[4];
+		VectorStoreAligned(OutNormalSimd, OutFloat);
+		OutNormal.X = OutFloat[0];
+		OutNormal.Y = OutFloat[1];
+		OutNormal.Z = OutFloat[2];
+
+		VectorStoreAligned(OutPositionSimd, OutFloat);
+		OutPosition.X = OutFloat[0];
+		OutPosition.Y = OutFloat[1];
+		OutPosition.Z = OutFloat[2];
+
+		return Result;
+	}
+
+	template <typename T = FReal, typename TGeometryA, typename TGeometryB>
+	bool GJKRaycast2SameSpace(const TGeometryA& A, const TGeometryB& B, const TVector<T, 3>& Start, const TVector<T, 3>& RayDir, const T RayLength,
+		T& OutTime, TVector<T, 3>& OutPosition, TVector<T, 3>& OutNormal, const T GivenThicknessA = 0, bool bComputeMTD = false, const TVector<T, 3>& InitialDir = TVector<T, 3>(1, 0, 0), const T GivenThicknessB = 0)
+	{
+		return GJKRaycast2SameSpaceImpl(A, B, Start, RayDir, RayLength, OutTime, OutPosition, OutNormal, GivenThicknessA, bComputeMTD, InitialDir, GivenThicknessB);
+	}
+
 	template <typename T, typename TGeometryA, typename TGeometryB>
 	UE_DEPRECATED(5.5, "Use GJKDistanceInitialVFromDirection if possible, or GJKDistanceInitialVFromRelativeTransform if original behaviour was required")
 	TVector<T, 3> GJKDistanceInitialV(const TGeometryA& A, const TGeometryB& B, const TRigidTransform<T, 3>& BToATM)
diff --git a/Engine/Source/Runtime/Experimental/Chaos/Public/Chaos/GJKShape.h b/Engine/Source/Runtime/Experimental/Chaos/Public/Chaos/GJKShape.h
--- a/Engine/Source/Runtime/Experimental/Chaos/Public/Chaos/GJKShape.h
+++ b/Engine/Source/Runtime/Experimental/Chaos/Public/Chaos/GJKShape.h
@@ -293,6 +293,90 @@ namespace Chaos
 		FReal Radius;
 	};
 
+	class FGJKSphereSingleSIMD
+	{
+	public:
+		FGJKSphereSingleSIMD(const FVec3f& InCenter, FRealSingle InRadius)
+			: Radius(InRadius)
+		{
+			Center = MakeVectorRegisterFloat(InCenter.X, InCenter.Y, InCenter.Z, 0.0f);
+		}
+
+		VectorRegister4Float SupportCoreSimd(const VectorRegister4Float V, const FRealSingle InMargin) const
+		{
+			return Center;
+		}
+
+		FRealSingle GetRadius() const
+		{
+			return Radius;
+		}
+
+		FRealSingle GetRadiusf() const
+		{
+			return Radius;
+		}
+
+		FRealSingle GetMargin() const
+		{
+			return Radius;
+		}
+
+		FRealSingle GetMarginf() const
+		{
+			return Radius;
+		}
+
+	private:
+		VectorRegister4Float Center;
+		FRealSingle Radius;
+	};
+
+	class FGJKCapsuleSingleSIMD
+	{
+	public:
+		FGJKCapsuleSingleSIMD(const FVec3f& InPointA, const FVec3f& InPointB, const FVec3f& InAxis, const FRealSingle InRadius)
+			: Radius(InRadius)
+		{
+			Axis = MakeVectorRegisterFloat(InAxis.X, InAxis.Y, InAxis.Z, 0.0f);
+			PointA = MakeVectorRegisterFloat(InPointA.X, InPointA.Y, InPointA.Z, 0.0f);
+			PointB = MakeVectorRegisterFloat(InPointB.X, InPointB.Y, InPointB.Z, 0.0f);
+		}
+
+		VectorRegister4Float SupportCoreSimd(const VectorRegister4Float V, FRealSingle InMargin) const
+		{
+			const VectorRegister4Float Dot = VectorDot3(V, Axis);
+			// const bool bIsSecond = Dot >= 0;
+			const VectorRegister4Float Mask = VectorCompareGE(Dot, GlobalVectorConstants::FloatZero);
+			return VectorSelect(Mask, PointB, PointA);
+		}
+
+		FRealSingle GetRadius() const
+		{
+			return Radius;
+		}
+
+		FRealSingle GetRadiusf() const
+		{
+			return Radius;
+		}
+
+		FRealSingle GetMargin() const
+		{
+			return Radius;
+		}
+
+		FRealSingle GetMarginf() const
+		{
+			return Radius;
+		}
+
+	private:
+		VectorRegister4Float Axis;
+		VectorRegister4Float PointA;
+		VectorRegister4Float PointB;
+		FRealSingle Radius;
+	};
 
 	// Utility for creating TGJKShape objects using template parameter deduction
 	template<typename T_SHAPE>
diff --git a/Engine/Source/Runtime/Experimental/Chaos/Public/Chaos/GeometryQueries.h b/Engine/Source/Runtime/Experimental/Chaos/Public/Chaos/GeometryQueries.h
--- a/Engine/Source/Runtime/Experimental/Chaos/Public/Chaos/GeometryQueries.h
+++ b/Engine/Source/Runtime/Experimental/Chaos/Public/Chaos/GeometryQueries.h
@@ -154,8 +154,12 @@ namespace Chaos
 
 	UE_INTERNAL CHAOS_API void TransformSweepResultsToWorld(const bool bResult, const FReal Time, const bool bComputeMTD, const FImplicitObject& TestGeom, const FRigidTransform3& TestGeomTM, const FVec3& LocalDir, const FVec3& LocalPosition, const FVec3& LocalNormal, const int32 FaceIndex, FVec3& OutWorldPosition, FVec3& OutWorldNormal, FVec3& OutWorldFaceNormal);
 	UE_INTERNAL CHAOS_API bool SweepSphereVsSphere(const FSphere& SweptSphere, const FRigidTransform3& SweptSphereTM, const FSphere& TestSphere, const FRigidTransform3& TestSphereTM, const FVec3& SweepDir, const FReal Length, const FReal Thickness, const bool bComputeMTD, FReal& OutTime, FVec3& OutPosition, FVec3& OutNormal, int32& OutFaceIndex, FVec3& OutFaceNormal);
+	UE_INTERNAL CHAOS_API bool SweepSphereVsBox(const FSphere& SweptSphere, const FRigidTransform3& SweptSphereTM, const TBox<FReal, 3>& TestBox, const FRigidTransform3& TestBoxTM, const FVec3& SweepDir, const FReal Length, const FReal Thickness, const bool bComputeMTD, FReal& OutTime, FVec3& OutPosition, FVec3& OutNormal, int32& OutFaceIndex, FVec3& OutFaceNormal);
 	UE_INTERNAL CHAOS_API bool SweepSphereVsCapsule(const FSphere& SweptSphere, const FRigidTransform3& SweptSphereTM, const FCapsule& TestCapsule, const FRigidTransform3& TestCapsuleTM, const FVec3& SweepDir, const FReal Length, const FReal Thickness, const bool bComputeMTD, FReal& OutTime, FVec3& OutPosition, FVec3& OutNormal, int32& OutFaceIndex, FVec3& OutFaceNormal);
+	UE_INTERNAL CHAOS_API bool SweepBoxVsSphere(const TBox<FReal, 3>& SweptBox, const FRigidTransform3& SweptBoxTM, const FSphere& TestSphere, const FRigidTransform3& TestSphereTM, const FVec3& SweepDir, const FReal Length, const FReal Thickness, const bool bComputeMTD, FReal& OutTime, FVec3& OutPosition, FVec3& OutNormal, int32& OutFaceIndex, FVec3& OutFaceNormal);
+	UE_INTERNAL CHAOS_API bool SweepBoxVsCapsule(const TBox<FReal, 3>& SweptBox, const FRigidTransform3& SweptBoxTM, const FCapsule& TestCapsule, const FRigidTransform3& TestCapsuleTM, const FVec3& SweepDir, const FReal Length, const FReal Thickness, const bool bComputeMTD, FReal& OutTime, FVec3& OutPosition, FVec3& OutNormal, int32& OutFaceIndex, FVec3& OutFaceNormal);
 	UE_INTERNAL CHAOS_API bool SweepCapsuleVsSphere(const FCapsule& SweptCapsule, const FRigidTransform3& SweptCapsuleTM, const FSphere& TestSphere, const FRigidTransform3& TestSphereTM, const FVec3& SweepDir, const FReal Length, const FReal Thickness, const bool bComputeMTD, FReal& OutTime, FVec3& OutPosition, FVec3& OutNormal, int32& OutFaceIndex, FVec3& OutFaceNormal);
+	UE_INTERNAL CHAOS_API bool SweepCapsuleVsBox(const FCapsule& SweptCapsule, const FRigidTransform3& SweptCapsuleTM, const TBox<FReal, 3>& TestBox, const FRigidTransform3& TestBoxTM, const FVec3& SweepDir, const FReal Length, const FReal Thickness, const bool bComputeMTD, FReal& OutTime, FVec3& OutPosition, FVec3& OutNormal, int32& OutFaceIndex, FVec3& OutFaceNormal);
 
 	template <EImplicitObjectType SweptType>
 	UE_INTERNAL bool TryRunSpecializedSweep(const FImplicitObject& SweptShape, const FRigidTransform3& SweptShapeTM, const FImplicitObject& TestGeom, const FRigidTransform3& TestGeomTM, const FVec3& SweepDir, const FReal Length, const FReal Thickness, const bool bComputeMTD, bool& bOutResult, FReal& OutTime, FVec3& OutPosition, FVec3& OutNormal, int32& OutFaceIndex, FVec3& OutFaceNormal)
@@ -170,6 +174,13 @@ namespace Chaos
 				bOutResult = SweepSphereVsSphere(SweptSphere, SweptShapeTM, TestSphere, TestGeomTM, SweepDir, Length, Thickness, bComputeMTD, OutTime, OutPosition, OutNormal, OutFaceIndex, OutFaceNormal);
 				return true;
 			}
+			else if (TestType == ImplicitObjectType::Box)
+			{
+				const FSphere& SweptSphere = SweptShape.GetObjectChecked<FSphere>();
+				const TBox<FReal, 3>& TestBox = TestGeom.GetObjectChecked<TBox<FReal, 3>>();
+				bOutResult = SweepSphereVsBox(SweptSphere, SweptShapeTM, TestBox, TestGeomTM, SweepDir, Length, Thickness, bComputeMTD, OutTime, OutPosition, OutNormal, OutFaceIndex, OutFaceNormal);
+				return true;
+			}
 			else if (TestType == ImplicitObjectType::Capsule)
 			{
 				const FSphere& SweptSphere = SweptShape.GetObjectChecked<FSphere>();
@@ -178,6 +189,23 @@ namespace Chaos
 				return true;
 			}
 		}
+		else if (SweptType == ImplicitObjectType::Box)
+		{
+			if (TestType == ImplicitObjectType::Sphere)
+			{
+				const TBox<FReal, 3>& SweptBox = SweptShape.GetObjectChecked<TBox<FReal, 3>>();
+				const FSphere& TestSphere = TestGeom.GetObjectChecked<FSphere>();
+				bOutResult = SweepBoxVsSphere(SweptBox, SweptShapeTM, TestSphere, TestGeomTM, SweepDir, Length, Thickness, bComputeMTD, OutTime, OutPosition, OutNormal, OutFaceIndex, OutFaceNormal);
+				return true;
+			}
+			else if (TestType == ImplicitObjectType::Capsule)
+			{
+				const TBox<FReal, 3>& SweptBox = SweptShape.GetObjectChecked<TBox<FReal, 3>>();
+				const FCapsule& TestCapsule = TestGeom.GetObjectChecked<FCapsule>();
+				bOutResult = SweepBoxVsCapsule(SweptBox, SweptShapeTM, TestCapsule, TestGeomTM, SweepDir, Length, Thickness, bComputeMTD, OutTime, OutPosition, OutNormal, OutFaceIndex, OutFaceNormal);
+				return true;
+			}
+		}
 		else if (SweptType == ImplicitObjectType::Capsule)
 		{
 			if (TestType == ImplicitObjectType::Sphere)
@@ -187,6 +215,13 @@ namespace Chaos
 				bOutResult = SweepCapsuleVsSphere(SweptCapsule, SweptShapeTM, TestSphere, TestGeomTM, SweepDir, Length, Thickness, bComputeMTD, OutTime, OutPosition, OutNormal, OutFaceIndex, OutFaceNormal);
 				return true;
 			}
+			else if (TestType == ImplicitObjectType::Box)
+			{
+				const FCapsule& SweptCapsule = SweptShape.GetObjectChecked<FCapsule>();
+				const TBox<FReal, 3>& TestBox = TestGeom.GetObjectChecked<TBox<FReal, 3>>();
+				bOutResult = SweepCapsuleVsBox(SweptCapsule, SweptShapeTM, TestBox, TestGeomTM, SweepDir, Length, Thickness, bComputeMTD, OutTime, OutPosition, OutNormal, OutFaceIndex, OutFaceNormal);
+				return true;
+			}
 		}
 		return false;
 	}
