DataflowNodeFactory undo API change that was part of code-size-reduction CL

diff --git a/Engine/Plugins/Experimental/Dataflow/Source/DataflowEditor/Private/Dataflow/DataflowAssetEditUtils.cpp b/Engine/Plugins/Experimental/Dataflow/Source/DataflowEditor/Private/Dataflow/DataflowAssetEditUtils.cpp
--- a/Engine/Plugins/Experimental/Dataflow/Source/DataflowEditor/Private/Dataflow/DataflowAssetEditUtils.cpp
+++ b/Engine/Plugins/Experimental/Dataflow/Source/DataflowEditor/Private/Dataflow/DataflowAssetEditUtils.cpp
@@ -125,11 +125,8 @@ namespace UE::Dataflow
 				{
 					.Guid = FGuid::NewGuid(),
 					.Type = NodeTypeName,
-					.NodeParams =
-					{
-						.Name = GenerateUniqueObjectName(Dataflow, NodeName),
-						.OwningObject = Dataflow,
-					},
+					.Name = GenerateUniqueObjectName(Dataflow, NodeName),
+					.OwningObject = Dataflow,
 				};
 
 				return Factory->NewNodeFromRegisteredType(*Dataflow->GetDataflow(), Parameters);
diff --git a/Engine/Source/Runtime/Experimental/Dataflow/Core/Private/Dataflow/DataflowGraph.cpp b/Engine/Source/Runtime/Experimental/Dataflow/Core/Private/Dataflow/DataflowGraph.cpp
--- a/Engine/Source/Runtime/Experimental/Dataflow/Core/Private/Dataflow/DataflowGraph.cpp
+++ b/Engine/Source/Runtime/Experimental/Dataflow/Core/Private/Dataflow/DataflowGraph.cpp
@@ -442,7 +442,7 @@ namespace UE::Dataflow
 			FName ArNodeName;
 			Ar << ArGuid << ArType << ArNodeName;
 
-			TSharedPtr<FDataflowNode> Node = FNodeFactory::GetInstance()->NewNodeFromRegisteredType(*InGraph, { ArGuid, ArType, { ArNodeName, OwningObject } });
+			TSharedPtr<FDataflowNode> Node = FNodeFactory::GetInstance()->NewNodeFromRegisteredType(*InGraph, { ArGuid, ArType, ArNodeName, OwningObject });
 			DATAFLOW_OPTIONAL_BLOCK_READ_BEGIN(Node != nullptr)
 			{
 				ensure(!NodeGuidMap.Contains(ArGuid));
diff --git a/Engine/Source/Runtime/Experimental/Dataflow/Core/Private/Dataflow/DataflowNodeFactory.cpp b/Engine/Source/Runtime/Experimental/Dataflow/Core/Private/Dataflow/DataflowNodeFactory.cpp
--- a/Engine/Source/Runtime/Experimental/Dataflow/Core/Private/Dataflow/DataflowNodeFactory.cpp
+++ b/Engine/Source/Runtime/Experimental/Dataflow/Core/Private/Dataflow/DataflowNodeFactory.cpp
@@ -47,11 +47,8 @@ namespace UE::Dataflow
 		{
 			.Guid = FGuid::NewGuid(),
 			.Type = FName{},
-			.NodeParams =
-			{
-				.Name = StaticType,
-				.OwningObject = GetTransientPackage()
-			}
+			.Name = StaticType,
+			.OwningObject = GetTransientPackage()
 		};
 
 		// Create a node and take over ownership temporarily...
diff --git a/Engine/Source/Runtime/Experimental/Dataflow/Core/Public/Dataflow/DataflowNodeFactory.h b/Engine/Source/Runtime/Experimental/Dataflow/Core/Public/Dataflow/DataflowNodeFactory.h
--- a/Engine/Source/Runtime/Experimental/Dataflow/Core/Public/Dataflow/DataflowNodeFactory.h
+++ b/Engine/Source/Runtime/Experimental/Dataflow/Core/Public/Dataflow/DataflowNodeFactory.h
@@ -15,7 +15,8 @@ namespace UE::Dataflow
 	{
 		FGuid Guid;
 		FName Type;
-		FNodeParameters NodeParams;
+		FName Name;
+		UObject* OwningObject = nullptr;
 	};
 
 	struct FFactoryParameters
@@ -87,7 +88,7 @@ namespace UE::Dataflow
 		{
 			auto CreationFunction = [](const UE::Dataflow::FNewNodeParameters& InParam)
 			{
-				TUniquePtr<FDataflowNode> Val = MakeUnique<T>(InParam.NodeParams, InParam.Guid);
+				TUniquePtr<FDataflowNode> Val = MakeUnique<T>(FNodeParameters { InParam.Name, InParam.OwningObject }, InParam.Guid);
 				Val->ValidateProperties();
 				Val->ValidateConnections();
 				return Val;
@@ -132,7 +133,7 @@ namespace UE::Dataflow
 
 		template<class T> TSharedPtr<T> NewNode(FGraph& Graph, const FNewNodeParameters& Param)
 		{
-			return Graph.AddNode(new T(Param.NodeParams.Name, Param.Guid));
+			return Graph.AddNode(new T(Param.Name, Param.Guid));
 		}
 
 		DATAFLOWCORE_API void RegisterDefaultNodes();
