Added support for generating duplicated iad and unduplicated chunks.

diff --git a/Engine/Source/Runtime/Engine/Private/StreamableManager.cpp b/Engine/Source/Runtime/Engine/Private/StreamableManager.cpp
--- a/Engine/Source/Runtime/Engine/Private/StreamableManager.cpp
+++ b/Engine/Source/Runtime/Engine/Private/StreamableManager.cpp
@@ -462,14 +462,19 @@ namespace UE::StreamableManager::Private
 
 			InstallArgs.Priority = Priority;
 
+			
 			InstallRequest = IoStore->Install(MoveTemp(InstallArgs),
-				[WeakOwningHandle = OwningHandle.ToWeakPtr()](const UE::IoStore::FOnDemandInstallResult& OnDemandInstallResult)
+				[WeakOwningHandle = OwningHandle.ToWeakPtr(), DebugString = OwningHandle->GetDebugName()](const UE::IoStore::FOnDemandInstallResult& OnDemandInstallResult)
 				{
 					TSharedPtr<FStreamableHandle> OwningHandle = WeakOwningHandle.Pin();
 					if (OwningHandle && OwningHandle->DownloadContext)
 					{
 						OwningHandle->DownloadContext->OnInstallComplete(OwningHandle.ToSharedRef(), OnDemandInstallResult);
 					}
+					else
+					{
+						UE_LOG(LogStreamableManager, Warning, TEXT("Owning handle was cleaned up before completion event was run. Owning handle name %s"), *DebugString);
+					}
 				},
 				[WeakOwningHandle = OwningHandle.ToWeakPtr()](const UE::IoStore::FOnDemandInstallProgress& InProgress)
 				{
diff --git a/Engine/Source/Runtime/Experimental/IoStore/OnDemand/Private/OnDemandContentInstaller.cpp b/Engine/Source/Runtime/Experimental/IoStore/OnDemand/Private/OnDemandContentInstaller.cpp
--- a/Engine/Source/Runtime/Experimental/IoStore/OnDemand/Private/OnDemandContentInstaller.cpp
+++ b/Engine/Source/Runtime/Experimental/IoStore/OnDemand/Private/OnDemandContentInstaller.cpp
@@ -986,6 +986,9 @@ void FOnDemandContentInstaller::CompleteInstallRequest(FRequest& Request)
 	FRequest::FInstall& InstallRequest = Request.AsInstall();
 	ensure(Request.ErrorCode != EIoErrorCode::Unknown);
 
+#if !UE_BUILD_SHIPPING
+	TArray<FString> DebugInfo;
+#endif
 	// Mark all resolved chunk(s) as referenced by the content handle and notify the package store to update
 	if (Request.ErrorCode == EIoErrorCode::Ok && InstallRequest.ResolvedChunkCount > 0)
 	{
@@ -1000,6 +1003,12 @@ void FOnDemandContentInstaller::CompleteInstallRequest(FRequest& Request)
 				for (int32 EntryIndex : ResolvedChunks.EntryIndices)
 				{
 					References.Indices[EntryIndex] = true;
+#if !UE_BUILD_SHIPPING
+					if (UE_LOG_ACTIVE(LogIoStoreOnDemand, Verbose))
+					{
+						DebugInfo.Add(FString::Printf(TEXT("ID:%s, Size:%d"), *LexToString(Container->ChunkEntries[EntryIndex].Hash), Container->ChunkEntries[EntryIndex].GetDiskSize() ));
+					}
+#endif
 				}
 			}
 		}
@@ -1026,6 +1035,10 @@ void FOnDemandContentInstaller::CompleteInstallRequest(FRequest& Request)
 		DurationCycles,
 		Request.ErrorCode);
 
+#if !UE_BUILD_SHIPPING
+	FString DebugInfoString = FString::Join(DebugInfo, TEXT(","));
+	UE_LOG(LogIoStoreOnDemand, Verbose, TEXT("ContentHandle=(%s), ResolvedChunks(%s)"), *LexToString(InstallRequest.Args.ContentHandle), *DebugInfoString);
+#endif
 	UE_LOG(LogIoStoreOnDemand, Verbose, TEXT("Install request completed, Result='%s', SeqNo=%u, Priority=%d, ContentHandle=(%s), ContentSize=%.2lf MiB, InstallSize=%.2lf MiB, CacheHitRatio=%d%%, Duration=%dms"),
 		GetIoErrorText(InstallResult.Status.GetErrorCode()), Request.SeqNo, Request.Priority, *LexToString(InstallRequest.Args.ContentHandle), double(InstallResult.Progress.TotalContentSize) / 1024.0 / 1024.0,
 		double(InstallResult.Progress.TotalInstallSize) / 1024.0 / 1024.0, int32(CacheHitRatio * 100), int32(InstallResult.DurationInSeconds * 1000));
