[ChaosVD] Restoring Debug Draw Box data wrapper Name that was changed by accident in a previous CL.

The name had a typo in the original version, and when I moved the definition to the cpp file, I corrected it unintentionally.

This caused an issue where CVD couldn't find the data processor for debug draw boxes, and therefore these are missing from the visualizations.

diff --git a/Engine/Source/Runtime/Experimental/ChaosVisualDebugger/Private/DataWrappers/ChaosVDDebugShapeDataWrapper.cpp b/Engine/Source/Runtime/Experimental/ChaosVisualDebugger/Private/DataWrappers/ChaosVDDebugShapeDataWrapper.cpp
--- a/Engine/Source/Runtime/Experimental/ChaosVisualDebugger/Private/DataWrappers/ChaosVDDebugShapeDataWrapper.cpp
+++ b/Engine/Source/Runtime/Experimental/ChaosVisualDebugger/Private/DataWrappers/ChaosVDDebugShapeDataWrapper.cpp
@@ -9,7 +9,7 @@
 FStringView FChaosVDDebugDrawSphereDataWrapper::WrapperTypeName = TEXT("FChaosVDDebugDrawSphereDataWrapper");
 FStringView FChaosVDDebugDrawLineDataWrapper::WrapperTypeName = TEXT("FChaosVDDebugDrawLineDataWrapper");
 FStringView FChaosVDDebugDrawImplicitObjectDataWrapper::WrapperTypeName = TEXT("FChaosVDDebugDrawImplicitObjectDataWrapper");
-FStringView FChaosVDDebugDrawBoxDataWrapper::WrapperTypeName = TEXT("FChaosVDDebugDrawBoxDataWrapper");
+FStringView FChaosVDDebugDrawBoxDataWrapper::WrapperTypeName = TEXT("FChaosVDDDebugDrawBoxDataWrapper");
 
 void FChaosVDDebugDrawShapeBase::SerializeBase_Internal(FArchive& Ar)
 {
