[ChaosVD] Adding initial support to enable CVD in Shipping configurations.

- Updated the helper method IsChaosVisualDebuggerSupported in Module rules to not take into account target configurations or trace support.
- Updated the bCompileChaosVisualDebuggerSupport flag to fallback as false for shipping configurations, unless it was explicitly set to true or false.
- Updated CVD's code to use the new Trace Minimal API, which is supported in shipping configurations.

#tests Tested that CVD is still compiled out in non-shipping builds by default., and that it works when it is explcitly compiled in Shipping builds (Editor/Program/Server/Client)

diff --git a/Engine/Source/Programs/ChaosVisualDebugger/ChaosVisualDebugger.Target.cs b/Engine/Source/Programs/ChaosVisualDebugger/ChaosVisualDebugger.Target.cs
--- a/Engine/Source/Programs/ChaosVisualDebugger/ChaosVisualDebugger.Target.cs
+++ b/Engine/Source/Programs/ChaosVisualDebugger/ChaosVisualDebugger.Target.cs
@@ -38,6 +38,8 @@ public class ChaosVisualDebuggerTarget : TargetRules
 		bCompileNavmeshSegmentLinks = false;
 		bWithLiveCoding = false;
 		bIncludePluginsForTargetPlatforms = false;
+
+		bCompileChaosVisualDebuggerSupport = true;
 		
 		OptedInModulePlatforms = new UnrealTargetPlatform[] { UnrealTargetPlatform.Win64 };
 
diff --git a/Engine/Source/Programs/UnrealBuildTool/Configuration/ModuleRules.cs b/Engine/Source/Programs/UnrealBuildTool/Configuration/ModuleRules.cs
--- a/Engine/Source/Programs/UnrealBuildTool/Configuration/ModuleRules.cs
+++ b/Engine/Source/Programs/UnrealBuildTool/Configuration/ModuleRules.cs
@@ -1648,7 +1648,7 @@ namespace UnrealBuildTool
 		/// </summary>
 		protected bool IsChaosVisualDebuggerSupported(ReadOnlyTargetRules target)
 		{
-			return target.bCompileChaosVisualDebuggerSupport && target.Configuration != UnrealTargetConfiguration.Shipping && target.bEnableTrace;
+			return target.bCompileChaosVisualDebuggerSupport;
 		}
 
 		/// <summary>
diff --git a/Engine/Source/Programs/UnrealBuildTool/Configuration/TargetRules.cs b/Engine/Source/Programs/UnrealBuildTool/Configuration/TargetRules.cs
--- a/Engine/Source/Programs/UnrealBuildTool/Configuration/TargetRules.cs
+++ b/Engine/Source/Programs/UnrealBuildTool/Configuration/TargetRules.cs
@@ -981,7 +981,21 @@ namespace UnrealBuildTool
 		[XmlConfigFile(Name = "bCompileChaosVisualDebuggerSupport")]
 		[CommandLine("-CompileChaosVisualDebuggerSupport")]
 		[RequiresUniqueBuildEnvironment]
-		public bool bCompileChaosVisualDebuggerSupport { get; set; } = true;
+		public bool bCompileChaosVisualDebuggerSupport 
+		{
+			get 
+			{
+				if (bEnableChaosVisualDebuggerSupport.HasValue)
+				{
+					return bEnableChaosVisualDebuggerSupport.Value;
+				}
+
+				// By default, if we didn't explicitly ask for CVD support, we should only enable it in non-shipping builds and only if trace is enabled
+				return Configuration != UnrealTargetConfiguration.Shipping && bEnableTrace;
+			}  
+			set => bEnableChaosVisualDebuggerSupport = value;
+		}
+		private bool? bEnableChaosVisualDebuggerSupport;
 
 		/// <summary>
 		/// Whether to include ICU unicode/i18n support in Core.
diff --git a/Engine/Source/Runtime/Experimental/Chaos/Private/ChaosVisualDebugger/ChaosVisualDebuggerTrace.cpp b/Engine/Source/Runtime/Experimental/Chaos/Private/ChaosVisualDebugger/ChaosVisualDebuggerTrace.cpp
--- a/Engine/Source/Runtime/Experimental/Chaos/Private/ChaosVisualDebugger/ChaosVisualDebuggerTrace.cpp
+++ b/Engine/Source/Runtime/Experimental/Chaos/Private/ChaosVisualDebugger/ChaosVisualDebuggerTrace.cpp
@@ -31,26 +31,28 @@
 #include "Misc/TransactionallySafeRWLock.h"
 #include "Templates/UniquePtr.h"
 
-UE_TRACE_EVENT_DEFINE(ChaosVDLogger, ChaosVDSolverFrameStart)
-
-UE_TRACE_EVENT_DEFINE(ChaosVDLogger, ChaosVDSolverFrameEnd)
-
-UE_TRACE_CHANNEL_DEFINE(ChaosVDChannel);
-UE_TRACE_EVENT_DEFINE(ChaosVDLogger, ChaosVDParticle)
-UE_TRACE_EVENT_DEFINE(ChaosVDLogger, ChaosVDParticleDestroyed)
-UE_TRACE_EVENT_DEFINE(ChaosVDLogger, ChaosVDSolverStepStart)
-UE_TRACE_EVENT_DEFINE(ChaosVDLogger, ChaosVDSolverStepEnd)
-UE_TRACE_EVENT_DEFINE(ChaosVDLogger, ChaosVDBinaryDataStart)
-UE_TRACE_EVENT_DEFINE(ChaosVDLogger, ChaosVDBinaryDataContent)
-UE_TRACE_EVENT_DEFINE(ChaosVDLogger, ChaosVDBinaryDataEnd)
-UE_TRACE_EVENT_DEFINE(ChaosVDLogger, ChaosVDSolverSimulationSpace)
-UE_TRACE_EVENT_DEFINE(ChaosVDLogger, ChaosVDDummyEvent)
-UE_TRACE_EVENT_DEFINE(ChaosVDLogger, ChaosVDNonSolverLocation)
-UE_TRACE_EVENT_DEFINE(ChaosVDLogger, ChaosVDNonSolverTransform)
-UE_TRACE_EVENT_DEFINE(ChaosVDLogger, ChaosVDNetworkTickOffset)
-UE_TRACE_EVENT_DEFINE(ChaosVDLogger, ChaosVDRolledBackDataID)
-UE_TRACE_EVENT_DEFINE(ChaosVDLogger, ChaosVDUsesAutoRTFM)
-UE_TRACE_EVENT_DEFINE(ChaosVDLogger, ChaosVDTraceRelevancyVolume)
+static_assert(CHAOS_VISUAL_DEBUGGER_WITHOUT_TRACE || UE_TRACE_MINIMAL_ENABLED || UE_TRACE_ENABLED, "Chaos Visual Debugger support requires Trace or Trace minimal support. Please enable trace support with the Unreal Build Tool flag bEnableTrace, disable cvd support with bCompileChaosVisualDebuggerSupport or allows CVD support without trace by defining CHAOS_VISUAL_DEBUGGER_WITHOUT_TRACE to 1");
+
+UE_TRACE_MINIMAL_EVENT_DEFINE(ChaosVDLogger, ChaosVDSolverFrameStart)
+
+UE_TRACE_MINIMAL_EVENT_DEFINE(ChaosVDLogger, ChaosVDSolverFrameEnd)
+
+UE_TRACE_MINIMAL_CHANNEL_DEFINE(ChaosVDChannel);
+UE_TRACE_MINIMAL_EVENT_DEFINE(ChaosVDLogger, ChaosVDParticle)
+UE_TRACE_MINIMAL_EVENT_DEFINE(ChaosVDLogger, ChaosVDParticleDestroyed)
+UE_TRACE_MINIMAL_EVENT_DEFINE(ChaosVDLogger, ChaosVDSolverStepStart)
+UE_TRACE_MINIMAL_EVENT_DEFINE(ChaosVDLogger, ChaosVDSolverStepEnd)
+UE_TRACE_MINIMAL_EVENT_DEFINE(ChaosVDLogger, ChaosVDBinaryDataStart)
+UE_TRACE_MINIMAL_EVENT_DEFINE(ChaosVDLogger, ChaosVDBinaryDataContent)
+UE_TRACE_MINIMAL_EVENT_DEFINE(ChaosVDLogger, ChaosVDBinaryDataEnd)
+UE_TRACE_MINIMAL_EVENT_DEFINE(ChaosVDLogger, ChaosVDSolverSimulationSpace)
+UE_TRACE_MINIMAL_EVENT_DEFINE(ChaosVDLogger, ChaosVDDummyEvent)
+UE_TRACE_MINIMAL_EVENT_DEFINE(ChaosVDLogger, ChaosVDNonSolverLocation)
+UE_TRACE_MINIMAL_EVENT_DEFINE(ChaosVDLogger, ChaosVDNonSolverTransform)
+UE_TRACE_MINIMAL_EVENT_DEFINE(ChaosVDLogger, ChaosVDNetworkTickOffset)
+UE_TRACE_MINIMAL_EVENT_DEFINE(ChaosVDLogger, ChaosVDRolledBackDataID)
+UE_TRACE_MINIMAL_EVENT_DEFINE(ChaosVDLogger, ChaosVDUsesAutoRTFM)
+UE_TRACE_MINIMAL_EVENT_DEFINE(ChaosVDLogger, ChaosVDTraceRelevancyVolume)
 
 namespace Chaos::VisualDebugger::Cvars
 {
@@ -208,7 +210,7 @@ void FChaosVisualDebuggerTrace::TraceParticleDestroyed(const Chaos::FGeometryPar
 
 	UE_AUTORTFM_ONCOMMIT(SolverID, Cycle, ParticleID)
 	{
-		UE_TRACE_LOG(ChaosVDLogger, ChaosVDParticleDestroyed, ChaosVDChannel)
+		UE_TRACE_MINIMAL_LOG(ChaosVDLogger, ChaosVDParticleDestroyed, ChaosVDChannel)
 			<< ChaosVDParticleDestroyed.SolverID(SolverID)
 			<< ChaosVDParticleDestroyed.Cycle(Cycle)
 			<< ChaosVDParticleDestroyed.ParticleID(ParticleID);
@@ -606,7 +608,7 @@ void FChaosVisualDebuggerTrace::TraceSolverFrameStart(const FChaosVDContext& Con
 
 	UE_AUTORTFM_OPEN
 	{
-		UE_TRACE_LOG(ChaosVDLogger, ChaosVDSolverFrameStart, ChaosVDChannel)
+		UE_TRACE_MINIMAL_LOG(ChaosVDLogger, ChaosVDSolverFrameStart, ChaosVDChannel)
 			<< ChaosVDSolverFrameStart.SolverID(ContextData.Id)
 			<< ChaosVDSolverFrameStart.Cycle(FPlatformTime::Cycles64())
 			<< ChaosVDSolverFrameStart.DebugName(*InDebugName, InDebugName.Len())
@@ -640,7 +642,7 @@ void FChaosVisualDebuggerTrace::TraceSolverFrameEnd(const FChaosVDContext& Conte
 
 	UE_AUTORTFM_OPEN
 	{
-		UE_TRACE_LOG(ChaosVDLogger, ChaosVDSolverFrameEnd, ChaosVDChannel)
+		UE_TRACE_MINIMAL_LOG(ChaosVDLogger, ChaosVDSolverFrameEnd, ChaosVDChannel)
 			<< ChaosVDSolverFrameEnd.SolverID(ContextData.Id)
 			<< ChaosVDSolverFrameEnd.Cycle(FPlatformTime::Cycles64());
 	};
@@ -663,7 +665,7 @@ void FChaosVisualDebuggerTrace::TraceSolverStepStart(FStringView StepName)
 
 	UE_AUTORTFM_OPEN
 	{
-		UE_TRACE_LOG(ChaosVDLogger, ChaosVDSolverStepStart, ChaosVDChannel)
+		UE_TRACE_MINIMAL_LOG(ChaosVDLogger, ChaosVDSolverStepStart, ChaosVDChannel)
 			<< ChaosVDSolverStepStart.Cycle(FPlatformTime::Cycles64())
 			<< ChaosVDSolverStepStart.SolverID(CVDContextData->Id)
 			<< ChaosVDSolverStepStart.StepName(StepName.GetData(), GetNum(StepName));
@@ -687,7 +689,7 @@ void FChaosVisualDebuggerTrace::TraceSolverStepEnd()
 
 	UE_AUTORTFM_OPEN
 	{
-		UE_TRACE_LOG(ChaosVDLogger, ChaosVDSolverStepEnd, ChaosVDChannel)
+		UE_TRACE_MINIMAL_LOG(ChaosVDLogger, ChaosVDSolverStepEnd, ChaosVDChannel)
 			<< ChaosVDSolverStepEnd.Cycle(FPlatformTime::Cycles64())
 			<< ChaosVDSolverStepEnd.SolverID(CVDContextData->Id);
 	};
@@ -710,7 +712,7 @@ void FChaosVisualDebuggerTrace::TraceSolverSimulationSpace(const Chaos::FRigidTr
 
 	UE_AUTORTFM_OPEN
 	{
-		UE_TRACE_LOG(ChaosVDLogger, ChaosVDSolverSimulationSpace, ChaosVDChannel)
+		UE_TRACE_MINIMAL_LOG(ChaosVDLogger, ChaosVDSolverSimulationSpace, ChaosVDChannel)
 			<< ChaosVDSolverSimulationSpace.Cycle(FPlatformTime::Cycles64())
 			<< ChaosVDSolverSimulationSpace.SolverID(CVDContextData->Id)
 			<< CVD_TRACE_VECTOR_ON_EVENT(ChaosVDSolverSimulationSpace, Position, Transform.GetLocation())
@@ -754,7 +756,7 @@ void FChaosVisualDebuggerTrace::TraceBinaryData(TConstArrayView<uint8> InData, F
 		constexpr uint32 MaxChunkSize = TNumericLimits<uint16>::Max();
 		const uint32 ChunkNum = (DataSize + MaxChunkSize - 1) / MaxChunkSize;
 
-		UE_TRACE_LOG(ChaosVDLogger, ChaosVDBinaryDataStart, ChaosVDChannel)
+		UE_TRACE_MINIMAL_LOG(ChaosVDLogger, ChaosVDBinaryDataStart, ChaosVDChannel)
 			<< ChaosVDBinaryDataStart.Cycle(FPlatformTime::Cycles64())
 			<< ChaosVDBinaryDataStart.TypeName(TypeName.GetData(), TypeName.Len())
 			<< ChaosVDBinaryDataStart.DataID(DataID)
@@ -768,7 +770,7 @@ void FChaosVisualDebuggerTrace::TraceBinaryData(TConstArrayView<uint8> InData, F
 			const uint16 Size = static_cast<uint16>(FMath::Min(RemainingSize, MaxChunkSize));
 			const uint8* ChunkData = DataViewToTrace.GetData() + MaxChunkSize * Index;
 
-			UE_TRACE_LOG(ChaosVDLogger, ChaosVDBinaryDataContent, ChaosVDChannel)
+			UE_TRACE_MINIMAL_LOG(ChaosVDLogger, ChaosVDBinaryDataContent, ChaosVDChannel)
 				<< ChaosVDBinaryDataContent.Cycle(FPlatformTime::Cycles64())
 				<< ChaosVDBinaryDataContent.DataID(DataID)
 				<< ChaosVDBinaryDataContent.RawData(ChunkData, Size);
@@ -789,14 +791,14 @@ void FChaosVisualDebuggerTrace::TraceBinaryData(TConstArrayView<uint8> InData, F
 	
 	UE_AUTORTFM_ONCOMMIT(DataID)
 	{
-		UE_TRACE_LOG(ChaosVDLogger, ChaosVDBinaryDataEnd, ChaosVDChannel)
+		UE_TRACE_MINIMAL_LOG(ChaosVDLogger, ChaosVDBinaryDataEnd, ChaosVDChannel)
 			<< ChaosVDBinaryDataEnd.Cycle(FPlatformTime::Cycles64())
 			<< ChaosVDBinaryDataEnd.DataID(DataID);
 	};
 
 	UE_AUTORTFM_ONABORT(DataID)
 	{
-		UE_TRACE_LOG(ChaosVDLogger, ChaosVDRolledBackDataID, ChaosVDChannel)
+		UE_TRACE_MINIMAL_LOG(ChaosVDLogger, ChaosVDRolledBackDataID, ChaosVDChannel)
 			<< ChaosVDRolledBackDataID.DataID(DataID);
 	};
 }
@@ -910,7 +912,7 @@ void FChaosVisualDebuggerTrace::TraceNonSolverLocation(const FVector& InLocation
 		return;
 	}
 
-	UE_TRACE_LOG(ChaosVDLogger, ChaosVDNonSolverLocation, ChaosVDChannel)
+	UE_TRACE_MINIMAL_LOG(ChaosVDLogger, ChaosVDNonSolverLocation, ChaosVDChannel)
 			<< ChaosVDNonSolverLocation.Cycle(FPlatformTime::Cycles64())
 			<< CVD_TRACE_VECTOR_ON_EVENT(ChaosVDNonSolverLocation, Position, InLocation)
 			<< ChaosVDNonSolverLocation.DebugName(DebugNameID.GetData(), DebugNameID.Len());
@@ -923,7 +925,7 @@ void FChaosVisualDebuggerTrace::TraceNonSolverTransform(const FTransform& InTran
 		return;
 	}
 
-	UE_TRACE_LOG(ChaosVDLogger, ChaosVDNonSolverTransform, ChaosVDChannel)
+	UE_TRACE_MINIMAL_LOG(ChaosVDLogger, ChaosVDNonSolverTransform, ChaosVDChannel)
 		<< ChaosVDNonSolverTransform.Cycle(FPlatformTime::Cycles64())
 		<< CVD_TRACE_VECTOR_ON_EVENT(ChaosVDNonSolverTransform, Position, InTransform.GetLocation())
 		<< CVD_TRACE_VECTOR_ON_EVENT(ChaosVDNonSolverTransform, Scale, InTransform.GetScale3D())
@@ -1078,7 +1080,7 @@ void FChaosVisualDebuggerTrace::TraceNetworkTickOffset(int32 TickOffset, int32 S
 
 	UE_AUTORTFM_OPEN
 	{
-		UE_TRACE_LOG(ChaosVDLogger, ChaosVDNetworkTickOffset, ChaosVDChannel)
+		UE_TRACE_MINIMAL_LOG(ChaosVDLogger, ChaosVDNetworkTickOffset, ChaosVDChannel)
 			<< ChaosVDNetworkTickOffset.Offset(TickOffset)
 			<< ChaosVDNetworkTickOffset.SolverID(SolverID);
 	};
@@ -1287,7 +1289,7 @@ void FChaosVisualDebuggerTrace::SetTraceRelevancyVolume(const FBox& InTraceRelev
 {
 	SessionState.TraceRelevancyVolume = InTraceRelevancyVolume;
 
-	UE_TRACE_LOG(ChaosVDLogger, ChaosVDTraceRelevancyVolume, ChaosVDChannel)
+	UE_TRACE_MINIMAL_LOG(ChaosVDLogger, ChaosVDTraceRelevancyVolume, ChaosVDChannel)
 		<< CVD_TRACE_VECTOR_ON_EVENT(ChaosVDTraceRelevancyVolume, BoxMin, InTraceRelevancyVolume.Min)
 		<< CVD_TRACE_VECTOR_ON_EVENT(ChaosVDTraceRelevancyVolume, BoxMax, InTraceRelevancyVolume.Max);
 }
@@ -1431,7 +1433,7 @@ void FChaosVisualDebuggerTrace::HandleRecordingStart()
 
 	TraceArchiveHeader();
 
-	UE_TRACE_LOG(ChaosVDLogger, ChaosVDUsesAutoRTFM, ChaosVDChannel)
+	UE_TRACE_MINIMAL_LOG(ChaosVDLogger, ChaosVDUsesAutoRTFM, ChaosVDChannel)
 			<< ChaosVDUsesAutoRTFM.bUsingAutoRTFM(static_cast<bool>(UE_AUTORTFM));
 
 	SessionState.bIsTracing = true;
diff --git a/Engine/Source/Runtime/Experimental/Chaos/Public/ChaosVisualDebugger/ChaosVDTraceMacros.h b/Engine/Source/Runtime/Experimental/Chaos/Public/ChaosVisualDebugger/ChaosVDTraceMacros.h
--- a/Engine/Source/Runtime/Experimental/Chaos/Public/ChaosVisualDebugger/ChaosVDTraceMacros.h
+++ b/Engine/Source/Runtime/Experimental/Chaos/Public/ChaosVisualDebugger/ChaosVDTraceMacros.h
@@ -6,17 +6,17 @@
 
 	#ifndef CVD_DEFINE_TRACE_VECTOR
 		#define CVD_DEFINE_TRACE_VECTOR(Type, Name) \
-		UE_TRACE_EVENT_FIELD(Type, Name##X) \
-		UE_TRACE_EVENT_FIELD(Type, Name##Y) \
-		UE_TRACE_EVENT_FIELD(Type, Name##Z)
+		UE_TRACE_MINIMAL_EVENT_FIELD(Type, Name##X) \
+		UE_TRACE_MINIMAL_EVENT_FIELD(Type, Name##Y) \
+		UE_TRACE_MINIMAL_EVENT_FIELD(Type, Name##Z)
 	#endif
 
 	#ifndef CVD_DEFINE_TRACE_ROTATOR
 		#define CVD_DEFINE_TRACE_ROTATOR(Type, Name) \
-		UE_TRACE_EVENT_FIELD(Type, Name##X) \
-		UE_TRACE_EVENT_FIELD(Type, Name##Y) \
-		UE_TRACE_EVENT_FIELD(Type, Name##Z) \
-		UE_TRACE_EVENT_FIELD(Type, Name##W)
+		UE_TRACE_MINIMAL_EVENT_FIELD(Type, Name##X) \
+		UE_TRACE_MINIMAL_EVENT_FIELD(Type, Name##Y) \
+		UE_TRACE_MINIMAL_EVENT_FIELD(Type, Name##Z) \
+		UE_TRACE_MINIMAL_EVENT_FIELD(Type, Name##W)
 	#endif
 
 	#ifndef CVD_TRACE_VECTOR_ON_EVENT
diff --git a/Engine/Source/Runtime/Experimental/Chaos/Public/ChaosVisualDebugger/ChaosVisualDebuggerTrace.h b/Engine/Source/Runtime/Experimental/Chaos/Public/ChaosVisualDebugger/ChaosVisualDebuggerTrace.h
--- a/Engine/Source/Runtime/Experimental/Chaos/Public/ChaosVisualDebugger/ChaosVisualDebuggerTrace.h
+++ b/Engine/Source/Runtime/Experimental/Chaos/Public/ChaosVisualDebugger/ChaosVisualDebuggerTrace.h
@@ -16,6 +16,10 @@
 	#define CHAOS_VISUAL_DEBUGGER_ENABLED WITH_CHAOS_VISUAL_DEBUGGER
 #endif
 
+#ifndef CHAOS_VISUAL_DEBUGGER_WITHOUT_TRACE
+	#define CHAOS_VISUAL_DEBUGGER_WITHOUT_TRACE 0
+#endif
+
 #include "Chaos/Framework/PhysicsSolverBase.h"
 #include "ChaosVisualDebugger/ChaosVDTraceMacros.h"
 
@@ -38,108 +42,108 @@ class FPBDConstraintContainer;
 class FPBDJointConstraints;
 }
 
-UE_TRACE_CHANNEL_EXTERN(ChaosVDChannel, CHAOS_API)
-
-UE_TRACE_EVENT_BEGIN_EXTERN(ChaosVDLogger, ChaosVDSolverFrameStart)
-	UE_TRACE_EVENT_FIELD(int32, SolverID)
-	UE_TRACE_EVENT_FIELD(uint64, Cycle)
-	UE_TRACE_EVENT_FIELD(UE::Trace::WideString, DebugName)
-	UE_TRACE_EVENT_FIELD(bool, IsKeyFrame)
-	UE_TRACE_EVENT_FIELD(bool, IsReSimulated)
-	UE_TRACE_EVENT_FIELD(int32, CurrentFrameNumber)
-UE_TRACE_EVENT_END()
-
-UE_TRACE_EVENT_BEGIN_EXTERN(ChaosVDLogger, ChaosVDSolverFrameEnd)
-	UE_TRACE_EVENT_FIELD(int32, SolverID)
-	UE_TRACE_EVENT_FIELD(uint64, Cycle)
-UE_TRACE_EVENT_END()
-
-UE_TRACE_EVENT_BEGIN_EXTERN(ChaosVDLogger, ChaosVDParticleCreated)
-	UE_TRACE_EVENT_FIELD(int32, SolverID)
-	UE_TRACE_EVENT_FIELD(uint64, Cycle)
-	UE_TRACE_EVENT_FIELD(int32, ParticleID)
-UE_TRACE_EVENT_END()
-
-UE_TRACE_EVENT_BEGIN_EXTERN(ChaosVDLogger, ChaosVDParticleDestroyed)
-	UE_TRACE_EVENT_FIELD(int32, SolverID)
-	UE_TRACE_EVENT_FIELD(uint64, Cycle)
-	UE_TRACE_EVENT_FIELD(int32, ParticleID)
-UE_TRACE_EVENT_END()
-
-UE_TRACE_EVENT_BEGIN_EXTERN(ChaosVDLogger, ChaosVDSolverStepStart)
-	UE_TRACE_EVENT_FIELD(int32, SolverID)
-	UE_TRACE_EVENT_FIELD(uint64, Cycle)
-	UE_TRACE_EVENT_FIELD(UE::Trace::WideString, StepName)
-UE_TRACE_EVENT_END()
-
-UE_TRACE_EVENT_BEGIN_EXTERN(ChaosVDLogger, ChaosVDSolverStepEnd)
-	UE_TRACE_EVENT_FIELD(int32, SolverID)
-	UE_TRACE_EVENT_FIELD(uint64, Cycle)
-	UE_TRACE_EVENT_FIELD(uint32, StepNumber)
-UE_TRACE_EVENT_END()
-
-UE_TRACE_EVENT_BEGIN_EXTERN(ChaosVDLogger, ChaosVDBinaryDataStart)
-	UE_TRACE_EVENT_FIELD(uint64, Cycle)
-	UE_TRACE_EVENT_FIELD(UE::Trace::WideString, TypeName)
-	UE_TRACE_EVENT_FIELD(int32, DataID)
-	UE_TRACE_EVENT_FIELD(uint32, DataSize)
-	UE_TRACE_EVENT_FIELD(uint32, OriginalSize)
-	UE_TRACE_EVENT_FIELD(bool, IsCompressed)
-UE_TRACE_EVENT_END()
-
-UE_TRACE_EVENT_BEGIN_EXTERN(ChaosVDLogger, ChaosVDBinaryDataContent)
-	UE_TRACE_EVENT_FIELD(uint64, Cycle)
-	UE_TRACE_EVENT_FIELD(int32, DataID)
-	UE_TRACE_EVENT_FIELD(uint8[], RawData)
-UE_TRACE_EVENT_END()
-
-UE_TRACE_EVENT_BEGIN_EXTERN(ChaosVDLogger, ChaosVDBinaryDataEnd)
-	UE_TRACE_EVENT_FIELD(uint64, Cycle)
-	UE_TRACE_EVENT_FIELD(int32, DataID)
-UE_TRACE_EVENT_END()
-
-UE_TRACE_EVENT_BEGIN_EXTERN(ChaosVDLogger, ChaosVDSolverSimulationSpace)
-	UE_TRACE_EVENT_FIELD(int32, SolverID)
-	UE_TRACE_EVENT_FIELD(uint64, Cycle)
+UE_TRACE_MINIMAL_CHANNEL_EXTERN(ChaosVDChannel, CHAOS_API)
+
+UE_TRACE_MINIMAL_EVENT_BEGIN_EXTERN(ChaosVDLogger, ChaosVDSolverFrameStart)
+	UE_TRACE_MINIMAL_EVENT_FIELD(int32, SolverID)
+	UE_TRACE_MINIMAL_EVENT_FIELD(uint64, Cycle)
+	UE_TRACE_MINIMAL_EVENT_FIELD(UE::Trace::WideString, DebugName)
+	UE_TRACE_MINIMAL_EVENT_FIELD(bool, IsKeyFrame)
+	UE_TRACE_MINIMAL_EVENT_FIELD(bool, IsReSimulated)
+	UE_TRACE_MINIMAL_EVENT_FIELD(int32, CurrentFrameNumber)
+UE_TRACE_MINIMAL_EVENT_END()
+
+UE_TRACE_MINIMAL_EVENT_BEGIN_EXTERN(ChaosVDLogger, ChaosVDSolverFrameEnd)
+	UE_TRACE_MINIMAL_EVENT_FIELD(int32, SolverID)
+	UE_TRACE_MINIMAL_EVENT_FIELD(uint64, Cycle)
+UE_TRACE_MINIMAL_EVENT_END()
+
+UE_TRACE_MINIMAL_EVENT_BEGIN_EXTERN(ChaosVDLogger, ChaosVDParticleCreated)
+	UE_TRACE_MINIMAL_EVENT_FIELD(int32, SolverID)
+	UE_TRACE_MINIMAL_EVENT_FIELD(uint64, Cycle)
+	UE_TRACE_MINIMAL_EVENT_FIELD(int32, ParticleID)
+UE_TRACE_MINIMAL_EVENT_END()
+
+UE_TRACE_MINIMAL_EVENT_BEGIN_EXTERN(ChaosVDLogger, ChaosVDParticleDestroyed)
+	UE_TRACE_MINIMAL_EVENT_FIELD(int32, SolverID)
+	UE_TRACE_MINIMAL_EVENT_FIELD(uint64, Cycle)
+	UE_TRACE_MINIMAL_EVENT_FIELD(int32, ParticleID)
+UE_TRACE_MINIMAL_EVENT_END()
+
+UE_TRACE_MINIMAL_EVENT_BEGIN_EXTERN(ChaosVDLogger, ChaosVDSolverStepStart)
+	UE_TRACE_MINIMAL_EVENT_FIELD(int32, SolverID)
+	UE_TRACE_MINIMAL_EVENT_FIELD(uint64, Cycle)
+	UE_TRACE_MINIMAL_EVENT_FIELD(UE::Trace::WideString, StepName)
+UE_TRACE_MINIMAL_EVENT_END()
+
+UE_TRACE_MINIMAL_EVENT_BEGIN_EXTERN(ChaosVDLogger, ChaosVDSolverStepEnd)
+	UE_TRACE_MINIMAL_EVENT_FIELD(int32, SolverID)
+	UE_TRACE_MINIMAL_EVENT_FIELD(uint64, Cycle)
+	UE_TRACE_MINIMAL_EVENT_FIELD(uint32, StepNumber)
+UE_TRACE_MINIMAL_EVENT_END()
+
+UE_TRACE_MINIMAL_EVENT_BEGIN_EXTERN(ChaosVDLogger, ChaosVDBinaryDataStart)
+	UE_TRACE_MINIMAL_EVENT_FIELD(uint64, Cycle)
+	UE_TRACE_MINIMAL_EVENT_FIELD(UE::Trace::WideString, TypeName)
+	UE_TRACE_MINIMAL_EVENT_FIELD(int32, DataID)
+	UE_TRACE_MINIMAL_EVENT_FIELD(uint32, DataSize)
+	UE_TRACE_MINIMAL_EVENT_FIELD(uint32, OriginalSize)
+	UE_TRACE_MINIMAL_EVENT_FIELD(bool, IsCompressed)
+UE_TRACE_MINIMAL_EVENT_END()
+
+UE_TRACE_MINIMAL_EVENT_BEGIN_EXTERN(ChaosVDLogger, ChaosVDBinaryDataContent)
+	UE_TRACE_MINIMAL_EVENT_FIELD(uint64, Cycle)
+	UE_TRACE_MINIMAL_EVENT_FIELD(int32, DataID)
+	UE_TRACE_MINIMAL_EVENT_FIELD(uint8[], RawData)
+UE_TRACE_MINIMAL_EVENT_END()
+
+UE_TRACE_MINIMAL_EVENT_BEGIN_EXTERN(ChaosVDLogger, ChaosVDBinaryDataEnd)
+	UE_TRACE_MINIMAL_EVENT_FIELD(uint64, Cycle)
+	UE_TRACE_MINIMAL_EVENT_FIELD(int32, DataID)
+UE_TRACE_MINIMAL_EVENT_END()
+
+UE_TRACE_MINIMAL_EVENT_BEGIN_EXTERN(ChaosVDLogger, ChaosVDSolverSimulationSpace)
+	UE_TRACE_MINIMAL_EVENT_FIELD(int32, SolverID)
+	UE_TRACE_MINIMAL_EVENT_FIELD(uint64, Cycle)
 	CVD_DEFINE_TRACE_VECTOR(Chaos::FReal, Position)
 	CVD_DEFINE_TRACE_ROTATOR(Chaos::FReal, Rotation)
-UE_TRACE_EVENT_END()
+UE_TRACE_MINIMAL_EVENT_END()
 
-UE_TRACE_EVENT_BEGIN_EXTERN(ChaosVDLogger, ChaosVDNonSolverLocation)
-	UE_TRACE_EVENT_FIELD(uint64, Cycle)
+UE_TRACE_MINIMAL_EVENT_BEGIN_EXTERN(ChaosVDLogger, ChaosVDNonSolverLocation)
+	UE_TRACE_MINIMAL_EVENT_FIELD(uint64, Cycle)
 	CVD_DEFINE_TRACE_VECTOR(Chaos::FReal, Position)
-	UE_TRACE_EVENT_FIELD(UE::Trace::WideString, DebugName)
-UE_TRACE_EVENT_END()
+	UE_TRACE_MINIMAL_EVENT_FIELD(UE::Trace::WideString, DebugName)
+UE_TRACE_MINIMAL_EVENT_END()
 
-UE_TRACE_EVENT_BEGIN_EXTERN(ChaosVDLogger, ChaosVDNonSolverTransform)
-	UE_TRACE_EVENT_FIELD(uint64, Cycle)
+UE_TRACE_MINIMAL_EVENT_BEGIN_EXTERN(ChaosVDLogger, ChaosVDNonSolverTransform)
+	UE_TRACE_MINIMAL_EVENT_FIELD(uint64, Cycle)
 	CVD_DEFINE_TRACE_VECTOR(Chaos::FReal, Position)
 	CVD_DEFINE_TRACE_VECTOR(Chaos::FReal, Scale)
 	CVD_DEFINE_TRACE_ROTATOR(Chaos::FReal, Rotation)
-	UE_TRACE_EVENT_FIELD(UE::Trace::WideString, DebugName)
-UE_TRACE_EVENT_END()
+	UE_TRACE_MINIMAL_EVENT_FIELD(UE::Trace::WideString, DebugName)
+UE_TRACE_MINIMAL_EVENT_END()
 
-UE_TRACE_EVENT_BEGIN_EXTERN(ChaosVDLogger, ChaosVDNetworkTickOffset)
-	UE_TRACE_EVENT_FIELD(int32, Offset)
-	UE_TRACE_EVENT_FIELD(int32, SolverID)
-UE_TRACE_EVENT_END()
+UE_TRACE_MINIMAL_EVENT_BEGIN_EXTERN(ChaosVDLogger, ChaosVDNetworkTickOffset)
+	UE_TRACE_MINIMAL_EVENT_FIELD(int32, Offset)
+	UE_TRACE_MINIMAL_EVENT_FIELD(int32, SolverID)
+UE_TRACE_MINIMAL_EVENT_END()
 
-UE_TRACE_EVENT_BEGIN_EXTERN(ChaosVDLogger, ChaosVDRolledBackDataID)
-	UE_TRACE_EVENT_FIELD(int32, DataID)
-UE_TRACE_EVENT_END()
+UE_TRACE_MINIMAL_EVENT_BEGIN_EXTERN(ChaosVDLogger, ChaosVDRolledBackDataID)
+	UE_TRACE_MINIMAL_EVENT_FIELD(int32, DataID)
+UE_TRACE_MINIMAL_EVENT_END()
 
-UE_TRACE_EVENT_BEGIN_EXTERN(ChaosVDLogger, ChaosVDUsesAutoRTFM)
-	UE_TRACE_EVENT_FIELD(bool, bUsingAutoRTFM)
-UE_TRACE_EVENT_END()
+UE_TRACE_MINIMAL_EVENT_BEGIN_EXTERN(ChaosVDLogger, ChaosVDUsesAutoRTFM)
+	UE_TRACE_MINIMAL_EVENT_FIELD(bool, bUsingAutoRTFM)
+UE_TRACE_MINIMAL_EVENT_END()
 
-UE_TRACE_EVENT_BEGIN_EXTERN(ChaosVDLogger, ChaosVDTraceRelevancyVolume)
+UE_TRACE_MINIMAL_EVENT_BEGIN_EXTERN(ChaosVDLogger, ChaosVDTraceRelevancyVolume)
 	CVD_DEFINE_TRACE_VECTOR(Chaos::FReal, BoxMin)
 	CVD_DEFINE_TRACE_VECTOR(Chaos::FReal, BoxMax)
-UE_TRACE_EVENT_END()
+UE_TRACE_MINIMAL_EVENT_END()
 
-UE_TRACE_EVENT_BEGIN_EXTERN(ChaosVDLogger, ChaosVDDummyEvent)
-	UE_TRACE_EVENT_FIELD(int32, SolverID)
-UE_TRACE_EVENT_END()
+UE_TRACE_MINIMAL_EVENT_BEGIN_EXTERN(ChaosVDLogger, ChaosVDDummyEvent)
+	UE_TRACE_MINIMAL_EVENT_FIELD(int32, SolverID)
+UE_TRACE_MINIMAL_EVENT_END()
 
 struct FChaosVDContext;
 struct FChaosVDQueryVisitStep;
diff --git a/Engine/Source/Runtime/Experimental/ChaosVisualDebugger/Private/ChaosVDRuntimeModule.cpp b/Engine/Source/Runtime/Experimental/ChaosVisualDebugger/Private/ChaosVDRuntimeModule.cpp
--- a/Engine/Source/Runtime/Experimental/ChaosVisualDebugger/Private/ChaosVDRuntimeModule.cpp
+++ b/Engine/Source/Runtime/Experimental/ChaosVisualDebugger/Private/ChaosVDRuntimeModule.cpp
@@ -100,7 +100,7 @@ void FChaosVDRuntimeModule::StartupModule()
 	else
 	{
 		
-#if UE_TRACE_ENABLED
+#if UE_TRACE_ENABLED || UE_TRACE_MINIMAL_ENABLED
 		UE::Trace::ToggleChannel(TEXT("ChaosVDChannel"), false);
 #endif
 
@@ -219,7 +219,7 @@ void FChaosVDRuntimeModule::StartRecording(TConstArrayView<FString> Args)
 	// Start with a generic Failure reason
 	FText FailureReason = LOCTEXT("SeeLogsForErrorDetailsText","Please see the logs for more details...");
 
-#if UE_TRACE_ENABLED
+#if UE_TRACE_ENABLED || UE_TRACE_MINIMAL_ENABLED
 
 	// Other tools could bee using trace
 	// This is aggressive but until Trace supports multi-sessions, just take over.
@@ -312,7 +312,7 @@ void FChaosVDRuntimeModule::StopRecording()
 
 	FTraceAuxiliary::OnTraceStopped.RemoveAll(this);
 
-#if UE_TRACE_ENABLED
+#if UE_TRACE_ENABLED || UE_TRACE_MINIMAL_ENABLED
 	RestoreTraceChannelsToPreRecordingState();
 
 	StopTrace();
@@ -398,7 +398,7 @@ void FChaosVDRuntimeModule::SaveAndDisabledCurrentEnabledTraceChannels()
 	// Until we support allowing other channels, indicate in the logs that we are disabling everything else
 	UE_LOG(LogChaosVDRuntime, Log, TEXT("[%s] Disabling additional trace channels..."), ANSI_TO_TCHAR(__FUNCTION__));
 
-#if UE_TRACE_ENABLED
+#if UE_TRACE_ENABLED || UE_TRACE_MINIMAL_ENABLED
 	OriginalTraceChannelsState.Reset();
 
 	// Disable any enabled additional channel
@@ -418,7 +418,7 @@ void FChaosVDRuntimeModule::SaveAndDisabledCurrentEnabledTraceChannels()
 
 void FChaosVDRuntimeModule::RestoreTraceChannelsToPreRecordingState()
 {
-#if UE_TRACE_ENABLED
+#if UE_TRACE_ENABLED || UE_TRACE_MINIMAL_ENABLED
 	UE_LOG(LogChaosVDRuntime, Log, TEXT("[%s] Restoring trace channels state..."), ANSI_TO_TCHAR(__FUNCTION__));
 
 	for (const TPair<FString, bool>& ChannelWithState : OriginalTraceChannelsState)
@@ -432,9 +432,12 @@ void FChaosVDRuntimeModule::RestoreTraceChannelsToPreRecordingState()
 
 void FChaosVDRuntimeModule::EnableRequiredTraceChannels()
 {
-#if UE_TRACE_ENABLED
+#if UE_TRACE_ENABLED || UE_TRACE_MINIMAL_ENABLED
 	UE::Trace::ToggleChannel(TEXT("ChaosVDChannel"), true); 
 	UE::Trace::ToggleChannel(TEXT("Frame"), true);
+#endif
+
+#if UE_TRACE_ENABLED
 	UE::Trace::ToggleChannel(TEXT("Log"), true);
 #endif
 }
diff --git a/Engine/Source/Runtime/TraceLog/Private/Trace/Trace.cpp b/Engine/Source/Runtime/TraceLog/Private/Trace/Trace.cpp
--- a/Engine/Source/Runtime/TraceLog/Private/Trace/Trace.cpp
+++ b/Engine/Source/Runtime/TraceLog/Private/Trace/Trace.cpp
@@ -43,7 +43,7 @@ UE_TRACE_MINIMAL_EVENT_END()
 
 void LogStall(uint64 Start, uint64 End)
 {
-	UE_TRACE_LOG($Trace, TraceStall, true)
+	UE_TRACE_MINIMAL_LOG($Trace, TraceStall, true)
 		<< TraceStall.StartCycle(Start)
 		<< TraceStall.EndCycle(End);
 }
