Defer transaction until activities effectively get scheduled.

#rnx

diff --git a/Engine/Source/Runtime/Experimental/IoStore/HttpClient/Private/Activity.inl b/Engine/Source/Runtime/Experimental/IoStore/HttpClient/Private/Activity.inl
--- a/Engine/Source/Runtime/Experimental/IoStore/HttpClient/Private/Activity.inl
+++ b/Engine/Source/Runtime/Experimental/IoStore/HttpClient/Private/Activity.inl
@@ -754,9 +754,7 @@ struct FActivity
 
 						FActivity(const FParams& Params);
 						~FActivity() = default;
-	void				Begin();
 	void				AddHeader(FAnsiStringView Key, FAnsiStringView Value);
-	void				End();
 	FOutcome			Tick(FHttpPeer& Peer, int32* MaxRecvSize=nullptr);
 	void				SetSink(FTicketSink&& InSink, UPTRINT Param);
 	void				SetDestination(FIoBuffer* InDest);
@@ -803,6 +801,7 @@ struct FActivity
 	uint8				MethodOffset;
 	uint8				PathOffset;
 	uint16				PathLength;
+	uint16				HeaderCount = 0;
 	uint32				StateParam = 0;
 #if IAS_HTTP_WITH_PERF
 	FStopwatch			Stopwatch;
@@ -815,6 +814,13 @@ struct FActivity
 	UPTRINT				SinkParam;
 	FTicketSink			Sink;
 	FBodyCloak			Transaction;
+
+	struct FHeaderRecord
+	{
+		int16			Key;
+		uint16			ValueLength;
+		char			Data[];
+	};
 };
 
 ////////////////////////////////////////////////////////////////////////////////
@@ -850,25 +856,23 @@ FActivity::FActivity(const FParams& Params)
 	LengthScore = uint8(Pow2);
 }
 
-////////////////////////////////////////////////////////////////////////////////
-void FActivity::Begin()
-{
-	FAnsiStringView Method(Transaction.GetData() + MethodOffset, PathOffset - MethodOffset);
-	FAnsiStringView Path(Transaction.GetData() + PathOffset, PathLength);
-	Transaction.Begin(Method, Path);
-}
-
 ////////////////////////////////////////////////////////////////////////////////
 void FActivity::AddHeader(FAnsiStringView Key, FAnsiStringView Value)
 {
-	Transaction.AddHeader(Key, Value);
-}
+	check(HeaderCount < 0xffff);
+	HeaderCount++;
 
-////////////////////////////////////////////////////////////////////////////////
-void FActivity::End()
-{
-	Transaction.End();
-	ChangeState(EState::Send);
+	static_assert(alignof(FHeaderRecord) == alignof(uint16));
+	uint32 Count = sizeof(FHeaderRecord) + Key.Len() + Value.Len();
+	Count = (Count + sizeof(uint16) - 1) / sizeof(uint16);
+	auto* Record = (FHeaderRecord*)(Transaction.Alloc<uint16>(Count));
+
+	Record->Key = int16(Key.Len());
+	Record->ValueLength = int16(Value.Len());
+
+	char* Cursor = Record->Data;
+	std::memcpy(Cursor, Key.GetData(), Key.Len());
+	std::memcpy(Cursor + Key.Len(), Value.GetData(), Value.Len());
 }
 
 ////////////////////////////////////////////////////////////////////////////////
@@ -1209,11 +1213,44 @@ FOutcome FActivity::Recv(FHttpPeer& Peer, int32& MaxRecvSize)
 ////////////////////////////////////////////////////////////////////////////////
 FOutcome FActivity::Tick(FHttpPeer& Peer, int32* MaxRecvSize)
 {
+	if (State == EState::Build)
+	{
+		FAnsiStringView Method(Transaction.GetData() + MethodOffset, PathOffset - MethodOffset);
+		FAnsiStringView Path(Transaction.GetData() + PathOffset, PathLength);
+		Transaction.Begin(Method, Path);
+
+		Transaction.AddHeader("Host", GetHost()->GetHostName());
+
+		// HTTP/1.1 is persistent by default thus "Connection" header isn't required
+		if (!GetHost()->IsPooled())
+		{
+			Transaction.AddHeader("Connection", "close");
+		}
+
+		UPTRINT RecordAddr = uint32(PathOffset) + uint32(PathLength);
+		for (uint32 i = 0, n = HeaderCount; i < n; ++i)
+		{
+			RecordAddr = (RecordAddr + alignof(FHeaderRecord) - 1) & ~(alignof(FHeaderRecord) - 1);
+			const auto* Record = (FHeaderRecord*)(Transaction.GetData() + RecordAddr);
+
+			FAnsiStringView Key(Record->Data, Record->Key);
+			FAnsiStringView Value(Record->Data + Record->Key, Record->ValueLength);
+			Transaction.AddHeader(Key, Value);
+
+			RecordAddr += sizeof(FHeaderRecord) + Record->Key + Record->ValueLength;
+		}
+
+		Transaction.End();
+
+		ChangeState(EState::Send);
+	}
+
 	if (State == EState::Send)
 	{
 		return Send(Peer);
 	}
 
+	check(State > EState::Send && State < EState::RecvDone);
 	check(MaxRecvSize != nullptr);
 	return Recv(Peer, *MaxRecvSize);
 }
diff --git a/Engine/Source/Runtime/Experimental/IoStore/HttpClient/Private/Loop.inl b/Engine/Source/Runtime/Experimental/IoStore/HttpClient/Private/Loop.inl
--- a/Engine/Source/Runtime/Experimental/IoStore/HttpClient/Private/Loop.inl
+++ b/Engine/Source/Runtime/Experimental/IoStore/HttpClient/Private/Loop.inl
@@ -776,17 +776,6 @@ FRequest FEventLoop::FImpl::Request(FActivityNode* Activity)
 {
 	Trace(Activity, ETrace::ActivityCreate, 0);
 
-	const FHost* Host = Activity->GetHost();
-
-	Activity->Begin();
-	Activity->AddHeader("Host", Host->GetHostName());
-
-	// HTTP/1.1 is persistent by default thus "Connection" header isn't required
-	if (!Host->IsPooled())
-	{
-		Activity->AddHeader("Connection", "close");
-	}
-
 	FRequest Ret;
 	Ret.Ptr = Activity;
 	return Ret;
@@ -797,8 +786,6 @@ FTicket FEventLoop::FImpl::Send(FActivityNode* Activity)
 {
 	Trace(Activity, ETrace::RequestBegin);
 
-	Activity->End();
-
 	uint64 Slot;
 	{
 		FScopeLock _(&Lock);
