Added accessors.

#rnx

diff --git a/Engine/Source/Runtime/Experimental/IoStore/HttpClient/Private/Activity.inl b/Engine/Source/Runtime/Experimental/IoStore/HttpClient/Private/Activity.inl
--- a/Engine/Source/Runtime/Experimental/IoStore/HttpClient/Private/Activity.inl
+++ b/Engine/Source/Runtime/Experimental/IoStore/HttpClient/Private/Activity.inl
@@ -761,6 +761,8 @@ struct FActivity
 	void				Cancel();
 	void				Done();
 	void				Fail(const FOutcome& Outcome);
+	FAnsiStringView		GetMethod() const;
+	FAnsiStringView		GetPath() const;
 	FHost*				GetHost() const;
 	UPTRINT				GetSinkParam() const;
 	FIoBuffer&			GetContent() const;
@@ -946,6 +948,18 @@ void FActivity::CallSink()
 	Sink(SinkArg);
 }
 
+////////////////////////////////////////////////////////////////////////////////
+FAnsiStringView FActivity::GetMethod() const
+{
+	return { Transaction.GetData() + MethodOffset, PathOffset - MethodOffset };
+}
+
+////////////////////////////////////////////////////////////////////////////////
+FAnsiStringView FActivity::GetPath() const
+{
+	return { Transaction.GetData() + PathOffset, PathLength };
+}
+
 ////////////////////////////////////////////////////////////////////////////////
 FHost* FActivity::GetHost() const
 {
@@ -1215,9 +1229,7 @@ FOutcome FActivity::Tick(FHttpPeer& Peer, int32* MaxRecvSize)
 {
 	if (State == EState::Build)
 	{
-		FAnsiStringView Method(Transaction.GetData() + MethodOffset, PathOffset - MethodOffset);
-		FAnsiStringView Path(Transaction.GetData() + PathOffset, PathLength);
-		Transaction.Begin(Method, Path);
+		Transaction.Begin(GetMethod(), GetPath());
 
 		Transaction.AddHeader("Host", GetHost()->GetHostName());
 
diff --git a/Engine/Source/Runtime/Experimental/IoStore/HttpClient/Private/Loop.inl b/Engine/Source/Runtime/Experimental/IoStore/HttpClient/Private/Loop.inl
--- a/Engine/Source/Runtime/Experimental/IoStore/HttpClient/Private/Loop.inl
+++ b/Engine/Source/Runtime/Experimental/IoStore/HttpClient/Private/Loop.inl
@@ -1082,7 +1082,7 @@ bool FEventLoop::Redirect(const FTicketStatus& Status, FTicketSink& OuterSink)
 	check(!Activity.Transaction.HasBody() || Response.GetContentLength() == 0); // should we ever hit this, we'll fix it
 
 	// Original method should remain unchanged
-	FAnsiStringView Method = Activity.Transaction.GetMethod();
+	FAnsiStringView Method = Activity.GetMethod();
 	check(!Method.IsEmpty());
 
 	FRequestParams RequestParams = {
