Lift enumeration accessor up.

#rnx

diff --git a/Engine/Source/Runtime/Experimental/IoStore/HttpClient/Private/Activity.inl b/Engine/Source/Runtime/Experimental/IoStore/HttpClient/Private/Activity.inl
--- a/Engine/Source/Runtime/Experimental/IoStore/HttpClient/Private/Activity.inl
+++ b/Engine/Source/Runtime/Experimental/IoStore/HttpClient/Private/Activity.inl
@@ -764,6 +764,7 @@ struct FActivity
 	FAnsiStringView		GetMethod() const;
 	FAnsiStringView		GetPath() const;
 	FHost*				GetHost() const;
+	void				EnumerateHeaders(FResponse::FHeaderSink HeaderSink) const;
 	UPTRINT				GetSinkParam() const;
 	FIoBuffer&			GetContent() const;
 	uint32				GetRemainingKiB() const;
@@ -967,6 +968,23 @@ FHost* FActivity::GetHost() const
 	return Host;
 }
 
+////////////////////////////////////////////////////////////////////////////////
+void FActivity::EnumerateHeaders(FResponse::FHeaderSink HeaderSink) const
+{
+	UPTRINT RecordAddr = uint32(PathOffset) + uint32(PathLength);
+	for (uint32 i = 0, n = HeaderCount; i < n; ++i)
+	{
+		RecordAddr = (RecordAddr + alignof(FHeaderRecord) - 1) & ~(alignof(FHeaderRecord) - 1);
+		const auto* Record = (FHeaderRecord*)(Transaction.GetData() + RecordAddr);
+
+		FAnsiStringView Key(Record->Data, Record->Key);
+		FAnsiStringView Value(Record->Data + Record->Key, Record->ValueLength);
+		HeaderSink(Key, Value);
+
+		RecordAddr += sizeof(FHeaderRecord) + Record->Key + Record->ValueLength;
+	}
+}
+
 ////////////////////////////////////////////////////////////////////////////////
 UPTRINT FActivity::GetSinkParam() const
 {
@@ -1239,18 +1257,10 @@ FOutcome FActivity::Tick(FHttpPeer& Peer, int32* MaxRecvSize)
 			Transaction.AddHeader("Connection", "close");
 		}
 
-		UPTRINT RecordAddr = uint32(PathOffset) + uint32(PathLength);
-		for (uint32 i = 0, n = HeaderCount; i < n; ++i)
-		{
-			RecordAddr = (RecordAddr + alignof(FHeaderRecord) - 1) & ~(alignof(FHeaderRecord) - 1);
-			const auto* Record = (FHeaderRecord*)(Transaction.GetData() + RecordAddr);
-
-			FAnsiStringView Key(Record->Data, Record->Key);
-			FAnsiStringView Value(Record->Data + Record->Key, Record->ValueLength);
+		EnumerateHeaders([this] (FAnsiStringView Key, FAnsiStringView Value) {
 			Transaction.AddHeader(Key, Value);
-
-			RecordAddr += sizeof(FHeaderRecord) + Record->Key + Record->ValueLength;
-		}
+			return true;
+		});
 
 		Transaction.End();
 
diff --git a/Engine/Source/Runtime/Experimental/IoStore/HttpClient/Private/Loop.inl b/Engine/Source/Runtime/Experimental/IoStore/HttpClient/Private/Loop.inl
--- a/Engine/Source/Runtime/Experimental/IoStore/HttpClient/Private/Loop.inl
+++ b/Engine/Source/Runtime/Experimental/IoStore/HttpClient/Private/Loop.inl
@@ -1115,12 +1115,9 @@ bool FEventLoop::Redirect(const FTicketStatus& Status, FTicketSink& OuterSink)
 	}
 
 	// Transfer original request headers
-	Activity.Transaction.EnumerateHeaders([&ForwardRequest] (FAnsiStringView Name, FAnsiStringView Value)
+	Activity.EnumerateHeaders([&ForwardRequest] (FAnsiStringView Name, FAnsiStringView Value)
 	{
-		if (Name != "Host" && Name != "Connection")
-		{
-			ForwardRequest.Header(Name, Value);
-		}
+		ForwardRequest.Header(Name, Value);
 		return true;
 	});
 
