Store errors as an outcome.

#rnx

diff --git a/Engine/Source/Runtime/Experimental/IoStore/HttpClient/Private/Activity.inl b/Engine/Source/Runtime/Experimental/IoStore/HttpClient/Private/Activity.inl
--- a/Engine/Source/Runtime/Experimental/IoStore/HttpClient/Private/Activity.inl
+++ b/Engine/Source/Runtime/Experimental/IoStore/HttpClient/Private/Activity.inl
@@ -768,8 +768,7 @@ struct FActivity
 	UPTRINT				GetSinkParam() const;
 	FIoBuffer&			GetContent() const;
 	uint32				GetRemainingKiB() const;
-	const char*			GetErrorReason() const;
-	uint32				GetErrorCode() const;
+	FOutcome			GetError() const;
 
 #if IAS_HTTP_WITH_PERF
 	const FStopwatch&	GetStopwatch() const;
@@ -812,7 +811,7 @@ struct FActivity
 	union {
 		FHost*			Host;
 		FIoBuffer*		Dest;
-		const char*		ErrorReason;
+		UPTRINT			Error;
 	};
 	UPTRINT				SinkParam;
 	FTicketSink			Sink;
@@ -919,11 +918,10 @@ void FActivity::Fail(const FOutcome& Outcome)
 		return;
 	}
 
-	ChangeState(EState::Failed);
-
-	ErrorReason = Outcome.GetMessage().GetData();
-	StateParam = uint32(Outcome.GetErrorCode());
+	static_assert(sizeof(Outcome) == sizeof(Error));
+	std::memcpy(&Error, &Outcome, sizeof(Error));
 
+	ChangeState(EState::Failed);
 	CallSink();
 }
 
@@ -1010,17 +1008,11 @@ uint32 FActivity::GetRemainingKiB() const
 }
 
 ////////////////////////////////////////////////////////////////////////////////
-const char* FActivity::GetErrorReason() const
+FOutcome FActivity::GetError() const
 {
-	check(State == EState::Failed);
-	return ErrorReason;
-}
-
-////////////////////////////////////////////////////////////////////////////////
-uint32 FActivity::GetErrorCode() const
-{
-	check(State == EState::Failed);
-	return uint32(StateParam);
+	FOutcome Outcome = FOutcome::None();
+	std::memcpy(&Outcome, &Error, sizeof(Outcome));
+	return Outcome;
 }
 
 ////////////////////////////////////////////////////////////////////////////////
diff --git a/Engine/Source/Runtime/Experimental/IoStore/HttpClient/Private/Api.inl b/Engine/Source/Runtime/Experimental/IoStore/HttpClient/Private/Api.inl
--- a/Engine/Source/Runtime/Experimental/IoStore/HttpClient/Private/Api.inl
+++ b/Engine/Source/Runtime/Experimental/IoStore/HttpClient/Private/Api.inl
@@ -230,7 +230,8 @@ FTicketStatus::FError FTicketStatus::GetError() const
 {
 	check(GetId() == EId::Error);
 	const auto* Activity = (FActivity*)this;
-	return { Activity->GetErrorReason(), Activity->GetErrorCode() };
+	FOutcome Error = Activity->GetError();
+	return { Error.GetMessage().GetData(), uint32(Error.GetErrorCode()) };
 }
 
 
