Only adjust max size after data in.

#rnx

diff --git a/Engine/Source/Runtime/Experimental/IoStore/HttpClient/Private/Activity.inl b/Engine/Source/Runtime/Experimental/IoStore/HttpClient/Private/Activity.inl
--- a/Engine/Source/Runtime/Experimental/IoStore/HttpClient/Private/Activity.inl
+++ b/Engine/Source/Runtime/Experimental/IoStore/HttpClient/Private/Activity.inl
@@ -1147,7 +1147,6 @@ FOutcome FActivity::RecvContent(FHttpPeer& Peer, int32& MaxRecvSize)
 
 	FMutableMemoryView View = Dest->GetMutableView();
 	View = View.Right(View.GetSize() - StateParam).Left(MaxRecvSize);
-	MaxRecvSize -= int32(View.GetSize());
 
 	FOutcome Outcome = Transaction.TryRecv(View, Peer);
 	if (Outcome.IsError())
@@ -1156,6 +1155,8 @@ FOutcome FActivity::RecvContent(FHttpPeer& Peer, int32& MaxRecvSize)
 		return Outcome;
 	}
 
+	MaxRecvSize -= Outcome.GetResult();
+
 	if (Outcome.IsWaiting())
 	{
 		StateParam += Outcome.GetResult();
