Reduce reliance on strings, removing assumption that they are still around.

#rnx

diff --git a/Engine/Source/Runtime/Experimental/IoStore/HttpClient/Private/Activity.inl b/Engine/Source/Runtime/Experimental/IoStore/HttpClient/Private/Activity.inl
--- a/Engine/Source/Runtime/Experimental/IoStore/HttpClient/Private/Activity.inl
+++ b/Engine/Source/Runtime/Experimental/IoStore/HttpClient/Private/Activity.inl
@@ -19,6 +19,9 @@ public:
 	void			EnumerateHeaders(FResponse::FHeaderSink Sink) const;
 	FOutcome		TrySendRequest(FHttpPeer& Peer);
 
+protected:
+	bool			IsHeadMethod() const;
+
 private:
 	FRequestCloak&	operator << (FAnsiStringView Value);
 	uint16			Bias = 0;
@@ -26,8 +29,15 @@ private:
 	uint16			AlreadySent = 0;
 	int16			HeaderRight = -1;
 	int8			MethodLength = -1;
+	bool			IsHead;
 };
 
+////////////////////////////////////////////////////////////////////////////////
+bool FRequestCloak::IsHeadMethod() const
+{
+	return IsHead;
+}
+
 ////////////////////////////////////////////////////////////////////////////////
 FAnsiStringView FRequestCloak::GetMethod() const
 {
@@ -52,6 +62,7 @@ void FRequestCloak::Begin(FAnsiStringView Method, FAnsiStringView Path)
 	*this << Method << " " << Path << " HTTP/1.1" "\r\n";
 	MethodLength = int8(Method.Len());
 	HeaderLeft = uint16(GetSize());
+	IsHead = Method.Equals("HEAD", ESearchCase::IgnoreCase);
 }
 
 ////////////////////////////////////////////////////////////////////////////////
@@ -317,7 +328,7 @@ FOutcome FResponseCloak::Parse()
 	);
 
 	bHasBody = (StatusCode != 304);
-	bHasBody &= !GetMethod().Equals("HEAD", ESearchCase::IgnoreCase);
+	bHasBody &= !IsHeadMethod();
 	bHasBody &= bChunked | (ContentLength > 0);
 
 	if (Flags & Flag_Length)
