Strategic reduction in transaction contract.

#rnx

diff --git a/Engine/Source/Runtime/Experimental/IoStore/HttpClient/Private/Activity.inl b/Engine/Source/Runtime/Experimental/IoStore/HttpClient/Private/Activity.inl
--- a/Engine/Source/Runtime/Experimental/IoStore/HttpClient/Private/Activity.inl
+++ b/Engine/Source/Runtime/Experimental/IoStore/HttpClient/Private/Activity.inl
@@ -15,8 +15,6 @@ public:
 	void			Begin(FAnsiStringView Method, FAnsiStringView Path);
 	void			AddHeader(FAnsiStringView Key, FAnsiStringView Value);
 	void			End();
-	FAnsiStringView	GetMethod() const;
-	void			EnumerateHeaders(FResponse::FHeaderSink Sink) const;
 	FOutcome		TrySendRequest(FHttpPeer& Peer);
 
 protected:
@@ -38,22 +36,6 @@ bool FRequestCloak::IsHeadMethod() const
 	return IsHead;
 }
 
-////////////////////////////////////////////////////////////////////////////////
-FAnsiStringView FRequestCloak::GetMethod() const
-{
-	check(MethodLength > 0);
-	return FAnsiStringView(GetData() + Bias, MethodLength);
-}
-
-////////////////////////////////////////////////////////////////////////////////
-void FRequestCloak::EnumerateHeaders(FResponse::FHeaderSink Sink) const
-{
-	check(HeaderRight >= 2);
-	int32 HeaderLength = HeaderRight - HeaderLeft;
-	FAnsiStringView Headers(GetData() + HeaderLeft, HeaderLength);
-	UE::IoStore::HTTP::EnumerateHeaders(Headers, Sink);
-}
-
 ////////////////////////////////////////////////////////////////////////////////
 void FRequestCloak::Begin(FAnsiStringView Method, FAnsiStringView Path)
 {
