Deleted unused member

#rnx

diff --git a/Engine/Source/Runtime/Experimental/IoStore/HttpClient/Private/Activity.inl b/Engine/Source/Runtime/Experimental/IoStore/HttpClient/Private/Activity.inl
--- a/Engine/Source/Runtime/Experimental/IoStore/HttpClient/Private/Activity.inl
+++ b/Engine/Source/Runtime/Experimental/IoStore/HttpClient/Private/Activity.inl
@@ -22,7 +22,6 @@ protected:
 
 private:
 	FRequestCloak&	operator << (FAnsiStringView Value);
-	uint16			Bias = 0;
 	uint16			HeaderLeft;
 	uint16			AlreadySent = 0;
 	int16			HeaderRight = -1;
@@ -40,7 +39,7 @@ bool FRequestCloak::IsHeadMethod() const
 void FRequestCloak::Begin(FAnsiStringView Method, FAnsiStringView Path)
 {
 	check(MethodLength < 0);
-	AlreadySent = Bias = uint16(GetSize());
+	AlreadySent = uint16(GetSize());
 	*this << Method << " " << Path << " HTTP/1.1" "\r\n";
 	MethodLength = int8(Method.Len());
 	HeaderLeft = uint16(GetSize());
