Accessor for bytes remaining

#rnx

diff --git a/Engine/Source/Runtime/Experimental/IoStore/HttpClient/Private/Activity.inl b/Engine/Source/Runtime/Experimental/IoStore/HttpClient/Private/Activity.inl
--- a/Engine/Source/Runtime/Experimental/IoStore/HttpClient/Private/Activity.inl
+++ b/Engine/Source/Runtime/Experimental/IoStore/HttpClient/Private/Activity.inl
@@ -358,7 +358,8 @@ class FBodyCloak
 public:
 	using FResponseCloak::FResponseCloak;
 
-	FOutcome		TryRecv(FMutableMemoryView Dest, FHttpPeer& Peer);
+	int64				GetRemaining() const;
+	FOutcome			TryRecv(FMutableMemoryView Dest, FHttpPeer& Peer);
 
 private:
 	enum class EState : uint8 { Init, Recv, Prologue, Chunk, Epilogue, Done };
@@ -382,6 +383,12 @@ private:
 	};
 };
 
+////////////////////////////////////////////////////////////////////////////////
+int64 FBodyCloak::GetRemaining() const
+{
+	return (State == EState::Recv) ? Remaining : -1;
+}
+
 ////////////////////////////////////////////////////////////////////////////////
 FOutcome FBodyCloak::TryRecv(FMutableMemoryView Dest, FHttpPeer& Peer)
 {
