There's no longer a need to track remaining in state parameter.

#rnx

diff --git a/Engine/Source/Runtime/Experimental/IoStore/HttpClient/Private/Activity.inl b/Engine/Source/Runtime/Experimental/IoStore/HttpClient/Private/Activity.inl
--- a/Engine/Source/Runtime/Experimental/IoStore/HttpClient/Private/Activity.inl
+++ b/Engine/Source/Runtime/Experimental/IoStore/HttpClient/Private/Activity.inl
@@ -1000,10 +1000,7 @@ uint32 FActivity::GetRemainingKiB() const
 {
 	if (State <= EState::RecvStream)  return MAX_uint32;
 	if (State >  EState::RecvContent) return 0;
-
-	uint32 ContentLength = uint32(Transaction.GetContentLength());
-	check(StateParam <= ContentLength);
-	return (ContentLength - StateParam) >> 10;
+	return uint32(Transaction.GetRemaining() >> 10);
 }
 
 ////////////////////////////////////////////////////////////////////////////////
@@ -1144,8 +1141,10 @@ FOutcome FActivity::RecvContent(FHttpPeer& Peer, int32& MaxRecvSize)
 {
 	TRACE_CPUPROFILER_EVENT_SCOPE(IasHttp::DoRecvContent);
 
+	int64 Remaining = Transaction.GetRemaining();
+
 	FMutableMemoryView View = Dest->GetMutableView();
-	View = View.Right(View.GetSize() - StateParam).Left(MaxRecvSize);
+	View = View.Right(Remaining).Left(MaxRecvSize);
 
 	FOutcome Outcome = Transaction.TryRecv(View, Peer);
 	if (Outcome.IsError())
@@ -1158,12 +1157,9 @@ FOutcome FActivity::RecvContent(FHttpPeer& Peer, int32& MaxRecvSize)
 
 	if (Outcome.IsWaiting())
 	{
-		StateParam += Outcome.GetResult();
 		return Outcome;
 	}
 
-	check(StateParam + Outcome.GetResult() == Transaction.GetContentLength());
-
 #if IAS_HTTP_WITH_PERF
 	Stopwatch.RecvEnd();
 #endif
