Removed state parameter.

#rnx

diff --git a/Engine/Source/Runtime/Experimental/IoStore/HttpClient/Private/Activity.inl b/Engine/Source/Runtime/Experimental/IoStore/HttpClient/Private/Activity.inl
--- a/Engine/Source/Runtime/Experimental/IoStore/HttpClient/Private/Activity.inl
+++ b/Engine/Source/Runtime/Experimental/IoStore/HttpClient/Private/Activity.inl
@@ -788,7 +788,7 @@ struct FActivity
 	};
 
 	void				CallSink();
-	void				ChangeState(EState InState, uint32 Param=0);
+	void				ChangeState(EState InState);
 	FOutcome			Send(FHttpPeer& Peer);
 	FOutcome			RecvMessage(FHttpPeer& Peer);
 	FOutcome			RecvContent(FHttpPeer& Peer, int32& MaxRecvSize);
@@ -803,7 +803,6 @@ struct FActivity
 	uint8				PathOffset;
 	uint16				PathLength;
 	uint16				HeaderCount = 0;
-	uint32				StateParam = 0;
 #if IAS_HTTP_WITH_PERF
 	FStopwatch			Stopwatch;
 #endif
@@ -925,14 +924,12 @@ void FActivity::Fail(const FOutcome& Outcome)
 }
 
 ////////////////////////////////////////////////////////////////////////////////
-void FActivity::ChangeState(EState InState, uint32 Param)
+void FActivity::ChangeState(EState InState)
 {
-	check(State != InState);
-
 	Trace(this, ETrace::StateChange, uint32(InState));
 
+	check(State != InState);
 	State = InState;
-	StateParam = Param;
 }
 
 ////////////////////////////////////////////////////////////////////////////////
