Extra check() to guard against user error.

#rnx

diff --git a/Engine/Source/Runtime/Experimental/IoStore/HttpClient/Private/Activity.inl b/Engine/Source/Runtime/Experimental/IoStore/HttpClient/Private/Activity.inl
--- a/Engine/Source/Runtime/Experimental/IoStore/HttpClient/Private/Activity.inl
+++ b/Engine/Source/Runtime/Experimental/IoStore/HttpClient/Private/Activity.inl
@@ -859,6 +859,8 @@ FActivity::FActivity(const FParams& Params)
 ////////////////////////////////////////////////////////////////////////////////
 void FActivity::AddHeader(FAnsiStringView Key, FAnsiStringView Value)
 {
+	check(State == EState::Build);
+
 	check(HeaderCount < 0xffff);
 	HeaderCount++;
 
