IAS upload tool now only warns if the IAD source data is missing.
This can happen if the bulk build is no longer required.

diff --git a/Engine/Source/Runtime/Experimental/IoStore/OnDemand/Private/Tool/Upload.cpp b/Engine/Source/Runtime/Experimental/IoStore/OnDemand/Private/Tool/Upload.cpp
--- a/Engine/Source/Runtime/Experimental/IoStore/OnDemand/Private/Tool/Upload.cpp
+++ b/Engine/Source/Runtime/Experimental/IoStore/OnDemand/Private/Tool/Upload.cpp
@@ -438,7 +438,8 @@ static TIoStatusOr<FUploadResult> UploadContainerFiles(
 
 	if (ContainerFiles.Num() == 0)
 	{
-		return FIoStatus(EIoErrorCode::InvalidParameter, TEXT("No container file(s) specified"));
+		UE_LOG(LogIas, Warning, TEXT("No container file(s) found"));
+		return FIoStatus(EIoErrorCode::Ok);
 	}
 
 	TSet<FIoHash> ExistingChunks;
