Change visibility of unused assignment operator

#rnx

diff --git a/Engine/Source/Runtime/Experimental/IoStore/HttpClient/Private/Misc.inl b/Engine/Source/Runtime/Experimental/IoStore/HttpClient/Private/Misc.inl
--- a/Engine/Source/Runtime/Experimental/IoStore/HttpClient/Private/Misc.inl
+++ b/Engine/Source/Runtime/Experimental/IoStore/HttpClient/Private/Misc.inl
@@ -443,7 +443,6 @@ public:
 								FBuffer() = default;
 								FBuffer(char* InData, uint32 InMax);
 								~FBuffer();
-	FBuffer&					operator = (FBuffer&& Rhs);
 	void						Resize(uint32 Size);
 	const char*					GetData() const;
 	uint32						GetSize() const;
@@ -471,6 +470,7 @@ private:
 								FBuffer(FBuffer&&) = delete;
 								FBuffer(const FBuffer&) = delete;
 	FBuffer&					operator = (const FBuffer&) = delete;
+	FBuffer&					operator = (FBuffer&& Rhs) = delete;
 };
 
 ////////////////////////////////////////////////////////////////////////////////
@@ -490,15 +490,6 @@ FBuffer::~FBuffer()
 	}
 }
 
-////////////////////////////////////////////////////////////////////////////////
-FBuffer& FBuffer::operator = (FBuffer&& Rhs)
-{
-	Swap(Data, Rhs.Data);
-	Swap(Max, Rhs.Max);
-	Swap(UsedInline, Rhs.UsedInline);
-	return *this;
-}
-
 ////////////////////////////////////////////////////////////////////////////////
 void FBuffer::Resize(uint32 Size)
 {
