Scoped enum

#rnx

diff --git a/Engine/Source/Runtime/Experimental/IoStore/HttpClient/Private/Activity.inl b/Engine/Source/Runtime/Experimental/IoStore/HttpClient/Private/Activity.inl
--- a/Engine/Source/Runtime/Experimental/IoStore/HttpClient/Private/Activity.inl
+++ b/Engine/Source/Runtime/Experimental/IoStore/HttpClient/Private/Activity.inl
@@ -744,11 +744,11 @@ struct FActivity
 		bool			bAllowChunked = true;
 	};
 
-	enum EStage : uint32
+	enum class EStage : int32
 	{
-		Request		= 0xffff'ffff,
-		Response	= 0xffff'fffe,
-		Content		= 0xffff'fffd,
+		Request			= -1,
+		Response		= -2,
+		Content			= -3,
 	};
 
 						FActivity(const FParams& Params);
@@ -1046,7 +1046,7 @@ FOutcome FActivity::Send(FHttpPeer& Peer)
 #endif
 
 	ChangeState(EState::RecvMessage);
-	return FOutcome::Ok(EStage::Request);
+	return FOutcome::Ok(int32(EStage::Request));
 }
 
 ////////////////////////////////////////////////////////////////////////////////
@@ -1125,14 +1125,14 @@ FOutcome FActivity::RecvMessage(FHttpPeer& Peer)
 	if (!Transaction.HasBody())
 	{
 		ChangeState(EState::RecvDone);
-		return FOutcome::Ok(EStage::Content);
+		return FOutcome::Ok(int32(EStage::Content));
 	}
 
 	// We're all set to go and get content
 	check(Dest != nullptr);
 	auto NextState = bChunked ? EState::RecvStream : EState::RecvContent;
 	ChangeState(NextState);
-	return FOutcome::Ok(EStage::Response);
+	return FOutcome::Ok(int32(EStage::Response));
 }
 
 ////////////////////////////////////////////////////////////////////////////////
@@ -1164,7 +1164,7 @@ FOutcome FActivity::RecvContent(FHttpPeer& Peer, int32& MaxRecvSize)
 #endif
 
 	Done();
-	return FOutcome::Ok(EStage::Content);
+	return FOutcome::Ok(int32(EStage::Content));
 }
 
 ////////////////////////////////////////////////////////////////////////////////
@@ -1213,7 +1213,7 @@ FOutcome FActivity::RecvStream(FHttpPeer& Peer, int32& MaxRecvSize)
 	*Dest = FIoBuffer();
 
 	Done();
-	return FOutcome::Ok(EStage::Content);
+	return FOutcome::Ok(int32(EStage::Content));
 }
 
 ////////////////////////////////////////////////////////////////////////////////
diff --git a/Engine/Source/Runtime/Experimental/IoStore/HttpClient/Private/Loop.inl b/Engine/Source/Runtime/Experimental/IoStore/HttpClient/Private/Loop.inl
--- a/Engine/Source/Runtime/Experimental/IoStore/HttpClient/Private/Loop.inl
+++ b/Engine/Source/Runtime/Experimental/IoStore/HttpClient/Private/Loop.inl
@@ -428,12 +428,12 @@ void FPeerGroup::RecvInternal(FTickState& State)
 
 	// If an okay result is -1, then message recv is done and content is next
 	// and it is likely that we have it.
-	if (Outcome.GetResult() == FActivity::EStage::Response)
+	if (Outcome.GetResult() == uint32(FActivity::EStage::Response))
 	{
 		return RecvInternal(State);
 	}
 
-	check(Outcome.GetResult() == FActivity::EStage::Content);
+	check(Outcome.GetResult() == uint32(FActivity::EStage::Content));
 
 	Recv = Activity->Next;
 	Activity->Next = State.DoneList;
@@ -469,7 +469,7 @@ void FPeerGroup::SendInternal(FTickState& State)
 		return;
 	}
 
-	check(Outcome.GetResult() == FActivity::EStage::Request);
+	check(Outcome.GetResult() == uint32(FActivity::EStage::Request));
 	Send = nullptr;
 
 	// Pass along this send to be received
