Built out a GetStage accessor

#rnx

diff --git a/Engine/Source/Runtime/Experimental/IoStore/HttpClient/Private/Activity.inl b/Engine/Source/Runtime/Experimental/IoStore/HttpClient/Private/Activity.inl
--- a/Engine/Source/Runtime/Experimental/IoStore/HttpClient/Private/Activity.inl
+++ b/Engine/Source/Runtime/Experimental/IoStore/HttpClient/Private/Activity.inl
@@ -746,9 +746,13 @@ struct FActivity
 
 	enum class EStage : int32
 	{
-		Request			= -1,
-		Response		= -2,
-		Content			= -3,
+		Build			= -1,
+		Request			= -2,
+		Response		= -3,
+		Content			= -4,
+		Done			= -5,
+		Cancelled		= -6,
+		Failed			= -7,
 	};
 
 						FActivity(const FParams& Params);
@@ -760,6 +764,7 @@ struct FActivity
 	void				Cancel();
 	void				Done();
 	void				Fail(const FOutcome& Outcome);
+	EStage				GetStage() const;
 	FAnsiStringView		GetMethod() const;
 	FAnsiStringView		GetPath() const;
 	FHost*				GetHost() const;
@@ -945,6 +950,25 @@ void FActivity::CallSink()
 	Sink(SinkArg);
 }
 
+////////////////////////////////////////////////////////////////////////////////
+FActivity::EStage FActivity::GetStage() const
+{
+	switch (State)
+	{
+	case EState::Build:			return EStage::Build;
+	case EState::Send:
+	case EState::RecvMessage:	return EStage::Response;
+	case EState::RecvContent:
+	case EState::RecvStream:
+	case EState::RecvDone:		return EStage::Content;
+	case EState::Completed:		return EStage::Done;
+	case EState::Cancelled:		return EStage::Cancelled;
+	case EState::Failed:		return EStage::Failed;
+	default:					break;
+	}
+	return EStage::Failed;
+}
+
 ////////////////////////////////////////////////////////////////////////////////
 FAnsiStringView FActivity::GetMethod() const
 {
diff --git a/Engine/Source/Runtime/Experimental/IoStore/HttpClient/Private/Api.inl b/Engine/Source/Runtime/Experimental/IoStore/HttpClient/Private/Api.inl
--- a/Engine/Source/Runtime/Experimental/IoStore/HttpClient/Private/Api.inl
+++ b/Engine/Source/Runtime/Experimental/IoStore/HttpClient/Private/Api.inl
@@ -160,14 +160,13 @@ void FResponse::SetDestination(FIoBuffer* Buffer)
 FTicketStatus::EId FTicketStatus::GetId() const
 {
 	const auto* Activity = (FActivity*)this;
-	switch (Activity->State)
+	switch (Activity->GetStage())
 	{
-	case FActivity::EState::RecvMessage:	return EId::Response;
-	case FActivity::EState::RecvStream:
-	case FActivity::EState::RecvDone:		return EId::Content;
-	case FActivity::EState::Cancelled:		return EId::Cancelled;
-	case FActivity::EState::Failed:			return EId::Error;
-	default:								check(false);
+	case FActivity::EStage::Response:	return EId::Response;
+	case FActivity::EStage::Content:	return EId::Content;
+	case FActivity::EStage::Cancelled:	return EId::Cancelled;
+	case FActivity::EStage::Failed:		return EId::Error;
+	default:							break;
 	}
 	return EId::Error;
 }
