Route access through an accessor

#rnx

diff --git a/Engine/Source/Runtime/Experimental/IoStore/HttpClient/Private/Activity.inl b/Engine/Source/Runtime/Experimental/IoStore/HttpClient/Private/Activity.inl
--- a/Engine/Source/Runtime/Experimental/IoStore/HttpClient/Private/Activity.inl
+++ b/Engine/Source/Runtime/Experimental/IoStore/HttpClient/Private/Activity.inl
@@ -765,6 +765,7 @@ struct FActivity
 	void				Done();
 	void				Fail(const FOutcome& Outcome);
 	EStage				GetStage() const;
+	const FBodyCloak*	GetTransaction() const;
 	FAnsiStringView		GetMethod() const;
 	FAnsiStringView		GetPath() const;
 	FHost*				GetHost() const;
@@ -969,6 +970,13 @@ FActivity::EStage FActivity::GetStage() const
 	return EStage::Failed;
 }
 
+////////////////////////////////////////////////////////////////////////////////
+const FBodyCloak* FActivity::GetTransaction() const
+{
+	check(State >= EState::Send && State <= EState::Completed);
+	return &Transaction;
+}
+
 ////////////////////////////////////////////////////////////////////////////////
 FAnsiStringView FActivity::GetMethod() const
 {
diff --git a/Engine/Source/Runtime/Experimental/IoStore/HttpClient/Private/Api.inl b/Engine/Source/Runtime/Experimental/IoStore/HttpClient/Private/Api.inl
--- a/Engine/Source/Runtime/Experimental/IoStore/HttpClient/Private/Api.inl
+++ b/Engine/Source/Runtime/Experimental/IoStore/HttpClient/Private/Api.inl
@@ -75,21 +75,21 @@ EStatusCodeClass FResponse::GetStatus() const
 uint32 FResponse::GetStatusCode() const
 {
 	const auto* Activity = (const FActivity*)this;
-	return Activity->Transaction.GetStatusCode();
+	return Activity->GetTransaction()->GetStatusCode();
 }
 
 ////////////////////////////////////////////////////////////////////////////////
 FAnsiStringView FResponse::GetStatusMessage() const
 {
 	const auto* Activity = (const FActivity*)this;
-	return Activity->Transaction.GetStatusMessage();
+	return Activity->GetTransaction()->GetStatusMessage();
 }
 
 ////////////////////////////////////////////////////////////////////////////////
 int64 FResponse::GetContentLength() const
 {
 	const auto* Activity = (const FActivity*)this;
-	return Activity->Transaction.GetContentLength();
+	return Activity->GetTransaction()->GetContentLength();
 }
 
 ////////////////////////////////////////////////////////////////////////////////
@@ -142,7 +142,7 @@ FAnsiStringView FResponse::GetHeader(FAnsiStringView Name) const
 void FResponse::ReadHeaders(FHeaderSink Sink) const
 {
 	const auto* Activity = (const FActivity*)this;
-	Activity->Transaction.ReadHeaders(Sink);
+	Activity->GetTransaction()->ReadHeaders(Sink);
 }
 
 ////////////////////////////////////////////////////////////////////////////////
@@ -205,7 +205,7 @@ uint32 FTicketStatus::GetContentLength() const
 {
 	check(GetId() <= EId::Content);
 	const auto* Activity = (FActivity*)this;
-	return uint32(Activity->Transaction.GetContentLength());
+	return uint32(Activity->GetTransaction()->GetContentLength());
 }
 
 ////////////////////////////////////////////////////////////////////////////////
@@ -263,7 +263,7 @@ FTicketPerf::FSample FTicketPerf::GetSample() const
 	uint32 Bps = ~0u;
 	if (Sample.RecvMs)
 	{
-		Bps = uint32((Activity.Transaction.GetContentLength() * 1000ull) / Sample.RecvMs);
+		Bps = uint32((Activity.GetTransaction()->GetContentLength() * 1000ull) / Sample.RecvMs);
 	}
 	Sample.RecvKiBps = Clamp(Bps >> 10);
 
diff --git a/Engine/Source/Runtime/Experimental/IoStore/HttpClient/Private/Loop.inl b/Engine/Source/Runtime/Experimental/IoStore/HttpClient/Private/Loop.inl
--- a/Engine/Source/Runtime/Experimental/IoStore/HttpClient/Private/Loop.inl
+++ b/Engine/Source/Runtime/Experimental/IoStore/HttpClient/Private/Loop.inl
@@ -392,7 +392,7 @@ void FPeerGroup::RecvInternal(FTickState& State)
 		return;
 	}
 
-	IsKeepAlive &= uint8(Activity->Transaction.IsKeepAlive() == true);
+	IsKeepAlive &= uint8(Activity->GetTransaction()->IsKeepAlive() == true);
 	LastUseMs = State.NowMs;
 	bWaiting |= Outcome.IsWaiting();
 
@@ -1079,7 +1079,7 @@ bool FEventLoop::Redirect(const FTicketStatus& Status, FTicketSink& OuterSink)
 
 	const auto& Activity = (FActivity&)Response; // todo: yuk
 
-	check(!Activity.Transaction.HasBody() || Response.GetContentLength() == 0); // should we ever hit this, we'll fix it
+	check(!Activity.GetTransaction()->HasBody() || Response.GetContentLength() == 0); // should we ever hit this, we'll fix it
 
 	// Original method should remain unchanged
 	FAnsiStringView Method = Activity.GetMethod();
