Separate out the build part of activity tick

#rnx

diff --git a/Engine/Source/Runtime/Experimental/IoStore/HttpClient/Private/Activity.inl b/Engine/Source/Runtime/Experimental/IoStore/HttpClient/Private/Activity.inl
--- a/Engine/Source/Runtime/Experimental/IoStore/HttpClient/Private/Activity.inl
+++ b/Engine/Source/Runtime/Experimental/IoStore/HttpClient/Private/Activity.inl
@@ -795,6 +795,7 @@ struct FActivity
 
 	void				CallSink();
 	void				ChangeState(EState InState);
+	FOutcome			Transact(FHttpPeer& Peer);
 	FOutcome			Send(FHttpPeer& Peer);
 	FOutcome			RecvMessage(FHttpPeer& Peer);
 	FOutcome			RecvContent(FHttpPeer& Peer, int32& MaxRecvSize);
@@ -1050,6 +1051,30 @@ const FStopwatch& FActivity::GetStopwatch() const
 }
 #endif // IAS_HTTP_WITH_PERF
 
+////////////////////////////////////////////////////////////////////////////////
+FOutcome FActivity::Transact(FHttpPeer& Peer)
+{
+	Transaction.Begin(GetMethod(), GetPath());
+
+	Transaction.AddHeader("Host", GetHost()->GetHostName());
+
+	// HTTP/1.1 is persistent by default thus "Connection" header isn't required
+	if (!GetHost()->IsPooled())
+	{
+		Transaction.AddHeader("Connection", "close");
+	}
+
+	EnumerateHeaders([this] (FAnsiStringView Key, FAnsiStringView Value)
+	{
+		Transaction.AddHeader(Key, Value);
+		return true;
+	});
+
+	Transaction.End();
+	ChangeState(EState::Send);
+	return FOutcome::Ok();
+}
+
 ////////////////////////////////////////////////////////////////////////////////
 FOutcome FActivity::Send(FHttpPeer& Peer)
 {
@@ -1266,24 +1291,11 @@ FOutcome FActivity::Tick(FHttpPeer& Peer, int32* MaxRecvSize)
 {
 	if (State == EState::Build)
 	{
-		Transaction.Begin(GetMethod(), GetPath());
-
-		Transaction.AddHeader("Host", GetHost()->GetHostName());
-
-		// HTTP/1.1 is persistent by default thus "Connection" header isn't required
-		if (!GetHost()->IsPooled())
+		FOutcome Outcome = Transact(Peer);
+		if (!Outcome.IsOk())
 		{
-			Transaction.AddHeader("Connection", "close");
+			return Outcome;
 		}
-
-		EnumerateHeaders([this] (FAnsiStringView Key, FAnsiStringView Value) {
-			Transaction.AddHeader(Key, Value);
-			return true;
-		});
-
-		Transaction.End();
-
-		ChangeState(EState::Send);
 	}
 
 	if (State == EState::Send)
