Placate Linux

#rnx

diff --git a/Engine/Source/Runtime/Experimental/IoStore/HttpClient/Private/Activity.inl b/Engine/Source/Runtime/Experimental/IoStore/HttpClient/Private/Activity.inl
--- a/Engine/Source/Runtime/Experimental/IoStore/HttpClient/Private/Activity.inl
+++ b/Engine/Source/Runtime/Experimental/IoStore/HttpClient/Private/Activity.inl
@@ -1039,7 +1039,8 @@ uint32 FActivity::GetRemainingKiB() const
 FOutcome FActivity::GetError() const
 {
 	FOutcome Outcome = FOutcome::None();
-	std::memcpy(&Outcome, &Error, sizeof(Outcome));
+	void* Ptr = &Outcome;
+	std::memcpy(Ptr, &Error, sizeof(Outcome));
 	return Outcome;
 }
 
