Deleted redundant S=MoveTemp(S)

#rnx

diff --git a/Engine/Source/Runtime/Experimental/IoStore/HttpClient/Private/ConnectionPool.inl b/Engine/Source/Runtime/Experimental/IoStore/HttpClient/Private/ConnectionPool.inl
--- a/Engine/Source/Runtime/Experimental/IoStore/HttpClient/Private/ConnectionPool.inl
+++ b/Engine/Source/Runtime/Experimental/IoStore/HttpClient/Private/ConnectionPool.inl
@@ -190,7 +190,6 @@ FOutcome FHost::Connect(FSocket& Socket)
 	// Socks connect in a blocking fashion so we're all set (ret=1)
 	if (bSocksConnected)
 	{
-		Socket = MoveTemp(Socket);
 		return FOutcome::Ok();
 	}
 
