Placate Clang

#rnx

diff --git a/Engine/Source/Runtime/Experimental/IoStore/HttpClient/Private/Misc.inl b/Engine/Source/Runtime/Experimental/IoStore/HttpClient/Private/Misc.inl
--- a/Engine/Source/Runtime/Experimental/IoStore/HttpClient/Private/Misc.inl
+++ b/Engine/Source/Runtime/Experimental/IoStore/HttpClient/Private/Misc.inl
@@ -412,7 +412,7 @@ static int32 ParseUrl(FAnsiStringView Url, FUrlOffsets& Out)
 
 	bool Bad = false;
 	Bad |= (Out.HostName.Len() == 0);
-	Bad |= bool(Out.UserInfo) & (Out.UserInfo.Len() == 0);
+	Bad |= bool(int32(Out.UserInfo) & (Out.UserInfo.Len() == 0));
 
 	if (Out.Port.Left)
 	{
diff --git a/Engine/Source/Runtime/Experimental/IoStore/HttpClient/Private/Socket.inl b/Engine/Source/Runtime/Experimental/IoStore/HttpClient/Private/Socket.inl
--- a/Engine/Source/Runtime/Experimental/IoStore/HttpClient/Private/Socket.inl
+++ b/Engine/Source/Runtime/Experimental/IoStore/HttpClient/Private/Socket.inl
@@ -107,7 +107,7 @@ static_assert(sizeof(sockaddr_in::sin_addr) == sizeof(uint32));
 #	define IsSocketResult(err)		(LastSocketResult() == WSA##err)
 #else
 	static int32 LastSocketResult() { return errno; }
-#	define IsSocketResult(err)		(LastSocketResult() == err)
+#	define IsSocketResult(err)		int32(LastSocketResult() == err)
 	static_assert(EWOULDBLOCK == EAGAIN);
 #endif
 
