Change visibility now everything is wrapped up in accessors

#rnx

diff --git a/Engine/Source/Runtime/Experimental/IoStore/HttpClient/Private/Activity.inl b/Engine/Source/Runtime/Experimental/IoStore/HttpClient/Private/Activity.inl
--- a/Engine/Source/Runtime/Experimental/IoStore/HttpClient/Private/Activity.inl
+++ b/Engine/Source/Runtime/Experimental/IoStore/HttpClient/Private/Activity.inl
@@ -720,6 +720,8 @@ void FStopwatch::Impl(uint32 Index)
 
 
 ////////////////////////////////////////////////////////////////////////////////
+static void Trace(const struct FActivity*, ETrace, uint32);
+
 static FLaneEstate* GActivityTraceEstate = LaneEstate_New({
 	.Name = "Iax/Activity",
 	.Group = "Iax",
@@ -779,6 +781,7 @@ struct FActivity
 	const FStopwatch&	GetStopwatch() const;
 #endif
 
+private:
 	enum class EState : uint8
 	{
 		Build,
@@ -828,6 +831,8 @@ struct FActivity
 		uint16			ValueLength;
 		char			Data[];
 	};
+
+    friend void Trace(const FActivity*, ETrace, uint32);
 };
 
 ////////////////////////////////////////////////////////////////////////////////
