Prevent IoStoreError if bulk container file is missing on platform.

diff --git a/Engine/Source/Runtime/Experimental/IoStore/OnDemand/Private/Tool/Upload.cpp b/Engine/Source/Runtime/Experimental/IoStore/OnDemand/Private/Tool/Upload.cpp
--- a/Engine/Source/Runtime/Experimental/IoStore/OnDemand/Private/Tool/Upload.cpp
+++ b/Engine/Source/Runtime/Experimental/IoStore/OnDemand/Private/Tool/Upload.cpp
@@ -438,8 +438,7 @@ static TIoStatusOr<FUploadResult> UploadContainerFiles(
 
 	if (ContainerFiles.Num() == 0)
 	{
-		UE_LOG(LogIas, Warning, TEXT("No container file(s) found"));
-		return FIoStatus(EIoErrorCode::Ok);
+		return FIoStatus(EIoErrorCode::InvalidParameter, TEXT("No container file(s) specified"));
 	}
 
 	TSet<FIoHash> ExistingChunks;
