Harden IAD cache timestamp tracking.
- Add a function that returns a month old data used for all missing timestamps.  This should help avoid issues where min or max timestamps were used.
- Add timestamp verification to startup verification.  If a timestamp is missing, fallback to block file modtime on filesytems that can support it.
- Add timestamps to the initial journal transaction when blocks are created.

diff --git a/Engine/Source/Runtime/Experimental/IoStore/OnDemand/Private/OnDemandInstallCache.cpp b/Engine/Source/Runtime/Experimental/IoStore/OnDemand/Private/OnDemandInstallCache.cpp
--- a/Engine/Source/Runtime/Experimental/IoStore/OnDemand/Private/OnDemandInstallCache.cpp
+++ b/Engine/Source/Runtime/Experimental/IoStore/OnDemand/Private/OnDemandInstallCache.cpp
@@ -40,6 +40,10 @@
 #define UE_ONDEMANDINSTALLCACHE_EXCLUSIVE_WRITE (0)
 #endif
 
+#ifndef UE_ONDEMANDINSTALLCACHE_USE_MODTIME
+#define UE_ONDEMANDINSTALLCACHE_USE_MODTIME (1)
+#endif
+
 #if UE_ONDEMANDINSTALLCACHE_EXCLUSIVE_WRITE
 #include "Tasks/Pipe.h"
 #endif // UE_ONDEMANDINSTALLCACHE_EXCLUSIVE_WRITE
@@ -184,11 +188,14 @@ struct FCas
 	FUniqueFileHandle			OpenWrite(FCasBlockId BlockId) const;
 	void						TrackAccess(FCasBlockId BlockId, int64 UtcTicks);
 	void						TrackAccess(FCasBlockId BlockId) { TrackAccess(BlockId, FDateTime::UtcNow().GetTicks()); }
-	void						TrackAccessIfNewer(FCasBlockId BlockId, int64 UtcTicks);
+	bool						TrackAccessIfNewer(FCasBlockId BlockId, int64 UtcTicks);
 	uint64						GetBlockInfo(FCasBlockInfoMap& OutBlockInfo);
 	void						Compact();
 	FIoStatus					Verify(TArray<FCasAddr>& OutAddrs);
 
+	// Returns the place holder timestamp to use if a timestamp is not found in the LastAccess table
+	static int64			GetTimestampForMissingLastAccess();
+
 	FStringView				RootDirectory;
 	FLookup					Lookup;
 	FBlockIdHandleCounts	BlockIds;
@@ -449,7 +456,7 @@ void FCas::TrackAccess(FCasBlockId BlockId, int64 UtcTicks)
 	LastAccess.FindOrAdd(BlockId, UtcTicks);
 }
 
-void FCas::TrackAccessIfNewer(FCasBlockId BlockId, int64 UtcTicks)
+bool FCas::TrackAccessIfNewer(FCasBlockId BlockId, int64 UtcTicks)
 {
 	check(BlockId.IsValid());
 	UE::TUniqueLock Lock(Mutex);
@@ -457,7 +464,10 @@ void FCas::TrackAccessIfNewer(FCasBlockId BlockId, int64 UtcTicks)
 	if (FoundTicks < UtcTicks)
 	{
 		FoundTicks = UtcTicks;
+		return true;
 	}
+
+	return false;
 }
 
 uint64 FCas::GetBlockInfo(FCasBlockInfoMap& OutBlockInfo)
@@ -509,7 +519,7 @@ uint64 FCas::GetBlockInfo(FCasBlockInfoMap& OutBlockInfo)
 			BlockInfo.Add(BlockId, FCasBlockInfo
 			{
 				.FileSize = uint64(FileSize),
-				.LastAccess = UtcTicks != nullptr ? *UtcTicks : 0
+				.LastAccess = UtcTicks != nullptr ? *UtcTicks : TimestampForMissingLastAccess
 			});
 			TotalSize += uint64(FileSize);
 
@@ -519,6 +529,7 @@ uint64 FCas::GetBlockInfo(FCasBlockInfoMap& OutBlockInfo)
 		IPlatformFile&		Ipf;
 		FCasBlockInfoMap&	BlockInfo;
 		FLastAccess			LastAccess;
+		const int64			TimestampForMissingLastAccess = FCas::GetTimestampForMissingLastAccess();
 		uint64				TotalSize = 0;
 	};
 
@@ -550,11 +561,55 @@ FIoStatus FCas::Verify(TArray<FCasAddr>& OutAddrs)
 	uint64				TotalVerifiedBytes = 0;
 	FIoStatus			Status = InstallCacheStatusBuilder(EIoErrorCode::Ok);
 
+	const int64 Now = FDateTime::UtcNow().GetTicks();
+	const int64 TimestampForMissingLastAccess = FCas::GetTimestampForMissingLastAccess();
+
+	IPlatformFile& Ipf = FPlatformFileManager::Get().GetPlatformFile();
+
+	UE::TUniqueLock Lock(Mutex);
+
 	for (auto BlockIt = BlockIds.CreateIterator(); BlockIt; ++BlockIt)
 	{
 		const FCasBlockId BlockId = BlockIt->Key;
 		if (const FCasBlockInfo* Info = BlockInfo.Find(BlockId))
 		{
+			if (int64* TimeStamp = LastAccess.Find(BlockId))
+			{
+				if (*TimeStamp > Now)
+				{
+					const FString Filename = GetBlockFilename(BlockId);
+					UE_LOG(LogIoStoreOnDemand, Warning, TEXT("Found future last access time for CAS block '%s'"), *Filename);
+					*TimeStamp = Now;
+				}
+			}
+			else
+			{
+				const FString Filename = GetBlockFilename(BlockId);
+				
+#if UE_ONDEMANDINSTALLCACHE_USE_MODTIME
+				FDateTime ModTime = Ipf.GetTimeStamp(*Filename);
+				if (ModTime != FDateTime::MinValue())
+				{
+					if (ModTime > Now)
+					{
+						UE_LOG(LogIoStoreOnDemand, Warning, TEXT("Missing last access time and found future mod time so using now time for CAS block '%s'"), *Filename);
+						LastAccess.Add(BlockId, Now);
+					}
+					else
+					{
+						UE_LOG(LogIoStoreOnDemand, Warning, TEXT("Missing last access time so using file mod time for CAS block '%s'"), *Filename);
+						LastAccess.Add(BlockId, ModTime.GetTicks());
+					}
+				}
+				else
+#endif // UE_ONDEMANDINSTALLCACHE_USE_MODTIME
+				{
+					// Failed to get mod time
+					UE_LOG(LogIoStoreOnDemand, Warning, TEXT("Missing last access time for CAS block '%s'"), *Filename);
+					LastAccess.Add(BlockId, TimestampForMissingLastAccess);
+				}
+			}
+
 			TotalVerifiedBytes += Info->FileSize;
 			continue;
 		}
@@ -570,7 +625,6 @@ FIoStatus FCas::Verify(TArray<FCasAddr>& OutAddrs)
 	UE_LOG(LogIoStoreOnDemand, Log, TEXT("Verified %d CAS blocks of total %.2lf MiB"),
 		BlockIds.Num(), ToMiB(TotalVerifiedBytes));
 
-	IPlatformFile& Ipf = FPlatformFileManager::Get().GetPlatformFile();
 	for (const TPair<FCasBlockId, FCasBlockInfo>& Kv : BlockInfo)
 	{
 		const FCasBlockId BlockId = Kv.Key;
@@ -592,7 +646,7 @@ FIoStatus FCas::Verify(TArray<FCasAddr>& OutAddrs)
 		if (!BlockIds.Contains(It->Value.BlockId))
 		{
 			MissingReferencedBlocks.Add(GetBlockFilename(It->Value.BlockId));
-			
+
 			OutAddrs.Add(It->Key);
 			It.RemoveCurrent();
 
@@ -608,6 +662,14 @@ FIoStatus FCas::Verify(TArray<FCasAddr>& OutAddrs)
 	return Status; 
 }
 
+int64 FCas::GetTimestampForMissingLastAccess()
+{
+	const FDateTime Now = FDateTime::UtcNow();
+	const FTimespan FourWeeks = FTimespan::FromDays(4*7);
+	const FDateTime FourWeeksAgo = Now - FourWeeks;
+	return FourWeeksAgo.GetTicks();
+}
+
 ///////////////////////////////////////////////////////////////////////////////
 struct FCasJournal
 {
@@ -1114,10 +1176,10 @@ bool FCasSnapshot::FFooter::IsValid() const
 
 TIoStatusOr<FCasSnapshot> FCasSnapshot::FromJournal(const FString& JournalFile)
 {
-	TMap<FCasAddr, FCasLocation>	CasLookup;
-	TMap<FCasBlockId, int64>		LastAccess;
-	TSet<FCasBlockId>				BlockIds;
-	FCasBlockId						CurrentBlockId;
+	FCas::FLookup		CasLookup;
+	FCas::FLastAccess	LastAccess;
+	TSet<FCasBlockId>	BlockIds;
+	FCasBlockId			CurrentBlockId;
 
 	FIoStatus ReplayStatus = FCasJournal::Replay(
 		JournalFile,
@@ -1170,6 +1232,8 @@ TIoStatusOr<FCasSnapshot> FCasSnapshot::FromJournal(const FString& JournalFile)
 		return ReplayStatus;
 	}
 
+	const int64 TimestampForMissingLastAccess = FCas::GetTimestampForMissingLastAccess();
+
 	FCasSnapshot Snapshot;
 	Snapshot.Blocks.Reserve(BlockIds.Num());
 	for (FCasBlockId BlockId : BlockIds)
@@ -1178,7 +1242,7 @@ TIoStatusOr<FCasSnapshot> FCasSnapshot::FromJournal(const FString& JournalFile)
 		Snapshot.Blocks.Add(FBlock
 		{
 			.BlockId	= BlockId,
-			.LastAccess = AccessTime != nullptr ? *AccessTime : FDateTime::UtcNow().GetTicks()
+			.LastAccess = AccessTime != nullptr ? *AccessTime : TimestampForMissingLastAccess
 		});
 	}
 
@@ -1695,6 +1759,7 @@ void FOnDemandInstallCache::Shutdown()
 	FCasJournal::FTransaction Transaction = FCasJournal::Begin(JournalFile);
 	for (const TPair<FCasBlockId, int64>& Kv : LastAccess)
 	{
+		// TODO: this is only set on shutdown so if we crash we won't have good timestamps
 		Transaction.BlockAccess(Kv.Key, Kv.Value);
 	}
 
@@ -1704,6 +1769,7 @@ void FOnDemandInstallCache::Shutdown()
 			*JournalFile, *Status.ToString());
 	}
 
+	// TODO: journal snapshot is also only made on shutdown? Feels like this would be better done on startup as well
 	IFileManager& Ifm = IFileManager::Get();
 	const FString JournalFilename = GetJournalFilename();
 	if (Ifm.FileSize(*JournalFile) > int64(MaxJournalSize))
@@ -3086,11 +3152,16 @@ FIoStatus FOnDemandInstallCache::FlushPendingChunksImpl(FPendingChunks& Chunks,
 
 			if (UtcAccessTicks)
 			{
-				Cas.TrackAccessIfNewer(CurrentBlockId, UtcAccessTicks);
+				if (Cas.TrackAccessIfNewer(CurrentBlockId, UtcAccessTicks))
+				{
+					Transaction.BlockAccess(CurrentBlockId, UtcAccessTicks);
+				}
 			}
 			else
 			{
-				Cas.TrackAccess(CurrentBlockId);
+				const int64 Now = FDateTime::UtcNow().GetTicks();
+				Cas.TrackAccess(CurrentBlockId, Now);
+				Transaction.BlockAccess(CurrentBlockId, Now);
 			}
 
 			if (CasFileHandle->Flush() == false)
