Dataflow : fix color ramp potential crash because copy method would copy from members of the other instance
- this would crash after a while ( likely after garbage collection )

diff --git a/Engine/Source/Runtime/Experimental/Dataflow/Engine/Private/Dataflow/DataflowColorRamp.cpp b/Engine/Source/Runtime/Experimental/Dataflow/Engine/Private/Dataflow/DataflowColorRamp.cpp
--- a/Engine/Source/Runtime/Experimental/Dataflow/Engine/Private/Dataflow/DataflowColorRamp.cpp
+++ b/Engine/Source/Runtime/Experimental/Dataflow/Engine/Private/Dataflow/DataflowColorRamp.cpp
@@ -4,6 +4,44 @@
 
 #include UE_INLINE_GENERATED_CPP_BY_NAME(DataflowColorRamp)
 
+FDataflowColorCurveOwner::FDataflowColorCurveOwner()
+{
+	ConstCurves.SetNum(4);
+	Curves.SetNum(4);
+
+	RichCurves.SetNum(4);
+
+	for (int32 Idx = 0; Idx < 4; ++Idx)
+	{
+		Curves[Idx].CurveToEdit = &RichCurves[Idx];
+		ConstCurves[Idx].CurveToEdit = &RichCurves[Idx];
+	}
+}
+
+FDataflowColorCurveOwner::FDataflowColorCurveOwner(const FDataflowColorCurveOwner& Other)
+{
+	this->operator=(Other);
+}
+
+FDataflowColorCurveOwner& FDataflowColorCurveOwner::operator=(const FDataflowColorCurveOwner& Other)
+{
+	ConstCurves.SetNum(4);
+	Curves.SetNum(4);
+
+	RichCurves.SetNum(4);
+
+	for (int32 Idx = 0; Idx < 4; ++Idx)
+	{
+		if (Other.RichCurves.IsValidIndex(Idx))
+		{
+			RichCurves[Idx] = Other.RichCurves[Idx];
+		}
+		Curves[Idx].CurveToEdit = &RichCurves[Idx];
+		ConstCurves[Idx].CurveToEdit = &RichCurves[Idx];
+	}
+	return *this;
+}
+
 TArray<FRichCurveEditInfoConst> FDataflowColorCurveOwner::GetCurves() const
 {
 	return ConstCurves;
diff --git a/Engine/Source/Runtime/Experimental/Dataflow/Engine/Public/Dataflow/DataflowColorRamp.h b/Engine/Source/Runtime/Experimental/Dataflow/Engine/Public/Dataflow/DataflowColorRamp.h
--- a/Engine/Source/Runtime/Experimental/Dataflow/Engine/Public/Dataflow/DataflowColorRamp.h
+++ b/Engine/Source/Runtime/Experimental/Dataflow/Engine/Public/Dataflow/DataflowColorRamp.h
@@ -13,19 +13,12 @@ public:
 	DECLARE_MULTICAST_DELEGATE_OneParam(FOnColorCurveChanged, TArray<FRichCurve*> /* ChangedCurves */);
 
 public:
-	FDataflowColorCurveOwner() 
-	{ 
-		ConstCurves.SetNum(4);
-		Curves.SetNum(4);
-
-		RichCurves.SetNum(4);
-
-		for (int32 Idx = 0; Idx < 4; ++Idx)
-		{
-			Curves[Idx].CurveToEdit = &RichCurves[Idx];
-			ConstCurves[Idx].CurveToEdit = &RichCurves[Idx];
-		}
-	}
+	DATAFLOWENGINE_API FDataflowColorCurveOwner();
+	DATAFLOWENGINE_API FDataflowColorCurveOwner(const FDataflowColorCurveOwner& Other);
+	DATAFLOWENGINE_API FDataflowColorCurveOwner& operator=(const FDataflowColorCurveOwner& Other);
+
+	FDataflowColorCurveOwner(FDataflowColorCurveOwner&& Other) = delete;
+	FDataflowColorCurveOwner& operator=(FDataflowColorCurveOwner&& Other) = delete;
 
 	/** implemenmt all the necessary virtual function from FCurveOwnerInterface */
 	DATAFLOWENGINE_API virtual TArray<FRichCurveEditInfoConst> GetCurves() const override;
