Add correct struct handling to dataflow context cache. UStructs can contain object references and where they do the context cache must report those references to the GC.

diff --git a/Engine/Source/Runtime/Experimental/Dataflow/Core/Private/Dataflow/DataflowContextCache.cpp b/Engine/Source/Runtime/Experimental/Dataflow/Core/Private/Dataflow/DataflowContextCache.cpp
--- a/Engine/Source/Runtime/Experimental/Dataflow/Core/Private/Dataflow/DataflowContextCache.cpp
+++ b/Engine/Source/Runtime/Experimental/Dataflow/Core/Private/Dataflow/DataflowContextCache.cpp
@@ -1,6 +1,7 @@
 // Copyright Epic Games, Inc. All Rights Reserved.
 
 #include "Dataflow/DataflowContextCache.h"
+#include "UObject/StructOnScope.h"
 
 namespace UE::Dataflow
 {
@@ -137,6 +138,16 @@ namespace UE::Dataflow
 		return MakeUnique<FContextCacheElementUStruct>(*this);
 	}
 
+	void FContextCacheElementUStruct::AddReferencedObjects(FReferenceCollector& Collector)
+	{
+		InstancedStruct.AddStructReferencedObjects(Collector);
+	}
+
+	FString FContextCacheElementUStruct::GetReferencerName() const
+	{
+		return TEXT("FContextCacheElementUStruct");
+	}
+
 	//////////////////////////////////////////////////////////////////////////////////////////////////////
 
 	const void* FContextCacheElementUStructArray::GetUntypedData(const IContextCacheStore& /*Context*/, const FProperty* /*PropertyIn*/) const
@@ -174,6 +185,16 @@ namespace UE::Dataflow
 		return MakeUnique<FContextCacheElementUStructArray>(*this);
 	}
 
+	void FContextCacheElementUStructArray::AddReferencedObjects(FReferenceCollector& Collector)
+	{
+		InstancedStructArray.AddReferencedObjects(Collector);
+	}
+
+	FString FContextCacheElementUStructArray::GetReferencerName() const
+	{
+		return TEXT("FContextCacheElementUStructArray");
+	}
+
 	FConstStructArrayView FContextCacheElementUStructArray::GetStructArrayView() const
 	{
 		return InstancedStructArray.GetScriptStruct() ?
@@ -244,6 +265,15 @@ namespace UE::Dataflow
 		}
 	}
 
+	void FContextCacheElementUStructArray::FInstancedStructArray::AddReferencedObjects(FReferenceCollector& Collector)
+	{
+		for(int32 Index = 0; IsValidIndex(Index); ++Index)
+		{
+			FStructOnScope ScopeStruct(ScriptStruct, GetData() + (GetStructureSize() * Index));
+			ScopeStruct.AddReferencedObjects(Collector);
+		}
+	}
+
 	//////////////////////////////////////////////////////////////////////////////////////////////////////
 
 	FContextValue::FContextValue(IContextCacheStore& InContext, FContextCacheKey InCacheKey)
diff --git a/Engine/Source/Runtime/Experimental/Dataflow/Core/Public/Dataflow/DataflowContextCache.h b/Engine/Source/Runtime/Experimental/Dataflow/Core/Public/Dataflow/DataflowContextCache.h
--- a/Engine/Source/Runtime/Experimental/Dataflow/Core/Public/Dataflow/DataflowContextCache.h
+++ b/Engine/Source/Runtime/Experimental/Dataflow/Core/Public/Dataflow/DataflowContextCache.h
@@ -392,7 +392,7 @@ namespace UE::Dataflow
 	//--------------------------------------------------------------------
 	// UStruct cache element
 	//--------------------------------------------------------------------
-	struct FContextCacheElementUStruct : public FContextCacheElementBase
+	struct FContextCacheElementUStruct : public FContextCacheElementBase, public FGCObject
 	{
 		explicit FContextCacheElementUStruct(FGuid InNodeGuid, const FProperty* InProperty, const FConstStructView& StructView, uint32 InNodeHash, FTimestamp Timestamp);
 
@@ -420,6 +420,11 @@ namespace UE::Dataflow
 
 		DATAFLOWCORE_API virtual TUniquePtr<FContextCacheElementBase> Clone(const IContextCacheStore& Context) const override;
 
+		//~ Begin FGCObject interface
+		DATAFLOWCORE_API void AddReferencedObjects(FReferenceCollector& Collector) override;
+		DATAFLOWCORE_API FString GetReferencerName() const override;
+		//~ End FGCObject interface
+
 	private:
 		FInstancedStruct InstancedStruct;
 	};
@@ -427,7 +432,7 @@ namespace UE::Dataflow
 	//--------------------------------------------------------------------
 	// UStruct array cache element
 	//--------------------------------------------------------------------
-	struct FContextCacheElementUStructArray : public FContextCacheElementBase
+	struct FContextCacheElementUStructArray : public FContextCacheElementBase, public FGCObject
 	{
 		explicit FContextCacheElementUStructArray(FGuid InNodeGuid, const FProperty* InProperty, const FConstStructArrayView& InStructArrayView, uint32 InNodeHash, FTimestamp Timestamp)
 			: FContextCacheElementBase(EType::CacheElementUStructArray, InNodeGuid, InProperty, InNodeHash, Timestamp)
@@ -469,6 +474,11 @@ namespace UE::Dataflow
 
 		virtual TUniquePtr<FContextCacheElementBase> Clone(const IContextCacheStore& Context) const override;
 
+		//~ Begin FGCObject interface
+		DATAFLOWCORE_API void AddReferencedObjects(FReferenceCollector& Collector) override;
+		DATAFLOWCORE_API FString GetReferencerName() const override;
+		//~ End FGCObject interface
+
 	private:
 		FConstStructArrayView GetStructArrayView() const;
 
@@ -494,6 +504,8 @@ namespace UE::Dataflow
 
 			const void* GetMemory() const;
 
+			void AddReferencedObjects(FReferenceCollector& Collector);
+
 			template<typename T>
 			const TArray<T>& Get() const
 			{
