IASTool ChunkPlugin command will now fail if it is unable to open a container file.

#rnx

- If we fail the upload but allow execution to continue we can end up with bad data and things won't work anyway so we should return an error value instead.
- Improve the logging if we cannot open the container file to make debugging easier.

diff --git a/Engine/Source/Runtime/Experimental/IoStore/OnDemand/Private/Tool/ChunkPlugin.cpp b/Engine/Source/Runtime/Experimental/IoStore/OnDemand/Private/Tool/ChunkPlugin.cpp
--- a/Engine/Source/Runtime/Experimental/IoStore/OnDemand/Private/Tool/ChunkPlugin.cpp
+++ b/Engine/Source/Runtime/Experimental/IoStore/OnDemand/Private/Tool/ChunkPlugin.cpp
@@ -241,8 +241,8 @@ static int32 ChunkContainer(const FString& ContainerUTocPath, const FChunkPlugin
 		FIoStatus Status = ContainerFileReader.Initialize(*FPaths::ChangeExtension(ContainerUTocPath, TEXT("")), EncryptionKeys);
 		if (!Status.IsOk())
 		{
-			UE_LOG(LogIoStoreOnDemand, Error, TEXT("Failed to open container '%s' for reading"), *ContainerUTocPath);
-			return 0; // TODO - Should this be -1 and stop execution?
+			UE_LOG(LogIoStoreOnDemand, Error, TEXT("Failed to open container '%s' for reading due to: %s"), *ContainerUTocPath, *Status.ToString());
+			return -1;
 		}
 	}
 
