Iris
* Make sure the TestNetObjectFactory unregisters all factories it created
* Fix netobject factory register log and flagged them verbose.

#rnx

diff --git a/Engine/Plugins/Runtime/ReplicationSystemTestPlugin/Source/Private/Tests/ReplicationSystem/TestNetObjectFactory.cpp b/Engine/Plugins/Runtime/ReplicationSystemTestPlugin/Source/Private/Tests/ReplicationSystem/TestNetObjectFactory.cpp
--- a/Engine/Plugins/Runtime/ReplicationSystemTestPlugin/Source/Private/Tests/ReplicationSystem/TestNetObjectFactory.cpp
+++ b/Engine/Plugins/Runtime/ReplicationSystemTestPlugin/Source/Private/Tests/ReplicationSystem/TestNetObjectFactory.cpp
@@ -95,7 +95,17 @@ UE_NET_TEST_FIXTURE(FNetObjectFactoryTestFixture, NetFactoryRegistration)
 		UE_NET_ASSERT_TRUE(FNetObjectFactoryRegistry::GetFactoryIdFromName(FactoryB) == InvalidNetObjectFactoryId);
 	}
 
+	// Create a replication system to test the factory spawning path
 	TUniquePtr<FReplicationSystemTestServer> Server(new FReplicationSystemTestServer(GetName()));
+
+	// Destroy the server
+	Server.Reset();
+
+	// Unregister all the factories registered in this unit test
+	FNetObjectFactoryRegistry::UnregisterFactory(FactoryA);
+	//B is already unregistered FNetObjectFactoryRegistry::UnregisterFactory(FactoryB);
+	FNetObjectFactoryRegistry::UnregisterFactory(FactoryC);
+	FNetObjectFactoryRegistry::UnregisterFactory(FactoryD);
 }
 
 } // end namespace UE::Net::Private
\ No newline at end of file
diff --git a/Engine/Source/Runtime/Experimental/Iris/Core/Private/Iris/ReplicationSystem/NetObjectFactoryRegistry.cpp b/Engine/Source/Runtime/Experimental/Iris/Core/Private/Iris/ReplicationSystem/NetObjectFactoryRegistry.cpp
--- a/Engine/Source/Runtime/Experimental/Iris/Core/Private/Iris/ReplicationSystem/NetObjectFactoryRegistry.cpp
+++ b/Engine/Source/Runtime/Experimental/Iris/Core/Private/Iris/ReplicationSystem/NetObjectFactoryRegistry.cpp
@@ -92,7 +92,7 @@ void FNetObjectFactoryRegistry::RegisterFactory(UClass* FactoryClass, FName Fact
 		NetFactories[Index].Id = IntCastChecked<FNetObjectFactoryId>(Index);
 	}
 
-	UE_LOG(LogIris, Log, TEXT("FNetObjectFactoryRegistry::UnregisterFactory is unregistering factory: %s name: %s id: %u"), *GetNameSafe(NetFactories[Index].NetFactoryClass.Get()), *NetFactories[Index].Name.ToString(), NetFactories[Index].Id);
+	UE_LOG(LogIris, Verbose, TEXT("FNetObjectFactoryRegistry::RegisterFactory has registered factory: %s name: %s id: %u"), *GetNameSafe(NetFactories[Index].NetFactoryClass.Get()), *NetFactories[Index].Name.ToString(), NetFactories[Index].Id);
 #endif
 }
 
@@ -114,7 +114,7 @@ void FNetObjectFactoryRegistry::UnregisterFactory(FName FactoryName)
 		const FFactoryData& Data = NetFactories[Index];
 		if (Data.Name == FactoryName)
 		{
-			UE_LOG(LogIris, Log, TEXT("FNetObjectFactoryRegistry::UnregisterFactory is unregistering factory: %s name: %s id: %u"), *GetNameSafe(Data.NetFactoryClass.Get()), *Data.Name.ToString(), Data.Id);
+			UE_LOG(LogIris, Verbose, TEXT("FNetObjectFactoryRegistry::UnregisterFactory is unregistering factory: %s name: %s id: %u"), *GetNameSafe(Data.NetFactoryClass.Get()), *Data.Name.ToString(), Data.Id);
 			
 			// Reset the entry to keep the Id's mapped to their index
 			NetFactories[Index] = FFactoryData{};
