[ChaosVD] Fix race condition when a CVD recording is started at the same time or right before PIE is started.

We were using the PIE started delegate to perform some initialization steps when a recording started. Originally these steps were just setting up a message on the screen.

Nowadays, that codepath also initializes the broadcasting of the recording state via message bus, and performs collision channels serialization. When these additional steps were added, we started using the global recording started delegate, but we never removed the PIE started delegate, therefore the code is now running twice, from different threads (the recording started delegate can run from the trace analysis thread).

In addition, this CL handles a potential case where a recording start command might be issued when we are already recording. That should be harmless, but I suspect it is causing issues with the change I made a few months ago that delays the actual start of a recording to when the trace analysis thread is ready.

Long term, this code path needs to be re-implemented with the control flow system as now it is partially async.

diff --git a/Engine/Source/Runtime/Experimental/ChaosSolverEngine/Private/Chaos/ChaosVDEngineEditorBridge.cpp b/Engine/Source/Runtime/Experimental/ChaosSolverEngine/Private/Chaos/ChaosVDEngineEditorBridge.cpp
--- a/Engine/Source/Runtime/Experimental/ChaosSolverEngine/Private/Chaos/ChaosVDEngineEditorBridge.cpp
+++ b/Engine/Source/Runtime/Experimental/ChaosSolverEngine/Private/Chaos/ChaosVDEngineEditorBridge.cpp
@@ -20,34 +20,55 @@
 #include "Serialization/MemoryWriter.h"
 #include "UObject/Package.h"
 
-
-void FChaosVDEngineEditorBridge::AddOnScreenRecordingMessage()
+bool FChaosVDEngineEditorBridge::AddOnScreenRecordingMessage(float DummyDeltaTime)
 {
+	constexpr bool bContinueLooping = false;
 	if (!GEngine)
 	{
-		return;
+		return bContinueLooping;
 	}
 
+	if (!IsInGameThread())
+	{
+		if (!DeferredShowMessageOnScreenHandle.IsValid())
+		{
+			DeferredShowMessageOnScreenHandle = FTSTicker::GetCoreTicker().AddTicker(FTickerDelegate::CreateRaw(this, &FChaosVDEngineEditorBridge::AddOnScreenRecordingMessage));
+		}
+
+		return bContinueLooping;
+	}
+
+	DeferredShowMessageOnScreenHandle.Reset();
+
 	static FText ChaosVDRecordingStartedMessage = NSLOCTEXT("ChaosVisualDebugger", "OnScreenChaosVDRecordingStartedMessage", "Chaos Visual Debugger recording in progress...");
 
 	if (CVDRecordingMessageKey == 0)
 	{
 		CVDRecordingMessageKey = GetTypeHash(ChaosVDRecordingStartedMessage.ToString());
 	}
-	
+
 	// Add a long duration value, we will remove the message manually when the recording stops
 	constexpr float MessageDurationSeconds = 3600.0f;
 	GEngine->AddOnScreenDebugMessage(CVDRecordingMessageKey, MessageDurationSeconds, FColor::Red,ChaosVDRecordingStartedMessage.ToString());
+
+	return bContinueLooping;
 }
 
 void FChaosVDEngineEditorBridge::RemoveOnScreenRecordingMessage()
 {
+	ensure(IsInGameThread());
+	
 	if (!GEngine)
 	{
 		return;
 	}
 
-	if (CVDRecordingMessageKey != 0)
+	if (DeferredShowMessageOnScreenHandle.IsValid())
+	{
+		FTSTicker::RemoveTicker(DeferredShowMessageOnScreenHandle);
+		DeferredShowMessageOnScreenHandle.Reset();
+	}
+	else if (CVDRecordingMessageKey != 0)
 	{
 		GEngine->RemoveOnScreenDebugMessage(CVDRecordingMessageKey);
 	}
@@ -59,8 +80,11 @@ void FChaosVDEngineEditorBridge::HandleCVDRecordingStarted()
 
 	BroadcastSessionStatus(FApp::GetDeltaTime());
 
-	constexpr float UpdateInterval = 0.5f;
-	RecordingStatusUpdateHandle = FTSTicker::GetCoreTicker().AddTicker(FTickerDelegate::CreateRaw(this, &FChaosVDEngineEditorBridge::BroadcastSessionStatus), UpdateInterval);
+	if (!RecordingStatusUpdateHandle.IsValid())
+	{
+		constexpr float UpdateInterval = 0.5f;
+		RecordingStatusUpdateHandle = FTSTicker::GetCoreTicker().AddTicker(FTickerDelegate::CreateRaw(this, &FChaosVDEngineEditorBridge::BroadcastSessionStatus), UpdateInterval);
+	}
 }
 
 void FChaosVDEngineEditorBridge::HandleCVDPostRecordingStarted()
@@ -73,6 +97,7 @@ void FChaosVDEngineEditorBridge::HandleCVDRecordingStopped()
 	RemoveOnScreenRecordingMessage();
 
 	FTSTicker::RemoveTicker(RecordingStatusUpdateHandle);
+	RecordingStatusUpdateHandle.Reset();
 	BroadcastSessionStatus(FApp::GetDeltaTime());
 }
 
@@ -92,8 +117,7 @@ void FChaosVDEngineEditorBridge::HandlePIEStarted(UGameInstance* GameInstance)
 	// If we were already recording, show the message
 	if (FChaosVDRuntimeModule::Get().IsRecording())
 	{
-		HandleCVDRecordingStarted();
-		HandleCVDPostRecordingStarted();
+		AddOnScreenRecordingMessage();
 	}
 }
 
@@ -256,7 +280,8 @@ void FChaosVDEngineEditorBridge::Initialize()
 	PIEStartedHandle = FWorldDelegates::OnPIEStarted.AddRaw(this, &FChaosVDEngineEditorBridge::HandlePIEStarted);
 #endif
 
-	// If we were already recording, show the message
+	// If we were already recording, we need to make sure we run the initialization step to setup the session broadcast ticker
+	// and the collision channel serialization
 	if (FChaosVDRuntimeModule::Get().IsRecording())
 	{
 		HandleCVDRecordingStarted();
@@ -266,6 +291,13 @@ void FChaosVDEngineEditorBridge::Initialize()
 
 void FChaosVDEngineEditorBridge::TearDown()
 {
+#if WITH_EDITOR
+	FWorldDelegates::OnPIEStarted.Remove(PIEStartedHandle);
+	PIEStartedHandle.Reset();
+#endif
+	FTSTicker::RemoveTicker(DeferredShowMessageOnScreenHandle);
+	DeferredShowMessageOnScreenHandle.Reset();
+	
 	FTSTicker::RemoveTicker(RecordingStatusUpdateHandle);
 
 	// Note: This works during engine shutdown because the Module Manager doesn't free the dll on module unload to account for use cases like this
@@ -274,12 +306,11 @@ void FChaosVDEngineEditorBridge::TearDown()
 	if (FChaosVDRuntimeModule::IsLoaded())
 	{
 		FChaosVDRuntimeModule::Get().RemoveRecordingStartedCallback(RecordingStartedHandle);
+		RecordingStartedHandle.Reset();
 		FChaosVDRuntimeModule::Get().RemovePostRecordingStartedCallback(PostRecordingStartedHandle);
+		PostRecordingStartedHandle.Reset();
 		FChaosVDRuntimeModule::Get().RemoveRecordingStopCallback(RecordingStoppedHandle);
-
-#if WITH_EDITOR
-		 FWorldDelegates::OnPIEStarted.Remove(PIEStartedHandle);
-#endif
+		RecordingStoppedHandle.Reset();
 
 		// Make sure of removing the message from the screen in case the recording didn't quite stopped yet
 		if (FChaosVDRuntimeModule::Get().IsRecording())
diff --git a/Engine/Source/Runtime/Experimental/ChaosSolverEngine/Public/Chaos/ChaosVDEngineEditorBridge.h b/Engine/Source/Runtime/Experimental/ChaosSolverEngine/Public/Chaos/ChaosVDEngineEditorBridge.h
--- a/Engine/Source/Runtime/Experimental/ChaosSolverEngine/Public/Chaos/ChaosVDEngineEditorBridge.h
+++ b/Engine/Source/Runtime/Experimental/ChaosSolverEngine/Public/Chaos/ChaosVDEngineEditorBridge.h
@@ -48,7 +48,7 @@ public:
 	using FChaosVDDataDataChannel = Chaos::VisualDebugger::FChaosVDOptionalDataChannel;
 
 private:
-	void AddOnScreenRecordingMessage();
+	bool AddOnScreenRecordingMessage(float DummyDeltaTime = 0.1f);
 	void RemoveOnScreenRecordingMessage();
 	void HandleCVDRecordingStarted();
 	void HandleCVDPostRecordingStarted();
@@ -77,6 +77,7 @@ private:
 	TSharedRef<FChaosVDRemoteSessionsManager> RemoteSessionsManager;
 	
 	FTSTicker::FDelegateHandle RecordingStatusUpdateHandle;
+	FTSTicker::FDelegateHandle DeferredShowMessageOnScreenHandle;
 };
 #else
 
diff --git a/Engine/Source/Runtime/Experimental/ChaosVisualDebugger/Private/ChaosVDRuntimeModule.cpp b/Engine/Source/Runtime/Experimental/ChaosVisualDebugger/Private/ChaosVDRuntimeModule.cpp
--- a/Engine/Source/Runtime/Experimental/ChaosVisualDebugger/Private/ChaosVDRuntimeModule.cpp
+++ b/Engine/Source/Runtime/Experimental/ChaosVisualDebugger/Private/ChaosVDRuntimeModule.cpp
@@ -180,12 +180,23 @@ bool FChaosVDRuntimeModule::RecordingTimerTick(float DeltaTime)
 
 void FChaosVDRuntimeModule::HandleTraceConnectionEstablished()
 {
+	UE_LOG(LogChaosVDRuntime, Log, TEXT("Trace connection established."));
+
+	if (OnConnectionDelegateHandle.IsValid())
+	{
+		FTraceAuxiliary::OnConnection.RemoveAll(this);
+		OnConnectionDelegateHandle.Reset();
+	}
+
+	if (!ensureAlwaysMsgf(bIsRecording, TEXT("Received a trace connection established callback but no CVD trace is active. This should not happen.")))
+	{
+		return;
+	}
+
 	{
 		UE::TReadScopeLock ReadLock(DelegatesRWLock);
 		RecordingStartedDelegate.Broadcast();
 	}
-
-	FTraceAuxiliary::OnConnection.RemoveAll(this);
 	
 	constexpr int32 MinAllowedTimeInSecondsBetweenCaptures = 1;
 	int32 ConfiguredTimeBetweenCaptures = CVarChaosVDGTimeBetweenFullCaptures->GetInt();
@@ -208,14 +219,12 @@ void FChaosVDRuntimeModule::HandleTraceConnectionEstablished()
 
 void FChaosVDRuntimeModule::StartRecording(TConstArrayView<FString> Args)
 {
-	if (bIsRecording)
+	if (!ensureAlwaysMsgf(!bIsRecording, TEXT("Received a recording start command while a recording is already active")))
 	{
-		return;
+		UE_LOG(LogChaosVDRuntime, Log, TEXT("[%hs] There is an active CVD recording. Attempting to stop it..."), __func__)
+		StopRecording();
 	}
 
-	// Start Listening for Trace Stopped events, in case Trace is stopped outside our control so we can gracefully stop CVD recording and log a warning 
-	FTraceAuxiliary::OnTraceStopped.AddRaw(this, &FChaosVDRuntimeModule::HandleTraceStopRequest);
-
 	// Start with a generic Failure reason
 	FText FailureReason = LOCTEXT("SeeLogsForErrorDetailsText","Please see the logs for more details...");
 
@@ -279,9 +288,18 @@ void FChaosVDRuntimeModule::StartRecording(TConstArrayView<FString> Args)
 	
 	AccumulatedRecordingTime = 0.0f;
 
-	if (ensure(bIsRecording))
+	if (ensureMsgf(bIsRecording, TEXT("Failed to start CVD recording | Reason [%s]"), *FailureReason.ToString()))
 	{
-		FTraceAuxiliary::OnConnection.AddRaw(this, &FChaosVDRuntimeModule::HandleTraceConnectionEstablished);
+		if (ensureMsgf(!OnConnectionDelegateHandle.IsValid(), TEXT("Starting a trace while we are waiting for a pending connection")))
+		{
+			OnConnectionDelegateHandle = FTraceAuxiliary::OnConnection.AddRaw(this, &FChaosVDRuntimeModule::HandleTraceConnectionEstablished);
+		}
+
+		// Start Listening for Trace Stopped events, in case Trace is stopped outside our control so we can gracefully stop CVD recording and log a warning
+		if (ensureMsgf(!OnTraceStoppedDelegateHandle.IsValid(), TEXT("Starting a trace while we are waiting for a trace session to stop.")))
+		{
+			OnTraceStoppedDelegateHandle = FTraceAuxiliary::OnTraceStopped.AddRaw(this, &FChaosVDRuntimeModule::HandleTraceStopRequest);
+		}
 	}
 	else
 	{
@@ -303,6 +321,10 @@ void FChaosVDRuntimeModule::StartRecording(TConstArrayView<FString> Args)
 void FChaosVDRuntimeModule::StopRecording()
 {
 	FTraceAuxiliary::OnConnection.RemoveAll(this);
+	OnConnectionDelegateHandle = FDelegateHandle();
+
+	FTraceAuxiliary::OnTraceStopped.RemoveAll(this);
+	OnTraceStoppedDelegateHandle = FDelegateHandle();
 
 	if (!bIsRecording)
 	{
@@ -310,8 +332,6 @@ void FChaosVDRuntimeModule::StopRecording()
 		return;
 	}
 
-	FTraceAuxiliary::OnTraceStopped.RemoveAll(this);
-
 #if UE_TRACE_ENABLED || UE_TRACE_MINIMAL_ENABLED
 	RestoreTraceChannelsToPreRecordingState();
 
diff --git a/Engine/Source/Runtime/Experimental/ChaosVisualDebugger/Public/ChaosVDRuntimeModule.h b/Engine/Source/Runtime/Experimental/ChaosVisualDebugger/Public/ChaosVDRuntimeModule.h
--- a/Engine/Source/Runtime/Experimental/ChaosVisualDebugger/Public/ChaosVDRuntimeModule.h
+++ b/Engine/Source/Runtime/Experimental/ChaosVisualDebugger/Public/ChaosVDRuntimeModule.h
@@ -161,6 +161,9 @@ private:
 	FTSTicker::FDelegateHandle FullCaptureRequesterHandle;
 	FTSTicker::FDelegateHandle RecordingTimerHandle;
 
+	FDelegateHandle OnConnectionDelegateHandle;
+	FDelegateHandle OnTraceStoppedDelegateHandle;
+
 	static UE_API FChaosVDRecordingStateChangedDelegate RecordingStartedDelegate;
 	static UE_API FChaosVDRecordingStateChangedDelegate PostRecordingStartedDelegate;
 	static UE_API FChaosVDRecordingStateChangedDelegate RecordingStopDelegate;
