Support Preskinned Local Bounds for nanite geometry collections.
There was no specific GetPreSkinnedLocalBounds() implementation for nanite geometry collection.
Added one which returns a cached copy of LocalBounds taken at scene proxy init time.

diff --git a/Engine/Source/Runtime/Experimental/GeometryCollectionEngine/Private/GeometryCollection/GeometryCollectionSceneProxy.cpp b/Engine/Source/Runtime/Experimental/GeometryCollectionEngine/Private/GeometryCollection/GeometryCollectionSceneProxy.cpp
--- a/Engine/Source/Runtime/Experimental/GeometryCollectionEngine/Private/GeometryCollection/GeometryCollectionSceneProxy.cpp
+++ b/Engine/Source/Runtime/Experimental/GeometryCollectionEngine/Private/GeometryCollection/GeometryCollectionSceneProxy.cpp
@@ -1345,6 +1345,8 @@ void FNaniteGeometryCollectionSceneProxy::CreateRenderThreadResources(FRHIComman
 	check(NaniteResources.RuntimeResourceID != INDEX_NONE && NaniteResources.HierarchyOffset != INDEX_NONE);
 
 	FGeometryCollectionSceneProxyBase::SetDynamicData_RenderThread(RHICmdList, DynamicData);
+
+	CachedPreSkinnedBounds = GetLocalBounds();
 }
 
 void FNaniteGeometryCollectionSceneProxy::DestroyRenderThreadResources()
@@ -1358,6 +1360,11 @@ SIZE_T FNaniteGeometryCollectionSceneProxy::GetTypeHash() const
 	return reinterpret_cast<size_t>(&UniquePointer);
 }
 
+void FNaniteGeometryCollectionSceneProxy::GetPreSkinnedLocalBounds(FBoxSphereBounds& OutBounds) const
+{
+	OutBounds = CachedPreSkinnedBounds;
+}
+
 FPrimitiveViewRelevance FNaniteGeometryCollectionSceneProxy::GetViewRelevance(const FSceneView* View) const
 {
 	LLM_SCOPE_BYTAG(Nanite);
diff --git a/Engine/Source/Runtime/Experimental/GeometryCollectionEngine/Private/GeometryCollection/GeometryCollectionSceneProxy.h b/Engine/Source/Runtime/Experimental/GeometryCollectionEngine/Private/GeometryCollection/GeometryCollectionSceneProxy.h
--- a/Engine/Source/Runtime/Experimental/GeometryCollectionEngine/Private/GeometryCollection/GeometryCollectionSceneProxy.h
+++ b/Engine/Source/Runtime/Experimental/GeometryCollectionEngine/Private/GeometryCollection/GeometryCollectionSceneProxy.h
@@ -319,6 +319,7 @@ public:
 	// FPrimitiveSceneProxy interface.
 	virtual void CreateRenderThreadResources(FRHICommandListBase& RHICmdList) override;
 	virtual void DestroyRenderThreadResources() override;
+	virtual void GetPreSkinnedLocalBounds(FBoxSphereBounds& OutBounds) const override;
 	virtual SIZE_T GetTypeHash() const override;
 	virtual FPrimitiveViewRelevance	GetViewRelevance(const FSceneView* View) const override;
 	virtual void GetDynamicMeshElements(const TArray<const FSceneView*>& Views, const FSceneViewFamily& ViewFamily, uint32 VisibilityMap, FMeshElementCollector& Collector) const override;
@@ -376,6 +377,8 @@ protected:
 	};
 	TArray<FGeometryNaniteData> GeometryNaniteData;
 
+	FBoxSphereBounds CachedPreSkinnedBounds;
+
 	uint32 bCastShadow : 1;
 	uint32 bReverseCulling : 1;
 	uint32 bHasMaterialErrors : 1;
