Iris - Avoid marking potentially invalid memory as reachable UObjects in FObjectNetSerializer::IsEqual.

#rnx

diff --git a/Engine/Source/Runtime/Experimental/Iris/Core/Private/Iris/Serialization/ObjectNetSerializer.cpp b/Engine/Source/Runtime/Experimental/Iris/Core/Private/Iris/Serialization/ObjectNetSerializer.cpp
--- a/Engine/Source/Runtime/Experimental/Iris/Core/Private/Iris/Serialization/ObjectNetSerializer.cpp
+++ b/Engine/Source/Runtime/Experimental/Iris/Core/Private/Iris/Serialization/ObjectNetSerializer.cpp
@@ -136,15 +136,15 @@ protected:
 	static void SetRawPtrValue(UObject*& Dst, UObject* Value) { Dst = Value; }
 
 	// TObjectPtr getters/setters
-	static UObject* GetRawPtrValue(TObjectPtr<UObject> Value) { return Value; }
-	static TObjectPtr<UObject> GetObjectPtrValue(TObjectPtr<UObject> Value) { return Value; }
+	static UObject* GetRawPtrValue(const TObjectPtr<UObject>& Value) { return Value; }
+	static const TObjectPtr<UObject>& GetObjectPtrValue(const TObjectPtr<UObject>& Value) { return Value; }
 
 	static void SetRawPtrValue(TObjectPtr<UObject>& Dst, UObject* Value) { Dst = Value; }
 	static void SetObjectPtrValue(TObjectPtr<UObject>& Dst, const FRemoteObjectReference& Value) { Dst = TObjectPtr<UObject>(Value.ToObjectPtr()); }
 
 	// Weak pointer getters/setters
 	static UObject* GetRawPtrValue(const TWeakObjectPtr<UObject>& Value) { return Value.Get(); }
-	static TWeakObjectPtr<UObject> GetObjectPtrValue(const TWeakObjectPtr<UObject>& Value) { return Value; }
+	static const TWeakObjectPtr<UObject>& GetObjectPtrValue(const TWeakObjectPtr<UObject>& Value) { return Value; }
 
 	static void SetRawPtrValue(TWeakObjectPtr<UObject>& Dst, UObject* Value) { Dst = TWeakObjectPtr<UObject>(Value); }
 #if UE_WITH_REMOTE_OBJECT_HANDLE
@@ -153,7 +153,7 @@ protected:
 
 	// For FScriptInterface we have a custom Dequantize so we only provide a value getter.
 	static UObject* GetRawPtrValue(const FScriptInterface& Value) { return Value.GetObject(); }
-	static TObjectPtr<UObject> GetObjectPtrValue(FScriptInterface& Value) { return Value.GetObjectRef(); }
+	static const TObjectPtr<UObject>& GetObjectPtrValue(FScriptInterface& Value) { return Value.GetObjectRef(); }
 
 	static UObject* ResolveObjectReference(FNetSerializationContext&, const FNetObjectReference&);
 
@@ -182,6 +182,8 @@ struct FObjectPtrNetSerializer : public Private::FObjectNetSerializerBase<TObjec
 	typedef FObjectPtrNetSerializerConfig ConfigType;
 
 	inline static const ConfigType DefaultConfig;
+
+	static bool IsEqual(FNetSerializationContext&, const FNetIsEqualArgs&);
 };
 
 UE_NET_IMPLEMENT_SERIALIZER_INTERNAL(FObjectPtrNetSerializer);
@@ -430,6 +432,27 @@ UObject* FObjectNetSerializerBase<T>::ResolveObjectReference(FNetSerializationCo
 namespace UE::Net
 {
 
+// FObjectPtrNetSerializer implementation
+bool FObjectPtrNetSerializer::IsEqual(FNetSerializationContext& Context, const FNetIsEqualArgs& Args)
+{
+	if (Args.bStateIsQuantized)
+	{
+		const QuantizedType& Value0 = *reinterpret_cast<const QuantizedType*>(Args.Source0);
+		const QuantizedType& Value1 = *reinterpret_cast<const QuantizedType*>(Args.Source1);
+
+		return Value0 == Value1;
+	}
+	else
+	{
+		const SourceType& Value0 = *reinterpret_cast<const SourceType*>(Args.Source0);
+		const SourceType& Value1 = *reinterpret_cast<const SourceType*>(Args.Source1);
+
+		// Avoid FObjectPtr operator == which passes by value and may mark objects as reachable. We can compare a formerly valid valid pointer that had its memory deleted or re-used.
+		// We don't care about UE_WITH_OBJECT_HANDLE_TYPE_SAFETY being defined or not.
+		return Value0.GetHandle() == Value1.GetHandle();
+	}
+}
+
 // FScriptInterfaceNetSerializer implementation
 void FScriptInterfaceNetSerializer::Dequantize(FNetSerializationContext& Context, const FNetDequantizeArgs& Args)
 {
