Default alignment is sufficient these days

#rnx

diff --git a/Engine/Source/Runtime/Experimental/IoStore/HttpClient/Private/Activity.inl b/Engine/Source/Runtime/Experimental/IoStore/HttpClient/Private/Activity.inl
--- a/Engine/Source/Runtime/Experimental/IoStore/HttpClient/Private/Activity.inl
+++ b/Engine/Source/Runtime/Experimental/IoStore/HttpClient/Private/Activity.inl
@@ -832,7 +832,9 @@ private:
 		char			Data[];
 	};
 
-    friend void Trace(const FActivity*, ETrace, uint32);
+	friend void Trace(const FActivity*, ETrace, uint32);
+
+	UE_NONCOPYABLE(FActivity);
 };
 
 ////////////////////////////////////////////////////////////////////////////////
