Added options to generate a corpus from the test server

#rnx

diff --git a/Engine/Source/Runtime/Experimental/IoStore/HttpClient/Private/TestServer.py b/Engine/Source/Runtime/Experimental/IoStore/HttpClient/Private/TestServer.py
--- a/Engine/Source/Runtime/Experimental/IoStore/HttpClient/Private/TestServer.py
+++ b/Engine/Source/Runtime/Experimental/IoStore/HttpClient/Private/TestServer.py
@@ -174,7 +174,7 @@ def tls_proxy_loop(httpd_port, root_cert, *server_pems):
 
 # {{{1 httpd ...................................................................
 
-payload_data = random.randbytes(2 << 20)
+payload_data = bytearray((x % (128 - 34)) + 33 for x in random.randbytes(2 << 20))
 
 def _make_payload(size):
     size = int(size)
@@ -277,6 +277,8 @@ def http_redirect(handler, style="abs", code=302, *dest):
     handler.end_headers()
 
 class Handler(http.server.BaseHTTPRequestHandler):
+    corpus_thunk = None
+
     def _preroll(self):
         self.protocol_version = "HTTP/1.1"
 
@@ -289,6 +291,14 @@ class Handler(http.server.BaseHTTPRequestHandler):
 
         self.send_error(404)
 
+    def parse_request(self):
+        ret = super().parse_request()
+
+        if self.corpus_thunk:
+            self.corpus_thunk(self.path, self.wfile)
+
+        return ret
+
     def do_HEAD(self):
         parts = self._preroll()
         if not parts:
@@ -349,15 +359,18 @@ class Handler(http.server.BaseHTTPRequestHandler):
 
         return self.send_error(404, f"not found '{self.path}'")
 
-def plain_httpd_loop(port, root_pem):
+def plain_httpd_loop(port, root_pem, corpus_thunk):
+    Handler.corpus_thunk = corpus_thunk
+
     print(f"httpd: {port}")
+    print("corpus: ", corpus_thunk.__name__ if corpus_thunk else "")
     server = http.server.ThreadingHTTPServer(("", port), Handler)
     server.ca_pem = root_pem
     server.serve_forever()
 
 
 
-# {{{1 main ....................................................................
+# {{{1 certs ...................................................................
 
 def gen_test_certs(openssl_bin):
     temp_dir_obj = tempfile.TemporaryDirectory(prefix="iastestserver_")
@@ -420,7 +433,46 @@ def gen_test_certs(openssl_bin):
     with key_path.open("rb") as inp:  k = inp.read()
     return r, s, k
 
-def main():
+
+
+# {{{1 main ....................................................................
+
+def setup_corpus(args):
+    if not (corpus_dir := args.corpusdir[0]):
+        return
+
+    corpus_dir.mkdir(parents=True, exist_ok=True)
+    for item in corpus_dir.glob("*.iax"):
+        item.unlink()
+
+    def corpus_collector(self, path, writer):
+        if random.random() > args.corpuschance[0]:
+            return
+
+        path = path[1:]
+        path = path.replace("/", "-")
+        path = path.replace("?", "-")
+        path = path.replace(":", "-")
+
+        name = "test_%08x_%s.iax" % (random.randrange(0, 0x7fffffff), path)
+        out = (corpus_dir / name).open("wb")
+
+        if not (orig_writer := getattr(writer, "orig_writer", None)):
+            writer.orig_writer = orig_writer = writer.write
+
+        def write(data):
+            out.write(data)
+            return orig_writer(data)
+
+        def close():
+            out.close()
+
+        writer.write = write
+        writer.close = close
+
+    return corpus_collector
+
+def main(args):
     for item in Path(__file__).parents:
         if (item / "GenerateProjectFiles.bat").is_file():
             os.chdir(item)
@@ -431,22 +483,24 @@ def main():
 
     if with_tls:
         print("\n## generating certificates")
-        openssl_bin = os.getenv("OPENSSL_BIN", "openssl")
+        openssl_bin = args.opensslbin[0] or os.getenv("OPENSSL_BIN", "openssl")
         if (x := Path(".") / "Engine/Binaries/DotNET/IOS/openssl.exe").is_file():
             openssl_bin = str(x)
         root_cert, server_cert, server_key = gen_test_certs(openssl_bin)
     else:
         root_cert = None
 
-    httpd_port = int(os.getenv("IasTestServerPort", 0))
+    httpd_port = args.port[0] or int(os.getenv("IasTestServerPort", 0))
     httpd_port = httpd_port or (int(random.random() * 0x8000) + 0x4000)
 
     def start_svc(target, *args):
         threading.Thread(target=target, args=args, daemon=True).start()
 
+    corpus_thunk = setup_corpus(args)
+
     print("\n## starting servers")
     start_svc(proxy_loop, httpd_port)
-    start_svc(plain_httpd_loop, httpd_port, root_cert)
+    start_svc(plain_httpd_loop, httpd_port, root_cert, corpus_thunk)
     if with_tls:
         start_svc(tls_proxy_loop, httpd_port, root_cert, server_key, server_cert)
 
@@ -456,6 +510,13 @@ def main():
         time.sleep(3600)
 
 if __name__ == "__main__":
-    main()
+    import argparse
+    parser = argparse.ArgumentParser(description="Test server to 'IasTool tests'")
+    parser.add_argument("--port",         nargs=1, default=(0,),    type=int,   help="Port to listen on", action="store")
+    parser.add_argument("--opensslbin",   nargs=1, default=(None,), type=Path,  help="Path to OpenSSL binary used to generate certs")
+    parser.add_argument("--corpusdir",    nargs=1, default=(None,), type=Path,  help="Directory to output response fuzzing corpus")
+    parser.add_argument("--corpuschance", nargs=1, default=(0.2,),  type=float, help="Chance of writing corpus item (0.0-1.0)")
+    args = parser.parse_args()
+    raise SystemExit(main(args))
 
 # vim: et
