Default alignment is sufficient these days

#rnx

diff --git a/Engine/Source/Runtime/Experimental/IoStore/HttpClient/Private/Misc.inl b/Engine/Source/Runtime/Experimental/IoStore/HttpClient/Private/Misc.inl
--- a/Engine/Source/Runtime/Experimental/IoStore/HttpClient/Private/Misc.inl
+++ b/Engine/Source/Runtime/Experimental/IoStore/HttpClient/Private/Misc.inl
@@ -431,7 +431,7 @@ static int32 ParseUrl(FAnsiStringView Url, FUrlOffsets& Out)
 // {{{1 buffer .................................................................
 
 ////////////////////////////////////////////////////////////////////////////////
-class alignas(16) FBuffer
+class FBuffer
 {
 public:
 	struct FMutableSection
