Formalise loop's list of activities.

#rnx

diff --git a/Engine/Source/Runtime/Experimental/IoStore/HttpClient/Private/Loop.inl b/Engine/Source/Runtime/Experimental/IoStore/HttpClient/Private/Loop.inl
--- a/Engine/Source/Runtime/Experimental/IoStore/HttpClient/Private/Loop.inl
+++ b/Engine/Source/Runtime/Experimental/IoStore/HttpClient/Private/Loop.inl
@@ -25,8 +25,6 @@ namespace UE::IoStore::HTTP
  */
 
 ////////////////////////////////////////////////////////////////////////////////
-using FActivityList = struct FActivityNode*;
-
 struct FActivityNode
 	: public FActivity
 {
@@ -157,6 +155,143 @@ void FActivityNode::Destroy(FActivityNode* Ptr)
 
 
 
+////////////////////////////////////////////////////////////////////////////////
+class FActivityList
+{
+public:
+					FActivityList()							= default;
+					~FActivityList();
+					FActivityList(FActivityList&& Rhs)		{ Swap(Head, Rhs.Head); }
+	void			operator = (FActivityList&& Rhs)		{ Swap(Head, Rhs.Head); }
+	int32			IsEmpty() const							{ return Head == nullptr; }
+	FActivityNode*	GetHead() const							{ return Head; }
+	FActivityNode*	Detach()								{ FActivityNode* Ret = Head; Head = nullptr; return Ret; }
+	void			Reverse();
+	void			Join(FActivityList& Rhs);
+	void			Prepend(FActivityNode* Node);
+	void			Append(FActivityNode* Node);
+	FActivityNode*	Pop();
+	void			PopPrepend(FActivityList& ToList);
+	void			PopAppend(FActivityList& ToList);
+
+	template <typename LAMBDA>
+	FActivityNode*	FindFirst(LAMBDA&& Predicate);
+
+private:
+	FActivityNode*	PopImpl();
+	FActivityNode*	Head = nullptr;
+
+private:
+					FActivityList(const FActivityList&) = delete;
+	FActivityList&	operator = (const FActivityList&) = delete;
+};
+
+////////////////////////////////////////////////////////////////////////////////
+FActivityList::~FActivityList()
+{
+	check(Head == nullptr);
+}
+
+////////////////////////////////////////////////////////////////////////////////
+template <typename LAMBDA>
+FActivityNode* FActivityList::FindFirst(LAMBDA&& Predicate)
+{
+	for (FActivityNode* Node = Head; Node != nullptr; Node = Node->Next)
+	{
+		if (Predicate(Node))
+		{
+			return Node;
+		}
+	}
+
+	return nullptr;
+}
+
+////////////////////////////////////////////////////////////////////////////////
+void FActivityList::Reverse()
+{
+	FActivityNode* Reverse = nullptr;
+	for (FActivityNode* Next; Head != nullptr; Head = Next)
+	{
+		Next = Head->Next;
+		Head->Next = Reverse;
+		Reverse = Head;
+	}
+	Head = Reverse;
+}
+
+////////////////////////////////////////////////////////////////////////////////
+void FActivityList::Join(FActivityList& Rhs)
+{
+	if (Head == nullptr)
+	{
+		Swap(Head, Rhs.Head);
+		return;
+	}
+
+	FActivityNode* Tail = Head;
+	for (; Tail->Next != nullptr; Tail = Tail->Next);
+	Swap(Tail->Next, Rhs.Head);
+}
+
+////////////////////////////////////////////////////////////////////////////////
+void FActivityList::Prepend(FActivityNode* Node)
+{
+	check(Node->Next == nullptr);
+	Node->Next = Head;
+	Head = Node;
+}
+
+////////////////////////////////////////////////////////////////////////////////
+void FActivityList::Append(FActivityNode* Node)
+{
+	if (Head == nullptr)
+	{
+		Head = Node;
+		return;
+	}
+
+	check(Head->Next == nullptr);
+	Head->Next = Node;
+}
+
+////////////////////////////////////////////////////////////////////////////////
+FActivityNode* FActivityList::Pop()
+{
+	if (Head == nullptr)
+	{
+		return Head;
+	}
+	return PopImpl();
+}
+
+////////////////////////////////////////////////////////////////////////////////
+FActivityNode* FActivityList::PopImpl()
+{
+	FActivityNode* Node = Head;
+	Head = Node->Next;
+	Node->Next = nullptr;
+	return Node;
+}
+
+////////////////////////////////////////////////////////////////////////////////
+void FActivityList::PopPrepend(FActivityList& ToList)
+{
+	check(Head != nullptr);
+	FActivityNode* Node = PopImpl();
+	ToList.Prepend(Node);
+}
+
+////////////////////////////////////////////////////////////////////////////////
+void FActivityList::PopAppend(FActivityList& ToList)
+{
+	check(Head != nullptr);
+	FActivityNode* Node = PopImpl();
+	ToList.Append(Node);
+}
+
+
+
 ////////////////////////////////////////////////////////////////////////////////
 struct FTickState
 {
@@ -177,13 +312,13 @@ class FWorkQueue
 public:
 						FWorkQueue() = default;
 						~FWorkQueue();
-	bool				HasWork() const { return List != nullptr; }
+	bool				HasWork() const { return !List.IsEmpty(); }
 	void				AddActivity(FActivityNode* Activity);
 	FActivityNode*		PopActivity();
 	void				TickCancels(FTickState& State);
 
 private:
-	FActivityList		List = nullptr;
+	FActivityList		List;
 	FActivityNode*		ListTail = nullptr;
 	uint64				ActiveSlots = 0;
 
@@ -193,7 +328,6 @@ private:
 ////////////////////////////////////////////////////////////////////////////////
 FWorkQueue::~FWorkQueue()
 {
-	check(List == nullptr);
 	check(ListTail == nullptr);
 }
 
@@ -204,11 +338,14 @@ void FWorkQueue::AddActivity(FActivityNode* Activity)
 
 	check(Activity->Next == nullptr);
 
-	if (ListTail != nullptr)
+	if (List.IsEmpty())
+	{
+		List.Prepend(Activity);
+	}
+	else
 	{
 		ListTail->Next = Activity;
 	}
-	List = (List == nullptr) ? Activity : List;
 	ListTail = Activity;
 
 	ActiveSlots |= (1ull << Activity->Slot);
@@ -217,13 +354,13 @@ void FWorkQueue::AddActivity(FActivityNode* Activity)
 ////////////////////////////////////////////////////////////////////////////////
 FActivityNode* FWorkQueue::PopActivity()
 {
-	if (List == nullptr)
+	FActivityNode* Activity = List.Pop();
+	if (Activity == nullptr)
 	{
 		return nullptr;
 	}
 
-	FActivityNode* Activity = List;
-	if ((List = List->Next) == nullptr)
+	if (List.IsEmpty())
 	{
 		ListTail = nullptr;
 	}
@@ -231,7 +368,7 @@ FActivityNode* FWorkQueue::PopActivity()
 	check(ActiveSlots & (1ull << Activity->Slot));
 	ActiveSlots ^= (1ull << Activity->Slot);
 
-	Activity->Next = nullptr;
+	check(Activity->Next == nullptr);
 	return Activity;
 }
 
@@ -246,9 +383,9 @@ void FWorkQueue::TickCancels(FTickState& State)
 	// We are going to rebuild the list of activities to maintain order as the
 	// activity list is singular.
 
-	check(List != nullptr);
-	FActivityNode* Activity = List;
-	List = ListTail = nullptr;
+	check(!List.IsEmpty());
+	FActivityNode* Activity = List.Detach();
+	ListTail = nullptr;
 	ActiveSlots = 0;
 
 	for (FActivityNode* Next; Activity != nullptr; Activity = Next)
@@ -264,8 +401,7 @@ void FWorkQueue::TickCancels(FTickState& State)
 
 		Activity->Cancel();
 
-		Activity->Next = State.DoneList;
-		State.DoneList = Activity;
+		State.DoneList.Prepend(Activity);
 	}
 }
 
@@ -289,8 +425,8 @@ private:
 	void				Negotiate(FTickState& State);
 	void				RecvInternal(FTickState& State);
 	void				SendInternal(FTickState& State);
-	FActivityList		Send = nullptr;
-	FActivityList		Recv = nullptr;
+	FActivityList		Send;
+	FActivityList		Recv;
 	FHttpPeer			Peer;
 	uint32				LastUseMs = 0;
 	uint8				IsKeepAlive = 0;
@@ -303,8 +439,6 @@ private:
 ////////////////////////////////////////////////////////////////////////////////
 FPeerGroup::~FPeerGroup()
 {
-	check(Send == nullptr);
-	check(Recv == nullptr);
 }
 
 ////////////////////////////////////////////////////////////////////////////////
@@ -318,34 +452,29 @@ FWaiter FPeerGroup::GetWaiter() const
 	FWaitable Waitable = Peer.GetWaitable();
 
 	FWaiter Waiter(MoveTemp(Waitable));
-	Waiter.WaitFor((Recv != nullptr) ? FWaiter::EWhat::Recv : FWaiter::EWhat::Send);
+	Waiter.WaitFor((!Recv.IsEmpty()) ? FWaiter::EWhat::Recv : FWaiter::EWhat::Send);
 	return Waiter;
 }
 
 ////////////////////////////////////////////////////////////////////////////////
 void FPeerGroup::Fail(FTickState& State, const char* Reason)
 {
-	// Any send left at this point is unrecoverable
-	if (Send != nullptr)
-	{
-		Send->Next = Recv;
-		Recv = Send;
-	}
+	// Any send left at this point is unrecoverable. Send is likely shorter.
+	Send.Join(Recv);
+	Recv = MoveTemp(Send);
 
 	// Failure is quite terminal and we need to abort everything
-	for (FActivityNode* Activity = Recv; Activity != nullptr;)
+	while (FActivityNode* Activity = Recv.Pop())
 	{
 		FOutcome Outcome = FOutcome::Error(Reason);
 		Activity->Fail(Outcome);
 
-		FActivityNode* Next = Activity->Next;
-		Activity->Next = State.DoneList;
-		State.DoneList = Activity;
-		Activity = Next;
+		State.DoneList.Prepend(Activity);
 	}
 
 	Peer = FHttpPeer();
-	Send = Recv = nullptr;
+	Send = FActivityList();
+	Recv = FActivityList();
 	bWaiting = false;
 	IsKeepAlive = 0;
 	bNegotiating = false;
@@ -355,7 +484,7 @@ void FPeerGroup::Fail(FTickState& State, const char* Reason)
 void FPeerGroup::Negotiate(FTickState& State)
 {
 	check(bNegotiating);
-	check(Send != nullptr);
+	check(!Send.IsEmpty());
 	check(Peer.IsValid());
 
 	FOutcome Outcome = Peer.Handshake();
@@ -379,9 +508,9 @@ void FPeerGroup::Negotiate(FTickState& State)
 void FPeerGroup::RecvInternal(FTickState& State)
 {
 	check(bNegotiating == false);
-	check(Recv != nullptr);
+	check(!Recv.IsEmpty());
 
-	FActivityNode* Activity = Recv;
+	FActivityNode* Activity = Recv.GetHead();
 
 	FOutcome Outcome = Activity->Tick(Peer, &State.RecvAllowance);
 
@@ -397,7 +526,7 @@ void FPeerGroup::RecvInternal(FTickState& State)
 	bWaiting |= Outcome.IsWaiting();
 
 	// If we've only a small amount left to receive we can start more work
-	if (IsKeepAlive & (Recv->Next == nullptr) & (Send == nullptr))
+	if (IsKeepAlive & (Activity->Next == nullptr) & Send.IsEmpty())
 	{
 		uint32 Remaining = Activity->GetRemainingKiB();
 		if (Remaining < uint32(GRecvWorkThresholdKiB))
@@ -407,9 +536,9 @@ void FPeerGroup::RecvInternal(FTickState& State)
 				Trace(Activity, ETrace::StartWork);
 	
 				check(Next->GetHost()->IsPooled());
-				check(Send == nullptr);
+				check(Send.IsEmpty());
 
-				Send = Next;
+				Send.Append(Next);
 				SendInternal(State);
 
 				if (!Peer.IsValid())
@@ -435,9 +564,7 @@ void FPeerGroup::RecvInternal(FTickState& State)
 
 	check(Outcome.GetResult() == uint32(FActivity::EStage::Content));
 
-	Recv = Activity->Next;
-	Activity->Next = State.DoneList;
-	State.DoneList = Activity;
+	Recv.PopPrepend(State.DoneList);
 
 	// If the server wants the connection closed then we can drop our peer
 	if (IsKeepAlive == 0)
@@ -451,9 +578,9 @@ void FPeerGroup::SendInternal(FTickState& State)
 {
 	check(bNegotiating == false);
 	check(IsKeepAlive == 1);
-	check(Send != nullptr);
+	check(!Send.IsEmpty());
 
-	FActivityNode* Activity = Send;
+	FActivityNode* Activity = Send.GetHead();
 
 	FOutcome Outcome = Activity->Tick(Peer);
 
@@ -470,17 +597,8 @@ void FPeerGroup::SendInternal(FTickState& State)
 	}
 
 	check(Outcome.GetResult() == uint32(FActivity::EStage::Request));
-	Send = nullptr;
-
-	// Pass along this send to be received
-	if (Recv == nullptr)
-	{
-		Recv = Activity;
-		return;
-	}
 
-	check(Recv->Next == nullptr);
-	Recv->Next = Activity;
+	Send.PopAppend(Recv);
 }
 
 ////////////////////////////////////////////////////////////////////////////////
@@ -491,24 +609,24 @@ bool FPeerGroup::Tick(FTickState& State)
 		Negotiate(State);
 	}
 
-	else if (Send != nullptr)
+	else if (!Send.IsEmpty())
 	{
 		SendInternal(State);
 	}
 
-	if (Recv != nullptr && State.RecvAllowance)
+	if (!Recv.IsEmpty() && State.RecvAllowance)
 	{
 		RecvInternal(State);
 	}
 
-	return !!IsKeepAlive | !!(UPTRINT(Send) | UPTRINT(Recv));
+	return !!IsKeepAlive | !(Send.IsEmpty() & Recv.IsEmpty());
 }
 
 ////////////////////////////////////////////////////////////////////////////////
 void FPeerGroup::TickSend(FTickState& State, FHost& Host, FPoller& Poller)
 {
 	// This path is only for those that are idle and have nothing to do
-	if (Send != nullptr || Recv != nullptr)
+	if (!(Send.IsEmpty() & Recv.IsEmpty()))
 	{
 		return;
 	}
@@ -547,8 +665,7 @@ void FPeerGroup::TickSend(FTickState& State, FHost& Host, FPoller& Poller)
 		if (Outcome.IsError())
 		{
 			// We failed to connect, let's bail.
-			Pending->Next = Recv;
-			Recv = Pending;
+			Recv.Prepend(Pending);
 			Fail(State, Outcome.GetMessage().GetData());
 			return;
 		}
@@ -562,7 +679,7 @@ void FPeerGroup::TickSend(FTickState& State, FHost& Host, FPoller& Poller)
 		Peer = FHttpPeer(MoveTemp(Socket), VerifyCert, HostName);
 	}
 
-	Send = Pending;
+	Send.Append(Pending);
 
 	if (!bWillBlock)
 	{
@@ -758,7 +875,7 @@ private:
 	std::atomic<uint64>		FreeSlots		= ~0ull;
 	std::atomic<uint64>		Cancels			= 0;
 	uint64					PrevFreeSlots	= ~0ull;
-	FActivityList			Pending			= nullptr;
+	FActivityList			Pending;
 	FThrottler				Throttler;
 	TArray<FHostGroup>		Groups;
 	int32					FailTimeoutMs	= GIdleMs;
@@ -809,8 +926,7 @@ FTicket FEventLoop::FImpl::Send(FActivityNode* Activity)
 
 		// This puts pending requests in reverse order of when they were made
 		// but this will be undone when ReceiveWork() traverses the list
-		Activity->Next = Pending;
-		Pending = Activity;
+		Pending.Prepend(Activity);
 	}
 
 	return Slot;
@@ -858,23 +974,17 @@ void FEventLoop::FImpl::ReceiveWork()
 	PrevFreeSlots = FreeSlotsLoad;
 
 	// Fetch the pending activities from out in the wild
-	FActivityNode* Activity = nullptr;
+	FActivityList Collected;
 	{
 		FScopeLock _(&Lock);
-		Swap(Activity, Pending);
+		Swap(Collected, Pending);
 	}
 
 	// Pending is in the reverse of the order that requests were made
-	FActivityNode* Reverse = nullptr;
-	for (FActivityNode* Next; Activity != nullptr; Activity = Next)
-	{
-		Next = Activity->Next;
-		Activity->Next = Reverse;
-		Reverse = Activity;
-	}
-	Activity = Reverse;
+	Collected.Reverse();
 
 	// Group activities by their host.
+	FActivityNode* Activity = Collected.Detach();
 	for (FActivityNode* Next; Activity != nullptr; Activity = Next)
 	{
 		Next = Activity->Next;
@@ -942,7 +1052,6 @@ uint32 FEventLoop::FImpl::Tick(int32 PollTimeoutMs)
 
 	// Tick groups and then remove ones that are idle
 	FTickState TickState = {
-		.DoneList = nullptr,
 		.Cancels = CancelsLoad,
 		.RecvAllowance = RecvAllowance,
 		.PollTimeoutMs = PollTimeoutMs,
@@ -969,7 +1078,7 @@ uint32 FEventLoop::FImpl::Tick(int32 PollTimeoutMs)
 	Throttler.ReturnUnused(RecvAllowance);
 
 	uint64 ReturnedSlots = 0;
-	for (FActivityNode* Activity = TickState.DoneList; Activity != nullptr;)
+	for (FActivityNode* Activity = TickState.DoneList.Detach(); Activity != nullptr;)
 	{
 		FActivityNode* Next = Activity->Next;
 		ReturnedSlots |= (1ull << Activity->Slot);
