Fail using outcome objects everywhere.

#rnx

diff --git a/Engine/Source/Runtime/Experimental/IoStore/HttpClient/Private/Loop.inl b/Engine/Source/Runtime/Experimental/IoStore/HttpClient/Private/Loop.inl
--- a/Engine/Source/Runtime/Experimental/IoStore/HttpClient/Private/Loop.inl
+++ b/Engine/Source/Runtime/Experimental/IoStore/HttpClient/Private/Loop.inl
@@ -419,7 +419,7 @@ public:
 	FWaiter				GetWaiter() const;
 	bool				Tick(FTickState& State);
 	void				TickSend(FTickState& State, FHost& Host, FPoller& Poller);
-	void				Fail(FTickState& State, const char* Reason);
+	void				Fail(FTickState& State, FOutcome Reason);
 
 private:
 	void				Negotiate(FTickState& State);
@@ -457,7 +457,7 @@ FWaiter FPeerGroup::GetWaiter() const
 }
 
 ////////////////////////////////////////////////////////////////////////////////
-void FPeerGroup::Fail(FTickState& State, const char* Reason)
+void FPeerGroup::Fail(FTickState& State, FOutcome Reason)
 {
 	// Any send left at this point is unrecoverable. Send is likely shorter.
 	Send.Join(Recv);
@@ -466,9 +466,7 @@ void FPeerGroup::Fail(FTickState& State, const char* Reason)
 	// Failure is quite terminal and we need to abort everything
 	while (FActivityNode* Activity = Recv.Pop())
 	{
-		FOutcome Outcome = FOutcome::Error(Reason);
-		Activity->Fail(Outcome);
-
+		Activity->Fail(Reason);
 		State.DoneList.Prepend(Activity);
 	}
 
@@ -490,7 +488,7 @@ void FPeerGroup::Negotiate(FTickState& State)
 	FOutcome Outcome = Peer.Handshake();
 	if (Outcome.IsError())
 	{
-		Fail(State, Outcome.GetMessage().GetData());
+		Fail(State, Outcome);
 		return;
 	}
 
@@ -517,7 +515,7 @@ void FPeerGroup::RecvInternal(FTickState& State)
 	// Any sort of error here is unrecoverable
 	if (Outcome.IsError())
 	{
-		Fail(State, Outcome.GetMessage().GetData());
+		Fail(State, Outcome);
 		return;
 	}
 
@@ -592,7 +590,7 @@ void FPeerGroup::SendInternal(FTickState& State)
 
 	if (Outcome.IsError())
 	{
-		Fail(State, Outcome.GetMessage().GetData());
+		Fail(State, Outcome);
 		return;
 	}
 
@@ -666,7 +664,7 @@ void FPeerGroup::TickSend(FTickState& State, FHost& Host, FPoller& Poller)
 		{
 			// We failed to connect, let's bail.
 			Recv.Prepend(Pending);
-			Fail(State, Outcome.GetMessage().GetData());
+			Fail(State, Outcome);
 			return;
 		}
 
@@ -824,9 +822,9 @@ void FHostGroup::Tick(FTickState& State)
 	// Wait on the groups that are
 	if (int32 Result = Wait(State); Result < 0)
 	{
-		const char* Reason = (Result == MIN_int32)
-			? "FailTimeout hit"
-			: "poll() returned an unexpected error";
+		FOutcome Reason = (Result == MIN_int32)
+			? FOutcome::Error("FailTimeout hit")
+			: FOutcome::Error("poll() returned an unexpected error");
 
 		for (FPeerGroup& Group : PeerGroups)
 		{
