Coordinated list building to remove a reversal

#rnx

diff --git a/Engine/Source/Runtime/Experimental/IoStore/HttpClient/Private/Loop.inl b/Engine/Source/Runtime/Experimental/IoStore/HttpClient/Private/Loop.inl
--- a/Engine/Source/Runtime/Experimental/IoStore/HttpClient/Private/Loop.inl
+++ b/Engine/Source/Runtime/Experimental/IoStore/HttpClient/Private/Loop.inl
@@ -319,7 +319,6 @@ public:
 
 private:
 	FActivityList		List;
-	FActivityNode*		ListTail = nullptr;
 	uint64				ActiveSlots = 0;
 
 	UE_NONCOPYABLE(FWorkQueue);
@@ -328,25 +327,16 @@ private:
 ////////////////////////////////////////////////////////////////////////////////
 FWorkQueue::~FWorkQueue()
 {
-	check(ListTail == nullptr);
 }
 
 ////////////////////////////////////////////////////////////////////////////////
 void FWorkQueue::AddActivity(FActivityNode* Activity)
 {
-	// We use a tail pointer here to maintain order that requests were made
+	// This prepends and thus reverses order. It is assumed this this is working
+	// in conjunction with the loop, negating the reversal what happens there.
 
 	check(Activity->Next == nullptr);
-
-	if (List.IsEmpty())
-	{
-		List.Prepend(Activity);
-	}
-	else
-	{
-		ListTail->Next = Activity;
-	}
-	ListTail = Activity;
+	List.Prepend(Activity);
 
 	ActiveSlots |= (1ull << Activity->Slot);
 }
@@ -360,11 +350,6 @@ FActivityNode* FWorkQueue::PopActivity()
 		return nullptr;
 	}
 
-	if (List.IsEmpty())
-	{
-		ListTail = nullptr;
-	}
-
 	check(ActiveSlots & (1ull << Activity->Slot));
 	ActiveSlots ^= (1ull << Activity->Slot);
 
@@ -384,10 +369,11 @@ void FWorkQueue::TickCancels(FTickState& State)
 	// activity list is singular.
 
 	check(!List.IsEmpty());
-	FActivityNode* Activity = List.Detach();
-	ListTail = nullptr;
 	ActiveSlots = 0;
 
+	List.Reverse();
+	FActivityNode* Activity = List.Detach();
+
 	for (FActivityNode* Next; Activity != nullptr; Activity = Next)
 	{
 		Next = Activity->Next;
@@ -923,7 +909,7 @@ FTicket FEventLoop::FImpl::Send(FActivityNode* Activity)
 		Activity->Slot = int8(63 - FMath::CountLeadingZeros64(Slot));
 
 		// This puts pending requests in reverse order of when they were made
-		// but this will be undone when ReceiveWork() traverses the list
+		// but this will be undone when ReceiveWork() adds work to the queue
 		Pending.Prepend(Activity);
 	}
 
@@ -978,8 +964,8 @@ void FEventLoop::FImpl::ReceiveWork()
 		Swap(Collected, Pending);
 	}
 
-	// Pending is in the reverse of the order that requests were made
-	Collected.Reverse();
+	// Pending is in the reverse of the order that requests were made. Adding
+	// activities to their corresponding group will reverse this reversal.
 
 	// Group activities by their host.
 	FActivityNode* Activity = Collected.Detach();
