Deleted unused destructors

#rnx

diff --git a/Engine/Source/Runtime/Experimental/IoStore/HttpClient/Private/Loop.inl b/Engine/Source/Runtime/Experimental/IoStore/HttpClient/Private/Loop.inl
--- a/Engine/Source/Runtime/Experimental/IoStore/HttpClient/Private/Loop.inl
+++ b/Engine/Source/Runtime/Experimental/IoStore/HttpClient/Private/Loop.inl
@@ -311,7 +311,6 @@ class FWorkQueue
 {
 public:
 						FWorkQueue() = default;
-						~FWorkQueue();
 	bool				HasWork() const { return !List.IsEmpty(); }
 	void				AddActivity(FActivityNode* Activity);
 	FActivityNode*		PopActivity();
@@ -324,11 +323,6 @@ private:
 	UE_NONCOPYABLE(FWorkQueue);
 };
 
-////////////////////////////////////////////////////////////////////////////////
-FWorkQueue::~FWorkQueue()
-{
-}
-
 ////////////////////////////////////////////////////////////////////////////////
 void FWorkQueue::AddActivity(FActivityNode* Activity)
 {
@@ -400,7 +394,6 @@ class FPeerGroup
 {
 public:
 						FPeerGroup() = default;
-						~FPeerGroup();
 	void				Unwait()			{ check(bWaiting); bWaiting = false; }
 	FWaiter				GetWaiter() const;
 	bool				Tick(FTickState& State);
@@ -422,11 +415,6 @@ private:
 	UE_NONCOPYABLE(FPeerGroup);
 };
 
-////////////////////////////////////////////////////////////////////////////////
-FPeerGroup::~FPeerGroup()
-{
-}
-
 ////////////////////////////////////////////////////////////////////////////////
 FWaiter FPeerGroup::GetWaiter() const
 {
