Gather rule-of-five

#rnx

diff --git a/Engine/Source/Runtime/Experimental/IoStore/HttpClient/Private/Loop.inl b/Engine/Source/Runtime/Experimental/IoStore/HttpClient/Private/Loop.inl
--- a/Engine/Source/Runtime/Experimental/IoStore/HttpClient/Private/Loop.inl
+++ b/Engine/Source/Runtime/Experimental/IoStore/HttpClient/Private/Loop.inl
@@ -163,6 +163,8 @@ public:
 					~FActivityList();
 					FActivityList(FActivityList&& Rhs)		{ Swap(Head, Rhs.Head); }
 	void			operator = (FActivityList&& Rhs)		{ Swap(Head, Rhs.Head); }
+					FActivityList(const FActivityList&)		= delete;
+	void			operator = (const FActivityList&)		= delete;
 	int32			IsEmpty() const							{ return Head == nullptr; }
 	FActivityNode*	GetHead() const							{ return Head; }
 	FActivityNode*	Detach()								{ FActivityNode* Ret = Head; Head = nullptr; return Ret; }
@@ -180,10 +182,6 @@ public:
 private:
 	FActivityNode*	PopImpl();
 	FActivityNode*	Head = nullptr;
-
-private:
-					FActivityList(const FActivityList&) = delete;
-	FActivityList&	operator = (const FActivityList&) = delete;
 };
 
 ////////////////////////////////////////////////////////////////////////////////
