Gather related parts into a single transaction.

#rnx

diff --git a/Engine/Source/Runtime/Experimental/IoStore/HttpClient/Private/Activity.inl b/Engine/Source/Runtime/Experimental/IoStore/HttpClient/Private/Activity.inl
--- a/Engine/Source/Runtime/Experimental/IoStore/HttpClient/Private/Activity.inl
+++ b/Engine/Source/Runtime/Experimental/IoStore/HttpClient/Private/Activity.inl
@@ -5,678 +5,6 @@
 namespace UE::IoStore::HTTP
 {
 
-////////////////////////////////////////////////////////////////////////////////
-class FRequestCloak
-	: public FBuffer
-{
-public:
-	using FBuffer::FBuffer;
-
-	void			Begin(FAnsiStringView Method, FAnsiStringView Path);
-	void			AddHeader(FAnsiStringView Key, FAnsiStringView Value);
-	void			End();
-	FOutcome		TrySendRequest(FHttpPeer& Peer);
-
-protected:
-	bool			IsHeadMethod() const;
-
-private:
-	FRequestCloak&	operator << (FAnsiStringView Value);
-	uint16			HeaderLeft;
-	uint16			AlreadySent = 0;
-	int16			HeaderRight = -1;
-	int8			MethodLength = -1;
-	bool			IsHead;
-};
-
-////////////////////////////////////////////////////////////////////////////////
-bool FRequestCloak::IsHeadMethod() const
-{
-	return IsHead;
-}
-
-////////////////////////////////////////////////////////////////////////////////
-void FRequestCloak::Begin(FAnsiStringView Method, FAnsiStringView Path)
-{
-	check(MethodLength < 0);
-	AlreadySent = uint16(GetSize());
-	*this << Method << " " << Path << " HTTP/1.1" "\r\n";
-	MethodLength = int8(Method.Len());
-	HeaderLeft = uint16(GetSize());
-	IsHead = Method.Equals("HEAD", ESearchCase::IgnoreCase);
-}
-
-////////////////////////////////////////////////////////////////////////////////
-void FRequestCloak::AddHeader(FAnsiStringView Key, FAnsiStringView Value)
-{
-	check(HeaderRight < 0);
-	*this << Key << ":" << Value << "\r\n";
-}
-
-////////////////////////////////////////////////////////////////////////////////
-void FRequestCloak::End()
-{
-	*this << "\r\n";
-	HeaderRight = int16(GetSize());
-}
-
-////////////////////////////////////////////////////////////////////////////////
-FRequestCloak& FRequestCloak::operator << (FAnsiStringView Value)
-{
-	uint32 Length = uint32(Value.Len());
-	FBuffer::FMutableSection Section = GetMutableFree(Length);
-	::memcpy(Section.Data, Value.GetData(), Length);
-	AdvanceUsed(Length);
-	return *this;
-}
-
-////////////////////////////////////////////////////////////////////////////////
-FOutcome FRequestCloak::TrySendRequest(FHttpPeer& Peer)
-{
-	const char* SendData = GetData();
-	int32 SendSize = GetSize();
-
-	SendData += AlreadySent;
-	SendSize -= AlreadySent;
-	check(SendSize > 0);
-
-	FOutcome Outcome = Peer.Send(SendData, SendSize);
-	if (!Outcome.IsOk())
-	{
-		return Outcome;
-	}
-
-	int32 Result = Outcome.GetResult();
-	AlreadySent += uint16(Result);
-	if (AlreadySent == GetSize())
-	{
-		return FOutcome::Ok();
-	}
-
-	check(AlreadySent < GetSize());
-	return FOutcome::Waiting();
-}
-
-
-
-////////////////////////////////////////////////////////////////////////////////
-class FResponseCloak
-	: public FRequestCloak
-{
-public:
-	using FRequestCloak::FRequestCloak;
-
-	FOutcome		TryRecvResponse(FHttpPeer& Peer);
-	bool			IsKeepAlive() const;
-	uint32			GetStatusCode() const;
-	FAnsiStringView	GetStatusMessage() const;
-	int64			GetContentLength() const;
-	bool			IsChunked() const;
-	bool			HasBody() const;
-	void			ReadHeaders(FResponse::FHeaderSink Sink) const;
-
-protected:
-	const char*		GetMessageRight() const;
-
-private:
-	FOutcome		Parse();
-	int64			ContentLength = -1;
-	FMessageOffsets	Offsets = {};
-	int16			MessageLeft = -1;
-	int16			MessageRight = -1;
-	uint16			StatusCode = 0;
-	uint8			bKeepAlive : 1;
-	uint8			bChunked : 1;
-	uint8			bHasBody : 1;
-	uint8			_Unused : 6;
-};
-
-////////////////////////////////////////////////////////////////////////////////
-bool FResponseCloak::IsKeepAlive() const
-{
-	return (MessageRight > 0) ? bKeepAlive : true;
-}
-
-////////////////////////////////////////////////////////////////////////////////
-uint32 FResponseCloak::GetStatusCode() const
-{
-	check(StatusCode >= 100);
-	return StatusCode;
-}
-
-////////////////////////////////////////////////////////////////////////////////
-FAnsiStringView	FResponseCloak::GetStatusMessage() const
-{
-	check(Offsets.Message > 0);
-	const char* Cursor = GetData() + MessageLeft;
-	return FAnsiStringView(Cursor + Offsets.Message, Offsets.Headers - Offsets.Message);
-}
-
-////////////////////////////////////////////////////////////////////////////////
-int64 FResponseCloak::GetContentLength() const
-{
-	check(Offsets.Headers > 0);
-	return ContentLength;
-}
-
-////////////////////////////////////////////////////////////////////////////////
-bool FResponseCloak::IsChunked() const
-{
-	check(MessageRight > 0);
-	return bChunked;
-}
-
-////////////////////////////////////////////////////////////////////////////////
-bool FResponseCloak::HasBody() const
-{
-	return bHasBody;
-}
-
-////////////////////////////////////////////////////////////////////////////////
-void FResponseCloak::ReadHeaders(FResponse::FHeaderSink Sink) const
-{
-	uint32 Offset = MessageLeft + Offsets.Headers;
-	uint32 Length = MessageRight - Offset - 2; // "-2" trims off '\r\n' that signals end of headers
-	const char* Cursor = GetData() + Offset;
-	FAnsiStringView Headers(Cursor, Length);
-	UE::IoStore::HTTP::EnumerateHeaders(Headers, Sink);
-}
-
-////////////////////////////////////////////////////////////////////////////////
-const char* FResponseCloak::GetMessageRight() const
-{
-	return GetData() + MessageRight;
-}
-
-////////////////////////////////////////////////////////////////////////////////
-FOutcome FResponseCloak::TryRecvResponse(FHttpPeer& Peer)
-{
-	auto FindMessageTerminal = [] (const char* Cursor, int32 Length)
-	{
-		for (int32 i = 4; i <= Length; ++i)
-		{
-			uint32 Candidate;
-			::memcpy(&Candidate, Cursor + i - 4, sizeof(Candidate));
-			if (Candidate == 0x0a0d0a0d)
-			{
-				return i;
-			}
-
-			i += (Cursor[i - 1] > 0x0d) ? 3 : 0;
-		}
-
-		return -1;
-	};
-
-	if (MessageLeft < 0)
-	{
-		MessageLeft = int16(GetSize());
-		bKeepAlive = true;
-		bChunked = false;
-	}
-
-	while (true)
-	{
-		static const uint32 PageSize = 256;
-		static const uint32 MaxHeaderSize = 8 << 10;
-
-		auto [Dest, DestSize] = GetMutableFree(0, PageSize);
-		FOutcome Outcome = Peer.Recv(Dest, DestSize);
-		if (!Outcome.IsOk())
-		{
-			return Outcome;
-		}
-
-		int32 Result = Outcome.GetResult();
-		AdvanceUsed(Result);
-
-		// Rewind a little to cover cases where the terminal is fragmented across
-		// recv() calls
-		uint32 DestBias = 0;
-		if (Dest - 3 >= GetData() + MessageLeft)
-		{
-			Dest -= (DestBias = 3);
-		}
-
-		int32 MessageEnd = FindMessageTerminal(Dest, Result + DestBias);
-		if (MessageEnd < 0)
-		{
-			if (GetSize() - MessageLeft > MaxHeaderSize)
-			{
-				return FOutcome::Error("Headers have grown larger than expected");
-			}
-
-			continue;
-		}
-
-		MessageRight = int16(ptrdiff_t(Dest + MessageEnd - GetData()));
-		return Parse();
-	}
-}
-
-////////////////////////////////////////////////////////////////////////////////
-FOutcome FResponseCloak::Parse()
-{
-	const char* Cursor = GetData() + MessageLeft;
-	FAnsiStringView MessageView(Cursor, MessageRight - MessageLeft);
-
-	if (ParseMessage(MessageView, Offsets) < 0)
-	{
-		return FOutcome::Error("Failed to parse message status");
-	}
-
-	FAnsiStringView StatusCodeView(
-		Cursor + Offsets.StatusCode,
-		Offsets.Message - Offsets.StatusCode
-	);
-	StatusCode = uint16(CrudeToInt(StatusCodeView));
-	if (uint32(StatusCode - 100) > 899)
-	{
-		return FOutcome::Error("Invalid status code", StatusCode);
-	}
-
-	static const uint32 Flag_Length		= 1 << 0;
-	static const uint32 Flag_XferEnc	= 1 << 1;
-	static const uint32 Flag_Connection = 1 << 2;
-	static const uint32 Flag_All		= 0x07;
-	uint32 Flags = 0;
-
-	ReadHeaders(
-		[this, &Flags] (FAnsiStringView Name, FAnsiStringView Value) mutable
-		{
-			// todo; may need smarter value handling; ;/, separated options & key-value pairs (ex. in rfc2068)
-
-			if (Name.Equals("Content-Length", ESearchCase::IgnoreCase))
-			{
-				ContentLength = int32(CrudeToInt(Value));
-				Flags |= Flag_Length;
-			}
-
-			else if (Name.Equals("Transfer-Encoding", ESearchCase::IgnoreCase))
-			{
-				bChunked = Value.Equals("chunked", ESearchCase::IgnoreCase);
-				Flags |= Flag_XferEnc;
-			}
-
-			else if (Name.Equals("Connection", ESearchCase::IgnoreCase))
-			{
-				bKeepAlive = !Value.Equals("close");
-				Flags |= Flag_Connection;
-			}
-
-			return Flags != Flag_All;
-		}
-	);
-
-	bHasBody = (StatusCode != 304);
-	bHasBody &= !IsHeadMethod();
-	bHasBody &= bChunked | (ContentLength > 0);
-
-	if (Flags & Flag_Length)
-	{
-		if ((StatusCode <= 100) | (StatusCode == 204) | (Flags & Flag_XferEnc))
-		{
-			return FOutcome::Error("Unexpected Content-Length field", StatusCode);
-		}
-
-		if (ContentLength < 0)
-		{
-			return FOutcome::Error("Invalid Content-Length field", uint32(ContentLength));
-		}
-
-		if ((GetSize() - MessageRight) > ContentLength)
-		{
-			return FOutcome::Error("More data received that expected");
-		}
-	}
-	else if (Flags & Flag_XferEnc)
-	{
-		if (!bChunked)
-		{
-			return FOutcome::Error("Unsupported Transfer-Encoding");
-		}
-		ContentLength = -1;
-	}
-	else if (bHasBody)
-	{
-		return FOutcome::Error("Missing Content-Length field");
-	}
-	else if (GetSize() > uint32(MessageRight))
-	{
-		return FOutcome::Error("Received content when none was expected");
-	}
-
-	return FOutcome::Ok();
-}
-
-
-
-////////////////////////////////////////////////////////////////////////////////
-class FBodyCloak
-	: public FResponseCloak
-{
-public:
-	using FResponseCloak::FResponseCloak;
-
-	int64				GetRemaining() const;
-	FOutcome			TryRecv(FMutableMemoryView Dest, FHttpPeer& Peer);
-
-private:
-	enum class EState : uint8 { Init, Recv, Prologue, Chunk, Epilogue, Done };
-
-	FOutcome			Gather(FMutableMemoryView Dest, FHttpPeer& Peer);
-	FOutcome			Init(FMutableMemoryView Dest, FHttpPeer& Peer);
-	FOutcome			Recv(FMutableMemoryView Dest, FHttpPeer& Peer);
-	FOutcome			Prologue(FMutableMemoryView Dest, FHttpPeer& Peer);
-	FOutcome			Chunk(FMutableMemoryView Dest, FHttpPeer& Peer);
-	FOutcome			Epilogue(FMutableMemoryView Dest, FHttpPeer& Peer);
-	FMemoryView			Overspill;
-	FMutableMemoryView	Scratch;
-	int64				Remaining = 0;
-	EState				State = EState::Init;
-
-	enum
-	{
-		HeaderBufSize	= 32,
-		CrLfLength		= 2,
-		EndOfXfer		= -1,
-	};
-};
-
-////////////////////////////////////////////////////////////////////////////////
-int64 FBodyCloak::GetRemaining() const
-{
-	return (State == EState::Recv) ? Remaining : -1;
-}
-
-////////////////////////////////////////////////////////////////////////////////
-FOutcome FBodyCloak::TryRecv(FMutableMemoryView Dest, FHttpPeer& Peer)
-{
-	FOutcome Outcome = FOutcome::None();
-
-	switch (State)
-	{
-	case EState::Init:	Outcome = Init(Dest, Peer);		break;
-	case EState::Recv:	Outcome = Recv(Dest, Peer);		break;
-	default:			Outcome = Gather(Dest, Peer);	break;
-	case EState::Done:	Outcome = FOutcome::Error("Transaction is complete"); break;
-	}
-
-	if (Outcome.IsOk())
-	{
-		State = EState::Done;
-	}
-
-	return Outcome;
-}
-
-////////////////////////////////////////////////////////////////////////////////
-FOutcome FBodyCloak::Gather(FMutableMemoryView Dest, FHttpPeer& Peer)
-{
-	FOutcome Outcome = FOutcome::None();
-
-	bool bDone = false;
-	int32 Size = 0;
-	do
-	{
-		switch (State)
-		{
-		case EState::Prologue:	Outcome = Prologue(Dest, Peer); break;
-		case EState::Chunk:		Outcome = Chunk(Dest, Peer);	break;
-		case EState::Epilogue:	Outcome = Epilogue(Dest, Peer);	break;
-		default:				Outcome = FOutcome::Error("Unexpected try-chunked state"); break;
-		}
-
-		if (Outcome.IsError())
-		{
-			return Outcome;
-		}
-
-		bDone = Outcome.IsOk();
-		int32 Result = Outcome.GetResult();
-		if (Result == 0)
-		{
-			break;
-		}
-
-		Size += Result;
-		Dest = Dest.Mid(Result);
-	}
-	while (!Dest.IsEmpty());
-
-	return bDone ? FOutcome::Ok(Size) : FOutcome::Waiting(Size);
-}
-
-////////////////////////////////////////////////////////////////////////////////
-FOutcome FBodyCloak::Init(FMutableMemoryView Dest, FHttpPeer& Peer)
-{
-	const char* DataEnd = GetData() + GetSize();
-	const char* MessageEnd = GetMessageRight();
-	uint32 AlreadyReceived = uint32(ptrdiff_t(DataEnd - MessageEnd));
-	if (AlreadyReceived)
-	{
-		Overspill = { (const uint8*)MessageEnd, AlreadyReceived };
-	}
-
-	if (AlreadyReceived < HeaderBufSize)
-	{
-		GetMutableFree(HeaderBufSize, HeaderBufSize);
-		MessageEnd = GetMessageRight();
-	}
-
-	Scratch = { (uint8*)MessageEnd, HeaderBufSize };
-	
-	if (IsChunked())
-	{
-		check(Remaining == 0);
-		State = EState::Prologue;
-	}
-	else
-	{
-		Remaining = GetContentLength();
-		check(Remaining > 0);
-		State = EState::Recv;
-	}
-
-	return TryRecv(Dest, Peer);
-}
-
-////////////////////////////////////////////////////////////////////////////////
-FOutcome FBodyCloak::Recv(FMutableMemoryView Dest, FHttpPeer& Peer)
-{
-	int64 OverspillSize = Overspill.GetSize();
-	if (OverspillSize != 0)
-	{
-		int64 Size = FMath::Min<int64>(Dest.GetSize(), OverspillSize);
-		::memcpy(Dest.GetData(), Overspill.GetData(), Size);
-		Dest = Dest.Mid(Size);
-		Overspill = Overspill.Mid(Size);
-		Remaining -= Size;
-		if (Remaining == 0 || Dest.IsEmpty())
-		{
-			int32 Result = int32(Size);
-			return Remaining ? FOutcome::Waiting(Result) : FOutcome::Ok(Result);
-		}
-	}
-
-	auto* Cursor = (char*)(Dest.GetData());
-	uint32 Size = uint32(FMath::Min<int64>(Dest.GetSize(), Remaining));
-	FOutcome Outcome = Peer.Recv(Cursor, Size);
-	if (Outcome.IsError())
-	{
-		return Outcome;
-	}
-
-	int32 Result = Outcome.GetResult();
-	Remaining -= Result;
-	if (Remaining < 0)
-	{
-		return FOutcome::Error("Unexpectedly received too much", int32(Remaining));
-	}
-
-	Result += int32(OverspillSize);
-	return Remaining ? FOutcome::Waiting(Result) : FOutcome::Ok(Result);
-}
-
-////////////////////////////////////////////////////////////////////////////////
-FOutcome FBodyCloak::Prologue(FMutableMemoryView Dest, FHttpPeer& Peer)
-{
-	auto FindHeader = [] (const char* Cursor, int32 Size)
-	{
-		// Base-16 size
-		int32 HexDigLen = 0;
-		for (; HexDigLen < Size; ++HexDigLen)
-		{
-			uint32 c = Cursor[HexDigLen];
-			if ((c - '0' >= 10) & (c - 'a' >= 6) & (c - 'A' >= 6))
-			{
-				break;
-			}
-			continue;
-		}
-
-		int32 Delta = Size - HexDigLen;
-
-		if (HexDigLen == 0)					return FOutcome::Error("Invalid chunk header");
-		if (Delta < 1)						return FOutcome::Waiting();
-		if (Cursor[HexDigLen] != '\r')		return FOutcome::Error("Chunk extensions are not supported (ERREXT)");
-		if (Delta < 2)						return FOutcome::Waiting();
-		if (Cursor[HexDigLen + 1] != '\n')	return FOutcome::Error("Invalid chunk CRLF");
-
-		return FOutcome::Ok(HexDigLen + 2);
-	};
-
-	const char* Header = nullptr;
-	FOutcome Outcome = FOutcome::None();
-	int32 HeaderLength = 0;
-
-	uint32 OverspillSize = uint32(Overspill.GetSize());
-	if (OverspillSize != 0)
-	{
-		// We try all the overspill first - we don't know how much we're expecting
-		// and overspill could contain both header and all data (thus further recvs
-		// could falsely hang us up).
-		Header = (char*)Overspill.GetData();
-
-		Outcome = FindHeader(Header, OverspillSize);
-		if (Outcome.IsError())
-		{
-			return Outcome;
-		}
-
-		HeaderLength = Outcome.GetResult();
-		Overspill = Overspill.Mid(HeaderLength);
-	}
-
-	if (HeaderLength == 0)
-	{
-		char* Cursor = (char*)Scratch.GetData();
-
-		if (OverspillSize != 0)
-		{
-			check(OverspillSize <= HeaderBufSize);
-			std::memmove(Cursor, Overspill.GetData(), OverspillSize);
-			Cursor += OverspillSize;
-		}
-
-		Outcome = Peer.Recv(Cursor, HeaderBufSize);
-		if (!Outcome.IsOk())
-		{
-			return Outcome;
-		}
-
-		Cursor -= OverspillSize;
-		OverspillSize += Outcome.GetResult();
-		Overspill = { (uint8*)Cursor, OverspillSize };
-
-		Outcome = FindHeader(Cursor, OverspillSize);
-		if (!Outcome.IsOk())
-		{
-			return Outcome;
-		}
-
-		Header = Cursor;
-		HeaderLength = Outcome.GetResult();
-		Overspill = Overspill.Mid(HeaderLength);
-	}
-
-	// Read chunk size
-	Header += Remaining;
-	int64 ChunkSize = CrudeToInt<16>(FAnsiStringView(Header, HeaderLength));
-	if (ChunkSize < 0)			return FOutcome::Error("Unparsable chunk size");
-	if (ChunkSize > 16 << 20)	return FOutcome::Error("Unacceptable chunk size");
-
-	if (ChunkSize == 0)
-	{
-		Remaining = EndOfXfer;
-		::memset(Scratch.GetData(), 0, CrLfLength);
-		State = EState::Epilogue;
-		return Epilogue(Dest, Peer);
-	}
-
-	Remaining = ChunkSize;
-	State = EState::Chunk;
-	return Chunk(Dest, Peer);
-}
-
-////////////////////////////////////////////////////////////////////////////////
-FOutcome FBodyCloak::Chunk(FMutableMemoryView Dest, FHttpPeer& Peer)
-{
-	check(Remaining != 0);
-
-	Dest = Dest.Left(Remaining);
-	FOutcome Outcome = Recv(Dest, Peer);
-	if (!Outcome.IsOk())
-	{
-		return Outcome;
-	}
-
-	check(Remaining == 0);
-	::memset(Scratch.GetData(), 0, CrLfLength);
-	State = EState::Epilogue;
-
-	int32 Result = Outcome.GetResult();
-	return FOutcome::Waiting(Result);
-}
-
-////////////////////////////////////////////////////////////////////////////////
-FOutcome FBodyCloak::Epilogue(FMutableMemoryView Dest, FHttpPeer& Peer)
-{
-	auto* Cursor = (uint8*)(Scratch.GetData());
-
-	FMutableMemoryView CrLfBuffer(Cursor, CrLfLength);
-	if (Cursor[0] != 0)
-	{
-		CrLfBuffer = { Cursor + 1, 1 };
-	}
-
-	check(Remaining <= 0);
-
-	int64 RemainingCache = Remaining;
-	Remaining = CrLfBuffer.GetSize();
-	FOutcome Outcome = Recv(CrLfBuffer, Peer);
-	if (!Outcome.IsOk())
-	{
-		Remaining = RemainingCache;
-		return Outcome;
-	}
-
-	if (Cursor[0] != '\r' || Cursor[1] != '\n')
-	{
-		return FOutcome::Error("Trailing headers are not supported (ERRTRAIL)");
-	}
-
-	if (RemainingCache > EndOfXfer)
-	{
-		State = EState::Prologue;
-		return Prologue(Dest, Peer);
-	}
-
-	return FOutcome::Ok();
-}
-
-
-
 #if IAS_HTTP_WITH_PERF
 
 ////////////////////////////////////////////////////////////////////////////////
@@ -767,7 +95,7 @@ struct FActivity
 	void				Done();
 	void				Fail(const FOutcome& Outcome);
 	EStage				GetStage() const;
-	const FBodyCloak*	GetTransaction() const;
+	const FTransaction*	GetTransaction() const;
 	FAnsiStringView		GetMethod() const;
 	FAnsiStringView		GetPath() const;
 	FHost*				GetHost() const;
@@ -823,7 +151,7 @@ private:
 	};
 	UPTRINT				SinkParam;
 	FTicketSink			Sink;
-	FBodyCloak			Transaction;
+	FTransaction		Transaction;
 
 	struct FHeaderRecord
 	{
@@ -979,7 +307,7 @@ FActivity::EStage FActivity::GetStage() const
 }
 
 ////////////////////////////////////////////////////////////////////////////////
-const FBodyCloak* FActivity::GetTransaction() const
+const FTransaction* FActivity::GetTransaction() const
 {
 	check(State >= EState::Send && State <= EState::Completed);
 	return &Transaction;
diff --git a/Engine/Source/Runtime/Experimental/IoStore/HttpClient/Private/Http.cpp b/Engine/Source/Runtime/Experimental/IoStore/HttpClient/Private/Http.cpp
--- a/Engine/Source/Runtime/Experimental/IoStore/HttpClient/Private/Http.cpp
+++ b/Engine/Source/Runtime/Experimental/IoStore/HttpClient/Private/Http.cpp
@@ -24,6 +24,7 @@
 #include "Misc.inl"
 #include "Socket.inl"
 #include "Peer.inl"
+#include "TransactionOne.inl"
 #include "Socks.inl"
 #include "ConnectionPool.inl"
 #include "Activity.inl"
diff --git a/Engine/Source/Runtime/Experimental/IoStore/HttpClient/Private/Misc.inl b/Engine/Source/Runtime/Experimental/IoStore/HttpClient/Private/Misc.inl
--- a/Engine/Source/Runtime/Experimental/IoStore/HttpClient/Private/Misc.inl
+++ b/Engine/Source/Runtime/Experimental/IoStore/HttpClient/Private/Misc.inl
@@ -134,61 +134,6 @@ FOutcome FOutcome::Error(const char* Message, int32 Code)
 
 
 
-////////////////////////////////////////////////////////////////////////////////
-template <typename LambdaType>
-static void EnumerateHeaders(FAnsiStringView Headers, LambdaType&& Lambda)
-{
-	// NB. here we are assuming that we will be dealing with servers that will
-	// not be returning headers with "obsolete line folding".
-
-	auto IsOws = [] (int32 c) { return (c == ' ') | (c == '\t'); };
-
-	const char* Cursor = Headers.GetData();
-	const char* End = Cursor + Headers.Len();
-	do
-	{
-		int32 ColonIndex = 0;
-		for (; Cursor + ColonIndex < End; ++ColonIndex)
-		{
-			if (Cursor[ColonIndex] == ':')
-			{
-				break;
-			}
-		}
-
-		Cursor += ColonIndex;
-
-		const char* Right = Cursor + 1;
-		while (Right < End)
-		{
-			if (Right[0] != '\r' || Right[1] != '\n')
-			{
-				Right += 1 + (Right[1] != '\r');
-				continue;
-			}
-
-			FAnsiStringView Name(Cursor - ColonIndex, ColonIndex);
-
-			const char* Left = Cursor + 1;
-			for (; IsOws(Left[0]); ++Left);
-
-			Cursor = Right;
-			for (; Cursor > Left + 1 && IsOws(Cursor[-1]); --Cursor);
-
-			FAnsiStringView Value (Left, int32(ptrdiff_t(Cursor - Left)));
-
-			if (!Lambda(Name, Value))
-			{
-				Right = End;
-			}
-
-			break;
-		}
-
-		Cursor = Right + 2;
-	} while (Cursor < End);
-}
-
 ////////////////////////////////////////////////////////////////////////////////
 template <uint32 Base=10>
 static int64 CrudeToInt(FAnsiStringView View)
@@ -226,71 +171,6 @@ static int64 CrudeToInt(FAnsiStringView View)
 
 template <uint32 Base=10> static int64 CrudeToInt(const char*) = delete;
 
-////////////////////////////////////////////////////////////////////////////////
-struct FMessageOffsets
-{
-	uint8	StatusCode;
-	uint8	Message;
-	uint16	Headers;
-};
-
-static int32 ParseMessage(FAnsiStringView Message, FMessageOffsets& Out)
-{
-	const FAnsiStringView Protocol("HTTP/1.1 ");
-
-	// Check there's enough data
-	if (Message.Len() < Protocol.Len() + 1) // "+1" accounts for at least one digit
-	{
-		return -1;
-	}
-
-	const char* Cursor = Message.GetData();
-
-	// Check for the expected protocol
-	if (FAnsiStringView(Cursor, 9) != Protocol)
-	{
-		return -1;
-	}
-	int32 i = Protocol.Len();
-
-	// Trim left and tightly reject anything adventurous
-	for (int n = 32; i < n && Cursor[i] == ' '; ++i);
-	Out.StatusCode = uint8(i);
-
-	// At least one status line digit. (Note to self; expect exactly three)
-	for (int n = 32; i < n && uint32(Cursor[i] - 0x30) <= 9; ++i);
-	if (uint32(i - Out.StatusCode - 1) > 32)
-	{
-		return -1;
-	}
-
-	// Trim left
-	for (int n = 32; i < n && Cursor[i] == ' '; ++i);
-	Out.Message = uint8(i);
-
-	// Extra conservative length allowance
-	if (i > 32)
-	{
-		return -1;
-	}
-
-	// Find \r\n
-	for (; Cursor[i] != '\r'; ++i)
-	{
-		if (i >= 2048)
-		{
-			return -1;
-		}
-	}
-	if (Cursor[i + 1] != '\n')
-	{
-		return -1;
-	}
-	Out.Headers = uint16(i + 2);
-
-	return 1;
-}
-
 ////////////////////////////////////////////////////////////////////////////////
 struct FUrlOffsets
 {
diff --git a/Engine/Source/Runtime/Experimental/IoStore/HttpClient/Private/TransactionOne.inl b/Engine/Source/Runtime/Experimental/IoStore/HttpClient/Private/TransactionOne.inl
new file mode 100644
--- /dev/null
+++ b/Engine/Source/Runtime/Experimental/IoStore/HttpClient/Private/TransactionOne.inl
@@ -0,0 +1,819 @@
+// Copyright Epic Games, Inc. All Rights Reserved.
+
+#pragma once
+
+namespace UE::IoStore::HTTP
+{
+
+////////////////////////////////////////////////////////////////////////////////
+struct FMessageOffsets
+{
+	uint8	StatusCode;
+	uint8	Message;
+	uint16	Headers;
+};
+
+static int32 ParseMessage(FAnsiStringView Message, FMessageOffsets& Out)
+{
+	const FAnsiStringView Protocol("HTTP/1.1 ");
+
+	// Check there's enough data
+	if (Message.Len() < Protocol.Len() + 1) // "+1" accounts for at least one digit
+	{
+		return -1;
+	}
+
+	const char* Cursor = Message.GetData();
+
+	// Check for the expected protocol
+	if (FAnsiStringView(Cursor, 9) != Protocol)
+	{
+		return -1;
+	}
+	int32 i = Protocol.Len();
+
+	// Trim left and tightly reject anything adventurous
+	for (int n = 32; i < n && Cursor[i] == ' '; ++i);
+	Out.StatusCode = uint8(i);
+
+	// At least one status line digit. (Note to self; expect exactly three)
+	for (int n = 32; i < n && uint32(Cursor[i] - 0x30) <= 9; ++i);
+	if (uint32(i - Out.StatusCode - 1) > 32)
+	{
+		return -1;
+	}
+
+	// Trim left
+	for (int n = 32; i < n && Cursor[i] == ' '; ++i);
+	Out.Message = uint8(i);
+
+	// Extra conservative length allowance
+	if (i > 32)
+	{
+		return -1;
+	}
+
+	// Find \r\n
+	for (; Cursor[i] != '\r'; ++i)
+	{
+		if (i >= 2048)
+		{
+			return -1;
+		}
+	}
+	if (Cursor[i + 1] != '\n')
+	{
+		return -1;
+	}
+	Out.Headers = uint16(i + 2);
+
+	return 1;
+}
+
+
+
+////////////////////////////////////////////////////////////////////////////////
+template <typename LambdaType>
+static void EnumerateHeaders(FAnsiStringView Headers, LambdaType&& Lambda)
+{
+	// NB. here we are assuming that we will be dealing with servers that will
+	// not be returning headers with "obsolete line folding".
+
+	auto IsOws = [] (int32 c) { return (c == ' ') | (c == '\t'); };
+
+	const char* Cursor = Headers.GetData();
+	const char* End = Cursor + Headers.Len();
+	do
+	{
+		int32 ColonIndex = 0;
+		for (; Cursor + ColonIndex < End; ++ColonIndex)
+		{
+			if (Cursor[ColonIndex] == ':')
+			{
+				break;
+			}
+		}
+
+		Cursor += ColonIndex;
+
+		const char* Right = Cursor + 1;
+		while (Right < End)
+		{
+			if (Right[0] != '\r' || Right[1] != '\n')
+			{
+				Right += 1 + (Right[1] != '\r');
+				continue;
+			}
+
+			FAnsiStringView Name(Cursor - ColonIndex, ColonIndex);
+
+			const char* Left = Cursor + 1;
+			for (; IsOws(Left[0]); ++Left);
+
+			Cursor = Right;
+			for (; Cursor > Left + 1 && IsOws(Cursor[-1]); --Cursor);
+
+			FAnsiStringView Value (Left, int32(ptrdiff_t(Cursor - Left)));
+
+			if (!Lambda(Name, Value))
+			{
+				Right = End;
+			}
+
+			break;
+		}
+
+		Cursor = Right + 2;
+	} while (Cursor < End);
+}
+
+
+
+namespace DetailOne
+{
+
+////////////////////////////////////////////////////////////////////////////////
+class FRequest
+	: public FBuffer
+{
+public:
+	using FBuffer::FBuffer;
+
+	void			Begin(FAnsiStringView Method, FAnsiStringView Path);
+	void			AddHeader(FAnsiStringView Key, FAnsiStringView Value);
+	void			End();
+	FOutcome		TrySendRequest(FHttpPeer& Peer);
+
+protected:
+	bool			IsHeadMethod() const;
+
+private:
+	FRequest&		operator << (FAnsiStringView Value);
+	uint16			HeaderLeft;
+	uint16			AlreadySent = 0;
+	int16			HeaderRight = -1;
+	int8			MethodLength = -1;
+	bool			IsHead;
+};
+
+////////////////////////////////////////////////////////////////////////////////
+bool FRequest::IsHeadMethod() const
+{
+	return IsHead;
+}
+
+////////////////////////////////////////////////////////////////////////////////
+void FRequest::Begin(FAnsiStringView Method, FAnsiStringView Path)
+{
+	check(MethodLength < 0);
+	AlreadySent = uint16(GetSize());
+	*this << Method << " " << Path << " HTTP/1.1" "\r\n";
+	MethodLength = int8(Method.Len());
+	HeaderLeft = uint16(GetSize());
+	IsHead = Method.Equals("HEAD", ESearchCase::IgnoreCase);
+}
+
+////////////////////////////////////////////////////////////////////////////////
+void FRequest::AddHeader(FAnsiStringView Key, FAnsiStringView Value)
+{
+	check(HeaderRight < 0);
+	*this << Key << ":" << Value << "\r\n";
+}
+
+////////////////////////////////////////////////////////////////////////////////
+void FRequest::End()
+{
+	*this << "\r\n";
+	HeaderRight = int16(GetSize());
+}
+
+////////////////////////////////////////////////////////////////////////////////
+FRequest& FRequest::operator << (FAnsiStringView Value)
+{
+	uint32 Length = uint32(Value.Len());
+	FBuffer::FMutableSection Section = GetMutableFree(Length);
+	::memcpy(Section.Data, Value.GetData(), Length);
+	AdvanceUsed(Length);
+	return *this;
+}
+
+////////////////////////////////////////////////////////////////////////////////
+FOutcome FRequest::TrySendRequest(FHttpPeer& Peer)
+{
+	const char* SendData = GetData();
+	int32 SendSize = GetSize();
+
+	SendData += AlreadySent;
+	SendSize -= AlreadySent;
+	check(SendSize > 0);
+
+	FOutcome Outcome = Peer.Send(SendData, SendSize);
+	if (!Outcome.IsOk())
+	{
+		return Outcome;
+	}
+
+	int32 Result = Outcome.GetResult();
+	AlreadySent += uint16(Result);
+	if (AlreadySent == GetSize())
+	{
+		return FOutcome::Ok();
+	}
+
+	check(AlreadySent < GetSize());
+	return FOutcome::Waiting();
+}
+
+
+
+////////////////////////////////////////////////////////////////////////////////
+class FStatusHeaders
+	: public FRequest
+{
+public:
+	using FRequest::FRequest;
+
+	FOutcome		TryRecvResponse(FHttpPeer& Peer);
+	bool			IsKeepAlive() const;
+	uint32			GetStatusCode() const;
+	FAnsiStringView	GetStatusMessage() const;
+	int64			GetContentLength() const;
+	bool			IsChunked() const;
+	bool			HasBody() const;
+	void			ReadHeaders(FResponse::FHeaderSink Sink) const;
+
+protected:
+	const char*		GetMessageRight() const;
+
+private:
+	FOutcome		Parse();
+	int64			ContentLength = -1;
+	FMessageOffsets	Offsets = {};
+	int16			MessageLeft = -1;
+	int16			MessageRight = -1;
+	uint16			StatusCode = 0;
+	uint8			bKeepAlive : 1;
+	uint8			bChunked : 1;
+	uint8			bHasBody : 1;
+	uint8			_Unused : 6;
+};
+
+////////////////////////////////////////////////////////////////////////////////
+bool FStatusHeaders::IsKeepAlive() const
+{
+	return (MessageRight > 0) ? bKeepAlive : true;
+}
+
+////////////////////////////////////////////////////////////////////////////////
+uint32 FStatusHeaders::GetStatusCode() const
+{
+	check(StatusCode >= 100);
+	return StatusCode;
+}
+
+////////////////////////////////////////////////////////////////////////////////
+FAnsiStringView	FStatusHeaders::GetStatusMessage() const
+{
+	check(Offsets.Message > 0);
+	const char* Cursor = GetData() + MessageLeft;
+	return FAnsiStringView(Cursor + Offsets.Message, Offsets.Headers - Offsets.Message);
+}
+
+////////////////////////////////////////////////////////////////////////////////
+int64 FStatusHeaders::GetContentLength() const
+{
+	check(Offsets.Headers > 0);
+	return ContentLength;
+}
+
+////////////////////////////////////////////////////////////////////////////////
+bool FStatusHeaders::IsChunked() const
+{
+	check(MessageRight > 0);
+	return bChunked;
+}
+
+////////////////////////////////////////////////////////////////////////////////
+bool FStatusHeaders::HasBody() const
+{
+	return bHasBody;
+}
+
+////////////////////////////////////////////////////////////////////////////////
+void FStatusHeaders::ReadHeaders(FResponse::FHeaderSink Sink) const
+{
+	uint32 Offset = MessageLeft + Offsets.Headers;
+	uint32 Length = MessageRight - Offset - 2; // "-2" trims off '\r\n' that signals end of headers
+	const char* Cursor = GetData() + Offset;
+	FAnsiStringView Headers(Cursor, Length);
+	UE::IoStore::HTTP::EnumerateHeaders(Headers, Sink);
+}
+
+////////////////////////////////////////////////////////////////////////////////
+const char* FStatusHeaders::GetMessageRight() const
+{
+	return GetData() + MessageRight;
+}
+
+////////////////////////////////////////////////////////////////////////////////
+FOutcome FStatusHeaders::TryRecvResponse(FHttpPeer& Peer)
+{
+	auto FindMessageTerminal = [] (const char* Cursor, int32 Length)
+	{
+		for (int32 i = 4; i <= Length; ++i)
+		{
+			uint32 Candidate;
+			::memcpy(&Candidate, Cursor + i - 4, sizeof(Candidate));
+			if (Candidate == 0x0a0d0a0d)
+			{
+				return i;
+			}
+
+			i += (Cursor[i - 1] > 0x0d) ? 3 : 0;
+		}
+
+		return -1;
+	};
+
+	if (MessageLeft < 0)
+	{
+		MessageLeft = int16(GetSize());
+		bKeepAlive = true;
+		bChunked = false;
+	}
+
+	while (true)
+	{
+		static const uint32 PageSize = 256;
+		static const uint32 MaxHeaderSize = 8 << 10;
+
+		auto [Dest, DestSize] = GetMutableFree(0, PageSize);
+		FOutcome Outcome = Peer.Recv(Dest, DestSize);
+		if (!Outcome.IsOk())
+		{
+			return Outcome;
+		}
+
+		int32 Result = Outcome.GetResult();
+		AdvanceUsed(Result);
+
+		// Rewind a little to cover cases where the terminal is fragmented across
+		// recv() calls
+		uint32 DestBias = 0;
+		if (Dest - 3 >= GetData() + MessageLeft)
+		{
+			Dest -= (DestBias = 3);
+		}
+
+		int32 MessageEnd = FindMessageTerminal(Dest, Result + DestBias);
+		if (MessageEnd < 0)
+		{
+			if (GetSize() - MessageLeft > MaxHeaderSize)
+			{
+				return FOutcome::Error("Headers have grown larger than expected");
+			}
+
+			continue;
+		}
+
+		MessageRight = int16(ptrdiff_t(Dest + MessageEnd - GetData()));
+		return Parse();
+	}
+}
+
+////////////////////////////////////////////////////////////////////////////////
+FOutcome FStatusHeaders::Parse()
+{
+	const char* Cursor = GetData() + MessageLeft;
+	FAnsiStringView MessageView(Cursor, MessageRight - MessageLeft);
+
+	if (ParseMessage(MessageView, Offsets) < 0)
+	{
+		return FOutcome::Error("Failed to parse message status");
+	}
+
+	FAnsiStringView StatusCodeView(
+		Cursor + Offsets.StatusCode,
+		Offsets.Message - Offsets.StatusCode
+	);
+	StatusCode = uint16(CrudeToInt(StatusCodeView));
+	if (uint32(StatusCode - 100) > 899)
+	{
+		return FOutcome::Error("Invalid status code", StatusCode);
+	}
+
+	static const uint32 Flag_Length		= 1 << 0;
+	static const uint32 Flag_XferEnc	= 1 << 1;
+	static const uint32 Flag_Connection = 1 << 2;
+	static const uint32 Flag_All		= 0x07;
+	uint32 Flags = 0;
+
+	ReadHeaders(
+		[this, &Flags] (FAnsiStringView Name, FAnsiStringView Value) mutable
+		{
+			// todo; may need smarter value handling; ;/, separated options & key-value pairs (ex. in rfc2068)
+
+			if (Name.Equals("Content-Length", ESearchCase::IgnoreCase))
+			{
+				ContentLength = int32(CrudeToInt(Value));
+				Flags |= Flag_Length;
+			}
+
+			else if (Name.Equals("Transfer-Encoding", ESearchCase::IgnoreCase))
+			{
+				bChunked = Value.Equals("chunked", ESearchCase::IgnoreCase);
+				Flags |= Flag_XferEnc;
+			}
+
+			else if (Name.Equals("Connection", ESearchCase::IgnoreCase))
+			{
+				bKeepAlive = !Value.Equals("close");
+				Flags |= Flag_Connection;
+			}
+
+			return Flags != Flag_All;
+		}
+	);
+
+	bHasBody = (StatusCode != 304);
+	bHasBody &= !IsHeadMethod();
+	bHasBody &= bChunked | (ContentLength > 0);
+
+	if (Flags & Flag_Length)
+	{
+		if ((StatusCode <= 100) | (StatusCode == 204) | (Flags & Flag_XferEnc))
+		{
+			return FOutcome::Error("Unexpected Content-Length field", StatusCode);
+		}
+
+		if (ContentLength < 0)
+		{
+			return FOutcome::Error("Invalid Content-Length field", uint32(ContentLength));
+		}
+
+		if ((GetSize() - MessageRight) > ContentLength)
+		{
+			return FOutcome::Error("More data received that expected");
+		}
+	}
+	else if (Flags & Flag_XferEnc)
+	{
+		if (!bChunked)
+		{
+			return FOutcome::Error("Unsupported Transfer-Encoding");
+		}
+		ContentLength = -1;
+	}
+	else if (bHasBody)
+	{
+		return FOutcome::Error("Missing Content-Length field");
+	}
+	else if (GetSize() > uint32(MessageRight))
+	{
+		return FOutcome::Error("Received content when none was expected");
+	}
+
+	return FOutcome::Ok();
+}
+
+
+
+////////////////////////////////////////////////////////////////////////////////
+class FBody
+	: public FStatusHeaders
+{
+public:
+	using FStatusHeaders::FStatusHeaders;
+
+	int64				GetRemaining() const;
+	FOutcome			TryRecv(FMutableMemoryView Dest, FHttpPeer& Peer);
+
+private:
+	enum class EState : uint8 { Init, Recv, Prologue, Chunk, Epilogue, Done };
+
+	FOutcome			Gather(FMutableMemoryView Dest, FHttpPeer& Peer);
+	FOutcome			Init(FMutableMemoryView Dest, FHttpPeer& Peer);
+	FOutcome			Recv(FMutableMemoryView Dest, FHttpPeer& Peer);
+	FOutcome			Prologue(FMutableMemoryView Dest, FHttpPeer& Peer);
+	FOutcome			Chunk(FMutableMemoryView Dest, FHttpPeer& Peer);
+	FOutcome			Epilogue(FMutableMemoryView Dest, FHttpPeer& Peer);
+	FMemoryView			Overspill;
+	FMutableMemoryView	Scratch;
+	int64				Remaining = 0;
+	EState				State = EState::Init;
+
+	enum
+	{
+		HeaderBufSize	= 32,
+		CrLfLength		= 2,
+		EndOfXfer		= -1,
+	};
+};
+
+////////////////////////////////////////////////////////////////////////////////
+int64 FBody::GetRemaining() const
+{
+	return (State == EState::Recv) ? Remaining : -1;
+}
+
+////////////////////////////////////////////////////////////////////////////////
+FOutcome FBody::TryRecv(FMutableMemoryView Dest, FHttpPeer& Peer)
+{
+	FOutcome Outcome = FOutcome::None();
+
+	switch (State)
+	{
+	case EState::Init:	Outcome = Init(Dest, Peer);		break;
+	case EState::Recv:	Outcome = Recv(Dest, Peer);		break;
+	default:			Outcome = Gather(Dest, Peer);	break;
+	case EState::Done:	Outcome = FOutcome::Error("Transaction is complete"); break;
+	}
+
+	if (Outcome.IsOk())
+	{
+		State = EState::Done;
+	}
+
+	return Outcome;
+}
+
+////////////////////////////////////////////////////////////////////////////////
+FOutcome FBody::Gather(FMutableMemoryView Dest, FHttpPeer& Peer)
+{
+	FOutcome Outcome = FOutcome::None();
+
+	bool bDone = false;
+	int32 Size = 0;
+	do
+	{
+		switch (State)
+		{
+		case EState::Prologue:	Outcome = Prologue(Dest, Peer); break;
+		case EState::Chunk:		Outcome = Chunk(Dest, Peer);	break;
+		case EState::Epilogue:	Outcome = Epilogue(Dest, Peer);	break;
+		default:				Outcome = FOutcome::Error("Unexpected try-chunked state"); break;
+		}
+
+		if (Outcome.IsError())
+		{
+			return Outcome;
+		}
+
+		bDone = Outcome.IsOk();
+		int32 Result = Outcome.GetResult();
+		if (Result == 0)
+		{
+			break;
+		}
+
+		Size += Result;
+		Dest = Dest.Mid(Result);
+	}
+	while (!Dest.IsEmpty());
+
+	return bDone ? FOutcome::Ok(Size) : FOutcome::Waiting(Size);
+}
+
+////////////////////////////////////////////////////////////////////////////////
+FOutcome FBody::Init(FMutableMemoryView Dest, FHttpPeer& Peer)
+{
+	const char* DataEnd = GetData() + GetSize();
+	const char* MessageEnd = GetMessageRight();
+	uint32 AlreadyReceived = uint32(ptrdiff_t(DataEnd - MessageEnd));
+	if (AlreadyReceived)
+	{
+		Overspill = { (const uint8*)MessageEnd, AlreadyReceived };
+	}
+
+	if (AlreadyReceived < HeaderBufSize)
+	{
+		GetMutableFree(HeaderBufSize, HeaderBufSize);
+		MessageEnd = GetMessageRight();
+	}
+
+	Scratch = { (uint8*)MessageEnd, HeaderBufSize };
+	
+	if (IsChunked())
+	{
+		check(Remaining == 0);
+		State = EState::Prologue;
+	}
+	else
+	{
+		Remaining = GetContentLength();
+		check(Remaining > 0);
+		State = EState::Recv;
+	}
+
+	return TryRecv(Dest, Peer);
+}
+
+////////////////////////////////////////////////////////////////////////////////
+FOutcome FBody::Recv(FMutableMemoryView Dest, FHttpPeer& Peer)
+{
+	int64 OverspillSize = Overspill.GetSize();
+	if (OverspillSize != 0)
+	{
+		int64 Size = FMath::Min<int64>(Dest.GetSize(), OverspillSize);
+		::memcpy(Dest.GetData(), Overspill.GetData(), Size);
+		Dest = Dest.Mid(Size);
+		Overspill = Overspill.Mid(Size);
+		Remaining -= Size;
+		if (Remaining == 0 || Dest.IsEmpty())
+		{
+			int32 Result = int32(Size);
+			return Remaining ? FOutcome::Waiting(Result) : FOutcome::Ok(Result);
+		}
+	}
+
+	auto* Cursor = (char*)(Dest.GetData());
+	uint32 Size = uint32(FMath::Min<int64>(Dest.GetSize(), Remaining));
+	FOutcome Outcome = Peer.Recv(Cursor, Size);
+	if (Outcome.IsError())
+	{
+		return Outcome;
+	}
+
+	int32 Result = Outcome.GetResult();
+	Remaining -= Result;
+	if (Remaining < 0)
+	{
+		return FOutcome::Error("Unexpectedly received too much", int32(Remaining));
+	}
+
+	Result += int32(OverspillSize);
+	return Remaining ? FOutcome::Waiting(Result) : FOutcome::Ok(Result);
+}
+
+////////////////////////////////////////////////////////////////////////////////
+FOutcome FBody::Prologue(FMutableMemoryView Dest, FHttpPeer& Peer)
+{
+	auto FindHeader = [] (const char* Cursor, int32 Size)
+	{
+		// Base-16 size
+		int32 HexDigLen = 0;
+		for (; HexDigLen < Size; ++HexDigLen)
+		{
+			uint32 c = Cursor[HexDigLen];
+			if ((c - '0' >= 10) & (c - 'a' >= 6) & (c - 'A' >= 6))
+			{
+				break;
+			}
+			continue;
+		}
+
+		int32 Delta = Size - HexDigLen;
+
+		if (HexDigLen == 0)					return FOutcome::Error("Invalid chunk header");
+		if (Delta < 1)						return FOutcome::Waiting();
+		if (Cursor[HexDigLen] != '\r')		return FOutcome::Error("Chunk extensions are not supported (ERREXT)");
+		if (Delta < 2)						return FOutcome::Waiting();
+		if (Cursor[HexDigLen + 1] != '\n')	return FOutcome::Error("Invalid chunk CRLF");
+
+		return FOutcome::Ok(HexDigLen + 2);
+	};
+
+	const char* Header = nullptr;
+	FOutcome Outcome = FOutcome::None();
+	int32 HeaderLength = 0;
+
+	uint32 OverspillSize = uint32(Overspill.GetSize());
+	if (OverspillSize != 0)
+	{
+		// We try all the overspill first - we don't know how much we're expecting
+		// and overspill could contain both header and all data (thus further recvs
+		// could falsely hang us up).
+		Header = (char*)Overspill.GetData();
+
+		Outcome = FindHeader(Header, OverspillSize);
+		if (Outcome.IsError())
+		{
+			return Outcome;
+		}
+
+		HeaderLength = Outcome.GetResult();
+		Overspill = Overspill.Mid(HeaderLength);
+	}
+
+	if (HeaderLength == 0)
+	{
+		char* Cursor = (char*)Scratch.GetData();
+
+		if (OverspillSize != 0)
+		{
+			check(OverspillSize <= HeaderBufSize);
+			std::memmove(Cursor, Overspill.GetData(), OverspillSize);
+			Cursor += OverspillSize;
+		}
+
+		Outcome = Peer.Recv(Cursor, HeaderBufSize);
+		if (!Outcome.IsOk())
+		{
+			return Outcome;
+		}
+
+		Cursor -= OverspillSize;
+		OverspillSize += Outcome.GetResult();
+		Overspill = { (uint8*)Cursor, OverspillSize };
+
+		Outcome = FindHeader(Cursor, OverspillSize);
+		if (!Outcome.IsOk())
+		{
+			return Outcome;
+		}
+
+		Header = Cursor;
+		HeaderLength = Outcome.GetResult();
+		Overspill = Overspill.Mid(HeaderLength);
+	}
+
+	// Read chunk size
+	Header += Remaining;
+	int64 ChunkSize = CrudeToInt<16>(FAnsiStringView(Header, HeaderLength));
+	if (ChunkSize < 0)			return FOutcome::Error("Unparsable chunk size");
+	if (ChunkSize > 16 << 20)	return FOutcome::Error("Unacceptable chunk size");
+
+	if (ChunkSize == 0)
+	{
+		Remaining = EndOfXfer;
+		::memset(Scratch.GetData(), 0, CrLfLength);
+		State = EState::Epilogue;
+		return Epilogue(Dest, Peer);
+	}
+
+	Remaining = ChunkSize;
+	State = EState::Chunk;
+	return Chunk(Dest, Peer);
+}
+
+////////////////////////////////////////////////////////////////////////////////
+FOutcome FBody::Chunk(FMutableMemoryView Dest, FHttpPeer& Peer)
+{
+	check(Remaining != 0);
+
+	Dest = Dest.Left(Remaining);
+	FOutcome Outcome = Recv(Dest, Peer);
+	if (!Outcome.IsOk())
+	{
+		return Outcome;
+	}
+
+	check(Remaining == 0);
+	::memset(Scratch.GetData(), 0, CrLfLength);
+	State = EState::Epilogue;
+
+	int32 Result = Outcome.GetResult();
+	return FOutcome::Waiting(Result);
+}
+
+////////////////////////////////////////////////////////////////////////////////
+FOutcome FBody::Epilogue(FMutableMemoryView Dest, FHttpPeer& Peer)
+{
+	auto* Cursor = (uint8*)(Scratch.GetData());
+
+	FMutableMemoryView CrLfBuffer(Cursor, CrLfLength);
+	if (Cursor[0] != 0)
+	{
+		CrLfBuffer = { Cursor + 1, 1 };
+	}
+
+	check(Remaining <= 0);
+
+	int64 RemainingCache = Remaining;
+	Remaining = CrLfBuffer.GetSize();
+	FOutcome Outcome = Recv(CrLfBuffer, Peer);
+	if (!Outcome.IsOk())
+	{
+		Remaining = RemainingCache;
+		return Outcome;
+	}
+
+	if (Cursor[0] != '\r' || Cursor[1] != '\n')
+	{
+		return FOutcome::Error("Trailing headers are not supported (ERRTRAIL)");
+	}
+
+	if (RemainingCache > EndOfXfer)
+	{
+		State = EState::Prologue;
+		return Prologue(Dest, Peer);
+	}
+
+	return FOutcome::Ok();
+}
+
+} // namespace DetailOne
+
+
+
+////////////////////////////////////////////////////////////////////////////////
+class FTransactionOne
+	: public DetailOne::FBody
+{
+public:
+	using DetailOne::FBody::FBody;
+};
+
+using FTransaction = FTransactionOne;
+
+} // namespace UE::IoStore::HTTP
