Short minimal-edit route to separating out two classes that had got confused.

#rnx

diff --git a/Engine/Source/Runtime/Experimental/IoStore/HttpClient/Private/Activity.inl b/Engine/Source/Runtime/Experimental/IoStore/HttpClient/Private/Activity.inl
--- a/Engine/Source/Runtime/Experimental/IoStore/HttpClient/Private/Activity.inl
+++ b/Engine/Source/Runtime/Experimental/IoStore/HttpClient/Private/Activity.inl
@@ -152,6 +152,7 @@ private:
 	UPTRINT				SinkParam;
 	FTicketSink			Sink;
 	FTransaction		Transaction;
+	FBuffer				Buffer;
 
 	struct FHeaderRecord
 	{
@@ -169,7 +170,8 @@ private:
 FActivity::FActivity(const FParams& Params)
 : bAllowChunked(Params.bAllowChunked)
 , Host(Params.Host)
-, Transaction(Params.Buffer, Params.BufferSize)
+, Transaction(Buffer)
+, Buffer(Params.Buffer, Params.BufferSize)
 {
 	check(Host != nullptr);
 
@@ -180,9 +182,9 @@ FActivity::FActivity(const FParams& Params)
 	auto Copy = [this] (FAnsiStringView Value)
 	{
 		uint32 Length = uint32(Value.Len());
-		char* Ptr = Transaction.Alloc<char>(Length);
+		char* Ptr = Buffer.Alloc<char>(Length);
 		std::memcpy(Ptr, Value.GetData(), Length);
-		uint32 Ret = uint32(ptrdiff_t(Ptr - Transaction.GetData()));
+		uint32 Ret = uint32(ptrdiff_t(Ptr - Buffer.GetData()));
 		check(Ret <= 255);
 		return uint8(Ret);
 	};
@@ -209,7 +211,7 @@ void FActivity::AddHeader(FAnsiStringView Key, FAnsiStringView Value)
 	static_assert(alignof(FHeaderRecord) == alignof(uint16));
 	uint32 Count = sizeof(FHeaderRecord) + Key.Len() + Value.Len();
 	Count = (Count + sizeof(uint16) - 1) / sizeof(uint16);
-	auto* Record = (FHeaderRecord*)(Transaction.Alloc<uint16>(Count));
+	auto* Record = (FHeaderRecord*)(Buffer.Alloc<uint16>(Count));
 
 	Record->Key = int16(Key.Len());
 	Record->ValueLength = int16(Value.Len());
@@ -316,13 +318,13 @@ const FTransaction* FActivity::GetTransaction() const
 ////////////////////////////////////////////////////////////////////////////////
 FAnsiStringView FActivity::GetMethod() const
 {
-	return { Transaction.GetData() + MethodOffset, PathOffset - MethodOffset };
+	return { Buffer.GetData() + MethodOffset, PathOffset - MethodOffset };
 }
 
 ////////////////////////////////////////////////////////////////////////////////
 FAnsiStringView FActivity::GetPath() const
 {
-	return { Transaction.GetData() + PathOffset, PathLength };
+	return { Buffer.GetData() + PathOffset, PathLength };
 }
 
 ////////////////////////////////////////////////////////////////////////////////
@@ -339,7 +341,7 @@ void FActivity::EnumerateHeaders(FResponse::FHeaderSink HeaderSink) const
 	for (uint32 i = 0, n = HeaderCount; i < n; ++i)
 	{
 		RecordAddr = (RecordAddr + alignof(FHeaderRecord) - 1) & ~(alignof(FHeaderRecord) - 1);
-		const auto* Record = (FHeaderRecord*)(Transaction.GetData() + RecordAddr);
+		const auto* Record = (FHeaderRecord*)(Buffer.GetData() + RecordAddr);
 
 		FAnsiStringView Key(Record->Data, Record->Key);
 		FAnsiStringView Value(Record->Data + Record->Key, Record->ValueLength);
diff --git a/Engine/Source/Runtime/Experimental/IoStore/HttpClient/Private/TransactionOne.inl b/Engine/Source/Runtime/Experimental/IoStore/HttpClient/Private/TransactionOne.inl
--- a/Engine/Source/Runtime/Experimental/IoStore/HttpClient/Private/TransactionOne.inl
+++ b/Engine/Source/Runtime/Experimental/IoStore/HttpClient/Private/TransactionOne.inl
@@ -132,12 +132,62 @@ static void EnumerateHeaders(FAnsiStringView Headers, LambdaType&& Lambda)
 namespace DetailOne
 {
 
+////////////////////////////////////////////////////////////////////////////////
+class FBase
+{
+public:
+					FBase(FBuffer& InBuffer);
+
+protected:
+	using FMutableSection = FBuffer::FMutableSection;
+
+	const char*		GetData() const;
+	uint32			GetSize() const;
+	FMutableSection	GetMutableFree(uint32 MinSize, uint32 PageSize=256);
+	void			AdvanceUsed(uint32 Delta);
+
+private:
+	FBuffer&		Buffer;
+};
+
+////////////////////////////////////////////////////////////////////////////////
+FBase::FBase(FBuffer& InBuffer)
+: Buffer(InBuffer)
+{
+}
+
+////////////////////////////////////////////////////////////////////////////////
+const char* FBase::GetData() const
+{
+	return Buffer.GetData();
+}
+
+////////////////////////////////////////////////////////////////////////////////
+uint32 FBase::GetSize() const
+{
+	return Buffer.GetSize();
+}
+
+////////////////////////////////////////////////////////////////////////////////
+FBase::FMutableSection FBase::GetMutableFree(uint32 MinSize, uint32 PageSize)
+{
+	return Buffer.GetMutableFree(MinSize, PageSize);
+}
+
+////////////////////////////////////////////////////////////////////////////////
+void FBase::AdvanceUsed(uint32 Delta)
+{
+	return Buffer.AdvanceUsed(Delta);
+}
+
+
+
 ////////////////////////////////////////////////////////////////////////////////
 class FRequest
-	: public FBuffer
+	: public FBase
 {
 public:
-	using FBuffer::FBuffer;
+	using FBase::FBase;
 
 	void			Begin(FAnsiStringView Method, FAnsiStringView Path);
 	void			AddHeader(FAnsiStringView Key, FAnsiStringView Value);
