Off-by-one error

#rnx

diff --git a/Engine/Source/Runtime/Experimental/IoStore/HttpClient/Private/TransactionOne.inl b/Engine/Source/Runtime/Experimental/IoStore/HttpClient/Private/TransactionOne.inl
--- a/Engine/Source/Runtime/Experimental/IoStore/HttpClient/Private/TransactionOne.inl
+++ b/Engine/Source/Runtime/Experimental/IoStore/HttpClient/Private/TransactionOne.inl
@@ -97,7 +97,7 @@ static void EnumerateHeaders(FAnsiStringView Headers, LambdaType&& Lambda)
 		Cursor += ColonIndex;
 
 		const char* Right = Cursor + 1;
-		while (Right < End)
+		while (Right + 1 < End)
 		{
 			if (Right[0] != '\r' || Right[1] != '\n')
 			{
