String comparison should be case-insensitive

#rnx

diff --git a/Engine/Source/Runtime/Experimental/IoStore/HttpClient/Private/TransactionOne.inl b/Engine/Source/Runtime/Experimental/IoStore/HttpClient/Private/TransactionOne.inl
--- a/Engine/Source/Runtime/Experimental/IoStore/HttpClient/Private/TransactionOne.inl
+++ b/Engine/Source/Runtime/Experimental/IoStore/HttpClient/Private/TransactionOne.inl
@@ -477,7 +477,7 @@ FOutcome FStatusHeaders::Parse()
 
 			else if (Name.Equals("Connection", ESearchCase::IgnoreCase))
 			{
-				bKeepAlive = !Value.Equals("close");
+				bKeepAlive = !Value.Equals("close", ESearchCase::IgnoreCase);
 				Flags |= Flag_Connection;
 			}
 
