Ask the peer to create appropriate transactions for us

#rnx

diff --git a/Engine/Source/Runtime/Experimental/IoStore/HttpClient/Private/Activity.inl b/Engine/Source/Runtime/Experimental/IoStore/HttpClient/Private/Activity.inl
--- a/Engine/Source/Runtime/Experimental/IoStore/HttpClient/Private/Activity.inl
+++ b/Engine/Source/Runtime/Experimental/IoStore/HttpClient/Private/Activity.inl
@@ -86,7 +86,7 @@ struct FActivity
 	};
 
 						FActivity(const FParams& Params);
-						~FActivity() = default;
+						~FActivity();
 	void				AddHeader(FAnsiStringView Key, FAnsiStringView Value);
 	FOutcome			Tick(FHttpPeer& Peer, int32* MaxRecvSize=nullptr);
 	void				SetSink(FTicketSink&& InSink, UPTRINT Param);
@@ -151,7 +151,7 @@ private:
 	};
 	UPTRINT				SinkParam;
 	FTicketSink			Sink;
-	FTransaction		Transaction;
+	FTransaction*		Transaction = nullptr;
 	FBuffer				Buffer;
 
 	struct FHeaderRecord
@@ -170,7 +170,6 @@ private:
 FActivity::FActivity(const FParams& Params)
 : bAllowChunked(Params.bAllowChunked)
 , Host(Params.Host)
-, Transaction(Buffer)
 , Buffer(Params.Buffer, Params.BufferSize)
 {
 	check(Host != nullptr);
@@ -200,6 +199,15 @@ FActivity::FActivity(const FParams& Params)
 	LengthScore = uint8(Pow2);
 }
 
+////////////////////////////////////////////////////////////////////////////////
+FActivity::~FActivity()
+{
+	if (Transaction != nullptr)
+	{
+		delete Transaction;
+	}
+}
+
 ////////////////////////////////////////////////////////////////////////////////
 void FActivity::AddHeader(FAnsiStringView Key, FAnsiStringView Value)
 {
@@ -312,7 +320,8 @@ FActivity::EStage FActivity::GetStage() const
 const FTransaction* FActivity::GetTransaction() const
 {
 	check(State >= EState::Send && State <= EState::Completed);
-	return &Transaction;
+	check(Transaction != nullptr);
+	return Transaction;
 }
 
 ////////////////////////////////////////////////////////////////////////////////
@@ -369,7 +378,7 @@ uint32 FActivity::GetRemainingKiB() const
 {
 	if (State <= EState::RecvStream)  return MAX_uint32;
 	if (State >  EState::RecvContent) return 0;
-	return uint32(Transaction.GetRemaining() >> 10);
+	return uint32(GetTransaction()->GetRemaining() >> 10);
 }
 
 ////////////////////////////////////////////////////////////////////////////////
@@ -392,23 +401,26 @@ const FStopwatch& FActivity::GetStopwatch() const
 ////////////////////////////////////////////////////////////////////////////////
 FOutcome FActivity::Transact(FHttpPeer& Peer)
 {
-	Transaction.Begin(GetMethod(), GetPath());
+	Transaction = Peer.Transact(Buffer);
+	check(Transaction != nullptr);
+
+	Transaction->Begin(GetMethod(), GetPath());
 
-	Transaction.AddHeader("Host", GetHost()->GetHostName());
+	Transaction->AddHeader("Host", GetHost()->GetHostName());
 
 	// HTTP/1.1 is persistent by default thus "Connection" header isn't required
 	if (!GetHost()->IsPooled())
 	{
-		Transaction.AddHeader("Connection", "close");
+		Transaction->AddHeader("Connection", "close");
 	}
 
 	EnumerateHeaders([this] (FAnsiStringView Key, FAnsiStringView Value)
 	{
-		Transaction.AddHeader(Key, Value);
+		Transaction->AddHeader(Key, Value);
 		return true;
 	});
 
-	Transaction.End();
+	Transaction->End();
 	ChangeState(EState::Send);
 	return FOutcome::Ok();
 }
@@ -422,7 +434,7 @@ FOutcome FActivity::Send(FHttpPeer& Peer)
 
 	TRACE_CPUPROFILER_EVENT_SCOPE(IasHttp::DoSend);
 
-	FOutcome Outcome = Transaction.TrySendRequest(Peer);
+	FOutcome Outcome = Transaction->TrySendRequest(Peer);
 	if (Outcome.IsError())
 	{
 		Fail(Outcome);
@@ -454,7 +466,7 @@ FOutcome FActivity::RecvMessage(FHttpPeer& Peer)
 	Stopwatch.RecvStart();
 #endif
 
-	FOutcome Outcome = Transaction.TryRecvResponse(Peer);
+	FOutcome Outcome = Transaction->TryRecvResponse(Peer);
 	if (Outcome.IsError())
 	{
 		Fail(Outcome);
@@ -468,8 +480,8 @@ FOutcome FActivity::RecvMessage(FHttpPeer& Peer)
 
 	check(Outcome.IsOk());
 
-	bool bChunked = Transaction.IsChunked();
-	int64 ContentLength = Transaction.GetContentLength();
+	bool bChunked = Transaction->IsChunked();
+	int64 ContentLength = Transaction->GetContentLength();
 
 	// Validate that the server's told us how and how much it will transmit
 	if (bChunked)
@@ -487,7 +499,7 @@ FOutcome FActivity::RecvMessage(FHttpPeer& Peer)
 	CallSink();
 
 	// Check the user gave us a destination for content
-	if (Transaction.HasBody())
+	if (Transaction->HasBody())
 	{
 		if (Dest == PriorDest)
 		{
@@ -517,7 +529,7 @@ FOutcome FActivity::RecvMessage(FHttpPeer& Peer)
 	}
 
 	// HEAD methods
-	if (!Transaction.HasBody())
+	if (!Transaction->HasBody())
 	{
 		ChangeState(EState::RecvDone);
 		return FOutcome::Ok(int32(EStage::Content));
@@ -535,12 +547,12 @@ FOutcome FActivity::RecvContent(FHttpPeer& Peer, int32& MaxRecvSize)
 {
 	TRACE_CPUPROFILER_EVENT_SCOPE(IasHttp::DoRecvContent);
 
-	int64 Remaining = Transaction.GetRemaining();
+	int64 Remaining = Transaction->GetRemaining();
 
 	FMutableMemoryView View = Dest->GetMutableView();
 	View = View.Right(Remaining).Left(MaxRecvSize);
 
-	FOutcome Outcome = Transaction.TryRecv(View, Peer);
+	FOutcome Outcome = Transaction->TryRecv(View, Peer);
 	if (Outcome.IsError())
 	{
 		Fail(Outcome);
@@ -571,7 +583,7 @@ FOutcome FActivity::RecvStream(FHttpPeer& Peer, int32& MaxRecvSize)
 	View = View.Left(MaxRecvSize);
 	MaxRecvSize -= int32(View.GetSize());
 
-	FOutcome Outcome = Transaction.TryRecv(View, Peer);
+	FOutcome Outcome = Transaction->TryRecv(View, Peer);
 	if (Outcome.IsError())
 	{
 		Fail(Outcome);
diff --git a/Engine/Source/Runtime/Experimental/IoStore/HttpClient/Private/Peer.inl b/Engine/Source/Runtime/Experimental/IoStore/HttpClient/Private/Peer.inl
--- a/Engine/Source/Runtime/Experimental/IoStore/HttpClient/Private/Peer.inl
+++ b/Engine/Source/Runtime/Experimental/IoStore/HttpClient/Private/Peer.inl
@@ -500,18 +500,22 @@ FOutcome FTlsPeer::Recv(char* Out, int32 MaxSize)
 
 
 
+////////////////////////////////////////////////////////////////////////////////
+class FTransaction* CreateTransactOne(FBuffer&);
+
 ////////////////////////////////////////////////////////////////////////////////
 class FHttpPeer
 	: public FTlsPeer
 {
 public:
-				FHttpPeer() = default;
-				FHttpPeer(FSocket InSocket, FCertRootsRef Certs=ECertRootsRefType::None, const char* HostName=nullptr);
-	FOutcome	Handshake();
+					FHttpPeer() = default;
+					FHttpPeer(FSocket InSocket, FCertRootsRef Certs=ECertRootsRefType::None, const char* HostName=nullptr);
+	FOutcome		Handshake();
+	FTransaction*	Transact(FBuffer& Buffer);
 
 private:
-	void		AssignProto();
-	int32		Proto = 0;
+	void			AssignProto();
+	int32			Proto = 0;
 };
 
 ////////////////////////////////////////////////////////////////////////////////
@@ -552,4 +556,16 @@ void FHttpPeer::AssignProto()
 	}
 }
 
+////////////////////////////////////////////////////////////////////////////////
+FTransaction* FHttpPeer::Transact(FBuffer& Buffer)
+{
+	if (Proto == 1)
+	{
+		return CreateTransactOne(Buffer);
+	}
+
+	check(!"unsupported");
+	return nullptr;
+}
+
 } // namespace UE::IoStore::HTTP
diff --git a/Engine/Source/Runtime/Experimental/IoStore/HttpClient/Private/TransactionOne.inl b/Engine/Source/Runtime/Experimental/IoStore/HttpClient/Private/TransactionOne.inl
--- a/Engine/Source/Runtime/Experimental/IoStore/HttpClient/Private/TransactionOne.inl
+++ b/Engine/Source/Runtime/Experimental/IoStore/HttpClient/Private/TransactionOne.inl
@@ -864,6 +864,15 @@ public:
 	using DetailOne::FBody::FBody;
 };
 
-using FTransaction = FTransactionOne;
+class FTransaction
+	: public FTransactionOne
+{
+	using FTransactionOne::FTransactionOne;
+};
+
+FTransaction* CreateTransactOne(FBuffer& Buffer)
+{
+	return new FTransaction(Buffer);
+}
 
 } // namespace UE::IoStore::HTTP
