Associate transactions with an ID

#rnx

diff --git a/Engine/Source/Runtime/Experimental/IoStore/HttpClient/Private/Activity.inl b/Engine/Source/Runtime/Experimental/IoStore/HttpClient/Private/Activity.inl
--- a/Engine/Source/Runtime/Experimental/IoStore/HttpClient/Private/Activity.inl
+++ b/Engine/Source/Runtime/Experimental/IoStore/HttpClient/Private/Activity.inl
@@ -95,6 +95,7 @@ struct FActivity
 	void				Done();
 	void				Fail(const FOutcome& Outcome);
 	EStage				GetStage() const;
+	FTransactId			GetTransactId() const;
 	const FTransaction*	GetTransaction() const;
 	FAnsiStringView		GetMethod() const;
 	FAnsiStringView		GetPath() const;
@@ -141,6 +142,7 @@ private:
 	uint8				PathOffset;
 	uint16				PathLength;
 	uint16				HeaderCount = 0;
+	FTransactId			TransactId = 0;
 #if IAS_HTTP_WITH_PERF
 	FStopwatch			Stopwatch;
 #endif
@@ -316,6 +318,13 @@ FActivity::EStage FActivity::GetStage() const
 	return EStage::Failed;
 }
 
+////////////////////////////////////////////////////////////////////////////////
+FTransactId FActivity::GetTransactId() const
+{
+	check(TransactId != 0);
+	return TransactId;
+}
+
 ////////////////////////////////////////////////////////////////////////////////
 const FTransaction* FActivity::GetTransaction() const
 {
@@ -420,7 +429,7 @@ FOutcome FActivity::Transact(FHttpPeer& Peer)
 		return true;
 	});
 
-	Transaction->End();
+	TransactId = Transaction->End();
 	ChangeState(EState::Send);
 	return FOutcome::Ok();
 }
diff --git a/Engine/Source/Runtime/Experimental/IoStore/HttpClient/Private/TransactionOne.inl b/Engine/Source/Runtime/Experimental/IoStore/HttpClient/Private/TransactionOne.inl
--- a/Engine/Source/Runtime/Experimental/IoStore/HttpClient/Private/TransactionOne.inl
+++ b/Engine/Source/Runtime/Experimental/IoStore/HttpClient/Private/TransactionOne.inl
@@ -5,6 +5,9 @@
 namespace UE::IoStore::HTTP
 {
 
+////////////////////////////////////////////////////////////////////////////////
+using FTransactId = uint32;
+
 ////////////////////////////////////////////////////////////////////////////////
 struct FMessageOffsets
 {
@@ -191,7 +194,7 @@ public:
 
 	void			Begin(FAnsiStringView Method, FAnsiStringView Path);
 	void			AddHeader(FAnsiStringView Key, FAnsiStringView Value);
-	void			End();
+	FTransactId		End();
 	FOutcome		TrySendRequest(FHttpPeer& Peer);
 
 protected:
@@ -231,10 +234,11 @@ void FRequest::AddHeader(FAnsiStringView Key, FAnsiStringView Value)
 }
 
 ////////////////////////////////////////////////////////////////////////////////
-void FRequest::End()
+FTransactId FRequest::End()
 {
 	*this << "\r\n";
 	HeaderRight = int16(GetSize());
+	return 1;
 }
 
 ////////////////////////////////////////////////////////////////////////////////
