Sketch of asking peer which transaction it will deal with next

#rnx

diff --git a/Engine/Source/Runtime/Experimental/IoStore/HttpClient/Private/Peer.inl b/Engine/Source/Runtime/Experimental/IoStore/HttpClient/Private/Peer.inl
--- a/Engine/Source/Runtime/Experimental/IoStore/HttpClient/Private/Peer.inl
+++ b/Engine/Source/Runtime/Experimental/IoStore/HttpClient/Private/Peer.inl
@@ -512,6 +512,7 @@ public:
 					FHttpPeer(FSocket InSocket, FCertRootsRef Certs=ECertRootsRefType::None, const char* HostName=nullptr);
 	FOutcome		Handshake();
 	FTransaction*	Transact(FBuffer& Buffer);
+	FOutcome		GetPendingTransactId();
 
 private:
 	void			AssignProto();
@@ -568,4 +569,15 @@ FTransaction* FHttpPeer::Transact(FBuffer& Buffer)
 	return nullptr;
 }
 
+////////////////////////////////////////////////////////////////////////////////
+FOutcome FHttpPeer::GetPendingTransactId()
+{
+	if (Proto == 1)
+	{
+		return FOutcome::Ok(1);
+	}
+
+	return FOutcome::Error("unsupported");
+}
+
 } // namespace UE::IoStore::HTTP
