Lookup activities with guidance from the peer.

#rnx

diff --git a/Engine/Source/Runtime/Experimental/IoStore/HttpClient/Private/Loop.inl b/Engine/Source/Runtime/Experimental/IoStore/HttpClient/Private/Loop.inl
--- a/Engine/Source/Runtime/Experimental/IoStore/HttpClient/Private/Loop.inl
+++ b/Engine/Source/Runtime/Experimental/IoStore/HttpClient/Private/Loop.inl
@@ -480,9 +480,34 @@ void FPeerGroup::RecvInternal(FTickState& State)
 	check(bNegotiating == false);
 	check(!Recv.IsEmpty());
 
-	FActivityNode* Activity = Recv.GetHead();
+	FActivityNode* Activity = nullptr;
+
+	FOutcome Outcome = Peer.GetPendingTransactId();
+	if (Outcome.IsWaiting())
+	{
+		bWaiting |= true;
+		return;
+	}
+
+	while (Outcome.IsOk())
+	{
+		uint32 TransactId = Outcome.GetResult();
+		Activity = Recv.FindFirst(
+			[TransactId] (FActivityNode* Activity)
+			{
+				return Activity->GetTransactId() == TransactId;
+			}
+		);
+
+		if (Activity == nullptr)
+		{
+			Outcome = FOutcome::Error("could not find work for transact id", TransactId);
+			break;
+		}
 
-	FOutcome Outcome = Activity->Tick(Peer, &State.RecvAllowance);
+		Outcome = Activity->Tick(Peer, &State.RecvAllowance);
+		break;
+	}
 
 	// Any sort of error here is unrecoverable
 	if (Outcome.IsError())
@@ -491,6 +516,8 @@ void FPeerGroup::RecvInternal(FTickState& State)
 		return;
 	}
 
+	check(Activity != nullptr);
+
 	IsKeepAlive &= uint8(Activity->GetTransaction()->IsKeepAlive() == true);
 	LastUseMs = State.NowMs;
 	bWaiting |= Outcome.IsWaiting();
