Iris
* Fix InitNetObjectFactories initializing invalid array entries

#rnx

diff --git a/Engine/Source/Runtime/Experimental/Iris/Core/Private/Iris/ReplicationSystem/ObjectReplicationBridge.cpp b/Engine/Source/Runtime/Experimental/Iris/Core/Private/Iris/ReplicationSystem/ObjectReplicationBridge.cpp
--- a/Engine/Source/Runtime/Experimental/Iris/Core/Private/Iris/ReplicationSystem/ObjectReplicationBridge.cpp
+++ b/Engine/Source/Runtime/Experimental/Iris/Core/Private/Iris/ReplicationSystem/ObjectReplicationBridge.cpp
@@ -235,13 +235,16 @@ void UObjectReplicationBridge::InitNetObjectFactories()
 	for(int32 Index=0; Index < RegisteredFactories.Num(); ++Index )
 	{
 		const FNetObjectFactoryRegistry::FFactoryData& Data = RegisteredFactories[Index];
-		check(Data.Id == (FNetObjectFactoryId)Index);
-		check(Data.NetFactoryClass.Get());
+		if (Data.Id != InvalidNetObjectFactoryId)
+		{
+			check(Data.Id == (FNetObjectFactoryId)Index);
+			check(Data.NetFactoryClass.Get());
 
-		UNetObjectFactory* Factory = NewObject<UNetObjectFactory>(this, Data.NetFactoryClass.Get());
-		NetObjectFactories.Add(Factory);
+			UNetObjectFactory* Factory = NewObject<UNetObjectFactory>(this, Data.NetFactoryClass.Get());
+			NetObjectFactories.Add(Factory);
 
-		Factory->Init(Data.Id, this);
+			Factory->Init(Data.Id, this);
+		}
 	}
 }
 
