IAD Cache
- Flush timestamp updates on block read with 1 hr granularity
- Fix timestamps not being updated after initial block creation (FindOrAdd should have been Add)

diff --git a/Engine/Source/Runtime/Experimental/IoStore/OnDemand/Private/OnDemandConfig.cpp b/Engine/Source/Runtime/Experimental/IoStore/OnDemand/Private/OnDemandConfig.cpp
--- a/Engine/Source/Runtime/Experimental/IoStore/OnDemand/Private/OnDemandConfig.cpp
+++ b/Engine/Source/Runtime/Experimental/IoStore/OnDemand/Private/OnDemandConfig.cpp
@@ -436,6 +436,24 @@ TIoStatusOr<FOnDemandInstallCacheConfig> TryParseInstallCacheConfig(const TCHAR*
 		OutConfig.JournalMaxSize = JournalMaxSize;
 	}
 
+	const double LastAccessGranularitySeconds = [CommandLine]
+	{
+		if (double ParamValue; FParse::Value(CommandLine, TEXT("-Iad.FileCache.LastAccessGranularitySeconds="), ParamValue))
+		{
+			return ParamValue;
+		}
+		else if (GConfig->GetDouble(TEXT("OnDemandInstall"), TEXT("FileCache.LastAccessGranularitySeconds"), ParamValue, GEngineIni))
+		{
+			return ParamValue;
+		}
+		return double(-1.0);
+	}();
+
+	if (LastAccessGranularitySeconds >= 0)
+	{
+		OutConfig.LastAccessGranularitySeconds = LastAccessGranularitySeconds;
+	}
+
 	OutConfig.RootDirectory = Config::GetInstallCacheDirectory(CommandLine);
 #if !UE_BUILD_SHIPPING
 	OutConfig.bDropCache = FParse::Param(CommandLine, TEXT("Iad.DropCache"));
diff --git a/Engine/Source/Runtime/Experimental/IoStore/OnDemand/Private/OnDemandContentInstaller.cpp b/Engine/Source/Runtime/Experimental/IoStore/OnDemand/Private/OnDemandContentInstaller.cpp
--- a/Engine/Source/Runtime/Experimental/IoStore/OnDemand/Private/OnDemandContentInstaller.cpp
+++ b/Engine/Source/Runtime/Experimental/IoStore/OnDemand/Private/OnDemandContentInstaller.cpp
@@ -289,6 +289,16 @@ void FOnDemandContentInstaller::EnqueueVerifyRequest(FOnDemandVerifyCacheComplet
 	TryExecuteNextRequest();
 }
 
+void FOnDemandContentInstaller::EnqueueFlushLastAccessRequest(FOnDemandFlushLastAccessCompleted&& OnCompleted)
+{
+	{
+		TUniqueLock Lock(Mutex);
+		RequestQueue.HeapPush(RequestAllocator.Construct(MoveTemp(OnCompleted)), RequestSortPredicate);
+	}
+
+	TryExecuteNextRequest();
+}
+
 void FOnDemandContentInstaller::CancelInstallRequest(FSharedInternalInstallRequest InstallRequest)
 {
 	InstallerPipe.Launch(
@@ -421,16 +431,21 @@ void FOnDemandContentInstaller::ExecuteRequest(FRequest& Request)
 			Installer.ExecutePurgeRequest(Request);
 		}
 
-		void operator()(FRequest::FDefrag& Defrag)
+		void operator()(FRequest::FDefrag&)
 		{
 			Installer.ExecuteDefragRequest(Request);
 		}
 
-		void operator()(FRequest::FVerify& Verify)
+		void operator()(FRequest::FVerify&)
 		{
 			Installer.ExecuteVerifyRequest(Request);
 		}
 
+		void operator()(FRequest::FFlushLastAccess&)
+		{
+			Installer.ExectuteFlushLastAccessRequest(Request);
+		}
+
 		FOnDemandContentInstaller&	Installer;
 		FRequest&					Request;
 	};
@@ -832,7 +847,7 @@ void FOnDemandContentInstaller::ExecuteVerifyRequest(FRequest& Request)
 	check(Request.IsVerify());
 	check(&Request == CurrentRequest);
 
-	UE_LOG(LogIoStoreOnDemand, Log, TEXT("Executing verify cache request"));	
+	UE_LOG(LogIoStoreOnDemand, Log, TEXT("Executing verify cache request"));
 
 	if (Request.ErrorCode.load(std::memory_order_relaxed) == EIoErrorCode::Cancelled)
 	{
@@ -850,6 +865,30 @@ void FOnDemandContentInstaller::ExecuteVerifyRequest(FRequest& Request)
 	CompleteVerifyRequest(Request); 
 }
 
+void FOnDemandContentInstaller::ExectuteFlushLastAccessRequest(FRequest& Request)
+{
+	check(Request.IsFlushLastAccess());
+	check(&Request == CurrentRequest);
+
+	UE_LOG(LogIoStoreOnDemand, Log, TEXT("Executing flush last access request"));
+
+	if (Request.ErrorCode.load(std::memory_order_relaxed) == EIoErrorCode::Cancelled)
+	{
+		CompleteFlushLastAccessRequest(Request);
+		return;
+	}
+
+	const FIoStatus Status	= IoStore.InstallCache->FlushLastAccess();
+	Request.ErrorCode		= Status.GetErrorCode();
+
+	if (!Status.IsOk())
+	{
+		Request.ErrorReason = Status.ToString();
+	}
+
+	CompleteFlushLastAccessRequest(Request);
+}
+
 void FOnDemandContentInstaller::ProcessDownloadedChunk(
 	FRequest& Request,
 	FChunkHttpRequestHandle& HttpRequest,
@@ -1234,6 +1273,35 @@ void FOnDemandContentInstaller::CompleteVerifyRequest(FRequest& Request)
 	OnCompleted(MoveTemp(VerifyResult));
 }
 
+void FOnDemandContentInstaller::CompleteFlushLastAccessRequest(FRequest& Request)
+{
+	const uint64 DurationCycles = FPlatformTime::Cycles64() - Request.StartTimeCycles;
+
+	FOnDemandFlushLastAccessResult Result;
+	Result.Status = Request.ErrorCode == EIoErrorCode::Ok ? FIoStatus::Ok : FIoStatus(Request.ErrorCode, Request.ErrorReason);
+	Result.DurationInSeconds = FPlatformTime::ToSeconds64(DurationCycles);
+
+	UE_LOG(LogIoStoreOnDemand, Log, TEXT("Flush last access request completed, Result='%s', Duration=%d ms"),
+		GetIoErrorText(Result.Status.GetErrorCode()), int32(Result.DurationInSeconds * 1000));
+
+	FOnDemandFlushLastAccessCompleted OnCompleted = MoveTemp(Request.AsFlushLastAccess().OnCompleted);
+	{
+		UE::TUniqueLock Lock(Mutex);
+		if (&Request == CurrentRequest)
+		{
+			CurrentRequest = nullptr;
+		}
+		RequestAllocator.Destroy(&Request);
+	}
+
+	if (!OnCompleted)
+	{
+		return;
+	}
+
+	OnCompleted(MoveTemp(Result));
+}
+
 void FOnDemandContentInstaller::Shutdown()
 {
 	bShuttingDown					= true;
diff --git a/Engine/Source/Runtime/Experimental/IoStore/OnDemand/Private/OnDemandContentInstaller.h b/Engine/Source/Runtime/Experimental/IoStore/OnDemand/Private/OnDemandContentInstaller.h
--- a/Engine/Source/Runtime/Experimental/IoStore/OnDemand/Private/OnDemandContentInstaller.h
+++ b/Engine/Source/Runtime/Experimental/IoStore/OnDemand/Private/OnDemandContentInstaller.h
@@ -120,7 +120,15 @@ class FOnDemandContentInstaller
 			FOnDemandVerifyCacheCompleted	OnCompleted;
 		};
 
-		using FRequestVariant = TVariant<FEmptyVariantState, FInstall, FPurge, FDefrag, FVerify>;
+		struct FFlushLastAccess
+		{
+			FFlushLastAccess(FOnDemandFlushLastAccessCompleted&& InOnCompleted)
+				: OnCompleted(MoveTemp(InOnCompleted)) { }
+
+			FOnDemandFlushLastAccessCompleted	OnCompleted;
+		};
+
+		using FRequestVariant = TVariant<FEmptyVariantState, FInstall, FPurge, FDefrag, FVerify, FFlushLastAccess>;
 
 		FRequest(FOnDemandInstallArgs&& Args, FOnDemandInstallCompleted&& OnCompleted, FOnDemandInstallProgressed&& OnProgress)
 		{
@@ -142,15 +150,22 @@ class FOnDemandContentInstaller
 			Variant.Emplace<FVerify>(MoveTemp(OnCompleted));
 		}
 
-		bool			IsInstall() const	{ return Variant.IsType<FInstall>(); }
-		bool			IsPurge() const		{ return Variant.IsType<FPurge>(); }
-		bool			IsDefrag() const	{ return Variant.IsType<FDefrag>(); }
-		bool			IsVerify() const	{ return Variant.IsType<FVerify>(); }
+		FRequest(FOnDemandFlushLastAccessCompleted&& OnCompleted)
+		{
+			Variant.Emplace<FFlushLastAccess>(MoveTemp(OnCompleted));
+		}
+
+		bool				IsInstall() const			{ return Variant.IsType<FInstall>(); }
+		bool				IsPurge() const				{ return Variant.IsType<FPurge>(); }
+		bool				IsDefrag() const			{ return Variant.IsType<FDefrag>(); }
+		bool				IsVerify() const			{ return Variant.IsType<FVerify>(); }
+		bool				IsFlushLastAccess() const	{ return Variant.IsType<FFlushLastAccess>(); }
 
-		FInstall&		AsInstall()			{ return Variant.Get<FInstall>(); }
-		FPurge&			AsPurge()			{ return Variant.Get<FPurge>(); }
-		FDefrag&		AsDefrag()			{ return Variant.Get<FDefrag>(); }
-		FVerify&		AsVerify()			{ return Variant.Get<FVerify>(); }
+		FInstall&			AsInstall()			{ return Variant.Get<FInstall>(); }
+		FPurge&				AsPurge()			{ return Variant.Get<FPurge>(); }
+		FDefrag&			AsDefrag()			{ return Variant.Get<FDefrag>(); }
+		FVerify&			AsVerify()			{ return Variant.Get<FVerify>(); }
+		FFlushLastAccess&	AsFlushLastAccess()	{ return Variant.Get<FFlushLastAccess>(); }
 
 		static uint32	NextSeqNo;
 
@@ -190,6 +205,7 @@ public:
 	void							EnqueuePurgeRequest(FOnDemandPurgeArgs&& Args, FOnDemandPurgeCompleted&& OnCompleted);
 	void							EnqueueDefragRequest(FOnDemandDefragArgs&& Args, FOnDemandDefragCompleted&& OnCompleted);
 	void							EnqueueVerifyRequest(FOnDemandVerifyCacheCompleted&& OnCompleted);
+	void							EnqueueFlushLastAccessRequest(FOnDemandFlushLastAccessCompleted&& OnCompleted);
 	void							CancelInstallRequest(FSharedInternalInstallRequest InstallRequest);
 	void							UpdateInstallRequestPriority(FSharedInternalInstallRequest InstallRequest, int32 NewPriority);
 	void							ReportAnalytics(TArray<FAnalyticsEventAttribute>& OutAnalyticsArray) const;
@@ -202,6 +218,7 @@ private:
 	void							ExecutePurgeRequest(FRequest& Request);
 	void							ExecuteDefragRequest(FRequest& Request);
 	void							ExecuteVerifyRequest(FRequest& Request);
+	void							ExectuteFlushLastAccessRequest(FRequest& Request);
 	void							ProcessDownloadedChunk(
 										FRequest& Request,
 										FChunkHttpRequestHandle& HttpRequest,
@@ -213,6 +230,7 @@ private:
 	void							CompletePurgeRequest(FRequest& Request);
 	void							CompleteDefragRequest(FRequest& Request);
 	void							CompleteVerifyRequest(FRequest& Request);
+	void							CompleteFlushLastAccessRequest(FRequest& Request);
 	void							Shutdown();
 
 	FOnDemandIoStore&				IoStore;
diff --git a/Engine/Source/Runtime/Experimental/IoStore/OnDemand/Private/OnDemandInstallCache.cpp b/Engine/Source/Runtime/Experimental/IoStore/OnDemand/Private/OnDemandInstallCache.cpp
--- a/Engine/Source/Runtime/Experimental/IoStore/OnDemand/Private/OnDemandInstallCache.cpp
+++ b/Engine/Source/Runtime/Experimental/IoStore/OnDemand/Private/OnDemandInstallCache.cpp
@@ -164,49 +164,93 @@ struct FCasBlockInfo
 using FCasBlockInfoMap = TMap<FCasBlockId, FCasBlockInfo>;
 
 ///////////////////////////////////////////////////////////////////////////////
+enum class ECasTrackAccessType : uint8
+{
+	Always,
+	Newer,
+	Granular
+};
+
+///////////////////////////////////////////////////////////////////////////////
+struct FCasSnapshot;
 struct FCas
 {
 	static constexpr uint32		DeleteBlockMaxWaitTimeMs = 10000;
+	static constexpr int64		DirtyTimestampMask = std::numeric_limits<int64>::lowest(); // sign bit
 
 	using FLookup				= TMap<FCasAddr, FCasLocation>;
 	using FReadHandles			= TMap<FCasBlockId, FWeakAsyncFileHandle>;
 	using FLastAccess			= TMap<FCasBlockId, int64>;
 	using FBlockIdHandleCounts	= TMap<FCasBlockId, int32>;
 
+	FCas(const FOnDemandInstallCacheConfig& Config);
+
 	void						Lock()		{ Mutex.Lock(); }
 	void						Unlock()	{ Mutex.Unlock(); }
 
 	FIoStatus					Initialize(FStringView Directory, bool bDeleteExisting = false);
 	FCasLocation				FindChunk(const FIoHash& Hash) const;
-	FCasBlockId					CreateBlock(int64 UtcTicks);
-	FCasBlockId					CreateBlock() { return CreateBlock(FDateTime::UtcNow().GetTicks()); }
+	FCasBlockId					CreateBlock();
 	FIoStatus					DeleteBlock(FCasBlockId BlockId, TArray<FCasAddr>& OutAddrs);
 	FString						GetBlockFilename(FCasBlockId BlockId) const;
 	FSharedFileOpenResult		OpenRead(FCasBlockId BlockId);
 	FSharedFileOpenAsyncResult	OpenAsyncRead(FCasBlockId BlockId);
 	void						OnFileHandleDeleted(FCasBlockId BlockId);
 	FUniqueFileHandle			OpenWrite(FCasBlockId BlockId) const;
-	void						TrackAccess(FCasBlockId BlockId, int64 UtcTicks);
-	void						TrackAccess(FCasBlockId BlockId) { TrackAccess(BlockId, FDateTime::UtcNow().GetTicks()); }
-	bool						TrackAccessIfNewer(FCasBlockId BlockId, int64 UtcTicks);
+	bool						TrackAccessIf(ECasTrackAccessType Type, FCasBlockId BlockId, int64 UtcTicks, bool bDirty);
+	bool						TrackAccessIf(ECasTrackAccessType Type, FCasBlockId BlockId, bool bDirty)
+	{
+		return TrackAccessIf(Type, BlockId, FDateTime::UtcNow().GetTicks(), bDirty);
+	};
 	uint64						GetBlockInfo(FCasBlockInfoMap& OutBlockInfo);
 	void						Compact();
 	FIoStatus					Verify(TArray<FCasAddr>& OutAddrs);
+	void						ReplaySnapshot(FCasSnapshot&& Snapshot);
+	FLastAccess					ConsumeLastAcccess()
+	{
+		FLastAccess Temp;
+		{
+			TUniqueLock Lock(Mutex);
+			Temp = MoveTemp(LastAccess);
+		}
+		return Temp;
+	}
+	// Returns the timestamps that are "dirty" and need to be flushed to disk and clears dirty flag
+	FCas::FLastAccess			GetAndClearDirtyLastAccess();
 
 	// Returns the place holder timestamp to use if a timestamp is not found in the LastAccess table
 	static int64			GetTimestampForMissingLastAccess();
+	static uint32			GetMaxBlockSize()
+	{
+		const uint32 MaxBlockSize = 32 << 20; //TODO: Make configurable	
+		return MaxBlockSize;
+	}
+
+	static uint32			GetMinBlockSize()
+	{
+		const uint32 MinBlockSize = 32 << 19; //TODO: Make configurable	
+		return MinBlockSize;
+	}
 
+private:
 	FStringView				RootDirectory;
+public: // TODO: FIXME: encapsulate these members
 	FLookup					Lookup;
 	FBlockIdHandleCounts	BlockIds;
+private:
 	FLastAccess				LastAccess;
 	FReadHandles			ReadHandles;
 	FEventRef				BlockReadsDoneEvent;
-	const uint32			MaxBlockSize = 32 << 20; //TODO: Make configurable	
-	const uint32			MinBlockSize = MaxBlockSize >> 1; //TODO: Make configurable	
+	const int64				LastAccessGranularityTicks;
 	mutable UE::FMutex		Mutex;
 };
 
+///////////////////////////////////////////////////////////////////////////////
+FCas::FCas(const FOnDemandInstallCacheConfig& Config)
+	: LastAccessGranularityTicks(FTimespan::FromSeconds(Config.LastAccessGranularitySeconds).GetTicks())
+{
+}
+
 ///////////////////////////////////////////////////////////////////////////////
 FIoStatus FCas::Initialize(FStringView Directory, bool bDeleteExisting)
 {
@@ -266,7 +310,7 @@ FCasLocation FCas::FindChunk(const FIoHash& Hash) const
 	return FCasLocation{};
 }
 
-FCasBlockId FCas::CreateBlock(int64 UtcTicks)
+FCasBlockId FCas::CreateBlock()
 {
 	IPlatformFile&	Ipf = FPlatformFileManager::Get().GetPlatformFile();
 	FCasBlockId		Out = FCasBlockId::Invalid;
@@ -289,7 +333,6 @@ FCasBlockId FCas::CreateBlock(int64 UtcTicks)
 		}
 
 		BlockIds.Add(BlockId, 0);
-		LastAccess.FindOrAdd(BlockId, UtcTicks);
 		Out = BlockId;
 	}
 
@@ -449,20 +492,50 @@ FUniqueFileHandle FCas::OpenWrite(FCasBlockId BlockId) const
 	return FUniqueFileHandle(Ipf.OpenWrite(*Filename, bAppend, bAllowRead));
 }
 
-void FCas::TrackAccess(FCasBlockId BlockId, int64 UtcTicks)
+bool FCas::TrackAccessIf(const ECasTrackAccessType Type, const FCasBlockId BlockId, int64 UtcTicks, bool bDirty)
 {
 	check(BlockId.IsValid());
-	UE::TUniqueLock Lock(Mutex);
-	LastAccess.FindOrAdd(BlockId, UtcTicks);
-}
 
-bool FCas::TrackAccessIfNewer(FCasBlockId BlockId, int64 UtcTicks)
-{
-	check(BlockId.IsValid());
+	const uint32 BlockIdHash = GetTypeHash(BlockId);
+
 	UE::TUniqueLock Lock(Mutex);
-	int64& FoundTicks = LastAccess.FindOrAdd(BlockId, FDateTime::MinValue().GetTicks());
-	if (FoundTicks < UtcTicks)
+
+	int64* MaybeFoundTicks = LastAccess.FindByHash(BlockIdHash, BlockId);
+	if (MaybeFoundTicks == nullptr)
+	{
+		if (bDirty)
+		{
+			UtcTicks |= DirtyTimestampMask;
+		}
+
+		LastAccess.AddByHash(BlockIdHash, BlockId, UtcTicks);
+
+		return true;
+	}
+
+	int64& FoundTicks = *MaybeFoundTicks;
+	const int64 PrevTicks = (FoundTicks & ~DirtyTimestampMask);
+	bDirty = bDirty || (FoundTicks & DirtyTimestampMask); // Don't clear dirty flag if already set
+
+	bool bUpdate = true;
+	switch (Type)
 	{
+	case ECasTrackAccessType::Newer:
+		bUpdate = PrevTicks < UtcTicks;
+		break;
+
+	case ECasTrackAccessType::Granular:
+		bUpdate = PrevTicks < UtcTicks && (UtcTicks - PrevTicks > LastAccessGranularityTicks);
+		break;
+	}
+
+	if (bUpdate)
+	{
+		if (bDirty)
+		{
+			UtcTicks |= DirtyTimestampMask;
+		}
+
 		FoundTicks = UtcTicks;
 		return true;
 	}
@@ -662,6 +735,25 @@ FIoStatus FCas::Verify(TArray<FCasAddr>& OutAddrs)
 	return Status; 
 }
 
+FCas::FLastAccess FCas::GetAndClearDirtyLastAccess()
+{
+	FCas::FLastAccess DirtyLastAccess;
+	{
+		TUniqueLock Lock(Mutex);
+		for (TPair<FCasBlockId, int64>& Kv : LastAccess)
+		{
+			int64& Timestamp = Kv.Value;
+			if (Timestamp & FCas::DirtyTimestampMask)
+			{
+				DirtyLastAccess.Add(Kv.Key, Timestamp);
+				Timestamp &= ~FCas::DirtyTimestampMask;
+			}
+		}
+	}
+
+	return DirtyLastAccess;
+}
+
 int64 FCas::GetTimestampForMissingLastAccess()
 {
 	const FDateTime Now = FDateTime::UtcNow();
@@ -783,7 +875,20 @@ struct FCasJournal
 };
 
 ///////////////////////////////////////////////////////////////////////////////
-const TCHAR* GetErrorText(FCasJournal::EErrorCode ErrorCode)
+static void JournalLastAccess(FCasJournal::FTransaction& Transaction, const FCas::FLastAccess& LastAccess)
+{
+	for (const TPair<FCasBlockId, int64>& Kv : LastAccess)
+	{
+		int64 Timestamp = Kv.Value;
+		if (Timestamp & FCas::DirtyTimestampMask)
+		{
+			Transaction.BlockAccess(Kv.Key, Timestamp & ~FCas::DirtyTimestampMask);
+		}
+	}
+}
+
+///////////////////////////////////////////////////////////////////////////////
+static const TCHAR* GetErrorText(FCasJournal::EErrorCode ErrorCode)
 {
 	switch (ErrorCode)
 	{
@@ -1413,6 +1518,26 @@ TIoStatusOr<int64> FCasSnapshot::TryCreateAndResetJournal(const FString& Snapsho
 	return SnapshotSize;
 }
 
+///////////////////////////////////////////////////////////////////////////////
+void FCas::ReplaySnapshot(FCasSnapshot&& Snapshot)
+{
+	TUniqueLock Lock(Mutex);
+
+	Lookup.Reserve(Snapshot.ChunkLocations.Num());
+	for (TPair<FCasAddr, FCasLocation>& Kv : Snapshot.ChunkLocations)
+	{
+		Lookup.Add(MoveTemp(Kv));
+	}
+
+	BlockIds.Reserve(Snapshot.Blocks.Num());
+	LastAccess.Reserve(Snapshot.Blocks.Num());
+	for (const FCasSnapshot::FBlock& Block : Snapshot.Blocks)
+	{
+		BlockIds.Add(Block.BlockId, 0);
+		LastAccess.Add(Block.BlockId, Block.LastAccess);
+	}
+}
+
 ///////////////////////////////////////////////////////////////////////////////
 class FOnDemandInstallCache final
 	: public IOnDemandInstallCache 
@@ -1547,6 +1672,7 @@ public:
 	virtual FIoStatus							DefragAll(const uint64* BytesToFree = nullptr) override;
 	virtual FIoStatus							Verify() override;
 	virtual FIoStatus							Flush() override;
+	virtual FIoStatus							FlushLastAccess() override;
 	virtual FOnDemandInstallCacheUsage			GetCacheUsage() override;
 
 private:
@@ -1597,13 +1723,14 @@ private:
 FOnDemandInstallCache::FOnDemandInstallCache(const FOnDemandInstallCacheConfig& Config, FOnDemandIoStore& InIoStore)
 	: IoStore(InIoStore)
 	, CacheDirectory(Config.RootDirectory)
+	, Cas(Config)
 	, MaxCacheSize(Config.DiskQuota)
 	, MaxJournalSize(Config.JournalMaxSize)
 {
 	UE_LOG(LogIoStoreOnDemand, Log, TEXT("Initializing install cache, MaxCacheSize=%.2lf MiB, MaxJournalSize=%.2lf KiB"),
 		ToMiB(MaxCacheSize), ToKiB(MaxJournalSize));
 
-	const uint64 MinDiskQuota = 2 * Cas.MaxBlockSize;
+	const uint64 MinDiskQuota = 2 * Cas.GetMaxBlockSize();
 	if (MaxCacheSize < MinDiskQuota)
 	{
 		UE_LOG(LogIoStoreOnDemand, Error, TEXT("Failed to initialize install cache - disk quota must be at least %.2lf MiB"), ToMiB(MinDiskQuota));
@@ -1611,7 +1738,7 @@ FOnDemandInstallCache::FOnDemandInstallCache(const FOnDemandInstallCacheConfig&
 	}
 
 	// Reserve one block of space for defragmentation overhead
-	MaxCacheSize -= Cas.MaxBlockSize;
+	MaxCacheSize -= Cas.GetMaxBlockSize();
 	UE_LOG(LogIoStoreOnDemand, Log, TEXT("Effective MaxCacheSize without defragmentation space is MaxCacheSize=%.2lf MiB"), ToMiB(MaxCacheSize));
 
 	FIoStatus Status = Cas.Initialize(CacheDirectory);
@@ -1632,65 +1759,61 @@ FOnDemandInstallCache::FOnDemandInstallCache(const FOnDemandInstallCacheConfig&
 			UE_LOG(LogIoStoreOnDemand, Log, TEXT("Loaded CAS snapshot '%s' %.2lf KiB with %d blocks and %d chunk locations"),
 					*SnapshotFile, ToKiB(SnapshotSize), Snapshot.Blocks.Num(), Snapshot.ChunkLocations.Num());
 
-			Cas.Lookup.Reserve(Snapshot.ChunkLocations.Num());
-			for (TPair<FCasAddr, FCasLocation>& Kv : Snapshot.ChunkLocations)
-			{
-				Cas.Lookup.Add(MoveTemp(Kv));
-			}
-
-			Cas.BlockIds.Reserve(Snapshot.Blocks.Num());
-			Cas.LastAccess.Reserve(Snapshot.Blocks.Num());
-			for (const FCasSnapshot::FBlock& Block : Snapshot.Blocks)
-			{
-				Cas.BlockIds.Add(Block.BlockId, 0);
-				Cas.LastAccess.Add(Block.BlockId, Block.LastAccess);
-			}
+			Cas.ReplaySnapshot(MoveTemp(Snapshot));
+		}
+		else if (SnapshotStatus.Status().GetErrorCode() != EIoErrorCode::NotFound) // Its ok for the snapshot to not exist, other errors are significant
+		{
+			Status = SnapshotStatus.Status();
 		}
 	}
 
-	// Replay the journal 
 	const FString JournalFile = GetJournalFilename();
-	Status = FCasJournal::Replay(JournalFile, [this](const FCasJournal::FEntry& JournalEntry)
+	if (Status.IsOk())
 	{
-		switch(JournalEntry.Type())
+		// Replay the journal 
+		Status = FCasJournal::Replay(JournalFile, [this](const FCasJournal::FEntry& JournalEntry)
 		{
-		case FCasJournal::FEntry::EType::ChunkLocation:
-		{
-			const FCasJournal::FEntry::FChunkLocation& ChunkLocation = JournalEntry.ChunkLocation;
-			if (ChunkLocation.CasLocation.IsValid())
+			switch(JournalEntry.Type())
 			{
-				FCasLocation& Loc = Cas.Lookup.FindOrAdd(ChunkLocation.CasAddr);
-				Loc = ChunkLocation.CasLocation;
+			case FCasJournal::FEntry::EType::ChunkLocation:
+			{
+				const FCasJournal::FEntry::FChunkLocation& ChunkLocation = JournalEntry.ChunkLocation;
+				if (ChunkLocation.CasLocation.IsValid())
+				{
+					FCasLocation& Loc = Cas.Lookup.FindOrAdd(ChunkLocation.CasAddr);
+					Loc = ChunkLocation.CasLocation;
+				}
+				else
+				{
+					Cas.Lookup.Remove(ChunkLocation.CasAddr);
+				}
+				break;
 			}
-			else
+			case FCasJournal::FEntry::EType::BlockCreated:
+			{
+				const FCasJournal::FEntry::FBlockOperation& Op = JournalEntry.BlockOperation;
+				CurrentBlock = Op.BlockId;
+				Cas.BlockIds.Add(Op.BlockId, 0);
+				break;
+			}
+			case FCasJournal::FEntry::EType::BlockDeleted:
 			{
-				Cas.Lookup.Remove(ChunkLocation.CasAddr);
+				const FCasJournal::FEntry::FBlockOperation& Op = JournalEntry.BlockOperation;
+				Cas.BlockIds.Remove(Op.BlockId);
+				FCasBlockId MaybeCurrentBlock = Op.BlockId;
+				CurrentBlock.compare_exchange_strong(MaybeCurrentBlock, FCasBlockId::Invalid);
+				break;
 			}
-			break;
-		}
-		case FCasJournal::FEntry::EType::BlockCreated:
-		{
-			const FCasJournal::FEntry::FBlockOperation& Op = JournalEntry.BlockOperation;
-			CurrentBlock = Op.BlockId;
-			Cas.BlockIds.Add(Op.BlockId, 0);
-			break;
-		}
-		case FCasJournal::FEntry::EType::BlockDeleted:
-		{
-			const FCasJournal::FEntry::FBlockOperation& Op = JournalEntry.BlockOperation;
-			Cas.BlockIds.Remove(Op.BlockId);
-			FCasBlockId MaybeCurrentBlock = Op.BlockId;
-			CurrentBlock.compare_exchange_strong(MaybeCurrentBlock, FCasBlockId::Invalid);
-			break;
-		}
-		case FCasJournal::FEntry::EType::BlockAccess:
-		{
-			const FCasJournal::FEntry::FBlockOperation& Op = JournalEntry.BlockOperation;
-			Cas.TrackAccess(Op.BlockId, Op.UtcTicks);
-			break;
-		}
-		};
-	});
+			case FCasJournal::FEntry::EType::BlockAccess:
+			{
+				const FCasJournal::FEntry::FBlockOperation& Op = JournalEntry.BlockOperation;
+				constexpr const bool bDirty = false;
+				check(Cas.TrackAccessIf(ECasTrackAccessType::Always, Op.BlockId, Op.UtcTicks, bDirty));
+				break;
+			}
+			};
+		});
+	}
 
 	// Initializing the cache for the first time
 	if (Status.GetErrorCode() == EIoErrorCode::NotFound)
@@ -1749,19 +1872,11 @@ void FOnDemandInstallCache::Initialize(FSharedBackendContextRef Context)
 
 void FOnDemandInstallCache::Shutdown()
 {
-	FCas::FLastAccess LastAccess;
-	{
-		TUniqueLock Lock(Cas.Mutex);
-		LastAccess = MoveTemp(Cas.LastAccess);
-	}
+	FCas::FLastAccess LastAccess = Cas.ConsumeLastAcccess();
 
 	const FString JournalFile = GetJournalFilename();
 	FCasJournal::FTransaction Transaction = FCasJournal::Begin(JournalFile);
-	for (const TPair<FCasBlockId, int64>& Kv : LastAccess)
-	{
-		// TODO: this is only set on shutdown so if we crash we won't have good timestamps
-		Transaction.BlockAccess(Kv.Key, Kv.Value);
-	}
+	JournalLastAccess(Transaction, LastAccess);
 
 	if (FIoStatus Status = FCasJournal::Commit(MoveTemp(Transaction)); !Status.IsOk())
 	{
@@ -1913,7 +2028,11 @@ bool FOnDemandInstallCache::Resolve(FIoRequestImpl* Request)
 	}
 
 	TRACE_IOSTORE_BACKEND_REQUEST_STARTED(Request, this);
-	Cas.TrackAccess(CasLoc.BlockId);
+	constexpr const bool bDirty = true;
+	if (Cas.TrackAccessIf(ECasTrackAccessType::Granular, CasLoc.BlockId, bDirty))
+	{
+		IoStore.FlushLastAccess(nullptr);
+	}
 
 #if UE_ONDEMANDINSTALLCACHE_EXCLUSIVE_WRITE
 	const bool bIsLocationInCurrentBlock = CasLoc.BlockId == CurrentBlock;
@@ -2308,7 +2427,7 @@ FIoStatus FOnDemandInstallCache::Verify()
 	TArray<FChunkLookup>				ChunkLookups;
 
 	{
-		TUniqueLock Lock(Cas.Mutex);
+		TUniqueLock Lock(Cas);
 		ChunkLocations.Reserve(Cas.Lookup.Num());
 		for (const TPair<FCasAddr, FCasLocation>& Kv : Cas.Lookup)
 		{
@@ -2800,7 +2919,7 @@ FIoStatus FOnDemandInstallCache::Defrag(
 					bPossibleToFreeBytes = true;
 				}
 			}
-			else if (Info.FileSize < Cas.MinBlockSize)
+			else if (Info.FileSize < Cas.GetMinBlockSize())
 			{
 				// Block is too small whether or not its fragmented
 				if (ensure(Info.RefSize <= Info.FileSize))
@@ -2837,7 +2956,7 @@ FIoStatus FOnDemandInstallCache::Defrag(
 
 				BlocksToDefrag.Add(FDefragBlock{ .BlockId = BlockId, .LastAccess = Info.LastAccess });
 			}
-			else if (Info.FileSize < Cas.MinBlockSize)
+			else if (Info.FileSize < Cas.GetMinBlockSize())
 			{
 				// Block is too small whether or not its fragmented
 				if (ensure(Info.RefSize <= Info.FileSize))
@@ -3021,6 +3140,30 @@ FIoStatus FOnDemandInstallCache::Flush()
 	return InstallCacheStatusBuilder(EIoErrorCode::Ok);
 }
 
+FIoStatus FOnDemandInstallCache::FlushLastAccess()
+{
+	const FCas::FLastAccess DirtyLastAccess = Cas.GetAndClearDirtyLastAccess();
+
+	if (DirtyLastAccess.IsEmpty())
+	{
+		return InstallCacheStatusBuilder(EIoErrorCode::Ok);
+	}
+
+	const FString JournalFile = GetJournalFilename();
+	FCasJournal::FTransaction Transaction = FCasJournal::Begin(JournalFile);
+	JournalLastAccess(Transaction, DirtyLastAccess);
+
+	if (FIoStatus Status = FCasJournal::Commit(MoveTemp(Transaction)); !Status.IsOk())
+	{
+		UE_LOG(LogIoStoreOnDemand, Error, TEXT("Failed to update CAS journal '%s' with block timestamp(s), reason '%s'"),
+			*JournalFile, *Status.ToString());
+
+		return Status;
+	}
+
+	return InstallCacheStatusBuilder(EIoErrorCode::Ok);
+}
+
 FOnDemandInstallCacheUsage FOnDemandInstallCache::GetCacheUsage()
 {
 	// If this is called from a thread other than the OnDemandIoStore tick thread
@@ -3102,14 +3245,7 @@ FIoStatus FOnDemandInstallCache::FlushPendingChunksImpl(FPendingChunks& Chunks,
 		
 		if (CurrentBlockId.IsValid() == false)
 		{
-			if (UtcAccessTicks)
-			{
-				CurrentBlockId = Cas.CreateBlock(UtcAccessTicks);
-			}
-			else
-			{
-				CurrentBlockId = Cas.CreateBlock();
-			}
+			CurrentBlockId = Cas.CreateBlock();
 			ensure(CurrentBlockId.IsValid());
 			CurrentBlock = CurrentBlockId;
 			Transaction.BlockCreated(CurrentBlockId);
@@ -3130,7 +3266,7 @@ FIoStatus FOnDemandInstallCache::FlushPendingChunksImpl(FPendingChunks& Chunks,
 
 		while (Chunks.IsEmpty() == false)
 		{
-			if (CasBlockOffset > 0 && CasBlockOffset + Ar.Tell() + Chunks.Chunks[0].GetSize() > Cas.MaxBlockSize)
+			if (CasBlockOffset > 0 && CasBlockOffset + Ar.Tell() + Chunks.Chunks[0].GetSize() > Cas.GetMaxBlockSize())
 			{
 				break;
 			}
@@ -3150,9 +3286,10 @@ FIoStatus FOnDemandInstallCache::FlushPendingChunksImpl(FPendingChunks& Chunks,
 					<< TEXT("Failed to serialize chunks to cache block");
 			}
 
+			constexpr const bool bDirty = false;
 			if (UtcAccessTicks)
 			{
-				if (Cas.TrackAccessIfNewer(CurrentBlockId, UtcAccessTicks))
+				if (Cas.TrackAccessIf(ECasTrackAccessType::Newer, CurrentBlockId, UtcAccessTicks, bDirty))
 				{
 					Transaction.BlockAccess(CurrentBlockId, UtcAccessTicks);
 				}
@@ -3160,7 +3297,7 @@ FIoStatus FOnDemandInstallCache::FlushPendingChunksImpl(FPendingChunks& Chunks,
 			else
 			{
 				const int64 Now = FDateTime::UtcNow().GetTicks();
-				Cas.TrackAccess(CurrentBlockId, Now);
+				check(Cas.TrackAccessIf(ECasTrackAccessType::Always, CurrentBlockId, Now, bDirty));
 				Transaction.BlockAccess(CurrentBlockId, Now);
 			}
 
diff --git a/Engine/Source/Runtime/Experimental/IoStore/OnDemand/Private/OnDemandInstallCache.h b/Engine/Source/Runtime/Experimental/IoStore/OnDemand/Private/OnDemandInstallCache.h
--- a/Engine/Source/Runtime/Experimental/IoStore/OnDemand/Private/OnDemandInstallCache.h
+++ b/Engine/Source/Runtime/Experimental/IoStore/OnDemand/Private/OnDemandInstallCache.h
@@ -40,6 +40,7 @@ public:
 	virtual FIoStatus							DefragAll(const uint64* BytesToFree = nullptr) = 0;
 	virtual FIoStatus							Verify() = 0;
 	virtual FIoStatus							Flush() = 0;
+	virtual FIoStatus							FlushLastAccess() = 0;
 	virtual FOnDemandInstallCacheUsage			GetCacheUsage() = 0;
 };
 
@@ -48,6 +49,7 @@ struct FOnDemandInstallCacheConfig
 	FString RootDirectory;
 	uint64	DiskQuota		= 1ull << 30;
 	uint64	JournalMaxSize	= 2ull << 20;
+	double	LastAccessGranularitySeconds = 60 * 60;
 	bool	bDropCache		= false;
 };
 
diff --git a/Engine/Source/Runtime/Experimental/IoStore/OnDemand/Private/OnDemandIoStore.cpp b/Engine/Source/Runtime/Experimental/IoStore/OnDemand/Private/OnDemandIoStore.cpp
--- a/Engine/Source/Runtime/Experimental/IoStore/OnDemand/Private/OnDemandIoStore.cpp
+++ b/Engine/Source/Runtime/Experimental/IoStore/OnDemand/Private/OnDemandIoStore.cpp
@@ -1782,6 +1782,12 @@ TArray<FSharedOnDemandContainer> FOnDemandIoStore::GetContainers(EOnDemandContai
 	return Out;
 }
 
+void FOnDemandIoStore::FlushLastAccess(FOnDemandFlushLastAccessCompleted&& OnCompleted)
+{
+	LLM_SCOPE_BYNAME(TEXT("FileSystem/OnDemandIoStore"));
+	Installer->EnqueueFlushLastAccessRequest(MoveTemp(OnCompleted));
+}
+
 FIoStatus FOnDemandIoStore::GetContainersAndPackagesForInstall(
 	FStringView MountId,
 	const TArray<FString>& TagSets,
diff --git a/Engine/Source/Runtime/Experimental/IoStore/OnDemand/Private/OnDemandIoStore.h b/Engine/Source/Runtime/Experimental/IoStore/OnDemand/Private/OnDemandIoStore.h
--- a/Engine/Source/Runtime/Experimental/IoStore/OnDemand/Private/OnDemandIoStore.h
+++ b/Engine/Source/Runtime/Experimental/IoStore/OnDemand/Private/OnDemandIoStore.h
@@ -308,6 +308,19 @@ void FOnDemandChunkInfo::GetUrl(FAnsiStringBuilderBase& Url) const
 		<< "/" << HashString << ANSITEXTVIEW(".iochunk");
 }
 
+///////////////////////////////////////////////////////////////////////////////
+/** Result from flushing cache last access times. */
+struct FOnDemandFlushLastAccessResult
+{
+	/** The status of the request. */
+	FIoStatus Status;
+	/** Duration in seconds. */
+	double DurationInSeconds = 0.0;
+};
+
+/** Flush last access completion callback. */
+using FOnDemandFlushLastAccessCompleted = TUniqueFunction<void(FOnDemandFlushLastAccessResult)>;
+
 ///////////////////////////////////////////////////////////////////////////////
 class FOnDemandIoStore
 	: public TSharedFromThis<FOnDemandIoStore, ESPMode::ThreadSafe>
@@ -351,7 +364,7 @@ public:
 											FOnDemandInstallProgressed&& OnProgress = nullptr) override;
 	virtual void						Purge(FOnDemandPurgeArgs&& Args, FOnDemandPurgeCompleted&& OnCompleted) override;
 	virtual void						Defrag(FOnDemandDefragArgs&& Args, FOnDemandDefragCompleted&& OnCompleted) override;
-	virtual void						Verify(FOnDemandVerifyCacheCompleted&& OnCompleted) override;
+	virtual void						Verify(FOnDemandVerifyCacheCompleted&& OnCompleted) override;	
 	virtual FIoStatus					Unmount(FStringView MountId) override;
 	virtual TIoStatusOr<FOnDemandInstallSizeResult> GetInstallSize(const FOnDemandGetInstallSizeArgs& Args) const override;
 	virtual FIoStatus					GetInstallSizesByMountId(const FOnDemandGetInstallSizeArgs& Args, TMap<FString, uint64>& OutSizesByMountId) const override;
@@ -364,7 +377,9 @@ public:
 	virtual void						CancelInstallRequest(FSharedInternalInstallRequest InstallRequest) override;
 	virtual void						UpdateInstallRequestPriority(FSharedInternalInstallRequest InstallRequest, int32 NewPriority) override;
 	virtual void						ReleaseContent(FOnDemandInternalContentHandle& ContentHandle) override;
+
 	TArray<FSharedOnDemandContainer>	GetContainers(EOnDemandContainerFlags ContainerFlags = EOnDemandContainerFlags::None) const;
+	void								FlushLastAccess(FOnDemandFlushLastAccessCompleted&& OnCompleted);
 
 private:
 	friend class FOnDemandContentInstaller;
