Geometry collection : add a way to query a leaf transform from a line trace and expose it to blueprint ( Function name : FindLeafTransformByLineTrace )

diff --git a/Engine/Source/Runtime/Experimental/GeometryCollectionEngine/Private/GeometryCollection/GeometryCollectionComponent.cpp b/Engine/Source/Runtime/Experimental/GeometryCollectionEngine/Private/GeometryCollection/GeometryCollectionComponent.cpp
--- a/Engine/Source/Runtime/Experimental/GeometryCollectionEngine/Private/GeometryCollection/GeometryCollectionComponent.cpp
+++ b/Engine/Source/Runtime/Experimental/GeometryCollectionEngine/Private/GeometryCollection/GeometryCollectionComponent.cpp
@@ -6,6 +6,7 @@
 #include "Async/ParallelFor.h"
 #include "Chaos/ChaosPhysicalMaterial.h"
 #include "Chaos/ChaosScene.h"
+#include "Chaos/PhysicsObjectCollisionInterface.h"
 #include "ChaosSolversModule.h"
 #include "ChaosStats.h"
 #include "ComponentRecreateRenderStateContext.h"
@@ -1408,6 +1409,46 @@ bool UGeometryCollectionComponent::IsFullyDecayed() const
 	return bAlreadyFullyDecayed;
 }
 
+bool UGeometryCollectionComponent::FindLeafTransformByLineTrace(const FVector Start, const FVector End, int32& OutLeafTransformIndex, FName& OutLeafTransformName) const
+{
+	OutLeafTransformIndex = INDEX_NONE;
+	OutLeafTransformName = NAME_None;
+
+	if (!PhysicsProxy)
+	{
+		return false;
+	}
+
+	TArray<Chaos::FPhysicsObjectHandle> Objects = GetAllPhysicsObjects();
+	FLockedReadPhysicsObjectExternalInterface Interface = FPhysicsObjectExternalInterface::LockRead(Objects);
+	Objects = Objects.FilterByPredicate(
+		[&Interface](Chaos::FPhysicsObjectHandle Handle)
+		{
+
+			return !Interface->HasChildren(Handle);
+		}
+	);
+
+	bool bHaveHit = false;
+
+	Chaos::FPhysicsObjectCollisionInterface_External CollisionInterface{ Interface.GetInterface() };
+	ChaosInterface::FRaycastHit BestHit;
+	if (CollisionInterface.LineTrace(Objects, Start, End, /*bTraceComplex*/false, BestHit))
+	{
+		if (BestHit.Actor)
+		{
+			if (const Chaos::FPBDRigidParticle* RigiParticle = BestHit.Actor->CastToRigidParticle())
+			{
+				const FGeometryCollectionItemIndex ItemIndex = PhysicsProxy->GetItemIndexFromGTParticleNoInternalCluster_External(RigiParticle);
+				OutLeafTransformIndex = ItemIndex.GetItemIndex();
+				OutLeafTransformName = PhysicsProxy->GetTransformName_External(ItemIndex);
+				bHaveHit = true;
+			}
+		}
+	}
+	return bHaveHit;
+}
+
 void UGeometryCollectionComponent::DispatchBreakEvent(const FChaosBreakEvent& Event)
 {
 	// native
diff --git a/Engine/Source/Runtime/Experimental/GeometryCollectionEngine/Public/GeometryCollection/GeometryCollectionComponent.h b/Engine/Source/Runtime/Experimental/GeometryCollectionEngine/Public/GeometryCollection/GeometryCollectionComponent.h
--- a/Engine/Source/Runtime/Experimental/GeometryCollectionEngine/Public/GeometryCollection/GeometryCollectionComponent.h
+++ b/Engine/Source/Runtime/Experimental/GeometryCollectionEngine/Public/GeometryCollection/GeometryCollectionComponent.h
@@ -1347,6 +1347,17 @@ public:
 
 	bool IsForceBrokenForCustomRenderer() const { return bForceBrokenForCustomRenderer; }
 
+	/**
+	* Find a leaf transform from a line trace
+	* @param Start Start position for the line trace
+	* @param End position for the line trace
+	* @param LeafTransformIndex the index of leaf transform found by the line trace
+	* @param LeafTransformName the name of leaf transform found by the line trace
+	* @return true if a leaf was found
+	*/
+	UFUNCTION(BlueprintCallable, Category = "Physics")
+	bool FindLeafTransformByLineTrace(const FVector Start, const FVector End, int32& LeafTransformIndex, FName& LeafTransformName) const;
+
 	/** 
 	* blueprint function to set a specific root proxy local transform 
 	* warning: when called from C++ and calling it on multiple root proxies it is recommended to use SetRootProxyLocalTransform instead and then call RefreshCustomRenderer
