Fix code analysis warning

diff --git a/Engine/Source/Runtime/Experimental/GeometryCollectionEngine/Private/GeometryCollection/GeometryCollectionComponent.cpp b/Engine/Source/Runtime/Experimental/GeometryCollectionEngine/Private/GeometryCollection/GeometryCollectionComponent.cpp
--- a/Engine/Source/Runtime/Experimental/GeometryCollectionEngine/Private/GeometryCollection/GeometryCollectionComponent.cpp
+++ b/Engine/Source/Runtime/Experimental/GeometryCollectionEngine/Private/GeometryCollection/GeometryCollectionComponent.cpp
@@ -1425,7 +1425,7 @@ bool UGeometryCollectionComponent::FindLeafTransformByLineTrace(const FVector St
 		[&Interface](Chaos::FPhysicsObjectHandle Handle)
 		{
 
-			return !Interface->HasChildren(Handle);
+			return !Interface.GetInterface().HasChildren(Handle);
 		}
 	);
 
