Overflow fix

#rnx

diff --git a/Engine/Source/Runtime/Experimental/IoStore/HttpClient/Private/TransactionOne.inl b/Engine/Source/Runtime/Experimental/IoStore/HttpClient/Private/TransactionOne.inl
--- a/Engine/Source/Runtime/Experimental/IoStore/HttpClient/Private/TransactionOne.inl
+++ b/Engine/Source/Runtime/Experimental/IoStore/HttpClient/Private/TransactionOne.inl
@@ -718,6 +718,7 @@ FOutcome FBody::Prologue(FMutableMemoryView Dest, FHttpPeer& Peer)
 
 		int32 Delta = Size - HexDigLen;
 
+		if (HexDigLen >= HeaderBufSize - 2)	return FOutcome::Error("Chunk size too large");
 		if (HexDigLen == 0)					return FOutcome::Error("Invalid chunk header");
 		if (Delta < 1)						return FOutcome::Waiting();
 		if (Cursor[HexDigLen] != '\r')		return FOutcome::Error("Chunk extensions are not supported (ERREXT)");
@@ -751,34 +752,40 @@ FOutcome FBody::Prologue(FMutableMemoryView Dest, FHttpPeer& Peer)
 
 	if (HeaderLength == 0)
 	{
-		char* Cursor = (char*)Scratch.GetData();
+		FMutableMemoryView View = Scratch;
 
 		if (OverspillSize != 0)
 		{
-			check(OverspillSize <= HeaderBufSize);
-			std::memmove(Cursor, Overspill.GetData(), OverspillSize);
-			Cursor += OverspillSize;
+			if (OverspillSize > uint32(View.GetSize()))
+			{
+				return FOutcome::Error("Unexpectedly large overspill");
+			}
+			std::memmove(View.GetData(), Overspill.GetData(), OverspillSize);
+			View = View.Mid(OverspillSize);
 		}
 
-		Outcome = Peer.Recv(Cursor, HeaderBufSize);
-		if (!Outcome.IsOk())
+		if (!View.IsEmpty())
 		{
-			return Outcome;
+			char* Cursor = (char*)View.GetData();
+			Outcome = Peer.Recv(Cursor, uint32(View.GetSize()));
+			if (!Outcome.IsOk())
+			{
+				return Outcome;
+			}
+			View = Scratch.Left(OverspillSize + Outcome.GetResult());
 		}
 
-		Cursor -= OverspillSize;
-		OverspillSize += Outcome.GetResult();
-		Overspill = { (uint8*)Cursor, OverspillSize };
-
-		Outcome = FindHeader(Cursor, OverspillSize);
+		char* Cursor = (char*)View.GetData();
+		Outcome = FindHeader(Cursor, uint32(View.GetSize()));
 		if (!Outcome.IsOk())
 		{
 			return Outcome;
 		}
 
-		Header = Cursor;
+		check(View.GetData() == Scratch.GetData());
+		Header = (char*)Scratch.GetData();
 		HeaderLength = Outcome.GetResult();
-		Overspill = Overspill.Mid(HeaderLength);
+		Overspill = View.Mid(HeaderLength);
 	}
 
 	// Read chunk size
