Reworked parse to better handle messages with odd lengths

#rnx

diff --git a/Engine/Source/Runtime/Experimental/IoStore/HttpClient/Private/TransactionOne.inl b/Engine/Source/Runtime/Experimental/IoStore/HttpClient/Private/TransactionOne.inl
--- a/Engine/Source/Runtime/Experimental/IoStore/HttpClient/Private/TransactionOne.inl
+++ b/Engine/Source/Runtime/Experimental/IoStore/HttpClient/Private/TransactionOne.inl
@@ -489,43 +489,29 @@ FOutcome FStatusHeaders::Parse()
 		}
 	);
 
-	bHasBody = (StatusCode != 304);
-	bHasBody &= !IsHeadMethod();
-	bHasBody &= bChunked | (ContentLength > 0);
+	bHasBody = !IsHeadMethod();
 
-	if (Flags & Flag_Length)
-	{
-		if ((StatusCode <= 100) | (StatusCode == 204) | (Flags & Flag_XferEnc))
-		{
-			return FOutcome::Error("Unexpected Content-Length field", StatusCode);
-		}
-
-		if (ContentLength < 0)
-		{
-			return FOutcome::Error("Invalid Content-Length field", uint32(ContentLength));
-		}
-
-		if ((GetSize() - MessageRight) > ContentLength)
-		{
-			return FOutcome::Error("More data received that expected");
-		}
+	if ((StatusCode < 200) | (StatusCode == 204) | (StatusCode == 304)) // | !connect
+	{	
+		ContentLength = -1;
+		bHasBody = false;
 	}
 	else if (Flags & Flag_XferEnc)
 	{
-		if (!bChunked)
-		{
-			return FOutcome::Error("Unsupported Transfer-Encoding");
-		}
+		if (!bChunked)				return FOutcome::Error("Unsupported Transfer-Encoding");
+		if (Flags & Flag_Length)	return FOutcome::Error("Chunked yet with length");
 		ContentLength = -1;
+		bHasBody &= true;
 	}
-	else if (bHasBody)
+	else if (uint32 Overshoot = GetSize() - MessageRight; Flags & Flag_Length)
 	{
-		return FOutcome::Error("Missing Content-Length field");
-	}
-	else if (GetSize() > uint32(MessageRight))
-	{
-		return FOutcome::Error("Received content when none was expected");
+		if (ContentLength < 0)			return FOutcome::Error("Invalid Content-Length field", uint32(ContentLength));
+		if (Overshoot > ContentLength)	return FOutcome::Error("More data received that expected");
+
+		bHasBody &= (ContentLength > 0);
 	}
+	else if (bHasBody)	return FOutcome::Error("Missing Content-Length field");
+	else if (Overshoot)	return FOutcome::Error("Received content when none was expected");
 
 	return FOutcome::Ok();
 }
