Missed a shelf

#rnx

diff --git a/Engine/Source/Runtime/Experimental/IoStore/HttpClient/Private/TransactionOne.inl b/Engine/Source/Runtime/Experimental/IoStore/HttpClient/Private/TransactionOne.inl
--- a/Engine/Source/Runtime/Experimental/IoStore/HttpClient/Private/TransactionOne.inl
+++ b/Engine/Source/Runtime/Experimental/IoStore/HttpClient/Private/TransactionOne.inl
@@ -493,7 +493,7 @@ FOutcome FStatusHeaders::Parse()
 
 	if ((StatusCode < 200) | (StatusCode == 204) | (StatusCode == 304)) // | !connect
 	{	
-		ContentLength = -1;
+		ContentLength = 0;
 		bHasBody = false;
 	}
 	else if (Flags & Flag_XferEnc)
@@ -559,6 +559,11 @@ int64 FBody::GetRemaining() const
 ////////////////////////////////////////////////////////////////////////////////
 FOutcome FBody::TryRecv(FMutableMemoryView Dest, FHttpPeer& Peer)
 {
+	if (Dest.GetSize() == 0)
+	{
+		return FOutcome::Error("Empty destination");
+	}
+
 	FOutcome Outcome = FOutcome::None();
 
 	switch (State)
