Task-Based Physics Advancement
Threading graph between integratioon, kinematics update and broad phase update
A new data structure storing the data to update in the Broad Phase has been added to be ablae to run milti threaded code safely :
- Instead of a map to add all data, data are stored in batch that can be use in different thread
Another class has been created in order to split physics simulation code and task dispatch logic.
The Dispatcher ensure all tasks can safely use the data thanks to the task graph.
The performance gain is present in most of scenario up to 25% improvement.

diff --git a/Engine/Source/Runtime/Experimental/Chaos/Private/Chaos/PBDRigidsEvolution.cpp b/Engine/Source/Runtime/Experimental/Chaos/Private/Chaos/PBDRigidsEvolution.cpp
--- a/Engine/Source/Runtime/Experimental/Chaos/Private/Chaos/PBDRigidsEvolution.cpp
+++ b/Engine/Source/Runtime/Experimental/Chaos/Private/Chaos/PBDRigidsEvolution.cpp
@@ -23,9 +23,6 @@ FAutoConsoleVariableRef CVarChaosNumPushOutIterationsOverride(TEXT("p.ChaosNumPu
 int32 ChaosNumContactIterationsOverride = -1;
 FAutoConsoleVariableRef CVarChaosNumContactIterationsOverride(TEXT("p.ChaosNumContactIterationsOverride"), ChaosNumContactIterationsOverride, TEXT("Override for num contact iterations if >= 0. [def:-1]"));
 
-bool bPruningParticleToInsertInAccelerationStructure = true;
-FAutoConsoleVariableRef CVarPruningParticleToInsertInAccelerationStructure(TEXT("p.ChaosPruningParticleToInsertInAccelerationStructure"), bPruningParticleToInsertInAccelerationStructure, TEXT("Enables pruning particle before updating them in the accceleration structure. This is an optimization that reduce the dynamic acceleration structure update."));
-
 namespace Chaos
 {
 	namespace CVars
@@ -224,6 +221,7 @@ namespace Chaos
 		, bIsSingleThreaded(InIsSingleThreaded)
 		, bCanStartAsyncTasks(true)
 		, LatestExternalTimestampConsumed_Internal(-1)
+		, TaskDispatcher(Particles, InternalAccelerationQueue)
 		, bAccelerationStructureTaskStarted(nullptr)
 		, bAccelerationStructureTaskSignalKill(nullptr)
 		, SpatialCollectionFactory(new FDefaultCollectionFactory())
@@ -630,82 +628,48 @@ namespace Chaos
 		// The bounding volume associated to the particle will be updated in the acceleration structure.
 		QUICK_SCOPE_CYCLE_COUNTER(STAT_PruneInternalPendingData);
 
-		const int32 NumWorker = !!GSingleThreadedPhysics ? 1 : FMath::Max(FTaskGraphInterface::Get().GetNumWorkerThreads(), 1);
-		const int32 NumParticles = InternalAccelerationQueue.PendingData.Num();
-		const int32 ParticleByContext = FMath::Max(FMath::DivideAndRoundUp(NumParticles, NumWorker), SmallBatchSize);
-		const int32 NumBatch = FMath::DivideAndRoundUp(NumParticles, ParticleByContext);
 
-		if (Chaos::ShouldExecuteTasksInParallel(NumBatch))
+		auto HasToBeUpdated = [&AccelerationStructure](const FPendingSpatialData& PendingData) -> bool
 		{
-			TArray<TArray<FPendingSpatialData>> NewPendingDataArray;
-			NewPendingDataArray.SetNum(NumBatch);
-			TArray<UE::Tasks::FTask> PendingTasks;
-			PendingTasks.Reserve(NumBatch);
-			for (int32 BatchIndex = 0; BatchIndex < NumBatch; BatchIndex++)
-			{
-				const int32 FirstIndex = BatchIndex * ParticleByContext;
-				int32 LastIndex = (BatchIndex + 1) * ParticleByContext;
-				LastIndex = BatchIndex == NumBatch - 1 ? FMath::Min(NumParticles, LastIndex) : LastIndex;
-				TArray<FPendingSpatialData>& NewPendingData = NewPendingDataArray[BatchIndex];
-				UE::Tasks::FTask PendingTask = UE::Tasks::Launch(UE_SOURCE_LOCATION, [this, &AccelerationStructure, FirstIndex, LastIndex, &NewPendingData]()
-				{
-					QUICK_SCOPE_CYCLE_COUNTER(STAT_PruneInternalPendingDataTask);
-					for (int32 Index = FirstIndex; Index < LastIndex; Index++)
-					{
-						const FPendingSpatialData& PendingData = InternalAccelerationQueue.PendingData[Index];
-						FGeometryParticleHandle* UpdateParticle = PendingData.AccelerationHandle.GetGeometryParticleHandle_PhysicsThread();
-
-						if (PendingData.Operation == EPendingSpatialDataOperation::Delete ||
-							(UpdateParticle->HasBounds() && AccelerationStructure.NeedUpdateElementIn(UpdateParticle, UpdateParticle->WorldSpaceInflatedBounds(), PendingData.SpatialIdx)))
-						{
-							NewPendingData.Add(PendingData);
-						}
-					}
-				}, LowLevelTasks::ETaskPriority::High);
-				PendingTasks.Add(PendingTask);
-			}
+			FGeometryParticleHandle* UpdateParticle = PendingData.AccelerationHandle.GetGeometryParticleHandle_PhysicsThread();
+			return (PendingData.Operation == EPendingSpatialDataOperation::Delete ||
+				(UpdateParticle->HasBounds() && AccelerationStructure.NeedUpdateElementIn(UpdateParticle, UpdateParticle->WorldSpaceInflatedBounds(), PendingData.SpatialIdx)));
 
-			UE::Tasks::Wait(PendingTasks);
-			// Once all threads are finished merge all data in the acceleration queue
-			// Maybe this could be multi threaded, but since it is really fast not sure if it is worth it.
-			InternalAccelerationQueue.PendingData.Reset();
-			for (int32 BatchIndex = 0; BatchIndex < NumBatch; BatchIndex++)
-			{
-				InternalAccelerationQueue.PendingData.Append(NewPendingDataArray[BatchIndex]);
-			}
+		};
+
+		if (!GSingleThreadedPhysics)
+		{
+			TaskDispatcher.PruneInternalPendingData(HasToBeUpdated);
 		}
 		else
 		{
-			// Small enough task to run immediately and avoid parallel overhead
-			QUICK_SCOPE_CYCLE_COUNTER(STAT_PruneInternalPendingDataTask);
-
-			TArray<FPendingSpatialData> OldPendingData = MoveTemp(InternalAccelerationQueue.PendingData);
-
-			InternalAccelerationQueue.PendingData.Reset();
-
-			for (int32 Index = 0; Index < NumParticles; Index++)
+			InternalAccelerationQueue.CleanUpDuplicated();
+			for (TArray<FPendingSpatialData>& PendingDataArray : InternalAccelerationQueue.PendingDataArrays)
 			{
-				const FPendingSpatialData& PendingData = OldPendingData[Index];
-				FGeometryParticleHandle* UpdateParticle = PendingData.AccelerationHandle.GetGeometryParticleHandle_PhysicsThread();
-
-				if (PendingData.Operation == EPendingSpatialDataOperation::Delete ||
-					(UpdateParticle->HasBounds() && AccelerationStructure.NeedUpdateElementIn(UpdateParticle, UpdateParticle->WorldSpaceInflatedBounds(), PendingData.SpatialIdx)))
+				TArray<FPendingSpatialData> NewPendingData;
+				NewPendingData.Reserve(PendingDataArray.Num());
+				for (const FPendingSpatialData& PendingData : PendingDataArray)
 				{
-					InternalAccelerationQueue.PendingData.Add(PendingData);
+					if (HasToBeUpdated(PendingData))
+					{
+						NewPendingData.Add(PendingData);
+					}
 				}
+				PendingDataArray = MoveTemp(NewPendingData);
 			}
 		}
 	}
 
 	void FPBDRigidsEvolutionBase::FlushInternalAccelerationQueue()
 	{
-		if (bPruningParticleToInsertInAccelerationStructure)
+		PruneInternalPendingData(*InternalAcceleration);
+		// Could be linked by dependent task, but could start as soon a previous task is finished
+		for (const TArray<FPendingSpatialData>& PendingDataArray : InternalAccelerationQueue.PendingDataArrays)
 		{
-			PruneInternalPendingData(*InternalAcceleration);
-		}
-		for (const FPendingSpatialData& PendingData : InternalAccelerationQueue.PendingData)
-		{
-			ApplyParticlePendingData(PendingData, *InternalAcceleration, false, true);
+			for (const FPendingSpatialData& PendingData : PendingDataArray)
+			{
+				ApplyParticlePendingData(PendingData, *InternalAcceleration, false, true);
+			}
 		}
 		InternalAcceleration->SetSyncTimestamp(LatestExternalTimestampConsumed_Internal);
 		InternalAccelerationQueue.Reset();
@@ -713,25 +677,28 @@ namespace Chaos
 
 	void FPBDRigidsEvolutionBase::FlushAsyncAccelerationQueue()
 	{
-		for (const FPendingSpatialData& PendingData : AsyncAccelerationQueue.PendingData)
+		auto FlushAsyncAccelerationQueueLambda = [this]() 
 		{
-			ApplyParticlePendingData(PendingData, *AsyncInternalAcceleration, true, false); //only the first queue needs to update the cached acceleration
-			ApplyParticlePendingData(PendingData, *AsyncExternalAcceleration, false, false);
-		}
-
-				//NOTE: This assumes that we are never creating a PT particle that is replicated to GT
-				//At the moment that is true, and it seems like we have enough mechanisms to avoid this direction
-				//If we want to support that, the UniqueIndex must be kept around until GT goes away
-				//This is hard to do, but would probably mean the ownership of the index is in the proxy
-		for (FUniqueIdx UniqueIdx : UniqueIndicesPendingRelease)
-		{
-			ReleaseIdx(UniqueIdx);
-		}
-		UniqueIndicesPendingRelease.Reset();
-		AsyncAccelerationQueue.Reset();
+			for (const FPendingSpatialData& PendingData : AsyncAccelerationQueue.PendingData)
+			{
+				ApplyParticlePendingData(PendingData, *AsyncInternalAcceleration, true, false); //only the first queue needs to update the cached acceleration
+				ApplyParticlePendingData(PendingData, *AsyncExternalAcceleration, false, false);
+			}
 
-		AsyncInternalAcceleration->SetSyncTimestamp(LatestExternalTimestampConsumed_Internal);
-		AsyncExternalAcceleration->SetSyncTimestamp(LatestExternalTimestampConsumed_Internal);
+			//NOTE: This assumes that we are never creating a PT particle that is replicated to GT
+			//At the moment that is true, and it seems like we have enough mechanisms to avoid this direction
+			//If we want to support that, the UniqueIndex must be kept around until GT goes away
+			//This is hard to do, but would probably mean the ownership of the index is in the proxy
+			for (FUniqueIdx UniqueIdx : UniqueIndicesPendingRelease)
+			{
+				ReleaseIdx(UniqueIdx);
+			}
+			UniqueIndicesPendingRelease.Reset();
+			AsyncAccelerationQueue.Reset();
+			AsyncInternalAcceleration->SetSyncTimestamp(LatestExternalTimestampConsumed_Internal);
+			AsyncExternalAcceleration->SetSyncTimestamp(LatestExternalTimestampConsumed_Internal);
+		};
+		TaskDispatcher.FlushAccelerationQueue(FlushAsyncAccelerationQueueLambda);
 	}
 
 	//TODO: make static and _External suffix
@@ -816,7 +783,6 @@ namespace Chaos
 		SCOPE_CYCLE_COUNTER(STAT_ComputeIntermediateSpatialAcceleration);
 		CHAOS_SCOPED_TIMER(ComputeIntermediateSpatialAcceleration);
 
-		bool ForceFullBuild = InternalAccelerationQueue.Num() > AccelerationStructureTimeSlicingMaxQueueSizeBeforeForce;
 
 		if (!AccelerationStructureTaskComplete)
 		{
@@ -826,6 +792,8 @@ namespace Chaos
 			AsyncInternalAcceleration = GetFreeSpatialAcceleration_Internal();
 			AsyncExternalAcceleration = GetFreeSpatialAcceleration_Internal();
 			FlushInternalAccelerationQueue();
+			
+			TaskDispatcher.WaitAsyncQueueTask();
 
 			// Give game thread an empty structure to pop
 			ExternalStructuresQueue.Enqueue(AsyncExternalAcceleration);
@@ -854,7 +822,7 @@ namespace Chaos
 				bNeedsReset = true;
 
 				FlushAsyncAccelerationQueue();
-				
+
 				if (AccelerationStructureSplitStaticAndDynamic == 1)
 				{
 					FlushInternalAccelerationQueue();
@@ -864,13 +832,14 @@ namespace Chaos
 					//other queues are no longer needed since we've flushed all operations and now have a pristine structure
 					InternalAccelerationQueue.Reset();
 				}
+				TaskDispatcher.WaitAsyncQueueTask();
 
 				// At this point the InternalAcceleration and the AsyncAcceleration structures will all be consistent with each other
 				// (If FlushInternalAccelerationQueue was flushed above)
 				// but the AsyncAcceleration now probably have fewer dirty elements (they were just rebuilt)
-				
+
 				//swap acceleration structure for new one
-					std::swap(InternalAcceleration, AsyncInternalAcceleration);
+				std::swap(InternalAcceleration, AsyncInternalAcceleration);
 				// The old InternalAcceleration is now in AsyncInternalAcceleration, and the constituent structures will be reused if no changes have been detected
 
 				if (AccelerationStructureUseDynamicTree)
@@ -900,7 +869,9 @@ namespace Chaos
 			}
 			else
 			{
-					FlushInternalAccelerationQueue();
+				FlushInternalAccelerationQueue();
+
+				TaskDispatcher.WaitAsyncQueueTask();
 
 				if (GAccelerationStructureCacheOverlappingLeaves  && AccelerationStructureUseDynamicTree)
 				{
@@ -912,12 +883,14 @@ namespace Chaos
 			if (bCanStartAsyncTasks)
 			{
 				// we run the task for both starting a new accel structure as well as for the time-slicing
+				bool ForceFullBuild = AsyncAccelerationQueue.Num() > AccelerationStructureTimeSlicingMaxQueueSizeBeforeForce;
 				AccelerationStructureTaskComplete = TGraphTask<FChaosAccelerationStructureTask>::CreateTask().ConstructAndDispatchWhenReady(*SpatialCollectionFactory, SpatialAccelerationCache, AsyncInternalAcceleration, AsyncExternalAcceleration, ForceFullBuild, bIsSingleThreaded, bNeedsReset, &bAccelerationStructureTaskStarted, &bAccelerationStructureTaskSignalKill);
 			}
 		}
 		else
 		{
 			FlushInternalAccelerationQueue();
+			TaskDispatcher.WaitAsyncQueueTask();
 		}
 	}
 
diff --git a/Engine/Source/Runtime/Experimental/Chaos/Private/Chaos/PBDRigidsEvolutionGBF.cpp b/Engine/Source/Runtime/Experimental/Chaos/Private/Chaos/PBDRigidsEvolutionGBF.cpp
--- a/Engine/Source/Runtime/Experimental/Chaos/Private/Chaos/PBDRigidsEvolutionGBF.cpp
+++ b/Engine/Source/Runtime/Experimental/Chaos/Private/Chaos/PBDRigidsEvolutionGBF.cpp
@@ -573,6 +573,7 @@ void FPBDRigidsEvolutionGBF::AdvanceOneTimeStepImpl(const FReal Dt, const FSubSt
 			CSV_SCOPED_TIMING_STAT(PhysicsVerbose, StepSolver_ComputeIntermediateSpatialAcceleration);
 			Base::ComputeIntermediateSpatialAcceleration();
 		}
+		TaskDispatcher.WaitTaskEndSpatial(); 
 	}
 
 	// Collision detection is sensitive to duplication bugs in the particle views
@@ -824,30 +825,17 @@ void FPBDRigidsEvolutionGBF::Integrate(FReal Dt)
 	const bool bForceMACD = CVars::bChaosForceMACD;
 
 	const TArray<FPBDRigidParticleHandle*>& ParticlesArray = Particles.GetActiveParticlesArray();
-	const int32 NumParticles = ParticlesArray.Num();
-	int32 NumTasks = !!GSingleThreadedPhysics ? 1 : FMath::Min(FTaskGraphInterface::Get().GetNumWorkerThreads(), Chaos::MaxNumWorkers);
-	NumTasks = FMath::Max(FMath::Min(NumTasks, NumParticles), 1);
-	constexpr int32 MinParticlePerTask = 40;
-	const int32 ParticleByTask = FMath::Max(FMath::DivideAndRoundUp(NumParticles, NumTasks), MinParticlePerTask);
-	int32 NumBatches = FMath::DivideAndRoundUp(NumParticles, ParticleByTask);
 
-	IntegrationPendingTasks.Reset(NumBatches);
-	for (int32 BatchIndex = 0; BatchIndex < NumBatches; BatchIndex++)
-	{
-		const int32 StartIndex = BatchIndex * ParticleByTask;
-		int32 EndIndex = (BatchIndex + 1) * ParticleByTask;
-		EndIndex = FMath::Min(NumParticles, EndIndex);
-
-		UE::Tasks::FTask PendingTask = UE::Tasks::Launch(UE_SOURCE_LOCATION, [this, StartIndex, EndIndex, &ParticlesArray,
+	auto IntegrateWork = [this, &ParticlesArray,
 		Dt, HackMaxAngularSpeedSq, HackMaxLinearSpeedSq, bAllowMACD, bForceMACD, VelocityBoundsMultiplierMACD, MaxVelocityBoundsExpansionMACD,
-		VelocityBoundsMultiplier, MaxVelocityBoundsExpansion, BoundsThickness]()
+		VelocityBoundsMultiplier, MaxVelocityBoundsExpansion, BoundsThickness](int32 BatchIndex, int32 StartIndex, int32 EndIndex)
+	{
+		for (int32 Index = StartIndex; Index < EndIndex; Index++)
 		{
-			QUICK_SCOPE_CYCLE_COUNTER(STAT_IntegrateTask);
-			for (int32 Index = StartIndex; Index < EndIndex; Index++)
+			//question: can we enforce this at the API layer? Right now islands contain non dynamic which makes this hard
+			if (FPBDRigidParticleHandle* PBDParticle = ParticlesArray[Index]->CastToRigidParticle();  PBDParticle)
 			{
-				//question: can we enforce this at the API layer? Right now islands contain non dynamic which makes this hard
-				FPBDRigidParticleHandle* PBDParticle = ParticlesArray[Index]->CastToRigidParticle();
-				if (PBDParticle && PBDParticle->ObjectState() == EObjectStateType::Dynamic)
+				if (PBDParticle->ObjectState() == EObjectStateType::Dynamic)
 				{
 					FPBDRigidParticleHandle& Particle = *PBDParticle;
 
@@ -980,19 +968,45 @@ void FPBDRigidsEvolutionGBF::Integrate(FReal Dt)
 						}
 					}
 				}
+
+				if (!Base::IsADisableCluster(*PBDParticle) && (PBDParticle->HasCollision() || ForceNoCollisionIntoSQ))
+				{
+					const FUniqueIdx UniqueIdx = PBDParticle->UniqueIdx();
+					FPendingSpatialData& PendingData = InternalAccelerationQueue.Add(BatchIndex + 1, UniqueIdx);
+					PendingData.Operation = EPendingSpatialDataOperation::Update;
+					PendingData.AccelerationHandle = FAccelerationStructureHandle(*PBDParticle);
+					PendingData.SpatialIdx = PBDParticle->SpatialIdx();
+				}
 			}
-		}, LowLevelTasks::ETaskPriority::High);
-		IntegrationPendingTasks.Add(PendingTask);
-	}
+		}
+	};
 
-	DirtyParticleTask = UE::Tasks::Launch(UE_SOURCE_LOCATION, [this, &ParticlesArray]()
+	auto DirtyAsync = [this, &ParticlesArray]()
 	{
-		QUICK_SCOPE_CYCLE_COUNTER(STAT_DirtyDynamicParticleTask);
 		for (FPBDRigidParticleHandle* Particle : ParticlesArray)
 		{
-			Base::DirtyParticle(*Particle);
+			Base::DirtyParticleAsync(*Particle);
 		}
-	}, LowLevelTasks::ETaskPriority::High);
+	};
+	
+
+	if (GSingleThreadedPhysics)
+	{
+		// In single threaded mode keep same logic with three batches inside PendingDataArrays: 
+		// 0 Dirty particles from the PushData
+		// 1 Dynamic particles added in Integrate
+		// 2 Kinematic particles added in Apply Kinematic Target
+		InternalAccelerationQueue.PendingDataArrays.SetNum(3);
+		InternalAccelerationQueue.KinematicBatchStartIndex = 2;
+		const int32 NumDynParticles = ParticlesArray.Num();
+		IntegrateWork(0, 0, NumDynParticles);
+		DirtyAsync();
+	}
+	else
+	{
+		TaskDispatcher.DispatchIntegrate(IntegrateWork);
+		TaskDispatcher.DispatchDynAsyncDirty(DirtyAsync);
+	}
 }
 
 void FPBDRigidsEvolutionGBF::ApplyKinematicTargets(const FReal Dt, const FReal StepFraction)
@@ -1000,14 +1014,14 @@ void FPBDRigidsEvolutionGBF::ApplyKinematicTargets(const FReal Dt, const FReal S
 	check(StepFraction > (FReal)0);
 	check(StepFraction <= (FReal)1);
 
-	const bool IsLastStep = (FMath::IsNearlyEqual(StepFraction, (FReal)1, (FReal)UE_KINDA_SMALL_NUMBER));
+	const bool bIsLastStep = (FMath::IsNearlyEqual(StepFraction, (FReal)1, (FReal)UE_KINDA_SMALL_NUMBER));
 
 	TParticleView<FPBDRigidParticles>& KienamticParticles = Particles.GetActiveMovingKinematicParticlesView();
 	using TSOA = typename TParticleView<FPBDRigidParticles>::TSOA;
 	using THandle = typename TSOA::THandleType;
 	using THandleBase = typename THandle::THandleBase;
 	using TTransientHandle = typename THandle::TTransientHandle;
-
+	int32 DispatchBatchIndex = 0;
 	for (int32 ViewIndex = 0; ViewIndex < KienamticParticles.SOAViews.Num(); ++ViewIndex)
 	{
 		const TSOAView<TSOA>& SOAView = KienamticParticles.SOAViews[ViewIndex];
@@ -1026,51 +1040,65 @@ void FPBDRigidsEvolutionGBF::ApplyKinematicTargets(const FReal Dt, const FReal S
 				};
 			using LambdaType = decltype(ReturnHandle);
 			const int32 NumHandles = CurHandlesArray->Num();
-			DispatchKinematicTargetsTasks<LambdaType>(NumHandles, Dt, StepFraction, IsLastStep, ReturnHandle);
+			DispatchKinematicTargetsTasks<LambdaType>(DispatchBatchIndex++, NumHandles, Dt, StepFraction, bIsLastStep, ReturnHandle);
 		}
 		else if (TSOA* Soa = SOAView.SOA)
 		{
 			auto ReturnHandle = [Soa](int32 Index) { return THandleBase(Soa, Index); };
 			using LambdaType = decltype(ReturnHandle);
-			DispatchKinematicTargetsTasks<LambdaType>(NumParticles, Dt, StepFraction, IsLastStep, ReturnHandle);
+			DispatchKinematicTargetsTasks<LambdaType>(DispatchBatchIndex++, NumParticles, Dt, StepFraction, bIsLastStep, ReturnHandle);
 		}
 	}
 
-	// To start dirtying Kinematics particles, the Dynamics particles must be finished
-	UE::Tasks::Wait(DirtyParticleTask);
-	DirtyParticleTask = UE::Tasks::FTask();
-	for(FTransientPBDRigidParticleHandle& Particle : Particles.GetActiveMovingKinematicParticlesView())
+	if (GSingleThreadedPhysics)
 	{
-		// Moving kinematics need to have SQ updates queued to apply substepped movement to the
-		// physics thread acceleration structure in order to get correct kinematic collisions
-		// and pass a correct structure back to the external thread
-		DirtyParticle(Particle);
+		for (FTransientPBDRigidParticleHandle& Particle : Particles.GetActiveMovingKinematicParticlesView())
+		{
+			// Moving kinematics need to have SQ updates queued to apply substepped movement to the
+			// physics thread acceleration structure in order to get correct kinematic collisions
+			// and pass a correct structure back to the external thread
+			DirtyParticleAsync(Particle);
+		}
+		// done with update, let's clear the tracking structures
+		if (bIsLastStep)
+		{
+			constexpr bool bUpdateView = false;
+			Particles.UpdateAllMovingKinematic(bUpdateView);
+		}
+		Particles.UpdateDirtyViews();
 	}
-
-	UE::Tasks::Wait(KinematicPendingTasks);
-	KinematicPendingTasks.Reset();
-	// done with update, let's clear the tracking structures
-	if (IsLastStep)
+	else
 	{
-		Particles.UpdateAllMovingKinematic();
+		auto DirtyKinAsync = [this]()
+		{
+			for (FTransientPBDRigidParticleHandle& Particle : Particles.GetActiveMovingKinematicParticlesView())
+			{
+				// Moving kinematics need to have SQ updates queued to apply substepped movement to the
+				// physics thread acceleration structure in order to get correct kinematic collisions
+				// and pass a correct structure back to the external thread
+				DirtyParticleAsync(Particle);
+			}
+		};
+
+		TaskDispatcher.DispatchKinAsyncDirtyAndUpdateKinematic(DirtyKinAsync, bIsLastStep);
 	}
+}
 
-	// Updating the views can be done only when all kinematics and dynamics update are finished.
-	UE::Tasks::Wait(IntegrationPendingTasks);
-	IntegrationPendingTasks.Reset();
-	// If we changed any particle state, the views need to be refreshed
-	Particles.UpdateDirtyViews();
+void FPBDRigidsEvolutionGBF::SafePostIntegrateCallback()
+{
+	TaskDispatcher.SafePostIntegrateCallback();
 }
 
 template <typename LambdaType>
-void FPBDRigidsEvolutionGBF::DispatchKinematicTargetsTasks(int32 NumParticles, FReal Dt, FReal StepFraction, bool IsLastStep, LambdaType ReturnHandle)
+void FPBDRigidsEvolutionGBF::DispatchKinematicTargetsTasks(int32 DispatchBatchIndex, int32 NumParticles, FReal Dt, FReal StepFraction, bool bIsLastStep, LambdaType ReturnHandle)
 {
 	using TSOA = typename TParticleView<FPBDRigidParticles>::TSOA;
 	using THandle = typename TSOA::THandleType;
 	using THandleBase = typename THandle::THandleBase;
 
-	auto KinematicTargetsLambda = [this, ReturnHandle, Dt, StepFraction, IsLastStep](int32 StartIndex, int32 EndIndex)
+	auto KinematicTargetsLambda = [this, ReturnHandle, Dt, StepFraction, bIsLastStep](int32 StartIndex, int32 EndIndex, int32 BatchIndex)
 	{
+		InternalAccelerationQueue.PendingDataArrays[InternalAccelerationQueue.KinematicBatchStartIndex + BatchIndex].Reserve(EndIndex- StartIndex);
 		for (int32 Index = StartIndex; Index < EndIndex; Index++)
 		{
 			THandleBase Handle = ReturnHandle(Index);
@@ -1106,7 +1134,7 @@ void FPBDRigidsEvolutionGBF::DispatchKinematicTargetsTasks(int32 NumParticles, F
 					// Target positions only need to be processed once, and we reset the velocity next frame (if no new target is set)
 					FVec3 NewX;
 					FRotation3f NewR;
-					if (IsLastStep)
+					if (bIsLastStep)
 					{
 						NewX = KinematicTarget.GetPosition();
 						NewR = KinematicTarget.GetRotation();
@@ -1176,35 +1204,24 @@ void FPBDRigidsEvolutionGBF::DispatchKinematicTargetsTasks(int32 NumParticles, F
 					}
 				}
 			}
+			if (!Base::IsADisableCluster(Particle) && (Particle.HasCollision() || ForceNoCollisionIntoSQ))
+			{
+				const FUniqueIdx UniqueIdx = Particle.UniqueIdx();
+				FPendingSpatialData& PendingData = InternalAccelerationQueue.Add(InternalAccelerationQueue.KinematicBatchStartIndex + BatchIndex, UniqueIdx);
+				PendingData.Operation = EPendingSpatialDataOperation::Update;
+				PendingData.AccelerationHandle = FAccelerationStructureHandle(Particle);
+				PendingData.SpatialIdx = Particle.SpatialIdx();
+			}
 		}
 	};
 
-	int32 NumTasks = !!GSingleThreadedPhysics ? 1 : FMath::Min(FTaskGraphInterface::Get().GetNumWorkerThreads(), Chaos::MaxNumWorkers);
-	NumTasks = FMath::Max(FMath::Min(NumTasks, NumParticles), 1);
-	constexpr int32 MinParticlePerTask = 40;
-	const int32 ParticleByTask = FMath::Max(FMath::DivideAndRoundUp(NumParticles, NumTasks), MinParticlePerTask);
-	int32 NumBatches = FMath::DivideAndRoundUp(NumParticles, ParticleByTask);
-
-	if (NumBatches > 1)
+	if (GSingleThreadedPhysics)
 	{
-		// We only write to particle state
-		for (int32 BatchIndex = 0; BatchIndex < NumBatches; BatchIndex++)
-		{
-			const int32 StartIndex = BatchIndex * ParticleByTask;
-			int32 EndIndex = (BatchIndex + 1) * ParticleByTask;
-			EndIndex = FMath::Min(NumParticles, EndIndex);
-
-			UE::Tasks::FTask PendingTask = UE::Tasks::Launch(UE_SOURCE_LOCATION, [KinematicTargetsLambda, StartIndex, EndIndex]()
-			{
-				QUICK_SCOPE_CYCLE_COUNTER(STAT_KinematicTargetsTasks);
-				KinematicTargetsLambda(StartIndex, EndIndex);
-			}, LowLevelTasks::ETaskPriority::High);
-			KinematicPendingTasks.Add(PendingTask);
-		}
+		KinematicTargetsLambda(0, NumParticles, 0);
 	}
 	else
 	{
-		KinematicTargetsLambda(0, NumParticles);
+		TaskDispatcher.DispatchKinematicsTarget(KinematicTargetsLambda, NumParticles, DispatchBatchIndex); 
 	}
 }
 
diff --git a/Engine/Source/Runtime/Experimental/Chaos/Private/PBDRigidsSolver.cpp b/Engine/Source/Runtime/Experimental/Chaos/Private/PBDRigidsSolver.cpp
--- a/Engine/Source/Runtime/Experimental/Chaos/Private/PBDRigidsSolver.cpp
+++ b/Engine/Source/Runtime/Experimental/Chaos/Private/PBDRigidsSolver.cpp
@@ -743,10 +743,19 @@ namespace Chaos
 
 	void FPBDRigidsSolver::OnEvolutionPostIntegrate(FReal Dt)
 	{
+		bool bHasSafeThreadCallback = false;
 		for (ISimCallbackObject* Callback : SimCallbackObjects)
 		{
 			if (Callback->HasOption(ESimCallbackOptions::PostIntegrate))
 			{
+				if (!bHasSafeThreadCallback)
+				{
+					// If a post integrate callback is called, 
+					// we need to wait for all workers threads integrating to finish.
+					MEvolution->SafePostIntegrateCallback();
+					bHasSafeThreadCallback = true;
+				}
+				
 				FScopedTraceSolverCallback TraceCallback(Callback);
 				Callback->PostIntegrate_Internal();
 			}
diff --git a/Engine/Source/Runtime/Experimental/Chaos/Public/Chaos/PBDRigidsEvolution.h b/Engine/Source/Runtime/Experimental/Chaos/Public/Chaos/PBDRigidsEvolution.h
--- a/Engine/Source/Runtime/Experimental/Chaos/Public/Chaos/PBDRigidsEvolution.h
+++ b/Engine/Source/Runtime/Experimental/Chaos/Public/Chaos/PBDRigidsEvolution.h
@@ -10,9 +10,10 @@
 #include "Chaos/Framework/DebugSubstep.h"
 #include "HAL/Event.h"
 #include "Chaos/PBDJointConstraints.h"
+#include "Chaos/PBDRigidsEvolutionFwd.h"
 #include "Chaos/PBDRigidsSOAs.h"
 #include "Chaos/SpatialAccelerationCollection.h"
-#include "Chaos/PBDRigidsEvolutionFwd.h"
+#include "Chaos/TaskDispatcherEvolution.h"
 #include "Chaos/Island/IslandManager.h"
 #include "Chaos/Island/IslandGroupManager.h"
 #include "Chaos/Defines.h"
@@ -494,28 +495,27 @@ public:
 	}
 
 	template <bool bPersistent>
-	FORCEINLINE_DEBUGGABLE void DirtyParticle(TGeometryParticleHandleImp<FReal, 3, bPersistent>& Particle, const EPendingSpatialDataOperation Op = EPendingSpatialDataOperation::Update)
+	FORCEINLINE_DEBUGGABLE bool IsADisableCluster(TGeometryParticleHandleImp<FReal, 3, bPersistent>& Particle)
 	{
-		ensure(Op != EPendingSpatialDataOperation::Delete); // Don't use the function to delete particles
 		const TPBDRigidParticleHandleImp<FReal, 3, bPersistent>* AsRigid = Particle.CastToRigidParticle();
-		if(AsRigid && AsRigid->Disabled())
+		if (AsRigid && AsRigid->Disabled())
 		{
 			TPBDRigidClusteredParticleHandleImp<FReal, 3, bPersistent>* AsClustered = Particle.CastToClustered();
 
-			if(AsClustered)
+			if (AsClustered)
 			{
 				// For clustered particles, they may appear disabled but they're being driven by an internal (solver-owned) cluster parent.
 				// If this is the case we let the spatial data update with those particles, otherwise skip.
 				// Alternatively, the particle may be being driven by a cluster union. Disabled children of a cluster union should not be added
 				// to the SQ.
 				// #BGTODO consider converting MDisabled into a bitfield for multiple disable types (Disabled, DisabledDriven, etc.)
-				if(FPBDRigidParticleHandle* ClusterParentBase = AsClustered->ClusterIds().Id)
+				if (FPBDRigidParticleHandle* ClusterParentBase = AsClustered->ClusterIds().Id)
 				{
-					if(Chaos::TPBDRigidClusteredParticleHandle<FReal, 3>* ClusterParent = ClusterParentBase->CastToClustered())
+					if (Chaos::TPBDRigidClusteredParticleHandle<FReal, 3>*ClusterParent = ClusterParentBase->CastToClustered())
 					{
-						if(!ClusterParent->InternalCluster())
+						if (!ClusterParent->InternalCluster())
 						{
-							return;
+							return true;
 						}
 
 						// We know we're an internal cluster now. If this was a GC internal proxy, we'd expect the 
@@ -523,44 +523,81 @@ public:
 						// we're in a cluster union.
 						if (ClusterParent->PhysicsProxy() != Particle.PhysicsProxy())
 						{
-							return;
+							return true;
 						}
 					}
 					else
 					{
 						// There's probably no way we get here since if we're clustering the parent should always be a clustered.
 						ensure(false);
-						return;
+						return true;
 					}
 				}
 				else
 				{
 					// Disabled cluster particle that doesn't have a parent cluster. MUST BE IGNORED.
-					return;
+					return true;
 				}
 			}
 			else
 			{
 				// Disabled particles take no immediate part in sim or query so shouldn't be added to the acceleration
-				return;
+				return true;
 			}
 		}
+		return false;
+	}
+
+	template <bool bPersistent>
+	FORCEINLINE_DEBUGGABLE void DirtyParticle(TGeometryParticleHandleImp<FReal, 3, bPersistent>& Particle, const EPendingSpatialDataOperation Op = EPendingSpatialDataOperation::Update)
+	{
+		ensure(Op != EPendingSpatialDataOperation::Delete); // Don't use the function to delete particles
+		
+		if (IsADisableCluster(Particle))
+		{
+			return;
+		}
+
+		//only add to acceleration structure if it has collision
+		if (Particle.HasCollision() || ForceNoCollisionIntoSQ)
+		{
+			//TODO: distinguish between new particles and dirty particles - Adds and updates are treated the same right now
+			const FUniqueIdx UniqueIdx = Particle.UniqueIdx();
+			FPendingSpatialData& PendingData = InternalAccelerationQueue.Add(0, UniqueIdx);
+			PendingData.Operation = Op;
+			PendingData.AccelerationHandle = FAccelerationStructureHandle(Particle);
+			PendingData.SpatialIdx = Particle.SpatialIdx();
+
+
+			FPendingSpatialData& AsyncSpatialData = AsyncAccelerationQueue.FindOrAdd(UniqueIdx);
+			AsyncSpatialData = PendingData;
+			// ensure(AsyncSpatialData.Operation != EPendingSpatialDataOperation::Delete); // TODO: This may be hit: Potentially due to UniqueIdx reuse?
+		}
+	}
+
+	template <bool bPersistent>
+	FORCEINLINE_DEBUGGABLE void DirtyParticleAsync(TGeometryParticleHandleImp<FReal, 3, bPersistent>& Particle, const EPendingSpatialDataOperation Op = EPendingSpatialDataOperation::Update)
+	{
+		ensure(Op != EPendingSpatialDataOperation::Delete); // Don't use the function to delete particles
+
+		if (IsADisableCluster(Particle))
+		{
+			return;
+		}
 
 		//only add to acceleration structure if it has collision
 		if (Particle.HasCollision() || ForceNoCollisionIntoSQ)
 		{
 			//TODO: distinguish between new particles and dirty particles - Adds and updates are treated the same right now
 			const FUniqueIdx UniqueIdx = Particle.UniqueIdx();
-			FPendingSpatialData& SpatialData = InternalAccelerationQueue.FindOrAdd(UniqueIdx);
-			ensure(SpatialData.Operation != EPendingSpatialDataOperation::Delete);
 
-			SpatialData.Operation = Op;
-			SpatialData.AccelerationHandle = FAccelerationStructureHandle(Particle);
-			SpatialData.SpatialIdx = Particle.SpatialIdx();
 
-			auto& AsyncSpatialData = AsyncAccelerationQueue.FindOrAdd(UniqueIdx);
+			FPendingSpatialData& AsyncSpatialData = AsyncAccelerationQueue.FindOrAdd(UniqueIdx);
+			AsyncSpatialData.Operation = Op;
+			AsyncSpatialData.AccelerationHandle = FAccelerationStructureHandle(Particle);
+			AsyncSpatialData.SpatialIdx = Particle.SpatialIdx();
+
 			// ensure(AsyncSpatialData.Operation != EPendingSpatialDataOperation::Delete); // TODO: This may be hit: Potentially due to UniqueIdx reuse?
-			AsyncSpatialData = SpatialData;
 		}
 	}
 
@@ -589,7 +626,6 @@ public:
 	{
 		if (const int32 NumNew = IslandManager.ReserveParticles(Num))
 		{
-			InternalAccelerationQueue.PendingData.Reserve(InternalAccelerationQueue.Num() + NumNew);
 			AsyncAccelerationQueue.PendingData.Reserve(AsyncAccelerationQueue.Num() + NumNew);
 		}
 	}
@@ -834,6 +870,9 @@ public:
 	void SetName(const FString& InName) { EvolutionName = InName; }
 	const FString& GetName() const { return EvolutionName; }
 
+	// Internal used only 
+	UE_INTERNAL virtual void SafePostIntegrateCallback() { check(false); };
+
 protected:
 	int32 NumConstraints() const
 	{
@@ -954,9 +993,11 @@ public:
 
 protected:
 
-	/** Pending operations for the internal acceleration structure */
-	FPendingSpatialDataQueue InternalAccelerationQueue;
+	Private::FTaskDispatcherEvolution TaskDispatcher;
 
+	/** Pending operations for the internal acceleration structure */
+	FPendingSpatialInternalDataQueue InternalAccelerationQueue;
+	
 	/** Pending operations for the acceleration structures being rebuilt asynchronously */
 	FPendingSpatialDataQueue AsyncAccelerationQueue;
 
diff --git a/Engine/Source/Runtime/Experimental/Chaos/Public/Chaos/PBDRigidsEvolutionGBF.h b/Engine/Source/Runtime/Experimental/Chaos/Public/Chaos/PBDRigidsEvolutionGBF.h
--- a/Engine/Source/Runtime/Experimental/Chaos/Public/Chaos/PBDRigidsEvolutionGBF.h
+++ b/Engine/Source/Runtime/Experimental/Chaos/Public/Chaos/PBDRigidsEvolutionGBF.h
@@ -273,6 +273,9 @@ namespace Chaos
 		UE_DEPRECATED(5.4, "Use Integrate(Dt)")
 		void Integrate(const TParticleView& InParticles, FReal Dt) { Integrate(Dt); }
 
+		// Internal used only
+		UE_INTERNAL virtual void SafePostIntegrateCallback() override;
+
 	protected:
 
 		CHAOS_API void AdvanceOneTimeStepImpl(const FReal Dt, const FSubStepInfo& SubStepInfo);
@@ -296,7 +299,7 @@ namespace Chaos
 		void UpdateCollisionSolverType();
 
 		template <typename LambdaType>
-		void DispatchKinematicTargetsTasks(int32 NumParticles, FReal Dt, FReal StepFraction, bool IsLastStep, LambdaType ReturnHandle);
+		void DispatchKinematicTargetsTasks(int32 DispatchBatchIndex, int32 NumParticles, FReal Dt, FReal StepFraction, bool IsLastStep, LambdaType ReturnHandle);
 
 		FRigidClustering Clustering;
 
@@ -324,10 +327,6 @@ namespace Chaos
 
 		FCCDManager CCDManager;
 
-		TArray<UE::Tasks::FTask> IntegrationPendingTasks;
-		TArray<UE::Tasks::FTask> KinematicPendingTasks;
-		UE::Tasks::FTask DirtyParticleTask;
-
 		bool bIsDeterministic;
 
 #if CHAOS_EVOLUTION_COLLISION_TESTMODE
diff --git a/Engine/Source/Runtime/Experimental/Chaos/Public/Chaos/PBDRigidsSOAs.h b/Engine/Source/Runtime/Experimental/Chaos/Public/Chaos/PBDRigidsSOAs.h
--- a/Engine/Source/Runtime/Experimental/Chaos/Public/Chaos/PBDRigidsSOAs.h
+++ b/Engine/Source/Runtime/Experimental/Chaos/Public/Chaos/PBDRigidsSOAs.h
@@ -826,7 +826,7 @@ public:
 		}
 	}
 
-	void UpdateAllMovingKinematic()
+	void UpdateAllMovingKinematic(const bool bUpdateViews = true)
 	{
 		// moving kinematics are going to a 'reset' mode then a 'none' mode
 		// 'reset' ones need to stay for another frame 
@@ -843,7 +843,7 @@ public:
 				// Particle may not be moving, but its velocity was just changed (to zero), so it
 				// needs to be in a dirty list. Since it are no longer in MovingKinematicsMapArray
 				// we must put it the transient dirty list
-				MarkTransientDirtyParticle(Rigid);
+				MarkTransientDirtyParticle(Rigid, bUpdateViews);
 
 				CheckParticleViewMask(Rigid);
 			}
diff --git a/Engine/Source/Runtime/Experimental/Chaos/Public/Chaos/PendingSpatialData.h b/Engine/Source/Runtime/Experimental/Chaos/Public/Chaos/PendingSpatialData.h
--- a/Engine/Source/Runtime/Experimental/Chaos/Public/Chaos/PendingSpatialData.h
+++ b/Engine/Source/Runtime/Experimental/Chaos/Public/Chaos/PendingSpatialData.h
@@ -5,6 +5,9 @@
 #include "Chaos/Transform.h"
 #include "Chaos/ISpatialAcceleration.h"
 #include "Chaos/Defines.h"
+#include "Tasks/Task.h"
+#include "Containers/BitArray.h"
+
 
 namespace Chaos
 { 
@@ -58,6 +61,125 @@ struct FPendingSpatialData
 	}
 };
 
+struct FPendingSpatialInternalDataQueue
+{
+	// PendingDataArrays store in 
+	// - [0] all elements added in push data
+	// - [1 -> KinematicBatchStartIndex] all dynamic particle contained in GetActiveParticlesArray
+	// - [KinematicBatchStartIndex -> PendingDataArrays.Num()] all kinematic particle in GetActiveMovingKinematicParticlesView
+	TArray<TArray<FPendingSpatialData>> PendingDataArrays;
+	// ParticleToPendingData ensure all particle in index [0] are unique
+	TArrayAsMap<FUniqueIdx, int32> ParticleToPendingData;
+	// This array is access to mark duplicated data from array [0], not we cannot use bit array because it could be used from different thread.
+	TArray<bool> Duplicated;
+	int32 KinematicBatchStartIndex = 0;
+
+	void Reset()
+	{
+		KinematicBatchStartIndex = 0;
+		for (TArray<FPendingSpatialData>& PendingData : PendingDataArrays)
+		{
+			PendingData.Reset();
+		}
+		PendingDataArrays.Reset();
+		ParticleToPendingData.Reset();
+		Duplicated.Reset();
+	}
+
+	int32 Num() const
+	{
+		int32 NumPendingData = 0;
+		for (const TArray<FPendingSpatialData>& PendingData : PendingDataArrays)
+		{
+			NumPendingData += PendingData.Num();
+		}
+		return NumPendingData;
+	}
+
+	FPendingSpatialData& Add(int32 BatchIndex, const FUniqueIdx UniqueIdx, EPendingSpatialDataOperation Operation = EPendingSpatialDataOperation::Add)
+	{
+		if (PendingDataArrays.IsEmpty())
+		{
+			check(BatchIndex == 0);
+			PendingDataArrays.SetNum(1);
+		}
+
+		check(PendingDataArrays[0].Num() == Duplicated.Num());
+		// Check if duplicate
+		if (BatchIndex != 0)
+		{
+			if (int32* Existing = ParticleToPendingData.Find(UniqueIdx))
+			{
+				Duplicated[*Existing] = true;
+			}
+			const int32 NewIdx = PendingDataArrays[BatchIndex].AddDefaulted(1);
+			PendingDataArrays[BatchIndex][NewIdx].Operation = Operation;
+			check(PendingDataArrays[0].Num() == Duplicated.Num());
+			return PendingDataArrays[BatchIndex][NewIdx];
+		}
+		else // Adding at Index 0 would be always in Single thread
+		{
+			if (int32* Existing = ParticleToPendingData.Find(UniqueIdx))
+			{
+				check(PendingDataArrays[0].Num() == Duplicated.Num());
+				return PendingDataArrays[0][*Existing];
+			}
+			const int32 NewIdx = PendingDataArrays[0].AddDefaulted(1);
+			ParticleToPendingData.Add(UniqueIdx, NewIdx);
+			PendingDataArrays[0][NewIdx].Operation = Operation;
+			Duplicated.Add(false);
+			check(PendingDataArrays[0].Num() == Duplicated.Num());
+			return PendingDataArrays[0][NewIdx];
+		}
+	}
+
+	void Remove(const FUniqueIdx UniqueIdx)
+	{
+		if (!PendingDataArrays.IsEmpty())
+		{
+			check(PendingDataArrays[0].Num() == Duplicated.Num());
+			TArray<FPendingSpatialData>& PendingData = PendingDataArrays[0];
+			if (int32* Existing = ParticleToPendingData.Find(UniqueIdx))
+			{
+				const int32 SlotIdx = *Existing;
+				if (SlotIdx + 1 < PendingData.Num())
+				{
+					const FUniqueIdx LastElemUniqueIdx = PendingData.Last().UniqueIdx();
+					ParticleToPendingData.FindChecked(LastElemUniqueIdx) = SlotIdx;	//We're going to swap elements so the last element is now in the position of the element we removed
+				}
+
+				Duplicated.RemoveAtSwap(SlotIdx);
+				PendingData.RemoveAtSwap(SlotIdx);
+				ParticleToPendingData.RemoveChecked(UniqueIdx);
+			}
+			check(PendingDataArrays[0].Num() == Duplicated.Num());
+		}
+	}
+
+	void CleanUpDuplicated()
+	{
+		if (!PendingDataArrays.IsEmpty())
+		{
+			check(PendingDataArrays[0].Num() == Duplicated.Num());
+			ParticleToPendingData.Reset();
+			TArray<FPendingSpatialData> NoDuplicatedData;
+			for (int32 Index = 0; Index < Duplicated.Num(); Index++)
+			{
+				const bool bIsDuplicated = Duplicated[Index];
+				if (!bIsDuplicated)
+				{
+					NoDuplicatedData.Add(PendingDataArrays[0][Index]);
+					ParticleToPendingData.Add(PendingDataArrays[0][Index].UniqueIdx(), Index);
+				}
+			}
+			Duplicated.Init(false, NoDuplicatedData.Num());
+			PendingDataArrays[0] = MoveTemp(NoDuplicatedData);
+			check(PendingDataArrays[0].Num() == Duplicated.Num());
+		}
+	}
+
+};
+
 struct FPendingSpatialDataQueue
 {
 	TArray<FPendingSpatialData> PendingData;
@@ -105,4 +227,4 @@ struct FPendingSpatialDataQueue
 	}
 };
 
-}
\ No newline at end of file
+}
diff --git a/Engine/Source/Runtime/Experimental/Chaos/Public/Chaos/TaskDispatcherEvolution.h b/Engine/Source/Runtime/Experimental/Chaos/Public/Chaos/TaskDispatcherEvolution.h
new file mode 100644
--- /dev/null
+++ b/Engine/Source/Runtime/Experimental/Chaos/Public/Chaos/TaskDispatcherEvolution.h
@@ -0,0 +1,309 @@
+// Copyright Epic Games, Inc. All Rights Reserved.
+#pragma once
+
+#include "Chaos/PBDRigidsSOAs.h"
+#include "Chaos/PendingSpatialData.h"
+#include "HAL/Event.h"
+
+namespace Chaos::Private
+{
+
+// This class is responsible to dispatch the evolution tasks
+// It orchestrates the task dependencies according to data flow.
+// This class shouldn't do any simulation logic and code.
+class FTaskDispatcherEvolution
+{
+private:
+	constexpr static int32 MinParticlePerTask = 40;
+
+	UE::Tasks::FTask AsyncQueueTask;
+	UE::Tasks::FTask UpdateViewTask;
+
+	TArray<UE::Tasks::FTask> IntegrationPendingTasks;
+	TArray<UE::Tasks::FTask> KinematicPendingTasks;
+
+	FPBDRigidsSOAs& Particles;
+	FPendingSpatialInternalDataQueue& InternalAccelerationQueue;
+
+	int32 NumTasks;
+	int32 NumDynParticles;
+	int32 NumKinBatches;
+
+public:
+	FTaskDispatcherEvolution(FPBDRigidsSOAs& ParticlesIn, FPendingSpatialInternalDataQueue& InternalAccelerationQueueIn)
+		: Particles(ParticlesIn)
+		, InternalAccelerationQueue(InternalAccelerationQueueIn)
+		, NumTasks(FMath::Min(FTaskGraphInterface::Get().GetNumWorkerThreads(), Chaos::MaxNumWorkers))
+		, NumDynParticles()
+		, NumKinBatches()
+	{
+	}
+
+	void ComputeKinematicBatch()
+	{
+		// Compute the number of kinematic particles to allocate the InternalAccelerationQueue pending data
+		const TParticleView<FPBDRigidParticles>& KienamticParticles = Particles.GetActiveMovingKinematicParticlesView();
+		const int32 NumView = KienamticParticles.SOAViews.Num();
+		NumKinBatches = 0;
+		for (int32 ViewIndex = 0; ViewIndex < NumView; ++ViewIndex)
+		{
+			const int32 NumKinParticles = KienamticParticles.SOAViews[ViewIndex].Size();
+			const int32 ParticleByTask = FMath::Max(FMath::DivideAndRoundUp(NumKinParticles, NumTasks), MinParticlePerTask);
+			NumKinBatches += FMath::DivideAndRoundUp(NumKinParticles, ParticleByTask);
+		}
+	}
+
+	template<typename Lambda>
+	void DispatchIntegrate(Lambda IntegrateWork)
+	{
+		NumDynParticles = Particles.GetActiveParticlesArray().Num();
+
+		int32 NumDynTask = FMath::Max(FMath::Min(NumTasks, NumDynParticles), 1);
+
+		const int32 ParticleByTask = FMath::Max(FMath::DivideAndRoundUp(NumDynParticles, NumDynTask), MinParticlePerTask);
+		int32 NumDynBatches = FMath::DivideAndRoundUp(NumDynParticles, ParticleByTask);
+
+		IntegrationPendingTasks.Reset(NumDynBatches);
+		InternalAccelerationQueue.KinematicBatchStartIndex = NumDynBatches + 1;
+		ComputeKinematicBatch();
+		InternalAccelerationQueue.PendingDataArrays.SetNum(InternalAccelerationQueue.KinematicBatchStartIndex + NumKinBatches); // Keep 0 for the dirty element from PushData
+		for (int32 BatchIndex = 0; BatchIndex < NumDynBatches; BatchIndex++)
+		{
+			const int32 StartIndex = BatchIndex * ParticleByTask;
+			int32 EndIndex = (BatchIndex + 1) * ParticleByTask;
+			EndIndex = FMath::Min(NumDynParticles, EndIndex);
+
+			UE::Tasks::FTask PendingTask = UE::Tasks::Launch(UE_SOURCE_LOCATION, [this, IntegrateWork, BatchIndex, StartIndex, EndIndex]()
+			{
+				QUICK_SCOPE_CYCLE_COUNTER(STAT_IntegrateTask);
+				IntegrateWork(BatchIndex, StartIndex, EndIndex);
+			}, LowLevelTasks::ETaskPriority::High);
+			IntegrationPendingTasks.Add(PendingTask);
+		}
+	}
+
+	template<typename Lambda>
+	void DispatchDynAsyncDirty(Lambda AsyncDirtyWork)
+	{
+		// Why do we add the dynamic particle to the async queue which is suppose to be for static tree ? 
+		AsyncQueueTask = UE::Tasks::Launch(UE_SOURCE_LOCATION, [AsyncDirtyWork]()
+		{
+			QUICK_SCOPE_CYCLE_COUNTER(STAT_DirtyAsyncDynamicParticleTask);
+			AsyncDirtyWork();
+		}, LowLevelTasks::ETaskPriority::High);
+	}
+
+	template<typename Lambda>
+	void DispatchKinAsyncDirtyAndUpdateKinematic(Lambda AsyncDirty, bool bIsLastStep)
+	{
+		// To start dirtying Kinematics particles, the Dynamics particles must be finished
+		// Is it necessary to update the kinematic particles for the async data ? 
+		AsyncQueueTask = UE::Tasks::Launch(UE_SOURCE_LOCATION, [AsyncDirty]()
+		{
+			QUICK_SCOPE_CYCLE_COUNTER(STAT_DirtyMovingKinematicsTask);
+			AsyncDirty();
+		}, AsyncQueueTask, LowLevelTasks::ETaskPriority::High);
+
+		TArray<UE::Tasks::FTask> PendingTasks;
+		PendingTasks.Add(AsyncQueueTask);
+
+		UE::Tasks::FTask MovingKinematicsTask;
+		// done with update, let's clear the tracking structures
+		if (bIsLastStep)
+		{
+			PendingTasks.Append(KinematicPendingTasks);
+			MovingKinematicsTask = UE::Tasks::Launch(UE_SOURCE_LOCATION, [this]()
+			{
+				QUICK_SCOPE_CYCLE_COUNTER(STAT_UpdateAllMovingKinematicTask);
+				constexpr bool bUpdateView = false;
+				Particles.UpdateAllMovingKinematic(bUpdateView);
+			}, PendingTasks, LowLevelTasks::ETaskPriority::High);
+
+		}
+		else
+		{
+			MovingKinematicsTask = MoveTemp(AsyncQueueTask);
+		}
+
+		// Updating the views can be done only when all kinematics and dynamics update are finished. // This comment is wrong
+		UpdateViewTask = UE::Tasks::Launch(UE_SOURCE_LOCATION, [this]()
+		{
+			QUICK_SCOPE_CYCLE_COUNTER(STAT_UpdateViewTask);
+			// If we changed any particle state, the views need to be refreshed
+			Particles.UpdateDirtyViews();
+
+		}, MovingKinematicsTask, LowLevelTasks::ETaskPriority::High);
+	}
+
+	template<typename Lambda>
+	void DispatchKinematicsTarget(Lambda KinematicTargetsWork, int32 NumParticles, int32 DispatchBatchIndex)
+	{
+		int32 NumKinTask = FMath::Max(FMath::Min(NumTasks, NumParticles), 1);
+		const int32 ParticleByTask = FMath::Max(FMath::DivideAndRoundUp(NumParticles, NumKinTask), MinParticlePerTask);
+		int32 NumBatches = FMath::DivideAndRoundUp(NumParticles, ParticleByTask);
+
+		check(NumBatches < InternalAccelerationQueue.PendingDataArrays.Num() + 1);
+		check(NumBatches > 0)
+		// We only write to particle state
+		for (int32 BatchIndex = 0; BatchIndex < NumBatches; BatchIndex++)
+		{
+			const int32 StartIndex = BatchIndex * ParticleByTask;
+			int32 EndIndex = (BatchIndex + 1) * ParticleByTask;
+			EndIndex = FMath::Min(NumParticles, EndIndex);
+
+			UE::Tasks::FTask PendingTask = UE::Tasks::Launch(UE_SOURCE_LOCATION, [KinematicTargetsWork, StartIndex, EndIndex, DispatchBatchIndex, BatchIndex]()
+			{
+				QUICK_SCOPE_CYCLE_COUNTER(STAT_KinematicTargetsTask);
+				KinematicTargetsWork(StartIndex, EndIndex, DispatchBatchIndex + BatchIndex);
+			}, LowLevelTasks::ETaskPriority::High);
+			KinematicPendingTasks.Add(PendingTask);
+		}
+	}
+
+	template <typename Lambda>
+	void PruneInternalPendingData(Lambda HasToBeUpdated)
+	{
+		const int32 NumBatchDyn = FMath::Min(InternalAccelerationQueue.KinematicBatchStartIndex, InternalAccelerationQueue.PendingDataArrays.Num()) - 1;
+		check(IntegrationPendingTasks.Num() == 0 || NumBatchDyn == IntegrationPendingTasks.Num());
+
+		TArray<UE::Tasks::FTask> CurrentPendingTasks;
+		CurrentPendingTasks.Reserve(InternalAccelerationQueue.PendingDataArrays.Num());
+
+		TArray<TArray<FPendingSpatialData>> NewDynPendingDataArray;
+		TArray<TArray<FPendingSpatialData>> NewKinPendingDataArray;
+
+		UE::Tasks::FTask CleanUpTask;
+		int32 NumActiveParticleTask = IntegrationPendingTasks.Num() + KinematicPendingTasks.Num();
+		if (NumActiveParticleTask > 0)
+		{
+			// Moving particles could be potentially duplicated in batch 0 (inside InternalAccelerationQueue.PendingDataArrays), 
+			// before processing batch 0, lets remove all duplicated particles in that batch.
+			TArray<UE::Tasks::FTask> MovingParticlesTasks;
+			MovingParticlesTasks.Reserve(NumActiveParticleTask);
+			MovingParticlesTasks.Append(IntegrationPendingTasks);
+			MovingParticlesTasks.Append(KinematicPendingTasks);
+			CleanUpTask = UE::Tasks::Launch(UE_SOURCE_LOCATION, [this]()
+			{
+				QUICK_SCOPE_CYCLE_COUNTER(STAT_CleanUpDuplicated);
+				InternalAccelerationQueue.CleanUpDuplicated();
+			}, MovingParticlesTasks, LowLevelTasks::ETaskPriority::High);
+
+			CurrentPendingTasks.Add(CleanUpTask);
+		}
+		// Prune Dynamics particles
+		if (NumBatchDyn > 0)
+		{
+			check(IntegrationPendingTasks.Num() == NumBatchDyn);
+			check(NumBatchDyn <= InternalAccelerationQueue.PendingDataArrays.Num());
+			NewDynPendingDataArray.SetNum(NumBatchDyn);
+			TArray<UE::Tasks::FTask> DynamicTasks;
+			for (int32 BatchIndex = 0; BatchIndex < NumBatchDyn; BatchIndex++)
+			{
+				// Dynamic particles start at index 1
+				TArray<FPendingSpatialData>& OldPendingData = InternalAccelerationQueue.PendingDataArrays[BatchIndex + 1];
+				TArray<FPendingSpatialData>& NewPendingData = NewDynPendingDataArray[BatchIndex];
+				UE::Tasks::FTask PendingTask = UE::Tasks::Launch(UE_SOURCE_LOCATION, [HasToBeUpdated, &OldPendingData, &NewPendingData]()
+				{
+					QUICK_SCOPE_CYCLE_COUNTER(STAT_PruneDynamicsPendingDataTask);
+					for (const FPendingSpatialData& PendingData : OldPendingData)
+					{
+						if (HasToBeUpdated(PendingData))
+						{
+							NewPendingData.Add(PendingData);
+						}
+					}
+				}, IntegrationPendingTasks[BatchIndex], LowLevelTasks::ETaskPriority::High);
+				DynamicTasks.Add(PendingTask);
+			}
+			// In order to safely copy all all pruned data, we need to be sure 
+			DynamicTasks.Add(CleanUpTask);
+			UE::Tasks::FTask CopyDynamicTask = UE::Tasks::Launch(UE_SOURCE_LOCATION, [this, &NewDynPendingDataArray, NumBatchDyn]()
+			{
+				QUICK_SCOPE_CYCLE_COUNTER(STAT_CopyPrunedDynamicPendingDataTask);
+				for (int32 BatchIndex = 0; BatchIndex < NumBatchDyn; BatchIndex++)
+				{
+					// Dynamic particles start at index 1
+					TArray<FPendingSpatialData>& OldPendingData = InternalAccelerationQueue.PendingDataArrays[BatchIndex + 1];
+					TArray<FPendingSpatialData>& NewPendingData = NewDynPendingDataArray[BatchIndex];
+					OldPendingData = MoveTemp(NewPendingData);
+				}
+			}, DynamicTasks, LowLevelTasks::ETaskPriority::High);
+			CurrentPendingTasks.Add(CopyDynamicTask);
+		}
+		int32 NumBatchKin = InternalAccelerationQueue.PendingDataArrays.Num() - InternalAccelerationQueue.KinematicBatchStartIndex;
+		// Prune Kinematics particles
+		if (NumBatchKin > 0 && InternalAccelerationQueue.KinematicBatchStartIndex > 0)
+		{
+			check(KinematicPendingTasks.Num() == NumBatchKin);
+			NewKinPendingDataArray.SetNum(NumBatchKin);
+			TArray<UE::Tasks::FTask> KinematicTasks;
+			for (int32 BatchIndex = 0; BatchIndex < NumBatchKin; BatchIndex++)
+			{
+				TArray<FPendingSpatialData>& OldPendingData = InternalAccelerationQueue.PendingDataArrays[InternalAccelerationQueue.KinematicBatchStartIndex + BatchIndex];
+				TArray<FPendingSpatialData>& NewPendingData = NewKinPendingDataArray[BatchIndex];
+				UE::Tasks::FTask PendingTask = UE::Tasks::Launch(UE_SOURCE_LOCATION, [this, &OldPendingData, &NewPendingData, HasToBeUpdated]()
+				{
+					QUICK_SCOPE_CYCLE_COUNTER(STAT_PruneKinematicPendingDataTask);
+					for (const FPendingSpatialData& PendingData : OldPendingData)
+					{
+						if (HasToBeUpdated(PendingData))
+						{
+							NewPendingData.Add(PendingData);
+						}
+					}
+				}, KinematicPendingTasks[BatchIndex], LowLevelTasks::ETaskPriority::High);
+				KinematicTasks.Add(PendingTask);
+			}
+			KinematicTasks.Add(CleanUpTask);
+			UE::Tasks::FTask CopyKinematicTask = UE::Tasks::Launch(UE_SOURCE_LOCATION, [this, &NewKinPendingDataArray, NumBatchKin]()
+			{
+				QUICK_SCOPE_CYCLE_COUNTER(STAT_CopyPrunedKinematicPendingDataTask);
+				for (int32 BatchIndex = 0; BatchIndex < NumBatchKin; BatchIndex++)
+				{
+					TArray<FPendingSpatialData>& OldPendingData = InternalAccelerationQueue.PendingDataArrays[InternalAccelerationQueue.KinematicBatchStartIndex + BatchIndex];
+					TArray<FPendingSpatialData>& NewPendingData = NewKinPendingDataArray[BatchIndex];
+					OldPendingData = MoveTemp(NewPendingData);
+				}
+			}, KinematicTasks, LowLevelTasks::ETaskPriority::High);
+			CurrentPendingTasks.Add(CopyKinematicTask);
+		}
+		UE::Tasks::Wait(CurrentPendingTasks);
+		IntegrationPendingTasks.Reset();
+		KinematicPendingTasks.Reset();
+	}
+
+	template<typename Lambda>
+	void FlushAccelerationQueue(Lambda FlushAccelerationQueueLambda)
+	{
+		AsyncQueueTask = UE::Tasks::Launch(UE_SOURCE_LOCATION, [FlushAccelerationQueueLambda]()
+		{
+			QUICK_SCOPE_CYCLE_COUNTER(STAT_FlushAsyncAccelerationQueueTask);
+			FlushAccelerationQueueLambda();
+		}, AsyncQueueTask, LowLevelTasks::ETaskPriority::High);
+	}
+
+	// This will be call if it is required to wait tasks result to invoke a callback.
+	// We never know what a callback is going to use.
+	void SafePostIntegrateCallback()
+	{
+		TArray<UE::Tasks::FTask> Tasks = { UpdateViewTask };
+		Tasks.Reserve(IntegrationPendingTasks.Num() + KinematicPendingTasks.Num() + 1);
+		Tasks.Append(IntegrationPendingTasks);
+		Tasks.Append(KinematicPendingTasks);
+		UE::Tasks::Wait(Tasks);
+		UpdateViewTask = UE::Tasks::FTask();
+		IntegrationPendingTasks.Reset();
+		KinematicPendingTasks.Reset();
+	}
+
+	void WaitAsyncQueueTask()
+	{
+		UE::Tasks::Wait(AsyncQueueTask);
+	}
+
+	void WaitTaskEndSpatial()
+	{
+		UE::Tasks::Wait(UpdateViewTask);
+	}
+};
+}
