Minor changes from review feedback for 44829205

diff --git a/Engine/Source/Runtime/Experimental/IoStore/OnDemand/Private/OnDemandInstallCache.cpp b/Engine/Source/Runtime/Experimental/IoStore/OnDemand/Private/OnDemandInstallCache.cpp
--- a/Engine/Source/Runtime/Experimental/IoStore/OnDemand/Private/OnDemandInstallCache.cpp
+++ b/Engine/Source/Runtime/Experimental/IoStore/OnDemand/Private/OnDemandInstallCache.cpp
@@ -205,13 +205,14 @@ struct FCas
 	uint64						GetBlockInfo(FCasBlockInfoMap& OutBlockInfo);
 	void						Compact();
 	FIoStatus					Verify(TArray<FCasAddr>& OutAddrs);
-	void						ReplaySnapshot(FCasSnapshot&& Snapshot);
+	void						LoadSnapshot(FCasSnapshot&& Snapshot);
 	FLastAccess					ConsumeLastAcccess()
 	{
 		FLastAccess Temp;
 		{
 			TUniqueLock Lock(Mutex);
 			Temp = MoveTemp(LastAccess);
+			LastAccess = FLastAccess();
 		}
 		return Temp;
 	}
@@ -1519,7 +1520,7 @@ TIoStatusOr<int64> FCasSnapshot::TryCreateAndResetJournal(const FString& Snapsho
 }
 
 ///////////////////////////////////////////////////////////////////////////////
-void FCas::ReplaySnapshot(FCasSnapshot&& Snapshot)
+void FCas::LoadSnapshot(FCasSnapshot&& Snapshot)
 {
 	TUniqueLock Lock(Mutex);
 
@@ -1759,7 +1760,7 @@ FOnDemandInstallCache::FOnDemandInstallCache(const FOnDemandInstallCacheConfig&
 			UE_LOG(LogIoStoreOnDemand, Log, TEXT("Loaded CAS snapshot '%s' %.2lf KiB with %d blocks and %d chunk locations"),
 					*SnapshotFile, ToKiB(SnapshotSize), Snapshot.Blocks.Num(), Snapshot.ChunkLocations.Num());
 
-			Cas.ReplaySnapshot(MoveTemp(Snapshot));
+			Cas.LoadSnapshot(MoveTemp(Snapshot));
 		}
 		else if (SnapshotStatus.Status().GetErrorCode() != EIoErrorCode::NotFound) // Its ok for the snapshot to not exist, other errors are significant
 		{
