Disable partial sleeping for non-collision constraints.

Only collision constraints can connect an active to a sleeping particle.
If an active and a sleeping particle are connected by a joint, for example, both of the particles will be woken up.

In practice, this is achieved by
(1) always propagating partial wake events through non-collision constraints and
(2) doing pre-processing ("re-waking") after putting particles asleep to make sure that there are no active-sleeping pairs between non-collision constraints.

diff --git a/Engine/Source/Programs/HeadlessChaos/Private/HeadlessChaosTestSleeping.cpp b/Engine/Source/Programs/HeadlessChaos/Private/HeadlessChaosTestSleeping.cpp
--- a/Engine/Source/Programs/HeadlessChaos/Private/HeadlessChaosTestSleeping.cpp
+++ b/Engine/Source/Programs/HeadlessChaos/Private/HeadlessChaosTestSleeping.cpp
@@ -959,5 +959,123 @@ namespace ChaosTest
 		CVarPostStepWakeThreshold->Set(PostStepWakeThresholdDefault);
 	}
 
+	// Sleep all particles except one which are connected by joint constraints.
+	// Since joint constraints don't support partial sleeping, the entire island wakes up.
+	TEST_P(SleepingTests, PartiallySleepingIslandWithJoints)
+	{
+		const bool bPartialSleepCollisionConstraintsOnly = GetParam();
+		IConsoleVariable* CVarPartialSleepCollisionConstraintsOnly = IConsoleManager::Get().FindConsoleVariable(TEXT("p.Chaos.Solver.Sleep.PartialSleepCollisionConstraintsOnly"), false);
+		check(CVarPartialSleepCollisionConstraintsOnly && CVarPartialSleepCollisionConstraintsOnly->IsVariableBool());
+		const bool bPartialSleepCollisionConstraintsOnlyDefault = CVarPartialSleepCollisionConstraintsOnly->GetBool();
+		CVarPartialSleepCollisionConstraintsOnly->Set(bPartialSleepCollisionConstraintsOnly);
+
+		const bool bPartialSleepEnabled = true;
+		FSleepingTest Test(bPartialSleepEnabled);
+		Test.MakeChainOfSpheres(5);
+
+		Test.ParticleHandles[0]->SetObjectState(EObjectStateType::Static);
+		Test.ParticleHandles[1]->SetObjectState(EObjectStateType::Sleeping);
+		Test.ParticleHandles[2]->SetObjectState(EObjectStateType::Sleeping);
+		Test.ParticleHandles[3]->SetObjectState(EObjectStateType::Sleeping);
+		Test.ParticleHandles[4]->SetObjectState(EObjectStateType::Dynamic);
+
+		Test.Advance();
+
+		if (bPartialSleepCollisionConstraintsOnly)
+		{
+			// All particles wake up since the connecting constraints don't support partial sleeping
+			for (int32 Id = 1; Id < Test.ParticleHandles.Num(); ++Id)
+			{
+				const FRigidBodyHandle_External* ParticleHandle = Test.ParticleHandles[Id];
+				EXPECT_TRUE(ParticleHandle->ObjectState() == EObjectStateType::Dynamic);
+			}
+		}
+		else
+		{
+			// Partial sleep supported, no state change
+			EXPECT_TRUE(Test.ParticleHandles[0]->ObjectState() == EObjectStateType::Static);
+			EXPECT_TRUE(Test.ParticleHandles[1]->ObjectState() == EObjectStateType::Sleeping);
+			EXPECT_TRUE(Test.ParticleHandles[2]->ObjectState() == EObjectStateType::Sleeping);
+			EXPECT_TRUE(Test.ParticleHandles[3]->ObjectState() == EObjectStateType::Sleeping);
+			EXPECT_TRUE(Test.ParticleHandles[4]->ObjectState() == EObjectStateType::Dynamic);
+		}
+		CVarPartialSleepCollisionConstraintsOnly->Set(bPartialSleepCollisionConstraintsOnlyDefault);
+	}
+
+	// TODO: Test with CVars::bChaosSolverPartialSleepCollisionConstraintsOnly == true
+	// Same as PartiallySleepingIslandWithJoints, just assign a different particle to be awake.
+	TEST_P(SleepingTests, PartiallySleepingIslandWithJoints_2)
+	{
+		const bool bPartialSleepCollisionConstraintsOnly = GetParam();
+		IConsoleVariable* CVarPartialSleepCollisionConstraintsOnly = IConsoleManager::Get().FindConsoleVariable(TEXT("p.Chaos.Solver.Sleep.PartialSleepCollisionConstraintsOnly"), false);
+		check(CVarPartialSleepCollisionConstraintsOnly && CVarPartialSleepCollisionConstraintsOnly->IsVariableBool());
+		const bool bPartialSleepCollisionConstraintsOnlyDefault = CVarPartialSleepCollisionConstraintsOnly->GetBool();
+		CVarPartialSleepCollisionConstraintsOnly->Set(bPartialSleepCollisionConstraintsOnly);
+
+		const bool bPartialSleepEnabled = true;
+		FSleepingTest Test(bPartialSleepEnabled);
+		Test.MakeChainOfSpheres(5);
+
+		Test.ParticleHandles[0]->SetObjectState(EObjectStateType::Static);
+		Test.ParticleHandles[1]->SetObjectState(EObjectStateType::Sleeping);
+		Test.ParticleHandles[2]->SetObjectState(EObjectStateType::Dynamic);
+		Test.ParticleHandles[3]->SetObjectState(EObjectStateType::Sleeping);
+		Test.ParticleHandles[4]->SetObjectState(EObjectStateType::Sleeping);
+
+		Test.Advance();
+
+		if (bPartialSleepCollisionConstraintsOnly)
+		{
+			// All particles wake up since the connecting constraints don't support partial sleeping
+			for (int32 Id = 1; Id < Test.ParticleHandles.Num(); ++Id)
+			{
+				const FRigidBodyHandle_External* ParticleHandle = Test.ParticleHandles[Id];
+				EXPECT_TRUE(ParticleHandle->ObjectState() == EObjectStateType::Dynamic);
+			}
+		}
+		else
+		{
+			// Partial sleep supported, no state change
+			EXPECT_TRUE(Test.ParticleHandles[0]->ObjectState() == EObjectStateType::Static);
+			EXPECT_TRUE(Test.ParticleHandles[1]->ObjectState() == EObjectStateType::Sleeping);
+			EXPECT_TRUE(Test.ParticleHandles[2]->ObjectState() == EObjectStateType::Dynamic);
+			EXPECT_TRUE(Test.ParticleHandles[3]->ObjectState() == EObjectStateType::Sleeping);
+			EXPECT_TRUE(Test.ParticleHandles[4]->ObjectState() == EObjectStateType::Sleeping);
+		}
+		CVarPartialSleepCollisionConstraintsOnly->Set(bPartialSleepCollisionConstraintsOnlyDefault);
+	}
+
+	// TODO: Test with CVars::bChaosSolverPartialSleepCollisionConstraintsOnly == true
+	// Sleep all particles except one which are connected by joint constraints.
+	// Since all constraints already sleep, the entire island remains asleep.
+	TEST_P(SleepingTests, FullySleepingIslandWithJoints)
+	{
+		const bool bPartialSleepCollisionConstraintsOnly = GetParam();
+		IConsoleVariable* CVarPartialSleepCollisionConstraintsOnly = IConsoleManager::Get().FindConsoleVariable(TEXT("p.Chaos.Solver.Sleep.PartialSleepCollisionConstraintsOnly"), false);
+		check(CVarPartialSleepCollisionConstraintsOnly && CVarPartialSleepCollisionConstraintsOnly->IsVariableBool());
+		const bool bPartialSleepCollisionConstraintsOnlyDefault = CVarPartialSleepCollisionConstraintsOnly->GetBool();
+		CVarPartialSleepCollisionConstraintsOnly->Set(bPartialSleepCollisionConstraintsOnly);
+
+		const bool bPartialSleepEnabled = true;
+		FSleepingTest Test(bPartialSleepEnabled);
+		Test.MakeChainOfSpheres(5);
+
+		Test.ParticleHandles[0]->SetObjectState(EObjectStateType::Static);
+		Test.ParticleHandles[1]->SetObjectState(EObjectStateType::Sleeping);
+		Test.ParticleHandles[2]->SetObjectState(EObjectStateType::Sleeping);
+		Test.ParticleHandles[3]->SetObjectState(EObjectStateType::Sleeping);
+		Test.ParticleHandles[4]->SetObjectState(EObjectStateType::Sleeping);
+
+		Test.Advance();
+
+		// All particles remain asleep since the entire island is asleep
+		for (int32 Id = 1; Id < Test.ParticleHandles.Num(); ++Id)
+		{
+			const FRigidBodyHandle_External* ParticleHandle = Test.ParticleHandles[Id];
+			EXPECT_TRUE(ParticleHandle->ObjectState() == EObjectStateType::Sleeping);
+		}
+		CVarPartialSleepCollisionConstraintsOnly->Set(bPartialSleepCollisionConstraintsOnlyDefault);
+	}
+
 	INSTANTIATE_TEST_SUITE_P(, SleepingTests, testing::Bool());
 }
\ No newline at end of file
diff --git a/Engine/Source/Runtime/Experimental/Chaos/Private/Chaos/Island/IslandManager.cpp b/Engine/Source/Runtime/Experimental/Chaos/Private/Chaos/Island/IslandManager.cpp
--- a/Engine/Source/Runtime/Experimental/Chaos/Private/Chaos/Island/IslandManager.cpp
+++ b/Engine/Source/Runtime/Experimental/Chaos/Private/Chaos/Island/IslandManager.cpp
@@ -66,7 +66,7 @@ namespace Chaos::CVars
 	/** Cvar to enable/disable triggering wake up for newly detected constraints with moving particles only (partial island sleeping only) */
 	bool bChaosSolverWakeOnMovingOnly = true;
 	FAutoConsoleVariableRef CVarChaosSolverWakeOnMovingOnly(TEXT("p.Chaos.Solver.Sleep.WakeOnMovingOnly"), bChaosSolverWakeOnMovingOnly, TEXT("[EXPERIMENTAL] Only newly detected constraints connected to a moving particle will trigger wake-up event (partial island sleep only)."));
-	
+
 	/** Cvar to adjust post-step waking triggered by velocity solver impulses.
 	* The solver impulse is converted to a change in velocity by division by the particle mass.
 	* Ideal setting = 10.0f, i.e. change in velocity > 10 will trigger wake of the sleeping particle. */
@@ -77,6 +77,10 @@ namespace Chaos::CVars
 	FRealSingle bChaosSolverMomentumPropagation = 0.0f;
 	FAutoConsoleVariableRef CVarChaosSolverMomentumPropagation(TEXT("p.Chaos.Solver.Sleep.MomentumPropagation"), bChaosSolverMomentumPropagation, TEXT("[EXPERIMENTAL] Ratio of the impact momentum transferred from one sleeping particle to another through a constraint during partial wake event (partial island sleeping only, 0 to disable, 1 to propagate with dissipation).[def:0.25]"));
 
+	/** Cvar to enable/disable partial sleep for collision constraints only, i.e. the sleep state of particles connect by other constraint types must match (partial island sleeping only) */
+	bool bChaosSolverPartialSleepCollisionConstraintsOnly = false;
+	FAutoConsoleVariableRef CVarChaosSolverPartialSleepCollisionConstraintsOnly(TEXT("p.Chaos.Solver.Sleep.PartialSleepCollisionConstraintsOnly"), bChaosSolverPartialSleepCollisionConstraintsOnly, TEXT("[EXPERIMENTAL] Only collision constraints permit partial sleep, i.e. the sleep state of particles connect by other constraint types must match (partial island sleep only)."));
+
 	/** Cvar to override the sleep counter threshold if necessary */
 	int32 ChaosSolverCollisionDefaultSleepCounterThreshold = 20;
 	FAutoConsoleVariableRef CVarChaosSolverCollisionDefaultSleepCounterThreshold(TEXT("p.Chaos.Solver.Sleep.Defaults.SleepCounterThreshold"), ChaosSolverCollisionDefaultSleepCounterThreshold, TEXT("Default counter threshold for sleeping.[def:20]"));
@@ -2541,6 +2545,16 @@ namespace Chaos::Private
 		// Make all non-moving particles in the island go to sleep
 		SleepNonMovingParticles(Island, Dt);
 
+		if (CVars::bChaosSolverPartialSleepCollisionConstraintsOnly)
+		{
+			// If this island contains any constraint which don't support partial sleeping,
+			// we need to traverse the constraints to make sure that connected particles are either both sleeping or both awake.
+			RewakeConstraintsNotSupportingPartialSleep(Island);
+
+			// NOTE: If (ChaosSolverPostStepWakeThreshold > 0.0f), we might encounter situations where an active-sleeping pair is connected by a non-supporting constraint. 
+			// However, we tolerate this since the state should only last for one tick before RewakeConstraintsNotSupportingPartialSleep will sync again.
+		}
+
 		// Deactivate all adjacent constraints between (a) 2 sleeping particles or (b) a sleeping and a non-dynamic particle
 		// Processing non-sleeping constraints is sufficient because no particles were woken up since the last call to @SyncConstraintSleepStateToParticles
 		PropagateParticleSleepToConstraints(Island);
@@ -2628,6 +2642,113 @@ namespace Chaos::Private
 		}
 	}
 
+	void FPBDIslandManager::RewakeConstraintsNotSupportingPartialSleep(FPBDIsland* Island)
+	{
+		// NOTE: We use node flags here since we want the updated state when processing.
+		const auto IsSleepingAndDynamic = [](const FPBDIslandParticle* Node) -> bool
+			{
+				return Node && !!Node->Flags.bIsSleeping && !!Node->Flags.bIsDynamic;
+			};
+
+		auto AddToWakeQueue = [](FPBDIslandParticle* Node, TArray<FGeometryParticleHandle*>& WakeUpQueue, TArray<FPBDIslandParticle*>& NodeQueue)
+			{
+				Node->Flags.bIsSleeping = false;
+				NodeQueue.Add(Node);
+				if (FGeometryParticleHandle* Particle = Node->GetParticle())
+				{
+					WakeUpQueue.Add(Particle);
+
+					// If the particle was re-woken, we need to restore the sleep counter from before the reset to 0 in SleepNonMovingParticles.
+					// We know that the particle could only go to sleep if Rigid->SleepCounter() > Node->SleepCounterThreshold.
+					// Thus, it holds that Rigid->SleepCounter() == Node->SleepCounterThreshold before the reset. 
+					// If the particle was not just deactivated, it would be woken up to prevent partial sleeping for non-supporting constraints and be set to the same counter.
+					if (FPBDRigidParticleHandle* Rigid = Particle->CastToRigidParticle())
+					{
+						Rigid->SetSleepCounter(static_cast<int8>(FMath::Min(Node->SleepCounterThreshold, TNumericLimits<int8>::Max())));
+					}
+				}
+			};
+
+		TArray<FGeometryParticleHandle*> WakeUpQueue;
+		TArray<FPBDIslandParticle*> NodeQueue;
+		WakeUpQueue.Reserve(Island->GetNumParticles());
+		NodeQueue.Reserve(Island->GetNumParticles());
+
+		// Check all awake constraints to make sure there are no active-sleeping particle pairs
+		// connected by constraints which do not support partial island sleep.
+		// Add all particles that need to re-woken to NodeQueue.
+		for (TArray<FPBDIslandConstraint*>& IslandEdges : Island->ContainerEdges)
+		{
+			for (FPBDIslandConstraint* IslandEdge : IslandEdges)
+			{
+				// Skip already sleeping constraints since must connect particles which were already 
+				// sleeping or kinematic during the solve and can therefore not have been just put to sleep.
+				if (!IslandEdge->Flags.bIsSleeping)
+				{
+					FPBDIslandParticle* Node0 = IslandEdge->Nodes[0];
+					FPBDIslandParticle* Node1 = IslandEdge->Nodes[1];
+					const FConstraintHandle* Constraint = IslandEdge->GetConstraint();
+
+					if (Node0 != nullptr && Node1 != nullptr)
+					{
+						// Reactivate particles between constraints which do not support partial sleeping
+						if (Constraint && !Constraint->SupportsPartialIslandSleeping())
+						{
+							// Cases for re-waking
+							// a) One particle asleep, one particle awake, both dynamic.
+							const bool bIsSleeping0 = IsSleepingAndDynamic(Node0);
+							const bool bIsSleeping1 = IsSleepingAndDynamic(Node1);
+							const bool bIsAwake0 = Node0->Flags.bIsDynamic && !bIsSleeping0;
+							const bool bIsAwake1 = Node1->Flags.bIsDynamic && !bIsSleeping1;
+							if (bIsSleeping0 && bIsAwake1)
+							{
+								IslandEdge->Flags.bWasProcessed = true;
+								AddToWakeQueue(Node0, WakeUpQueue, NodeQueue);
+							}
+							else if (bIsAwake0 && bIsSleeping1)
+							{
+								IslandEdge->Flags.bWasProcessed = true;
+								AddToWakeQueue(Node1, WakeUpQueue, NodeQueue);
+							}
+						}
+					}
+				}
+			}
+		}
+
+		// Walk the constraint graph to re-wake all particles connected to particles in NodeQueue
+		// by constraints which do not support partial sleeping.
+		// This ensures that all particles connected by non-supporting constraints have the same sleep/awake state.
+		for (int32 NodeIndex = 0; NodeIndex < NodeQueue.Num(); ++NodeIndex)
+		{
+			FPBDIslandParticle* ThisNode = NodeQueue[NodeIndex];
+			for (FPBDIslandConstraint* Edge : ThisNode->Edges)
+			{
+				if (!Edge->Flags.bWasProcessed)
+				{
+					Edge->Flags.bWasProcessed = true;
+
+					const FConstraintHandle* Constraint = Edge->GetConstraint();
+					if (Constraint && !Constraint->SupportsPartialIslandSleeping())
+					{
+						FPBDIslandParticle* OtherNode = (ThisNode == Edge->Nodes[0]) ? Edge->Nodes[1] : Edge->Nodes[0];
+						if ((OtherNode != nullptr) && OtherNode->Flags.bIsDynamic)
+						{
+							if (IsSleepingAndDynamic(OtherNode))
+							{
+								AddToWakeQueue(OtherNode, WakeUpQueue, NodeQueue);
+							}
+						}
+					}
+				}
+			}
+		}
+
+		// Don't reset the sleep counter here since they have already been set in AddToWakeQueue.
+		WakeUpParticles(WakeUpQueue, /*bResetSleepCounter=*/ false);
+		ResetEdgeFlags(Island);
+	}
+
 	void FPBDIslandManager::PropagateParticleSleepToConstraints(FPBDIsland* Island)
 	{
 		for (TArray<FPBDIslandConstraint*>& IslandEdges : Island->ContainerEdges)
@@ -2882,10 +3003,12 @@ namespace Chaos::Private
 						const bool bIsTriggerNode = NodeIndex < NumNodesInQueue;
 						const bool bIsConfirmedPropagation = bIsTriggerNode && (Edge->Flags.bIsNew || bTriggerNodesAlwaysPropagate);
 						const bool bIsUpwardPropagation = OtherNode->Level >= Edge->Level;	// Propagation upward through the levels, i.e. away from closest kinematic
+						const bool bEdgePropagatesPartialWake = CVars::bChaosSolverPartialSleepCollisionConstraintsOnly && Edge->GetConstraint() && !Edge->GetConstraint()->SupportsPartialIslandSleeping();
 
 						// 1) We process one of the trigger nodes and propagate through a new constraint (if bTriggerNodesAlwaysPropagate == false).
 						// 2) We propagate upward in the constraint graph levels (away from kinematic).
-						if (bIsConfirmedPropagation || bIsUpwardPropagation)
+						// 3) Partial sleep is enabled for collisions only and the constraint type does not support connecting sleeping and active particles.
+						if (bIsConfirmedPropagation || bIsUpwardPropagation || bEdgePropagatesPartialWake)
 						{
 							const FPBDParticleMomentumSq* NodeMomentum = MomentumQueue.IsEmpty() ? nullptr : &MomentumQueue[NodeIndex];
 							ProcessParticleWake(OtherNode, NodeQueue[NodeIndex], NodeMomentum, NodeQueue, MomentumQueue, ParticlesToWakeUp, /*bAlwaysPropagate=*/ true);
diff --git a/Engine/Source/Runtime/Experimental/Chaos/Public/Chaos/Collision/PBDCollisionConstraint.h b/Engine/Source/Runtime/Experimental/Chaos/Public/Chaos/Collision/PBDCollisionConstraint.h
--- a/Engine/Source/Runtime/Experimental/Chaos/Public/Chaos/Collision/PBDCollisionConstraint.h
+++ b/Engine/Source/Runtime/Experimental/Chaos/Public/Chaos/Collision/PBDCollisionConstraint.h
@@ -417,6 +417,7 @@ namespace Chaos
 		bool UsePerContactInitialPhi() const { return Flags.bUsePerContactInitialPhi; }
 
 		virtual bool SupportsSleeping() const override final { return true; }
+		virtual bool SupportsPartialIslandSleeping() const override final { return true; }
 		CHAOS_API virtual bool IsSleeping() const override final;
 		CHAOS_API virtual void SetIsSleeping(const bool bInIsSleeping) override final;
 
diff --git a/Engine/Source/Runtime/Experimental/Chaos/Public/Chaos/ConstraintHandle.h b/Engine/Source/Runtime/Experimental/Chaos/Public/Chaos/ConstraintHandle.h
--- a/Engine/Source/Runtime/Experimental/Chaos/Public/Chaos/ConstraintHandle.h
+++ b/Engine/Source/Runtime/Experimental/Chaos/Public/Chaos/ConstraintHandle.h
@@ -166,6 +166,9 @@ namespace Chaos
 		// Does this constraint have the concept of sleep? (only really used for debug validation)
 		virtual bool SupportsSleeping() const { return false; }
 
+		// Can this constraint connect sleeping and active particles for partial island sleeping?
+		virtual bool SupportsPartialIslandSleeping() const { return false; }
+
 		virtual bool IsSleeping() const { return false; }
 		virtual void SetIsSleeping(const bool bInIsSleeping) {}
 
diff --git a/Engine/Source/Runtime/Experimental/Chaos/Public/Chaos/Island/IslandManager.h b/Engine/Source/Runtime/Experimental/Chaos/Public/Chaos/Island/IslandManager.h
--- a/Engine/Source/Runtime/Experimental/Chaos/Public/Chaos/Island/IslandManager.h
+++ b/Engine/Source/Runtime/Experimental/Chaos/Public/Chaos/Island/IslandManager.h
@@ -613,6 +613,7 @@ namespace Chaos::Private
 		// Partial Island Sleeping
 		CHAOS_API void ProcessPartialIslandSleep(FPBDIsland* Island, const FRealSingle Dt);
 		CHAOS_API void SleepNonMovingParticles(FPBDIsland* Island, const FRealSingle Dt);
+		CHAOS_API void RewakeConstraintsNotSupportingPartialSleep(FPBDIsland* Island);
 		CHAOS_API void PropagateParticleSleepToConstraints(FPBDIsland* Island);
 		CHAOS_API void PostStepPartialWake(FPBDIsland* Island);
 
