Clean up
Rename the integration barrier:  WaitIntegrationComplete

diff --git a/Engine/Source/Runtime/Experimental/Chaos/Private/Chaos/PBDRigidsEvolutionGBF.cpp b/Engine/Source/Runtime/Experimental/Chaos/Private/Chaos/PBDRigidsEvolutionGBF.cpp
--- a/Engine/Source/Runtime/Experimental/Chaos/Private/Chaos/PBDRigidsEvolutionGBF.cpp
+++ b/Engine/Source/Runtime/Experimental/Chaos/Private/Chaos/PBDRigidsEvolutionGBF.cpp
@@ -1084,9 +1084,9 @@ void FPBDRigidsEvolutionGBF::ApplyKinematicTargets(const FReal Dt, const FReal S
 	}
 }
 
-void FPBDRigidsEvolutionGBF::SafePostIntegrateCallback()
+void FPBDRigidsEvolutionGBF::WaitIntegrationComplete()
 {
-	TaskDispatcher.SafePostIntegrateCallback();
+	TaskDispatcher.WaitIntegrationComplete();
 }
 
 template <typename LambdaType>
diff --git a/Engine/Source/Runtime/Experimental/Chaos/Private/PBDRigidsSolver.cpp b/Engine/Source/Runtime/Experimental/Chaos/Private/PBDRigidsSolver.cpp
--- a/Engine/Source/Runtime/Experimental/Chaos/Private/PBDRigidsSolver.cpp
+++ b/Engine/Source/Runtime/Experimental/Chaos/Private/PBDRigidsSolver.cpp
@@ -743,17 +743,17 @@ namespace Chaos
 
 	void FPBDRigidsSolver::OnEvolutionPostIntegrate(FReal Dt)
 	{
-		bool bHasSafeThreadCallback = false;
+		bool bHasIntegrationCompleted = false;
 		for (ISimCallbackObject* Callback : SimCallbackObjects)
 		{
 			if (Callback->HasOption(ESimCallbackOptions::PostIntegrate))
 			{
-				if (!bHasSafeThreadCallback)
+				if (!bHasIntegrationCompleted)
 				{
 					// If a post integrate callback is called, 
 					// we need to wait for all workers threads integrating to finish.
-					MEvolution->SafePostIntegrateCallback();
-					bHasSafeThreadCallback = true;
+					MEvolution->WaitIntegrationComplete();
+					bHasIntegrationCompleted = true;
 				}
 				
 				FScopedTraceSolverCallback TraceCallback(Callback);
diff --git a/Engine/Source/Runtime/Experimental/Chaos/Public/Chaos/PBDRigidsEvolution.h b/Engine/Source/Runtime/Experimental/Chaos/Public/Chaos/PBDRigidsEvolution.h
--- a/Engine/Source/Runtime/Experimental/Chaos/Public/Chaos/PBDRigidsEvolution.h
+++ b/Engine/Source/Runtime/Experimental/Chaos/Public/Chaos/PBDRigidsEvolution.h
@@ -871,7 +871,7 @@ public:
 	const FString& GetName() const { return EvolutionName; }
 
 	// Internal used only 
-	UE_INTERNAL virtual void SafePostIntegrateCallback() { check(false); };
+	UE_INTERNAL virtual void WaitIntegrationComplete() { check(false); };
 
 protected:
 	int32 NumConstraints() const
diff --git a/Engine/Source/Runtime/Experimental/Chaos/Public/Chaos/PBDRigidsEvolutionGBF.h b/Engine/Source/Runtime/Experimental/Chaos/Public/Chaos/PBDRigidsEvolutionGBF.h
--- a/Engine/Source/Runtime/Experimental/Chaos/Public/Chaos/PBDRigidsEvolutionGBF.h
+++ b/Engine/Source/Runtime/Experimental/Chaos/Public/Chaos/PBDRigidsEvolutionGBF.h
@@ -274,7 +274,7 @@ namespace Chaos
 		void Integrate(const TParticleView& InParticles, FReal Dt) { Integrate(Dt); }
 
 		// Internal used only
-		UE_INTERNAL virtual void SafePostIntegrateCallback() override;
+		UE_INTERNAL virtual void WaitIntegrationComplete() override;
 
 	protected:
 
diff --git a/Engine/Source/Runtime/Experimental/Chaos/Public/Chaos/TaskDispatcherEvolution.h b/Engine/Source/Runtime/Experimental/Chaos/Public/Chaos/TaskDispatcherEvolution.h
--- a/Engine/Source/Runtime/Experimental/Chaos/Public/Chaos/TaskDispatcherEvolution.h
+++ b/Engine/Source/Runtime/Experimental/Chaos/Public/Chaos/TaskDispatcherEvolution.h
@@ -282,9 +282,9 @@ public:
 		}, AsyncQueueTask, LowLevelTasks::ETaskPriority::High);
 	}
 
-	// This will be call if it is required to wait tasks result to invoke a callback.
+	// This will be called if it is required to wait tasks result to invoke a callback.
 	// We never know what a callback is going to use.
-	void SafePostIntegrateCallback()
+	void WaitIntegrationComplete()
 	{
 		TArray<UE::Tasks::FTask> Tasks = { UpdateViewTask };
 		Tasks.Reserve(IntegrationPendingTasks.Num() + KinematicPendingTasks.Num() + 1);
