Fix crash in flush asynch
Task dependecies could have been overwritten, and them task dependencies were not respected anymore
Clean up task and created one task for each job.

diff --git a/Engine/Source/Runtime/Experimental/Chaos/Public/Chaos/TaskDispatcherEvolution.h b/Engine/Source/Runtime/Experimental/Chaos/Public/Chaos/TaskDispatcherEvolution.h
--- a/Engine/Source/Runtime/Experimental/Chaos/Public/Chaos/TaskDispatcherEvolution.h
+++ b/Engine/Source/Runtime/Experimental/Chaos/Public/Chaos/TaskDispatcherEvolution.h
@@ -16,7 +16,9 @@ class FTaskDispatcherEvolution
 private:
 	constexpr static int32 MinParticlePerTask = 40;
 
-	UE::Tasks::FTask AsyncQueueTask;
+	UE::Tasks::FTask AsyncQueueDynTask;
+	UE::Tasks::FTask AsyncQueueKinTask;
+	UE::Tasks::FTask FlushAsyncQueueTask;
 	UE::Tasks::FTask UpdateViewTask;
 
 	TArray<UE::Tasks::FTask> IntegrationPendingTasks;
@@ -86,7 +88,7 @@ public:
 	void DispatchDynAsyncDirty(Lambda AsyncDirtyWork)
 	{
 		// Why do we add the dynamic particle to the async queue which is suppose to be for static tree ? 
-		AsyncQueueTask = UE::Tasks::Launch(UE_SOURCE_LOCATION, [AsyncDirtyWork]()
+		AsyncQueueDynTask = UE::Tasks::Launch(UE_SOURCE_LOCATION, [AsyncDirtyWork]()
 		{
 			QUICK_SCOPE_CYCLE_COUNTER(STAT_DirtyAsyncDynamicParticleTask);
 			AsyncDirtyWork();
@@ -98,14 +100,14 @@ public:
 	{
 		// To start dirtying Kinematics particles, the Dynamics particles must be finished
 		// Is it necessary to update the kinematic particles for the async data ? 
-		AsyncQueueTask = UE::Tasks::Launch(UE_SOURCE_LOCATION, [AsyncDirty]()
+		AsyncQueueKinTask = UE::Tasks::Launch(UE_SOURCE_LOCATION, [AsyncDirty]()
 		{
 			QUICK_SCOPE_CYCLE_COUNTER(STAT_DirtyMovingKinematicsTask);
 			AsyncDirty();
-		}, AsyncQueueTask, LowLevelTasks::ETaskPriority::High);
+		}, AsyncQueueDynTask, LowLevelTasks::ETaskPriority::High);
 
 		TArray<UE::Tasks::FTask> PendingTasks;
-		PendingTasks.Add(AsyncQueueTask);
+		PendingTasks.Add(AsyncQueueKinTask);
 
 		UE::Tasks::FTask MovingKinematicsTask;
 		// done with update, let's clear the tracking structures
@@ -122,7 +124,7 @@ public:
 		}
 		else
 		{
-			MovingKinematicsTask = MoveTemp(AsyncQueueTask);
+			MovingKinematicsTask = AsyncQueueKinTask;
 		}
 
 		// Updating the views can be done only when all kinematics and dynamics update are finished. // This comment is wrong
@@ -275,11 +277,11 @@ public:
 	template<typename Lambda>
 	void FlushAccelerationQueue(Lambda FlushAccelerationQueueLambda)
 	{
-		AsyncQueueTask = UE::Tasks::Launch(UE_SOURCE_LOCATION, [FlushAccelerationQueueLambda]()
+		FlushAsyncQueueTask = UE::Tasks::Launch(UE_SOURCE_LOCATION, [FlushAccelerationQueueLambda]()
 		{
 			QUICK_SCOPE_CYCLE_COUNTER(STAT_FlushAsyncAccelerationQueueTask);
 			FlushAccelerationQueueLambda();
-		}, AsyncQueueTask, LowLevelTasks::ETaskPriority::High);
+		}, TArray<UE::Tasks::FTask>{AsyncQueueDynTask, AsyncQueueKinTask}, LowLevelTasks::ETaskPriority::High);
 	}
 
 	// This will be called if it is required to wait tasks result to invoke a callback.
@@ -298,7 +300,7 @@ public:
 
 	void WaitAsyncQueueTask()
 	{
-		UE::Tasks::Wait(AsyncQueueTask);
+		UE::Tasks::Wait(FlushAsyncQueueTask);
 	}
 
 	void WaitTaskEndSpatial()
