Chaos - expose IsIntialContact to contact modifiers

diff --git a/Engine/Source/Runtime/Experimental/Chaos/Private/Chaos/ContactModification.cpp b/Engine/Source/Runtime/Experimental/Chaos/Private/Chaos/ContactModification.cpp
--- a/Engine/Source/Runtime/Experimental/Chaos/Private/Chaos/ContactModification.cpp
+++ b/Engine/Source/Runtime/Experimental/Chaos/Private/Chaos/ContactModification.cpp
@@ -60,6 +60,17 @@ namespace Chaos
 		return DeepestIdx;
 	}
 
+	bool FContactPairModifier::GetIsInitialContact() const
+	{
+		return Constraint->IsInitialContact();
+	}
+
+	void FContactPairModifier::SetIsInitialContact(bool bIsInitialContact) const
+	{
+		Constraint->SetModifierApplied();
+
+		Constraint->SetIsInitialContact(bIsInitialContact);
+	}
 
 	const FImplicitObject* FContactPairModifier::GetContactGeometry(int32 ParticleIdx)
 	{
diff --git a/Engine/Source/Runtime/Experimental/Chaos/Public/Chaos/ContactModification.h b/Engine/Source/Runtime/Experimental/Chaos/Public/Chaos/ContactModification.h
--- a/Engine/Source/Runtime/Experimental/Chaos/Public/Chaos/ContactModification.h
+++ b/Engine/Source/Runtime/Experimental/Chaos/Public/Chaos/ContactModification.h
@@ -76,6 +76,17 @@ namespace Chaos
 		*/
 		CHAOS_API int32 GetDeepestContactIndex() const;
 
+		/**
+		* @return true if this contact is an initial contact, i.e., the pair of shapes in the contact were not touching last frame.
+		*/
+		CHAOS_API bool GetIsInitialContact() const;
+
+		/**
+		* Set whether this contact is considered an initial contact. This controls how the contact behaves with regard to the initial
+		* overlap depenetration system.
+		*/
+		CHAOS_API void SetIsInitialContact(bool bIsInitialContact) const;
+
 		/**
 		* @return Geometry of specified particle in contact pair.
 		*/
