Fix CIS

Undo changelist 44852070

diff --git a/Engine/Source/Runtime/Experimental/Chaos/Private/Chaos/PBDRigidsEvolution.cpp b/Engine/Source/Runtime/Experimental/Chaos/Private/Chaos/PBDRigidsEvolution.cpp
--- a/Engine/Source/Runtime/Experimental/Chaos/Private/Chaos/PBDRigidsEvolution.cpp
+++ b/Engine/Source/Runtime/Experimental/Chaos/Private/Chaos/PBDRigidsEvolution.cpp
@@ -23,6 +23,9 @@ FAutoConsoleVariableRef CVarChaosNumPushOutIterationsOverride(TEXT("p.ChaosNumPu
 int32 ChaosNumContactIterationsOverride = -1;
 FAutoConsoleVariableRef CVarChaosNumContactIterationsOverride(TEXT("p.ChaosNumContactIterationsOverride"), ChaosNumContactIterationsOverride, TEXT("Override for num contact iterations if >= 0. [def:-1]"));
 
+bool bPruningParticleToInsertInAccelerationStructure = true;
+FAutoConsoleVariableRef CVarPruningParticleToInsertInAccelerationStructure(TEXT("p.ChaosPruningParticleToInsertInAccelerationStructure"), bPruningParticleToInsertInAccelerationStructure, TEXT("Enables pruning particle before updating them in the accceleration structure. This is an optimization that reduce the dynamic acceleration structure update."));
+
 namespace Chaos
 {
 	namespace CVars
@@ -221,7 +224,6 @@ namespace Chaos
 		, bIsSingleThreaded(InIsSingleThreaded)
 		, bCanStartAsyncTasks(true)
 		, LatestExternalTimestampConsumed_Internal(-1)
-		, TaskDispatcher(Particles, InternalAccelerationQueue)
 		, bAccelerationStructureTaskStarted(nullptr)
 		, bAccelerationStructureTaskSignalKill(nullptr)
 		, SpatialCollectionFactory(new FDefaultCollectionFactory())
@@ -628,48 +630,82 @@ namespace Chaos
 		// The bounding volume associated to the particle will be updated in the acceleration structure.
 		QUICK_SCOPE_CYCLE_COUNTER(STAT_PruneInternalPendingData);
 
+		const int32 NumWorker = !!GSingleThreadedPhysics ? 1 : FMath::Max(FTaskGraphInterface::Get().GetNumWorkerThreads(), 1);
+		const int32 NumParticles = InternalAccelerationQueue.PendingData.Num();
+		const int32 ParticleByContext = FMath::Max(FMath::DivideAndRoundUp(NumParticles, NumWorker), SmallBatchSize);
+		const int32 NumBatch = FMath::DivideAndRoundUp(NumParticles, ParticleByContext);
 
-		auto HasToBeUpdated = [&AccelerationStructure](const FPendingSpatialData& PendingData) -> bool
+		if (Chaos::ShouldExecuteTasksInParallel(NumBatch))
 		{
-			FGeometryParticleHandle* UpdateParticle = PendingData.AccelerationHandle.GetGeometryParticleHandle_PhysicsThread();
-			return (PendingData.Operation == EPendingSpatialDataOperation::Delete ||
-				(UpdateParticle->HasBounds() && AccelerationStructure.NeedUpdateElementIn(UpdateParticle, UpdateParticle->WorldSpaceInflatedBounds(), PendingData.SpatialIdx)));
-
-		};
+			TArray<TArray<FPendingSpatialData>> NewPendingDataArray;
+			NewPendingDataArray.SetNum(NumBatch);
+			TArray<UE::Tasks::FTask> PendingTasks;
+			PendingTasks.Reserve(NumBatch);
+			for (int32 BatchIndex = 0; BatchIndex < NumBatch; BatchIndex++)
+			{
+				const int32 FirstIndex = BatchIndex * ParticleByContext;
+				int32 LastIndex = (BatchIndex + 1) * ParticleByContext;
+				LastIndex = BatchIndex == NumBatch - 1 ? FMath::Min(NumParticles, LastIndex) : LastIndex;
+				TArray<FPendingSpatialData>& NewPendingData = NewPendingDataArray[BatchIndex];
+				UE::Tasks::FTask PendingTask = UE::Tasks::Launch(UE_SOURCE_LOCATION, [this, &AccelerationStructure, FirstIndex, LastIndex, &NewPendingData]()
+				{
+					QUICK_SCOPE_CYCLE_COUNTER(STAT_PruneInternalPendingDataTask);
+					for (int32 Index = FirstIndex; Index < LastIndex; Index++)
+					{
+						const FPendingSpatialData& PendingData = InternalAccelerationQueue.PendingData[Index];
+						FGeometryParticleHandle* UpdateParticle = PendingData.AccelerationHandle.GetGeometryParticleHandle_PhysicsThread();
+
+						if (PendingData.Operation == EPendingSpatialDataOperation::Delete ||
+							(UpdateParticle->HasBounds() && AccelerationStructure.NeedUpdateElementIn(UpdateParticle, UpdateParticle->WorldSpaceInflatedBounds(), PendingData.SpatialIdx)))
+						{
+							NewPendingData.Add(PendingData);
+						}
+					}
+				}, LowLevelTasks::ETaskPriority::High);
+				PendingTasks.Add(PendingTask);
+			}
 
-		if (!GSingleThreadedPhysics)
-		{
-			TaskDispatcher.PruneInternalPendingData(HasToBeUpdated);
+			UE::Tasks::Wait(PendingTasks);
+			// Once all threads are finished merge all data in the acceleration queue
+			// Maybe this could be multi threaded, but since it is really fast not sure if it is worth it.
+			InternalAccelerationQueue.PendingData.Reset();
+			for (int32 BatchIndex = 0; BatchIndex < NumBatch; BatchIndex++)
+			{
+				InternalAccelerationQueue.PendingData.Append(NewPendingDataArray[BatchIndex]);
+			}
 		}
 		else
 		{
-			InternalAccelerationQueue.CleanUpDuplicated();
-			for (TArray<FPendingSpatialData>& PendingDataArray : InternalAccelerationQueue.PendingDataArrays)
+			// Small enough task to run immediately and avoid parallel overhead
+			QUICK_SCOPE_CYCLE_COUNTER(STAT_PruneInternalPendingDataTask);
+
+			TArray<FPendingSpatialData> OldPendingData = MoveTemp(InternalAccelerationQueue.PendingData);
+
+			InternalAccelerationQueue.PendingData.Reset();
+
+			for (int32 Index = 0; Index < NumParticles; Index++)
 			{
-				TArray<FPendingSpatialData> NewPendingData;
-				NewPendingData.Reserve(PendingDataArray.Num());
-				for (const FPendingSpatialData& PendingData : PendingDataArray)
+				const FPendingSpatialData& PendingData = OldPendingData[Index];
+				FGeometryParticleHandle* UpdateParticle = PendingData.AccelerationHandle.GetGeometryParticleHandle_PhysicsThread();
+
+				if (PendingData.Operation == EPendingSpatialDataOperation::Delete ||
+					(UpdateParticle->HasBounds() && AccelerationStructure.NeedUpdateElementIn(UpdateParticle, UpdateParticle->WorldSpaceInflatedBounds(), PendingData.SpatialIdx)))
 				{
-					if (HasToBeUpdated(PendingData))
-					{
-						NewPendingData.Add(PendingData);
-					}
+					InternalAccelerationQueue.PendingData.Add(PendingData);
 				}
-				PendingDataArray = MoveTemp(NewPendingData);
 			}
 		}
 	}
 
 	void FPBDRigidsEvolutionBase::FlushInternalAccelerationQueue()
 	{
-		PruneInternalPendingData(*InternalAcceleration);
-		// Could be linked by dependent task, but could start as soon a previous task is finished
-		for (const TArray<FPendingSpatialData>& PendingDataArray : InternalAccelerationQueue.PendingDataArrays)
+		if (bPruningParticleToInsertInAccelerationStructure)
 		{
-			for (const FPendingSpatialData& PendingData : PendingDataArray)
-			{
-				ApplyParticlePendingData(PendingData, *InternalAcceleration, false, true);
-			}
+			PruneInternalPendingData(*InternalAcceleration);
+		}
+		for (const FPendingSpatialData& PendingData : InternalAccelerationQueue.PendingData)
+		{
+			ApplyParticlePendingData(PendingData, *InternalAcceleration, false, true);
 		}
 		InternalAcceleration->SetSyncTimestamp(LatestExternalTimestampConsumed_Internal);
 		InternalAccelerationQueue.Reset();
@@ -677,28 +713,25 @@ namespace Chaos
 
 	void FPBDRigidsEvolutionBase::FlushAsyncAccelerationQueue()
 	{
-		auto FlushAsyncAccelerationQueueLambda = [this]() 
+		for (const FPendingSpatialData& PendingData : AsyncAccelerationQueue.PendingData)
 		{
-			for (const FPendingSpatialData& PendingData : AsyncAccelerationQueue.PendingData)
-			{
-				ApplyParticlePendingData(PendingData, *AsyncInternalAcceleration, true, false); //only the first queue needs to update the cached acceleration
-				ApplyParticlePendingData(PendingData, *AsyncExternalAcceleration, false, false);
-			}
+			ApplyParticlePendingData(PendingData, *AsyncInternalAcceleration, true, false); //only the first queue needs to update the cached acceleration
+			ApplyParticlePendingData(PendingData, *AsyncExternalAcceleration, false, false);
+		}
 
-			//NOTE: This assumes that we are never creating a PT particle that is replicated to GT
-			//At the moment that is true, and it seems like we have enough mechanisms to avoid this direction
-			//If we want to support that, the UniqueIndex must be kept around until GT goes away
-			//This is hard to do, but would probably mean the ownership of the index is in the proxy
-			for (FUniqueIdx UniqueIdx : UniqueIndicesPendingRelease)
-			{
-				ReleaseIdx(UniqueIdx);
-			}
-			UniqueIndicesPendingRelease.Reset();
-			AsyncAccelerationQueue.Reset();
-			AsyncInternalAcceleration->SetSyncTimestamp(LatestExternalTimestampConsumed_Internal);
-			AsyncExternalAcceleration->SetSyncTimestamp(LatestExternalTimestampConsumed_Internal);
-		};
-		TaskDispatcher.FlushAccelerationQueue(FlushAsyncAccelerationQueueLambda);
+				//NOTE: This assumes that we are never creating a PT particle that is replicated to GT
+				//At the moment that is true, and it seems like we have enough mechanisms to avoid this direction
+				//If we want to support that, the UniqueIndex must be kept around until GT goes away
+				//This is hard to do, but would probably mean the ownership of the index is in the proxy
+		for (FUniqueIdx UniqueIdx : UniqueIndicesPendingRelease)
+		{
+			ReleaseIdx(UniqueIdx);
+		}
+		UniqueIndicesPendingRelease.Reset();
+		AsyncAccelerationQueue.Reset();
+
+		AsyncInternalAcceleration->SetSyncTimestamp(LatestExternalTimestampConsumed_Internal);
+		AsyncExternalAcceleration->SetSyncTimestamp(LatestExternalTimestampConsumed_Internal);
 	}
 
 	//TODO: make static and _External suffix
@@ -783,6 +816,7 @@ namespace Chaos
 		SCOPE_CYCLE_COUNTER(STAT_ComputeIntermediateSpatialAcceleration);
 		CHAOS_SCOPED_TIMER(ComputeIntermediateSpatialAcceleration);
 
+		bool ForceFullBuild = InternalAccelerationQueue.Num() > AccelerationStructureTimeSlicingMaxQueueSizeBeforeForce;
 
 		if (!AccelerationStructureTaskComplete)
 		{
@@ -792,8 +826,6 @@ namespace Chaos
 			AsyncInternalAcceleration = GetFreeSpatialAcceleration_Internal();
 			AsyncExternalAcceleration = GetFreeSpatialAcceleration_Internal();
 			FlushInternalAccelerationQueue();
-			
-			TaskDispatcher.WaitAsyncQueueTask();
 
 			// Give game thread an empty structure to pop
 			ExternalStructuresQueue.Enqueue(AsyncExternalAcceleration);
@@ -822,7 +854,7 @@ namespace Chaos
 				bNeedsReset = true;
 
 				FlushAsyncAccelerationQueue();
-
+				
 				if (AccelerationStructureSplitStaticAndDynamic == 1)
 				{
 					FlushInternalAccelerationQueue();
@@ -832,14 +864,13 @@ namespace Chaos
 					//other queues are no longer needed since we've flushed all operations and now have a pristine structure
 					InternalAccelerationQueue.Reset();
 				}
-				TaskDispatcher.WaitAsyncQueueTask();
 
 				// At this point the InternalAcceleration and the AsyncAcceleration structures will all be consistent with each other
 				// (If FlushInternalAccelerationQueue was flushed above)
 				// but the AsyncAcceleration now probably have fewer dirty elements (they were just rebuilt)
-
+				
 				//swap acceleration structure for new one
-				std::swap(InternalAcceleration, AsyncInternalAcceleration);
+					std::swap(InternalAcceleration, AsyncInternalAcceleration);
 				// The old InternalAcceleration is now in AsyncInternalAcceleration, and the constituent structures will be reused if no changes have been detected
 
 				if (AccelerationStructureUseDynamicTree)
@@ -869,9 +900,7 @@ namespace Chaos
 			}
 			else
 			{
-				FlushInternalAccelerationQueue();
-
-				TaskDispatcher.WaitAsyncQueueTask();
+					FlushInternalAccelerationQueue();
 
 				if (GAccelerationStructureCacheOverlappingLeaves  && AccelerationStructureUseDynamicTree)
 				{
@@ -883,14 +912,12 @@ namespace Chaos
 			if (bCanStartAsyncTasks)
 			{
 				// we run the task for both starting a new accel structure as well as for the time-slicing
-				bool ForceFullBuild = AsyncAccelerationQueue.Num() > AccelerationStructureTimeSlicingMaxQueueSizeBeforeForce;
 				AccelerationStructureTaskComplete = TGraphTask<FChaosAccelerationStructureTask>::CreateTask().ConstructAndDispatchWhenReady(*SpatialCollectionFactory, SpatialAccelerationCache, AsyncInternalAcceleration, AsyncExternalAcceleration, ForceFullBuild, bIsSingleThreaded, bNeedsReset, &bAccelerationStructureTaskStarted, &bAccelerationStructureTaskSignalKill);
 			}
 		}
 		else
 		{
 			FlushInternalAccelerationQueue();
-			TaskDispatcher.WaitAsyncQueueTask();
 		}
 	}
 
diff --git a/Engine/Source/Runtime/Experimental/Chaos/Private/Chaos/PBDRigidsEvolutionGBF.cpp b/Engine/Source/Runtime/Experimental/Chaos/Private/Chaos/PBDRigidsEvolutionGBF.cpp
--- a/Engine/Source/Runtime/Experimental/Chaos/Private/Chaos/PBDRigidsEvolutionGBF.cpp
+++ b/Engine/Source/Runtime/Experimental/Chaos/Private/Chaos/PBDRigidsEvolutionGBF.cpp
@@ -573,7 +573,6 @@ void FPBDRigidsEvolutionGBF::AdvanceOneTimeStepImpl(const FReal Dt, const FSubSt
 			CSV_SCOPED_TIMING_STAT(PhysicsVerbose, StepSolver_ComputeIntermediateSpatialAcceleration);
 			Base::ComputeIntermediateSpatialAcceleration();
 		}
-		TaskDispatcher.WaitTaskEndSpatial(); 
 	}
 
 	// Collision detection is sensitive to duplication bugs in the particle views
@@ -825,17 +824,30 @@ void FPBDRigidsEvolutionGBF::Integrate(FReal Dt)
 	const bool bForceMACD = CVars::bChaosForceMACD;
 
 	const TArray<FPBDRigidParticleHandle*>& ParticlesArray = Particles.GetActiveParticlesArray();
+	const int32 NumParticles = ParticlesArray.Num();
+	int32 NumTasks = !!GSingleThreadedPhysics ? 1 : FMath::Min(FTaskGraphInterface::Get().GetNumWorkerThreads(), Chaos::MaxNumWorkers);
+	NumTasks = FMath::Max(FMath::Min(NumTasks, NumParticles), 1);
+	constexpr int32 MinParticlePerTask = 40;
+	const int32 ParticleByTask = FMath::Max(FMath::DivideAndRoundUp(NumParticles, NumTasks), MinParticlePerTask);
+	int32 NumBatches = FMath::DivideAndRoundUp(NumParticles, ParticleByTask);
 
-	auto IntegrateWork = [this, &ParticlesArray,
-		Dt, HackMaxAngularSpeedSq, HackMaxLinearSpeedSq, bAllowMACD, bForceMACD, VelocityBoundsMultiplierMACD, MaxVelocityBoundsExpansionMACD,
-		VelocityBoundsMultiplier, MaxVelocityBoundsExpansion, BoundsThickness](int32 BatchIndex, int32 StartIndex, int32 EndIndex)
+	IntegrationPendingTasks.Reset(NumBatches);
+	for (int32 BatchIndex = 0; BatchIndex < NumBatches; BatchIndex++)
 	{
-		for (int32 Index = StartIndex; Index < EndIndex; Index++)
+		const int32 StartIndex = BatchIndex * ParticleByTask;
+		int32 EndIndex = (BatchIndex + 1) * ParticleByTask;
+		EndIndex = FMath::Min(NumParticles, EndIndex);
+
+		UE::Tasks::FTask PendingTask = UE::Tasks::Launch(UE_SOURCE_LOCATION, [this, StartIndex, EndIndex, &ParticlesArray,
+		Dt, HackMaxAngularSpeedSq, HackMaxLinearSpeedSq, bAllowMACD, bForceMACD, VelocityBoundsMultiplierMACD, MaxVelocityBoundsExpansionMACD,
+		VelocityBoundsMultiplier, MaxVelocityBoundsExpansion, BoundsThickness]()
 		{
-			//question: can we enforce this at the API layer? Right now islands contain non dynamic which makes this hard
-			if (FPBDRigidParticleHandle* PBDParticle = ParticlesArray[Index]->CastToRigidParticle();  PBDParticle)
+			QUICK_SCOPE_CYCLE_COUNTER(STAT_IntegrateTask);
+			for (int32 Index = StartIndex; Index < EndIndex; Index++)
 			{
-				if (PBDParticle->ObjectState() == EObjectStateType::Dynamic)
+				//question: can we enforce this at the API layer? Right now islands contain non dynamic which makes this hard
+				FPBDRigidParticleHandle* PBDParticle = ParticlesArray[Index]->CastToRigidParticle();
+				if (PBDParticle && PBDParticle->ObjectState() == EObjectStateType::Dynamic)
 				{
 					FPBDRigidParticleHandle& Particle = *PBDParticle;
 
@@ -968,45 +980,19 @@ void FPBDRigidsEvolutionGBF::Integrate(FReal Dt)
 						}
 					}
 				}
-
-				if (!Base::IsADisableCluster(*PBDParticle) && (PBDParticle->HasCollision() || ForceNoCollisionIntoSQ))
-				{
-					const FUniqueIdx UniqueIdx = PBDParticle->UniqueIdx();
-					FPendingSpatialData& PendingData = InternalAccelerationQueue.Add(BatchIndex + 1, UniqueIdx);
-					PendingData.Operation = EPendingSpatialDataOperation::Update;
-					PendingData.AccelerationHandle = FAccelerationStructureHandle(*PBDParticle);
-					PendingData.SpatialIdx = PBDParticle->SpatialIdx();
-				}
 			}
-		}
-	};
+		}, LowLevelTasks::ETaskPriority::High);
+		IntegrationPendingTasks.Add(PendingTask);
+	}
 
-	auto DirtyAsync = [this, &ParticlesArray]()
+	DirtyParticleTask = UE::Tasks::Launch(UE_SOURCE_LOCATION, [this, &ParticlesArray]()
 	{
+		QUICK_SCOPE_CYCLE_COUNTER(STAT_DirtyDynamicParticleTask);
 		for (FPBDRigidParticleHandle* Particle : ParticlesArray)
 		{
-			Base::DirtyParticleAsync(*Particle);
+			Base::DirtyParticle(*Particle);
 		}
-	};
-	
-
-	if (GSingleThreadedPhysics)
-	{
-		// In single threaded mode keep same logic with three batches inside PendingDataArrays: 
-		// 0 Dirty particles from the PushData
-		// 1 Dynamic particles added in Integrate
-		// 2 Kinematic particles added in Apply Kinematic Target
-		InternalAccelerationQueue.PendingDataArrays.SetNum(3);
-		InternalAccelerationQueue.KinematicBatchStartIndex = 2;
-		const int32 NumDynParticles = ParticlesArray.Num();
-		IntegrateWork(0, 0, NumDynParticles);
-		DirtyAsync();
-	}
-	else
-	{
-		TaskDispatcher.DispatchIntegrate(IntegrateWork);
-		TaskDispatcher.DispatchDynAsyncDirty(DirtyAsync);
-	}
+	}, LowLevelTasks::ETaskPriority::High);
 }
 
 void FPBDRigidsEvolutionGBF::ApplyKinematicTargets(const FReal Dt, const FReal StepFraction)
@@ -1014,14 +1000,14 @@ void FPBDRigidsEvolutionGBF::ApplyKinematicTargets(const FReal Dt, const FReal S
 	check(StepFraction > (FReal)0);
 	check(StepFraction <= (FReal)1);
 
-	const bool bIsLastStep = (FMath::IsNearlyEqual(StepFraction, (FReal)1, (FReal)UE_KINDA_SMALL_NUMBER));
+	const bool IsLastStep = (FMath::IsNearlyEqual(StepFraction, (FReal)1, (FReal)UE_KINDA_SMALL_NUMBER));
 
 	TParticleView<FPBDRigidParticles>& KienamticParticles = Particles.GetActiveMovingKinematicParticlesView();
 	using TSOA = typename TParticleView<FPBDRigidParticles>::TSOA;
 	using THandle = typename TSOA::THandleType;
 	using THandleBase = typename THandle::THandleBase;
 	using TTransientHandle = typename THandle::TTransientHandle;
-	int32 DispatchBatchIndex = 0;
+
 	for (int32 ViewIndex = 0; ViewIndex < KienamticParticles.SOAViews.Num(); ++ViewIndex)
 	{
 		const TSOAView<TSOA>& SOAView = KienamticParticles.SOAViews[ViewIndex];
@@ -1040,65 +1026,51 @@ void FPBDRigidsEvolutionGBF::ApplyKinematicTargets(const FReal Dt, const FReal S
 				};
 			using LambdaType = decltype(ReturnHandle);
 			const int32 NumHandles = CurHandlesArray->Num();
-			DispatchKinematicTargetsTasks<LambdaType>(DispatchBatchIndex++, NumHandles, Dt, StepFraction, bIsLastStep, ReturnHandle);
+			DispatchKinematicTargetsTasks<LambdaType>(NumHandles, Dt, StepFraction, IsLastStep, ReturnHandle);
 		}
 		else if (TSOA* Soa = SOAView.SOA)
 		{
 			auto ReturnHandle = [Soa](int32 Index) { return THandleBase(Soa, Index); };
 			using LambdaType = decltype(ReturnHandle);
-			DispatchKinematicTargetsTasks<LambdaType>(DispatchBatchIndex++, NumParticles, Dt, StepFraction, bIsLastStep, ReturnHandle);
+			DispatchKinematicTargetsTasks<LambdaType>(NumParticles, Dt, StepFraction, IsLastStep, ReturnHandle);
 		}
 	}
 
-	if (GSingleThreadedPhysics)
+	// To start dirtying Kinematics particles, the Dynamics particles must be finished
+	UE::Tasks::Wait(DirtyParticleTask);
+	DirtyParticleTask = UE::Tasks::FTask();
+	for(FTransientPBDRigidParticleHandle& Particle : Particles.GetActiveMovingKinematicParticlesView())
 	{
-		for (FTransientPBDRigidParticleHandle& Particle : Particles.GetActiveMovingKinematicParticlesView())
-		{
-			// Moving kinematics need to have SQ updates queued to apply substepped movement to the
-			// physics thread acceleration structure in order to get correct kinematic collisions
-			// and pass a correct structure back to the external thread
-			DirtyParticleAsync(Particle);
-		}
-		// done with update, let's clear the tracking structures
-		if (bIsLastStep)
-		{
-			constexpr bool bUpdateView = false;
-			Particles.UpdateAllMovingKinematic(bUpdateView);
-		}
-		Particles.UpdateDirtyViews();
+		// Moving kinematics need to have SQ updates queued to apply substepped movement to the
+		// physics thread acceleration structure in order to get correct kinematic collisions
+		// and pass a correct structure back to the external thread
+		DirtyParticle(Particle);
 	}
-	else
-	{
-		auto DirtyKinAsync = [this]()
-		{
-			for (FTransientPBDRigidParticleHandle& Particle : Particles.GetActiveMovingKinematicParticlesView())
-			{
-				// Moving kinematics need to have SQ updates queued to apply substepped movement to the
-				// physics thread acceleration structure in order to get correct kinematic collisions
-				// and pass a correct structure back to the external thread
-				DirtyParticleAsync(Particle);
-			}
-		};
 
-		TaskDispatcher.DispatchKinAsyncDirtyAndUpdateKinematic(DirtyKinAsync, bIsLastStep);
+	UE::Tasks::Wait(KinematicPendingTasks);
+	KinematicPendingTasks.Reset();
+	// done with update, let's clear the tracking structures
+	if (IsLastStep)
+	{
+		Particles.UpdateAllMovingKinematic();
 	}
-}
 
-void FPBDRigidsEvolutionGBF::WaitIntegrationComplete()
-{
-	TaskDispatcher.WaitIntegrationComplete();
+	// Updating the views can be done only when all kinematics and dynamics update are finished.
+	UE::Tasks::Wait(IntegrationPendingTasks);
+	IntegrationPendingTasks.Reset();
+	// If we changed any particle state, the views need to be refreshed
+	Particles.UpdateDirtyViews();
 }
 
 template <typename LambdaType>
-void FPBDRigidsEvolutionGBF::DispatchKinematicTargetsTasks(int32 DispatchBatchIndex, int32 NumParticles, FReal Dt, FReal StepFraction, bool bIsLastStep, LambdaType ReturnHandle)
+void FPBDRigidsEvolutionGBF::DispatchKinematicTargetsTasks(int32 NumParticles, FReal Dt, FReal StepFraction, bool IsLastStep, LambdaType ReturnHandle)
 {
 	using TSOA = typename TParticleView<FPBDRigidParticles>::TSOA;
 	using THandle = typename TSOA::THandleType;
 	using THandleBase = typename THandle::THandleBase;
 
-	auto KinematicTargetsLambda = [this, ReturnHandle, Dt, StepFraction, bIsLastStep](int32 StartIndex, int32 EndIndex, int32 BatchIndex)
+	auto KinematicTargetsLambda = [this, ReturnHandle, Dt, StepFraction, IsLastStep](int32 StartIndex, int32 EndIndex)
 	{
-		InternalAccelerationQueue.PendingDataArrays[InternalAccelerationQueue.KinematicBatchStartIndex + BatchIndex].Reserve(EndIndex- StartIndex);
 		for (int32 Index = StartIndex; Index < EndIndex; Index++)
 		{
 			THandleBase Handle = ReturnHandle(Index);
@@ -1134,7 +1106,7 @@ void FPBDRigidsEvolutionGBF::DispatchKinematicTargetsTasks(int32 DispatchBatchIn
 					// Target positions only need to be processed once, and we reset the velocity next frame (if no new target is set)
 					FVec3 NewX;
 					FRotation3f NewR;
-					if (bIsLastStep)
+					if (IsLastStep)
 					{
 						NewX = KinematicTarget.GetPosition();
 						NewR = KinematicTarget.GetRotation();
@@ -1204,24 +1176,35 @@ void FPBDRigidsEvolutionGBF::DispatchKinematicTargetsTasks(int32 DispatchBatchIn
 					}
 				}
 			}
-			if (!Base::IsADisableCluster(Particle) && (Particle.HasCollision() || ForceNoCollisionIntoSQ))
-			{
-				const FUniqueIdx UniqueIdx = Particle.UniqueIdx();
-				FPendingSpatialData& PendingData = InternalAccelerationQueue.Add(InternalAccelerationQueue.KinematicBatchStartIndex + BatchIndex, UniqueIdx);
-				PendingData.Operation = EPendingSpatialDataOperation::Update;
-				PendingData.AccelerationHandle = FAccelerationStructureHandle(Particle);
-				PendingData.SpatialIdx = Particle.SpatialIdx();
-			}
 		}
 	};
 
-	if (GSingleThreadedPhysics)
+	int32 NumTasks = !!GSingleThreadedPhysics ? 1 : FMath::Min(FTaskGraphInterface::Get().GetNumWorkerThreads(), Chaos::MaxNumWorkers);
+	NumTasks = FMath::Max(FMath::Min(NumTasks, NumParticles), 1);
+	constexpr int32 MinParticlePerTask = 40;
+	const int32 ParticleByTask = FMath::Max(FMath::DivideAndRoundUp(NumParticles, NumTasks), MinParticlePerTask);
+	int32 NumBatches = FMath::DivideAndRoundUp(NumParticles, ParticleByTask);
+
+	if (NumBatches > 1)
 	{
-		KinematicTargetsLambda(0, NumParticles, 0);
+		// We only write to particle state
+		for (int32 BatchIndex = 0; BatchIndex < NumBatches; BatchIndex++)
+		{
+			const int32 StartIndex = BatchIndex * ParticleByTask;
+			int32 EndIndex = (BatchIndex + 1) * ParticleByTask;
+			EndIndex = FMath::Min(NumParticles, EndIndex);
+
+			UE::Tasks::FTask PendingTask = UE::Tasks::Launch(UE_SOURCE_LOCATION, [KinematicTargetsLambda, StartIndex, EndIndex]()
+			{
+				QUICK_SCOPE_CYCLE_COUNTER(STAT_KinematicTargetsTasks);
+				KinematicTargetsLambda(StartIndex, EndIndex);
+			}, LowLevelTasks::ETaskPriority::High);
+			KinematicPendingTasks.Add(PendingTask);
+		}
 	}
 	else
 	{
-		TaskDispatcher.DispatchKinematicsTarget(KinematicTargetsLambda, NumParticles, DispatchBatchIndex); 
+		KinematicTargetsLambda(0, NumParticles);
 	}
 }
 
diff --git a/Engine/Source/Runtime/Experimental/Chaos/Private/PBDRigidsSolver.cpp b/Engine/Source/Runtime/Experimental/Chaos/Private/PBDRigidsSolver.cpp
--- a/Engine/Source/Runtime/Experimental/Chaos/Private/PBDRigidsSolver.cpp
+++ b/Engine/Source/Runtime/Experimental/Chaos/Private/PBDRigidsSolver.cpp
@@ -743,19 +743,10 @@ namespace Chaos
 
 	void FPBDRigidsSolver::OnEvolutionPostIntegrate(FReal Dt)
 	{
-		bool bHasIntegrationCompleted = false;
 		for (ISimCallbackObject* Callback : SimCallbackObjects)
 		{
 			if (Callback->HasOption(ESimCallbackOptions::PostIntegrate))
 			{
-				if (!bHasIntegrationCompleted)
-				{
-					// If a post integrate callback is called, 
-					// we need to wait for all workers threads integrating to finish.
-					MEvolution->WaitIntegrationComplete();
-					bHasIntegrationCompleted = true;
-				}
-				
 				FScopedTraceSolverCallback TraceCallback(Callback);
 				Callback->PostIntegrate_Internal();
 			}
diff --git a/Engine/Source/Runtime/Experimental/Chaos/Public/Chaos/PBDRigidsEvolution.h b/Engine/Source/Runtime/Experimental/Chaos/Public/Chaos/PBDRigidsEvolution.h
--- a/Engine/Source/Runtime/Experimental/Chaos/Public/Chaos/PBDRigidsEvolution.h
+++ b/Engine/Source/Runtime/Experimental/Chaos/Public/Chaos/PBDRigidsEvolution.h
@@ -10,10 +10,9 @@
 #include "Chaos/Framework/DebugSubstep.h"
 #include "HAL/Event.h"
 #include "Chaos/PBDJointConstraints.h"
-#include "Chaos/PBDRigidsEvolutionFwd.h"
 #include "Chaos/PBDRigidsSOAs.h"
 #include "Chaos/SpatialAccelerationCollection.h"
-#include "Chaos/TaskDispatcherEvolution.h"
+#include "Chaos/PBDRigidsEvolutionFwd.h"
 #include "Chaos/Island/IslandManager.h"
 #include "Chaos/Island/IslandGroupManager.h"
 #include "Chaos/Defines.h"
@@ -495,27 +494,28 @@ public:
 	}
 
 	template <bool bPersistent>
-	FORCEINLINE_DEBUGGABLE bool IsADisableCluster(TGeometryParticleHandleImp<FReal, 3, bPersistent>& Particle)
+	FORCEINLINE_DEBUGGABLE void DirtyParticle(TGeometryParticleHandleImp<FReal, 3, bPersistent>& Particle, const EPendingSpatialDataOperation Op = EPendingSpatialDataOperation::Update)
 	{
+		ensure(Op != EPendingSpatialDataOperation::Delete); // Don't use the function to delete particles
 		const TPBDRigidParticleHandleImp<FReal, 3, bPersistent>* AsRigid = Particle.CastToRigidParticle();
-		if (AsRigid && AsRigid->Disabled())
+		if(AsRigid && AsRigid->Disabled())
 		{
 			TPBDRigidClusteredParticleHandleImp<FReal, 3, bPersistent>* AsClustered = Particle.CastToClustered();
 
-			if (AsClustered)
+			if(AsClustered)
 			{
 				// For clustered particles, they may appear disabled but they're being driven by an internal (solver-owned) cluster parent.
 				// If this is the case we let the spatial data update with those particles, otherwise skip.
 				// Alternatively, the particle may be being driven by a cluster union. Disabled children of a cluster union should not be added
 				// to the SQ.
 				// #BGTODO consider converting MDisabled into a bitfield for multiple disable types (Disabled, DisabledDriven, etc.)
-				if (FPBDRigidParticleHandle* ClusterParentBase = AsClustered->ClusterIds().Id)
+				if(FPBDRigidParticleHandle* ClusterParentBase = AsClustered->ClusterIds().Id)
 				{
-					if (Chaos::TPBDRigidClusteredParticleHandle<FReal, 3>*ClusterParent = ClusterParentBase->CastToClustered())
+					if(Chaos::TPBDRigidClusteredParticleHandle<FReal, 3>* ClusterParent = ClusterParentBase->CastToClustered())
 					{
-						if (!ClusterParent->InternalCluster())
+						if(!ClusterParent->InternalCluster())
 						{
-							return true;
+							return;
 						}
 
 						// We know we're an internal cluster now. If this was a GC internal proxy, we'd expect the 
@@ -523,81 +523,44 @@ public:
 						// we're in a cluster union.
 						if (ClusterParent->PhysicsProxy() != Particle.PhysicsProxy())
 						{
-							return true;
+							return;
 						}
 					}
 					else
 					{
 						// There's probably no way we get here since if we're clustering the parent should always be a clustered.
 						ensure(false);
-						return true;
+						return;
 					}
 				}
 				else
 				{
 					// Disabled cluster particle that doesn't have a parent cluster. MUST BE IGNORED.
-					return true;
+					return;
 				}
 			}
 			else
 			{
 				// Disabled particles take no immediate part in sim or query so shouldn't be added to the acceleration
-				return true;
+				return;
 			}
 		}
-		return false;
-	}
-
-	template <bool bPersistent>
-	FORCEINLINE_DEBUGGABLE void DirtyParticle(TGeometryParticleHandleImp<FReal, 3, bPersistent>& Particle, const EPendingSpatialDataOperation Op = EPendingSpatialDataOperation::Update)
-	{
-		ensure(Op != EPendingSpatialDataOperation::Delete); // Don't use the function to delete particles
-		
-		if (IsADisableCluster(Particle))
-		{
-			return;
-		}
-
-		//only add to acceleration structure if it has collision
-		if (Particle.HasCollision() || ForceNoCollisionIntoSQ)
-		{
-			//TODO: distinguish between new particles and dirty particles - Adds and updates are treated the same right now
-			const FUniqueIdx UniqueIdx = Particle.UniqueIdx();
-			FPendingSpatialData& PendingData = InternalAccelerationQueue.Add(0, UniqueIdx);
-			PendingData.Operation = Op;
-			PendingData.AccelerationHandle = FAccelerationStructureHandle(Particle);
-			PendingData.SpatialIdx = Particle.SpatialIdx();
-
-
-			FPendingSpatialData& AsyncSpatialData = AsyncAccelerationQueue.FindOrAdd(UniqueIdx);
-			AsyncSpatialData = PendingData;
-			// ensure(AsyncSpatialData.Operation != EPendingSpatialDataOperation::Delete); // TODO: This may be hit: Potentially due to UniqueIdx reuse?
-		}
-	}
-
-	template <bool bPersistent>
-	FORCEINLINE_DEBUGGABLE void DirtyParticleAsync(TGeometryParticleHandleImp<FReal, 3, bPersistent>& Particle, const EPendingSpatialDataOperation Op = EPendingSpatialDataOperation::Update)
-	{
-		ensure(Op != EPendingSpatialDataOperation::Delete); // Don't use the function to delete particles
-
-		if (IsADisableCluster(Particle))
-		{
-			return;
-		}
 
 		//only add to acceleration structure if it has collision
 		if (Particle.HasCollision() || ForceNoCollisionIntoSQ)
 		{
 			//TODO: distinguish between new particles and dirty particles - Adds and updates are treated the same right now
 			const FUniqueIdx UniqueIdx = Particle.UniqueIdx();
+			FPendingSpatialData& SpatialData = InternalAccelerationQueue.FindOrAdd(UniqueIdx);
+			ensure(SpatialData.Operation != EPendingSpatialDataOperation::Delete);
 
+			SpatialData.Operation = Op;
+			SpatialData.AccelerationHandle = FAccelerationStructureHandle(Particle);
+			SpatialData.SpatialIdx = Particle.SpatialIdx();
 
-			FPendingSpatialData& AsyncSpatialData = AsyncAccelerationQueue.FindOrAdd(UniqueIdx);
-			AsyncSpatialData.Operation = Op;
-			AsyncSpatialData.AccelerationHandle = FAccelerationStructureHandle(Particle);
-			AsyncSpatialData.SpatialIdx = Particle.SpatialIdx();
-
+			auto& AsyncSpatialData = AsyncAccelerationQueue.FindOrAdd(UniqueIdx);
 			// ensure(AsyncSpatialData.Operation != EPendingSpatialDataOperation::Delete); // TODO: This may be hit: Potentially due to UniqueIdx reuse?
+			AsyncSpatialData = SpatialData;
 		}
 	}
 
@@ -626,6 +589,7 @@ public:
 	{
 		if (const int32 NumNew = IslandManager.ReserveParticles(Num))
 		{
+			InternalAccelerationQueue.PendingData.Reserve(InternalAccelerationQueue.Num() + NumNew);
 			AsyncAccelerationQueue.PendingData.Reserve(AsyncAccelerationQueue.Num() + NumNew);
 		}
 	}
@@ -870,9 +834,6 @@ public:
 	void SetName(const FString& InName) { EvolutionName = InName; }
 	const FString& GetName() const { return EvolutionName; }
 
-	// Internal used only 
-	UE_INTERNAL virtual void WaitIntegrationComplete() { check(false); };
-
 protected:
 	int32 NumConstraints() const
 	{
@@ -994,13 +955,11 @@ public:
 protected:
 
 	/** Pending operations for the internal acceleration structure */
-	FPendingSpatialInternalDataQueue InternalAccelerationQueue;
-	
+	FPendingSpatialDataQueue InternalAccelerationQueue;
+
 	/** Pending operations for the acceleration structures being rebuilt asynchronously */
 	FPendingSpatialDataQueue AsyncAccelerationQueue;
 
-	Private::FTaskDispatcherEvolution TaskDispatcher;
-
 	/*void SerializePendingMap(FChaosArchive& Ar, TMap<FGeometryParticleHandle*, FPendingSpatialData>& Map)
 	{
 		TArray<FGeometryParticleHandle*> Keys;
diff --git a/Engine/Source/Runtime/Experimental/Chaos/Public/Chaos/PBDRigidsEvolutionGBF.h b/Engine/Source/Runtime/Experimental/Chaos/Public/Chaos/PBDRigidsEvolutionGBF.h
--- a/Engine/Source/Runtime/Experimental/Chaos/Public/Chaos/PBDRigidsEvolutionGBF.h
+++ b/Engine/Source/Runtime/Experimental/Chaos/Public/Chaos/PBDRigidsEvolutionGBF.h
@@ -273,9 +273,6 @@ namespace Chaos
 		UE_DEPRECATED(5.4, "Use Integrate(Dt)")
 		void Integrate(const TParticleView& InParticles, FReal Dt) { Integrate(Dt); }
 
-		// Internal used only
-		UE_INTERNAL virtual void WaitIntegrationComplete() override;
-
 	protected:
 
 		CHAOS_API void AdvanceOneTimeStepImpl(const FReal Dt, const FSubStepInfo& SubStepInfo);
@@ -299,7 +296,7 @@ namespace Chaos
 		void UpdateCollisionSolverType();
 
 		template <typename LambdaType>
-		void DispatchKinematicTargetsTasks(int32 DispatchBatchIndex, int32 NumParticles, FReal Dt, FReal StepFraction, bool IsLastStep, LambdaType ReturnHandle);
+		void DispatchKinematicTargetsTasks(int32 NumParticles, FReal Dt, FReal StepFraction, bool IsLastStep, LambdaType ReturnHandle);
 
 		FRigidClustering Clustering;
 
@@ -327,6 +324,10 @@ namespace Chaos
 
 		FCCDManager CCDManager;
 
+		TArray<UE::Tasks::FTask> IntegrationPendingTasks;
+		TArray<UE::Tasks::FTask> KinematicPendingTasks;
+		UE::Tasks::FTask DirtyParticleTask;
+
 		bool bIsDeterministic;
 
 #if CHAOS_EVOLUTION_COLLISION_TESTMODE
diff --git a/Engine/Source/Runtime/Experimental/Chaos/Public/Chaos/PBDRigidsSOAs.h b/Engine/Source/Runtime/Experimental/Chaos/Public/Chaos/PBDRigidsSOAs.h
--- a/Engine/Source/Runtime/Experimental/Chaos/Public/Chaos/PBDRigidsSOAs.h
+++ b/Engine/Source/Runtime/Experimental/Chaos/Public/Chaos/PBDRigidsSOAs.h
@@ -826,7 +826,7 @@ public:
 		}
 	}
 
-	void UpdateAllMovingKinematic(const bool bUpdateViews = true)
+	void UpdateAllMovingKinematic()
 	{
 		// moving kinematics are going to a 'reset' mode then a 'none' mode
 		// 'reset' ones need to stay for another frame 
@@ -843,7 +843,7 @@ public:
 				// Particle may not be moving, but its velocity was just changed (to zero), so it
 				// needs to be in a dirty list. Since it are no longer in MovingKinematicsMapArray
 				// we must put it the transient dirty list
-				MarkTransientDirtyParticle(Rigid, bUpdateViews);
+				MarkTransientDirtyParticle(Rigid);
 
 				CheckParticleViewMask(Rigid);
 			}
diff --git a/Engine/Source/Runtime/Experimental/Chaos/Public/Chaos/PendingSpatialData.h b/Engine/Source/Runtime/Experimental/Chaos/Public/Chaos/PendingSpatialData.h
--- a/Engine/Source/Runtime/Experimental/Chaos/Public/Chaos/PendingSpatialData.h
+++ b/Engine/Source/Runtime/Experimental/Chaos/Public/Chaos/PendingSpatialData.h
@@ -5,9 +5,6 @@
 #include "Chaos/Transform.h"
 #include "Chaos/ISpatialAcceleration.h"
 #include "Chaos/Defines.h"
-#include "Tasks/Task.h"
-#include "Containers/BitArray.h"
-
 
 namespace Chaos
 { 
@@ -61,125 +58,6 @@ struct FPendingSpatialData
 	}
 };
 
-struct FPendingSpatialInternalDataQueue
-{
-	// PendingDataArrays store in 
-	// - [0] all elements added in push data
-	// - [1 -> KinematicBatchStartIndex] all dynamic particle contained in GetActiveParticlesArray
-	// - [KinematicBatchStartIndex -> PendingDataArrays.Num()] all kinematic particle in GetActiveMovingKinematicParticlesView
-	TArray<TArray<FPendingSpatialData>> PendingDataArrays;
-	// ParticleToPendingData ensure all particle in index [0] are unique
-	TArrayAsMap<FUniqueIdx, int32> ParticleToPendingData;
-	// This array is access to mark duplicated data from array [0], not we cannot use bit array because it could be used from different thread.
-	TArray<bool> Duplicated;
-	int32 KinematicBatchStartIndex = 0;
-
-	void Reset()
-	{
-		KinematicBatchStartIndex = 0;
-		for (TArray<FPendingSpatialData>& PendingData : PendingDataArrays)
-		{
-			PendingData.Reset();
-		}
-		PendingDataArrays.Reset();
-		ParticleToPendingData.Reset();
-		Duplicated.Reset();
-	}
-
-	int32 Num() const
-	{
-		int32 NumPendingData = 0;
-		for (const TArray<FPendingSpatialData>& PendingData : PendingDataArrays)
-		{
-			NumPendingData += PendingData.Num();
-		}
-		return NumPendingData;
-	}
-
-	FPendingSpatialData& Add(int32 BatchIndex, const FUniqueIdx UniqueIdx, EPendingSpatialDataOperation Operation = EPendingSpatialDataOperation::Add)
-	{
-		if (PendingDataArrays.IsEmpty())
-		{
-			check(BatchIndex == 0);
-			PendingDataArrays.SetNum(1);
-		}
-
-		check(PendingDataArrays[0].Num() == Duplicated.Num());
-		// Check if duplicate
-		if (BatchIndex != 0)
-		{
-			if (int32* Existing = ParticleToPendingData.Find(UniqueIdx))
-			{
-				Duplicated[*Existing] = true;
-			}
-			const int32 NewIdx = PendingDataArrays[BatchIndex].AddDefaulted(1);
-			PendingDataArrays[BatchIndex][NewIdx].Operation = Operation;
-			check(PendingDataArrays[0].Num() == Duplicated.Num());
-			return PendingDataArrays[BatchIndex][NewIdx];
-		}
-		else // Adding at Index 0 would be always in Single thread
-		{
-			if (int32* Existing = ParticleToPendingData.Find(UniqueIdx))
-			{
-				check(PendingDataArrays[0].Num() == Duplicated.Num());
-				return PendingDataArrays[0][*Existing];
-			}
-			const int32 NewIdx = PendingDataArrays[0].AddDefaulted(1);
-			ParticleToPendingData.Add(UniqueIdx, NewIdx);
-			PendingDataArrays[0][NewIdx].Operation = Operation;
-			Duplicated.Add(false);
-			check(PendingDataArrays[0].Num() == Duplicated.Num());
-			return PendingDataArrays[0][NewIdx];
-		}
-	}
-
-	void Remove(const FUniqueIdx UniqueIdx)
-	{
-		if (!PendingDataArrays.IsEmpty())
-		{
-			check(PendingDataArrays[0].Num() == Duplicated.Num());
-			TArray<FPendingSpatialData>& PendingData = PendingDataArrays[0];
-			if (int32* Existing = ParticleToPendingData.Find(UniqueIdx))
-			{
-				const int32 SlotIdx = *Existing;
-				if (SlotIdx + 1 < PendingData.Num())
-				{
-					const FUniqueIdx LastElemUniqueIdx = PendingData.Last().UniqueIdx();
-					ParticleToPendingData.FindChecked(LastElemUniqueIdx) = SlotIdx;	//We're going to swap elements so the last element is now in the position of the element we removed
-				}
-
-				Duplicated.RemoveAtSwap(SlotIdx);
-				PendingData.RemoveAtSwap(SlotIdx);
-				ParticleToPendingData.RemoveChecked(UniqueIdx);
-			}
-			check(PendingDataArrays[0].Num() == Duplicated.Num());
-		}
-	}
-
-	void CleanUpDuplicated()
-	{
-		if (!PendingDataArrays.IsEmpty())
-		{
-			check(PendingDataArrays[0].Num() == Duplicated.Num());
-			ParticleToPendingData.Reset();
-			TArray<FPendingSpatialData> NoDuplicatedData;
-			for (int32 Index = 0; Index < Duplicated.Num(); Index++)
-			{
-				const bool bIsDuplicated = Duplicated[Index];
-				if (!bIsDuplicated)
-				{
-					NoDuplicatedData.Add(PendingDataArrays[0][Index]);
-					ParticleToPendingData.Add(PendingDataArrays[0][Index].UniqueIdx(), Index);
-				}
-			}
-			Duplicated.Init(false, NoDuplicatedData.Num());
-			PendingDataArrays[0] = MoveTemp(NoDuplicatedData);
-			check(PendingDataArrays[0].Num() == Duplicated.Num());
-		}
-	}
-
-};
-
 struct FPendingSpatialDataQueue
 {
 	TArray<FPendingSpatialData> PendingData;
@@ -227,4 +105,4 @@ struct FPendingSpatialDataQueue
 	}
 };
 
-}
+}
\ No newline at end of file
diff --git a/Engine/Source/Runtime/Experimental/Chaos/Public/Chaos/TaskDispatcherEvolution.h b/Engine/Source/Runtime/Experimental/Chaos/Public/Chaos/TaskDispatcherEvolution.h
deleted file mode 100644
