Chaos
- Evolution now uses shared string for debug name
- Add shared debug name to collision detector for better logging and ease of debugging

#rnx

diff --git a/Engine/Source/Runtime/Experimental/Chaos/Public/Chaos/ChaosDebugName.h b/Engine/Source/Runtime/Experimental/Chaos/Public/Chaos/ChaosDebugName.h
--- a/Engine/Source/Runtime/Experimental/Chaos/Public/Chaos/ChaosDebugName.h
+++ b/Engine/Source/Runtime/Experimental/Chaos/Public/Chaos/ChaosDebugName.h
@@ -19,6 +19,8 @@ namespace Chaos
 		FSharedDebugName() = default;
 		FSharedDebugName(const FSharedDebugName& Other) = default;
 		FSharedDebugName(FSharedDebugName&& Other) = default;
+		FSharedDebugName& operator=(const FSharedDebugName& Other) = default;
+		FSharedDebugName& operator=(FSharedDebugName&& Other) = default;
 
 		FSharedDebugName(const FString& S);
 		FSharedDebugName(FString&& S);
diff --git a/Engine/Source/Runtime/Experimental/Chaos/Public/Chaos/Collision/CollisionContext.h b/Engine/Source/Runtime/Experimental/Chaos/Public/Chaos/Collision/CollisionContext.h
--- a/Engine/Source/Runtime/Experimental/Chaos/Public/Chaos/Collision/CollisionContext.h
+++ b/Engine/Source/Runtime/Experimental/Chaos/Public/Chaos/Collision/CollisionContext.h
@@ -1,6 +1,7 @@
 // Copyright Epic Games, Inc. All Rights Reserved.
 #pragma once
 
+#include "Chaos/ChaosDebugName.h"
 #include "Chaos/Core.h"
 
 namespace Chaos
@@ -27,6 +28,7 @@ namespace Chaos
 			, bAllowManifoldReuse(true)
 			, bAllowCCD(true)
 			, bAllowMACD(false)
+			, EvolutionName()
 		{
 		}
 
@@ -65,6 +67,13 @@ namespace Chaos
 		// Whether MACD is allowed (disabled for RBAN)
 		bool bAllowMACD;
 
+		// Useful for debugging (esp. PIE client/server)
+		FSharedDebugName EvolutionName;
+
+		const FString& GetName() const
+		{
+			return EvolutionName.Value();
+		}
 	};
 
 	/**
diff --git a/Engine/Source/Runtime/Experimental/Chaos/Public/Chaos/PBDCollisionConstraints.h b/Engine/Source/Runtime/Experimental/Chaos/Public/Chaos/PBDCollisionConstraints.h
--- a/Engine/Source/Runtime/Experimental/Chaos/Public/Chaos/PBDCollisionConstraints.h
+++ b/Engine/Source/Runtime/Experimental/Chaos/Public/Chaos/PBDCollisionConstraints.h
@@ -69,6 +69,9 @@ public:
 
 	CHAOS_API virtual ~FPBDCollisionConstraints();
 
+	const FString& GetName() const { return DetectorSettings.GetName(); }
+	void SetName(const FSharedDebugName& InName) { DetectorSettings.EvolutionName = InName; }
+
 	/**
 	 * Whether this container provides constraint handles (simple solvers do not need them)
 	 */
diff --git a/Engine/Source/Runtime/Experimental/Chaos/Public/Chaos/PBDRigidsEvolution.h b/Engine/Source/Runtime/Experimental/Chaos/Public/Chaos/PBDRigidsEvolution.h
--- a/Engine/Source/Runtime/Experimental/Chaos/Public/Chaos/PBDRigidsEvolution.h
+++ b/Engine/Source/Runtime/Experimental/Chaos/Public/Chaos/PBDRigidsEvolution.h
@@ -2,6 +2,7 @@
 #pragma once
 
 #include "Chaos/Character/CharacterGroundConstraintContainer.h"
+#include "Chaos/ChaosDebugName.h"
 #include "Chaos/PBDCollisionConstraints.h"
 #include "Chaos/PBDRigidClustering.h"
 #include "Chaos/PBDRigidParticles.h"
@@ -831,8 +832,8 @@ public:
 	const TArray<FRemovalData>& GetAllRemovals() { return MAllRemovals; }
 	void ResetAllRemovals() { MAllRemovals.Reset(); }
 
-	void SetName(const FString& InName) { EvolutionName = InName; }
-	const FString& GetName() const { return EvolutionName; }
+	virtual void SetName(const FString& InName) { EvolutionName = FSharedDebugName(InName); }
+	const FString& GetName() const { return EvolutionName.Value(); }
 
 protected:
 	int32 NumConstraints() const
@@ -1030,7 +1031,7 @@ protected:
 	bool bIsReset = false;
 
 	// Useful name for debugging. E.g., Indicates whether we are on client or server
-	FString EvolutionName;
+	FSharedDebugName EvolutionName;
 
 #if CHAOS_EVOLUTION_COLLISION_TESTMODE
 	// Test Mode for Collision issues (resets particle positions every tick for repeatable testing)
diff --git a/Engine/Source/Runtime/Experimental/Chaos/Public/Chaos/PBDRigidsEvolutionGBF.h b/Engine/Source/Runtime/Experimental/Chaos/Public/Chaos/PBDRigidsEvolutionGBF.h
--- a/Engine/Source/Runtime/Experimental/Chaos/Public/Chaos/PBDRigidsEvolutionGBF.h
+++ b/Engine/Source/Runtime/Experimental/Chaos/Public/Chaos/PBDRigidsEvolutionGBF.h
@@ -77,6 +77,12 @@ namespace Chaos
 			bool InIsSingleThreaded = false);
 		CHAOS_API ~FPBDRigidsEvolutionGBF();
 
+		virtual void SetName(const FString& InName)
+		{ 
+			FPBDRigidsEvolutionBase::SetName(InName);
+			CollisionConstraints.SetName(EvolutionName);
+		}
+
 		void SetPreIntegrateCallback(const FPBDRigidsEvolutionCallback& Cb)
 		{
 			PreIntegrateCallback = Cb;
