Gather configuration into a single argument

#rnx

diff --git a/Engine/Source/Runtime/Experimental/IoStore/HttpClient/Private/Loop.inl b/Engine/Source/Runtime/Experimental/IoStore/HttpClient/Private/Loop.inl
--- a/Engine/Source/Runtime/Experimental/IoStore/HttpClient/Private/Loop.inl
+++ b/Engine/Source/Runtime/Experimental/IoStore/HttpClient/Private/Loop.inl
@@ -671,9 +671,12 @@ void FPeerGroup::TickSend(FTickState& State, FHost& Host, FPoller& Poller)
 		bNegotiating = true;
 		bWillBlock = Outcome.IsWaiting();
 
-		FCertRootsRef VerifyCert = Host.GetVerifyCert();
-		const char* HostName = Host.GetHostName().GetData();
-		Peer = FHttpPeer(MoveTemp(Socket), VerifyCert, HostName);
+		FHttpPeer::FParams Params = {
+			.Socket = MoveTemp(Socket),
+			.Certs = Host.GetVerifyCert(),
+			.HostName = Host.GetHostName().GetData(),
+		};
+		Peer = FHttpPeer(MoveTemp(Params));
 	}
 
 	Send.Append(Pending);
diff --git a/Engine/Source/Runtime/Experimental/IoStore/HttpClient/Private/Peer.inl b/Engine/Source/Runtime/Experimental/IoStore/HttpClient/Private/Peer.inl
--- a/Engine/Source/Runtime/Experimental/IoStore/HttpClient/Private/Peer.inl
+++ b/Engine/Source/Runtime/Experimental/IoStore/HttpClient/Private/Peer.inl
@@ -508,8 +508,15 @@ class FHttpPeer
 	: public FTlsPeer
 {
 public:
+	struct FParams
+	{
+		FSocket			Socket;
+		FCertRootsRef 	Certs = ECertRootsRefType::None;
+		const char* 	HostName = nullptr;
+	};
+
 					FHttpPeer() = default;
-					FHttpPeer(FSocket InSocket, FCertRootsRef Certs=ECertRootsRefType::None, const char* HostName=nullptr);
+					FHttpPeer(FParams&& Params);
 	FOutcome		Handshake();
 	FTransaction*	Transact(FBuffer& Buffer);
 	FOutcome		GetPendingTransactId();
@@ -520,8 +527,8 @@ private:
 };
 
 ////////////////////////////////////////////////////////////////////////////////
-FHttpPeer::FHttpPeer(FSocket InSocket, FCertRootsRef Certs, const char* HostName)
-: FTlsPeer(MoveTemp(InSocket), Certs, HostName)
+FHttpPeer::FHttpPeer(FParams&& Params)
+: FTlsPeer(MoveTemp(Params.Socket), Params.Certs, Params.HostName)
 {
 	if (Ssl != nullptr)
 	{
