Rule of five

#rnx

diff --git a/Engine/Source/Runtime/Experimental/IoStore/HttpClient/Private/Peer.inl b/Engine/Source/Runtime/Experimental/IoStore/HttpClient/Private/Peer.inl
--- a/Engine/Source/Runtime/Experimental/IoStore/HttpClient/Private/Peer.inl
+++ b/Engine/Source/Runtime/Experimental/IoStore/HttpClient/Private/Peer.inl
@@ -516,6 +516,10 @@ public:
 	};
 
 					FHttpPeer() = default;
+					FHttpPeer(FHttpPeer&& Rhs) = delete;
+					FHttpPeer(const FHttpPeer& Rhs) = delete;
+	FHttpPeer&		operator = (FHttpPeer&& Rhs);
+	FHttpPeer&		operator = (const FHttpPeer& Rhs) = delete;
 					FHttpPeer(FParams&& Params);
 	FOutcome		Handshake();
 	FTransaction*	Transact(FBuffer& Buffer);
@@ -536,6 +540,14 @@ FHttpPeer::FHttpPeer(FParams&& Params)
 	}
 }
 
+////////////////////////////////////////////////////////////////////////////////
+FHttpPeer& FHttpPeer::operator = (FHttpPeer&& Rhs)
+{
+	FTlsPeer::operator = (MoveTemp(Rhs));
+	Swap(Proto, Rhs.Proto);
+	return *this;
+}
+
 ////////////////////////////////////////////////////////////////////////////////
 FOutcome FHttpPeer::Handshake()
 {
