[NetPhysics] Fix snapping and inconsistencies in post-resim error rotation during render interpolation.

diff --git a/Engine/Source/Runtime/Experimental/Chaos/Public/Chaos/Framework/PhysicsProxyBase.h b/Engine/Source/Runtime/Experimental/Chaos/Public/Chaos/Framework/PhysicsProxyBase.h
--- a/Engine/Source/Runtime/Experimental/Chaos/Public/Chaos/Framework/PhysicsProxyBase.h
+++ b/Engine/Source/Runtime/Experimental/Chaos/Public/Chaos/Framework/PhysicsProxyBase.h
@@ -266,7 +266,7 @@ public:
 	virtual const FQuat GetErrorR(const Chaos::FRealSingle Alpha) const { return FQuat::Identity; }
 
 	/** Add X and R error onto current error to correct through interpolation */
-	virtual void AccumlateErrorXR(const Chaos::FVec3 X, const FQuat R, const int32 CurrentSimTick, const int32 ErrorSmoothDuration) { };
+	virtual void AccumlateErrorXR(const Chaos::FVec3 X, const FQuat R, const int32 CurrentSimTick, const int32 ErrorSmoothDuration) { }
 
 	/** Tick current error data and decay error */
 	virtual const bool UpdateError(const int32 CurrentSimTick, const Chaos::FReal AsyncFixedTimeStep) { return false; }
@@ -289,7 +289,7 @@ public:
 	virtual const Chaos::FVec3 GetErrorVelocitySmoothingX(const Chaos::FRealSingle Alpha) const { return Chaos::FVec3::ZeroVector; }
 
 	/** EXPERIMENTAL - Register the current velocity and position for use in velocity correction calculations */
-	virtual void SetVelocitySmoothing(const Chaos::FVec3 CurrV, const Chaos::FVec3 CurrX, const int32 ErrorVelocitySmoothDuration) { };
+	virtual void SetVelocitySmoothing(const Chaos::FVec3 CurrV, const Chaos::FVec3 CurrX, const int32 ErrorVelocitySmoothDuration) { }
 
 	/** Get FErrorInterpolationSettings which stores custom settings for render interpolation error corrections, returns nullptr if there are no custom settings */
 	virtual FErrorInterpolationSettings* GetErrorInterpolationSettings() { return nullptr; }
@@ -311,7 +311,7 @@ struct FProxyInterpolationError : FProxyInterpolationBase
 
 	virtual const Chaos::FVec3 GetErrorX(const Chaos::FRealSingle Alpha) const override { return FMath::Lerp(ErrorXPrev, ErrorX, Alpha); }
 
-	virtual const FQuat GetErrorR(const Chaos::FRealSingle Alpha) const override { return FMath::Lerp(ErrorRPrev, ErrorR, Alpha); }
+	virtual const FQuat GetErrorR(const Chaos::FRealSingle Alpha) const override { return FQuat::Slerp(ErrorRPrev, ErrorR, Alpha); }
 
 	virtual void AccumlateErrorXR(const Chaos::FVec3 X, const FQuat R, const int32 CurrentSimTick, const int32 ErrorSmoothDuration) override
 	{
@@ -322,9 +322,7 @@ struct FProxyInterpolationError : FProxyInterpolationBase
 		if (IsErrorSmoothing())
 		{
 			ErrorX += X;
-			ErrorXPrev = ErrorX;
-			ErrorR *= R;
-			ErrorRPrev = ErrorR;
+			ErrorR = R * ErrorR;
 		}
 		else
 		{
@@ -404,7 +402,7 @@ protected:
 			ErrorXPrev = ErrorX;
 			ErrorX *= Alpha;
 			ErrorRPrev = ErrorR;
-			ErrorR = FMath::Lerp(FQuat::Identity, ErrorR, Alpha);
+			ErrorR = FQuat::Slerp(FQuat::Identity, ErrorR, Alpha);
 		}
 
 		return true;
