IAD Cache - Update last access time for chunk when a content handle is released.

diff --git a/Engine/Source/Runtime/Experimental/IoStore/OnDemand/Private/OnDemandInstallCache.cpp b/Engine/Source/Runtime/Experimental/IoStore/OnDemand/Private/OnDemandInstallCache.cpp
--- a/Engine/Source/Runtime/Experimental/IoStore/OnDemand/Private/OnDemandInstallCache.cpp
+++ b/Engine/Source/Runtime/Experimental/IoStore/OnDemand/Private/OnDemandInstallCache.cpp
@@ -90,6 +90,10 @@ using FWeakAsyncFileHandle			= TWeakPtr<IAsyncReadFileHandle>;
 using FSharedFileOpenAsyncResult	= TValueOrError<FSharedAsyncFileHandle, FFileSystemError>;
 
 using FCasAddr						= FHash96;
+static const FCasAddr& AsCasAddr(const FIoHash& IoHash)
+{
+	return *reinterpret_cast<const FCasAddr*>(&IoHash);
+}
 
 ///////////////////////////////////////////////////////////////////////////////
 struct FCasBlockId
@@ -202,6 +206,7 @@ struct FCas
 	{
 		return TrackAccessIf(Type, BlockId, FDateTime::UtcNow().GetTicks(), bDirty);
 	};
+	bool						UnlockedTrackAccessIf(ECasTrackAccessType Type, uint32 BlockIdHash, FCasBlockId BlockId, int64 UtcTicks, bool bDirty);
 	uint64						GetBlockInfo(FCasBlockInfoMap& OutBlockInfo);
 	void						Compact();
 	FIoStatus					Verify(TArray<FCasAddr>& OutAddrs);
@@ -298,11 +303,11 @@ FIoStatus FCas::Initialize(FStringView Directory, bool bDeleteExisting)
 
 FCasLocation FCas::FindChunk(const FIoHash& Hash) const
 {
-	const FCasAddr* Addr	= reinterpret_cast<const FCasAddr*>(&Hash);
-	const uint32 TypeHash	= GetTypeHash(*Addr);
+	const FCasAddr& Addr	= AsCasAddr(Hash);
+	const uint32 TypeHash	= GetTypeHash(Addr);
 	{
 		UE::TUniqueLock Lock(Mutex);
-		if (const FCasLocation* Loc = Lookup.FindByHash(TypeHash, *Addr))
+		if (const FCasLocation* Loc = Lookup.FindByHash(TypeHash, Addr))
 		{
 			return *Loc;
 		}
@@ -493,14 +498,19 @@ FUniqueFileHandle FCas::OpenWrite(FCasBlockId BlockId) const
 	return FUniqueFileHandle(Ipf.OpenWrite(*Filename, bAppend, bAllowRead));
 }
 
-bool FCas::TrackAccessIf(const ECasTrackAccessType Type, const FCasBlockId BlockId, int64 UtcTicks, bool bDirty)
+bool FCas::TrackAccessIf(const ECasTrackAccessType Type, const FCasBlockId BlockId, const int64 UtcTicks, const bool bDirty)
 {
-	check(BlockId.IsValid());
-
 	const uint32 BlockIdHash = GetTypeHash(BlockId);
 
 	UE::TUniqueLock Lock(Mutex);
 
+	return UnlockedTrackAccessIf(Type, BlockIdHash, BlockId, UtcTicks, bDirty);
+}
+
+bool FCas::UnlockedTrackAccessIf(const ECasTrackAccessType Type, const uint32 BlockIdHash, const FCasBlockId BlockId, int64 UtcTicks, bool bDirty)
+{
+	check(BlockId.IsValid());
+
 	int64* MaybeFoundTicks = LastAccess.FindByHash(BlockIdHash, BlockId);
 	if (MaybeFoundTicks == nullptr)
 	{
@@ -1674,6 +1684,7 @@ public:
 	virtual FIoStatus							Verify() override;
 	virtual FIoStatus							Flush() override;
 	virtual FIoStatus							FlushLastAccess() override;
+	virtual void								UpdateLastAccess(TConstArrayView<FIoHash> ChunkHashes) override;
 	virtual FOnDemandInstallCacheUsage			GetCacheUsage() override;
 
 private:
@@ -2450,7 +2461,7 @@ FIoStatus FOnDemandInstallCache::Verify()
 		Lookup.AddrToIndex.Reserve(Container->ChunkEntries.Num());
 		for (int32 EntryIndex = 0; const FOnDemandChunkEntry& Entry : Container->ChunkEntries)
 		{
-			const FCasAddr& Addr = *reinterpret_cast<const FCasAddr*>(&Entry.Hash);
+			const FCasAddr& Addr = AsCasAddr(Entry.Hash);
 			Lookup.AddrToIndex.Add(Addr, EntryIndex++);
 		}
 	}
@@ -3165,6 +3176,35 @@ FIoStatus FOnDemandInstallCache::FlushLastAccess()
 	return InstallCacheStatusBuilder(EIoErrorCode::Ok);
 }
 
+void FOnDemandInstallCache::UpdateLastAccess(TConstArrayView<FIoHash> ChunkHashes)
+{
+	bool bLastAccessDirty = false;
+	{
+		const int64 Now = FDateTime::UtcNow().GetTicks();
+		constexpr const bool bDirty = true;
+
+		TUniqueLock Lock(Cas);
+
+		for (const FIoHash& ChunkHash : ChunkHashes)
+		{
+			const FCasAddr& Addr = AsCasAddr(ChunkHash);
+			if (FCasLocation* CasLoc = Cas.Lookup.Find(Addr))
+			{
+				const FCasBlockId BlockId = CasLoc->BlockId;
+				const uint32 BlockIdHash = GetTypeHash(BlockId);
+				const bool bUpdatedLastAccess = Cas.UnlockedTrackAccessIf(
+					ECasTrackAccessType::Granular, BlockIdHash, BlockId, Now, bDirty);
+				bLastAccessDirty = bLastAccessDirty || bUpdatedLastAccess;
+			}
+		}
+	}
+
+	if (bLastAccessDirty)
+	{
+		IoStore.FlushLastAccess(nullptr);
+	}
+}
+
 FOnDemandInstallCacheUsage FOnDemandInstallCache::GetCacheUsage()
 {
 	// If this is called from a thread other than the OnDemandIoStore tick thread
diff --git a/Engine/Source/Runtime/Experimental/IoStore/OnDemand/Private/OnDemandInstallCache.h b/Engine/Source/Runtime/Experimental/IoStore/OnDemand/Private/OnDemandInstallCache.h
--- a/Engine/Source/Runtime/Experimental/IoStore/OnDemand/Private/OnDemandInstallCache.h
+++ b/Engine/Source/Runtime/Experimental/IoStore/OnDemand/Private/OnDemandInstallCache.h
@@ -41,6 +41,7 @@ public:
 	virtual FIoStatus							Verify() = 0;
 	virtual FIoStatus							Flush() = 0;
 	virtual FIoStatus							FlushLastAccess() = 0;
+	virtual void								UpdateLastAccess(TConstArrayView<FIoHash> ChunkHashes) = 0;
 	virtual FOnDemandInstallCacheUsage			GetCacheUsage() = 0;
 };
 
diff --git a/Engine/Source/Runtime/Experimental/IoStore/OnDemand/Private/OnDemandIoStore.cpp b/Engine/Source/Runtime/Experimental/IoStore/OnDemand/Private/OnDemandIoStore.cpp
--- a/Engine/Source/Runtime/Experimental/IoStore/OnDemand/Private/OnDemandIoStore.cpp
+++ b/Engine/Source/Runtime/Experimental/IoStore/OnDemand/Private/OnDemandIoStore.cpp
@@ -1896,26 +1896,54 @@ void FOnDemandIoStore::ReleaseContent(FOnDemandInternalContentHandle& ContentHan
 	LLM_SCOPE_BYNAME(TEXT("FileSystem/OnDemandIoStore"));
 	UE_LOG(LogIoStoreOnDemand, Verbose, TEXT("Releasing content handle, %s"), *LexToString(ContentHandle));
 
-	UE::TUniqueLock Lock(ContainerMutex);
+	TArray<FIoHash> ReleasedChunks;
 
-	for (FSharedOnDemandContainer& Container : Containers)
 	{
-		if (EnumHasAnyFlags(Container->Flags, EOnDemandContainerFlags::StreamOnDemand))
+		UE::TUniqueLock Lock(ContainerMutex);
+
+		for (FSharedOnDemandContainer& Container : Containers)
 		{
-			continue;
-		}
+			if (EnumHasAnyFlags(Container->Flags, EOnDemandContainerFlags::StreamOnDemand))
+			{
+				continue;
+			}
 
-		Container->ChunkEntryReferences.SetNum(
-			Algo::RemoveIf(
-				Container->ChunkEntryReferences,
+			const int32 RemoveIndex = Container->ChunkEntryReferences.IndexOfByPredicate(
 				[ContentHandleId = ContentHandle.HandleId()](const FOnDemandChunkEntryReferences& Refs)
-				{ 
+				{
 					return Refs.ContentHandleId == ContentHandleId;
-				}));
+				}
+			);
+
+			if (RemoveIndex == INDEX_NONE)
+			{
+				continue;
+			}
+
+			// Note: A referencer is only added to the first container found that has a chunk, and the
+			// referencer is not added the same container more than once. So once the refrencer is removed,
+			// we only need to inspect the current container for unreffed chunks.
+
+			FOnDemandChunkEntryReferences RemovedRefs = MoveTemp(Container->ChunkEntryReferences[RemoveIndex]);
+			Container->ChunkEntryReferences.RemoveAtSwap(RemoveIndex);
+			check(RemovedRefs.Indices.Num() == Container->ChunkEntries.Num());
+
+			TBitArray<> ReffedChunks = Container->GetReferencedChunkEntries();
+
+			for (int32 i = 0; i < RemovedRefs.Indices.Num(); ++i)
+			{
+				if (RemovedRefs.Indices[i] && (ReffedChunks.IsEmpty() || !ReffedChunks[i]))
+				{
+					ReleasedChunks.Add(Container->ChunkEntries[i].Hash);
+				}
+			}
+		}
 	}
 
 	// Package store needs to be aware of install state
 	PackageStoreBackend->NeedsUpdate(EOnDemandPackageStoreUpdateMode::ReferencedPackages);
+
+	InstallCache->UpdateLastAccess(ReleasedChunks);
 }
 
 void FOnDemandIoStore::GetReferencedContent(TArray<FSharedOnDemandContainer>& OutContainers, TArray<TBitArray<>>& OutChunkEntryIndices, bool bPackageStore)
