GeometryCollection: Move storage of PreskinnedBounds in geom collection RenderData out of non-nanite MeshDescription and into common nanite/non-nanite storage.
This allows us to access the precalculated preskinned bounds correctly from both nanite and non-nanite scene proxies.
This replaces an earlier solution  for caching preskinned bounds on the nanite geom collection scene proxy which incorrectly depended on the state of the component when the scene proxy was created.

diff --git a/Engine/Source/Runtime/Experimental/GeometryCollectionEngine/Private/GeometryCollection/DerivedDataGeometryCollectionCooker.cpp b/Engine/Source/Runtime/Experimental/GeometryCollectionEngine/Private/GeometryCollection/DerivedDataGeometryCollectionCooker.cpp
--- a/Engine/Source/Runtime/Experimental/GeometryCollectionEngine/Private/GeometryCollection/DerivedDataGeometryCollectionCooker.cpp
+++ b/Engine/Source/Runtime/Experimental/GeometryCollectionEngine/Private/GeometryCollection/DerivedDataGeometryCollectionCooker.cpp
@@ -124,7 +124,7 @@ bool FDerivedDataGeometryCollectionRenderDataCooker::Build(TArray<uint8>& OutDat
 
 const TCHAR* FDerivedDataGeometryCollectionRenderDataCooker::GetVersionString() const
 {
-	const TCHAR* VersionString = TEXT("EADF0CE05C96495E8621D418077E9C02");
+	const TCHAR* VersionString = TEXT("B1206D187DE34AF580A635426580C021");
 
 	static FString CachedNaniteVersionString;
 	if (CachedNaniteVersionString.IsEmpty())
diff --git a/Engine/Source/Runtime/Experimental/GeometryCollectionEngine/Private/GeometryCollection/GeometryCollectionRenderData.cpp b/Engine/Source/Runtime/Experimental/GeometryCollectionEngine/Private/GeometryCollection/GeometryCollectionRenderData.cpp
--- a/Engine/Source/Runtime/Experimental/GeometryCollectionEngine/Private/GeometryCollection/GeometryCollectionRenderData.cpp
+++ b/Engine/Source/Runtime/Experimental/GeometryCollectionEngine/Private/GeometryCollection/GeometryCollectionRenderData.cpp
@@ -125,7 +125,7 @@ void FGeometryCollectionMeshResources::Serialize(FArchive& Ar)
 
 void FGeometryCollectionMeshDescription::Serialize(FArchive& Ar)
 {
-	Ar << NumVertices << NumTriangles << PreSkinnedBounds << Sections << SectionsNoInternal;
+	Ar << NumVertices << NumTriangles << Sections << SectionsNoInternal;
 	
 	// need subsection data on cooked builds when using ray tracing
 	// to be able to build one BLAS per subsection
@@ -633,9 +633,6 @@ void CreateMeshData(FGeometryCollection& InCollection, FGeometryCollectionBuiltM
 	OutRenderData.MeshDescription.Sections = MoveTemp(SectionData.Sections);
 	OutRenderData.MeshDescription.SectionsNoInternal = MoveTemp(SectionData.SectionsNoInternal);
 	OutRenderData.MeshDescription.SubSections = MoveTemp(SectionData.SubSections);
-
-	// Calculate bounds.
-	OutRenderData.MeshDescription.PreSkinnedBounds = BuildPreSkinnedBounds(InCollection);
 }
 
 void CreateNaniteData(FGeometryCollectionBuiltMeshData&& InMeshData, FGeometryCollectionRenderData& OutRenderData)
@@ -696,6 +693,8 @@ TUniquePtr<FGeometryCollectionRenderData> FGeometryCollectionRenderData::Create(
 		{
 			CreateNaniteData(MoveTemp(MeshBuildData), *RenderData.Get());
 		}
+
+		RenderData->PreSkinnedBounds = BuildPreSkinnedBounds(InCollection);
 	}
 
 	return RenderData;
@@ -772,6 +771,8 @@ void FGeometryCollectionRenderData::Serialize(FArchive& Ar, UGeometryCollection&
 	{
 		ClearNaniteResources(NaniteResourcesPtr);
 	}
+
+	Ar << PreSkinnedBounds;
 }
 
 void FGeometryCollectionRenderData::InitResources(UGeometryCollection const& Owner)
diff --git a/Engine/Source/Runtime/Experimental/GeometryCollectionEngine/Private/GeometryCollection/GeometryCollectionRenderData.h b/Engine/Source/Runtime/Experimental/GeometryCollectionEngine/Private/GeometryCollection/GeometryCollectionRenderData.h
--- a/Engine/Source/Runtime/Experimental/GeometryCollectionEngine/Private/GeometryCollection/GeometryCollectionRenderData.h
+++ b/Engine/Source/Runtime/Experimental/GeometryCollectionEngine/Private/GeometryCollection/GeometryCollectionRenderData.h
@@ -92,9 +92,6 @@ struct FGeometryCollectionMeshDescription
 	uint32 NumVertices = 0;
 	uint32 NumTriangles = 0;
 
-	// Bounds before any movement or skinning ever happens.
-	FBoxSphereBounds PreSkinnedBounds;
-
 	// Info on grouped material sections within vertex buffers. 
 	// A section has a uniform material index and can contain multiple per-bone subsections.
 	TArray<FGeometryCollectionMeshElement> Sections;
@@ -145,6 +142,9 @@ public:
 	/** Nanite resources. */
 	TPimplPtr<Nanite::FResources> NaniteResourcesPtr;
 
+	/** Bounds before any movement or skinning ever happens. */
+	FBoxSphereBounds PreSkinnedBounds = FBoxSphereBounds(ForceInitToZero);
+
 private:
 	bool bIsInitialized = false;
 };
diff --git a/Engine/Source/Runtime/Experimental/GeometryCollectionEngine/Private/GeometryCollection/GeometryCollectionSceneProxy.cpp b/Engine/Source/Runtime/Experimental/GeometryCollectionEngine/Private/GeometryCollection/GeometryCollectionSceneProxy.cpp
--- a/Engine/Source/Runtime/Experimental/GeometryCollectionEngine/Private/GeometryCollection/GeometryCollectionSceneProxy.cpp
+++ b/Engine/Source/Runtime/Experimental/GeometryCollectionEngine/Private/GeometryCollection/GeometryCollectionSceneProxy.cpp
@@ -189,6 +189,8 @@ FGeometryCollectionSceneProxyBase::FGeometryCollectionSceneProxyBase(UGeometryCo
 		}
 	}
 
+	PreSkinnedBounds = Component->GetRestCollection()->RenderData->PreSkinnedBounds;
+
 	DynamicData = Component->InitDynamicData(true);
 }
 
@@ -1134,7 +1136,7 @@ HHitProxy* FGeometryCollectionSceneProxy::CreateHitProxies(UPrimitiveComponent*
 
 void FGeometryCollectionSceneProxy::GetPreSkinnedLocalBounds(FBoxSphereBounds& OutBounds) const
 {
-	OutBounds = MeshDescription.PreSkinnedBounds; 
+	OutBounds = PreSkinnedBounds; 
 }
 
 uint32 FGeometryCollectionSceneProxy::GetAllocatedSize() const
@@ -1345,8 +1347,6 @@ void FNaniteGeometryCollectionSceneProxy::CreateRenderThreadResources(FRHIComman
 	check(NaniteResources.RuntimeResourceID != INDEX_NONE && NaniteResources.HierarchyOffset != INDEX_NONE);
 
 	FGeometryCollectionSceneProxyBase::SetDynamicData_RenderThread(RHICmdList, DynamicData);
-
-	CachedPreSkinnedBounds = GetLocalBounds();
 }
 
 void FNaniteGeometryCollectionSceneProxy::DestroyRenderThreadResources()
@@ -1362,7 +1362,7 @@ SIZE_T FNaniteGeometryCollectionSceneProxy::GetTypeHash() const
 
 void FNaniteGeometryCollectionSceneProxy::GetPreSkinnedLocalBounds(FBoxSphereBounds& OutBounds) const
 {
-	OutBounds = CachedPreSkinnedBounds;
+	OutBounds = PreSkinnedBounds;
 }
 
 FPrimitiveViewRelevance FNaniteGeometryCollectionSceneProxy::GetViewRelevance(const FSceneView* View) const
diff --git a/Engine/Source/Runtime/Experimental/GeometryCollectionEngine/Private/GeometryCollection/GeometryCollectionSceneProxy.h b/Engine/Source/Runtime/Experimental/GeometryCollectionEngine/Private/GeometryCollection/GeometryCollectionSceneProxy.h
--- a/Engine/Source/Runtime/Experimental/GeometryCollectionEngine/Private/GeometryCollection/GeometryCollectionSceneProxy.h
+++ b/Engine/Source/Runtime/Experimental/GeometryCollectionEngine/Private/GeometryCollection/GeometryCollectionSceneProxy.h
@@ -192,6 +192,8 @@ protected:
 	FGeometryCollectionMeshDescription MeshDescription;
 	FMaterialRelevance MaterialRelevance;
 
+	FBoxSphereBounds PreSkinnedBounds;
+
 	FGeometryCollectionVertexFactory VertexFactory;
 	const bool bUseShaderBoneTransform;
 	const bool bSupportsTripleBufferVertexUpload;
@@ -242,8 +244,6 @@ class FGeometryCollectionSceneProxy final : public FPrimitiveSceneProxy, private
 
 	FCollisionResponseContainer CollisionResponse;
 
-	FBoxSphereBounds PreSkinnedBounds;
-
 	bool bRenderResourcesCreated = false;
 
 #if WITH_EDITOR
@@ -377,8 +377,6 @@ protected:
 	};
 	TArray<FGeometryNaniteData> GeometryNaniteData;
 
-	FBoxSphereBounds CachedPreSkinnedBounds;
-
 	uint32 bCastShadow : 1;
 	uint32 bReverseCulling : 1;
 	uint32 bHasMaterialErrors : 1;
