Geometry collection physics proxy : better handling of missing implicits collision attribute on initialization
Also added logs when this happen in order to capture the name of the corresponding assets in logs

diff --git a/Engine/Source/Runtime/Experimental/Chaos/Private/PhysicsProxy/GeometryCollectionPhysicsProxy.cpp b/Engine/Source/Runtime/Experimental/Chaos/Private/PhysicsProxy/GeometryCollectionPhysicsProxy.cpp
--- a/Engine/Source/Runtime/Experimental/Chaos/Private/PhysicsProxy/GeometryCollectionPhysicsProxy.cpp
+++ b/Engine/Source/Runtime/Experimental/Chaos/Private/PhysicsProxy/GeometryCollectionPhysicsProxy.cpp
@@ -1264,7 +1264,11 @@ void FGeometryCollectionPhysicsProxy::InitializeBodiesPT(Chaos::FPBDRigidsSolver
 		const TManagedArray<uint8>& DynamicState = DynamicCollection.DynamicState;
 		const TManagedArray<bool>& SimulatableParticles = DynamicCollection.SimulatableParticles;
 		const TManagedArray<FTransform>& MassToLocal = RestCollection->GetAttribute<FTransform>(MassToLocalAttributeName, FTransformCollection::TransformGroup);
-		const TManagedArray<Chaos::FImplicitObjectPtr>& Implicits = DynamicCollection.GetAttribute<Chaos::FImplicitObjectPtr>(FGeometryDynamicCollection::ImplicitsAttribute, FTransformCollection::TransformGroup);
+		const TManagedArray<Chaos::FImplicitObjectPtr>* Implicits = DynamicCollection.FindAttribute<Chaos::FImplicitObjectPtr>(FGeometryDynamicCollection::ImplicitsAttribute, FTransformCollection::TransformGroup);
+		if (!ensure(Implicits))
+		{
+			UE_LOG(LogChaos, Error, TEXT("Geometry collection %s - InitializeBodiesPT called with collision missing in the collection "), Owner ? *Owner->GetFullName(): TEXT("Unknown"));
+		}
 
 		using FSimplicialUniquePtr = TUniquePtr<FCollisionStructureManager::FSimplicial>;
 		const TManagedArray<FSimplicialUniquePtr>* Simplicials = nullptr;
@@ -1321,11 +1325,13 @@ void FGeometryCollectionPhysicsProxy::InitializeBodiesPT(Chaos::FPBDRigidsSolver
 				const Chaos::FReal ScaledMass = AdjustMassForScale(Mass[TransformGroupIndex]);
 				const Chaos::FVec3f ScaledInertia = AdjustInertiaForScale((Chaos::FVec3f)InertiaTensor[TransformGroupIndex]);
 
+				const bool bUseImplicit = (Implicits != nullptr) && (BuildGeometryForChildrenOnPT(Parameters) || (TransformGroupIndex == Parameters.InitialRootIndex));
+
 				PopulateSimulatedParticle(
 					Handle,
 					Parameters.Shared,
 					Simplicials? (*Simplicials)[TransformGroupIndex].Get(): nullptr,
-					(BuildGeometryForChildrenOnPT(Parameters) || (TransformGroupIndex == Parameters.InitialRootIndex)) ? Implicits[TransformGroupIndex] : nullptr,
+					bUseImplicit? (*Implicits)[TransformGroupIndex] : nullptr,
 					SimFilter,
 					QueryFilter,
 					ScaledMass,
@@ -1741,7 +1747,7 @@ void FGeometryCollectionPhysicsProxy::CreateChildrenGeometry_Internal()
 
 			check(NumTransforms == PhysicsThreadCollection.NumElements(FTransformCollection::TransformGroup));
 			const TManagedArray<FTransform>& MassToLocal = RestCollection->GetAttribute<FTransform>(MassToLocalAttributeName, FTransformCollection::TransformGroup);
-			const TManagedArray<Chaos::FImplicitObjectPtr>& Implicits = PhysicsThreadCollection.GetAttribute<Chaos::FImplicitObjectPtr>(FGeometryDynamicCollection::ImplicitsAttribute, FTransformCollection::TransformGroup);
+			const TManagedArray<Chaos::FImplicitObjectPtr>* Implicits = PhysicsThreadCollection.FindAttribute<Chaos::FImplicitObjectPtr>(FGeometryDynamicCollection::ImplicitsAttribute, FTransformCollection::TransformGroup);
 
 			using FSimplicialUniquePtr = TUniquePtr<FCollisionStructureManager::FSimplicial>;
 			const TManagedArray<FSimplicialUniquePtr>* Simplicials = nullptr;
@@ -1771,7 +1777,7 @@ void FGeometryCollectionPhysicsProxy::CreateChildrenGeometry_Internal()
 					{
 						const FTransform WorldTransform = MassToLocal[TransformIndex] * Transform[TransformIndex] * Parameters.WorldTransform;
 						const FCollisionStructureManager::FSimplicial* Simplicial = Simplicials ? (*Simplicials)[TransformIndex].Get() : nullptr;
-						SetImplicitToPTParticles(Handle, Parameters.Shared, Simplicial, Implicits[TransformIndex], SimFilter, QueryFilter, WorldTransform, CollisionParticlesPerObjectFraction);
+						SetImplicitToPTParticles(Handle, Parameters.Shared, Simplicial, Implicits? (*Implicits)[TransformIndex]: nullptr, SimFilter, QueryFilter, WorldTransform, CollisionParticlesPerObjectFraction);
 						check(Handle->GetGeometry() != nullptr);
 					}
 				}
@@ -1950,7 +1956,7 @@ Chaos::TPBDGeometryCollectionParticleHandle<Chaos::FReal, 3>* FGeometryCollectio
 	FGeometryDynamicCollection& DynamicCollection = PhysicsThreadCollection;
 	TManagedArray<uint8>& DynamicState = DynamicCollection.DynamicState;
 	const TManagedArray<FTransform>& MassToLocal = Parameters.RestCollectionShared->GetAttribute<FTransform>(MassToLocalAttributeName, FTransformCollection::TransformGroup);
-	const TManagedArray<Chaos::FImplicitObjectPtr>& Implicits = DynamicCollection.GetAttribute<Chaos::FImplicitObjectPtr>(FGeometryDynamicCollection::ImplicitsAttribute, FTransformCollection::TransformGroup);
+	const TManagedArray<Chaos::FImplicitObjectPtr>* Implicits = DynamicCollection.FindAttribute<Chaos::FImplicitObjectPtr>(FGeometryDynamicCollection::ImplicitsAttribute, FTransformCollection::TransformGroup);
 
 	using FSimplicialUniquePtr = TUniquePtr<FCollisionStructureManager::FSimplicial>;
 	const TManagedArray<FSimplicialUniquePtr>* Simplicials = nullptr;
@@ -2001,13 +2007,15 @@ Chaos::TPBDGeometryCollectionParticleHandle<Chaos::FReal, 3>* FGeometryCollectio
 
 	check(Parameters.RestCollectionShared);
 
+	const bool bUseImplicit = (Implicits != nullptr) && (BuildGeometryForChildrenOnPT(Parameters) || (CollectionClusterIndex == Parameters.InitialRootIndex));
+
 	// NOTE: The particle creation call above (CreateClusterParticle) activates the particle, which does various things including adding
 	// it to the constraint graph. It seems a bit dodgy to be completely changing what the particle is after it has been enabled. Maybe we should fix this.
 	PopulateSimulatedParticle(
 		Handle,
 		Parameters.Shared,
 		Simplicials? (*Simplicials)[CollectionClusterIndex].Get(): nullptr,
-		(BuildGeometryForChildrenOnPT(Parameters) || (CollectionClusterIndex == Parameters.InitialRootIndex)) ? Implicits[CollectionClusterIndex] : nullptr,
+		bUseImplicit? (*Implicits)[CollectionClusterIndex] : nullptr,
 		SimFilter,
 		QueryFilter,
 		Mass,
@@ -2074,7 +2082,7 @@ FGeometryCollectionPhysicsProxy::BuildClusters_Internal(
 	FGeometryDynamicCollection& DynamicCollection = PhysicsThreadCollection;
 	TManagedArray<uint8>& DynamicState = DynamicCollection.DynamicState;
 	const TManagedArray<FTransform>& MassToLocal = Parameters.RestCollectionShared->GetAttribute<FTransform>(MassToLocalAttributeName, FTransformCollection::TransformGroup);
-	const TManagedArray<Chaos::FImplicitObjectPtr>& Implicits = DynamicCollection.GetAttribute<Chaos::FImplicitObjectPtr>(FGeometryDynamicCollection::ImplicitsAttribute, FTransformCollection::TransformGroup);
+	const TManagedArray<Chaos::FImplicitObjectPtr>* Implicits = DynamicCollection.FindAttribute<Chaos::FImplicitObjectPtr>(FGeometryDynamicCollection::ImplicitsAttribute, FTransformCollection::TransformGroup);
 
 	//If we are a root particle use the world transform, otherwise set the relative transform
 	const FTransform CollectionSpaceTransform = GeometryCollectionAlgo::Private::GlobalMatrix(DynamicCollection, CollectionClusterIndex);
@@ -2109,7 +2117,7 @@ FGeometryCollectionPhysicsProxy::BuildClusters_Internal(
 			(!DynamicCollection.GetHasParent(CollectionClusterIndex)) ? Parameters.ClusterGroupIndex : 0, 
 			MoveTemp(ChildHandlesCopy),
 			ClusterCreationParameters,
-			Implicits[CollectionClusterIndex], // union from children if null
+			Implicits? (*Implicits)[CollectionClusterIndex]: nullptr, // union from children if null
 			&ParticleTM,
 			ExistingIndex
 			);
@@ -2995,7 +3003,7 @@ TBitArray<> FGeometryCollectionPhysicsProxy::CalculateClustersToCreateFromChildr
 	{
 		ClustersToGenerate.Init(false, NumTransforms);
 
-		const TManagedArray<Chaos::FImplicitObjectPtr>& Implicits = DynamicCollection.GetAttribute<Chaos::FImplicitObjectPtr>(FGeometryDynamicCollection::ImplicitsAttribute, FTransformCollection::TransformGroup);
+		const TManagedArray<Chaos::FImplicitObjectPtr>* Implicits = DynamicCollection.FindAttribute<Chaos::FImplicitObjectPtr>(FGeometryDynamicCollection::ImplicitsAttribute, FTransformCollection::TransformGroup);
 
 		// All the leaf objects
 		TArray<int32> ChildrenToCheckForParentFix;
@@ -3022,7 +3030,7 @@ TBitArray<> FGeometryCollectionPhysicsProxy::CalculateClustersToCreateFromChildr
 			// step 2: test the parent for having children with geometry
 			for (const int32 ParentToFixIndex : ParentToPotentiallyFix)
 			{
-				const bool bParentHasCollision = (Implicits[ParentToFixIndex] != nullptr) || ClustersToGenerate[ParentToFixIndex];
+				const bool bParentHasCollision = Implicits && (((*Implicits)[ParentToFixIndex] != nullptr) || ClustersToGenerate[ParentToFixIndex]);
 
 				if (!bParentHasCollision)
 				{
@@ -3031,7 +3039,7 @@ TBitArray<> FGeometryCollectionPhysicsProxy::CalculateClustersToCreateFromChildr
 					DynamicCollection.IterateThroughChildren(ParentToFixIndex, [&](int32 ChildIndex)
 					{
 						// defer if any of the children is a cluster with no collision yet generated 
-						const bool bChildHasCollision = (Implicits[ChildIndex] != nullptr) || ClustersToGenerate[ChildIndex];
+						const bool bChildHasCollision = Implicits && (((*Implicits)[ChildIndex] != nullptr) || ClustersToGenerate[ChildIndex]);
 						if (!bChildHasCollision && DynamicCollection.HasChildren(ChildIndex))
 						{
 							bAllChildrenHaveCollision = false;
@@ -3061,7 +3069,7 @@ int32 FGeometryCollectionPhysicsProxy::CalculateEffectiveParticles(const FGeomet
 	int32 NumEffectiveParticlesFound = 0;
 	EffectiveParticles.Init(false, NumTransform);
 	const int32 MaxSimulatedLevel = FMath::Min(GlobalMaxSimulatedLevel, InMaxSimulatedLevel);
-	const TManagedArray<Chaos::FImplicitObjectPtr>& Implicits = DynamicCollection.GetAttribute<Chaos::FImplicitObjectPtr>(FGeometryDynamicCollection::ImplicitsAttribute, FTransformCollection::TransformGroup);
+	const TManagedArray<Chaos::FImplicitObjectPtr>* Implicits = DynamicCollection.FindAttribute<Chaos::FImplicitObjectPtr>(FGeometryDynamicCollection::ImplicitsAttribute, FTransformCollection::TransformGroup);
 	TManagedArrayAccessor<int32> Levels = DynamicCollection.GetInitialLevels();
 
 	for (int32 TransformIndex = 0; TransformIndex < NumTransform; ++TransformIndex)
@@ -3070,7 +3078,7 @@ int32 FGeometryCollectionPhysicsProxy::CalculateEffectiveParticles(const FGeomet
 		if (Level <= MaxSimulatedLevel || !bEnableClustering)
 		{
 			const bool bIsClusterUsingChildGeometry = (ClustersUsingChildGeometry.Num() > 0) && ClustersUsingChildGeometry[TransformIndex];
-			const bool bHasGeometry = Implicits[TransformIndex].IsValid();
+			const bool bHasGeometry = Implicits && (*Implicits)[TransformIndex].IsValid();
 
 			if (bHasGeometry || bIsClusterUsingChildGeometry)
 			{
