Updated SphereConvex, CapsuleConvex and CapsuleCapsule sweeps to use GjkRaycastSameSpace.
Added a few more tests.

diff --git a/Engine/Source/Runtime/Experimental/Chaos/Private/Chaos/GeometryQueries.cpp b/Engine/Source/Runtime/Experimental/Chaos/Private/Chaos/GeometryQueries.cpp
--- a/Engine/Source/Runtime/Experimental/Chaos/Private/Chaos/GeometryQueries.cpp
+++ b/Engine/Source/Runtime/Experimental/Chaos/Private/Chaos/GeometryQueries.cpp
@@ -110,6 +110,18 @@ namespace Chaos
 		return bResult;
 	}
 
+	bool SweepSphereVsConvex(const FSphere& SweptSphere, const FRigidTransform3& SweptSphereTM, const FImplicitObject& TestObject, const FRigidTransform3& TestObjectTM,
+		const FVec3& SweepDir, const FReal Length, const FReal Thickness, const bool bComputeMTD,
+		FReal& OutTime, FVec3& OutPosition, FVec3& OutNormal, int32& OutFaceIndex, FVec3& OutFaceNormal)
+	{
+		const FGJKSphereSingleSIMD LocalSweptSphere = ComputeLocalSphere(SweptSphere, SweptSphereTM, TestObjectTM);
+		return Utilities::CastHelperNoUnwrap(TestObject, TestObjectTM,
+			[&LocalSweptSphere, &SweepDir, Length, Thickness, bComputeMTD, &OutTime, &OutPosition, &OutNormal, &OutFaceIndex, &OutFaceNormal]<typename T>(const T& TestObjectDowncast, const FRigidTransform3& TestObjectTM)
+			{
+				return LocalSweep(LocalSweptSphere, TestObjectDowncast, TestObjectDowncast, TestObjectTM, SweepDir, Length, Thickness, bComputeMTD, OutTime, OutPosition, OutNormal, OutFaceIndex, OutFaceNormal);
+			});
+	}
+
 	bool SweepBoxVsSphere(const TBox<FReal, 3>& SweptBox, const FRigidTransform3& SweptBoxTM, const FSphere& TestSphere, const FRigidTransform3& TestSphereTM,
 		const FVec3& SweepDir, const FReal Length, const FReal Thickness, const bool bComputeMTD,
 		FReal& OutTime, FVec3& OutPosition, FVec3& OutNormal, int32& OutFaceIndex, FVec3& OutFaceNormal)
@@ -118,8 +130,8 @@ namespace Chaos
 		return LocalSweep(SweptBox, LocalTestSphere, TestSphere, SweptBoxTM, SweepDir, Length, Thickness, bComputeMTD, OutTime, OutPosition, OutNormal, OutFaceIndex, OutFaceNormal);
 	}
 
-	bool SweepBoxVsCapsule(const TBox<FReal, 3>& SweptBox, const FRigidTransform3& SweptBoxTM, const FCapsule& TestCapsule, const FRigidTransform3& TestCapsuleTM,
-		const FVec3& SweepDir, const FReal Length, const FReal Thickness, const bool bComputeMTD,
+	bool SweepBoxVsCapsule(const TBox<FReal, 3>& SweptBox, const FRigidTransform3& SweptBoxTM, const FCapsule& TestCapsule, const FRigidTransform3& TestCapsuleTM, 
+		const FVec3& SweepDir, const FReal Length, const FReal Thickness, const bool bComputeMTD, 
 		FReal& OutTime, FVec3& OutPosition, FVec3& OutNormal, int32& OutFaceIndex, FVec3& OutFaceNormal)
 	{
 		const FGJKCapsuleSingleSIMD LocalTestCapsule = ComputeLocalCapsule(TestCapsule, TestCapsuleTM, SweptBoxTM);
@@ -153,4 +165,20 @@ namespace Chaos
 		const FGJKCapsuleSingleSIMD LocalSweptCapsule = ComputeLocalCapsule(SweptCapsule, SweptCapsuleTM, TestBoxTM);
 		return LocalSweep(LocalSweptCapsule, TestBox, TestBox, TestBoxTM, SweepDir, Length, Thickness, bComputeMTD, OutTime, OutPosition, OutNormal, OutFaceIndex, OutFaceNormal);
 	}
+
+	bool SweepCapsuleVsCapsule(const FCapsule& SweptCapsule, const FRigidTransform3& SweptCapsuleTM, const FCapsule& TestCapsule, const FRigidTransform3& TestCapsuleTM, const FVec3& SweepDir, const FReal Length, const FReal Thickness, const bool bComputeMTD, FReal& OutTime, FVec3& OutPosition, FVec3& OutNormal, int32& OutFaceIndex, FVec3& OutFaceNormal)
+	{
+		const FGJKCapsuleSingleSIMD LocalSweptCapsule = ComputeLocalCapsule(SweptCapsule, SweptCapsuleTM, TestCapsuleTM);
+		return LocalSweep(LocalSweptCapsule, TestCapsule, TestCapsule, TestCapsuleTM, SweepDir, Length, Thickness, bComputeMTD, OutTime, OutPosition, OutNormal, OutFaceIndex, OutFaceNormal);
+	}
+
+	bool SweepCapsuleVsConvex(const FCapsule& SweptCapsule, const FRigidTransform3& SweptCapsuleTM, const FImplicitObject& TestObject, const FRigidTransform3& TestObjectTM, const FVec3& SweepDir, const FReal Length, const FReal Thickness, const bool bComputeMTD, FReal& OutTime, FVec3& OutPosition, FVec3& OutNormal, int32& OutFaceIndex, FVec3& OutFaceNormal)
+	{
+		const FGJKCapsuleSingleSIMD LocalSweptCapsule = ComputeLocalCapsule(SweptCapsule, SweptCapsuleTM, TestObjectTM);
+		return Utilities::CastHelperNoUnwrap(TestObject, TestObjectTM,
+			[&LocalSweptCapsule, &SweepDir, Length, Thickness, bComputeMTD, &OutTime, &OutPosition, &OutNormal, &OutFaceIndex, &OutFaceNormal]<typename T>(const T& TestObjectDowncast, const FRigidTransform3& TestObjectTM)
+			{
+				return LocalSweep(LocalSweptCapsule, TestObjectDowncast, TestObjectDowncast, TestObjectTM, SweepDir, Length, Thickness, bComputeMTD, OutTime, OutPosition, OutNormal, OutFaceIndex, OutFaceNormal);
+			});
+	}
 }
diff --git a/Engine/Source/Runtime/Experimental/Chaos/Public/Chaos/GeometryQueries.h b/Engine/Source/Runtime/Experimental/Chaos/Public/Chaos/GeometryQueries.h
--- a/Engine/Source/Runtime/Experimental/Chaos/Public/Chaos/GeometryQueries.h
+++ b/Engine/Source/Runtime/Experimental/Chaos/Public/Chaos/GeometryQueries.h
@@ -156,10 +156,13 @@ namespace Chaos
 	UE_INTERNAL CHAOS_API bool SweepSphereVsSphere(const FSphere& SweptSphere, const FRigidTransform3& SweptSphereTM, const FSphere& TestSphere, const FRigidTransform3& TestSphereTM, const FVec3& SweepDir, const FReal Length, const FReal Thickness, const bool bComputeMTD, FReal& OutTime, FVec3& OutPosition, FVec3& OutNormal, int32& OutFaceIndex, FVec3& OutFaceNormal);
 	UE_INTERNAL CHAOS_API bool SweepSphereVsBox(const FSphere& SweptSphere, const FRigidTransform3& SweptSphereTM, const TBox<FReal, 3>& TestBox, const FRigidTransform3& TestBoxTM, const FVec3& SweepDir, const FReal Length, const FReal Thickness, const bool bComputeMTD, FReal& OutTime, FVec3& OutPosition, FVec3& OutNormal, int32& OutFaceIndex, FVec3& OutFaceNormal);
 	UE_INTERNAL CHAOS_API bool SweepSphereVsCapsule(const FSphere& SweptSphere, const FRigidTransform3& SweptSphereTM, const FCapsule& TestCapsule, const FRigidTransform3& TestCapsuleTM, const FVec3& SweepDir, const FReal Length, const FReal Thickness, const bool bComputeMTD, FReal& OutTime, FVec3& OutPosition, FVec3& OutNormal, int32& OutFaceIndex, FVec3& OutFaceNormal);
+	UE_INTERNAL CHAOS_API bool SweepSphereVsConvex(const FSphere& SweptSphere, const FRigidTransform3& SweptSphereTM, const FImplicitObject& TestObject, const FRigidTransform3& TestObjectTM, const FVec3& SweepDir, const FReal Length, const FReal Thickness, const bool bComputeMTD, FReal& OutTime, FVec3& OutPosition, FVec3& OutNormal, int32& OutFaceIndex, FVec3& OutFaceNormal);
 	UE_INTERNAL CHAOS_API bool SweepBoxVsSphere(const TBox<FReal, 3>& SweptBox, const FRigidTransform3& SweptBoxTM, const FSphere& TestSphere, const FRigidTransform3& TestSphereTM, const FVec3& SweepDir, const FReal Length, const FReal Thickness, const bool bComputeMTD, FReal& OutTime, FVec3& OutPosition, FVec3& OutNormal, int32& OutFaceIndex, FVec3& OutFaceNormal);
 	UE_INTERNAL CHAOS_API bool SweepBoxVsCapsule(const TBox<FReal, 3>& SweptBox, const FRigidTransform3& SweptBoxTM, const FCapsule& TestCapsule, const FRigidTransform3& TestCapsuleTM, const FVec3& SweepDir, const FReal Length, const FReal Thickness, const bool bComputeMTD, FReal& OutTime, FVec3& OutPosition, FVec3& OutNormal, int32& OutFaceIndex, FVec3& OutFaceNormal);
 	UE_INTERNAL CHAOS_API bool SweepCapsuleVsSphere(const FCapsule& SweptCapsule, const FRigidTransform3& SweptCapsuleTM, const FSphere& TestSphere, const FRigidTransform3& TestSphereTM, const FVec3& SweepDir, const FReal Length, const FReal Thickness, const bool bComputeMTD, FReal& OutTime, FVec3& OutPosition, FVec3& OutNormal, int32& OutFaceIndex, FVec3& OutFaceNormal);
 	UE_INTERNAL CHAOS_API bool SweepCapsuleVsBox(const FCapsule& SweptCapsule, const FRigidTransform3& SweptCapsuleTM, const TBox<FReal, 3>& TestBox, const FRigidTransform3& TestBoxTM, const FVec3& SweepDir, const FReal Length, const FReal Thickness, const bool bComputeMTD, FReal& OutTime, FVec3& OutPosition, FVec3& OutNormal, int32& OutFaceIndex, FVec3& OutFaceNormal);
+	UE_INTERNAL CHAOS_API bool SweepCapsuleVsCapsule(const FCapsule& SweptCapsule, const FRigidTransform3& SweptCapsuleTM, const FCapsule& TestCapsule, const FRigidTransform3& TestCapsuleTM, const FVec3& SweepDir, const FReal Length, const FReal Thickness, const bool bComputeMTD, FReal& OutTime, FVec3& OutPosition, FVec3& OutNormal, int32& OutFaceIndex, FVec3& OutFaceNormal);
+	UE_INTERNAL CHAOS_API bool SweepCapsuleVsConvex(const FCapsule& SweptCapsule, const FRigidTransform3& SweptCapsuleTM, const FImplicitObject& TestObject, const FRigidTransform3& TestObjectTM, const FVec3& SweepDir, const FReal Length, const FReal Thickness, const bool bComputeMTD, FReal& OutTime, FVec3& OutPosition, FVec3& OutNormal, int32& OutFaceIndex, FVec3& OutFaceNormal);
 
 	template <EImplicitObjectType SweptType>
 	UE_INTERNAL bool TryRunSpecializedSweep(const FImplicitObject& SweptShape, const FRigidTransform3& SweptShapeTM, const FImplicitObject& TestGeom, const FRigidTransform3& TestGeomTM, const FVec3& SweepDir, const FReal Length, const FReal Thickness, const bool bComputeMTD, bool& bOutResult, FReal& OutTime, FVec3& OutPosition, FVec3& OutNormal, int32& OutFaceIndex, FVec3& OutFaceNormal)
@@ -188,6 +191,12 @@ namespace Chaos
 				bOutResult = SweepSphereVsCapsule(SweptSphere, SweptShapeTM, TestCapsule, TestGeomTM, SweepDir, Length, Thickness, bComputeMTD, OutTime, OutPosition, OutNormal, OutFaceIndex, OutFaceNormal);
 				return true;
 			}
+			else if (TestGeom.IsConvex())
+			{
+				const FSphere& SweptSphere = SweptShape.GetObjectChecked<FSphere>();
+				bOutResult = SweepSphereVsConvex(SweptSphere, SweptShapeTM, TestGeom, TestGeomTM, SweepDir, Length, Thickness, bComputeMTD, OutTime, OutPosition, OutNormal, OutFaceIndex, OutFaceNormal);
+				return true;
+			}
 		}
 		else if (SweptType == ImplicitObjectType::Box)
 		{
@@ -222,6 +231,19 @@ namespace Chaos
 				bOutResult = SweepCapsuleVsBox(SweptCapsule, SweptShapeTM, TestBox, TestGeomTM, SweepDir, Length, Thickness, bComputeMTD, OutTime, OutPosition, OutNormal, OutFaceIndex, OutFaceNormal);
 				return true;
 			}
+			else if (TestType == ImplicitObjectType::Capsule)
+			{
+				const FCapsule& SweptCapsule = SweptShape.GetObjectChecked<FCapsule>();
+				const FCapsule& TestCapsule = TestGeom.GetObjectChecked<FCapsule>();
+				bOutResult = SweepCapsuleVsCapsule(SweptCapsule, SweptShapeTM, TestCapsule, TestGeomTM, SweepDir, Length, Thickness, bComputeMTD, OutTime, OutPosition, OutNormal, OutFaceIndex, OutFaceNormal);
+				return true;
+			}
+			else if (TestGeom.IsConvex())
+			{
+				const FCapsule& SweptCapsule = SweptShape.GetObjectChecked<FCapsule>();
+				bOutResult = SweepCapsuleVsConvex(SweptCapsule, SweptShapeTM, TestGeom, TestGeomTM, SweepDir, Length, Thickness, bComputeMTD, OutTime, OutPosition, OutNormal, OutFaceIndex, OutFaceNormal);
+				return true;
+			}
 		}
 		return false;
 	}
