IAD installer - Fix deadlock from double locking mutex on request destruction.  This happens when an install request is the last reference to a content handle.  The content handle destructor needs to be called outside of the installer mutex because it may enque a request to the installer to update last access times

diff --git a/Engine/Source/Runtime/Core/Internal/IO/IoAllocators.h b/Engine/Source/Runtime/Core/Internal/IO/IoAllocators.h
--- a/Engine/Source/Runtime/Core/Internal/IO/IoAllocators.h
+++ b/Engine/Source/Runtime/Core/Internal/IO/IoAllocators.h
@@ -8,6 +8,22 @@
 template <typename T, uint16 SlabSize = 4096>
 class TSingleThreadedSlabAllocator
 {
+private:
+	struct FSlab;
+
+	struct FElement
+	{
+		TTypeCompatibleBytes<T> Data;
+		FSlab* Slab = nullptr;
+	};
+
+	struct FSlab
+	{
+		uint16 Allocated = 0;
+		uint16 Freed = 0;
+		FElement Elements[SlabSize];
+	};
+
 public:
 	TSingleThreadedSlabAllocator()
 	{
@@ -32,22 +48,6 @@ public:
 		Free(Ptr);
 	}
 
-private:
-	struct FSlab;
-
-	struct FElement
-	{
-		TTypeCompatibleBytes<T> Data;
-		FSlab* Slab = nullptr;
-	};
-
-	struct FSlab
-	{
-		uint16 Allocated = 0;
-		uint16 Freed = 0;
-		FElement Elements[SlabSize];
-	};
-
 	T* Alloc()
 	{
 		uint16 ElementIndex = CurrentSlab->Allocated++;
@@ -74,5 +74,6 @@ private:
 		}
 	}
 
+private:
 	FSlab* CurrentSlab = nullptr;
 };
diff --git a/Engine/Source/Runtime/Experimental/IoStore/OnDemand/Private/OnDemandContentInstaller.cpp b/Engine/Source/Runtime/Experimental/IoStore/OnDemand/Private/OnDemandContentInstaller.cpp
--- a/Engine/Source/Runtime/Experimental/IoStore/OnDemand/Private/OnDemandContentInstaller.cpp
+++ b/Engine/Source/Runtime/Experimental/IoStore/OnDemand/Private/OnDemandContentInstaller.cpp
@@ -1018,6 +1018,15 @@ void FOnDemandContentInstaller::NotifyInstallProgress(FRequest& Request)
 	}
 }
 
+void FOnDemandContentInstaller::DestroyRequest(FRequest& Request)
+{
+	// Call destructor outside the lock to avoid double locking
+	Request.~FRequest();
+
+	UE::TUniqueLock Lock(Mutex);
+	RequestAllocator.Free(&Request);
+}
+
 void FOnDemandContentInstaller::CompleteInstallRequest(FRequest& Request)
 {
 	using namespace UE::IoStore::Private;
@@ -1111,10 +1120,7 @@ void FOnDemandContentInstaller::CompleteInstallRequest(FRequest& Request)
 	if (!InstallRequest.OnCompleted)
 	{
 		InstallRequest.Request->Status.store(FinalRequestStatus);
-		{
-			UE::TUniqueLock Lock(Mutex);
-			RequestAllocator.Destroy(&Request);
-		}
+		DestroyRequest(Request);
 		return;
 	}
 
@@ -1134,10 +1140,7 @@ void FOnDemandContentInstaller::CompleteInstallRequest(FRequest& Request)
 				ensure(InstallRequest.bNotifyingProgressOnGameThread == false);
 				OnCompleted(MoveTemp(InstallResult));
 				InstallRequest.Request->Status.store(FinalRequestStatus);
-				{
-					UE::TUniqueLock Lock(Mutex);
-					RequestAllocator.Destroy(&Request);
-				}
+				DestroyRequest(Request);
 			});
 	}
 	else
@@ -1145,10 +1148,7 @@ void FOnDemandContentInstaller::CompleteInstallRequest(FRequest& Request)
 		FOnDemandInstallCompleted OnCompleted = MoveTemp(InstallRequest.OnCompleted);
 		OnCompleted(MoveTemp(InstallResult));
 		InstallRequest.Request->Status.store(FinalRequestStatus);
-		{
-			UE::TUniqueLock Lock(Mutex);
-			RequestAllocator.Destroy(&Request);
-		}
+		DestroyRequest(Request);
 	}
 }
 
@@ -1173,8 +1173,8 @@ void FOnDemandContentInstaller::CompletePurgeRequest(FRequest& Request)
 		{
 			CurrentRequest = nullptr;
 		}
-		RequestAllocator.Destroy(&Request);
 	}
+	DestroyRequest(Request);
 
 	TryExecuteNextRequest();
 
@@ -1219,10 +1219,10 @@ void FOnDemandContentInstaller::CompleteDefragRequest(FRequest& Request)
 		{
 			CurrentRequest = nullptr;
 		}
-		RequestAllocator.Destroy(&Request);
 	}
+	DestroyRequest(Request);
 
-	TryExecuteNextRequest();	
+	TryExecuteNextRequest();
 
 	if (!OnCompleted)
 	{
@@ -1262,8 +1262,8 @@ void FOnDemandContentInstaller::CompleteVerifyRequest(FRequest& Request)
 		{
 			CurrentRequest = nullptr;
 		}
-		RequestAllocator.Destroy(&Request);
 	}
+	DestroyRequest(Request);
 
 	if (!OnCompleted)
 	{
@@ -1291,8 +1291,8 @@ void FOnDemandContentInstaller::CompleteFlushLastAccessRequest(FRequest& Request
 		{
 			CurrentRequest = nullptr;
 		}
-		RequestAllocator.Destroy(&Request);
 	}
+	DestroyRequest(Request);
 
 	if (!OnCompleted)
 	{
diff --git a/Engine/Source/Runtime/Experimental/IoStore/OnDemand/Private/OnDemandContentInstaller.h b/Engine/Source/Runtime/Experimental/IoStore/OnDemand/Private/OnDemandContentInstaller.h
--- a/Engine/Source/Runtime/Experimental/IoStore/OnDemand/Private/OnDemandContentInstaller.h
+++ b/Engine/Source/Runtime/Experimental/IoStore/OnDemand/Private/OnDemandContentInstaller.h
@@ -226,6 +226,7 @@ private:
 										FString&& ErrorReason,
 										FIoBuffer&& Chunk);
 	void							NotifyInstallProgress(FRequest& Request);
+	void							DestroyRequest(FRequest& Request);
 	void							CompleteInstallRequest(FRequest& Request);
 	void							CompletePurgeRequest(FRequest& Request);
 	void							CompleteDefragRequest(FRequest& Request);
