Add null checks.  Bug hawk suggested these.  They would most likely only apply to a server waiting to fork but I reviewed the code and it does look possible to hit some of these if a content handle or mount request was made before successul cache initialization.

diff --git a/Engine/Source/Runtime/Experimental/IoStore/OnDemand/Private/OnDemandIoStore.cpp b/Engine/Source/Runtime/Experimental/IoStore/OnDemand/Private/OnDemandIoStore.cpp
--- a/Engine/Source/Runtime/Experimental/IoStore/OnDemand/Private/OnDemandIoStore.cpp
+++ b/Engine/Source/Runtime/Experimental/IoStore/OnDemand/Private/OnDemandIoStore.cpp
@@ -705,7 +705,7 @@ FIoStatus FOnDemandIoStore::Unmount(FStringView MountId)
 			return false;
 		}));
 
-		if (bRemoved)
+		if (bRemoved && PackageStoreBackend)
 		{
 			PackageStoreBackend->NeedsUpdate();
 		}
@@ -1399,7 +1399,10 @@ FIoStatus FOnDemandIoStore::TickMountRequest(FMountRequest& MountRequest)
 		PendingContainerHeaders.Append(MoveTemp(RequestedContainerHeaders));
 	}
 
-	PackageStoreBackend->NeedsUpdate();
+	if (PackageStoreBackend)
+	{
+		PackageStoreBackend->NeedsUpdate();
+	}
 
 	return MountResult; 
 }
@@ -1490,7 +1493,7 @@ void FOnDemandIoStore::OnEncryptionKeyAdded(const FGuid& Id, const FAES::FAESKey
 		}
 	}
 
-	if (bAddedContainerHeaders)
+	if (bAddedContainerHeaders && PackageStoreBackend)
 	{
 		PackageStoreBackend->NeedsUpdate();
 	}
@@ -1941,9 +1944,15 @@ void FOnDemandIoStore::ReleaseContent(FOnDemandInternalContentHandle& ContentHan
 	}
 
 	// Package store needs to be aware of install state
-	PackageStoreBackend->NeedsUpdate(EOnDemandPackageStoreUpdateMode::ReferencedPackages);
+	if (PackageStoreBackend)
+	{
+		PackageStoreBackend->NeedsUpdate(EOnDemandPackageStoreUpdateMode::ReferencedPackages);
+	}
 
-	InstallCache->UpdateLastAccess(ReleasedChunks);
+	if (InstallCache)
+	{
+		InstallCache->UpdateLastAccess(ReleasedChunks);
+	}
 }
 
 void FOnDemandIoStore::GetReferencedContent(TArray<FSharedOnDemandContainer>& OutContainers, TArray<TBitArray<>>& OutChunkEntryIndices, bool bPackageStore)
