Chaos - expose IsCCD on contact modifier (read only)

#rnx

diff --git a/Engine/Source/Runtime/Experimental/Chaos/Private/Chaos/ContactModification.cpp b/Engine/Source/Runtime/Experimental/Chaos/Private/Chaos/ContactModification.cpp
--- a/Engine/Source/Runtime/Experimental/Chaos/Private/Chaos/ContactModification.cpp
+++ b/Engine/Source/Runtime/Experimental/Chaos/Private/Chaos/ContactModification.cpp
@@ -513,6 +513,12 @@ namespace Chaos
 		return Constraint->GetImplicit(ParticleIdx);
 	}
 
+	bool FContactPairModifier::GetIsCCD() const
+	{
+		return Constraint->GetCCDEnabled() && Constraint->GetCCDSweepEnabled();
+	}
+
+
 	bool FContactPairModifier::IsEdgeContactPoint(int32 ContactPointIdx) const
 	{
 		const FManifoldPoint& ManifoldPoint = Constraint->GetManifoldPoint(ContactPointIdx);
diff --git a/Engine/Source/Runtime/Experimental/Chaos/Public/Chaos/ContactModification.h b/Engine/Source/Runtime/Experimental/Chaos/Public/Chaos/ContactModification.h
--- a/Engine/Source/Runtime/Experimental/Chaos/Public/Chaos/ContactModification.h
+++ b/Engine/Source/Runtime/Experimental/Chaos/Public/Chaos/ContactModification.h
@@ -270,6 +270,13 @@ namespace Chaos
 		*/
 		CHAOS_API const FConstImplicitObjectRef GetImplicit(int32 ParticleIdx) const;
 
+		/*
+		* Whether this contact was generated from a CCD hit. Even if CCD is enabled for the bodies,
+		* the hit may be a non-CCD contact if they are moving slowly.
+		* NOTE: You can disable CCD in a MidPhase Modifier, but it is too late to do that here.
+		*/
+		CHAOS_API bool GetIsCCD() const;
+
 		/*
 		* Check to see if a contact point index is an edge contact.
 		*/
