Chaos - fix comments on GetCCDSweepEnabled

#rnx

diff --git a/Engine/Source/Runtime/Experimental/Chaos/Public/Chaos/Collision/PBDCollisionConstraint.h b/Engine/Source/Runtime/Experimental/Chaos/Public/Chaos/Collision/PBDCollisionConstraint.h
--- a/Engine/Source/Runtime/Experimental/Chaos/Public/Chaos/Collision/PBDCollisionConstraint.h
+++ b/Engine/Source/Runtime/Experimental/Chaos/Public/Chaos/Collision/PBDCollisionConstraint.h
@@ -305,14 +305,21 @@ namespace Chaos
 		}
 
 		/**
-		 * For CCD constraints, do we need to run the initial sweep/rewind step.
-		 * This may change from tick to tick as an object's velocity changes.
+		 * For CCD constraints, a transient flag (reset every frame) that is used to both enable/disable CCD sweeps for
+		 * the frame, and to report whether the sweep was run and the results retained. It will be set to false if
+		 * the bodies are moving less than the CCD threshold this tick.
+		 * 
+		 * If this is true in a collision callback then: 1) a CCD sweep was performed; 2) a CCD impact was detected; and
+		 * 3) the collision is being treated as a CCD contact. The bodies will be moved back to first CCD contact.
+		 * 
+		 * NOTE: If a sweep is performed and a hit detected, but the resultant contact is shallow enough to be
+		 * handled by the standard collision solver path, this flag is set to false and no rollback will be performed.
 		 */
 		bool GetCCDSweepEnabled() const { return Flags.bCCDSweepEnabled; }
 
 		/**
-		 * For CCD constraints, enable or disable the pre-solve sweep/rewind for this tick. Called
-		 * every tick for CCD constraints based on current velocity versus size etc.
+		 * For CCD constraints, a transient flag (reset every frame) that is used to both enable/disable CCD sweeps for
+		 * the frame, and to retport whether the sweep was run and the results retained. See GetCCDSweepEnabled().
 		 */
 		void SetCCDSweepEnabled(const bool bCCDSweepEnabled)
 		{
