Single selection for ALPN setup

#rnx

diff --git a/Engine/Source/Runtime/Experimental/IoStore/HttpClient/Private/Peer.inl b/Engine/Source/Runtime/Experimental/IoStore/HttpClient/Private/Peer.inl
--- a/Engine/Source/Runtime/Experimental/IoStore/HttpClient/Private/Peer.inl
+++ b/Engine/Source/Runtime/Experimental/IoStore/HttpClient/Private/Peer.inl
@@ -252,13 +252,18 @@ static void Ssl_AssociatePeer(SSL* Ssl, FPeer* Peer)
 }
 
 ////////////////////////////////////////////////////////////////////////////////
-static void Ssl_SetupAlpn(SSL* Ssl, int32 /*MaxHttpVersion*/)
+static void Ssl_SetupAlpn(SSL* Ssl, int32 HttpVersion)
 {
-	static const uint8 AlpnProtos[] =
-		"\x08" "http/1.1"
-	//	"\x02" "h2"
-		;
-	SSL_set_alpn_protos(Ssl, AlpnProtos, sizeof(AlpnProtos) - 1);
+	FAnsiStringView AlpnProtos;
+	switch (HttpVersion)
+	{
+	case 1:		AlpnProtos = "\x08" "http/1.1";
+//	case 2:		AlpnProtos = "\x02" "h2";
+	default:	break;
+	}
+
+	check(!AlpnProtos.IsEmpty());
+	SSL_set_alpn_protos(Ssl, (uint8*)(AlpnProtos.GetData()), int32(AlpnProtos.Len()));
 }
 
 ////////////////////////////////////////////////////////////////////////////////
