Use the lowest-common-denominator

#rnx

diff --git a/Engine/Source/Runtime/Experimental/IoStore/HttpClient/Private/TransactionOne.inl b/Engine/Source/Runtime/Experimental/IoStore/HttpClient/Private/TransactionOne.inl
--- a/Engine/Source/Runtime/Experimental/IoStore/HttpClient/Private/TransactionOne.inl
+++ b/Engine/Source/Runtime/Experimental/IoStore/HttpClient/Private/TransactionOne.inl
@@ -195,7 +195,7 @@ public:
 	void			Begin(FAnsiStringView Method, FAnsiStringView Path);
 	void			AddHeader(FAnsiStringView Key, FAnsiStringView Value);
 	FTransactId		End();
-	FOutcome		TrySendRequest(FHttpPeer& Peer);
+	FOutcome		TrySendRequest(FTlsPeer& Peer);
 
 protected:
 	bool			IsHeadMethod() const;
@@ -252,7 +252,7 @@ FRequest& FRequest::operator << (FAnsiStringView Value)
 }
 
 ////////////////////////////////////////////////////////////////////////////////
-FOutcome FRequest::TrySendRequest(FHttpPeer& Peer)
+FOutcome FRequest::TrySendRequest(FTlsPeer& Peer)
 {
 	const char* SendData = GetData();
 	int32 SendSize = GetSize();
@@ -287,7 +287,7 @@ class FStatusHeaders
 public:
 	using FRequest::FRequest;
 
-	FOutcome		TryRecvResponse(FHttpPeer& Peer);
+	FOutcome		TryRecvResponse(FTlsPeer& Peer);
 	bool			IsKeepAlive() const;
 	uint32			GetStatusCode() const;
 	FAnsiStringView	GetStatusMessage() const;
@@ -370,7 +370,7 @@ const char* FStatusHeaders::GetMessageRight() const
 }
 
 ////////////////////////////////////////////////////////////////////////////////
-FOutcome FStatusHeaders::TryRecvResponse(FHttpPeer& Peer)
+FOutcome FStatusHeaders::TryRecvResponse(FTlsPeer& Peer)
 {
 	auto FindMessageTerminal = [] (const char* Cursor, int32 Length)
 	{
@@ -526,17 +526,17 @@ public:
 	using FStatusHeaders::FStatusHeaders;
 
 	int64				GetRemaining() const;
-	FOutcome			TryRecv(FMutableMemoryView Dest, FHttpPeer& Peer);
+	FOutcome			TryRecv(FMutableMemoryView Dest, FTlsPeer& Peer);
 
 private:
 	enum class EState : uint8 { Init, Recv, Prologue, Chunk, Epilogue, Done };
 
-	FOutcome			Gather(FMutableMemoryView Dest, FHttpPeer& Peer);
-	FOutcome			Init(FMutableMemoryView Dest, FHttpPeer& Peer);
-	FOutcome			Recv(FMutableMemoryView Dest, FHttpPeer& Peer);
-	FOutcome			Prologue(FMutableMemoryView Dest, FHttpPeer& Peer);
-	FOutcome			Chunk(FMutableMemoryView Dest, FHttpPeer& Peer);
-	FOutcome			Epilogue(FMutableMemoryView Dest, FHttpPeer& Peer);
+	FOutcome			Gather(FMutableMemoryView Dest, FTlsPeer& Peer);
+	FOutcome			Init(FMutableMemoryView Dest, FTlsPeer& Peer);
+	FOutcome			Recv(FMutableMemoryView Dest, FTlsPeer& Peer);
+	FOutcome			Prologue(FMutableMemoryView Dest, FTlsPeer& Peer);
+	FOutcome			Chunk(FMutableMemoryView Dest, FTlsPeer& Peer);
+	FOutcome			Epilogue(FMutableMemoryView Dest, FTlsPeer& Peer);
 	FMemoryView			Overspill;
 	FMutableMemoryView	Scratch;
 	int64				Remaining = 0;
@@ -557,7 +557,7 @@ int64 FBody::GetRemaining() const
 }
 
 ////////////////////////////////////////////////////////////////////////////////
-FOutcome FBody::TryRecv(FMutableMemoryView Dest, FHttpPeer& Peer)
+FOutcome FBody::TryRecv(FMutableMemoryView Dest, FTlsPeer& Peer)
 {
 	if (Dest.GetSize() == 0)
 	{
@@ -583,7 +583,7 @@ FOutcome FBody::TryRecv(FMutableMemoryView Dest, FHttpPeer& Peer)
 }
 
 ////////////////////////////////////////////////////////////////////////////////
-FOutcome FBody::Gather(FMutableMemoryView Dest, FHttpPeer& Peer)
+FOutcome FBody::Gather(FMutableMemoryView Dest, FTlsPeer& Peer)
 {
 	FOutcome Outcome = FOutcome::None();
 
@@ -620,7 +620,7 @@ FOutcome FBody::Gather(FMutableMemoryView Dest, FHttpPeer& Peer)
 }
 
 ////////////////////////////////////////////////////////////////////////////////
-FOutcome FBody::Init(FMutableMemoryView Dest, FHttpPeer& Peer)
+FOutcome FBody::Init(FMutableMemoryView Dest, FTlsPeer& Peer)
 {
 	const char* DataEnd = GetData() + GetSize();
 	const char* MessageEnd = GetMessageRight();
@@ -654,7 +654,7 @@ FOutcome FBody::Init(FMutableMemoryView Dest, FHttpPeer& Peer)
 }
 
 ////////////////////////////////////////////////////////////////////////////////
-FOutcome FBody::Recv(FMutableMemoryView Dest, FHttpPeer& Peer)
+FOutcome FBody::Recv(FMutableMemoryView Dest, FTlsPeer& Peer)
 {
 	int64 OverspillSize = Overspill.GetSize();
 	if (OverspillSize != 0)
@@ -691,7 +691,7 @@ FOutcome FBody::Recv(FMutableMemoryView Dest, FHttpPeer& Peer)
 }
 
 ////////////////////////////////////////////////////////////////////////////////
-FOutcome FBody::Prologue(FMutableMemoryView Dest, FHttpPeer& Peer)
+FOutcome FBody::Prologue(FMutableMemoryView Dest, FTlsPeer& Peer)
 {
 	auto FindHeader = [] (const char* Cursor, int32 Size)
 	{
@@ -799,7 +799,7 @@ FOutcome FBody::Prologue(FMutableMemoryView Dest, FHttpPeer& Peer)
 }
 
 ////////////////////////////////////////////////////////////////////////////////
-FOutcome FBody::Chunk(FMutableMemoryView Dest, FHttpPeer& Peer)
+FOutcome FBody::Chunk(FMutableMemoryView Dest, FTlsPeer& Peer)
 {
 	check(Remaining != 0);
 
@@ -819,7 +819,7 @@ FOutcome FBody::Chunk(FMutableMemoryView Dest, FHttpPeer& Peer)
 }
 
 ////////////////////////////////////////////////////////////////////////////////
-FOutcome FBody::Epilogue(FMutableMemoryView Dest, FHttpPeer& Peer)
+FOutcome FBody::Epilogue(FMutableMemoryView Dest, FTlsPeer& Peer)
 {
 	auto* Cursor = (uint8*)(Scratch.GetData());
 
