Easiest route to enact a file split.

#rnx

diff --git a/Engine/Source/Runtime/Experimental/IoStore/HttpClient/Private/Http.cpp b/Engine/Source/Runtime/Experimental/IoStore/HttpClient/Private/Http.cpp
--- a/Engine/Source/Runtime/Experimental/IoStore/HttpClient/Private/Http.cpp
+++ b/Engine/Source/Runtime/Experimental/IoStore/HttpClient/Private/Http.cpp
@@ -24,7 +24,6 @@
 #include "Misc.inl"
 #include "Socket.inl"
 #include "Peer.inl"
-#include "TransactionOne.inl"
 #include "Socks.inl"
 #include "ConnectionPool.inl"
 #include "Activity.inl"
diff --git a/Engine/Source/Runtime/Experimental/IoStore/HttpClient/Private/Peer.inl b/Engine/Source/Runtime/Experimental/IoStore/HttpClient/Private/Peer.inl
--- a/Engine/Source/Runtime/Experimental/IoStore/HttpClient/Private/Peer.inl
+++ b/Engine/Source/Runtime/Experimental/IoStore/HttpClient/Private/Peer.inl
@@ -503,10 +503,15 @@ FOutcome FTlsPeer::Recv(char* Out, int32 MaxSize)
 	return Ssl_Read(Ssl, Out, MaxSize);
 }
 
+} // namespace UE::IoStore::HTTP
+
 
 
 ////////////////////////////////////////////////////////////////////////////////
-class FTransaction* CreateTransactOne(FBuffer&);
+#include "TransactionOne.inl"
+
+namespace UE::IoStore::HTTP
+{
 
 ////////////////////////////////////////////////////////////////////////////////
 class FHttpPeer
