Target specific version instead of decaying

#rnx

diff --git a/Engine/Source/Runtime/Experimental/IoStore/HttpClient/Private/Peer.inl b/Engine/Source/Runtime/Experimental/IoStore/HttpClient/Private/Peer.inl
--- a/Engine/Source/Runtime/Experimental/IoStore/HttpClient/Private/Peer.inl
+++ b/Engine/Source/Runtime/Experimental/IoStore/HttpClient/Private/Peer.inl
@@ -536,17 +536,17 @@ public:
 	FOutcome		GetPendingTransactId();
 
 private:
-	void			AssignProto();
-	int32			Proto = 0;
+	uint8			Proto;
 };
 
 ////////////////////////////////////////////////////////////////////////////////
 FHttpPeer::FHttpPeer(FParams&& Params)
 : FTlsPeer(MoveTemp(Params.Socket), Params.Certs, Params.HostName)
+, Proto(1)
 {
 	if (Ssl != nullptr)
 	{
-		Ssl_SetupAlpn(Ssl, 1);
+		Ssl_SetupAlpn(Ssl, Proto);
 	}
 }
 
@@ -561,29 +561,23 @@ FHttpPeer& FHttpPeer::operator = (FHttpPeer&& Rhs)
 ////////////////////////////////////////////////////////////////////////////////
 FOutcome FHttpPeer::Handshake()
 {
-	FOutcome Outcome = FTlsPeer::Handshake();
-	if (Outcome.IsOk())
+	if (Ssl != nullptr)
 	{
-		AssignProto();
+		return FOutcome::Ok();
 	}
 
-	return Outcome;
-}
-
-////////////////////////////////////////////////////////////////////////////////
-void FHttpPeer::AssignProto()
-{
-	Proto = 1;
-
-	if (Ssl == nullptr)
+	if (FOutcome Outcome = FTlsPeer::Handshake(); !Outcome.IsOk())
 	{
-		return;
+		return Outcome;
 	}
 
-	if (int32 Negotiated = Ssl_GetProtocolVersion(Ssl); Negotiated > 0)
+	int32 Negotiated = Ssl_GetProtocolVersion(Ssl);
+	if (Negotiated != Proto)
 	{
-		Proto = Negotiated;
+		return FOutcome::Error("Failed to negotiate protocol version");
 	}
+
+	return FOutcome::Ok();
 }
 
 ////////////////////////////////////////////////////////////////////////////////
