Initial surfacing of version selection

#rnx

diff --git a/Engine/Source/Runtime/Experimental/IoStore/HttpClient/Private/Loop.inl b/Engine/Source/Runtime/Experimental/IoStore/HttpClient/Private/Loop.inl
--- a/Engine/Source/Runtime/Experimental/IoStore/HttpClient/Private/Loop.inl
+++ b/Engine/Source/Runtime/Experimental/IoStore/HttpClient/Private/Loop.inl
@@ -675,6 +675,7 @@ void FPeerGroup::TickSend(FTickState& State, FHost& Host, FPoller& Poller)
 			.Socket = MoveTemp(Socket),
 			.Certs = Host.GetVerifyCert(),
 			.HostName = Host.GetHostName().GetData(),
+			.Version = 1,
 		};
 		Peer = FHttpPeer(MoveTemp(Params));
 	}
diff --git a/Engine/Source/Runtime/Experimental/IoStore/HttpClient/Private/Peer.inl b/Engine/Source/Runtime/Experimental/IoStore/HttpClient/Private/Peer.inl
--- a/Engine/Source/Runtime/Experimental/IoStore/HttpClient/Private/Peer.inl
+++ b/Engine/Source/Runtime/Experimental/IoStore/HttpClient/Private/Peer.inl
@@ -523,6 +523,7 @@ public:
 		FSocket			Socket;
 		FCertRootsRef 	Certs = ECertRootsRefType::None;
 		const char* 	HostName = nullptr;
+		uint8			Version = 1;
 	};
 
 					FHttpPeer() = default;
@@ -542,7 +543,7 @@ private:
 ////////////////////////////////////////////////////////////////////////////////
 FHttpPeer::FHttpPeer(FParams&& Params)
 : FTlsPeer(MoveTemp(Params.Socket), Params.Certs, Params.HostName)
-, Proto(1)
+, Proto(Params.Version)
 {
 	if (Ssl != nullptr)
 	{
@@ -563,7 +564,12 @@ FOutcome FHttpPeer::Handshake()
 {
 	if (Ssl != nullptr)
 	{
-		return FOutcome::Ok();
+		if (Proto == 1)
+		{
+			return FOutcome::Ok();
+		}
+
+		return FOutcome::Error("Unsupported version");
 	}
 
 	if (FOutcome Outcome = FTlsPeer::Handshake(); !Outcome.IsOk())
