Better handshake logic that is better suited to the future.

#rnx

diff --git a/Engine/Source/Runtime/Experimental/IoStore/HttpClient/Private/Peer.inl b/Engine/Source/Runtime/Experimental/IoStore/HttpClient/Private/Peer.inl
--- a/Engine/Source/Runtime/Experimental/IoStore/HttpClient/Private/Peer.inl
+++ b/Engine/Source/Runtime/Experimental/IoStore/HttpClient/Private/Peer.inl
@@ -564,26 +564,24 @@ FOutcome FHttpPeer::Handshake()
 {
 	if (Ssl != nullptr)
 	{
-		if (Proto == 1)
+		if (FOutcome Outcome = FTlsPeer::Handshake(); !Outcome.IsOk())
 		{
-			return FOutcome::Ok();
+			return Outcome;
 		}
 
-		return FOutcome::Error("Unsupported version");
-	}
-
-	if (FOutcome Outcome = FTlsPeer::Handshake(); !Outcome.IsOk())
-	{
-		return Outcome;
+		int32 Negotiated = Ssl_GetProtocolVersion(Ssl);
+		if (Negotiated >= 0 && Negotiated != Proto)
+		{
+			return FOutcome::Error("Failed to negotiate protocol version");
+		}
 	}
 
-	int32 Negotiated = Ssl_GetProtocolVersion(Ssl);
-	if (Negotiated != Proto)
+	if (Proto == 1)
 	{
-		return FOutcome::Error("Failed to negotiate protocol version");
+		return FOutcome::Ok();
 	}
 
-	return FOutcome::Ok();
+	return FOutcome::Error("Unsupported version");
 }
 
 ////////////////////////////////////////////////////////////////////////////////
