Differentiate when alpn is present or not.

#rnx

diff --git a/Engine/Source/Runtime/Experimental/IoStore/HttpClient/Private/Peer.inl b/Engine/Source/Runtime/Experimental/IoStore/HttpClient/Private/Peer.inl
--- a/Engine/Source/Runtime/Experimental/IoStore/HttpClient/Private/Peer.inl
+++ b/Engine/Source/Runtime/Experimental/IoStore/HttpClient/Private/Peer.inl
@@ -279,6 +279,8 @@ static int32 Ssl_GetProtocolVersion(SSL* Ssl)
 		return Proto;
 	}
 
+	Proto = 0;
+
 	FAnsiStringView Needle(AlpnProto, AlpnProtoLen);
 	FAnsiStringView Candidates[] = {
 		"http/1.1",
