Pre-integration waking for teleporting and partial sleeping

Execution order:
1) Evolution: Register particle state changes and explicitly set transform
2) IslandManager::UpdateExplicit: (pre-integration, pre-CD): Teleport the particle out of its existing island, perform level-based waking (particle velocity = 0)
3) PBDCollisionConstraints::AddConstraintsToGraph (post-CD): In case of penetration in new collision constraints, estimate the velocity which would have let to the penetration caused by the teleport
4) IslandManager::FinalizeIslands (pre-solve, post-CD): Teleport the particle into its new island, estimated velocity allows for propagation downward in case of momentum transfer.

diff --git a/Engine/Source/Programs/HeadlessChaos/Private/HeadlessChaosTestGraphEvolution.cpp b/Engine/Source/Programs/HeadlessChaos/Private/HeadlessChaosTestGraphEvolution.cpp
--- a/Engine/Source/Programs/HeadlessChaos/Private/HeadlessChaosTestGraphEvolution.cpp
+++ b/Engine/Source/Programs/HeadlessChaos/Private/HeadlessChaosTestGraphEvolution.cpp
@@ -1494,7 +1494,12 @@ namespace ChaosTest
 	// Teleporting means explicitly updating the transform of the particle.
 	TEST_P(GraphEvolutionTests, TestConstraintGraph_TeleportSleeping)
 	{
-		FGraphEvolutionTest Test(3, GetParam());
+		IConsoleVariable* CVarPartialWakePreIntegration = IConsoleManager::Get().FindConsoleVariable(TEXT("p.Chaos.Solver.Sleep.PartialWakePreIntegration"), false);
+		check(CVarPartialWakePreIntegration && CVarPartialWakePreIntegration->IsVariableBool());
+		const bool bPartialWakePreIntegrationDefault = CVarPartialWakePreIntegration->GetBool();
+		CVarPartialWakePreIntegration->Set(true);
+
+		FGraphEvolutionTest Test(5, GetParam());
 		Test.MakeChain();
 		Test.Evolution.SetParticleObjectState(Test.ParticleHandles[0], EObjectStateType::Kinematic);
 
@@ -1506,29 +1511,59 @@ namespace ChaosTest
 		EXPECT_TRUE(Test.ParticleHandles[0]->IsInConstraintGraph());
 		EXPECT_TRUE(Test.ParticleHandles[1]->IsInConstraintGraph());
 		EXPECT_TRUE(Test.ParticleHandles[2]->IsInConstraintGraph());
+		EXPECT_TRUE(Test.ParticleHandles[3]->IsInConstraintGraph());
+		EXPECT_TRUE(Test.ParticleHandles[4]->IsInConstraintGraph());
 		EXPECT_TRUE(FConstGenericParticleHandle(Test.ParticleHandles[0])->IsKinematic());
 		EXPECT_TRUE(FConstGenericParticleHandle(Test.ParticleHandles[1])->Sleeping());
 		EXPECT_TRUE(FConstGenericParticleHandle(Test.ParticleHandles[2])->Sleeping());
+		EXPECT_TRUE(FConstGenericParticleHandle(Test.ParticleHandles[3])->Sleeping());
+		EXPECT_TRUE(FConstGenericParticleHandle(Test.ParticleHandles[4])->Sleeping());
 
 		EXPECT_TRUE(Test.ConstraintHandles[0]->IsInConstraintGraph());
 		EXPECT_TRUE(Test.ConstraintHandles[1]->IsInConstraintGraph());
+		EXPECT_TRUE(Test.ConstraintHandles[2]->IsInConstraintGraph());
+		EXPECT_TRUE(Test.ConstraintHandles[3]->IsInConstraintGraph());
 		EXPECT_TRUE(Test.ConstraintHandles[0]->IsSleeping());
 		EXPECT_TRUE(Test.ConstraintHandles[1]->IsSleeping());
+		EXPECT_TRUE(Test.ConstraintHandles[2]->IsSleeping());
+		EXPECT_TRUE(Test.ConstraintHandles[3]->IsSleeping());
 
 		// Teleport a particle by setting its transform explicitly
 		const bool bIsTeleport = true;
 		const bool bWakeUp = true;
-		Test.Evolution.SetParticleTransform(Test.ParticleHandles[2], FVec3(0, 3, 0), FRotation3(), bIsTeleport, bWakeUp);
+		Test.Evolution.SetParticleObjectState(Test.ParticleHandles[4], EObjectStateType::Dynamic);
+		Test.Evolution.SetParticleTransform(Test.ParticleHandles[4], FVec3(0, 3, 0), FRotation3(), bIsTeleport, bWakeUp);
 
 		Test.Advance();
 
 		// Should wake up the entire island
 		EXPECT_EQ(Test.IslandManager->GetNumIslands(), 1);
 		EXPECT_FALSE(Test.IslandManager->GetIsland(0)->IsSleeping());
-		EXPECT_FALSE(FConstGenericParticleHandle(Test.ParticleHandles[1])->Sleeping());
-		EXPECT_FALSE(FConstGenericParticleHandle(Test.ParticleHandles[2])->Sleeping());
-		EXPECT_FALSE(Test.ConstraintHandles[0]->IsSleeping());
-		EXPECT_FALSE(Test.ConstraintHandles[1]->IsSleeping());
+
+		if (!GetParam())
+		{
+			EXPECT_FALSE(FConstGenericParticleHandle(Test.ParticleHandles[1])->Sleeping());
+			EXPECT_FALSE(FConstGenericParticleHandle(Test.ParticleHandles[2])->Sleeping());
+			EXPECT_FALSE(FConstGenericParticleHandle(Test.ParticleHandles[3])->Sleeping());
+			EXPECT_FALSE(FConstGenericParticleHandle(Test.ParticleHandles[4])->Sleeping());
+			EXPECT_FALSE(Test.ConstraintHandles[0]->IsSleeping());
+			EXPECT_FALSE(Test.ConstraintHandles[1]->IsSleeping());
+			EXPECT_FALSE(Test.ConstraintHandles[2]->IsSleeping());
+			EXPECT_FALSE(Test.ConstraintHandles[3]->IsSleeping());
+		}
+		else // partial island sleep
+		{
+			EXPECT_TRUE(FConstGenericParticleHandle(Test.ParticleHandles[1])->Sleeping());
+			EXPECT_TRUE(FConstGenericParticleHandle(Test.ParticleHandles[2])->Sleeping());
+			EXPECT_FALSE(FConstGenericParticleHandle(Test.ParticleHandles[3])->Sleeping());
+			EXPECT_FALSE(FConstGenericParticleHandle(Test.ParticleHandles[4])->Sleeping());
+			EXPECT_TRUE(Test.ConstraintHandles[0]->IsSleeping());
+			EXPECT_TRUE(Test.ConstraintHandles[1]->IsSleeping());
+			EXPECT_FALSE(Test.ConstraintHandles[2]->IsSleeping());
+			EXPECT_FALSE(Test.ConstraintHandles[3]->IsSleeping());
+		}
+
+		CVarPartialWakePreIntegration->Set(bPartialWakePreIntegrationDefault);
 	}
 
 	// Test validating that a sleeping island wakes up when adding an impulse, acceleration, etc. to one particle.
@@ -2263,9 +2298,9 @@ namespace ChaosTest
 			Test.MakeChain();
 
 			// Number of steps is a few less than the sleep counter threshold
-			IConsoleVariable* CVarPartialSleeping = IConsoleManager::Get().FindConsoleVariable(TEXT("p.Chaos.Solver.Sleep.Defaults.SleepCounterThreshold"), false);
-			check(CVarPartialSleeping && CVarPartialSleeping->IsVariableInt());			
-			const int32 SleepCounterThreshold = CVarPartialSleeping->GetInt();
+			IConsoleVariable* CVarSleepCounterThreshold = IConsoleManager::Get().FindConsoleVariable(TEXT("p.Chaos.Solver.Sleep.Defaults.SleepCounterThreshold"), false);
+			check(CVarSleepCounterThreshold && CVarSleepCounterThreshold->IsVariableInt());
+			const int32 SleepCounterThreshold = CVarSleepCounterThreshold->GetInt();
 			const int32 Steps = SleepCounterThreshold - 2;
 
 			for (int32 I = 0; I < Steps; ++I)
@@ -2313,9 +2348,9 @@ namespace ChaosTest
 			Test.MakeChain();
 
 			// Number of steps is the sleep counter threshold
-			IConsoleVariable* CVarPartialSleeping = IConsoleManager::Get().FindConsoleVariable(TEXT("p.Chaos.Solver.Sleep.Defaults.SleepCounterThreshold"), false);
-			check(CVarPartialSleeping && CVarPartialSleeping->IsVariableInt());
-			const int32 SleepCounterThreshold = CVarPartialSleeping->GetInt();
+			IConsoleVariable* CVarSleepCounterThreshold = IConsoleManager::Get().FindConsoleVariable(TEXT("p.Chaos.Solver.Sleep.Defaults.SleepCounterThreshold"), false);
+			check(CVarSleepCounterThreshold && CVarSleepCounterThreshold->IsVariableInt());
+			const int32 SleepCounterThreshold = CVarSleepCounterThreshold->GetInt();
 			const int32 Steps = SleepCounterThreshold;
 
 			for (int32 I = 0; I < Steps; ++I)
diff --git a/Engine/Source/Programs/HeadlessChaos/Private/HeadlessChaosTestSleeping.cpp b/Engine/Source/Programs/HeadlessChaos/Private/HeadlessChaosTestSleeping.cpp
--- a/Engine/Source/Programs/HeadlessChaos/Private/HeadlessChaosTestSleeping.cpp
+++ b/Engine/Source/Programs/HeadlessChaos/Private/HeadlessChaosTestSleeping.cpp
@@ -169,14 +169,14 @@ namespace ChaosTest
 			return Joint;
 		}
 
-		void MakeStackOfCubes(const int32 Num, const FReal Size = 200)
+		void MakeStackOfCubes(const int32 Num, const FReal Size = 200, const FVec3 Pos = FVec3(0))
 		{
 			MakeFloor();
 			const FReal InitialHeight = Size / 2.0f;
 			for (int32 Id = 0; Id < Num; ++Id)
 			{
 				const FReal Z = InitialHeight + Size * Id;
-				MakeCube(FVec3(0, 0, Z), Size);
+				MakeCube(Pos + FVec3(0, 0, Z), Size);
 			}
 		}
 
@@ -521,6 +521,11 @@ namespace ChaosTest
 		const FReal Dimension = 200;
 		const bool bPartialSleepEnabled = GetParam();
 
+		IConsoleVariable* CVarPartialWakePreIntegration = IConsoleManager::Get().FindConsoleVariable(TEXT("p.Chaos.Solver.Sleep.PartialWakePreIntegration"), false);
+		check(CVarPartialWakePreIntegration && CVarPartialWakePreIntegration->IsVariableBool());
+		const bool bPartialWakePreIntegrationDefault = CVarPartialWakePreIntegration->GetBool();
+		CVarPartialWakePreIntegration->Set(true);
+
 		// Create a stack of cubes
 		FSleepingTest Test(bPartialSleepEnabled);
 		Test.MakeStackOfCubes(NumOfStackedObjects, /*Size=*/Dimension);
@@ -549,14 +554,94 @@ namespace ChaosTest
 		}
 		else
 		{
-			// @todo (Chaos): The entire stack wakes up. We want to change that.
-			EXPECT_TRUE(Test.ParticleHandles[0]->ObjectState() == EObjectStateType::Dynamic);
+			// With pre-integration partial waking, only one particle below the teleported one wakes up.
+			EXPECT_TRUE(Test.ParticleHandles[0]->ObjectState() == EObjectStateType::Sleeping);
 			EXPECT_TRUE(Test.ParticleHandles[1]->ObjectState() == EObjectStateType::Dynamic);
 			EXPECT_TRUE(Test.ParticleHandles[2]->ObjectState() == EObjectStateType::Dynamic);
 			EXPECT_TRUE(Test.ParticleHandles[3]->ObjectState() == EObjectStateType::Dynamic);
 			EXPECT_TRUE(Test.ParticleHandles[4]->ObjectState() == EObjectStateType::Dynamic);
 		}
 
+		CVarPartialWakePreIntegration->Set(bPartialWakePreIntegrationDefault);
+	}
+
+	TEST_P(SleepingTests, TeleportCubeFromOneStackToAnother)
+	{
+		const int32 NumOfStackedObjects = 5;
+		const FReal Dimension = 200;
+		const bool bPartialSleepEnabled = true;
+		const FRealSingle MomentumPropagation = GetParam() ? 0.5f : 0.0f;
+
+		IConsoleVariable* CVarMomentumPropagation = IConsoleManager::Get().FindConsoleVariable(TEXT("p.Chaos.Solver.Sleep.MomentumPropagation"), false);
+		check(CVarMomentumPropagation && CVarMomentumPropagation->IsVariableFloat());
+		const FRealSingle MomentumPropagationDefault = CVarMomentumPropagation->GetFloat();
+		CVarMomentumPropagation->Set(MomentumPropagation);
+
+		IConsoleVariable* CVarPartialWakePreIntegration = IConsoleManager::Get().FindConsoleVariable(TEXT("p.Chaos.Solver.Sleep.PartialWakePreIntegration"), false);
+		check(CVarPartialWakePreIntegration && CVarPartialWakePreIntegration->IsVariableBool());
+		const bool bPartialWakePreIntegrationDefault = CVarPartialWakePreIntegration->GetBool();
+		CVarPartialWakePreIntegration->Set(true);
+
+		// Create two identical stacks of cubes which don't touch
+		FSleepingTest Test(bPartialSleepEnabled);
+		Test.MakeStackOfCubes(NumOfStackedObjects, /*Size=*/Dimension);
+		Test.MakeStackOfCubes(NumOfStackedObjects, /*Size=*/Dimension, /*Pos=*/FVec3(0, 500, 0));
+
+		// Simulate until all particles are sleeping
+		Test.AdvanceUntilSleeping();
+
+		// Check that everything is sleeping
+		for (FRigidBodyHandle_External* ParticleHandle : Test.ParticleHandles)
+		{
+			EXPECT_TRUE(ParticleHandle->ObjectState() == EObjectStateType::Sleeping);
+		}
+
+		// Teleport one of the cubes out of the stack into the other stack
+		const FVec3 NewX = Test.ParticleHandles[4]->GetX() + FVec3(0, 500, 0);
+		Test.ParticleHandles[4]->SetX(NewX, true);
+
+		Test.Advance();
+
+		if (MomentumPropagation)
+		{
+			// Stack we teleport out of
+			// Only one particle above the one which is teleported will wake up (level-based waking).
+			// Momentum transfer does not affect the stack we teleport out of.
+			EXPECT_TRUE(Test.ParticleHandles[0]->ObjectState() == EObjectStateType::Sleeping);
+			EXPECT_TRUE(Test.ParticleHandles[1]->ObjectState() == EObjectStateType::Sleeping);
+			EXPECT_TRUE(Test.ParticleHandles[2]->ObjectState() == EObjectStateType::Sleeping);
+			EXPECT_TRUE(Test.ParticleHandles[3]->ObjectState() == EObjectStateType::Dynamic);
+			EXPECT_TRUE(Test.ParticleHandles[4]->ObjectState() == EObjectStateType::Dynamic);
+
+			// Stack we teleport into
+			// Everything wakes up if momentum transfer is enabled
+			EXPECT_TRUE(Test.ParticleHandles[5]->ObjectState() == EObjectStateType::Dynamic);
+			EXPECT_TRUE(Test.ParticleHandles[6]->ObjectState() == EObjectStateType::Dynamic);
+			EXPECT_TRUE(Test.ParticleHandles[7]->ObjectState() == EObjectStateType::Dynamic);
+			EXPECT_TRUE(Test.ParticleHandles[8]->ObjectState() == EObjectStateType::Dynamic);
+			EXPECT_TRUE(Test.ParticleHandles[9]->ObjectState() == EObjectStateType::Dynamic);
+		}
+		else
+		{
+			// Stack we teleport out of
+			// Only one particle above the one which is teleported will wake up (level-based waking).
+			EXPECT_TRUE(Test.ParticleHandles[0]->ObjectState() == EObjectStateType::Sleeping);
+			EXPECT_TRUE(Test.ParticleHandles[1]->ObjectState() == EObjectStateType::Sleeping);
+			EXPECT_TRUE(Test.ParticleHandles[2]->ObjectState() == EObjectStateType::Sleeping);
+			EXPECT_TRUE(Test.ParticleHandles[3]->ObjectState() == EObjectStateType::Dynamic);
+			EXPECT_TRUE(Test.ParticleHandles[4]->ObjectState() == EObjectStateType::Dynamic);
+
+			// Stack we teleport into
+			// Only one particle above the one which is teleported into the stack will wake up (level-based waking).
+			EXPECT_TRUE(Test.ParticleHandles[5]->ObjectState() == EObjectStateType::Sleeping);
+			EXPECT_TRUE(Test.ParticleHandles[6]->ObjectState() == EObjectStateType::Sleeping);
+			EXPECT_TRUE(Test.ParticleHandles[7]->ObjectState() == EObjectStateType::Sleeping);
+			EXPECT_TRUE(Test.ParticleHandles[8]->ObjectState() == EObjectStateType::Dynamic);
+			EXPECT_TRUE(Test.ParticleHandles[9]->ObjectState() == EObjectStateType::Dynamic);
+		}
+
+		CVarPartialWakePreIntegration->Set(bPartialWakePreIntegrationDefault);
+		CVarMomentumPropagation->Set(MomentumPropagationDefault);
 	}
 
 	// Uses a traditional brick wall where every 2nd row is shifted (see SleepingTests::MakeBrickWall)
diff --git a/Engine/Source/Runtime/Experimental/Chaos/Private/Chaos/Island/IslandManager.cpp b/Engine/Source/Runtime/Experimental/Chaos/Private/Chaos/Island/IslandManager.cpp
--- a/Engine/Source/Runtime/Experimental/Chaos/Private/Chaos/Island/IslandManager.cpp
+++ b/Engine/Source/Runtime/Experimental/Chaos/Private/Chaos/Island/IslandManager.cpp
@@ -81,6 +81,10 @@ namespace Chaos::CVars
 	bool bChaosSolverPartialSleepCollisionConstraintsOnly = false;
 	FAutoConsoleVariableRef CVarChaosSolverPartialSleepCollisionConstraintsOnly(TEXT("p.Chaos.Solver.Sleep.PartialSleepCollisionConstraintsOnly"), bChaosSolverPartialSleepCollisionConstraintsOnly, TEXT("[EXPERIMENTAL] Only collision constraints permit partial sleep, i.e. the sleep state of particles connect by other constraint types must match (partial island sleep only)."));
 
+	/** Cvar to enable partial waking pre-integration (UpdateExplicitSleep). */
+	bool bChaosSolverPartialWakePreIntegration = false;
+	FAutoConsoleVariableRef CVarChaosSolverPartialWakePreIntegration(TEXT("p.Chaos.Solver.Sleep.PartialWakePreIntegration"), bChaosSolverPartialWakePreIntegration, TEXT("[EXPERIMENTAL] Partial wake events also happen pre-integration when the user explicilty sets a particle transform or applies external impulses or forces (partial island sleeping only)."));
+
 	/** Cvar to override the sleep counter threshold if necessary */
 	int32 ChaosSolverCollisionDefaultSleepCounterThreshold = 20;
 	FAutoConsoleVariableRef CVarChaosSolverCollisionDefaultSleepCounterThreshold(TEXT("p.Chaos.Solver.Sleep.Defaults.SleepCounterThreshold"), ChaosSolverCollisionDefaultSleepCounterThreshold, TEXT("Default counter threshold for sleeping.[def:20]"));
@@ -819,7 +823,7 @@ namespace Chaos::Private
 	void FPBDIslandManager::WakeParticleIslands(FGeometryParticleHandle* Particle)
 	{
 		// We want to wake for at least one frame (i.e., prevent re-sleeping this frame)
-		const bool bIsSleepAllowed = false;
+		bool bIsSleepAllowed = false;
 
 		// When we explicitly wake we reset sleep counters etc
 		if (FPBDIslandParticle* Node = GetGraphNode(Particle))
@@ -832,6 +836,13 @@ namespace Chaos::Private
 			// was made dynamic right before calling WakeParticleIsland
 			if (Node->Island != nullptr)
 			{
+				// We permit sleeping for this island if partial island sleeping and pre-integration waking are enabled.
+				// This will allows us to perform partial waking if a particle is moved into the island by updating its transform explicitly.
+				if (CVars::bChaosSolverPartialIslandSleep && CVars::bChaosSolverPartialWakePreIntegration)
+				{
+					bIsSleepAllowed = true;
+				}
+
 				// This is a dynamic particle and we know its island
 				EnqueueIslandCheckSleep(Node->Island, bIsSleepAllowed);
 				UpdateIslandPartialSleepAllowed(Node->Island, Particle);
@@ -872,6 +883,40 @@ namespace Chaos::Private
 		}
 	}
 
+	void FPBDIslandManager::FlagParticleUpdated(const FGeometryParticleHandle* Particle)
+	{
+		if (CVars::bChaosSolverPartialIslandSleep && CVars::bChaosSolverPartialWakePreIntegration)
+		{
+			if (FPBDIslandParticle* Node = GetGraphNode(Particle))
+			{
+				Node->Flags.bWasUpdated = true;
+			}
+		}
+	}
+
+	bool FPBDIslandManager::WasParticleTransformSetExplicitly(const FGeometryParticleHandle* Particle)
+	{
+		if (CVars::bChaosSolverPartialIslandSleep && CVars::bChaosSolverPartialWakePreIntegration)
+		{
+			if (FPBDIslandParticle* Node = GetGraphNode(Particle))
+			{
+				return Node->Flags.bWasTransformSetExplicitly;
+			}
+		}
+		return false;
+	}
+
+	void FPBDIslandManager::FlagParticleTransformSetExplicitly(const FGeometryParticleHandle* Particle)
+	{
+		if (CVars::bChaosSolverPartialIslandSleep && CVars::bChaosSolverPartialWakePreIntegration)
+		{
+			if (FPBDIslandParticle* Node = GetGraphNode(Particle))
+			{
+				Node->Flags.bWasTransformSetExplicitly = true;
+			}
+		}
+	}
+
 	void FPBDIslandManager::AddParticle(FGeometryParticleHandle* Particle)
 	{
 		// Particles get auto-added via AddConstraint (see CreateGraphEdge, GetOrCreateGraphNode)
@@ -1365,14 +1410,19 @@ namespace Chaos::Private
 			}
 			else if (CVars::bChaosSolverPartialIslandSleep && !CVars::bChaosSolverWakeOnParticleStateChanged)
 			{
-				// Register wake-up events for updates in the particles (partial island sleeping only)
-				// Case 1) Switch from non-moving kinematic to non-sleeping dynamic OR
-				// Case 2) Switch dynamic particle from sleeping to not-sleeping
-				const bool bSwitchStaticKinematicToAwakeDynamic = (!bWasDynamic && !bWasMoving) && (bIsDynamic && !bIsSleeping); // Case 1
-				const bool bSwitchSleepingDynamicToAwake = (bIsDynamic && bWasSleeping && !bIsSleeping); // Case 2
+				// Only set bWasUpdated flag if it isn't set yet.
 				const bool bIsSleepAllowed = true;
-				const bool bTriggerWakeUpEvent = (bSwitchStaticKinematicToAwakeDynamic || bSwitchSleepingDynamicToAwake);
-				Node->Flags.bWasUpdated = bTriggerWakeUpEvent ? true : false;
+				if (!Node->Flags.bWasUpdated)
+				{
+					// Register wake-up events for updates in the particles (partial island sleeping only)
+					// Case 1) Switch from non-moving kinematic to non-sleeping dynamic OR
+					// Case 2) Switch dynamic particle from sleeping to not-sleeping
+					const bool bSwitchStaticKinematicToAwakeDynamic = (!bWasDynamic && !bWasMoving) && (bIsDynamic && !bIsSleeping); // Case 1
+					const bool bSwitchSleepingDynamicToAwake = (bIsDynamic && bWasSleeping && !bIsSleeping); // Case 2
+					const bool bTriggerWakeUpEvent = (bSwitchStaticKinematicToAwakeDynamic || bSwitchSleepingDynamicToAwake);
+
+					Node->Flags.bWasUpdated = bTriggerWakeUpEvent ? true : false;
+				}
 				EnqueueIslandCheckSleep(Node->Island, bIsSleepAllowed);
 			}
 			else
@@ -2204,7 +2254,7 @@ namespace Chaos::Private
 				{
 					if (!Island->Flags.bIsSleeping)
 					{
-						ProcessPartialIslandWake(Island);
+						ProcessPartialIslandWake(Island, /*bProcessEdgeBasedWake=*/ true);
 					}
 					SyncConstraintSleepStateToParticles(Island);
 				}
@@ -2213,12 +2263,6 @@ namespace Chaos::Private
 					PropagateIslandSleep(Island);
 				}
 				Island->UpdateSyncState();
-
-				if (CVars::bChaosSolverPartialIslandSleep)
-				{
-					ResetEdgeFlags(Island);
-					ResetNodeFlags(Island);
-				}
 			}
 
 			// Remove islands without constraints. NOTE: must come after PropagateIslandSleep
@@ -2272,6 +2316,12 @@ namespace Chaos::Private
 			Island->ResimFrame = INDEX_NONE;
 		}
 
+		if (CVars::bChaosSolverPartialIslandSleep)
+		{
+			ResetEdgeFlags();
+			ResetNodeFlags();
+		}
+
 		Validate();
 	}
 
@@ -2311,10 +2361,30 @@ namespace Chaos::Private
 					Island->SleepCounter = 0;
 				}
 
-				// NOTE: we only get here if particle sleep state was changed. We need to ensure that
-				// constraint sleep state matches the particle sleep state. Ideally we would only
-				// do this if we know that they don't match, but that's hard to determine.
-				PropagateIslandSleep(Island);
+				if (CVars::bChaosSolverPartialIslandSleep && CVars::bChaosSolverPartialWakePreIntegration && !!Island->Flags.bIsSleepAllowed && !!Island->Flags.bIsPartialSleepAllowed)
+				{
+					if (!Island->Flags.bIsSleeping)
+					{
+						// Pre-CD call, no new edges to process waking for.
+						ProcessPartialIslandWake(Island, /*bProcessEdgeBasedWake=*/ false);
+					}
+					SyncConstraintSleepStateToParticles(Island);
+
+					// We don't want to reset Node->Flags.bWasUpdated and Node->Flags.bWasTransformSetExplicitly since we need to check them in post-integration waking (FinalizeIslands).
+					// Hence, do not call ResetNodeFlags(Island), but reset the Node->Flags.bWasProcessed manually here.
+					ResetEdgeFlags(Island);
+					for (FPBDIslandParticle* Node : Island->Nodes)
+					{
+						Node->Flags.bWasProcessed = false;
+					}
+				}
+				else // For island sleeping to sync the sleep state OR if (!bIsSleepAllowed) to wake up the entire island
+				{
+					// NOTE: we only get here if particle sleep state was changed. We need to ensure that
+					// constraint sleep state matches the particle sleep state. Ideally we would only
+					// do this if we know that they don't match, but that's hard to determine.
+					PropagateIslandSleep(Island);
+				}
 
 				Island->Flags.bCheckSleep = false;
 			}
@@ -2890,7 +2960,9 @@ namespace Chaos::Private
 		WakeUpParticles(WakeUpQueue, /*bResetSleepCounter=*/false);
 	}
 
-	void FPBDIslandManager::ProcessPartialIslandWake(FPBDIsland* Island)
+	// If (bProcessEdgeBasedWake), we first process all updated nodes and triggered wake events, then we process all new edges and wake events.
+	// If (!bProcessEdgeBasedWake), we only process updated nodes and their wake events.
+	void FPBDIslandManager::ProcessPartialIslandWake(FPBDIsland* Island, const bool bProcessEdgeBasedWake)
 	{
 		TArray<FGeometryParticleHandle*> ParticlesToWakeUp; // Particles which will be woken up after processing all partial wake events.
 		TArray<FPBDIslandParticle*> NodeQueue;				// Queue of island nodes to be traversed by partial waking.
@@ -2901,36 +2973,35 @@ namespace Chaos::Private
 			MomentumQueue.Reserve(Island->Nodes.Num());
 		}
 
-		// Phase 1 - Waking triggered by new constraints
-		// Step 1a - Find trigger particles (i.e. particles connected to new constraints with significant motion)
-		for (TArray<FPBDIslandConstraint*>& IslandEdges : Island->ContainerEdges)
+		// Phase 1 - Waking triggered by particle state changes
+		// Step 1a - Find trigger particles (i.e. particles which were just woken up the user) 
+		// Start with node-based waking since it always triggers independent of the particle velocity.
+		for (FPBDIslandParticle* Node : Island->Nodes)
 		{
-			for (FPBDIslandConstraint* IslandEdge : IslandEdges)
+			if (Node->Flags.bWasUpdated && !Node->Flags.bWasProcessed)
 			{
-				if (IslandEdge->Flags.bIsNew)
+				ProcessTriggerParticle(Node, NodeQueue, MomentumQueue, /*bAlwaysTrigger=*/ true);
+				Node->Flags.bWasProcessed = true;
+
+				// Reset this this to pre-integration velocity for nodes whose transforms were set explicitly.
+				// We previously overwrote VSmooth with the penetration velocity (i.e. velocity to cause penetration after explicitly setting the transform).
+				// We also check (bProcessEdgeBasedWake) to make sure we call this function from FinalizeIslands not from UpdateExplicitSleep.
+				if (Node->Flags.bWasTransformSetExplicitly && bProcessEdgeBasedWake)
 				{
-					FPBDIslandParticle* Node0 = IslandEdge->Nodes[0];
-					FPBDIslandParticle* Node1 = IslandEdge->Nodes[1];
-					if (Node0 && !Node0->Flags.bWasProcessed)
+					if (FPBDRigidParticleHandle* Rigid = Node->Particle->CastToRigidParticle())
 					{
-						ProcessTriggerParticle(Node0, NodeQueue, MomentumQueue, /*bAlwaysTrigger=*/ !CVars::bChaosSolverWakeOnMovingOnly);
-						Node0->Flags.bWasProcessed = true;
-					}
-					if (Node1 && !Node1->Flags.bWasProcessed)
-					{
-						ProcessTriggerParticle(Node1, NodeQueue, MomentumQueue, /*bAlwaysTrigger=*/ !CVars::bChaosSolverWakeOnMovingOnly);
-						Node1->Flags.bWasProcessed = true;
+						Rigid->SetVSmooth(Rigid->GetPreV());
 					}
 				}
 			}
 		}
-		const int32 NumNodesConstraintWake = NodeQueue.Num();
+		const int32 NumNodesNodeWake = NodeQueue.Num();
 		ValidateWakeQueues(NodeQueue, MomentumQueue);
 		bool bResetEdgeFlags = false;
 
 		// Step 1b - Propagate wake event
 		// Partial island wake disabled, but at least one trigger particle found? Then, perform a full island wake.
-		if (!CVars::bChaosSolverPartialIslandWake && NumNodesConstraintWake > 0)
+		if (!CVars::bChaosSolverPartialIslandWake && NumNodesNodeWake > 0)
 		{
 			PropagateIslandSleep(Island);
 
@@ -2939,41 +3010,59 @@ namespace Chaos::Private
 		}
 		else // Otherwise, propagate the partial wake event through the constraint graph.
 		{
-			bResetEdgeFlags = PropagatePartialIslandWake(NodeQueue, MomentumQueue, ParticlesToWakeUp, NumNodesConstraintWake, /*bTriggerNodesAlwaysPropagate=*/ false);
+			bResetEdgeFlags = PropagatePartialIslandWake(NodeQueue, MomentumQueue, ParticlesToWakeUp, NumNodesNodeWake, /*bTriggerNodesAlwaysPropagate=*/ true);
 		}
 
-		// Phase 2 - Waking triggered by particle state changes
-		// Step 2a - Find trigger particles (i.e. particles which were just woken up the user) 
 		NodeQueue.Reset();
 		MomentumQueue.Reset();
-		for (FPBDIslandParticle* Node : Island->Nodes)
+
+		if (bProcessEdgeBasedWake)
 		{
-			if (Node->Flags.bWasUpdated)
+			// Phase 2 - Waking triggered by new constraints
+			// Step 2a - Find trigger particles (i.e. particles connected to new constraints with significant motion)
+			// Edge-based waking will only be triggered for nodes that haven't been processed in node-based waking yet.
+			for (TArray<FPBDIslandConstraint*>& IslandEdges : Island->ContainerEdges)
 			{
-				ProcessTriggerParticle(Node, NodeQueue, MomentumQueue, /*bAlwaysTrigger=*/ true);
+				for (FPBDIslandConstraint* IslandEdge : IslandEdges)
+				{
+					if (IslandEdge->Flags.bIsNew)
+					{
+						FPBDIslandParticle* Node0 = IslandEdge->Nodes[0];
+						FPBDIslandParticle* Node1 = IslandEdge->Nodes[1];
+						if (Node0 && !Node0->Flags.bWasProcessed)
+						{
+							ProcessTriggerParticle(Node0, NodeQueue, MomentumQueue, /*bAlwaysTrigger=*/ !CVars::bChaosSolverWakeOnMovingOnly);
+							Node0->Flags.bWasProcessed = true;
+						}
+						if (Node1 && !Node1->Flags.bWasProcessed)
+						{
+							ProcessTriggerParticle(Node1, NodeQueue, MomentumQueue, /*bAlwaysTrigger=*/ !CVars::bChaosSolverWakeOnMovingOnly);
+							Node1->Flags.bWasProcessed = true;
+						}
+					}
+				}
 			}
-		}
-		const int32 NumNodesNodeWake = NodeQueue.Num();
-		ValidateWakeQueues(NodeQueue, MomentumQueue);
-
+			const int32 NumNodesConstraintWake = NodeQueue.Num();
+			ValidateWakeQueues(NodeQueue, MomentumQueue);
 
-		// Step 2b - Propagate wake event
-		// Partial island wake disabled, but at least one trigger particle found? Then, perform a full island wake.
-		if (!CVars::bChaosSolverPartialIslandWake && NumNodesNodeWake > 0)
-		{
-			PropagateIslandSleep(Island);
+			// Step 2b - Propagate wake event
+			// Partial island wake disabled, but at least one trigger particle found? Then, perform a full island wake.
+			if (!CVars::bChaosSolverPartialIslandWake && NumNodesConstraintWake > 0)
+			{
+				PropagateIslandSleep(Island);
 
-			// Particle sleep counter reset is needed to prevent immediate re-sleeping in partial island sleeping.
-			ResetParticleSleepCounter(Island);
-		}
-		else // Otherwise, propagate the partial wake event through the constraint graph.
-		{
-			// Reset edge flags if we previously processed partial wake events.
-			if (bResetEdgeFlags && !NodeQueue.IsEmpty())
+				// Particle sleep counter reset is needed to prevent immediate re-sleeping in partial island sleeping.
+				ResetParticleSleepCounter(Island);
+			}
+			else // Otherwise, propagate the partial wake event through the constraint graph.
 			{
-				ResetEdgeFlags(Island);
+				// Reset edge flags if we previously processed partial wake events.
+				if (bResetEdgeFlags && !NodeQueue.IsEmpty())
+				{
+					ResetEdgeFlags(Island);
+				}
+				PropagatePartialIslandWake(NodeQueue, MomentumQueue, ParticlesToWakeUp, NumNodesConstraintWake, /*bTriggerNodesAlwaysPropagate=*/ false);
 			}
-			PropagatePartialIslandWake(NodeQueue, MomentumQueue, ParticlesToWakeUp, NumNodesNodeWake, /*bTriggerNodesAlwaysPropagate=*/ true);
 		}
 
 		// Phase 3 - Wake up all relevant particles and reset sleep counter
@@ -3063,12 +3152,12 @@ namespace Chaos::Private
 
 		const auto IsAwake = [](const FPBDIslandParticle* Node) -> bool
 			{
-				return Node && !Node->Flags.bIsSleeping && !!Node->Flags.bIsDynamic;
+				return Node && !IsParticleSleeping(Node->GetParticle()) && IsParticleDynamic(Node->GetParticle());
 			};
 
 		const auto IsDynamic = [](const FPBDIslandParticle* Node) -> bool
 			{
-				return Node && !!Node->Flags.bIsDynamic;
+				return Node && IsParticleDynamic(Node->GetParticle());
 			};
 
 		const auto ComputeLinearMomentumSq = [](const FPBDIslandParticle* Node) -> FRealSingle
@@ -3168,12 +3257,12 @@ namespace Chaos::Private
 	{
 		const auto IsAwake = [](const FPBDIslandParticle* Node) -> bool
 			{
-				return Node && !Node->Flags.bIsSleeping && !!Node->Flags.bIsDynamic;
+				return Node && !IsParticleSleeping(Node->GetParticle()) && IsParticleDynamic(Node->GetParticle());
 			};
 
 		const auto IsDynamic = [](const FPBDIslandParticle* Node) -> bool
 			{
-				return Node && !!Node->Flags.bIsDynamic;
+				return Node && IsParticleDynamic(Node->GetParticle());
 			};
 
 		// Stop wake-up event if the node or particle is empty
@@ -3263,12 +3352,17 @@ namespace Chaos::Private
 
 				if (Node0 != nullptr && Node1 != nullptr)
 				{
+					const bool bIsSleeping0 = IsParticleSleeping(Node0->GetParticle());
+					const bool bIsSleeping1 = IsParticleSleeping(Node1->GetParticle());
+					const bool bIsDynamic0 = IsParticleDynamic(Node0->GetParticle());
+					const bool bIsDynamic1 = IsParticleDynamic(Node1->GetParticle());
+
 					// Deactivate constraint if:
 					// Case (a) Both particles are sleeping
 					// Case (b) One particle is sleeping and the other is not dynamic (i.e. kinematic, static, etc.)
-					const bool bAreBothSleeping = !!Node0->Flags.bIsSleeping && !!Node1->Flags.bIsSleeping;
-					const bool bIsSleeping0Kinematic1 = !!Node0->Flags.bIsSleeping && !Node1->Flags.bIsDynamic;
-					const bool bIsSleeping1Kinematic0 = !!Node1->Flags.bIsSleeping && !Node0->Flags.bIsDynamic;
+					const bool bAreBothSleeping = bIsSleeping0 && bIsSleeping1;
+					const bool bIsSleeping0Kinematic1 = bIsSleeping0 && !bIsDynamic1;
+					const bool bIsSleeping1Kinematic0 = bIsSleeping1 && !bIsDynamic0;
 					if (bAreBothSleeping || bIsSleeping0Kinematic1 || bIsSleeping1Kinematic0)
 					{
 						IslandEdge->Flags.bIsSleeping = true;
@@ -3277,8 +3371,8 @@ namespace Chaos::Private
 
 					// Activate the constraint if:
 					// Case (a) at least one particle is awake
-					const bool bIsAwakeDynamic0 = !Node0->Flags.bIsSleeping && !!Node0->Flags.bIsDynamic;
-					const bool bIsAwakeDynamic1 = !Node1->Flags.bIsSleeping && !!Node1->Flags.bIsDynamic;
+					const bool bIsAwakeDynamic0 = !bIsSleeping0 && bIsDynamic0;
+					const bool bIsAwakeDynamic1 = !bIsSleeping1 && bIsDynamic1;
 					if (bIsAwakeDynamic0 || bIsAwakeDynamic1)
 					{
 						Constraint->SetIsSleeping(false);
@@ -3368,9 +3462,30 @@ namespace Chaos::Private
 		{
 			Node->Flags.bWasUpdated = false;
 			Node->Flags.bWasProcessed = false;
+			Node->Flags.bWasTransformSetExplicitly = false;
+		}
+	}
+
+	void FPBDIslandManager::ResetEdgeFlags()
+	{
+		for (FPBDIslandConstraint* Edge : Edges)
+		{
+			Edge->Flags.bIsNew = false;
+			Edge->Flags.bWasProcessed = false;
+		}
+	}
+
+	void FPBDIslandManager::ResetNodeFlags()
+	{
+		for (FPBDIslandParticle* Node : Nodes)
+		{
+			Node->Flags.bWasUpdated = false;
+			Node->Flags.bWasProcessed = false;
+			Node->Flags.bWasTransformSetExplicitly = false;
 		}
 	}
 
+
 	void FPBDIslandManager::ProcessDisable(TFunctionRef<void(FPBDRigidParticleHandle*)> ParticleDisableFunctor)
 	{
 		// @todo(chaos): parallelize
diff --git a/Engine/Source/Runtime/Experimental/Chaos/Private/Chaos/PBDCollisionConstraints.cpp b/Engine/Source/Runtime/Experimental/Chaos/Private/Chaos/PBDCollisionConstraints.cpp
--- a/Engine/Source/Runtime/Experimental/Chaos/Private/Chaos/PBDCollisionConstraints.cpp
+++ b/Engine/Source/Runtime/Experimental/Chaos/Private/Chaos/PBDCollisionConstraints.cpp
@@ -716,6 +716,40 @@ namespace Chaos
 		for (FPBDCollisionConstraintHandle* ConstraintHandle : TempCollisions)
 		{
 			IslandManager.AddConstraint(GetContainerId(), ConstraintHandle, ConstraintHandle->GetConstrainedParticles());
+
+			// Check for nodes whose transform was set explicitly, estimate the wake-up momentum with new constraints
+			// Update the smooth velocity to an estimate of the motion which would have lead to the penetration detected.
+			const auto UpdateSmoothVelocity = [](FGeometryParticleHandle* Particle, const FPBDCollisionConstraint& Contact)  -> void
+				{
+					if (FPBDRigidParticleHandle* Rigid = Particle->CastToRigidParticle())
+					{
+						// @todo(Chaos): Not great because
+						// 1) The time step is hard-coded which might lead to unwanted side effect for frame rate changes
+						// 2) VSmooth logic is spread out over various files and classes (IslandManager, PBDRigidsEvolutionGBF)
+						// 3) Penetration velocity estimation is not what we would expect when calling AddToConstraintGraph
+						const FReal FakeInvDt = 30.0; // hard-coded inverse time step
+						const FReal Phi = Contact.GetPhi();
+						const FVec3 Normal = Contact.CalculateWorldContactNormal();
+						const FVec3 PenetrationVelocity = (Phi * Normal) * FakeInvDt;
+						if (Rigid->VSmooth().SizeSquared() + UE_KINDA_SMALL_NUMBER < PenetrationVelocity.SizeSquared())
+						{
+							Rigid->SetVSmooth(PenetrationVelocity);
+						}
+					}
+				};
+
+
+			const FParticlePair& ParticlePair = ConstraintHandle->GetConstrainedParticles();
+			const FPBDCollisionConstraint& Contact = ConstraintHandle->GetContact();
+
+			if (IslandManager.WasParticleTransformSetExplicitly(ParticlePair[0]))
+			{
+				UpdateSmoothVelocity(ParticlePair[0], Contact);
+			}
+			if (IslandManager.WasParticleTransformSetExplicitly(ParticlePair[1]))
+			{
+				UpdateSmoothVelocity(ParticlePair[1], Contact);
+			}
 		}
 
 		TempCollisions.Reset();
diff --git a/Engine/Source/Runtime/Experimental/Chaos/Private/Chaos/PBDRigidsEvolutionGBF.cpp b/Engine/Source/Runtime/Experimental/Chaos/Private/Chaos/PBDRigidsEvolutionGBF.cpp
--- a/Engine/Source/Runtime/Experimental/Chaos/Private/Chaos/PBDRigidsEvolutionGBF.cpp
+++ b/Engine/Source/Runtime/Experimental/Chaos/Private/Chaos/PBDRigidsEvolutionGBF.cpp
@@ -1553,6 +1553,15 @@ void FPBDRigidsEvolutionGBF::SetParticleTransform(FGeometryParticleHandle* InPar
 
 	FGenericParticleHandle(InParticle)->SetTransform(InPos, InRot);
 
+	// We only want to set the IsParticleTransformSetExplicitly flag if the particle is awake and this is a teleport.
+	// Since only then we want to estimate an overlap velocity based on penetration.
+	// (see @FPBDCollisionConstraints::AddConstraintsToGraph)
+	if (bWakeUp && bIsTeleport)
+	{
+		GetIslandManager().FlagParticleTransformSetExplicitly(InParticle);
+	}
+	GetIslandManager().FlagParticleUpdated(InParticle);
+
 	OnParticleMoved(InParticle, PrevX, PrevR, bIsTeleport, bWakeUp);
 
 	// This is a bit overkill, if this becomes an issue we will need to look into adding support
diff --git a/Engine/Source/Runtime/Experimental/Chaos/Private/PBDRigidsSolver.cpp b/Engine/Source/Runtime/Experimental/Chaos/Private/PBDRigidsSolver.cpp
--- a/Engine/Source/Runtime/Experimental/Chaos/Private/PBDRigidsSolver.cpp
+++ b/Engine/Source/Runtime/Experimental/Chaos/Private/PBDRigidsSolver.cpp
@@ -2979,6 +2979,8 @@ TRACE_COUNTER_SET(ChaosTraceCounter_##Name, Value)
 
 		GetEvolution()->SetParticleObjectState(Rigid, DynamicMisc.ObjectState());
 		GetEvolution()->SetParticleSleepType(Rigid, DynamicMisc.SleepType());
+
+		GetEvolution()->GetIslandManager().FlagParticleUpdated(Rigid->Handle());
 	}
 
 	void FPBDRigidsSolver::ApplyConfig(const FChaosSolverConfiguration& InConfig)
diff --git a/Engine/Source/Runtime/Experimental/Chaos/Public/Chaos/Island/IslandManager.h b/Engine/Source/Runtime/Experimental/Chaos/Public/Chaos/Island/IslandManager.h
--- a/Engine/Source/Runtime/Experimental/Chaos/Public/Chaos/Island/IslandManager.h
+++ b/Engine/Source/Runtime/Experimental/Chaos/Public/Chaos/Island/IslandManager.h
@@ -97,6 +97,7 @@ namespace Chaos::Private
 				uint32 bNeedsResim : 1;
 				uint32 bWasUpdated : 1;
 				uint32 bWasProcessed : 1;
+				uint32 bWasTransformSetExplicitly : 1;
 			};
 
 		private:
@@ -458,6 +459,9 @@ namespace Chaos::Private
 		CHAOS_API int32 GetParticleLevel(FGeometryParticleHandle* Particle) const;
 		CHAOS_API void WakeParticleIslands(FGeometryParticleHandle* Particle);
 		CHAOS_API void SleepParticle(FGeometryParticleHandle* Particle);
+		CHAOS_API void FlagParticleUpdated(const FGeometryParticleHandle* Particle);
+		CHAOS_API bool WasParticleTransformSetExplicitly(const FGeometryParticleHandle* Particle);
+		CHAOS_API void FlagParticleTransformSetExplicitly(const FGeometryParticleHandle* Particle);
 
 		// Constraint management
 		int32 GetNumConstraints() const { return Edges.Num(); }
@@ -618,7 +622,7 @@ namespace Chaos::Private
 		CHAOS_API void PostStepPartialWake(FPBDIsland* Island);
 
 		// Partial Island Waking
-		CHAOS_API void ProcessPartialIslandWake(FPBDIsland* Island);
+		CHAOS_API void ProcessPartialIslandWake(FPBDIsland* Island, const bool bProcessEdgeBasedWake);
 		CHAOS_API bool PropagatePartialIslandWake(TArray<FPBDIslandParticle*>& NodeQueue, TArray<FPBDParticleMomentumSq>& MomentumQueue,
 			TArray<FGeometryParticleHandle*>& ParticlesToWakeUp, const int32 NumNodesConstraintWake, const bool bTriggerNodesAlwaysPropagate);
 		CHAOS_API bool ProcessTriggerParticle(FPBDIslandParticle* Node, TArray<FPBDIslandParticle*>& NodeQueue, TArray<FPBDParticleMomentumSq>& MomentumQueue, const bool bAlwaysTrigger);
@@ -630,6 +634,8 @@ namespace Chaos::Private
 		CHAOS_API void ResetParticleSleepCounter(FPBDIsland* Island);
 		CHAOS_API void ResetEdgeFlags(FPBDIsland* Island);
 		CHAOS_API void ResetNodeFlags(FPBDIsland* Island);
+		CHAOS_API void ResetEdgeFlags();
+		CHAOS_API void ResetNodeFlags();
 
 		// Disabling
 		CHAOS_API void ProcessDisable(TFunctionRef<void(FPBDRigidParticleHandle*)> ParticleDisableFunctor);
diff --git a/Engine/Source/Runtime/Experimental/Chaos/Public/Chaos/PBDRigidsEvolutionGBF.h b/Engine/Source/Runtime/Experimental/Chaos/Public/Chaos/PBDRigidsEvolutionGBF.h
--- a/Engine/Source/Runtime/Experimental/Chaos/Public/Chaos/PBDRigidsEvolutionGBF.h
+++ b/Engine/Source/Runtime/Experimental/Chaos/Public/Chaos/PBDRigidsEvolutionGBF.h
@@ -273,6 +273,9 @@ namespace Chaos
 				const FVec3 PredictedAngularVelocity = Particle.GetW() + Particle.AngularAcceleration() * FakeDT + Particle.AngularImpulseVelocity();
 				Particle.WSmooth() = FMath::Lerp(Particle.WSmooth(), PredictedAngularVelocity, SmoothRate);
 			}
+
+			// Flag the particle as potentially updated when calling ResetVSmoothFromForces
+			GetIslandManager().FlagParticleUpdated(Particle.Handle());
 		}
 
 		template<typename TParticleView> 
