BugHawk suggested comment fixes

diff --git a/Engine/Source/Runtime/Experimental/Chaos/Private/Chaos/Island/IslandManager.cpp b/Engine/Source/Runtime/Experimental/Chaos/Private/Chaos/Island/IslandManager.cpp
--- a/Engine/Source/Runtime/Experimental/Chaos/Private/Chaos/Island/IslandManager.cpp
+++ b/Engine/Source/Runtime/Experimental/Chaos/Private/Chaos/Island/IslandManager.cpp
@@ -2625,8 +2625,9 @@ namespace Chaos::Private
 			// However, we tolerate this since the state should only last for one tick before RewakeConstraintsNotSupportingPartialSleep will sync again.
 		}
 
-		// Deactivate all adjacent constraints between (a) 2 sleeping particles or (b) a sleeping and a non-dynamic particle
+		// Deactivate all adjacent constraints between (a) 2 sleeping particles or (b) a sleeping and a non-dynamic particle.
 		// Processing non-sleeping constraints is sufficient because no particles were woken up since the last call to @SyncConstraintSleepStateToParticles
+		// (Except for particles which were re-woken but whose connected constraints are still active).
 		PropagateParticleSleepToConstraints(Island);
 
 		if (CVars::ChaosSolverPostStepWakeThreshold > 0.0f)
diff --git a/Engine/Source/Runtime/Experimental/Chaos/Private/Chaos/PBDRigidsEvolutionGBF.cpp b/Engine/Source/Runtime/Experimental/Chaos/Private/Chaos/PBDRigidsEvolutionGBF.cpp
--- a/Engine/Source/Runtime/Experimental/Chaos/Private/Chaos/PBDRigidsEvolutionGBF.cpp
+++ b/Engine/Source/Runtime/Experimental/Chaos/Private/Chaos/PBDRigidsEvolutionGBF.cpp
@@ -1553,7 +1553,7 @@ void FPBDRigidsEvolutionGBF::SetParticleTransform(FGeometryParticleHandle* InPar
 
 	FGenericParticleHandle(InParticle)->SetTransform(InPos, InRot);
 
-	// We only want to set the IsParticleTransformSetExplicitly flag if the particle is awake and this is a teleport.
+	// We only want to set the IsParticleTransformSetExplicitly flag if the particle will be woken up and this is a teleport.
 	// Since only then we want to estimate an overlap velocity based on penetration.
 	// (see @FPBDCollisionConstraints::AddConstraintsToGraph)
 	if (bWakeUp && bIsTeleport)
