GeometryCollection: Move storage of PreskinnedBounds in geom collection RenderData out of non-nanite MeshDescription and into common nanite/non-nanite storage.
This allows us to access the precalculated preskinned bounds correctly from both nanite and non-nanite scene proxies.
This replaces an earlier solution  for caching preskinned bounds on the nanite geom collection scene proxy which incorrectly depended on the state of the component when the scene proxy was created.
Resubmit of previous backed out version. Fixing serialization bug when stripping all render data, and forcing incremental cook.
#was-backout 44922689

diff --git a/Engine/Source/Runtime/Experimental/GeometryCollectionEngine/Private/GeometryCollection/DerivedDataGeometryCollectionCooker.cpp b/Engine/Source/Runtime/Experimental/GeometryCollectionEngine/Private/GeometryCollection/DerivedDataGeometryCollectionCooker.cpp
--- a/Engine/Source/Runtime/Experimental/GeometryCollectionEngine/Private/GeometryCollection/DerivedDataGeometryCollectionCooker.cpp
+++ b/Engine/Source/Runtime/Experimental/GeometryCollectionEngine/Private/GeometryCollection/DerivedDataGeometryCollectionCooker.cpp
@@ -122,10 +122,10 @@ bool FDerivedDataGeometryCollectionRenderDataCooker::Build(TArray<uint8>& OutDat
 	return false;
 }
 
+const TCHAR* FDerivedDataGeometryCollectionRenderDataCooker::VersionString = TEXT("2087D33DBF524754A1F8FD20DCB7BEAB");
+
 const TCHAR* FDerivedDataGeometryCollectionRenderDataCooker::GetVersionString() const
 {
-	const TCHAR* VersionString = TEXT("9B7017DECCD24F8C83D647B6009BF0C8");
-
 	static FString CachedNaniteVersionString;
 	if (CachedNaniteVersionString.IsEmpty())
 	{
diff --git a/Engine/Source/Runtime/Experimental/GeometryCollectionEngine/Private/GeometryCollection/GeometryCollectionObject.cpp b/Engine/Source/Runtime/Experimental/GeometryCollectionEngine/Private/GeometryCollection/GeometryCollectionObject.cpp
--- a/Engine/Source/Runtime/Experimental/GeometryCollectionEngine/Private/GeometryCollection/GeometryCollectionObject.cpp
+++ b/Engine/Source/Runtime/Experimental/GeometryCollectionEngine/Private/GeometryCollection/GeometryCollectionObject.cpp
@@ -1978,6 +1978,14 @@ bool UGeometryCollection::Modify(bool bAlwaysMarkDirty /*= true*/)
 	return bSuperResult;
 }
 
+void UGeometryCollection::AppendToClassSchema(FAppendToClassSchemaContext& Context)
+{
+	Super::AppendToClassSchema(Context);
+
+	FStringView RenderDataVersionString(FDerivedDataGeometryCollectionRenderDataCooker::VersionString);
+	Context.Update(RenderDataVersionString.GetData(), RenderDataVersionString.Len() * sizeof(RenderDataVersionString[0]));
+}
+
 void UGeometryCollection::EnsureDataIsCooked(bool bInitResources, bool bIsTransacting, bool bIsPersistant, bool bAllowCopyFromDDC)
 {
 	EnsureSimulationDataIsCooked(bIsTransacting, bAllowCopyFromDDC);
diff --git a/Engine/Source/Runtime/Experimental/GeometryCollectionEngine/Private/GeometryCollection/GeometryCollectionRenderData.cpp b/Engine/Source/Runtime/Experimental/GeometryCollectionEngine/Private/GeometryCollection/GeometryCollectionRenderData.cpp
--- a/Engine/Source/Runtime/Experimental/GeometryCollectionEngine/Private/GeometryCollection/GeometryCollectionRenderData.cpp
+++ b/Engine/Source/Runtime/Experimental/GeometryCollectionEngine/Private/GeometryCollection/GeometryCollectionRenderData.cpp
@@ -125,7 +125,7 @@ void FGeometryCollectionMeshResources::Serialize(FArchive& Ar)
 
 void FGeometryCollectionMeshDescription::Serialize(FArchive& Ar)
 {
-	Ar << NumVertices << NumTriangles << PreSkinnedBounds << Sections << SectionsNoInternal;
+	Ar << NumVertices << NumTriangles << Sections << SectionsNoInternal;
 	
 	// need subsection data on cooked builds when using ray tracing
 	// to be able to build one BLAS per subsection
@@ -633,9 +633,6 @@ void CreateMeshData(FGeometryCollection& InCollection, FGeometryCollectionBuiltM
 	OutRenderData.MeshDescription.Sections = MoveTemp(SectionData.Sections);
 	OutRenderData.MeshDescription.SectionsNoInternal = MoveTemp(SectionData.SectionsNoInternal);
 	OutRenderData.MeshDescription.SubSections = MoveTemp(SectionData.SubSections);
-
-	// Calculate bounds.
-	OutRenderData.MeshDescription.PreSkinnedBounds = BuildPreSkinnedBounds(InCollection);
 }
 
 void CreateNaniteData(FGeometryCollectionBuiltMeshData&& InMeshData, FGeometryCollectionRenderData& OutRenderData)
@@ -696,6 +693,8 @@ TUniquePtr<FGeometryCollectionRenderData> FGeometryCollectionRenderData::Create(
 		{
 			CreateNaniteData(MoveTemp(MeshBuildData), *RenderData.Get());
 		}
+
+		RenderData->PreSkinnedBounds = BuildPreSkinnedBounds(InCollection);
 	}
 
 	return RenderData;
@@ -712,24 +711,28 @@ FGeometryCollectionRenderData::~FGeometryCollectionRenderData()
 
 void FGeometryCollectionRenderData::Serialize(FArchive& Ar, UGeometryCollection& Owner)
 {
+#if WITH_EDITOR
+	// Optionally don't cook the rendering data.
+	// This is used if we expect to use a custom GC render path such as the ISM pool.
 	const bool bIsArchiveValidCandidateForStrip = (Ar.IsCooking() || (Ar.IsCountingMemory() && Ar.IsFilterEditorOnly()));
-	if (Owner.bStripRenderDataOnCook && bIsArchiveValidCandidateForStrip)
+	const bool bStripAllRenderDataFromCook = Owner.bStripRenderDataOnCook && bIsArchiveValidCandidateForStrip;
+	const bool bStripMeshDataFromCook = bStripAllRenderDataFromCook;
+	// Also strip nanite data for non-nanite platforms.
+	const bool bStripNaniteDataFromCook = bStripAllRenderDataFromCook || (Ar.IsCooking() && !DoesTargetPlatformSupportNanite(Ar.CookingTarget()));
+#else
+	const bool bStripMeshDataFromCook = false;
+	const bool bStripNaniteDataFromCook = false;
+#endif
+
+	if (bStripMeshDataFromCook)
 	{
-		// Don't cook rendering data.
-		// This is used if we expect to use a custom GC render path such as the ISM pool.
 		bool bDummy = false;
-		Ar << bDummy; // bHasMeshData
-		Ar << bDummy; // bHasNaniteData
-		return;
+		Ar << bDummy;
+	}
+	else
+	{
+		Ar << bHasMeshData;
 	}
-
-	Ar << bHasMeshData;
-
-#if WITH_EDITOR
-	bool bStripNaniteDataFromCook = Ar.IsCooking() && !DoesTargetPlatformSupportNanite(Ar.CookingTarget());
-#else
-	bool bStripNaniteDataFromCook = false;
-#endif
 
 	if (bStripNaniteDataFromCook)
 	{
@@ -741,7 +744,7 @@ void FGeometryCollectionRenderData::Serialize(FArchive& Ar, UGeometryCollection&
 		Ar << bHasNaniteData;
 	}
 
-	if (bHasMeshData)
+	if (bHasMeshData && !bStripMeshDataFromCook)
 	{
 		MeshResource.Serialize(Ar);
 		MeshDescription.Serialize(Ar);
@@ -772,6 +775,9 @@ void FGeometryCollectionRenderData::Serialize(FArchive& Ar, UGeometryCollection&
 	{
 		ClearNaniteResources(NaniteResourcesPtr);
 	}
+
+	// Always write bounds.
+	Ar << PreSkinnedBounds;
 }
 
 void FGeometryCollectionRenderData::InitResources(UGeometryCollection const& Owner)
diff --git a/Engine/Source/Runtime/Experimental/GeometryCollectionEngine/Private/GeometryCollection/GeometryCollectionRenderData.h b/Engine/Source/Runtime/Experimental/GeometryCollectionEngine/Private/GeometryCollection/GeometryCollectionRenderData.h
--- a/Engine/Source/Runtime/Experimental/GeometryCollectionEngine/Private/GeometryCollection/GeometryCollectionRenderData.h
+++ b/Engine/Source/Runtime/Experimental/GeometryCollectionEngine/Private/GeometryCollection/GeometryCollectionRenderData.h
@@ -92,9 +92,6 @@ struct FGeometryCollectionMeshDescription
 	uint32 NumVertices = 0;
 	uint32 NumTriangles = 0;
 
-	// Bounds before any movement or skinning ever happens.
-	FBoxSphereBounds PreSkinnedBounds;
-
 	// Info on grouped material sections within vertex buffers. 
 	// A section has a uniform material index and can contain multiple per-bone subsections.
 	TArray<FGeometryCollectionMeshElement> Sections;
@@ -133,7 +130,9 @@ public:
 	/** Releases the render resources. */
 	void ReleaseResources();
 
+	/** False if we don't store valid non-nanite mesh data. */
 	bool bHasMeshData = false;
+	/** False if we don't store valid nanite data. */
 	bool bHasNaniteData = false;
 
 	/** Mesh GPU resources. */
@@ -145,6 +144,9 @@ public:
 	/** Nanite resources. */
 	TPimplPtr<Nanite::FResources> NaniteResourcesPtr;
 
+	/** Object local bounds before any animation. */
+	FBoxSphereBounds PreSkinnedBounds = FBoxSphereBounds(ForceInitToZero);
+
 private:
 	bool bIsInitialized = false;
 };
diff --git a/Engine/Source/Runtime/Experimental/GeometryCollectionEngine/Private/GeometryCollection/GeometryCollectionSceneProxy.cpp b/Engine/Source/Runtime/Experimental/GeometryCollectionEngine/Private/GeometryCollection/GeometryCollectionSceneProxy.cpp
--- a/Engine/Source/Runtime/Experimental/GeometryCollectionEngine/Private/GeometryCollection/GeometryCollectionSceneProxy.cpp
+++ b/Engine/Source/Runtime/Experimental/GeometryCollectionEngine/Private/GeometryCollection/GeometryCollectionSceneProxy.cpp
@@ -189,6 +189,8 @@ FGeometryCollectionSceneProxyBase::FGeometryCollectionSceneProxyBase(UGeometryCo
 		}
 	}
 
+	PreSkinnedBounds = Component->GetRestCollection()->RenderData->PreSkinnedBounds;
+
 	DynamicData = Component->InitDynamicData(true);
 }
 
@@ -1134,7 +1136,7 @@ HHitProxy* FGeometryCollectionSceneProxy::CreateHitProxies(UPrimitiveComponent*
 
 void FGeometryCollectionSceneProxy::GetPreSkinnedLocalBounds(FBoxSphereBounds& OutBounds) const
 {
-	OutBounds = MeshDescription.PreSkinnedBounds; 
+	OutBounds = PreSkinnedBounds; 
 }
 
 uint32 FGeometryCollectionSceneProxy::GetAllocatedSize() const
@@ -1345,8 +1347,6 @@ void FNaniteGeometryCollectionSceneProxy::CreateRenderThreadResources(FRHIComman
 	check(NaniteResources.RuntimeResourceID != INDEX_NONE && NaniteResources.HierarchyOffset != INDEX_NONE);
 
 	FGeometryCollectionSceneProxyBase::SetDynamicData_RenderThread(RHICmdList, DynamicData);
-
-	CachedPreSkinnedBounds = GetLocalBounds();
 }
 
 void FNaniteGeometryCollectionSceneProxy::DestroyRenderThreadResources()
@@ -1362,7 +1362,7 @@ SIZE_T FNaniteGeometryCollectionSceneProxy::GetTypeHash() const
 
 void FNaniteGeometryCollectionSceneProxy::GetPreSkinnedLocalBounds(FBoxSphereBounds& OutBounds) const
 {
-	OutBounds = CachedPreSkinnedBounds;
+	OutBounds = PreSkinnedBounds;
 }
 
 FPrimitiveViewRelevance FNaniteGeometryCollectionSceneProxy::GetViewRelevance(const FSceneView* View) const
diff --git a/Engine/Source/Runtime/Experimental/GeometryCollectionEngine/Private/GeometryCollection/GeometryCollectionSceneProxy.h b/Engine/Source/Runtime/Experimental/GeometryCollectionEngine/Private/GeometryCollection/GeometryCollectionSceneProxy.h
--- a/Engine/Source/Runtime/Experimental/GeometryCollectionEngine/Private/GeometryCollection/GeometryCollectionSceneProxy.h
+++ b/Engine/Source/Runtime/Experimental/GeometryCollectionEngine/Private/GeometryCollection/GeometryCollectionSceneProxy.h
@@ -192,6 +192,8 @@ protected:
 	FGeometryCollectionMeshDescription MeshDescription;
 	FMaterialRelevance MaterialRelevance;
 
+	FBoxSphereBounds PreSkinnedBounds;
+
 	FGeometryCollectionVertexFactory VertexFactory;
 	const bool bUseShaderBoneTransform;
 	const bool bSupportsTripleBufferVertexUpload;
@@ -242,8 +244,6 @@ class FGeometryCollectionSceneProxy final : public FPrimitiveSceneProxy, private
 
 	FCollisionResponseContainer CollisionResponse;
 
-	FBoxSphereBounds PreSkinnedBounds;
-
 	bool bRenderResourcesCreated = false;
 
 #if WITH_EDITOR
@@ -377,8 +377,6 @@ protected:
 	};
 	TArray<FGeometryNaniteData> GeometryNaniteData;
 
-	FBoxSphereBounds CachedPreSkinnedBounds;
-
 	uint32 bCastShadow : 1;
 	uint32 bReverseCulling : 1;
 	uint32 bHasMaterialErrors : 1;
diff --git a/Engine/Source/Runtime/Experimental/GeometryCollectionEngine/Public/GeometryCollection/DerivedDataGeometryCollectionCooker.h b/Engine/Source/Runtime/Experimental/GeometryCollectionEngine/Public/GeometryCollection/DerivedDataGeometryCollectionCooker.h
--- a/Engine/Source/Runtime/Experimental/GeometryCollectionEngine/Public/GeometryCollection/DerivedDataGeometryCollectionCooker.h
+++ b/Engine/Source/Runtime/Experimental/GeometryCollectionEngine/Public/GeometryCollection/DerivedDataGeometryCollectionCooker.h
@@ -67,6 +67,8 @@ public:
 
 	bool CanBuild() const { return true; }
 
+	static const TCHAR* VersionString;
+
 private:
 	UGeometryCollection& GeometryCollection;
 };
diff --git a/Engine/Source/Runtime/Experimental/GeometryCollectionEngine/Public/GeometryCollection/GeometryCollectionObject.h b/Engine/Source/Runtime/Experimental/GeometryCollectionEngine/Public/GeometryCollection/GeometryCollectionObject.h
--- a/Engine/Source/Runtime/Experimental/GeometryCollectionEngine/Public/GeometryCollection/GeometryCollectionObject.h
+++ b/Engine/Source/Runtime/Experimental/GeometryCollectionEngine/Public/GeometryCollection/GeometryCollectionObject.h
@@ -394,6 +394,7 @@ public:
 	GEOMETRYCOLLECTIONENGINE_API void Serialize(FArchive& Ar);
 
 #if WITH_EDITOR
+	GEOMETRYCOLLECTIONENGINE_API static void AppendToClassSchema(FAppendToClassSchemaContext& Context);
 	GEOMETRYCOLLECTIONENGINE_API void EnsureDataIsCooked(bool bInitResources, bool bIsTransacting, bool bIsPersistant, bool bAllowCopyFromDDC = true);
 #endif
 
