Review feedback for 44894454 - add ensures to null checks

diff --git a/Engine/Source/Runtime/Experimental/IoStore/OnDemand/Private/OnDemandIoStore.cpp b/Engine/Source/Runtime/Experimental/IoStore/OnDemand/Private/OnDemandIoStore.cpp
--- a/Engine/Source/Runtime/Experimental/IoStore/OnDemand/Private/OnDemandIoStore.cpp
+++ b/Engine/Source/Runtime/Experimental/IoStore/OnDemand/Private/OnDemandIoStore.cpp
@@ -705,7 +705,7 @@ FIoStatus FOnDemandIoStore::Unmount(FStringView MountId)
 			return false;
 		}));
 
-		if (bRemoved && PackageStoreBackend)
+		if (bRemoved && ensureMsgf(PackageStoreBackend, TEXT("PackageStoreBackend is null, is this a prefork server?")))
 		{
 			PackageStoreBackend->NeedsUpdate();
 		}
@@ -1399,7 +1399,7 @@ FIoStatus FOnDemandIoStore::TickMountRequest(FMountRequest& MountRequest)
 		PendingContainerHeaders.Append(MoveTemp(RequestedContainerHeaders));
 	}
 
-	if (PackageStoreBackend)
+	if (ensureMsgf(PackageStoreBackend, TEXT("PackageStoreBackend is null, is this a prefork server?")))
 	{
 		PackageStoreBackend->NeedsUpdate();
 	}
@@ -1493,7 +1493,7 @@ void FOnDemandIoStore::OnEncryptionKeyAdded(const FGuid& Id, const FAES::FAESKey
 		}
 	}
 
-	if (bAddedContainerHeaders && PackageStoreBackend)
+	if (bAddedContainerHeaders && ensureMsgf(PackageStoreBackend, TEXT("PackageStoreBackend is null, is this a prefork server?")))
 	{
 		PackageStoreBackend->NeedsUpdate();
 	}
@@ -1944,12 +1944,12 @@ void FOnDemandIoStore::ReleaseContent(FOnDemandInternalContentHandle& ContentHan
 	}
 
 	// Package store needs to be aware of install state
-	if (PackageStoreBackend)
+	if (ensureMsgf(PackageStoreBackend, TEXT("PackageStoreBackend is null, is this a prefork server?")))
 	{
 		PackageStoreBackend->NeedsUpdate(EOnDemandPackageStoreUpdateMode::ReferencedPackages);
 	}
 
-	if (InstallCache)
+	if (ensureMsgf(InstallCache, TEXT("InstallCache is null, is this a prefork server?")))
 	{
 		InstallCache->UpdateLastAccess(ReleasedChunks);
 	}
