Dataflow : Ed node fixes
* Fix array pin not showing the right icon under certain circumstances
* Drop not create inline editable fields for properties that are nested in structure and the ones with a EditCondition

diff --git a/Engine/Source/Runtime/Experimental/Dataflow/Engine/Private/Dataflow/DataflowEdNode.cpp b/Engine/Source/Runtime/Experimental/Dataflow/Engine/Private/Dataflow/DataflowEdNode.cpp
--- a/Engine/Source/Runtime/Experimental/Dataflow/Engine/Private/Dataflow/DataflowEdNode.cpp
+++ b/Engine/Source/Runtime/Experimental/Dataflow/Engine/Private/Dataflow/DataflowEdNode.cpp
@@ -204,6 +204,10 @@ void UDataflowEdNode::UpdatePinsFromDataflowNode()
 								Pin->PinType.PinCategory = DataflowInput->GetType();
 								Pin->PinType.PinSubCategory = NAME_None;
 								Pin->PinType.PinSubCategoryObject = nullptr;
+								if (FDataflowArrayTypePolicy::SupportsTypeStatic(DataflowInput->GetType()))
+								{
+									Pin->PinType.ContainerType = EPinContainerType::Array;
+								}
 							}
 						}
 						else if (Pin->Direction == EEdGraphPinDirection::EGPD_Output)
@@ -221,6 +225,10 @@ void UDataflowEdNode::UpdatePinsFromDataflowNode()
 								Pin->PinType.PinCategory = DataflowOutput->GetType();
 								Pin->PinType.PinSubCategory = NAME_None;
 								Pin->PinType.PinSubCategoryObject = nullptr;
+								if (FDataflowArrayTypePolicy::SupportsTypeStatic(DataflowOutput->GetType()))
+								{
+									Pin->PinType.ContainerType = EPinContainerType::Array;
+								}
 							}
 						}
 					}
@@ -373,6 +381,25 @@ bool UDataflowEdNode::SupportsEditablePinType(const UEdGraphPin& Pin)
 {
 	if (const FDataflowConnection* Connection = GetConnectionFromPin(&Pin))
 	{
+		if (const FProperty* Property = Connection->GetProperty())
+		{
+			// We do not handle EditConditions
+			static const FName Name_EditConditionHides("EditCondition");
+			if (Property->HasMetaData(Name_EditConditionHides))
+			{
+				return false;
+			}
+
+			// only for direct properties of the node
+			if (UStruct* OwnerStruct = Property->GetOwnerStruct())
+			{
+				if (!OwnerStruct->IsChildOf<FDataflowNode>())
+				{
+					return false;
+				}
+			}
+		}
+
 		// todo(dataflow): Add support for anytype in the future 
 		if (!Connection->IsAnyType())
 		{
