Dataflow : fix typo

diff --git a/Engine/Source/Runtime/Experimental/Dataflow/Engine/Private/Dataflow/DataflowEdNode.cpp b/Engine/Source/Runtime/Experimental/Dataflow/Engine/Private/Dataflow/DataflowEdNode.cpp
--- a/Engine/Source/Runtime/Experimental/Dataflow/Engine/Private/Dataflow/DataflowEdNode.cpp
+++ b/Engine/Source/Runtime/Experimental/Dataflow/Engine/Private/Dataflow/DataflowEdNode.cpp
@@ -384,8 +384,8 @@ bool UDataflowEdNode::SupportsEditablePinType(const UEdGraphPin& Pin)
 		if (const FProperty* Property = Connection->GetProperty())
 		{
 			// We do not handle EditConditions
-			static const FName Name_EditConditionHides("EditCondition");
-			if (Property->HasMetaData(Name_EditConditionHides))
+			static const FName Name_EditCondition("EditCondition");
+			if (Property->HasMetaData(Name_EditCondition))
 			{
 				return false;
 			}
