IAD - fix cache blocks sometimes being larger than their max size.  This was caused by always adding the first chunk size instead of the size of the chunk being written. Chunks are now flushed in buffer order instead of reverese order as well.

diff --git a/Engine/Source/Runtime/Experimental/IoStore/OnDemand/Private/OnDemandContentInstaller.cpp b/Engine/Source/Runtime/Experimental/IoStore/OnDemand/Private/OnDemandContentInstaller.cpp
--- a/Engine/Source/Runtime/Experimental/IoStore/OnDemand/Private/OnDemandContentInstaller.cpp
+++ b/Engine/Source/Runtime/Experimental/IoStore/OnDemand/Private/OnDemandContentInstaller.cpp
@@ -919,6 +919,8 @@ void FOnDemandContentInstaller::ProcessDownloadedChunk(
 
 			if (ChunkEntry.Hash == ChunkHash)
 			{
+				check(ChunkEntry.GetDiskSize() == Chunk.GetSize());
+
 				if (FIoStatus Status = IoStore.InstallCache->PutChunk(MoveTemp(Chunk), ChunkHash); !Status.IsOk())
 				{
 					Request.ErrorCode	= Status.GetErrorCode();
diff --git a/Engine/Source/Runtime/Experimental/IoStore/OnDemand/Private/OnDemandInstallCache.cpp b/Engine/Source/Runtime/Experimental/IoStore/OnDemand/Private/OnDemandInstallCache.cpp
--- a/Engine/Source/Runtime/Experimental/IoStore/OnDemand/Private/OnDemandInstallCache.cpp
+++ b/Engine/Source/Runtime/Experimental/IoStore/OnDemand/Private/OnDemandInstallCache.cpp
@@ -1630,16 +1630,6 @@ class FOnDemandInstallCache final
 			Chunks.Add(MoveTemp(Chunk));
 		}
 
-		FIoBuffer Pop(FIoHash& OutChunkHash)
-		{
-			check(Chunks.Num() == ChunkHashes.Num());
-			check(Chunks.IsEmpty() == false);
-			FIoBuffer Chunk = Chunks.Pop(EAllowShrinking::No);
-			TotalSize		= TotalSize - Chunk.GetSize();
-			OutChunkHash	= ChunkHashes.Pop(EAllowShrinking::No);
-			return Chunk;
-		}
-
 		void Reset()
 		{
 			Chunks.Reset();
@@ -3277,10 +3267,14 @@ FIoStatus FOnDemandInstallCache::FlushPendingChunksImpl(FPendingChunks& Chunks,
 {
 	ON_SCOPE_EXIT { Chunks.Reset(); };
 
+	FLargeMemoryWriter Ar(FMath::Min(Chunks.TotalSize, static_cast<uint64>(Cas.GetMaxBlockSize())));
+
 	// This should be the only thread writing to CurrentBlock
 	FCasBlockId CurrentBlockId = CurrentBlock;
 
-	while (Chunks.IsEmpty() == false)
+	int32 ChunkIdx = 0;
+
+	while (ChunkIdx < Chunks.Chunks.Num())
 	{
 		FCasJournal::FTransaction Transaction = FCasJournal::Begin(GetJournalFilename());
 		
@@ -3301,19 +3295,22 @@ FIoStatus FOnDemandInstallCache::FlushPendingChunksImpl(FPendingChunks& Chunks,
 
 		const int64 CasBlockOffset = CasFileHandle->Tell();
 
-		FLargeMemoryWriter	Ar(Chunks.TotalSize);
 		TArray<FIoHash>		ChunkHashes;
 		TArray<int64>		Offsets;
 
-		while (Chunks.IsEmpty() == false)
+		while (ChunkIdx < Chunks.Chunks.Num())
 		{
-			if (CasBlockOffset > 0 && CasBlockOffset + Ar.Tell() + Chunks.Chunks[0].GetSize() > Cas.GetMaxBlockSize())
+			if (CasBlockOffset > 0 && CasBlockOffset + Ar.Tell() + Chunks.Chunks[ChunkIdx].GetSize() > Cas.GetMaxBlockSize())
 			{
 				break;
 			}
-			FIoBuffer Chunk = Chunks.Pop(ChunkHashes.AddDefaulted_GetRef());
+
+			FIoBuffer Chunk = MoveTemp(Chunks.Chunks[ChunkIdx]);
+			ChunkHashes.Add(MoveTemp(Chunks.ChunkHashes[ChunkIdx]));
 			Offsets.Add(CasBlockOffset + Ar.Tell());
 			Ar.Serialize(Chunk.GetData(), Chunk.GetSize());
+
+			++ChunkIdx;
 		}
 
 		if (Ar.Tell() > 0)
@@ -3321,6 +3318,8 @@ FIoStatus FOnDemandInstallCache::FlushPendingChunksImpl(FPendingChunks& Chunks,
 			UE_LOG(LogIoStoreOnDemand, Log, TEXT("Writing %.2lf MiB to CAS block %u"),
 				ToMiB(Ar.Tell()), CurrentBlockId.Id);
 
+			check(CasBlockOffset + Ar.Tell() <= Cas.GetMaxBlockSize());
+
 			if (CasFileHandle->Write(Ar.GetData(), Ar.Tell()) == false)
 			{
 				return InstallCacheStatusBuilder(EIoErrorCode::WriteError)
@@ -3366,6 +3365,8 @@ FIoStatus FOnDemandInstallCache::FlushPendingChunksImpl(FPendingChunks& Chunks,
 					Transaction.ChunkLocation(Loc, CasAddr);
 				}
 			}
+
+			Ar.Seek(0);
 		}
 
 		if (FIoStatus Status = FCasJournal::Commit(MoveTemp(Transaction)); Status.IsOk() == false)
@@ -3373,7 +3374,7 @@ FIoStatus FOnDemandInstallCache::FlushPendingChunksImpl(FPendingChunks& Chunks,
 			return Status;
 		}
 
-		if (Chunks.IsEmpty() == false)
+		if (ChunkIdx < Chunks.Chunks.Num())
 		{
 			CurrentBlockId = FCasBlockId::Invalid;
 		}
