IASTool - Allow cryptokeys to be passed to the ChunkPlugin command.

#rnx

- Added a new cmdline arg '-CryptoKeys' to the ChunkPlugin command that allows the caller to pass in a path to a json file containing the encryption keys needed to decrypt the container(s) being chunked.
- Add the EIoContainerFlags::Encrypted flag (if needed) to any ondemand toc we are writing so that they are handled correctly at runtime.
- Allow the paranoid data validation check we have to decrypt the chunks being checked if needed.
- Order includes to work with ICYU.

diff --git a/Engine/Source/Runtime/Experimental/IoStore/OnDemand/Private/Tool/ChunkPlugin.cpp b/Engine/Source/Runtime/Experimental/IoStore/OnDemand/Private/Tool/ChunkPlugin.cpp
--- a/Engine/Source/Runtime/Experimental/IoStore/OnDemand/Private/Tool/ChunkPlugin.cpp
+++ b/Engine/Source/Runtime/Experimental/IoStore/OnDemand/Private/Tool/ChunkPlugin.cpp
@@ -8,14 +8,15 @@
 #include "Containers/StringView.h"
 #include "HAL/FileManager.h"
 #include "IO/IoChunkEncoding.h"
+#include "IO/IoContainerHeader.h"
 #include "IO/IoStatus.h"
 #include "IO/IoStore.h"
 #include "IO/IoStoreOnDemand.h"
-#include "IO/IoContainerHeader.h"
 #include "Misc/DateTime.h"
 #include "Misc/FileHelper.h"
-#include "Misc/Paths.h"
+#include "Misc/KeyChainUtilities.h"
 #include "Misc/PathViews.h"
+#include "Misc/Paths.h"
 #include "Misc/StringBuilder.h"
 #include "S3/S3Client.h"
 #include "Serialization/JsonReader.h"
@@ -268,9 +269,11 @@ static int32 ChunkContainer(const FString& ContainerUTocPath, const FChunkPlugin
 	ContainerEntry.ContainerId = ContainerFileReader.GetContainerId();
 	ContainerEntry.ContainerName = FPaths::GetBaseFilename(ContainerUTocPath);
 
-	if (EnumHasAnyFlags(ContainerFileReader.GetContainerFlags(), EIoContainerFlags::Encrypted))
+	const bool bIsContainerEncrypted = EnumHasAnyFlags(ContainerFileReader.GetContainerFlags(), EIoContainerFlags::Encrypted);
+	if (bIsContainerEncrypted)
 	{
 		ContainerEntry.EncryptionKeyGuid = LexToString(ContainerFileReader.GetEncryptionKeyGuid());
+		ContainerEntry.ContainerFlags |= static_cast<uint8>(EIoContainerFlags::Encrypted);
 	}
 
 	if (Settings.PackageSets.Num() > 0 && OutFoundTagSets != nullptr)
@@ -340,7 +343,7 @@ static int32 ChunkContainer(const FString& ContainerUTocPath, const FChunkPlugin
 
 	for (const FIoStoreTocChunkInfo& ChunkInfo : ChunkInfos)
 	{
-		const bool bDecrypt = false;
+		const bool bDecrypt = false; // We do not want to decrypt the results. We will upload encrypted chunks and decrypt on the client
 		TIoStatusOr<FIoStoreCompressedReadResult> Status = ContainerFileReader.ReadCompressed(ChunkInfo.Id, FIoReadOptions(), bDecrypt);
 		if (!Status.IsOk())
 		{
@@ -389,12 +392,20 @@ static int32 ChunkContainer(const FString& ContainerUTocPath, const FChunkPlugin
 			return -1;
 		}
 
+		// Additional validation to make sure that the chunks can be decrpted and decoded
 		{
 			FIoChunkDecodingParams DecodingParams;
 			DecodingParams.TotalRawSize = RawChunkSize;
 			DecodingParams.EncodedBlockSize = TConstArrayView<uint32>(ContainerEntry.BlockSizes.GetData() + BlockOffset, BlockCount);
 			DecodingParams.BlockHash = TConstArrayView<FIoBlockHash>(ContainerEntry.BlockHashes.GetData() + BlockOffset, BlockCount);
 
+			if (bIsContainerEncrypted)
+			{
+				const FAES::FAESKey* EncryptionKey = EncryptionKeys.Find(ContainerFileReader.GetEncryptionKeyGuid());
+				check(EncryptionKey);
+				DecodingParams.EncryptionKey = FMemoryView(EncryptionKey->Key, EncryptionKey->KeySize);
+			}
+
 			FIoBuffer OutputBuffer(RawChunkSize);
 			bool bDecoded = FIoChunkEncoding::Decode(DecodingParams, ReadResult.IoBuffer.GetView(), OutputBuffer.GetMutableView(), EIoDecodeFlags::None);
 			check(bDecoded);
@@ -486,6 +497,20 @@ int32 WriteOnDemandToc(FOnDemandToc& OnDemandToc, const FString& ContainerName,
 	}
 }
 
+static FKeyChain LoadCryptoKeys(const FContext& Context)
+{
+	FKeyChain Ret;
+
+	auto Path = Context.Get<FStringView>(TEXT("-CryptoKeys"));
+	if (Path.IsEmpty())
+	{
+		return Ret;
+	}
+
+	KeyChainUtilities::LoadKeyChainFromFile(Path.GetData(), Ret);
+	return Ret;
+}
+
 ////////////////////////////////////////////////////////////////////////////////
 static int32 ChunkPluginCommandEntry(const FContext& Context)
 {
@@ -623,6 +648,13 @@ static int32 ChunkPluginCommandEntry(const FContext& Context)
 	FChunkPluginStats Stats;
 
 	TMap<FGuid, FAES::FAESKey> EncryptionKeys;
+
+	FKeyChain KeyChain = LoadCryptoKeys(Context);
+	for (const TPair<FGuid, FNamedAESKey>& KeyPair : KeyChain.GetEncryptionKeys())
+	{
+		EncryptionKeys.Add(KeyPair.Key, KeyPair.Value.Key);
+	}
+
 	if (FileMgr.DirectoryExists(*ContainerFolder) == false)
 	{
 		UE_LOG(LogIoStoreOnDemand, Error, TEXT("Directory '%s' does not exist"), *ContainerFolder);
@@ -883,6 +915,7 @@ static FCommand ChunkPluginCommand(
 		TArgument<FStringView>(TEXT("-SettingsFile"),		TEXT("Optional settings file.")),
 		TArgument<FStringView>(TEXT("-OutputStatsJson"),	TEXT("Path to write a json file with statistics.")),
 		TArgument<FStringView>(TEXT("-HostGroupName"),		TEXT("Host group name or URL")),
+		TArgument<FStringView>(TEXT("-CryptoKeys"),			TEXT("")),
 		TArgument<bool>(TEXT("-IncludeSigPak"),				TEXT("Include .sig and .pak file in the uondemandtoc")),
 		TArgument<bool>(TEXT("-KeepContainerFiles"),		TEXT("Should we keep the container files after processing them.")),
 		TArgument<bool>(TEXT("-StreamOnDemand"),			TEXT("Set the content to be streamed on-demand")),
