Iris - Implemented RPC state stomp detection which can detect state buffer writes between quantization and serialization. This code path is only enabled on dedicated servers. The cvar to control this is net.Iris.RPC.ServerStompDetectionEnabled which is false by default.

diff --git a/Engine/Source/Runtime/Experimental/Iris/Core/Private/Iris/ReplicationSystem/NetBlob/NetBlob.cpp b/Engine/Source/Runtime/Experimental/Iris/Core/Private/Iris/ReplicationSystem/NetBlob/NetBlob.cpp
--- a/Engine/Source/Runtime/Experimental/Iris/Core/Private/Iris/ReplicationSystem/NetBlob/NetBlob.cpp
+++ b/Engine/Source/Runtime/Experimental/Iris/Core/Private/Iris/ReplicationSystem/NetBlob/NetBlob.cpp
@@ -11,6 +11,8 @@
 #include "Iris/Serialization/InternalNetSerializationContext.h"
 #include "Iris/Serialization/NetSerializationContext.h"
 #include "Iris/Serialization/ObjectNetSerializer.h"
+#include "HAL/PlatformMemory.h"
+#include "Templates/AlignmentTemplates.h"
 
 namespace UE::Net
 {
@@ -30,6 +32,7 @@ FNetBlob::~FNetBlob()
 		{
 			if (uint8* StateBuffer = QuantizedBlobState.GetStateBuffer())
 			{
+				QuantizedBlobState.Unprotect();
 				Private::FReplicationStateOperationsInternal::FreeDynamicState(StateBuffer, Descriptor);
 			}
 		}
@@ -169,12 +172,50 @@ void FNetObjectAttachment::DeserializeSubObjectReference(FNetSerializationContex
 	ReadFullNetObjectReference(Context, TargetObjectReference);
 }
 
-FNetBlob::FQuantizedBlobState::FQuantizedBlobState(uint32 Size, uint32 Alignment)
-: StateBuffer(static_cast<uint8*>(FMemory::Malloc(Size, Alignment)))
+FNetBlob::FQuantizedBlobState::FQuantizedBlobState(uint32 Size, uint32 Alignment, FQuantizedBlobState::EMemoryAllocationFlags InMemoryAllocationFlags)
+: MemoryAllocationFlags(InMemoryAllocationFlags)
+{
+	if (InMemoryAllocationFlags == EMemoryAllocationFlags::None)
+	{
+		StateBuffer = static_cast<uint8*>(FMemory::MallocZeroed(Size, Alignment));
+		AllocationSize = Size;
+	}
+	else
+	{
+		const FPlatformMemoryConstants& MemoryConstants = FPlatformMemory::GetConstants();
+		const uint32 NewAlignment = Align(Alignment, MemoryConstants.PageSize);
+		AllocationSize = static_cast<uint32>(Align(Size, MemoryConstants.PageSize));
+		StateBuffer = static_cast<uint8*>(FMemory::MallocZeroed(AllocationSize, NewAlignment));
+	}
+}
+
+
+FNetBlob::FQuantizedBlobState::~FQuantizedBlobState()
 {
 	if (StateBuffer)
 	{
-		FMemory::Memzero(StateBuffer, Size);
+		Unprotect();
+		FMemory::Free(StateBuffer);
+	}
+}
+
+void FNetBlob::FQuantizedBlobState::Protect()
+{
+	if (MemoryAllocationFlags == EMemoryAllocationFlags::Protectable)
+	{
+		constexpr bool bCanRead = true;
+		constexpr bool bCanWrite = false;
+		FPlatformMemory::PageProtect(StateBuffer, AllocationSize, bCanRead, bCanWrite);
+	}
+}
+
+void FNetBlob::FQuantizedBlobState::Unprotect()
+{
+	if (MemoryAllocationFlags == EMemoryAllocationFlags::Protectable)
+	{
+		constexpr bool bCanRead = true;
+		constexpr bool bCanWrite = true;
+		FPlatformMemory::PageProtect(StateBuffer, AllocationSize, bCanRead, bCanWrite);
 	}
 }
 
diff --git a/Engine/Source/Runtime/Experimental/Iris/Core/Private/Iris/ReplicationSystem/NetBlob/NetRPC.cpp b/Engine/Source/Runtime/Experimental/Iris/Core/Private/Iris/ReplicationSystem/NetBlob/NetRPC.cpp
--- a/Engine/Source/Runtime/Experimental/Iris/Core/Private/Iris/ReplicationSystem/NetBlob/NetRPC.cpp
+++ b/Engine/Source/Runtime/Experimental/Iris/Core/Private/Iris/ReplicationSystem/NetBlob/NetRPC.cpp
@@ -2,6 +2,7 @@
 
 #include "Iris/ReplicationSystem/NetBlob/NetRPC.h"
 
+#include "HAL/IConsoleManager.h"
 #include "Iris/Core/BitTwiddling.h"
 #include "Net/Core/Trace/NetDebugName.h"
 #include "Net/Core/Trace/NetTrace.h"
@@ -32,6 +33,13 @@ DEFINE_LOG_CATEGORY_STATIC(LogIrisRpc, Log, All);
 namespace UE::Net::Private
 {
 
+static bool bEnableRPCServerStateStompDetection = false;
+static FAutoConsoleVariableRef CVarEnableRPCStateStompDetection(
+	TEXT("net.Iris.RPC.ServerStompDetectionEnabled"),
+	bEnableRPCServerStateStompDetection,
+	TEXT("If enabled we can detect if RPC state is stomped between quantization and serialization on a dedicated server.")
+);
+
 static bool NetRPC_GetFunctionLocator(const UReplicationSystem* ReplicationSystem, const FNetObjectReference& ObjectReference, const UFunction* Function, FNetRPC::FFunctionLocator& OutFunctionLocator, const FReplicationStateMemberFunctionDescriptor*& OutFunctionDescriptor);
 
 static UObject* NetRPC_GetObject(FNetSerializationContext& Context, const FNetObjectReference& ObjectReference);
@@ -401,7 +409,12 @@ FNetRPC* FNetRPC::Create(UReplicationSystem* ReplicationSystem, const FNetBlobCr
 	// Don't spend CPU cycles on quantizing zero parameters
 	if (BlobDescriptor != nullptr && BlobDescriptor->InternalSize)
 	{
-		QuantizedBlobState = FQuantizedBlobState(BlobDescriptor->InternalSize, BlobDescriptor->InternalAlignment);
+#if WITH_SERVER_CODE
+		const FQuantizedBlobState::EMemoryAllocationFlags MemoryAllocationFlags = bEnableRPCServerStateStompDetection ? FQuantizedBlobState::EMemoryAllocationFlags::Protectable : FQuantizedBlobState::EMemoryAllocationFlags::None;
+#else
+		const FQuantizedBlobState::EMemoryAllocationFlags MemoryAllocationFlags = FQuantizedBlobState::EMemoryAllocationFlags::None;
+#endif
+		QuantizedBlobState = FQuantizedBlobState(FunctionDescriptor->Descriptor->InternalSize, FunctionDescriptor->Descriptor->InternalAlignment, MemoryAllocationFlags);
 
 		// Setup Context
 		FNetSerializationContext Context;
@@ -411,6 +424,13 @@ FNetRPC* FNetRPC::Create(UReplicationSystem* ReplicationSystem, const FNetBlobCr
 
 		// Quantize the function parameters
 		FReplicationStateOperations::Quantize(Context, QuantizedBlobState.GetStateBuffer(), static_cast<const uint8*>(FunctionParameters), BlobDescriptor);
+
+#if WITH_SERVER_CODE
+		if (MemoryAllocationFlags == FQuantizedBlobState::EMemoryAllocationFlags::Protectable)
+		{
+			QuantizedBlobState.Protect();
+		}
+#endif
 	}
 
 	FNetRPC* NetRPC = new FNetRPC(CreationInfo);
diff --git a/Engine/Source/Runtime/Experimental/Iris/Core/Public/Iris/ReplicationSystem/NetBlob/NetBlob.h b/Engine/Source/Runtime/Experimental/Iris/Core/Public/Iris/ReplicationSystem/NetBlob/NetBlob.h
--- a/Engine/Source/Runtime/Experimental/Iris/Core/Public/Iris/ReplicationSystem/NetBlob/NetBlob.h
+++ b/Engine/Source/Runtime/Experimental/Iris/Core/Public/Iris/ReplicationSystem/NetBlob/NetBlob.h
@@ -63,11 +63,18 @@ class FNetBlob
 public:
 	struct FQuantizedBlobState
 	{
+		enum class EMemoryAllocationFlags : uint32
+		{
+			None = 0,
+			// If the state is constructed as protectable one can call Protect() to detect memory stomps.
+			Protectable = 1,
+		};
+
 		FQuantizedBlobState() = default;
 		FQuantizedBlobState(const FQuantizedBlobState&) = delete;
 		FQuantizedBlobState(FQuantizedBlobState&&);
-		IRISCORE_API FQuantizedBlobState(uint32 Size, uint32 Alignment);
-		~FQuantizedBlobState();
+		IRISCORE_API FQuantizedBlobState(uint32 Size, uint32 Alignment, EMemoryAllocationFlags MemoryAllocationFlags = EMemoryAllocationFlags::None);
+		IRISCORE_API ~FQuantizedBlobState();
 
 		FQuantizedBlobState& operator=(const FQuantizedBlobState&) = delete;
 		FQuantizedBlobState& operator=(FQuantizedBlobState&&);
@@ -75,8 +82,16 @@ public:
 		uint8* GetStateBuffer() { return StateBuffer; }
 		const uint8* GetStateBuffer() const { return StateBuffer; }
 
+		/** Protects the state buffer such that a page fault will occur if something tries to modify its contents. The state must have been allocated with MakeProtextableState(). */
+		IRISCORE_API void Protect();
+
+		/** Allows the state buffer to be modified again if it was protected. */
+		IRISCORE_API void Unprotect();
+
 	private:
 		uint8* StateBuffer = nullptr;
+		EMemoryAllocationFlags MemoryAllocationFlags = EMemoryAllocationFlags::None;
+		uint32 AllocationSize = 0;
 	};
 
 	/** Construct a NetBlob with reference count zero. */
@@ -208,26 +223,29 @@ protected:
 inline FNetBlob::FQuantizedBlobState::FQuantizedBlobState(FNetBlob::FQuantizedBlobState&& Other)
 {
 	StateBuffer = Other.StateBuffer;
-	Other.StateBuffer = nullptr;
-}
+	MemoryAllocationFlags = Other.MemoryAllocationFlags;
+	AllocationSize = Other.AllocationSize;
 
-inline FNetBlob::FQuantizedBlobState::~FQuantizedBlobState()
-{
-	if (StateBuffer)
-	{
-		FMemory::Free(StateBuffer);
-	}
+	Other.StateBuffer = nullptr;
+	Other.MemoryAllocationFlags = EMemoryAllocationFlags::None;
+	Other.AllocationSize = 0;
 }
 
 inline FNetBlob::FQuantizedBlobState& FNetBlob::FQuantizedBlobState::operator=(FNetBlob::FQuantizedBlobState&& Other)
 {
 	if (StateBuffer)
 	{
+		Unprotect();
 		FMemory::Free(StateBuffer);
 	}
 
 	StateBuffer = Other.StateBuffer;
+	MemoryAllocationFlags = Other.MemoryAllocationFlags;
+	AllocationSize = Other.AllocationSize;
+
 	Other.StateBuffer = nullptr;
+	Other.MemoryAllocationFlags = EMemoryAllocationFlags::None;
+	Other.AllocationSize = 0;
 	return *this;
 }
 
