Fixed function call documentation.

#rnx

diff --git a/Engine/Source/Runtime/Experimental/Iris/Core/Public/Iris/ReplicationSystem/NetBlob/NetBlob.h b/Engine/Source/Runtime/Experimental/Iris/Core/Public/Iris/ReplicationSystem/NetBlob/NetBlob.h
--- a/Engine/Source/Runtime/Experimental/Iris/Core/Public/Iris/ReplicationSystem/NetBlob/NetBlob.h
+++ b/Engine/Source/Runtime/Experimental/Iris/Core/Public/Iris/ReplicationSystem/NetBlob/NetBlob.h
@@ -82,7 +82,7 @@ public:
 		uint8* GetStateBuffer() { return StateBuffer; }
 		const uint8* GetStateBuffer() const { return StateBuffer; }
 
-		/** Protects the state buffer such that a page fault will occur if something tries to modify its contents. The state must have been allocated with MakeProtextableState(). */
+		/** Protects the state buffer such that a page fault will occur if something tries to modify its contents. The state must have been constructed with EMemoryAllocationFlags::Protectable. */
 		IRISCORE_API void Protect();
 
 		/** Allows the state buffer to be modified again if it was protected. */
