IASTool - Improve error handling when running the 'ChunkPlugin' command

#rnx

- Instead of asserting on failed conditions we should log an actual error message and return an error value.

diff --git a/Engine/Source/Runtime/Experimental/IoStore/OnDemand/Private/Tool/ChunkPlugin.cpp b/Engine/Source/Runtime/Experimental/IoStore/OnDemand/Private/Tool/ChunkPlugin.cpp
--- a/Engine/Source/Runtime/Experimental/IoStore/OnDemand/Private/Tool/ChunkPlugin.cpp
+++ b/Engine/Source/Runtime/Experimental/IoStore/OnDemand/Private/Tool/ChunkPlugin.cpp
@@ -392,7 +392,7 @@ static int32 ChunkContainer(const FString& ContainerUTocPath, const FChunkPlugin
 			return -1;
 		}
 
-		// Additional validation to make sure that the chunks can be decrpted and decoded
+		// Additional validation to make sure that the chunks can be decrypted and decoded
 		{
 			FIoChunkDecodingParams DecodingParams;
 			DecodingParams.TotalRawSize = RawChunkSize;
@@ -402,13 +402,21 @@ static int32 ChunkContainer(const FString& ContainerUTocPath, const FChunkPlugin
 			if (bIsContainerEncrypted)
 			{
 				const FAES::FAESKey* EncryptionKey = EncryptionKeys.Find(ContainerFileReader.GetEncryptionKeyGuid());
-				check(EncryptionKey);
+				if (EncryptionKey == nullptr)
+				{
+					UE_LOG(LogIoStoreOnDemand, Error, TEXT("Failed to find encryption key '%s', Container='%s"), *LexToString(ContainerFileReader.GetEncryptionKeyGuid()), *ContainerUTocPath);
+					return -1;
+				}
 				DecodingParams.EncryptionKey = FMemoryView(EncryptionKey->Key, EncryptionKey->KeySize);
 			}
 
 			FIoBuffer OutputBuffer(RawChunkSize);
 			bool bDecoded = FIoChunkEncoding::Decode(DecodingParams, ReadResult.IoBuffer.GetView(), OutputBuffer.GetMutableView(), EIoDecodeFlags::None);
-			check(bDecoded);
+			if (!bDecoded)
+			{
+				UE_LOG(LogIoStoreOnDemand, Error, TEXT("Failed to decode chunk from disk, Container='%s', ChunkId='%s'"), *ContainerUTocPath, *LexToString(ChunkInfo.Id));
+				return -1;
+			}
 		}
 
 		FIoStatus WriteStatus = ChunkWriter->WriteChunk(ChunksRelativeFolder, ReadResult.IoBuffer, ChunkHash);
