IAD cache -

Fix always printing
"Error: Failed to append critical error to journal"
when Defrag doesn't have enough space to run.

The error was being logged even if the journal commit succeeded.

diff --git a/Engine/Source/Runtime/Experimental/IoStore/OnDemand/Private/OnDemandInstallCache.cpp b/Engine/Source/Runtime/Experimental/IoStore/OnDemand/Private/OnDemandInstallCache.cpp
--- a/Engine/Source/Runtime/Experimental/IoStore/OnDemand/Private/OnDemandInstallCache.cpp
+++ b/Engine/Source/Runtime/Experimental/IoStore/OnDemand/Private/OnDemandInstallCache.cpp
@@ -2857,7 +2857,7 @@ FIoStatus FOnDemandInstallCache::Defrag(
 		// Append a critical error entry to clear the cache at next startup
 		FCasJournal::FTransaction Transaction = FCasJournal::Begin(GetJournalFilename());
 		Transaction.CriticalError(FCasJournal::EErrorCode::DefragOutOfDiskSpace);
-		if (FIoStatus JournalStatus = FCasJournal::Commit(MoveTemp(Transaction)); Status.IsOk() == false)
+		if (FIoStatus JournalStatus = FCasJournal::Commit(MoveTemp(Transaction)); JournalStatus.IsOk() == false)
 		{
 			UE_LOG(LogIoStoreOnDemand, Error, TEXT("Failed to append critical error to journal, reason '%s'"),
 				*JournalStatus.ToString());
