Chaos - change log verbosity when creating debug draw (avoid commandlet spam)

#rnx

diff --git a/Engine/Source/Runtime/Experimental/ChaosSolverEngine/Private/Chaos/ChaosDebugDrawSubsystem.cpp b/Engine/Source/Runtime/Experimental/ChaosSolverEngine/Private/Chaos/ChaosDebugDrawSubsystem.cpp
--- a/Engine/Source/Runtime/Experimental/ChaosSolverEngine/Private/Chaos/ChaosDebugDrawSubsystem.cpp
+++ b/Engine/Source/Runtime/Experimental/ChaosSolverEngine/Private/Chaos/ChaosDebugDrawSubsystem.cpp
@@ -284,7 +284,14 @@ bool UChaosDebugDrawSubsystem::ShouldCreateSubsystem(UObject* Outer) const
 		const bool bIsWorldTypeSupported = World->IsPreviewWorld() ? Chaos::CVars::bChaosDebugDraw_PreviewWorldsEnabled : true;
 		const bool bCreateDebugDraw = !IsRunningCommandlet() && bIsWorldTypeSupported;
 
-		UE_LOGFMT(LogChaosDD, Log, "{0} Chaos Debug Draw Scene for world {1}", bCreateDebugDraw ? TEXT("Creating") : TEXT("Not creating"), World->GetName());
+		if (bCreateDebugDraw)
+		{
+			UE_LOGFMT(LogChaosDD, Log, "Creating Chaos Debug Draw Scene for world {0}",World->GetName());
+		}
+		else
+		{
+			UE_LOGFMT(LogChaosDD, Verbose, "Not creating Chaos Debug Draw Scene for world {0}", World->GetName());
+		}
 		
 		return bCreateDebugDraw;
 	}
