Iris
-Implemented validation of class type when serializing object references

diff --git a/Engine/Plugins/Runtime/ReplicationSystemTestPlugin/Source/Private/Tests/ReplicationSystem/ReplicatedTestObject.cpp b/Engine/Plugins/Runtime/ReplicationSystemTestPlugin/Source/Private/Tests/ReplicationSystem/ReplicatedTestObject.cpp
--- a/Engine/Plugins/Runtime/ReplicationSystemTestPlugin/Source/Private/Tests/ReplicationSystem/ReplicatedTestObject.cpp
+++ b/Engine/Plugins/Runtime/ReplicationSystemTestPlugin/Source/Private/Tests/ReplicationSystem/ReplicatedTestObject.cpp
@@ -368,6 +368,9 @@ void UTestReplicatedIrisObjectWithObjectReference::GetLifetimeReplicatedProps( T
 	DOREPLIFETIME(ThisClass, RawObjectPtrRef);
 	DOREPLIFETIME(ThisClass, WeakObjectPtrObjectRef);
 	DOREPLIFETIME(ThisClass, SoftObjectPtrRef);
+	DOREPLIFETIME(ThisClass, TypedRawObjectPtrRef);
+	DOREPLIFETIME(ThisClass, TypedWeakObjectPtrObjectRef);
+	DOREPLIFETIME(ThisClass, TypedSoftObjectPtrRef);
 }
 
 UTestReplicatedIrisObjectWithObjectReference::UTestReplicatedIrisObjectWithObjectReference()
diff --git a/Engine/Plugins/Runtime/ReplicationSystemTestPlugin/Source/Private/Tests/ReplicationSystem/ReplicatedTestObject.h b/Engine/Plugins/Runtime/ReplicationSystemTestPlugin/Source/Private/Tests/ReplicationSystem/ReplicatedTestObject.h
--- a/Engine/Plugins/Runtime/ReplicationSystemTestPlugin/Source/Private/Tests/ReplicationSystem/ReplicatedTestObject.h
+++ b/Engine/Plugins/Runtime/ReplicationSystemTestPlugin/Source/Private/Tests/ReplicationSystem/ReplicatedTestObject.h
@@ -386,6 +386,15 @@ public:
 	UPROPERTY(Transient, Replicated)
 	TSoftObjectPtr<UObject> SoftObjectPtrRef;
 
+	UPROPERTY(Transient, Replicated)
+	TObjectPtr<UTestReplicatedIrisObjectWithObjectReference> TypedRawObjectPtrRef;
+
+	UPROPERTY(Transient, Replicated)
+	TWeakObjectPtr<UTestReplicatedIrisObjectWithObjectReference> TypedWeakObjectPtrObjectRef;
+
+	UPROPERTY(Transient, Replicated)
+	TSoftObjectPtr<UTestReplicatedIrisObjectWithObjectReference> TypedSoftObjectPtrRef;
+
 public:
 	// Network data only for test
 	TArray<UE::Net::FReplicationFragment*> ReplicationFragments;
diff --git a/Engine/Plugins/Runtime/ReplicationSystemTestPlugin/Source/Private/Tests/Serialization/TestObjectNetSerializer.cpp b/Engine/Plugins/Runtime/ReplicationSystemTestPlugin/Source/Private/Tests/Serialization/TestObjectNetSerializer.cpp
--- a/Engine/Plugins/Runtime/ReplicationSystemTestPlugin/Source/Private/Tests/Serialization/TestObjectNetSerializer.cpp
+++ b/Engine/Plugins/Runtime/ReplicationSystemTestPlugin/Source/Private/Tests/Serialization/TestObjectNetSerializer.cpp
@@ -220,4 +220,45 @@ UE_NET_TEST_FIXTURE(FReplicationSystemServerClientTestFixture, TestGarbageObject
 	Server->UpdateAndSend({Client});
 }
 
+UE_NET_TEST_FIXTURE(FReplicationSystemServerClientTestFixture, TestTypedObjectReference)
+{
+	// Add a client
+	FReplicationSystemTestClient* Client = CreateClient();
+
+	// Object on which we conduct the test
+	UTestReplicatedIrisObjectWithObjectReference* ServerObject = Server->CreateObject<UTestReplicatedIrisObjectWithObjectReference>();
+
+	// Spawn different types of objects on server
+	UTestReplicatedIrisObjectWithObjectReference* ReplicatedIrisObjectWithObjectReference = Server->CreateObject<UTestReplicatedIrisObjectWithObjectReference>();
+	UTestReplicatedIrisObject* ReplicatedIrisObject = Server->CreateObject<UTestReplicatedIrisObject>();
+
+	// Untyped should always work
+	ServerObject->RawObjectPtrRef = ReplicatedIrisObject;
+	ServerObject->WeakObjectPtrObjectRef = ReplicatedIrisObject;
+	ServerObject->SoftObjectPtrRef = ReplicatedIrisObject;
+
+	// HACK: to set invalid types.
+ 	// We expect this to serialize as an invalid object reference
+	*(TObjectPtr<UTestReplicatedIrisObject>*)(&ServerObject->TypedRawObjectPtrRef) = ReplicatedIrisObject;
+	*(TWeakObjectPtr<UTestReplicatedIrisObject>*)(&ServerObject->TypedWeakObjectPtrObjectRef) = ReplicatedIrisObject;
+	*(TSoftObjectPtr<UTestReplicatedIrisObject>*)(&ServerObject->TypedSoftObjectPtrRef) = ReplicatedIrisObject;
+
+	// Send and deliver packet
+	Server->UpdateAndSend({Client});
+
+	UTestReplicatedIrisObjectWithObjectReference* ClientObject = Client->GetObjectAs<UTestReplicatedIrisObjectWithObjectReference>(ServerObject->NetRefHandle);
+	UE_NET_ASSERT_NE(ClientObject, nullptr);
+
+	// We expect these to be resolvable
+	UE_NET_ASSERT_NE(ClientObject->RawObjectPtrRef.Get(), nullptr);
+	UE_NET_ASSERT_NE(ClientObject->WeakObjectPtrObjectRef.Get(), nullptr);
+	UE_NET_ASSERT_NE(ClientObject->SoftObjectPtrRef.Get(), nullptr);
+
+	// We expect these to be invalid
+	UE_NET_ASSERT_EQ(ClientObject->TypedRawObjectPtrRef.Get(), nullptr);
+	UE_NET_ASSERT_EQ(ClientObject->TypedWeakObjectPtrObjectRef.Get(), nullptr);
+	UE_NET_ASSERT_EQ(ClientObject->TypedSoftObjectPtrRef.Get(), nullptr);
+}
+
+
 }
diff --git a/Engine/Source/Runtime/Experimental/Iris/Core/Private/Iris/ReplicationState/DefaultPropertyNetSerializerInfos.cpp b/Engine/Source/Runtime/Experimental/Iris/Core/Private/Iris/ReplicationState/DefaultPropertyNetSerializerInfos.cpp
--- a/Engine/Source/Runtime/Experimental/Iris/Core/Private/Iris/ReplicationState/DefaultPropertyNetSerializerInfos.cpp
+++ b/Engine/Source/Runtime/Experimental/Iris/Core/Private/Iris/ReplicationState/DefaultPropertyNetSerializerInfos.cpp
@@ -56,6 +56,16 @@ struct FObjectPropertyNetSerializerInfo : public FPropertyNetSerializerInfo
 	{
 		return &UE_NET_GET_SERIALIZER(FObjectNetSerializer);
 	}
+
+	virtual bool CanUseDefaultConfig(const FProperty*) const override { return false; }
+
+	virtual FNetSerializerConfig* BuildNetSerializerConfig(void* NetSerializerConfigBuffer, const FProperty* Property) const override
+	{
+		FObjectNetSerializerConfig* Config = new (NetSerializerConfigBuffer) FObjectNetSerializerConfig();
+		const FObjectPropertyBase* ObjectPtrPropertyBase = static_cast<const FObjectPropertyBase*>(Property);
+		Config->PropertyClass = ObjectPtrPropertyBase->PropertyClass;
+		return Config;
+	}
 };
 UE_NET_IMPLEMENT_NETSERIALIZER_INFO(FObjectPropertyNetSerializerInfo);
 
@@ -77,10 +87,49 @@ struct FObjectPtrPropertyNetSerializerInfo : public FPropertyNetSerializerInfo
 	{
 		return &UE_NET_GET_SERIALIZER(FObjectPtrNetSerializer);
 	}
+
+	virtual bool CanUseDefaultConfig(const FProperty*) const override { return false; }
+
+	virtual FNetSerializerConfig* BuildNetSerializerConfig(void* NetSerializerConfigBuffer, const FProperty* Property) const override
+	{
+		FObjectPtrNetSerializerConfig* Config = new (NetSerializerConfigBuffer) FObjectPtrNetSerializerConfig();
+		const FObjectPropertyBase* ObjectPtrPropertyBase = static_cast<const FObjectPropertyBase*>(Property);
+		Config->PropertyClass = ObjectPtrPropertyBase->PropertyClass;
+		return Config;
+	}
 };
 UE_NET_IMPLEMENT_NETSERIALIZER_INFO(FObjectPtrPropertyNetSerializerInfo);
 
-UE_NET_IMPLEMENT_SIMPLE_NETSERIALIZER_INFO(FWeakObjectProperty, FWeakObjectNetSerializer);
+struct FWeakObjectPtrPropertyNetSerializerInfo : public FPropertyNetSerializerInfo
+{
+	virtual const FFieldClass* GetPropertyTypeClass() const override
+	{
+		return FWeakObjectProperty::StaticClass();
+	}
+
+	virtual bool IsSupported(const FProperty* Property) const
+	{
+		return true;
+	}
+
+	virtual const FNetSerializer* GetNetSerializer(const FProperty* Property) const override
+	{
+		return &UE_NET_GET_SERIALIZER(FWeakObjectNetSerializer);
+	}
+
+	virtual bool CanUseDefaultConfig(const FProperty*) const override { return false; }
+
+	virtual FNetSerializerConfig* BuildNetSerializerConfig(void* NetSerializerConfigBuffer, const FProperty* Property) const override
+	{
+		FWeakObjectNetSerializerConfig* Config = new (NetSerializerConfigBuffer) FWeakObjectNetSerializerConfig();
+		const FObjectPropertyBase* ObjectPtrPropertyBase = static_cast<const FObjectPropertyBase*>(Property);
+		Config->PropertyClass = ObjectPtrPropertyBase->PropertyClass;
+		return Config;
+	}
+};
+UE_NET_IMPLEMENT_NETSERIALIZER_INFO(FWeakObjectPtrPropertyNetSerializerInfo);
+
+
 UE_NET_IMPLEMENT_SIMPLE_NETSERIALIZER_INFO(FSoftObjectProperty, FSoftObjectNetSerializer);
 UE_NET_IMPLEMENT_SIMPLE_NETSERIALIZER_INFO(FSoftClassProperty, FSoftObjectNetSerializer);
 UE_NET_IMPLEMENT_SIMPLE_NETSERIALIZER_INFO(FFieldPathProperty, FFieldPathNetSerializer);
@@ -183,7 +232,7 @@ struct FScriptInterfacePropertyNetSerializerInfo : public FPropertyNetSerializer
 	{
 		FScriptInterfaceNetSerializerConfig* Config = new (NetSerializerConfigBuffer) FScriptInterfaceNetSerializerConfig();
 		const FInterfaceProperty* InterfaceProperty = static_cast<const FInterfaceProperty*>(Property);
-		Config->InterfaceClass = InterfaceProperty->InterfaceClass;
+		Config->PropertyClass = InterfaceProperty->InterfaceClass;
 		return Config;
 	}
 };
@@ -238,7 +287,7 @@ void RegisterDefaultPropertyNetSerializerInfos()
 	// Object and field types
 	UE_NET_REGISTER_NETSERIALIZER_INFO(FObjectPropertyNetSerializerInfo);
 	UE_NET_REGISTER_NETSERIALIZER_INFO(FObjectPtrPropertyNetSerializerInfo);
-	UE_NET_REGISTER_NETSERIALIZER_INFO(FWeakObjectProperty);
+	UE_NET_REGISTER_NETSERIALIZER_INFO(FWeakObjectPtrPropertyNetSerializerInfo);
 	UE_NET_REGISTER_NETSERIALIZER_INFO(FScriptInterfacePropertyNetSerializerInfo);
 	UE_NET_REGISTER_NETSERIALIZER_INFO(FSoftObjectProperty);
 	UE_NET_REGISTER_NETSERIALIZER_INFO(FFieldPathProperty);
diff --git a/Engine/Source/Runtime/Experimental/Iris/Core/Private/Iris/Serialization/ObjectNetSerializer.cpp b/Engine/Source/Runtime/Experimental/Iris/Core/Private/Iris/Serialization/ObjectNetSerializer.cpp
--- a/Engine/Source/Runtime/Experimental/Iris/Core/Private/Iris/Serialization/ObjectNetSerializer.cpp
+++ b/Engine/Source/Runtime/Experimental/Iris/Core/Private/Iris/Serialization/ObjectNetSerializer.cpp
@@ -132,6 +132,7 @@ protected:
 	// Raw pointer getters/setters
 	static UObject* GetRawPtrValue(UObject* Value) { return Value; }
 	static TObjectPtr<UObject> GetObjectPtrValue(UObject* Value) { return Value; }
+	static UObject* GetValidatedRawPtrValue(UObject* Value, const UClass* PropertyClass);
 
 	static void SetRawPtrValue(UObject*& Dst, UObject* Value) { Dst = Value; }
 
@@ -341,13 +342,14 @@ void FObjectNetSerializerBase<T>::Quantize(FNetSerializationContext& Context, co
 template<typename T>
 void FObjectNetSerializerBase<T>::Dequantize(FNetSerializationContext& Context, const FNetDequantizeArgs& Args)
 {
+	const FObjectNetSerializerConfig& Config = *static_cast<const FObjectNetSerializerConfig*>(Args.NetSerializerConfig);
 	const QuantizedType& Source = *reinterpret_cast<const QuantizedType*>(Args.Source);
 	T& Target = *reinterpret_cast<T*>(Args.Target);
-	
+
 	if (Source.IsNetReference())
 	{
 		UObject* DequantizedObject = ResolveObjectReference(Context, Source.NetReference);
-		SetRawPtrValue(Target, DequantizedObject);
+		SetRawPtrValue(Target, GetValidatedRawPtrValue(DequantizedObject, Config.PropertyClass));
 		return;
 	}
 
@@ -418,6 +420,20 @@ void FObjectNetSerializerBase<T>::FreeDynamicState(FNetSerializationContext& Con
 	SourceValue.FreeRemoteReference(Context, Args);
 }
 
+template<typename T>
+UObject* FObjectNetSerializerBase<T>::GetValidatedRawPtrValue(UObject* Value, const UClass* PropertyClass)
+{
+	if (PropertyClass && Value && !Value->IsA(PropertyClass))
+	{
+		UE_LOG(LogIris, Warning, TEXT("Forged object: got %s, expecting %s"), *Value->GetFullName(), *PropertyClass->GetFullName());
+		return nullptr;
+	}
+	else
+	{
+		return Value;
+	}
+}
+
 template<typename T>
 UObject* FObjectNetSerializerBase<T>::ResolveObjectReference(FNetSerializationContext& Context, const FNetObjectReference& NetObjectReference)
 {
@@ -481,7 +497,8 @@ void FScriptInterfaceNetSerializer::Dequantize(FNetSerializationContext& Context
 	if (Object != nullptr)
 	{
 		const FScriptInterfaceNetSerializerConfig& Config = *static_cast<const FScriptInterfaceNetSerializerConfig*>(Args.NetSerializerConfig);
-		void* Interface = Object->GetInterfaceAddress(Config.InterfaceClass);
+		UClass* ConfigInterfaceClass = Config.PropertyClass;
+		void* Interface = Object->GetInterfaceAddress(ConfigInterfaceClass);
 
 		Target.SetObject(Object);
 		Target.SetInterface(Interface);
@@ -498,9 +515,10 @@ bool FScriptInterfaceNetSerializer::Validate(FNetSerializationContext& Context,
 	const SourceType& Value = *reinterpret_cast<const SourceType*>(Args.Source);
 	if (UObject* Object = Value.GetObject())
 	{
-		void* Interface = Value.GetInterface();
+		void* Interface = Value.GetInterface();		
 		const FScriptInterfaceNetSerializerConfig& Config = *static_cast<const FScriptInterfaceNetSerializerConfig*>(Args.NetSerializerConfig);
-		if (Interface != Object->GetInterfaceAddress(Config.InterfaceClass))
+		UClass* ConfigInterfaceClass = Config.PropertyClass;
+		if (Interface != Object->GetInterfaceAddress(ConfigInterfaceClass))
 		{
 			return false;
 		}
diff --git a/Engine/Source/Runtime/Experimental/Iris/Core/Public/Iris/Serialization/ObjectNetSerializer.h b/Engine/Source/Runtime/Experimental/Iris/Core/Public/Iris/Serialization/ObjectNetSerializer.h
--- a/Engine/Source/Runtime/Experimental/Iris/Core/Public/Iris/Serialization/ObjectNetSerializer.h
+++ b/Engine/Source/Runtime/Experimental/Iris/Core/Public/Iris/Serialization/ObjectNetSerializer.h
@@ -17,27 +17,27 @@ USTRUCT()
 struct FObjectNetSerializerConfig : public FNetSerializerConfig
 {
 	GENERATED_BODY()
+
+	UPROPERTY()
+	TObjectPtr<UClass> PropertyClass = nullptr;
 };
 
 USTRUCT()
-struct FObjectPtrNetSerializerConfig : public FNetSerializerConfig
+struct FObjectPtrNetSerializerConfig : public FObjectNetSerializerConfig
 {
 	GENERATED_BODY()
 };
 
 USTRUCT()
-struct FWeakObjectNetSerializerConfig : public FNetSerializerConfig
+struct FWeakObjectNetSerializerConfig : public FObjectNetSerializerConfig
 {
 	GENERATED_BODY()
 };
 
 USTRUCT()
-struct FScriptInterfaceNetSerializerConfig : public FNetSerializerConfig
+struct FScriptInterfaceNetSerializerConfig : public FObjectNetSerializerConfig
 {
 	GENERATED_BODY()
-
-	UPROPERTY()
-	TObjectPtr<UClass> InterfaceClass = nullptr;
 };
 
 namespace UE::Net
