Prevent rare case where cancelling an IoRequest could result in a crash or the request never completing.

#rnx

### Problem
- When the IAX http thread completes a request, the request is destroyed and returned to the slab allocator but the higher level owner FChunkRequest would retain the original handle.
- There is a small slice of time between the http request completing and the IAS backend completing the IoRequest.
- If the IoRequest was cancelled during this sliver of time then the FChunkRequest for the IoRequest would still be active and would try to mark it's http request as cancelled as the handle was still valid.
- At this point one of three things could occur
-- The handle is unallocated in the slab allocator, this is effectively a NOP
-- The handle was originally allocated in an additional slab, due to overflowing the initial size of the allocator but this additional slab has since been destroyed. The handle either points to invalid memory (crash) or has been reallocated elsewhere (memory stomp)
-- The handle has been reallocated to a different FChunkRequest. This request would be marked as cancelled so the IAS backend will not mark it with the IoDispatcher as completed as it believes that the IoDispatcher will already deal with that as part of the (fake) cancellation call. Since that cancellation was not real no clean up will ever be run and the IoRequest will never be considered complete.

### Fix
- For now once the http request completed callback is triggered for a  FChunkRequest we invalidate it's http request. If it is cancelled at a later date it will not try to access a http request in turn.
- The cleaning up of the handle must be done while FChunkRequests::Mutex is locked as that will prevent the IoRequest cancellation being run on a different thread at the same time.
- The fix can be disabled by setting the cvar 'ias.EnableRequestHandleClear' to false.

diff --git a/Engine/Source/Runtime/Experimental/IoStore/OnDemand/Private/OnDemandIoDispatcherBackend.cpp b/Engine/Source/Runtime/Experimental/IoStore/OnDemand/Private/OnDemandIoDispatcherBackend.cpp
--- a/Engine/Source/Runtime/Experimental/IoStore/OnDemand/Private/OnDemandIoDispatcherBackend.cpp
+++ b/Engine/Source/Runtime/Experimental/IoStore/OnDemand/Private/OnDemandIoDispatcherBackend.cpp
@@ -153,6 +153,13 @@ static FAutoConsoleVariableRef CVar_WriteOnlyDecoding(
 	TEXT("Enables the use of 'WriteOnly' flag when decoding to buffers with the 'HardwareTargetBuffer' flag")
 );
 
+bool GIasEnableRequestHandleClear = true;
+static FAutoConsoleVariableRef CVar_EnableRequestHandleClear(
+	TEXT("ias.EnableRequestHandleClear"),
+	GIasEnableRequestHandleClear,
+	TEXT("When enabled FChunkRequest::HttpRequest will be cleared when the request callback is triggered")
+);
+
 #if UE_ALLOW_DISABLE_CANCELLING
 bool GIasCancelationEnabled = true;
 static FAutoConsoleVariableRef CVar_IasCancelationEnabled(
@@ -600,6 +607,15 @@ class FOnDemandIoBackend final
 			return false;
 		}
 
+		void OnHttpRequestCompleted(FChunkRequest* Request)
+		{
+			if (GIasEnableRequestHandleClear)
+			{
+				FScopeLock _(&Mutex);
+				Request->HttpRequest = 0;
+			}
+		}
+
 		FIoChunkId GetChunkId(FChunkRequest* Request)
 		{
 			FScopeLock _(&Mutex);
@@ -844,6 +860,7 @@ void FOnDemandIoBackend::IssueHttpRequest(FChunkRequest* ChunkRequest)
 	ChunkRequest->HttpRequest = HttpClient.IssueRequest(ChunkRequest->Params.ChunkInfo, ChunkRequest->Params.ChunkRange, ChunkRequest->Priority,
 		[this, ChunkRequest](uint32 StatusCode, FStringView ErrorReason, FIoBuffer&& Data)
 		{
+			ChunkRequests.OnHttpRequestCompleted(ChunkRequest); // Sets HttpRequest to 0
 			ChunkRequest->Chunk = MoveTemp(Data);
 
 			if (IsHttpStatusOk(StatusCode))
