Synchronize access to IAD HTTP handles.

- this change fixes a rare race condition when cancelling IAD HTTP requests.

#rnx

diff --git a/Engine/Source/Runtime/Experimental/IoStore/OnDemand/Private/OnDemandContentInstaller.cpp b/Engine/Source/Runtime/Experimental/IoStore/OnDemand/Private/OnDemandContentInstaller.cpp
--- a/Engine/Source/Runtime/Experimental/IoStore/OnDemand/Private/OnDemandContentInstaller.cpp
+++ b/Engine/Source/Runtime/Experimental/IoStore/OnDemand/Private/OnDemandContentInstaller.cpp
@@ -368,9 +368,13 @@ void FOnDemandContentInstaller::UpdateInstallRequestPriority(FSharedInternalInst
 					{
 						PendingHttpRequest.IoHandle.UpdatePriorty(Request.Priority);
 					}
-					else if (PendingHttpRequest.Handle != nullptr)
+					else 
 					{
-						HttpClient.ReprioritizeRequest(PendingHttpRequest.Handle, Request.Priority);
+						TUniqueLock HttpLock(HttpMutex);
+						if (PendingHttpRequest.Handle != nullptr)
+						{
+							HttpClient.ReprioritizeRequest(PendingHttpRequest.Handle, Request.Priority);
+						}
 					}
 				}
 			}
@@ -768,6 +772,10 @@ void FOnDemandContentInstaller::ExecuteInstallRequest(FRequest& Request, bool bR
 				[this, &Request, &HttpRequest](uint32 HttpStatusCode, FStringView ErrorReason, FIoBuffer&& Chunk)
 				{
 					TRACE_CPUPROFILER_EVENT_SCOPE(FOnDemandContentInstaller::Callback);
+					{
+						TUniqueLock Lock(HttpMutex);
+						HttpRequest.Handle = nullptr;
+					}
 					InstallerPipe.Launch(
 						TEXT("ProcessIoStoreOnDemandDownloadedChunk"),
 						[this, &Request, &HttpRequest, HttpStatusCode, ErrorReason = FString(ErrorReason), Chunk = MoveTemp(Chunk)]() mutable
@@ -899,7 +907,6 @@ void FOnDemandContentInstaller::ProcessDownloadedChunk(
 	TRACE_CPUPROFILER_EVENT_SCOPE(FOnDemandContentInstaller::ProcessDownloadedChunk);
 
 	FRequest::FInstall& InstallRequest			= Request.AsInstall();
-	HttpRequest.Handle							= nullptr;
 	FSharedOnDemandContainer& Container			= InstallRequest.ResolvedChunks[HttpRequest.ContainerIndex].Container;
 	const FOnDemandChunkEntry& ChunkEntry		= Container->ChunkEntries[HttpRequest.EntryIndex];
 	const FIoChunkId& ChunkId					= Container->ChunkIds[HttpRequest.EntryIndex];
@@ -963,17 +970,28 @@ void FOnDemandContentInstaller::ProcessDownloadedChunk(
 		if (Request.ErrorCode != EIoErrorCode::Unknown)
 		{
 			int32 NumCancelled = 0;
-			for (FChunkHttpRequestHandle& PendingHttpRequest : InstallRequest.HttpRequestHandles)
+			if (FHttpIoDispatcher::IsInitialized())
 			{
-				if (PendingHttpRequest.IoHandle.IsValid())
+				for (FChunkHttpRequestHandle& PendingHttpRequest : InstallRequest.HttpRequestHandles)
 				{
-					PendingHttpRequest.IoHandle.Cancel();
-					++NumCancelled;
+					if (PendingHttpRequest.IoHandle.IsValid())
+					{
+						PendingHttpRequest.IoHandle.Cancel();
+						++NumCancelled;
+					}
 				}
-				else if (PendingHttpRequest.Handle != nullptr)
+			}
+			else
+			{
+				TUniqueLock Lock(HttpMutex);
+				for (FChunkHttpRequestHandle& PendingHttpRequest : InstallRequest.HttpRequestHandles)
 				{
-					HttpClient.CancelRequest(PendingHttpRequest.Handle);
-					++NumCancelled;
+					if (PendingHttpRequest.Handle != nullptr)
+					{
+						HttpClient.CancelRequest(PendingHttpRequest.Handle);
+						PendingHttpRequest.Handle = nullptr;
+						++NumCancelled;
+					}
 				}
 			}
 
diff --git a/Engine/Source/Runtime/Experimental/IoStore/OnDemand/Private/OnDemandContentInstaller.h b/Engine/Source/Runtime/Experimental/IoStore/OnDemand/Private/OnDemandContentInstaller.h
--- a/Engine/Source/Runtime/Experimental/IoStore/OnDemand/Private/OnDemandContentInstaller.h
+++ b/Engine/Source/Runtime/Experimental/IoStore/OnDemand/Private/OnDemandContentInstaller.h
@@ -239,6 +239,7 @@ private:
 	UE::Tasks::FPipe				InstallerPipe;
 
 	FMutex							Mutex;
+	FMutex							HttpMutex;
 	FRequestAllocator				RequestAllocator;
 	TArray<FRequest*>				RequestQueue;
 	FRequest*						CurrentRequest = nullptr;
