Constraints: make sure that the controller is using the appropriate world when needed

opening a new asset editor will register new world(s) that can update the current constraint manager controller's world.
if that happens when adding a constraint via a Sequencer evaluation, it can result in the constraint being associated with an invalid world (preview, thumbnail, etc.), which can cause future issues when registering/unregistering constraints in that invalid world.

#rnx

diff --git a/Engine/Plugins/Animation/ControlRig/Source/ControlRig/Private/Sequencer/MovieSceneControlRigParameterTemplate.cpp b/Engine/Plugins/Animation/ControlRig/Source/ControlRig/Private/Sequencer/MovieSceneControlRigParameterTemplate.cpp
--- a/Engine/Plugins/Animation/ControlRig/Source/ControlRig/Private/Sequencer/MovieSceneControlRigParameterTemplate.cpp
+++ b/Engine/Plugins/Animation/ControlRig/Source/ControlRig/Private/Sequencer/MovieSceneControlRigParameterTemplate.cpp
@@ -1206,10 +1206,9 @@ static UTickableConstraint* CreateConstraintIfNeeded(UWorld* InWorld, const FCon
 	// it's possible that we have it but it's not in the manager, due to manager not being saved with it (due to spawning or undo/redo).
 	if (InWorld)
 	{
-		const FConstraintsManagerController& Controller = FConstraintsManagerController::Get(InWorld);
-		if (Controller.GetConstraint(Constraint->ConstraintID) == nullptr)
+		if (!FConstraintsManagerController::Get(InWorld).GetConstraint(Constraint->ConstraintID))
 		{
-			Controller.AddConstraint(InConstraintValue.Constraint.Get());
+			FConstraintsManagerController::Get(InWorld).AddConstraint(InConstraintValue.Constraint.Get());
 			//need to reconstuct channels here.. note this is now lazy and so will recreate it next time view requests it
 			//but only do it if the control rig has a valid world it may not for example in PIE
 			if (InSection->GetControlRig() && InSection->GetControlRig()->GetWorld())
diff --git a/Engine/Plugins/Animation/ControlRig/Source/ControlRig/Private/Sequencer/MovieSceneControlRigSystem.cpp b/Engine/Plugins/Animation/ControlRig/Source/ControlRig/Private/Sequencer/MovieSceneControlRigSystem.cpp
--- a/Engine/Plugins/Animation/ControlRig/Source/ControlRig/Private/Sequencer/MovieSceneControlRigSystem.cpp
+++ b/Engine/Plugins/Animation/ControlRig/Source/ControlRig/Private/Sequencer/MovieSceneControlRigSystem.cpp
@@ -147,10 +147,9 @@ UTickableConstraint* CreateConstraintIfNeeded(UWorld* InWorld, UTickableConstrai
 	// it's possible that we have it but it's not in the manager, due to manager not being saved with it (due to spawning or undo/redo).
 	if (InWorld)
 	{
-		const FConstraintsManagerController& Controller = FConstraintsManagerController::Get(InWorld);
-		if (Controller.GetConstraint(Constraint->ConstraintID) == nullptr)
+		if (!FConstraintsManagerController::Get(InWorld).GetConstraint(Constraint->ConstraintID))
 		{
-			Controller.AddConstraint(Constraint);
+			FConstraintsManagerController::Get(InWorld).AddConstraint(Constraint);
 			//need to reconstruct channels here.. note this is now lazy and so will recreate it next time view requests it
 			//but only do it if the control rig has a valid world it may not for example in PIE
 			if (InSection->GetControlRig() && InSection->GetControlRig()->GetWorld())
diff --git a/Engine/Source/Editor/MovieSceneTools/Private/Constraints/TransformConstraintChannelInterface.cpp b/Engine/Source/Editor/MovieSceneTools/Private/Constraints/TransformConstraintChannelInterface.cpp
--- a/Engine/Source/Editor/MovieSceneTools/Private/Constraints/TransformConstraintChannelInterface.cpp
+++ b/Engine/Source/Editor/MovieSceneTools/Private/Constraints/TransformConstraintChannelInterface.cpp
@@ -183,12 +183,11 @@ void ITransformConstraintChannelInterface::UnregisterConstraints(IMovieSceneCons
 		return;
 	}
 
-	FConstraintsManagerController& Controller = FConstraintsManagerController::Get(InWorld);
 	for (const FConstraintAndActiveChannel& ConstraintAndActiveChannel : Constraints)
 	{
 		if (UTickableConstraint* Constraint = ConstraintAndActiveChannel.GetConstraint())
 		{
-			Controller.UnregisterConstraint(Constraint);
+			FConstraintsManagerController::Get(InWorld).UnregisterConstraint(Constraint);
 		}
 	}		
 }
diff --git a/Engine/Source/Editor/MovieSceneTools/Private/TrackEditors/TransformTrackEditor.cpp b/Engine/Source/Editor/MovieSceneTools/Private/TrackEditors/TransformTrackEditor.cpp
--- a/Engine/Source/Editor/MovieSceneTools/Private/TrackEditors/TransformTrackEditor.cpp
+++ b/Engine/Source/Editor/MovieSceneTools/Private/TrackEditors/TransformTrackEditor.cpp
@@ -1899,12 +1899,12 @@ void F3DTransformTrackEditor::HandleConstraintPropertyChanged(UTickableTransform
 void F3DTransformTrackEditor::ClearOutConstraintDelegates() 
 {
 	UWorld* World = GCurrentLevelEditingViewportClient ? GCurrentLevelEditingViewportClient->GetWorld() : nullptr;
-	FConstraintsManagerController& Controller = FConstraintsManagerController::Get(World);
+	FConstraintsManagerNotifyDelegate& NotifyDelegate = FConstraintsManagerController::Get(World).GetNotifyDelegate();
 	for (FDelegateHandle& Handle : ConstraintHandlesToClear)
 	{
 		if (Handle.IsValid())
 		{
-			Controller.GetNotifyDelegate().Remove(Handle);
+			NotifyDelegate.Remove(Handle);
 		}
 	}
 	ConstraintHandlesToClear.Reset();
diff --git a/Engine/Source/Runtime/Experimental/Animation/Constraints/Private/ConstraintsActor.cpp b/Engine/Source/Runtime/Experimental/Animation/Constraints/Private/ConstraintsActor.cpp
--- a/Engine/Source/Runtime/Experimental/Animation/Constraints/Private/ConstraintsActor.cpp
+++ b/Engine/Source/Runtime/Experimental/Animation/Constraints/Private/ConstraintsActor.cpp
@@ -78,10 +78,9 @@ void AConstraintsActor::RegisterConstraintsTickFunctions() const
 			});
 			
 			// ensure registration, we now do this by adding it to the controller (which registers and sets it with the subsystem)
-			FConstraintsManagerController& Controller = FConstraintsManagerController::Get(GetWorld());
 			for (const TObjectPtr<UTickableConstraint>& ConstPtr: ConstraintsManager->Constraints)
 			{
-				Controller.AddConstraint(ConstPtr);
+				FConstraintsManagerController::Get(GetWorld()).AddConstraint(ConstPtr);
 			}
 		}
 	}
diff --git a/Engine/Source/Runtime/Experimental/Animation/Constraints/Private/ConstraintsManager.cpp b/Engine/Source/Runtime/Experimental/Animation/Constraints/Private/ConstraintsManager.cpp
--- a/Engine/Source/Runtime/Experimental/Animation/Constraints/Private/ConstraintsManager.cpp
+++ b/Engine/Source/Runtime/Experimental/Animation/Constraints/Private/ConstraintsManager.cpp
@@ -251,12 +251,13 @@ void UConstraintsManager::OnActorDestroyed(AActor* InActor)
 
 		if (!IndicesToRemove.IsEmpty())
 		{
-			FConstraintsManagerController& Controller = FConstraintsManagerController::Get(InActor->GetWorld());
 			if (FConstraintsManagerController::bDoNotRemoveConstraint == false)
 			{
+				UWorld* ActorWorld = InActor->GetWorld();
+				constexpr bool bDoNotCompensate = true;
 				for (int32 Index = IndicesToRemove.Num() - 1; Index >= 0; --Index)
 				{
-					Controller.RemoveConstraint(Index,/*do not compensate*/ true);
+					FConstraintsManagerController::Get(ActorWorld).RemoveConstraint(Index, bDoNotCompensate);
 				}
 			}
 		}
diff --git a/Engine/Source/Runtime/Experimental/Animation/Constraints/Private/Transform/DependencyBuilder.cpp b/Engine/Source/Runtime/Experimental/Animation/Constraints/Private/Transform/DependencyBuilder.cpp
--- a/Engine/Source/Runtime/Experimental/Animation/Constraints/Private/Transform/DependencyBuilder.cpp
+++ b/Engine/Source/Runtime/Experimental/Animation/Constraints/Private/Transform/DependencyBuilder.cpp
@@ -171,7 +171,6 @@ void FDependencyBuilder::BuildSelfDependencies(UWorld* InWorld, UTickableTransfo
 
 	const UObject* SelfTarget = ChildTarget;
 
-	const FConstraintsManagerController& Controller = FConstraintsManagerController::Get(InWorld);
 	auto SelfTargetPredicate = [InConstraint, SelfTarget](const ConstraintWeakPtr& Constraint)
 	{
 		const UTickableTransformConstraint* TransformConstraint = Cast<UTickableTransformConstraint>(Constraint.Get());
@@ -183,7 +182,7 @@ void FDependencyBuilder::BuildSelfDependencies(UWorld* InWorld, UTickableTransfo
 		const UObject* ChildTarget = GetHandleTarget(TransformConstraint->ChildTRSHandle);
 		return ParentTarget == SelfTarget && ChildTarget == SelfTarget;
 	};
-	const TArray< ConstraintWeakPtr > SelfConstraints = Controller.GetConstraintsByPredicate(SelfTargetPredicate);
+	const TArray< ConstraintWeakPtr > SelfConstraints = FConstraintsManagerController::Get(InWorld).GetConstraintsByPredicate(SelfTargetPredicate);
 	
 	for (const ConstraintWeakPtr& SelfConstraint: SelfConstraints)
 	{
@@ -192,24 +191,24 @@ void FDependencyBuilder::BuildSelfDependencies(UWorld* InWorld, UTickableTransfo
 		// if the new handles depend on that constraint child then, TransformConstraint should tick before
 		if (ParentHandle->HasDirectDependencyWith(*TransformConstraint->ChildTRSHandle))
 		{
-			Controller.SetConstraintsDependencies(TransformConstraint->ConstraintID, InConstraint->ConstraintID);
+			FConstraintsManagerController::Get(InWorld).SetConstraintsDependencies(TransformConstraint->ConstraintID, InConstraint->ConstraintID);
 			LogDependency(SelfDependencyDesc, TransformConstraint->ChildTRSHandle, ParentHandle, TransformConstraint, InConstraint);
 		}
 		else if (ChildHandle->HasDirectDependencyWith(*TransformConstraint->ChildTRSHandle))
 		{
-			Controller.SetConstraintsDependencies(TransformConstraint->ConstraintID, InConstraint->ConstraintID);
+			FConstraintsManagerController::Get(InWorld).SetConstraintsDependencies(TransformConstraint->ConstraintID, InConstraint->ConstraintID);
 			LogDependency(SelfDependencyDesc, TransformConstraint->ChildTRSHandle, ChildHandle, TransformConstraint, InConstraint);
 		}
 
 		// if the TransformConstraint handles depend on the new constraint child then, TransformConstraint should tick after
 		if (TransformConstraint->ParentTRSHandle->HasDirectDependencyWith(*ChildHandle))
 		{
-			Controller.SetConstraintsDependencies(InConstraint->ConstraintID, TransformConstraint->ConstraintID);
+			FConstraintsManagerController::Get(InWorld).SetConstraintsDependencies(InConstraint->ConstraintID, TransformConstraint->ConstraintID);
 			LogDependency(SelfDependencyDesc, ChildHandle, TransformConstraint->ParentTRSHandle, InConstraint, TransformConstraint);
 		}
 		else if (TransformConstraint->ChildTRSHandle->HasDirectDependencyWith(*ChildHandle))
 		{
-			Controller.SetConstraintsDependencies(InConstraint->ConstraintID, TransformConstraint->ConstraintID);
+			FConstraintsManagerController::Get(InWorld).SetConstraintsDependencies(InConstraint->ConstraintID, TransformConstraint->ConstraintID);
 			LogDependency(SelfDependencyDesc, ChildHandle, TransformConstraint->ChildTRSHandle, InConstraint, TransformConstraint);
 		}
 	}
@@ -236,8 +235,6 @@ void FDependencyBuilder::BuildExternalDependencies(UWorld* InWorld, UTickableTra
 	{
 		return;
 	}
-
-	const FConstraintsManagerController& Controller = FConstraintsManagerController::Get(InWorld);
 	
 	{
 		// get all constraints acting on the same target 
@@ -251,7 +248,7 @@ void FDependencyBuilder::BuildExternalDependencies(UWorld* InWorld, UTickableTra
 			const UObject* OtherChildTarget = GetHandleTarget(TransformConstraint->ChildTRSHandle);
 			return OtherChildTarget && OtherChildTarget == ChildTarget;
 		};
-		TArray< ConstraintWeakPtr > ChildTargetParentConstraints = Controller.GetConstraintsByPredicate(SameChildTargetPredicate);
+		TArray< ConstraintWeakPtr > ChildTargetParentConstraints = FConstraintsManagerController::Get(InWorld).GetConstraintsByPredicate(SameChildTargetPredicate);
 
 		// store constraint index in this array
 		const int32 ConstrainIndex = ChildTargetParentConstraints.IndexOfByKey(InConstraint);
@@ -276,13 +273,13 @@ void FDependencyBuilder::BuildExternalDependencies(UWorld* InWorld, UTickableTra
 				// if the new handles depend on that constraint child then, TransformConstraint should tick before
 				if (ParentHandle->HasDirectDependencyWith(*TransformConstraint->ChildTRSHandle))
 				{
-					Controller.SetConstraintsDependencies(TransformConstraint->ConstraintID, InConstraint->ConstraintID);
+					FConstraintsManagerController::Get(InWorld).SetConstraintsDependencies(TransformConstraint->ConstraintID, InConstraint->ConstraintID);
 					ManagedDependencies[Index] = true;
 					LogDependency(ExternalDependencyDesc, TransformConstraint->ChildTRSHandle, ParentHandle, TransformConstraint, InConstraint);
 				}
 				else if (ChildHandle->HasDirectDependencyWith(*TransformConstraint->ChildTRSHandle))
 				{
-					Controller.SetConstraintsDependencies(TransformConstraint->ConstraintID, InConstraint->ConstraintID);
+					FConstraintsManagerController::Get(InWorld).SetConstraintsDependencies(TransformConstraint->ConstraintID, InConstraint->ConstraintID);
 					ManagedDependencies[Index] = true;
 					LogDependency(ExternalDependencyDesc, TransformConstraint->ChildTRSHandle, ChildHandle, TransformConstraint, InConstraint);
 				}
@@ -290,13 +287,13 @@ void FDependencyBuilder::BuildExternalDependencies(UWorld* InWorld, UTickableTra
 				// if the TransformConstraint handles depend on the new constraint child then, TransformConstraint should tick after
 				if (TransformConstraint->ParentTRSHandle->HasDirectDependencyWith(*ChildHandle))
 				{
-					Controller.SetConstraintsDependencies(InConstraint->ConstraintID, TransformConstraint->ConstraintID);
+					FConstraintsManagerController::Get(InWorld).SetConstraintsDependencies(InConstraint->ConstraintID, TransformConstraint->ConstraintID);
 					ManagedDependencies[Index] = true;
 					LogDependency(ExternalDependencyDesc, ChildHandle, TransformConstraint->ParentTRSHandle, InConstraint, TransformConstraint);
 				}
 				else if (TransformConstraint->ChildTRSHandle->HasDirectDependencyWith(*ChildHandle))
 				{
-					Controller.SetConstraintsDependencies(InConstraint->ConstraintID, TransformConstraint->ConstraintID);
+					FConstraintsManagerController::Get(InWorld).SetConstraintsDependencies(InConstraint->ConstraintID, TransformConstraint->ConstraintID);
 					ManagedDependencies[Index] = true;
 					LogDependency(ExternalDependencyDesc, ChildHandle, TransformConstraint->ChildTRSHandle, InConstraint, TransformConstraint);
 				}
@@ -317,7 +314,7 @@ void FDependencyBuilder::BuildExternalDependencies(UWorld* InWorld, UTickableTra
 						if (!FConstraintCycleChecker::HasPrerequisiteDependencyWith(&TickFunction, &ParentTickFunctionToCheck, VisitedFunctions) &&
 							!FConstraintCycleChecker::HasPrerequisiteDependencyWith(&ParentTickFunctionToCheck, &TickFunction, VisitedFunctions))
 						{
-							Controller.SetConstraintsDependencies(LastConstraintSharingSameTarget->ConstraintID, InConstraint->ConstraintID);
+							FConstraintsManagerController::Get(InWorld).SetConstraintsDependencies(LastConstraintSharingSameTarget->ConstraintID, InConstraint->ConstraintID);
 							if (bDebugDependencies)
 							{
 								UE_LOG(LogTemp, Warning, TEXT("Creation Order Dependency: '%s' must tick before '%s' to respect creation order."),
@@ -343,7 +340,7 @@ void FDependencyBuilder::BuildExternalDependencies(UWorld* InWorld, UTickableTra
 			const UObject* OtherChildTarget = GetHandleTarget(TransformConstraint->ChildTRSHandle);
 			return OtherChildTarget && OtherChildTarget == ParentTarget;
 		};
-		TArray< ConstraintWeakPtr > ParentTargetParentConstraints = Controller.GetConstraintsByPredicate(SameParentTargetPredicate);
+		TArray< ConstraintWeakPtr > ParentTargetParentConstraints = FConstraintsManagerController::Get(InWorld).GetConstraintsByPredicate(SameParentTargetPredicate);
 
 		// store constraint index in this array
 		const int32 ConstrainIndex = ParentTargetParentConstraints.IndexOfByKey(InConstraint);
@@ -368,7 +365,7 @@ void FDependencyBuilder::BuildExternalDependencies(UWorld* InWorld, UTickableTra
 				// if the new handles depend on that constraint child then, TransformConstraint should tick before
 				if (ParentHandle->HasDirectDependencyWith(*TransformConstraint->ChildTRSHandle))
 				{
-					Controller.SetConstraintsDependencies(TransformConstraint->ConstraintID, InConstraint->ConstraintID);
+					FConstraintsManagerController::Get(InWorld).SetConstraintsDependencies(TransformConstraint->ConstraintID, InConstraint->ConstraintID);
 					ManagedDependencies[Index] = true;
 					LogDependency(ExternalDependencyDesc, TransformConstraint->ChildTRSHandle, ParentHandle, TransformConstraint, InConstraint);
 				}
@@ -390,7 +387,7 @@ void FDependencyBuilder::BuildExternalDependencies(UWorld* InWorld, UTickableTra
 						if (!FConstraintCycleChecker::HasPrerequisiteDependencyWith(&TickFunction, &ParentTickFunctionToCheck, VisitedFunctions) &&
 							!FConstraintCycleChecker::HasPrerequisiteDependencyWith(&ParentTickFunctionToCheck, &TickFunction, VisitedFunctions))
 						{
-							Controller.SetConstraintsDependencies(LastConstraintSharingSameTarget->ConstraintID, InConstraint->ConstraintID);
+							FConstraintsManagerController::Get(InWorld).SetConstraintsDependencies(LastConstraintSharingSameTarget->ConstraintID, InConstraint->ConstraintID);
 							if (bDebugDependencies)
 							{
 								UE_LOG(LogTemp, Warning, TEXT("Creation Order Dependency: '%s' must tick before '%s' to respect creation order."),
@@ -416,7 +413,7 @@ void FDependencyBuilder::BuildExternalDependencies(UWorld* InWorld, UTickableTra
 			const UObject* OtherParentTarget = GetHandleTarget(TransformConstraint->ParentTRSHandle);
 			return OtherParentTarget && OtherParentTarget == ChildTarget;
 		};
-		TArray< ConstraintWeakPtr > ParentTargetChildConstraints = Controller.GetConstraintsByPredicate(SameParentAsChildTargetPredicate);
+		TArray< ConstraintWeakPtr > ParentTargetChildConstraints = FConstraintsManagerController::Get(InWorld).GetConstraintsByPredicate(SameParentAsChildTargetPredicate);
 
 		const FConstraintTickFunction& TickFunction = InConstraint->GetTickFunction(InWorld);
 		for (const ConstraintWeakPtr& Constraint: ParentTargetChildConstraints)
@@ -429,7 +426,7 @@ void FDependencyBuilder::BuildExternalDependencies(UWorld* InWorld, UTickableTra
 			if (!FConstraintCycleChecker::HasPrerequisiteDependencyWith(&TickFunction, &ParentTickFunctionToCheck, VisitedFunctions0) &&
 				!FConstraintCycleChecker::HasPrerequisiteDependencyWith(&ParentTickFunctionToCheck, &TickFunction, VisitedFunctions1))
 			{
-				Controller.SetConstraintsDependencies(InConstraint->ConstraintID, ConstraintSharingParentTarget->ConstraintID);
+				FConstraintsManagerController::Get(InWorld).SetConstraintsDependencies(InConstraint->ConstraintID, ConstraintSharingParentTarget->ConstraintID);
 				if (bDebugDependencies)
 				{
 					UE_LOG(LogTemp, Warning, TEXT("External Dependency: '%s' must tick before '%s' as it acts on its parent."),
@@ -464,8 +461,7 @@ bool FDependencyBuilder::BuildDependencies(UWorld* InWorld, UTickableTransformCo
 	const UTransformableHandle* ChildHandle = InConstraint->ChildTRSHandle.Get();
 	
  	// get previous child constraints
-	const FConstraintsManagerController& Controller = FConstraintsManagerController::Get(InWorld);
-	TArray<ConstraintWeakPtr> ChildParentConstraints = Controller.GetParentConstraints(ChildHandle->GetHash(), true);
+	TArray<ConstraintWeakPtr> ChildParentConstraints = FConstraintsManagerController::Get(InWorld).GetParentConstraints(ChildHandle->GetHash(), true);
 	ChildParentConstraints.Remove(InConstraint);
 
 	// add dependencies with the last child constraint
@@ -476,7 +472,7 @@ bool FDependencyBuilder::BuildDependencies(UWorld* InWorld, UTickableTransformCo
 		{
 			UE_LOG(LogTemp, Warning, TEXT("Order Dependency: tick after last constraint."));
 		}
-		Controller.SetConstraintsDependencies(LastChildConstraintID, InConstraint->ConstraintID);
+		FConstraintsManagerController::Get(InWorld).SetConstraintsDependencies(LastChildConstraintID, InConstraint->ConstraintID);
 	}
 
 	const UObject* ParentTarget = ParentHandle->GetTarget().Get();
@@ -499,7 +495,7 @@ bool FDependencyBuilder::BuildDependencies(UWorld* InWorld, UTickableTransformCo
 	UE::TransformConstraintUtil::GetChildrenConstraints(InWorld, InConstraint, ChildChildConstraints, !bSelf);
 	for (const ConstraintWeakPtr& ChildConstraint: ChildChildConstraints)
 	{
-		Controller.SetConstraintsDependencies(InConstraint->ConstraintID, ChildConstraint->ConstraintID);
+		FConstraintsManagerController::Get(InWorld).SetConstraintsDependencies(InConstraint->ConstraintID, ChildConstraint->ConstraintID);
 		if (const UTickableTransformConstraint* TransformConstraint = Cast<UTickableTransformConstraint>(ChildConstraint))
 		{
 			LogDependency(ChildDependencyDesc, ChildHandle, TransformConstraint->ChildTRSHandle, InConstraint, TransformConstraint);
@@ -521,7 +517,7 @@ bool FDependencyBuilder::BuildDependencies(UWorld* InWorld, UTickableTransformCo
 	}
 
 	// invalidate graph
-	Controller.InvalidateEvaluationGraph();
+	FConstraintsManagerController::Get(InWorld).InvalidateEvaluationGraph();
 	
 	return true;
 }
@@ -540,8 +536,6 @@ void FDependencyBuilder::BuildAttachmentsDependencies(UWorld* InWorld, const UTi
 	{
 		return;
 	}
-
-	const FConstraintsManagerController& Controller = FConstraintsManagerController::Get(InWorld);
 	
 	static const TCHAR* AttachmentDependencyDesc = TEXT("Attachment Dependency");
 	
@@ -569,10 +563,10 @@ void FDependencyBuilder::BuildAttachmentsDependencies(UWorld* InWorld, const UTi
 					return ChildHandle->GetHash() == ChildHash;
 				};
 				
-				const TArray< ConstraintWeakPtr > AttachChildConstraints = Controller.GetConstraintsByPredicate(IsHashChildOfConstraint);
+				const TArray< ConstraintWeakPtr > AttachChildConstraints = FConstraintsManagerController::Get(InWorld).GetConstraintsByPredicate(IsHashChildOfConstraint);
 				for (const ConstraintWeakPtr& AttachChildConstraint: AttachChildConstraints)
 				{
-					Controller.SetConstraintsDependencies(InConstraint->ConstraintID, AttachChildConstraint->ConstraintID);
+					FConstraintsManagerController::Get(InWorld).SetConstraintsDependencies(InConstraint->ConstraintID, AttachChildConstraint->ConstraintID);
 					if (const UTickableTransformConstraint* AttachTransformConstraint = Cast<UTickableTransformConstraint>(AttachChildConstraint))
 					{
 						LogAttachmentDependency(AttachmentDependencyDesc, ChildHandle, AttachTransformConstraint->ChildTRSHandle,  InConstraint, AttachTransformConstraint);
diff --git a/Engine/Source/Runtime/MovieSceneTracks/Private/Systems/MovieSceneConstraintSystem.cpp b/Engine/Source/Runtime/MovieSceneTracks/Private/Systems/MovieSceneConstraintSystem.cpp
--- a/Engine/Source/Runtime/MovieSceneTracks/Private/Systems/MovieSceneConstraintSystem.cpp
+++ b/Engine/Source/Runtime/MovieSceneTracks/Private/Systems/MovieSceneConstraintSystem.cpp
@@ -21,21 +21,20 @@
 #include UE_INLINE_GENERATED_CPP_BY_NAME(MovieSceneConstraintSystem)
 
 
-static UTickableConstraint* CreateConstraintIfNeeded(FConstraintsManagerController& Controller, FConstraintAndActiveChannel* ConstraintAndActiveChannel, const  FConstraintComponentData& ConstraintChannel)
+static UTickableConstraint* CreateConstraintIfNeeded(UWorld* InWorld, const FConstraintAndActiveChannel* InConstraintAndActiveChannel)
 {
-	UTickableConstraint* Constraint = ConstraintAndActiveChannel->GetConstraint();
+	UTickableConstraint* Constraint = InConstraintAndActiveChannel->GetConstraint();
 	if (!Constraint) //if constraint doesn't exist it probably got unspawned so recreate it and add it
 	{
 		return nullptr;
 	}
-	else // it's possible that we have it but it's not in the manager, due to manager not being saved with it (due to spawning or undo/redo).
+	
+	// it's possible that we have it but it's not in the manager, due to manager not being saved with it (due to spawning or undo/redo).
+	if (!FConstraintsManagerController::Get(InWorld).GetConstraint(Constraint->ConstraintID))
 	{
-		const TArray< TWeakObjectPtr<UTickableConstraint>>& ConstraintsArray = Controller.GetConstraintsArray();
-		if (Controller.GetConstraint(Constraint->ConstraintID) == nullptr)
-		{
-			Controller.AddConstraint(Constraint);
-		}
+		FConstraintsManagerController::Get(InWorld).AddConstraint(Constraint);
 	}
+	
 	return Constraint;
 }
 
@@ -84,18 +83,15 @@ struct FPreAnimatedConstraintTraits : FBoundObjectPreAnimatedStateTraits
 	static FPreAnimatedConstraint CachePreAnimatedValue(UObject* InBoundObject, const FConstraintComponentData& ConstraintData)
 	{
 		USceneComponent* SceneComponent = CastChecked<USceneComponent>(InBoundObject);
-		FConstraintsManagerController& Controller = FConstraintsManagerController::Get(SceneComponent->GetWorld());
+		UWorld* World = SceneComponent->GetWorld();
 		if (FConstraintAndActiveChannel* ConstraintAndActiveChannel = ConstraintData.Section->GetConstraintChannel(ConstraintData.ConstraintID))
 		{
 			UTickableConstraint* Constraint = ConstraintAndActiveChannel->GetConstraint().Get();
-			const bool bIsActive = false; //we set static mesh constraints to be inactive to matchup how CR's work in not being valid
+			constexpr bool bIsActive = false; //we set static mesh constraints to be inactive to matchup how CR's work in not being valid
 			const FGuid ID = Constraint ? Constraint->ConstraintID : ConstraintData.ConstraintID;
 			//also need to create it
-			Constraint = CreateConstraintIfNeeded(Controller, ConstraintAndActiveChannel, ConstraintData);
-			if (Constraint)
-			{
-				Controller.AddConstraint(Constraint);
-			}
+			Constraint = CreateConstraintIfNeeded(World, ConstraintAndActiveChannel);
+			
 			return FPreAnimatedConstraint{ ID, (Constraint == nullptr), bIsActive, SceneComponent, ConstraintData.Section };
 		}
 		return FPreAnimatedConstraint{ ConstraintData.ConstraintID, false,false, SceneComponent, ConstraintData.Section };
@@ -143,17 +139,13 @@ struct FPreAnimatedConstraintTraits : FBoundObjectPreAnimatedStateTraits
 struct FEvaluateConstraintChannels
 {
 	FEvaluateConstraintChannels(UWorld* InWorld, FInstanceRegistry* InInstanceRegistry, UMovieSceneConstraintSystem* InSystem)
-		: World(InWorld), InstanceRegistry(InInstanceRegistry), Controller(nullptr), System(InSystem)
+		: World(InWorld), InstanceRegistry(InInstanceRegistry), System(InSystem)
 	{
 		check(World && InstanceRegistry && System);
 	}
 
 	void PreTask()
 	{
-		// Retrieve the controller at the start of the task since this task is now stored persistently
-		Controller = &FConstraintsManagerController::Get(World);
-		check(Controller);
-
 		System->DynamicOffsets.Reset();
 	}
 
@@ -165,7 +157,7 @@ struct FEvaluateConstraintChannels
 			return;
 		}
 
-		UTickableConstraint* Constraint = CreateConstraintIfNeeded(*Controller, ConstraintAndActiveChannel, ConstraintChannel);
+		UTickableConstraint* Constraint = CreateConstraintIfNeeded(World, ConstraintAndActiveChannel);
 		if (!Constraint)
 		{
 			return;
@@ -204,7 +196,6 @@ struct FEvaluateConstraintChannels
 
 	UWorld* World;
 	FInstanceRegistry* InstanceRegistry;
-	FConstraintsManagerController* Controller;
 	UMovieSceneConstraintSystem* System;
 };
 
