Geometry collection : add ResetState function ( C++ , blueprint and live in editor as a button )

diff --git a/Engine/Source/Runtime/Experimental/GeometryCollectionEngine/Private/GeometryCollection/GeometryCollectionComponent.cpp b/Engine/Source/Runtime/Experimental/GeometryCollectionEngine/Private/GeometryCollection/GeometryCollectionComponent.cpp
--- a/Engine/Source/Runtime/Experimental/GeometryCollectionEngine/Private/GeometryCollection/GeometryCollectionComponent.cpp
+++ b/Engine/Source/Runtime/Experimental/GeometryCollectionEngine/Private/GeometryCollection/GeometryCollectionComponent.cpp
@@ -1920,11 +1920,8 @@ void UGeometryCollectionComponent::PostEditChangeProperty(FPropertyChangedEvent&
 
 	if (PropertyChangedEvent.Property && PropertyChangedEvent.Property->GetFName() == GET_MEMBER_NAME_CHECKED(UGeometryCollectionComponent, ChaosSolverActor))
 	{
-		// Reset the dynamic collection 
-		ResetDynamicCollection();
-
-		// Recreate the dynamic state since the dynamic collection has been reset
-		RecreatePhysicsState();
+		// call the setter method to make sure everything is proeprly reset
+		SetSolverActor(ChaosSolverActor);
 	}
 }
 #endif
@@ -5077,6 +5074,18 @@ void UGeometryCollectionComponent::SetRestCollection(const UGeometryCollection*
 	}
 }
 
+void UGeometryCollectionComponent::ResetState()
+{
+	// Reset the dynamic collection
+	ResetDynamicCollection();
+
+	// Rebuild the physics state since the dynamic collection has changed
+	RecreatePhysicsState();
+
+	// if we are using a custom render, make sure everything is properly reset 
+	RefreshCustomRenderer();
+}
+
 FString UGeometryCollectionComponent::GetDebugInfo()
 {
 	// print the game thread side of things
@@ -5165,6 +5174,8 @@ FGeometryCollectionEdit::~FGeometryCollectionEdit()
 			{
 				ToProcess->RecreatePhysicsState();
 			}
+
+			ToProcess->RefreshCustomRenderer();
 		};
 
 		if (bPropagateToAllMatchingComponents)
@@ -5969,11 +5980,7 @@ void UGeometryCollectionComponent::SetSolverActor(AChaosSolverActor* InSolverAct
 {
 	ChaosSolverActor = InSolverActor;
 
-	// Reset the dynamic collection
-	ResetDynamicCollection();
-
-	// Rebuild the physics state since the proxy could now be registered to another actor
-	RecreatePhysicsState();
+	ResetState();
 }
 
 bool UGeometryCollectionComponent::CanRunSimulationInEditor() const
diff --git a/Engine/Source/Runtime/Experimental/GeometryCollectionEngine/Public/GeometryCollection/GeometryCollectionComponent.h b/Engine/Source/Runtime/Experimental/GeometryCollectionEngine/Public/GeometryCollection/GeometryCollectionComponent.h
--- a/Engine/Source/Runtime/Experimental/GeometryCollectionEngine/Public/GeometryCollection/GeometryCollectionComponent.h
+++ b/Engine/Source/Runtime/Experimental/GeometryCollectionEngine/Public/GeometryCollection/GeometryCollectionComponent.h
@@ -857,6 +857,10 @@ public:
 	UFUNCTION(BlueprintCallable, Category = "ChaosPhysics")
 	GEOMETRYCOLLECTIONENGINE_API void SetRestCollection(const UGeometryCollection * RestCollectionIn, bool bApplyAssetDefaults = true);
 
+	/** Reset the state of the component : object will be rest to be unbroken and physics will reset accordingly */
+	UFUNCTION(BlueprintCallable, CallInEditor, Category = "ChaosPhysics")
+	GEOMETRYCOLLECTIONENGINE_API void ResetState();
+
 	/** RestCollection */
 	UFUNCTION(BlueprintCallable, Category = "ChaosPhysics")
 	GEOMETRYCOLLECTIONENGINE_API FString GetDebugInfo();
