Constraints: new evaluation turned on by default

diff --git a/Engine/Source/Runtime/Experimental/Animation/Constraints/Private/Transform/TransformableHandleUtils.cpp b/Engine/Source/Runtime/Experimental/Animation/Constraints/Private/Transform/TransformableHandleUtils.cpp
--- a/Engine/Source/Runtime/Experimental/Animation/Constraints/Private/Transform/TransformableHandleUtils.cpp
+++ b/Engine/Source/Runtime/Experimental/Animation/Constraints/Private/Transform/TransformableHandleUtils.cpp
@@ -10,7 +10,7 @@ namespace TransformableHandleUtils
 {
 
 static int32 SkeletalMeshTickingMode = 2;
-static bool	bNewEvaluation = false;
+static bool	bNewEvaluation = true;
 
 static FAutoConsoleVariableRef CVarSkeletalMeshTickingMode(
 	TEXT("Constraints.SkeletalMesh.TickingMode"),
