Iris - Detect some NetToken deserialization errors and handle them gracefully.

diff --git a/Engine/Plugins/Runtime/ReplicationSystemTestPlugin/Source/Private/Tests/ReplicationSystem/ReplicationSystemServerClientTestFixture.cpp b/Engine/Plugins/Runtime/ReplicationSystemTestPlugin/Source/Private/Tests/ReplicationSystem/ReplicationSystemServerClientTestFixture.cpp
--- a/Engine/Plugins/Runtime/ReplicationSystemTestPlugin/Source/Private/Tests/ReplicationSystem/ReplicationSystemServerClientTestFixture.cpp
+++ b/Engine/Plugins/Runtime/ReplicationSystemTestPlugin/Source/Private/Tests/ReplicationSystem/ReplicationSystemServerClientTestFixture.cpp
@@ -477,8 +477,8 @@ void FReplicationSystemTestNode::RecvUpdate(uint32 ConnectionId, FNetSerializati
 
 	Connection.DataStreamManager->ReadData(Context);
 
-	check(!Context.HasErrorOrOverflow());
-	check(Context.GetBitStreamReader()->GetBitsLeft() == 0U);
+	ensure(!Context.HasErrorOrOverflow());
+	ensure(Context.GetBitStreamReader()->GetBitsLeft() == 0U);
 }
 
 uint32 FReplicationSystemTestNode::GetReplicationSystemId() const
diff --git a/Engine/Plugins/Runtime/ReplicationSystemTestPlugin/Source/Private/Tests/ReplicationSystem/TestStructNetTokenDataStore.cpp b/Engine/Plugins/Runtime/ReplicationSystemTestPlugin/Source/Private/Tests/ReplicationSystem/TestStructNetTokenDataStore.cpp
--- a/Engine/Plugins/Runtime/ReplicationSystemTestPlugin/Source/Private/Tests/ReplicationSystem/TestStructNetTokenDataStore.cpp
+++ b/Engine/Plugins/Runtime/ReplicationSystemTestPlugin/Source/Private/Tests/ReplicationSystem/TestStructNetTokenDataStore.cpp
@@ -5,6 +5,8 @@
 #include "Iris/ReplicationSystem/ReplicationFragmentUtil.h"
 #include "Iris/ReplicationSystem/StructNetTokenDataStore.h"
 #include "Iris/ReplicationState/PropertyNetSerializerInfoRegistry.h"
+#include "Iris/Serialization/NetBitStreamUtil.h"
+#include "Iris/Serialization/NetBitStreamWriter.h"
 #include "Logging/LogScopedVerbosityOverride.h"
 
 // Note: this is in engine code.
@@ -203,6 +205,107 @@ UE_NET_TEST_FIXTURE(FTestStructNetTokensFixture, StructAsNetTokenUsingLastResort
 	UE_NET_ASSERT_EQ(ServerToken, ClientToken);
 }
 
+// Test that we can receive corrupt token data in the form a bad token index without crashing
+UE_NET_TEST_FIXTURE(FTestStructNetTokensFixture, CanReceiveInvalidTokenIndexGracefully)
+{
+	// Create token which is fairly large since we're going to corrupt the packet and need some leeway.
+	FTestStructNetTokenStoreStruct StructA;
+	StructA.IntA = 1;
+	StructA.IntB = 2;
+	StructA.ByteArray.SetNumZeroed(64);
+
+	FNetToken StructTokenA = CreateAndExportNetToken(StructA);
+
+	// Send  packet
+	Server->NetUpdate();
+	Server->SendTo(Client);
+	Server->PostSendUpdate();
+
+	// Corrupt packet
+	{
+		FReplicationSystemTestNode::FConnectionInfo& ConnectionInfo = Server->GetConnectionInfo(Client->ConnectionIdOnServer);
+		UE_NET_ASSERT_GT(ConnectionInfo.WrittenPackets.Count(), SIZE_T(0));
+
+		FReplicationSystemTestNode::FPacketData& Packet = ConnectionInfo.WrittenPackets.Poke();
+
+		FNetBitStreamWriter Writer;
+		Writer.InitBytes(Packet.PacketBuffer, sizeof(Packet.PacketBuffer));
+
+		// The bit offset for the FNetToken at the time this test was written.
+		constexpr uint32 ExpectedTokenBitOffset = 9U;
+		Writer.Seek(ExpectedTokenBitOffset);
+
+		// Mimic FNetTokenStore::InternalWriteNetToken but write a bad token index
+		constexpr uint32 BadTokenIndex = FNetToken::MaxNetTokenCount + 4711U;
+		WritePackedUint32(&Writer, BadTokenIndex);
+		Writer.WriteBool(StructTokenA.IsAssignedByAuthority());
+		Writer.WriteBits(StructTokenA.GetTypeId(), FNetToken::TokenTypeIdBits);
+
+		Writer.CommitWrites();
+	}
+
+	LOG_SCOPE_VERBOSITY_OVERRIDE(LogNetToken, ELogVerbosity::Fatal);
+
+	// Deliver packet
+	{
+		FTestEnsureScope SuppressEnsureScope;
+		Server->DeliverTo(Client, DeliverPacket);
+	}
+
+	// If corruption occured we should not be able to resolve the token on the client and we'll get some sort of invalid state back which shouldn't match the server state.
+	UE_NET_ASSERT_NE(StructA, ClientStructTokenStore->ResolveRemoteToken(StructTokenA, *ClientRemoteNetTokenState));
+}
+
+// Test that we can receive corrupt token data in the form a non-existing token type without crashing
+UE_NET_TEST_FIXTURE(FTestStructNetTokensFixture, CanReceiveInvalidTokenTypeGracefully)
+{
+	// Create token which is fairly large since we're going to corrupt the packet and need some leeway.
+	FTestStructNetTokenStoreStruct StructA;
+	StructA.IntA = 1;
+	StructA.IntB = 2;
+	StructA.ByteArray.SetNumZeroed(64);
+
+	FNetToken StructTokenA = CreateAndExportNetToken(StructA);
+
+	// Send  packet
+	Server->NetUpdate();
+	Server->SendTo(Client);
+	Server->PostSendUpdate();
+
+	// Corrupt packet
+	{
+		FReplicationSystemTestNode::FConnectionInfo& ConnectionInfo = Server->GetConnectionInfo(Client->ConnectionIdOnServer);
+		UE_NET_ASSERT_GT(ConnectionInfo.WrittenPackets.Count(), SIZE_T(0));
+
+		FReplicationSystemTestNode::FPacketData& Packet = ConnectionInfo.WrittenPackets.Poke();
+
+		FNetBitStreamWriter Writer;
+		Writer.InitBytes(Packet.PacketBuffer, sizeof(Packet.PacketBuffer));
+
+		// The bit offset for the FNetToken at the time this test was written.
+		constexpr uint32 ExpectedTokenBitOffset = 9U;
+		Writer.Seek(ExpectedTokenBitOffset);
+
+		// Mimic FNetTokenStore::InternalWriteNetToken but write a bad token type ID
+		WritePackedUint32(&Writer, StructTokenA.GetIndex());
+		Writer.WriteBool(StructTokenA.IsAssignedByAuthority());
+		Writer.WriteBits(FNetToken::MaxTypeIdCount - 1U, FNetToken::TokenTypeIdBits);
+
+		Writer.CommitWrites();
+	}
+
+	LOG_SCOPE_VERBOSITY_OVERRIDE(LogNetToken, ELogVerbosity::Fatal);
+
+	// Deliver packet
+	{
+		FTestEnsureScope SuppressEnsureScope;
+		Server->DeliverTo(Client, DeliverPacket);
+	}
+
+	// If corruption occured we should not be able to resolve the token on the client and we'll get some sort of invalid state back which shouldn't match the server state.
+	UE_NET_ASSERT_NE(StructA, ClientStructTokenStore->ResolveRemoteToken(StructTokenA, *ClientRemoteNetTokenState));
+}
+
 } // end namespace UE::Net::Private
 
 
diff --git a/Engine/Plugins/Runtime/ReplicationSystemTestPlugin/Source/Private/Tests/ReplicationSystem/TestStructNetTokenDataStore.h b/Engine/Plugins/Runtime/ReplicationSystemTestPlugin/Source/Private/Tests/ReplicationSystem/TestStructNetTokenDataStore.h
--- a/Engine/Plugins/Runtime/ReplicationSystemTestPlugin/Source/Private/Tests/ReplicationSystem/TestStructNetTokenDataStore.h
+++ b/Engine/Plugins/Runtime/ReplicationSystemTestPlugin/Source/Private/Tests/ReplicationSystem/TestStructNetTokenDataStore.h
@@ -2,6 +2,7 @@
 
 #pragma once
 
+#include "Containers/Array.h"
 #include "ReplicatedTestObject.h"
 #include "Iris/Serialization/NetSerializer.h"
 #include "Iris/ReplicationSystem/NetTokenStructDefines.h"
@@ -16,7 +17,7 @@ struct FTestStructNetTokenStoreStruct
 
 	uint64 GetUniqueKey() const
 	{
-		return HashCombineFast(IntA, IntB);
+		return HashCombineFast(HashCombineFast(IntA, IntB), GetTypeHash(ByteArray));
 	}
 
 	UPROPERTY(Transient)
@@ -24,6 +25,9 @@ struct FTestStructNetTokenStoreStruct
 
 	UPROPERTY(Transient)
 	int32 IntB;
+
+	UPROPERTY(Transient)
+	TArray<uint8> ByteArray;
 };
 	
 UE_NET_DECLARE_NAMED_NETTOKEN_STRUCT_SERIALIZERS(TestStructNetTokenStoreStruct, REPLICATIONSYSTEMTESTPLUGIN_API)
diff --git a/Engine/Source/Runtime/Experimental/Iris/Core/Private/Iris/ReplicationSystem/NetTokenStore.cpp b/Engine/Source/Runtime/Experimental/Iris/Core/Private/Iris/ReplicationSystem/NetTokenStore.cpp
--- a/Engine/Source/Runtime/Experimental/Iris/Core/Private/Iris/ReplicationSystem/NetTokenStore.cpp
+++ b/Engine/Source/Runtime/Experimental/Iris/Core/Private/Iris/ReplicationSystem/NetTokenStore.cpp
@@ -134,10 +134,18 @@ FNetToken FNetTokenStore::InternalReadNetToken(UE::Net::FNetSerializationContext
 			TokenTypeId = Reader->ReadBits(FNetToken::TokenTypeIdBits);
 		}
 
-		if (!Reader->IsOverflown())
+		if (!Context.HasErrorOrOverflow())
 		{
 			ReadToken = FNetTokenStore::MakeNetToken(TokenTypeId, TokenIndex, bIsAssignedByAuthority ? FNetToken::ENetTokenAuthority::Authority : FNetToken::ENetTokenAuthority::None);
 			UE_NET_TRACE_SET_SCOPE_NAME(TokenScope, *ReadToken.ToString());
+
+			// We expected a valid token
+			if (!ReadToken.IsValid())
+			{
+				Context.SetError(GNetError_InvalidValue);
+				UE_LOG(LogNetToken, Error, TEXT("Expected a valid token but got an invalid one. TokenTypeId: %u TokenIndex: %u."), TokenTypeId, TokenIndex);
+				ensure(false);
+			}
 		}
 	}
 
@@ -435,7 +443,7 @@ void FNetTokenStore::ReadTokenData(FNetSerializationContext& Context, const FNet
 
 		if (TokenTypeId >= (uint32)TokenDataStores.Num())
 		{
-			Reader->DoOverflow();
+			Context.SetError(GNetError_InvalidValue);
 			UE_LOG(LogNetToken, Error, TEXT("Failed to read ReadTokenData for %s."), *NetToken.ToString());
 			ensure(false);
 			return;
@@ -446,7 +454,7 @@ void FNetTokenStore::ReadTokenData(FNetSerializationContext& Context, const FNet
 
 		if (!ValidateAndStoreNetTokenData(*DataStore, RemoteNetTokenStoreState, NetToken, StoreKey))
 		{
-			Reader->DoOverflow();
+			Context.SetError(GNetError_InvalidValue);
 			UE_LOG(LogNetToken, Error, TEXT("Failed to read ReadTokenData for %s."), *NetToken.ToString());
 			ensure(false);
 			return;
diff --git a/Engine/Source/Runtime/Experimental/Iris/Core/Private/Iris/ReplicationSystem/ObjectReferenceCache.cpp b/Engine/Source/Runtime/Experimental/Iris/Core/Private/Iris/ReplicationSystem/ObjectReferenceCache.cpp
--- a/Engine/Source/Runtime/Experimental/Iris/Core/Private/Iris/ReplicationSystem/ObjectReferenceCache.cpp
+++ b/Engine/Source/Runtime/Experimental/Iris/Core/Private/Iris/ReplicationSystem/ObjectReferenceCache.cpp
@@ -823,7 +823,7 @@ UObject* FObjectReferenceCache::ResolveObjectReferenceHandleInternal(FNetRefHand
 	if (!CacheObjectPtr->RelativePath.IsValid())
 	{
 		// If we don't have a path, assume this is a dynamic handle
-		check(RefHandle.IsDynamic());
+		ensure(RefHandle.IsDynamic());
 		
 		return nullptr;
 	}
@@ -1092,19 +1092,19 @@ UObject* FObjectReferenceCache::ResolveObjectReferenceHandleInternal(FNetRefHand
 		if (FCachedNetObjectReference* ExistingCachedObjectPtr = ReferenceHandleToCachedReference.Find(*ExistingHandle))
 		{
 			FCachedNetObjectReference& ExistingCachedObject = *ExistingCachedObjectPtr;
-		if (ExistingCachedObject.Object.IsStale())
-		{
-			UE_LOG_REFERENCECACHE(Verbose, TEXT("Invalidating object for Stale Ref %s with %s as the object has been reused"), *ExistingCachedObject.NetRefHandle.ToString(), *CompleteRefHandle.ToString());
-		}
-		else if (ExistingCachedObject.Object.HasSameIndexAndSerialNumber(CacheObjectPtr->Object))
-		{
-			bShouldAssign = CompleteRefHandle.GetId() > (*ExistingHandle).GetId();
-			if (bShouldAssign)
+			if (ExistingCachedObject.Object.IsStale())
+			{
+				UE_LOG_REFERENCECACHE(Verbose, TEXT("Invalidating object for Stale Ref %s with %s as the object has been reused"), *ExistingCachedObject.NetRefHandle.ToString(), *CompleteRefHandle.ToString());
+			}
+			else if (ExistingCachedObject.Object.HasSameIndexAndSerialNumber(CacheObjectPtr->Object))
 			{
-				UE_LOG_REFERENCECACHE(Verbose, TEXT("Invalidating object for conflicting Ref %s with %s"), *ExistingCachedObject.NetRefHandle.ToString(), *CompleteRefHandle.ToString());
-			}	
+				bShouldAssign = CompleteRefHandle.GetId() > (*ExistingHandle).GetId();
+				if (bShouldAssign)
+				{
+					UE_LOG_REFERENCECACHE(Verbose, TEXT("Invalidating object for conflicting Ref %s with %s"), *ExistingCachedObject.NetRefHandle.ToString(), *CompleteRefHandle.ToString());
+				}
+			}
 		}
-	}
 		else
 		{
 			UE_LOG_REFERENCECACHE(Warning, TEXT("ObjectReferenceCache::ResolveObjectReferenceHandleInternal removed %s from ObjectToNetReferenceHandle due to mismatch with cache. Object %s (0x%p) "), *ExistingHandle->ToString(), *GetNameSafe(Object), Object);
@@ -1351,8 +1351,7 @@ void FObjectReferenceCache::ReadFullReferenceInternal(FNetSerializationContext&
 	if (RecursionCount > INTERNAL_READ_REF_RECURSION_LIMIT) 
 	{
 		UE_LOG_REFERENCECACHE(Warning, TEXT("ReadFullReferenceInternal: Hit recursion limit."));
-		check(false);
-		Reader->DoOverflow();
+		Context.SetError(GNetError_InternalError);
 		OutObjectRef = FNetObjectReference();
 
 		return;
@@ -1371,7 +1370,8 @@ void FObjectReferenceCache::ReadFullReferenceInternal(FNetSerializationContext&
 	const bool bIsExported = Reader->ReadBool();
 	if (Reader->IsOverflown())
 	{
-		check(false);
+		UE_LOG_REFERENCECACHE(Error, TEXT("ReadFullReferenceInternal: Bitstream overflow"));
+		Context.SetError(GNetError_BitStreamOverflow);
 		OutObjectRef = FNetObjectReference();
 		return;
 	}
@@ -1392,7 +1392,8 @@ void FObjectReferenceCache::ReadFullReferenceInternal(FNetSerializationContext&
 		bNoLoad = Reader->ReadBits(1U) & 1;
 		if (Reader->IsOverflown())
 		{
-			check(false);
+			UE_LOG_REFERENCECACHE(Error, TEXT("ReadFullReferenceInternal: Bitstream overflow"));
+			Context.SetError(GNetError_BitStreamOverflow);
 			OutObjectRef = FNetObjectReference();
 			return;
 		}
@@ -1407,7 +1408,8 @@ void FObjectReferenceCache::ReadFullReferenceInternal(FNetSerializationContext&
 
 	if (Reader->IsOverflown())
 	{
-		check(false);
+		UE_LOG_REFERENCECACHE(Error, TEXT("ReadFullReferenceInternal: Bitstream overflow"));
+		Context.SetError(GNetError_BitStreamOverflow);
 		OutObjectRef = FNetObjectReference();
 		return;
 	}
@@ -1488,8 +1490,22 @@ void FObjectReferenceCache::ReadFullReference(FNetSerializationContext& Context,
 		FNetToken RelativePath = StringTokenStore->ReadNetToken(Context);
 		NetTokenStore->ConditionalReadNetTokenData(Context, RelativePath);
 
-		OutRef.RefHandle = NetRefHandle;
-		OutRef.PathToken = RelativePath;
+		// We expect a valid RelativePath
+		if (!RelativePath.IsValid())
+		{
+			UE_LOG_REFERENCECACHE(Error, TEXT("ReadFullReference: Invalid relative path"));
+			Context.SetError(GNetError_InvalidValue);
+		}
+
+		if (Context.HasErrorOrOverflow())
+		{
+			OutRef = FNetObjectReference();
+		}
+		else
+		{
+			OutRef.RefHandle = NetRefHandle;
+			OutRef.PathToken = RelativePath;
+		}
 
 		return;
 	}
@@ -1596,8 +1612,22 @@ void FObjectReferenceCache::ReadReference(FNetSerializationContext& Context, FNe
 		const FNetRefHandle NetRefHandle = ReadNetRefHandle(Context);
 		FNetToken RelativePath = StringTokenStore->ReadNetToken(Context);
 
-		OutRef.RefHandle = NetRefHandle;
-		OutRef.PathToken = RelativePath;
+		// We expect a valid RelativePath
+		if (!RelativePath.IsValid())
+		{
+			UE_LOG_REFERENCECACHE(Error, TEXT("ReadReference: Invalid relative path"));
+			Context.SetError(GNetError_InvalidValue);
+		}
+
+		if (Context.HasErrorOrOverflow())
+		{
+			OutRef = FNetObjectReference();
+		}
+		else
+		{
+			OutRef.RefHandle = NetRefHandle;
+			OutRef.PathToken = RelativePath;
+		}
 
 		return;
 	}
@@ -1606,7 +1636,8 @@ void FObjectReferenceCache::ReadReference(FNetSerializationContext& Context, FNe
 
 	if (Reader->IsOverflown())
 	{
-		check(false);
+		UE_LOG_REFERENCECACHE(Error, TEXT("ReadReference: Bitstream overflow"));
+		Context.SetError(GNetError_BitStreamOverflow);
 		OutRef = FNetObjectReference();
 		return;
 	}
@@ -1763,7 +1794,14 @@ FObjectReferenceCache::EWriteExportsResult FObjectReferenceCache::WritePendingEx
 
 bool FObjectReferenceCache::ReadExports(const FNetRefHandle& NetObjectHandle, FNetSerializationContext& Context, TArray<FNetRefHandle>* MustBeMappedExports, EIrisAsyncLoadingPriority& OutIrisAsyncLoadingPriority)
 {
-	check(MustBeMappedExports);
+	if (!MustBeMappedExports)
+	{
+		Context.SetError(GNetError_InvalidValue);
+		UE_LOG_REFERENCECACHE(Error, TEXT("ReadExports: Null MustBeMappedExports"));
+		ensure(false);
+		return false;
+	}
+
 	FNetBitStreamReader& Reader = *Context.GetBitStreamReader();
 
 	// Read NetToken exports
@@ -1871,7 +1909,13 @@ bool FObjectReferenceCache::WriteMustBeMappedExports(FNetSerializationContext& C
 
 void FObjectReferenceCache::ReadMustBeMappedExports(const FNetRefHandle& NetObjectHandle, FNetSerializationContext& Context, TArray<FNetRefHandle>* MustBeMappedExports, EIrisAsyncLoadingPriority& OutIrisAsyncLoadingPriority)
 {
-	check(MustBeMappedExports);
+	if (!MustBeMappedExports)
+	{
+		Context.SetError(GNetError_InvalidValue);
+		UE_LOG_REFERENCECACHE(Error, TEXT("ReadMustBeMappedExports: Null MustBeMappedExports"));
+		ensure(false);
+		return;
+	}
 
 	FNetBitStreamReader& Reader = *Context.GetBitStreamReader();
 
diff --git a/Engine/Source/Runtime/Experimental/Iris/Core/Public/Iris/ReplicationSystem/NetTokenStore.h b/Engine/Source/Runtime/Experimental/Iris/Core/Public/Iris/ReplicationSystem/NetTokenStore.h
--- a/Engine/Source/Runtime/Experimental/Iris/Core/Public/Iris/ReplicationSystem/NetTokenStore.h
+++ b/Engine/Source/Runtime/Experimental/Iris/Core/Public/Iris/ReplicationSystem/NetTokenStore.h
@@ -393,10 +393,12 @@ private:
 
 	static FNetToken MakeNetToken(uint32 TypeId, uint32 Index, FNetToken::ENetTokenAuthority Authority)
 	{ 
-		check(TypeId < FNetToken::MaxTypeIdCount);
-		check(Index < FNetToken::MaxNetTokenCount);
-
-		return FNetToken(TypeId, Index, Authority);
+		if ((TypeId < FNetToken::MaxTypeIdCount) && (Index < FNetToken::MaxNetTokenCount))
+		{
+			return FNetToken(TypeId, Index, Authority);
+		}
+		
+		return FNetToken();
 	}
 
 	TUniquePtr<FNetTokenStoreState> LocalNetTokenStoreState;
