Various bughawk fixes found after submitting CL 45051761

diff --git a/Engine/Plugins/ChaosClothAsset/Source/ChaosClothAsset/Private/ChaosClothAsset/ClothCollection.cpp b/Engine/Plugins/ChaosClothAsset/Source/ChaosClothAsset/Private/ChaosClothAsset/ClothCollection.cpp
--- a/Engine/Plugins/ChaosClothAsset/Source/ChaosClothAsset/Private/ChaosClothAsset/ClothCollection.cpp
+++ b/Engine/Plugins/ChaosClothAsset/Source/ChaosClothAsset/Private/ChaosClothAsset/ClothCollection.cpp
@@ -550,6 +550,13 @@ namespace UE::Chaos::ClothAsset
 			// Sim Vertices 3D Group
 			PreResizedSimPosition3D = &GetManagedArrayCollection()->AddAttribute<FVector3f>(ClothCollectionAttribute::PreResizedSimPosition3D, ClothCollectionGroup::SimVertices3D);
 
+			// Sim Accessory Meshes Group
+			SimAccessoryMeshName = &GetManagedArrayCollection()->AddAttribute<FName>(ClothCollectionAttribute::SimAccessoryMeshName, ClothCollectionGroup::SimAccessoryMeshes);
+			SimAccessoryMeshPosition3DAttribute = &GetManagedArrayCollection()->AddAttribute<FName>(ClothCollectionAttribute::SimAccessoryMeshPosition3DAttribute, ClothCollectionGroup::SimAccessoryMeshes);
+			SimAccessoryMeshNormalAttribute = &GetManagedArrayCollection()->AddAttribute<FName>(ClothCollectionAttribute::SimAccessoryMeshNormalAttribute, ClothCollectionGroup::SimAccessoryMeshes);
+			SimAccessoryMeshBoneIndicesAttribute = &GetManagedArrayCollection()->AddAttribute<FName>(ClothCollectionAttribute::SimAccessoryMeshBoneIndicesAttribute, ClothCollectionGroup::SimAccessoryMeshes);
+			SimAccessoryMeshBoneWeightsAttribute = &GetManagedArrayCollection()->AddAttribute<FName>(ClothCollectionAttribute::SimAccessoryMeshBoneWeightsAttribute, ClothCollectionGroup::SimAccessoryMeshes);
+
 			return;
 		}
 
diff --git a/Engine/Plugins/ChaosClothAsset/Source/ChaosClothAsset/Private/ChaosClothAsset/ClothGeometryTools.cpp b/Engine/Plugins/ChaosClothAsset/Source/ChaosClothAsset/Private/ChaosClothAsset/ClothGeometryTools.cpp
--- a/Engine/Plugins/ChaosClothAsset/Source/ChaosClothAsset/Private/ChaosClothAsset/ClothGeometryTools.cpp
+++ b/Engine/Plugins/ChaosClothAsset/Source/ChaosClothAsset/Private/ChaosClothAsset/ClothGeometryTools.cpp
@@ -1978,7 +1978,7 @@ PRAGMA_ENABLE_DEPRECATION_WARNINGS
 			// Can just initialize directly 
 			const int32 NumToSimVertices3D = ToCloth.GetNumSimVertices3D();
 			FCollectionClothSimAccessoryMeshFacade AccessoryFacade = ToCloth.AddGetSimAccessoryMesh();
-			AccessoryFacade.Initialize(GenerateUniqueAccessoryMeshName(), ToCloth.GetSimPosition3D().Left(NumToSimVertices3D), ToCloth.GetSimNormal().Left(NumToSimVertices3D), ToCloth.GetSimBoneIndices().Left(NumToSimVertices3D), ToCloth.GetSimBoneWeights().Left(NumToSimVertices3D));
+			AccessoryFacade.Initialize(GenerateUniqueAccessoryMeshName(), FromCloth.GetSimPosition3D().Left(NumToSimVertices3D), FromCloth.GetSimNormal().Left(NumToSimVertices3D), FromCloth.GetSimBoneIndices().Left(NumToSimVertices3D), FromCloth.GetSimBoneWeights().Left(NumToSimVertices3D));
 			return AccessoryFacade.GetSimAccessoryMeshIndex();
 		}
 
diff --git a/Engine/Plugins/ChaosClothAsset/Source/ChaosClothAssetEngine/Private/ChaosClothAsset/ClothSimulationAccessoryMesh.cpp b/Engine/Plugins/ChaosClothAsset/Source/ChaosClothAssetEngine/Private/ChaosClothAsset/ClothSimulationAccessoryMesh.cpp
--- a/Engine/Plugins/ChaosClothAsset/Source/ChaosClothAssetEngine/Private/ChaosClothAsset/ClothSimulationAccessoryMesh.cpp
+++ b/Engine/Plugins/ChaosClothAsset/Source/ChaosClothAssetEngine/Private/ChaosClothAsset/ClothSimulationAccessoryMesh.cpp
@@ -23,7 +23,7 @@ namespace UE::Chaos::ClothAsset
 			check(SimBoneIndices[VertexIndex].Num() == SimBoneWeights[VertexIndex].Num());
 			check(SimBoneIndices[VertexIndex].Num() <= FClothVertBoneData::MaxTotalInfluences);
 			
-			BoneData[VertexIndex].NumInfluences = SimBoneIndices[VertexIndex].Num();
+			BoneData[VertexIndex].NumInfluences = FMath::Min(SimBoneIndices[VertexIndex].Num(), FClothVertBoneData::MaxTotalInfluences);
 			for (int32 BoneIndex = 0; BoneIndex < BoneData[VertexIndex].NumInfluences; ++BoneIndex)
 			{
 				BoneData[VertexIndex].BoneIndices[BoneIndex] = SimBoneIndices[VertexIndex][BoneIndex];
diff --git a/Engine/Source/Runtime/Experimental/Chaos/Private/Chaos/PBDSphericalConstraint.cpp b/Engine/Source/Runtime/Experimental/Chaos/Private/Chaos/PBDSphericalConstraint.cpp
--- a/Engine/Source/Runtime/Experimental/Chaos/Private/Chaos/PBDSphericalConstraint.cpp
+++ b/Engine/Source/Runtime/Experimental/Chaos/Private/Chaos/PBDSphericalConstraint.cpp
@@ -107,8 +107,7 @@ void FPBDSphericalBackstopConstraint::SetProperties(
 	SetScale(MeshScale);
 }
 
-template<typename SolverParticlesOrRange>
-void FPBDSphericalBackstopConstraint::ApplyLegacyHelperISPC(SolverParticlesOrRange& Particles, const FSolverReal /*Dt*/) const
+void FPBDSphericalBackstopConstraint::ApplyLegacyHelperISPC(FSolverParticlesRange& Particles, const FSolverReal /*Dt*/) const
 {
 	check(bRealTypeCompatibleWithISPC);
 
@@ -129,7 +128,7 @@ void FPBDSphericalBackstopConstraint::ApplyLegacyHelperISPC(SolverParticlesOrRan
 				BackstopDistanceRange * GetScale(),
 				BackstopRadiusBase * GetScale(),
 				BackstopRadiusRange * GetScale(),
-				ParticleOffset,
+				0,
 				ParticleCount);
 		}
 		else
@@ -142,7 +141,7 @@ void FPBDSphericalBackstopConstraint::ApplyLegacyHelperISPC(SolverParticlesOrRan
 				BackstopDistanceBase * GetScale(),
 				BackstopRadiusBase * GetScale(),
 				BackstopRadiusRange * GetScale(),
-				ParticleOffset,
+				0,
 				ParticleCount);
 		}
 	}
@@ -156,7 +155,7 @@ void FPBDSphericalBackstopConstraint::ApplyLegacyHelperISPC(SolverParticlesOrRan
 			BackstopDistanceBase * GetScale(),
 			BackstopDistanceRange * GetScale(),
 			BackstopRadiusBase * GetScale(),
-			ParticleOffset,
+			0,
 			ParticleCount);
 	}
 	else
@@ -167,16 +166,13 @@ void FPBDSphericalBackstopConstraint::ApplyLegacyHelperISPC(SolverParticlesOrRan
 			(const ispc::FVector3f*)AnimationNormalsView.GetData(),
 			BackstopDistanceBase * GetScale(),
 			BackstopRadiusBase * GetScale(),
-			ParticleOffset,
+			0,
 			ParticleCount);
 	}
 #endif
 }
-template CHAOS_API void FPBDSphericalBackstopConstraint::ApplyLegacyHelperISPC(FSolverParticles& Particles, const FSolverReal /*Dt*/) const;
-template CHAOS_API void FPBDSphericalBackstopConstraint::ApplyLegacyHelperISPC(FSolverParticlesRange& Particles, const FSolverReal /*Dt*/) const;
 
-template<typename SolverParticlesOrRange>
-void FPBDSphericalBackstopConstraint::ApplyHelperISPC(SolverParticlesOrRange& Particles, const FSolverReal /*Dt*/) const
+void FPBDSphericalBackstopConstraint::ApplyHelperISPC(FSolverParticlesRange& Particles, const FSolverReal /*Dt*/) const
 {
 	check(bRealTypeCompatibleWithISPC);
 
@@ -197,7 +193,7 @@ void FPBDSphericalBackstopConstraint::ApplyHelperISPC(SolverParticlesOrRange& Pa
 				BackstopDistanceRange * GetScale(),
 				BackstopRadiusBase * GetScale(),
 				BackstopRadiusRange * GetScale(),
-				ParticleOffset,
+				0,
 				ParticleCount);
 		}
 		else
@@ -210,7 +206,7 @@ void FPBDSphericalBackstopConstraint::ApplyHelperISPC(SolverParticlesOrRange& Pa
 				BackstopDistanceBase * GetScale(),
 				BackstopRadiusBase * GetScale(),
 				BackstopRadiusRange * GetScale(),
-				ParticleOffset,
+				0,
 				ParticleCount);
 		}
 	}
@@ -224,7 +220,7 @@ void FPBDSphericalBackstopConstraint::ApplyHelperISPC(SolverParticlesOrRange& Pa
 			BackstopDistanceBase * GetScale(),
 			BackstopDistanceRange * GetScale(),
 			BackstopRadiusBase * GetScale(),
-			ParticleOffset,
+			0,
 			ParticleCount);
 	}
 	else
@@ -235,12 +231,10 @@ void FPBDSphericalBackstopConstraint::ApplyHelperISPC(SolverParticlesOrRange& Pa
 			(const ispc::FVector3f*)AnimationNormalsView.GetData(),
 			BackstopDistanceBase * GetScale(),
 			BackstopRadiusBase * GetScale(),
-			ParticleOffset,
+			0,
 			ParticleCount);
 	}
 #endif
 }
-template CHAOS_API void FPBDSphericalBackstopConstraint::ApplyHelperISPC(FSolverParticles& Particles, const FSolverReal /*Dt*/) const;
-template CHAOS_API void FPBDSphericalBackstopConstraint::ApplyHelperISPC(FSolverParticlesRange& Particles, const FSolverReal /*Dt*/) const;
 
 }  // End namespace Chaos::Softs
diff --git a/Engine/Source/Runtime/Experimental/Chaos/Public/Chaos/PBDSphericalConstraint.h b/Engine/Source/Runtime/Experimental/Chaos/Public/Chaos/PBDSphericalConstraint.h
--- a/Engine/Source/Runtime/Experimental/Chaos/Public/Chaos/PBDSphericalConstraint.h
+++ b/Engine/Source/Runtime/Experimental/Chaos/Public/Chaos/PBDSphericalConstraint.h
@@ -231,8 +231,13 @@ public:
 	void SetEnabled(bool bInEnabled) { bEnabled = bInEnabled; }
 	bool IsEnabled() const { return bEnabled; }
 
-	template<typename SolverParticlesOrRange>
-	void Apply(SolverParticlesOrRange& Particles, const FSolverReal Dt) const
+	void Apply(FSolverParticles& Particles, const FSolverReal Dt) const
+	{
+		FSolverParticlesRange ParticlesAsRange(&Particles, ParticleOffset, ParticleCount);
+		return Apply(ParticlesAsRange, Dt);
+	}
+
+	void Apply(FSolverParticlesRange& Particles, const FSolverReal Dt) const
 	{
 		SCOPE_CYCLE_COUNTER(STAT_PBD_SphericalBackstop);
 
@@ -339,15 +344,13 @@ public:
 
 private:
 
-	template<bool bHasBackstopDistance, bool bHasBackstopRadius, typename SolverParticlesOrRange>
-	void ApplyHelper(SolverParticlesOrRange& Particles, const FSolverReal Dt) const
+	template<bool bHasBackstopDistance, bool bHasBackstopRadius>
+	void ApplyHelper(FSolverParticlesRange& Particles, const FSolverReal Dt) const
 	{
 		const TConstArrayView<FSolverVec3> AnimationPositionsView = bMeshUseGlobalIndexation ? Particles.GetConstArrayView(AnimationPositions) : TConstArrayView<FSolverVec3>(AnimationPositions);
 		const TConstArrayView<FSolverVec3> AnimationNormalsView = bMeshUseGlobalIndexation ? Particles.GetConstArrayView(AnimationNormals) : TConstArrayView<FSolverVec3>(AnimationNormals);
-		PhysicsParallelFor(ParticleCount, [this, &Particles, Dt, &AnimationPositionsView, &AnimationNormalsView](int32 Index)  // TODO: profile need for parallel loop based on particle count
+		PhysicsParallelFor(ParticleCount, [this, &Particles, Dt, &AnimationPositionsView, &AnimationNormalsView](int32 ParticleIndex)  // TODO: profile need for parallel loop based on particle count
 		{
-			const int32 ParticleIndex = ParticleOffset + Index;
-
 			if (Particles.InvM(ParticleIndex) == 0)
 			{
 				return;
@@ -356,8 +359,8 @@ private:
 			const FSolverVec3& AnimationPosition = AnimationPositionsView[ParticleIndex];
 			const FSolverVec3& AnimationNormal = AnimationNormalsView[ParticleIndex];
 
-			const FSolverReal SphereOffsetDistance = (bHasBackstopDistance ? BackstopDistanceBase + BackstopDistanceRange * SphereOffsetDistances[Index] : BackstopDistanceBase) * Scale;
-			const FSolverReal Radius = (bHasBackstopRadius ? BackstopRadiusBase + BackstopRadiusRange * SphereRadii[Index] : BackstopRadiusBase) * Scale;
+			const FSolverReal SphereOffsetDistance = (bHasBackstopDistance ? BackstopDistanceBase + BackstopDistanceRange * SphereOffsetDistances[ParticleIndex] : BackstopDistanceBase) * Scale;
+			const FSolverReal Radius = (bHasBackstopRadius ? BackstopRadiusBase + BackstopRadiusRange * SphereRadii[ParticleIndex] : BackstopRadiusBase) * Scale;
 
 			const FSolverVec3 Center = AnimationPosition - (Radius + SphereOffsetDistance) * AnimationNormal;  // Non legacy version adds radius to the distance
 			const FSolverVec3 CenterToParticle = Particles.P(ParticleIndex) - Center;
@@ -377,15 +380,13 @@ private:
 		});
 	}
 
-	template<bool bHasBackstopDistance, bool bHasBackstopRadius, typename SolverParticlesOrRange>
-	void ApplyLegacyHelper(SolverParticlesOrRange& Particles, const FSolverReal Dt) const
+	template<bool bHasBackstopDistance, bool bHasBackstopRadius>
+	void ApplyLegacyHelper(FSolverParticlesRange& Particles, const FSolverReal Dt) const
 	{
 		const TConstArrayView<FSolverVec3> AnimationPositionsView = bMeshUseGlobalIndexation ? Particles.GetConstArrayView(AnimationPositions) : TConstArrayView<FSolverVec3>(AnimationPositions);
 		const TConstArrayView<FSolverVec3> AnimationNormalsView = bMeshUseGlobalIndexation ? Particles.GetConstArrayView(AnimationNormals) : TConstArrayView<FSolverVec3>(AnimationNormals);
-		PhysicsParallelFor(ParticleCount, [this, &Particles, Dt, &AnimationPositionsView, &AnimationNormalsView](int32 Index)  // TODO: profile need for parallel loop based on particle count
+		PhysicsParallelFor(ParticleCount, [this, &Particles, Dt, &AnimationPositionsView, &AnimationNormalsView](int32 ParticleIndex)  // TODO: profile need for parallel loop based on particle count
 		{
-			const int32 ParticleIndex = ParticleOffset + Index;
-
 			if (Particles.InvM(ParticleIndex) == 0)
 			{
 				return;
@@ -394,8 +395,8 @@ private:
 			const FSolverVec3& AnimationPosition = AnimationPositionsView[ParticleIndex];
 			const FSolverVec3& AnimationNormal = AnimationNormalsView[ParticleIndex];
 
-			const FSolverReal SphereOffsetDistance = (bHasBackstopDistance ? BackstopDistanceBase + BackstopDistanceRange * SphereOffsetDistances[Index] : BackstopDistanceBase) * Scale;
-			const FSolverReal Radius = (bHasBackstopRadius ? BackstopRadiusBase + BackstopRadiusRange * SphereRadii[Index] : BackstopRadiusBase) * Scale;
+			const FSolverReal SphereOffsetDistance = (bHasBackstopDistance ? BackstopDistanceBase + BackstopDistanceRange * SphereOffsetDistances[ParticleIndex] : BackstopDistanceBase) * Scale;
+			const FSolverReal Radius = (bHasBackstopRadius ? BackstopRadiusBase + BackstopRadiusRange * SphereRadii[ParticleIndex] : BackstopRadiusBase) * Scale;
 
 			const FSolverVec3 Center = AnimationPosition - SphereOffsetDistance * AnimationNormal;  // Legacy version already includes the radius within the distance
 			const FSolverVec3 CenterToParticle = Particles.P(ParticleIndex) - Center;
@@ -415,10 +416,8 @@ private:
 		});
 	}
 
-	template<typename SolverParticlesOrRange>
-	CHAOS_API void ApplyLegacyHelperISPC(SolverParticlesOrRange& Particles, const FSolverReal Dt) const;
-	template<typename SolverParticlesOrRange>
-	CHAOS_API void ApplyHelperISPC(SolverParticlesOrRange& Particles, const FSolverReal Dt) const;
+	CHAOS_API void ApplyLegacyHelperISPC(FSolverParticlesRange& Particles, const FSolverReal Dt) const;
+	CHAOS_API void ApplyHelperISPC(FSolverParticlesRange& Particles, const FSolverReal Dt) const;
 
 private:
 	const TArray<FSolverVec3>& AnimationPositions;  // Positions of spheres
