IAD installer - fix chunks not being properly pinned if install requests are queued to run later.

diff --git a/Engine/Source/Runtime/Experimental/IoStore/OnDemand/Private/OnDemandContentInstaller.cpp b/Engine/Source/Runtime/Experimental/IoStore/OnDemand/Private/OnDemandContentInstaller.cpp
--- a/Engine/Source/Runtime/Experimental/IoStore/OnDemand/Private/OnDemandContentInstaller.cpp
+++ b/Engine/Source/Runtime/Experimental/IoStore/OnDemand/Private/OnDemandContentInstaller.cpp
@@ -2,6 +2,7 @@
 
 #include "OnDemandContentInstaller.h"
 
+#include "Algo/Accumulate.h"
 #include "Algo/RemoveIf.h"
 #include "Async/UniqueLock.h"
 #include "Containers/RingBuffer.h"
@@ -427,7 +428,7 @@ void FOnDemandContentInstaller::ExecuteRequest(FRequest& Request)
 
 		void operator()(FRequest::FInstall&)
 		{
-			Installer.ExecuteInstallRequest(Request, /*bRemoveAlreadyCachedChunks*/ true);
+			Installer.ExecuteInstallRequest(Request);
 		}
 
 		void operator()(FRequest::FPurge&)
@@ -458,6 +459,43 @@ void FOnDemandContentInstaller::ExecuteRequest(FRequest& Request)
 	Visit(Visitor, Request.Variant);
 }
 
+void FOnDemandContentInstaller::PinCachedChunks(FOnDemandContentInstaller::FRequest::FInstall& InstallRequest, TFunctionRef<void(int32, int32, bool)> OnChunkFound) const
+{
+	const TSharedPtr<FOnDemandInternalContentHandle>& ContentHandle = InstallRequest.Args.ContentHandle.Handle;
+
+	// Find all chunks we need to fetch from the resolved chunk(s)
+	
+	for (int32 ContainerIndex = 0; Private::FResolvedContainerChunks& ResolvedChunks : InstallRequest.ResolvedChunks)
+	{
+		TArray<int32, TInlineAllocator<64>> CachedEntryIndices;
+		for (int32 EntryIndex : ResolvedChunks.EntryIndices)
+		{
+			const FOnDemandChunkEntry& Entry = ResolvedChunks.Container->ChunkEntries[EntryIndex];
+			const bool bCached = IoStore.InstallCache->IsChunkCached(Entry.Hash);
+			if (bCached)
+			{
+				CachedEntryIndices.Add(EntryIndex);
+			}
+			
+			OnChunkFound(ContainerIndex, EntryIndex, bCached);
+		}
+
+		// Add references to existing chunk(s)
+		if (CachedEntryIndices.IsEmpty() == false)
+		{
+			TUniqueLock Lock(IoStore.ContainerMutex);
+
+			FOnDemandChunkEntryReferences& References = ResolvedChunks.Container->FindOrAddChunkEntryReferences(*ContentHandle);
+			for (int32 EntryIndex : CachedEntryIndices)
+			{
+				References.Indices[EntryIndex] = true;
+			}
+		}
+
+		ContainerIndex++;
+	}
+}
+
 void FOnDemandContentInstaller::ProcessInstallRequest(FRequest& Request)
 {
 	using namespace UE::IoStore::Private;
@@ -564,57 +602,47 @@ void FOnDemandContentInstaller::ProcessInstallRequest(FRequest& Request)
 		return CompleteInstallRequest(Request);
 	}
 
-	uint64 TotalContentSize	= 0;
-	uint64 TotalInstallSize	= 0;
-
-	// Find all chunks we need to fetch from the resolved chunk(s)
+	bool bExecuteRequest = false;
 	{
-		for (int32 ContainerIndex = 0; FResolvedContainerChunks& ResolvedChunks : InstallRequest.ResolvedChunks)
+		TUniqueLock Lock(Mutex);
+		if (CurrentRequest == nullptr)
 		{
-			TArray<int32, TInlineAllocator<64>> CachedEntryIndices;
-			for (int32 EntryIndex : ResolvedChunks.EntryIndices)
-			{
-				const FOnDemandChunkEntry& Entry = ResolvedChunks.Container->ChunkEntries[EntryIndex];
-				const bool bCached = IoStore.InstallCache->IsChunkCached(Entry.Hash);
-				if (bCached)
-				{
-					CachedEntryIndices.Add(EntryIndex);
-				}
-				else
-				{
-					InstallRequest.HttpRequestHandles.Add(FChunkHttpRequestHandle
-					{
-						.Handle			= nullptr,
-						.ContainerIndex = ContainerIndex,
-						.EntryIndex		= EntryIndex
-					});
-					TotalInstallSize += Entry.GetDiskSize();
-				}
-				++InstallRequest.ResolvedChunkCount;
-				TotalContentSize += Entry.GetDiskSize();
-			}
+			CurrentRequest = &Request;
+			bExecuteRequest = true;
+		}
+	}
 
-			// Add references to existing chunk(s)
-			if (CachedEntryIndices.IsEmpty() == false)
-			{
-				TUniqueLock Lock(IoStore.ContainerMutex);
+	if (bExecuteRequest)
+	{
+		// Execute immediately, let ExecuteInstallRequest handle chunk pinning
 
-				FOnDemandChunkEntryReferences& References = ResolvedChunks.Container->FindOrAddChunkEntryReferences(*ContentHandle);
-				for (int32 EntryIndex : CachedEntryIndices)
-				{
-					References.Indices[EntryIndex] = true;
-				}
-			}
+		check(CurrentRequest == &Request);
+		ExecuteInstallRequest(Request);
+		return;
+	}
+
+	uint64 TotalContentSize	= 0;
+	uint64 TotalInstallSize	= 0;
 
-			ContainerIndex++;
+	// Pin any already cached chunks while waiting for execution
+	PinCachedChunks(InstallRequest, [&InstallRequest, &TotalContentSize, &TotalInstallSize](const int32 ContainerIndex, const int32 EntryIndex, const bool bCached)
+	{
+		const FOnDemandChunkEntry& Entry = InstallRequest.ResolvedChunks[ContainerIndex].Container->ChunkEntries[EntryIndex];
+		TotalContentSize += Entry.GetDiskSize();
+
+		if (bCached)
+		{
+			return;
 		}
-	}
+		
+		TotalInstallSize += Entry.GetDiskSize();
+	});
 
 	InstallRequest.Progress.TotalContentSize	= TotalContentSize;
 	InstallRequest.Progress.TotalInstallSize	= TotalInstallSize;
 	InstallRequest.Progress.CurrentInstallSize	= 0;
 
-	if (InstallRequest.HttpRequestHandles.IsEmpty())
+	if (TotalInstallSize == 0)
 	{
 		Request.ErrorCode = EIoErrorCode::Ok;
 
@@ -626,35 +654,19 @@ void FOnDemandContentInstaller::ProcessInstallRequest(FRequest& Request)
 		return CompleteInstallRequest(Request);
 	}
 
-	bool bExecuteRequest = false;
 	{
 		TUniqueLock Lock(Mutex);
-		if (CurrentRequest == nullptr)
-		{
-			CurrentRequest	= &Request;
-			bExecuteRequest	= true;
-		}
-		else
-		{
-			RequestQueue.HeapPush(&Request, RequestSortPredicate);
-		}
-	}
-
-	if (bExecuteRequest)
-	{
-		check(CurrentRequest == &Request);
-		const bool bRemoveAlreadyCachedChunks = false;
-		ExecuteInstallRequest(Request, bRemoveAlreadyCachedChunks);
+		RequestQueue.HeapPush(&Request, RequestSortPredicate);
 	}
 }
 
-void FOnDemandContentInstaller::ExecuteInstallRequest(FRequest& Request, bool bRemoveAlreadyCachedChunks)
+void FOnDemandContentInstaller::ExecuteInstallRequest(FRequest& Request)
 {
 	check(Request.IsInstall());
 	check(&Request == CurrentRequest);
 
 	FRequest::FInstall& InstallRequest = Request.AsInstall();
-	check(InstallRequest.HttpRequestHandles.IsEmpty() == false);
+	check(InstallRequest.HttpRequestHandles.IsEmpty());
 
 	UE_LOG(LogIoStoreOnDemand, Verbose, TEXT("Executing install request, SeqNo=%u, Priority=%d, ContentHandle=(%s)"),
 		Request.SeqNo, Request.Priority, *LexToString(InstallRequest.Args.ContentHandle));
@@ -664,30 +676,37 @@ void FOnDemandContentInstaller::ExecuteInstallRequest(FRequest& Request, bool bR
 		return CompleteInstallRequest(Request);
 	}
 
-	if (bRemoveAlreadyCachedChunks)
-	{
-		InstallRequest.HttpRequestHandles.SetNum(
-			Algo::RemoveIf(
-				InstallRequest.HttpRequestHandles,
-				[this, &InstallRequest](FChunkHttpRequestHandle& HttpRequest)
-				{
-					FSharedOnDemandContainer& Container		= InstallRequest.ResolvedChunks[HttpRequest.ContainerIndex].Container;
-					const FOnDemandChunkEntry& ChunkEntry	= Container->ChunkEntries[HttpRequest.EntryIndex];
-
-					if (IoStore.InstallCache->IsChunkCached(ChunkEntry.Hash))
-					{
-						InstallRequest.Progress.TotalInstallSize -= ChunkEntry.GetDiskSize();
-						return true;
-					}
+	uint64 TotalContentSize = 0;
+	uint64 TotalInstallSize	= 0;
 
-					return false;
-				}));
+	PinCachedChunks(InstallRequest, [&InstallRequest, &TotalContentSize, &TotalInstallSize](const int32 ContainerIndex, const int32 EntryIndex, const bool bCached)
+	{
+		const FOnDemandChunkEntry& Entry = InstallRequest.ResolvedChunks[ContainerIndex].Container->ChunkEntries[EntryIndex];
+		TotalContentSize += Entry.GetDiskSize();
 
-		if (InstallRequest.HttpRequestHandles.IsEmpty())
+		if (bCached)
 		{
-			Request.ErrorCode = EIoErrorCode::Ok;
-			return CompleteInstallRequest(Request);
+			return;
 		}
+
+		InstallRequest.HttpRequestHandles.Add(FChunkHttpRequestHandle
+		{
+			.Handle			= nullptr,
+			.ContainerIndex = ContainerIndex,
+			.EntryIndex		= EntryIndex
+		});
+		TotalInstallSize += Entry.GetDiskSize();
+	});
+
+	// TotalInstallSize may be different now
+	InstallRequest.Progress.TotalContentSize = TotalContentSize;
+	InstallRequest.Progress.TotalInstallSize = TotalInstallSize;
+	InstallRequest.Progress.CurrentInstallSize = 0;
+
+	if (InstallRequest.HttpRequestHandles.IsEmpty())
+	{
+		Request.ErrorCode = EIoErrorCode::Ok;
+		return CompleteInstallRequest(Request);
 	}
 
 	// Make sure we have enough space in the cache
@@ -1057,8 +1076,18 @@ void FOnDemandContentInstaller::CompleteInstallRequest(FRequest& Request)
 #if !UE_BUILD_SHIPPING
 	TArray<FString> DebugInfo;
 #endif
+
+	const uint64 ResolvedChunkCount = Algo::TransformAccumulate(
+		InstallRequest.ResolvedChunks, 
+		[](const Private::FResolvedContainerChunks& ContainerChunks)
+		{
+			return ContainerChunks.EntryIndices.Num();
+		},
+		uint64(0)
+	);
+
 	// Mark all resolved chunk(s) as referenced by the content handle and notify the package store to update
-	if (Request.ErrorCode == EIoErrorCode::Ok && InstallRequest.ResolvedChunkCount > 0)
+	if (Request.ErrorCode == EIoErrorCode::Ok && ResolvedChunkCount > 0)
 	{
 		{
 			FOnDemandContentHandle& ContentHandle = InstallRequest.Args.ContentHandle;
@@ -1095,7 +1124,7 @@ void FOnDemandContentInstaller::CompleteInstallRequest(FRequest& Request)
 	InstallResult.DurationInSeconds		= FPlatformTime::ToSeconds64(DurationCycles);
 
 	FOnDemandContentInstallerStats::OnRequestCompleted(
-		InstallRequest.ResolvedChunkCount,
+		ResolvedChunkCount,
 		InstallResult.Progress.TotalContentSize,
 		static_cast<uint64>(InstallRequest.HttpRequestHandles.Num()),
 		InstallResult.Progress.TotalInstallSize,
diff --git a/Engine/Source/Runtime/Experimental/IoStore/OnDemand/Private/OnDemandContentInstaller.h b/Engine/Source/Runtime/Experimental/IoStore/OnDemand/Private/OnDemandContentInstaller.h
--- a/Engine/Source/Runtime/Experimental/IoStore/OnDemand/Private/OnDemandContentInstaller.h
+++ b/Engine/Source/Runtime/Experimental/IoStore/OnDemand/Private/OnDemandContentInstaller.h
@@ -85,7 +85,6 @@ class FOnDemandContentInstaller
 			TArray<Private::FResolvedContainerChunks>	ResolvedChunks;
 			TArray<FChunkHttpRequestHandle>				HttpRequestHandles;
 			FOnDemandInstallProgress					Progress;
-			uint64										ResolvedChunkCount = 0;
 			uint64										DownloadedChunkCount = 0;
 			uint64										LastProgressCycles = 0;
 			std::atomic_bool							bHttpRequestsIssued{false};
@@ -213,8 +212,9 @@ public:
 private:
 	void							TryExecuteNextRequest();
 	void							ExecuteRequest(FRequest& Request);
+	void							PinCachedChunks(FRequest::FInstall& InstallRequest, TFunctionRef<void(int32, int32, bool)> OnChunkFound) const;
 	void							ProcessInstallRequest(FRequest& Request);
-	void							ExecuteInstallRequest(FRequest& Request, bool bRemoveAlreadyCachedChunks);
+	void							ExecuteInstallRequest(FRequest& Request);
 	void							ExecutePurgeRequest(FRequest& Request);
 	void							ExecuteDefragRequest(FRequest& Request);
 	void							ExecuteVerifyRequest(FRequest& Request);
