Iris
* Fix creation header bitguard feature triggering an assert when the creation data overflows due to running out of packet space

#rnx

diff --git a/Engine/Source/Runtime/Experimental/Iris/Core/Private/Iris/ReplicationSystem/NetObjectFactory.cpp b/Engine/Source/Runtime/Experimental/Iris/Core/Private/Iris/ReplicationSystem/NetObjectFactory.cpp
--- a/Engine/Source/Runtime/Experimental/Iris/Core/Private/Iris/ReplicationSystem/NetObjectFactory.cpp
+++ b/Engine/Source/Runtime/Experimental/Iris/Core/Private/Iris/ReplicationSystem/NetObjectFactory.cpp
@@ -91,7 +91,7 @@ bool UNetObjectFactory::WriteHeader(UE::Net::FNetRefHandle Handle, UE::Net::FNet
 	const bool bSuccess = SerializeHeader(FCreationHeaderContext(Handle, Bridge, this, Serialization), Header);
 
 #if IRIS_CREATIONHEADER_BITGUARD
-	if (bSuccess)
+	if (bSuccess && !Writer->IsOverflown())
 	{
 		// Go back and write the final header bit size
 		const uint32 BitsWritten = Writer->GetPosBits() - StartPos;
