Dataflow : fix dataflow async evaluation
* level of granularity is now the node level as input may cause concurrency issue if an run of eveluate may write to two output at the same time causing some cache entries to be modified while the next node reads it
* this is also preparation for using the evaluator class more generally

diff --git a/Engine/Source/Runtime/Experimental/Dataflow/Core/Private/Dataflow/DataflowContextEvaluator.cpp b/Engine/Source/Runtime/Experimental/Dataflow/Core/Private/Dataflow/DataflowContextEvaluator.cpp
--- a/Engine/Source/Runtime/Experimental/Dataflow/Core/Private/Dataflow/DataflowContextEvaluator.cpp
+++ b/Engine/Source/Runtime/Experimental/Dataflow/Core/Private/Dataflow/DataflowContextEvaluator.cpp
@@ -7,54 +7,27 @@ namespace UE::Dataflow
 {
 	void FContextEvaluator::ScheduleNodeEvaluation(const FDataflowNode& Node, FOnPostEvaluationFunction OnPostEvaluation)
 	{
-		if (Node.NumOutputs() == 0)
+		FEvaluationEntry Entry
 		{
-			// likely a terminal node or no output set the output to default Guid
-			FEvaluationEntry Entry
-			{
-				.WeakNode = Node.AsWeak(),
-				.Id = FNodeOutputId{ Node.GetGuid(), FGuid() },
-				.OnPostEvaluation = OnPostEvaluation,
-			};
-			ScheduleEvaluation(Entry);
-		}
-		else
-		{
-			// schedule all available outputs
-			for (const FDataflowOutput* Output : Node.GetOutputs())
-			{
-				if (Output)
-				{
-					ScheduleEvaluation(*Output, OnPostEvaluation);
-				}
-			}
-		}
+			.Id = FNodeId{ Node.GetGuid() },
+			.WeakNode = Node.AsWeak(),
+			.OnPostEvaluation = OnPostEvaluation,
+		};
+		ScheduleEvaluation(Entry);
 		Process();
 	}
 
 	void FContextEvaluator::ScheduleOutputEvaluation(const FDataflowOutput& Output, FOnPostEvaluationFunction OnPostEvaluation)
-	{
-		ScheduleEvaluation(Output, OnPostEvaluation);
-		Process();
-	}
-
-	void FContextEvaluator::ScheduleEvaluation(const FDataflowOutput& Output, FOnPostEvaluationFunction OnPostEvaluation)
 	{
 		if (const FDataflowNode* Node = Output.GetOwningNode())
 		{
-			FEvaluationEntry Entry
-			{
-				.WeakNode = Node->AsWeak(),
-				.Id = FNodeOutputId{ Node->GetGuid(), Output.GetGuid() },
-				.OnPostEvaluation = OnPostEvaluation,
-			};
-			ScheduleEvaluation(Entry);
+			ScheduleNodeEvaluation(*Node, OnPostEvaluation);
 		}
 	}
 
 	void FContextEvaluator::ScheduleEvaluation(const FEvaluationEntry& Entry)
 	{
-		if (RunningTasks.Contains(Entry.Id) || PendingEvaluationEntries.Contains(Entry.Id))
+		if (IsScheduledOrRunning(Entry.Id))
 		{
 			UE_LOG(LogTemp, Warning, TEXT("FContextEvaluator::ScheduleEvaluation : skipped [%s]"), *Entry.ToString());
 			return;
@@ -62,30 +35,40 @@ namespace UE::Dataflow
 		UE_LOG(LogTemp, Warning, TEXT("FContextEvaluator::ScheduleEvaluation : [%s]"), *Entry.ToString());
 		PendingEvaluationEntries.Add(Entry.Id, Entry);
 
-		// add upstream outputs
-		TArray<const FDataflowOutput*> InvalidUpstreamOutputs;
+		// add upstream nodes
 		if (TSharedPtr<const FDataflowNode> Node = Entry.WeakNode.Pin())
 		{
-			InvalidUpstreamOutputs.Reset();
-			FindInvalidUpstreamOutputs(*Node, InvalidUpstreamOutputs);
-			for (const FDataflowOutput* UpstreamOutput : InvalidUpstreamOutputs)
+			TArray<const FDataflowNode*> InvalidUpstreamNodes;
+			FindInvalidUpstreamNodes(*Node, InvalidUpstreamNodes);
+			for (const FDataflowNode* UpstreamNode: InvalidUpstreamNodes)
 			{
-				UE_LOG(LogTemp, Warning, TEXT("FContextEvaluator::ScheduleEvaluation :  [%s] -- Invalid Upstream output [%s]"), *Entry.ToString(), *UpstreamOutput->GetName().ToString());
+				UE_LOG(LogTemp, Warning, TEXT("FContextEvaluator::ScheduleEvaluation :  [%s] -- Invalid Upstream Node [%s]"), *Entry.ToString(), *UpstreamNode->GetName().ToString());
 			}
-			for (const FDataflowOutput* UpstreamOutput : InvalidUpstreamOutputs)
+			for (const FDataflowNode* UpstreamNode : InvalidUpstreamNodes)
 			{
-				ScheduleEvaluation(*UpstreamOutput, {});
+				FEvaluationEntry UpstreamEntry
+				{
+					.Id = FNodeId{ UpstreamNode->GetGuid() },
+					.WeakNode = UpstreamNode->AsWeak(),
+					.OnPostEvaluation = {},
+				};
+				ScheduleEvaluation(UpstreamEntry);
 			}
 		}
 	}
 
+	bool FContextEvaluator::IsScheduledOrRunning(const FNodeId& Id) const
+	{
+		return (RunningTasks.Contains(Id) || PendingEvaluationEntries.Contains(Id));
+	}
+
 	void FContextEvaluator::Cancel()
 	{
 		PendingEvaluationEntries.Reset();
 		CompletedTasks.Reset();
 	}
 
-	void FContextEvaluator::FindInvalidUpstreamOutputs(const FDataflowNode& Node, TArray<const FDataflowOutput*>& OutInvalidUpstreamOutputs)
+	void FContextEvaluator::FindInvalidUpstreamNodes(const FDataflowNode& Node, TArray<const FDataflowNode*>& OutInvalidUpstreamNodes)
 	{
 		for (const FDataflowInput* Input : Node.GetInputs())
 		{
@@ -93,15 +76,18 @@ namespace UE::Dataflow
 			{
 				if (const FDataflowOutput* UpstreamOutput = Input->GetConnection())
 				{
-					UE_LOG(LogTemp, Warning, TEXT("FContextEvaluator::FindInvalidUpstreamOutputs :  [%s] input[%s] -> output [%s]"), 
-						*Node.GetName().ToString(), 
+					UE_LOG(LogTemp, Warning, TEXT("FContextEvaluator::FindInvalidUpstreamOutputs :  [%s] input[%s] -> output [%s]"),
+						*Node.GetName().ToString(),
 						*Input->GetName().ToString(),
 						*UpstreamOutput->GetName().ToString()
 					);
 
 					if (!UpstreamOutput->HasValidData(OwningContext))
 					{
-						OutInvalidUpstreamOutputs.Add(UpstreamOutput);
+						if (const FDataflowNode* UpstreamNode = UpstreamOutput->GetOwningNode())
+						{
+							OutInvalidUpstreamNodes.AddUnique(UpstreamNode);
+						}
 					}
 				}
 			}
@@ -128,12 +114,12 @@ namespace UE::Dataflow
 
 	bool FContextEvaluator::TryScheduleTask(const FEvaluationEntry& Entry)
 	{
-		TArray<const FDataflowOutput*> InvalidUpstreamOutputs;
 		if (TSharedPtr<const FDataflowNode> Node = Entry.WeakNode.Pin())
 		{
-			InvalidUpstreamOutputs.Reset();
-			FindInvalidUpstreamOutputs(*Node, InvalidUpstreamOutputs);
-			if (InvalidUpstreamOutputs.IsEmpty())
+			// todo(ccaillaud) : can optimize by findingthe first one
+			TArray<const FDataflowNode*> InvalidUpstreamNodes;
+			FindInvalidUpstreamNodes(*Node, InvalidUpstreamNodes);
+			if (InvalidUpstreamNodes.IsEmpty())
 			{
 				ScheduleTask(Entry);
 				return true;
@@ -167,7 +153,7 @@ namespace UE::Dataflow
 		if (NumScheduleTasks == 0)
 		{
 			UE_LOG(LogTemp, Warning, TEXT("FContextEvaluator::Process : No Task Scheduled NumPendingTasks=[%d]"), PendingEvaluationEntries.Num());
-			for (const TPair<FNodeOutputId, FEvaluationEntry>& Entry : PendingEvaluationEntries)
+			for (const TPair<FNodeId, FEvaluationEntry>& Entry : PendingEvaluationEntries)
 			{
 				UE_LOG(LogTemp, Warning, TEXT("FContextEvaluator::Process : \t -[%s]"), *Entry.Value.ToString());
 			}
@@ -216,16 +202,24 @@ namespace UE::Dataflow
 						if (TSharedPtr<const FDataflowNode> Node = Entry.WeakNode.Pin())
 						{
 							Node->SetAsyncEvaluating(true);
-							if (const FDataflowOutput* Output = Node->FindOutput(Entry.Id.OutputId))
+							if (Node->NumOutputs() == 0)
 							{
-								Node->Evaluate(*ContextPtr, Output);
+								Node->Evaluate(*ContextPtr, nullptr);
 							}
-							else if (Node->NumOutputs() == 0)
+							else
 							{
-								Node->Evaluate(*ContextPtr, nullptr);
+								const TArray<FDataflowOutput*> Outputs = Node->GetOutputs();
+								for (const FDataflowOutput* Output : Outputs)
+								{
+									// todo(ccaillaud) : should we check if the oputput is frozen?
+									// todo(ccaillaud) : shoudl we have an option to skip the non connected outputs ? 
+									if (!Output->HasValidData(*ContextPtr))
+									{
+										Node->Evaluate(*ContextPtr, Output);
+									}
+								}
 							}
 							Node->SetAsyncEvaluating(false);
-
 							UE_LOG(LogTemp, Warning, TEXT("FContextEvaluator::EndTask : [%s]"), *Entry.ToString());
 						}
 					}
@@ -261,16 +255,11 @@ namespace UE::Dataflow
 	{
 		static FName UnknownName("-Unknown-");
 		FName NodeName = UnknownName;
-		FName OutputName = UnknownName;
 		if (TSharedPtr<const FDataflowNode> Node = WeakNode.Pin())
 		{
 			NodeName = Node->GetName();
-			if (const FDataflowOutput* Output = Node->FindOutput(Id.OutputId))
-			{
-				OutputName = Output->GetName();
-			}
 		}
-		return FString::Format(TEXT("{0}.{1}"), { NodeName.ToString(), OutputName.ToString() });
+		return NodeName.ToString();
 	}
 };
 
diff --git a/Engine/Source/Runtime/Experimental/Dataflow/Core/Public/Dataflow/DataflowContextEvaluator.h b/Engine/Source/Runtime/Experimental/Dataflow/Core/Public/Dataflow/DataflowContextEvaluator.h
--- a/Engine/Source/Runtime/Experimental/Dataflow/Core/Public/Dataflow/DataflowContextEvaluator.h
+++ b/Engine/Source/Runtime/Experimental/Dataflow/Core/Public/Dataflow/DataflowContextEvaluator.h
@@ -3,6 +3,7 @@
 #pragma once
 
 #include "Async/TaskGraphInterfaces.h"
+#include "Containers/Array.h"
 #include "Misc/Guid.h"
 #include "Templates/Function.h"
 #include "Templates/SharedPointer.h"
@@ -25,30 +26,19 @@ namespace UE::Dataflow
 	struct FContextEvaluator
 	{
 	public:
-		struct FNodeOutputId
-		{
-			FGuid NodeId;
-			FGuid OutputId;
-			bool operator == (const FNodeOutputId& Other) const
-			{
-				return NodeId == Other.NodeId && OutputId == Other.OutputId;
-			}
-			friend inline uint32 GetTypeHash(const FNodeOutputId& Id)
-			{
-				return HashCombine(GetTypeHash(Id.NodeId), GetTypeHash(Id.OutputId));
-			}
-		};
+		using FNodeId = FGuid;
+
 		struct FEvaluationEntry
 		{
+			FNodeId Id;
 			TWeakPtr<const FDataflowNode> WeakNode;
-			FNodeOutputId Id;
 			FOnPostEvaluationFunction OnPostEvaluation;
 
 			bool operator == (const FEvaluationEntry& Other) const
 			{
 				return Id == Other.Id;
 			}
-			bool operator == (const FNodeOutputId& Other) const
+			bool operator == (const FNodeId& Other) const
 			{
 				return (Id == Other);
 			}
@@ -72,10 +62,10 @@ namespace UE::Dataflow
 		int32 GetNumCompletedTasks() const { return CompletedTasks.Num(); }
 
 	private:
-		void FindInvalidUpstreamOutputs(const FDataflowNode& Node, TArray<const FDataflowOutput*>& OutInvalidUpstreamOutputs);
-		bool ShouldRunOnGameThread(const FDataflowNode& Node);
+		bool IsScheduledOrRunning(const FNodeId& Id) const;
 
-		void ScheduleEvaluation(const FDataflowOutput& Output, FOnPostEvaluationFunction OnPostEvaluation);
+		void FindInvalidUpstreamNodes(const FDataflowNode& Node, TArray<const FDataflowNode*>& OutInvalidUpstreamNodes);
+		bool ShouldRunOnGameThread(const FDataflowNode& Node);
 		void ScheduleEvaluation(const FEvaluationEntry& Entry);
 
 		bool TryScheduleTask(const FEvaluationEntry& Entry);
@@ -83,9 +73,9 @@ namespace UE::Dataflow
 		void ClearCompletedTasks();
 
 		FContext& OwningContext;
-		TMap<FNodeOutputId, FEvaluationEntry> PendingEvaluationEntries;
-		TMap<FNodeOutputId, FGraphEventRef> RunningTasks;
-		TSet<FNodeOutputId> CompletedTasks;
+		TMap<FNodeId, FEvaluationEntry> PendingEvaluationEntries;
+		TMap<FNodeId, FGraphEventRef> RunningTasks;
+		TSet<FNodeId> CompletedTasks;
 	};
 }
 
