Fix a bug causing ias.DisplayOnScreenStatistics to show an incorrect value for the max number of concurrent requests allowed.

#rnx

- Replaced use of hardcoded max number of concurrent requests with the cvar value.

diff --git a/Engine/Source/Runtime/Experimental/IoStore/OnDemand/Private/Statistics.cpp b/Engine/Source/Runtime/Experimental/IoStore/OnDemand/Private/Statistics.cpp
--- a/Engine/Source/Runtime/Experimental/IoStore/OnDemand/Private/Statistics.cpp
+++ b/Engine/Source/Runtime/Experimental/IoStore/OnDemand/Private/Statistics.cpp
@@ -68,6 +68,7 @@ static FAutoConsoleVariableRef CVar_DisplayOnScreenStatistics(
 #endif // UE_ENABLE_ONSCREEN_STATISTICS
 
 extern int32 GIasHttpRateLimitKiBPerSecond;
+extern int32 GIasHttpConcurrentRequests;
 
 ////////////////////////////////////////////////////////////////////////////////
 static FOnDemandImmediateAnalyticHandler OnDemandImmediateAnalyticEventHandler;
@@ -786,9 +787,8 @@ void FOnDemandIoBackendStats::PrintOnScreenStatistics(FCoreDelegates::FSeverityM
 				Builder << TEXT(" DataRate Cap: None");
 			}
 
-			const int32 MaxNumRequests = 8;
-			Builder << TEXT(" ConcurrentRequests: ") << GHttpNumConcurrentRequests << TEXT("/") << MaxNumRequests;
-			if (GHttpNumConcurrentRequests >= MaxNumRequests)
+			Builder << TEXT(" ConcurrentRequests: ") << GHttpNumConcurrentRequests << TEXT("/") << GIasHttpConcurrentRequests;
+			if (GHttpNumConcurrentRequests >= GIasHttpConcurrentRequests)
 			{
 				Builder << TEXT("(saturated)");
 			}
