Use CurrentInstallSize for FOnDemandContentInstallerStats::OnRequestCompleted.  If a request fails, it may not have downloaded all required chunks.

diff --git a/Engine/Source/Runtime/Experimental/IoStore/OnDemand/Private/OnDemandContentInstaller.cpp b/Engine/Source/Runtime/Experimental/IoStore/OnDemand/Private/OnDemandContentInstaller.cpp
--- a/Engine/Source/Runtime/Experimental/IoStore/OnDemand/Private/OnDemandContentInstaller.cpp
+++ b/Engine/Source/Runtime/Experimental/IoStore/OnDemand/Private/OnDemandContentInstaller.cpp
@@ -1127,7 +1127,7 @@ void FOnDemandContentInstaller::CompleteInstallRequest(FRequest& Request)
 		ResolvedChunkCount,
 		InstallResult.Progress.TotalContentSize,
 		static_cast<uint64>(InstallRequest.HttpRequestHandles.Num()),
-		InstallResult.Progress.TotalInstallSize,
+		InstallResult.Progress.CurrentInstallSize,
 		CacheHitRatio,
 		DurationCycles,
 		Request.ErrorCode);
