Disable loading of test certs if we're building out of tree.

#rnx

diff --git a/Engine/Source/Runtime/Experimental/IoStore/HttpClient/Private/Test.inl b/Engine/Source/Runtime/Experimental/IoStore/HttpClient/Private/Test.inl
--- a/Engine/Source/Runtime/Experimental/IoStore/HttpClient/Private/Test.inl
+++ b/Engine/Source/Runtime/Experimental/IoStore/HttpClient/Private/Test.inl
@@ -211,6 +211,7 @@ static void ThrottleTest(FAnsiStringView TestUrl)
 ////////////////////////////////////////////////////////////////////////////////
 static void TlsLoadRootCerts()
 {
+#if !defined(NO_UE_INCLUDES)
 	IFileManager& Ifm = IFileManager::Get();
 	FString PemPath = FPaths::EngineDir() / TEXT("Content/Certificates/ThirdParty/cacert.pem");
 	FArchive* Reader = Ifm.CreateFileReader(*PemPath);
@@ -224,6 +225,7 @@ static void TlsLoadRootCerts()
 	FCertRoots::SetDefault(MoveTemp(CaRoots));
 
 	delete Reader;
+#endif // NO_UE_INCLUDES
 }
 
 ////////////////////////////////////////////////////////////////////////////////
