Wrap up transactions in a simple reference object

#rnx

diff --git a/Engine/Source/Runtime/Experimental/IoStore/HttpClient/Private/Activity.inl b/Engine/Source/Runtime/Experimental/IoStore/HttpClient/Private/Activity.inl
--- a/Engine/Source/Runtime/Experimental/IoStore/HttpClient/Private/Activity.inl
+++ b/Engine/Source/Runtime/Experimental/IoStore/HttpClient/Private/Activity.inl
@@ -153,7 +153,7 @@ private:
 	};
 	UPTRINT				SinkParam;
 	FTicketSink			Sink;
-	FTransaction*		Transaction = nullptr;
+	FTransactRef		Transaction;
 	FBuffer				Buffer;
 
 	struct FHeaderRecord
@@ -204,10 +204,6 @@ FActivity::FActivity(const FParams& Params)
 ////////////////////////////////////////////////////////////////////////////////
 FActivity::~FActivity()
 {
-	if (Transaction != nullptr)
-	{
-		delete Transaction;
-	}
 }
 
 ////////////////////////////////////////////////////////////////////////////////
@@ -329,8 +325,8 @@ FTransactId FActivity::GetTransactId() const
 const FTransaction* FActivity::GetTransaction() const
 {
 	check(State >= EState::Send && State <= EState::Completed);
-	check(Transaction != nullptr);
-	return Transaction;
+	check(Transaction.IsValid());
+	return Transaction.Get();
 }
 
 ////////////////////////////////////////////////////////////////////////////////
@@ -411,7 +407,7 @@ const FStopwatch& FActivity::GetStopwatch() const
 FOutcome FActivity::Transact(FHttpPeer& Peer)
 {
 	Transaction = Peer.Transact(Buffer);
-	check(Transaction != nullptr);
+	check(Transaction.IsValid());
 
 	Transaction->Begin(GetMethod(), GetPath());
 
diff --git a/Engine/Source/Runtime/Experimental/IoStore/HttpClient/Private/Loop.inl b/Engine/Source/Runtime/Experimental/IoStore/HttpClient/Private/Loop.inl
--- a/Engine/Source/Runtime/Experimental/IoStore/HttpClient/Private/Loop.inl
+++ b/Engine/Source/Runtime/Experimental/IoStore/HttpClient/Private/Loop.inl
@@ -42,7 +42,7 @@ struct FActivityNode
 
 private:
 							FActivityNode(const FParams& Params);
-							~FActivityNode();
+							~FActivityNode() = default;
 };
 
 ////////////////////////////////////////////////////////////////////////////////
@@ -150,6 +150,7 @@ FActivityNode* FActivityNode::Create(FParams& Params, FAnsiStringView Url, FCert
 ////////////////////////////////////////////////////////////////////////////////
 void FActivityNode::Destroy(FActivityNode* Ptr)
 {
+	Ptr->~FActivityNode();
 	FMemory::Free(Ptr);
 }
 
diff --git a/Engine/Source/Runtime/Experimental/IoStore/HttpClient/Private/Peer.inl b/Engine/Source/Runtime/Experimental/IoStore/HttpClient/Private/Peer.inl
--- a/Engine/Source/Runtime/Experimental/IoStore/HttpClient/Private/Peer.inl
+++ b/Engine/Source/Runtime/Experimental/IoStore/HttpClient/Private/Peer.inl
@@ -535,7 +535,7 @@ public:
 	FHttpPeer&		operator = (const FHttpPeer& Rhs) = delete;
 					FHttpPeer(FParams&& Params);
 	FOutcome		Handshake();
-	FTransaction*	Transact(FBuffer& Buffer);
+	FTransactRef	Transact(FBuffer& Buffer);
 	FOutcome		GetPendingTransactId();
 
 private:
@@ -587,15 +587,15 @@ FOutcome FHttpPeer::Handshake()
 }
 
 ////////////////////////////////////////////////////////////////////////////////
-FTransaction* FHttpPeer::Transact(FBuffer& Buffer)
+FTransactRef FHttpPeer::Transact(FBuffer& Buffer)
 {
 	if (Proto == 1)
 	{
-		return CreateTransactOne(Buffer);
+		return FTransactRef(CreateTransactOne(Buffer));
 	}
 
 	check(!"unsupported");
-	return nullptr;
+	return FTransactRef();
 }
 
 ////////////////////////////////////////////////////////////////////////////////
diff --git a/Engine/Source/Runtime/Experimental/IoStore/HttpClient/Private/TransactionOne.inl b/Engine/Source/Runtime/Experimental/IoStore/HttpClient/Private/TransactionOne.inl
--- a/Engine/Source/Runtime/Experimental/IoStore/HttpClient/Private/TransactionOne.inl
+++ b/Engine/Source/Runtime/Experimental/IoStore/HttpClient/Private/TransactionOne.inl
@@ -866,15 +866,82 @@ public:
 	using DetailOne::FBody::FBody;
 };
 
+UPTRINT CreateTransactOne(FBuffer& Buffer)
+{
+	void* Ptr = FMemory::Malloc(sizeof(FTransactionOne), alignof(FTransactionOne));
+	new (Ptr) FTransactionOne(Buffer);
+	return UPTRINT(Ptr);
+}
+
+
+
+////////////////////////////////////////////////////////////////////////////////
 class FTransaction
-	: public FTransactionOne
 {
-	using FTransactionOne::FTransactionOne;
+public:
+#define TR_METHOD(x, ...) \
+	template <typename... Arg> auto x(Arg&&... Args) __VA_ARGS__ { \
+		return ((FTransactionOne*)this)->x(Forward<Arg>(Args)...); \
+	}
+	TR_METHOD(Begin)
+	TR_METHOD(AddHeader)
+	TR_METHOD(End)
+	TR_METHOD(TrySendRequest)
+	TR_METHOD(TryRecvResponse)
+	TR_METHOD(IsKeepAlive, const)
+	TR_METHOD(GetStatusCode, const)
+	TR_METHOD(GetStatusMessage, const)
+	TR_METHOD(GetContentLength, const)
+	TR_METHOD(IsChunked, const)
+	TR_METHOD(HasBody, const)
+	TR_METHOD(ReadHeaders, const)
+	TR_METHOD(GetRemaining, const)
+	TR_METHOD(TryRecv)
+#undef TR_METHOD
+
+private:
+	friend class FTransactRef;
+
+	void Destruct()
+	{
+		((FTransactionOne*)this)->~FTransactionOne();
+		FMemory::Free(this);
+	}
+};
+
+
+
+////////////////////////////////////////////////////////////////////////////////
+class FTransactRef
+{
+public:
+						FTransactRef() = default;
+						~FTransactRef();
+						FTransactRef(UPTRINT Ptr) : Self(Ptr) {}
+						FTransactRef(FTransactRef&&) = delete;
+						FTransactRef(const FTransactRef&) = delete;
+	void				operator = (const FTransactRef& Rhs) = delete;
+	void				operator = (FTransactRef&& Rhs) { Swap(Self, Rhs.Self); };
+	const FTransaction*	operator -> () const	{ return Get(); }
+	FTransaction*		operator -> ()			{ return Get(); }
+	bool				IsValid() const			{ return Get() != nullptr; }
+	const FTransaction*	Get() const				{ return (FTransaction*)Self; }
+	FTransaction*		Get()					{ return (FTransaction*)Self; }
+
+
+private:
+	UPTRINT				Self = 0;
 };
 
-FTransaction* CreateTransactOne(FBuffer& Buffer)
+////////////////////////////////////////////////////////////////////////////////
+FTransactRef::~FTransactRef()
 {
-	return new FTransaction(Buffer);
+	if (Self == 0)
+	{
+		return;
+	}
+
+	Get()->Destruct();
 }
 
 } // namespace UE::IoStore::HTTP
