Deleted unused destructor.

#rnx

diff --git a/Engine/Source/Runtime/Experimental/IoStore/HttpClient/Private/Activity.inl b/Engine/Source/Runtime/Experimental/IoStore/HttpClient/Private/Activity.inl
--- a/Engine/Source/Runtime/Experimental/IoStore/HttpClient/Private/Activity.inl
+++ b/Engine/Source/Runtime/Experimental/IoStore/HttpClient/Private/Activity.inl
@@ -86,7 +86,6 @@ struct FActivity
 	};
 
 						FActivity(const FParams& Params);
-						~FActivity();
 	void				AddHeader(FAnsiStringView Key, FAnsiStringView Value);
 	FOutcome			Tick(FHttpPeer& Peer, int32* MaxRecvSize=nullptr);
 	void				SetSink(FTicketSink&& InSink, UPTRINT Param);
@@ -201,11 +200,6 @@ FActivity::FActivity(const FParams& Params)
 	LengthScore = uint8(Pow2);
 }
 
-////////////////////////////////////////////////////////////////////////////////
-FActivity::~FActivity()
-{
-}
-
 ////////////////////////////////////////////////////////////////////////////////
 void FActivity::AddHeader(FAnsiStringView Key, FAnsiStringView Value)
 {
