Missing trace event

#rnx

diff --git a/Engine/Source/Runtime/Experimental/IoStore/HttpClient/Private/Loop.inl b/Engine/Source/Runtime/Experimental/IoStore/HttpClient/Private/Loop.inl
--- a/Engine/Source/Runtime/Experimental/IoStore/HttpClient/Private/Loop.inl
+++ b/Engine/Source/Runtime/Experimental/IoStore/HttpClient/Private/Loop.inl
@@ -150,6 +150,7 @@ FActivityNode* FActivityNode::Create(FParams& Params, FAnsiStringView Url, FCert
 ////////////////////////////////////////////////////////////////////////////////
 void FActivityNode::Destroy(FActivityNode* Ptr)
 {
+	Trace(Ptr, ETrace::ActivityDestroy, 0);
 	Ptr->~FActivityNode();
 	FMemory::Free(Ptr);
 }
