Remove the paranoid thread safety checks from FMultiEndpointHttpClient.

#rnx

- These checks were added when all systems were moved over to use FMultiEndpointHttpClient to cacth any potential problems but no issues were ever found and there is no value in keeping the checks.

diff --git a/Engine/Source/Runtime/Experimental/IoStore/OnDemand/IoStoreOnDemand.Build.cs b/Engine/Source/Runtime/Experimental/IoStore/OnDemand/IoStoreOnDemand.Build.cs
--- a/Engine/Source/Runtime/Experimental/IoStore/OnDemand/IoStoreOnDemand.Build.cs
+++ b/Engine/Source/Runtime/Experimental/IoStore/OnDemand/IoStoreOnDemand.Build.cs
@@ -26,14 +26,5 @@ public class IoStoreOnDemand : ModuleRules
 		{
 			PrivateDependencyModuleNames.AddRange(new string[] { "S3Client", "RSA" });
 		}
-
-		if (Target.Configuration != UnrealTargetConfiguration.Shipping)
-		{
-			PrivateDefinitions.Add("UE_HTTPCLIENT_THREADSAFE_CHECKS=1");
-		}
-		else
-		{
-			PrivateDefinitions.Add("UE_HTTPCLIENT_THREADSAFE_CHECKS=0");
-		}
 	}
 }
diff --git a/Engine/Source/Runtime/Experimental/IoStore/OnDemand/Private/OnDemandHttpClient.cpp b/Engine/Source/Runtime/Experimental/IoStore/OnDemand/Private/OnDemandHttpClient.cpp
--- a/Engine/Source/Runtime/Experimental/IoStore/OnDemand/Private/OnDemandHttpClient.cpp
+++ b/Engine/Source/Runtime/Experimental/IoStore/OnDemand/Private/OnDemandHttpClient.cpp
@@ -12,12 +12,6 @@
 #include "Misc/Paths.h"
 #include "Statistics.h"
 
-#if UE_HTTPCLIENT_THREADSAFE_CHECKS
-	#define UE_CHECK_THREADSAFETY(MethodName) checkf(OwningThread == 0 || OwningThread == FPlatformTLS::GetCurrentThreadId() || !FGenericPlatformProcess::SupportsMultithreading(), #MethodName TEXT("was called on the on the wrong thread!"))
-#else
-	#define UE_CHECK_THREADSAFETY(MethodName)
-#endif //UE_HTTPCLIENT_THREADSAFE_CHECKS
-
 namespace UE::IoStore
 {
 
@@ -144,13 +138,6 @@ FMultiEndpointHttpClient::FMultiEndpointHttpClient(const FMultiEndpointHttpClien
 	: Config(InConfig)
 {
 	EventLoop.SetFailTimeout(Config.TimeoutMs);
-
-#if UE_HTTPCLIENT_THREADSAFE_CHECKS
-	if (Config.bEnableThreadSafetyChecks)
-	{
-		OwningThread = FPlatformTLS::GetCurrentThreadId();
-	}
-#endif //UE_HTTPCLIENT_THREADSAFE_CHECKS
 }
 
 TUniquePtr<FMultiEndpointHttpClient> FMultiEndpointHttpClient::Create(const FMultiEndpointHttpClientConfig& Config)
@@ -233,8 +220,6 @@ void FMultiEndpointHttpClient::Get(const FOnDemandHostGroup& HostGroup, FAnsiStr
 
 void FMultiEndpointHttpClient::Get(const FOnDemandHostGroup& HostGroup, FAnsiStringView RelativeUrl, const FIoOffsetAndLength& ChunkRange, TArray<FAnsiString>&& Headers, FOnHttpResponse&& OnResponse)
 {
-	UE_CHECK_THREADSAFETY(FMultiEndpointHttpClient::Get);
-
 	FConnection& Connection = GetConnection(HostGroup);
 
 	IssueRequest(FRequest
@@ -251,8 +236,6 @@ void FMultiEndpointHttpClient::Get(const FOnDemandHostGroup& HostGroup, FAnsiStr
 
 bool FMultiEndpointHttpClient::Tick(int32 WaitTimeMs, uint32 MaxKiBPerSecond)
 {
-	UE_CHECK_THREADSAFETY(FMultiEndpointHttpClient::Tick);
-
 	EventLoop.Throttle(MaxKiBPerSecond);
 	const uint32 TicketCount = EventLoop.Tick(WaitTimeMs);
 
@@ -289,8 +272,6 @@ void FMultiEndpointHttpClient::UpdateConnections()
 {
 	TRACE_CPUPROFILER_EVENT_SCOPE(FMultiEndpointHttpClient::UpdateConnections);
 
-	UE_CHECK_THREADSAFETY(FMultiEndpointHttpClient::UpdateConnections);
-
 	for (TUniquePtr<FConnection>& Connection : Connections)
 	{
 		check(Connection.IsValid());
@@ -309,8 +290,6 @@ void FMultiEndpointHttpClient::UpdateConnections()
 
 void FMultiEndpointHttpClient::IssueRequest(FRequest&& Request)
 {
-	UE_CHECK_THREADSAFETY(FMultiEndpointHttpClient::IssueRequest);
-
 	using namespace UE::IoStore::HTTP;
 
 	check(Request.Connection.HostGroup.IsEmpty() == false);
@@ -464,8 +443,6 @@ void FMultiEndpointHttpClient::IssueRequest(FRequest&& Request)
 
 void FMultiEndpointHttpClient::RetryRequest(FRequest&& Request)
 {
-	UE_CHECK_THREADSAFETY(FMultiEndpointHttpClient::RetryRequest);
-
 	FConnection& Connection = Request.Connection;
 
 	// Try a different host URL after the first retry
@@ -506,8 +483,6 @@ TUniquePtr<HTTP::FConnectionPool> FMultiEndpointHttpClient::CreateConnection(FAn
 
 FMultiEndpointHttpClient::FConnection& FMultiEndpointHttpClient::GetConnection(const FOnDemandHostGroup& HostGroup)
 {
-	UE_CHECK_THREADSAFETY(FMultiEndpointHttpClient::GetConnection);
-
 	for (TUniquePtr<FConnection>& Conn : Connections)
 	{
 		check(Conn.IsValid());
@@ -531,8 +506,6 @@ FMultiEndpointHttpClient::FConnection& FMultiEndpointHttpClient::GetConnection(c
 
 FMultiEndpointHttpClient::FConnection* FMultiEndpointHttpClient::FindConnection(const FOnDemandHostGroup& HostGroup)
 {
-	UE_CHECK_THREADSAFETY(FMultiEndpointHttpClient::FindConnection);
-
 	for (TUniquePtr<FConnection>& Conn : Connections)
 	{
 		check(Conn.IsValid());
@@ -596,5 +569,3 @@ void FMultiEndpointHttpClient::Log(const FRequest& Request, const HTTP::FTicketS
 }
 
 } // namespace UE::IoStore
-
-#undef UE_CHECK_THREADSAFETY
diff --git a/Engine/Source/Runtime/Experimental/IoStore/OnDemand/Private/OnDemandHttpClient.h b/Engine/Source/Runtime/Experimental/IoStore/OnDemand/Private/OnDemandHttpClient.h
--- a/Engine/Source/Runtime/Experimental/IoStore/OnDemand/Private/OnDemandHttpClient.h
+++ b/Engine/Source/Runtime/Experimental/IoStore/OnDemand/Private/OnDemandHttpClient.h
@@ -169,10 +169,6 @@ private:
 	TArray<TUniquePtr<FConnection>>		Connections;
 	HTTP::FEventLoop					EventLoop;
 	TArray<FRequest>					Retries;
-
-#if UE_HTTPCLIENT_THREADSAFE_CHECKS
-	uint32							OwningThread = 0;
-#endif //UE_HTTPCLIENT_THREADSAFE_CHECKS
 };
 
 } // namespace UE::IoStore
