Addressed narrowing warning.

#rnx

diff --git a/Engine/Source/Runtime/Experimental/IoStore/HttpClient/Private/Peer.inl b/Engine/Source/Runtime/Experimental/IoStore/HttpClient/Private/Peer.inl
--- a/Engine/Source/Runtime/Experimental/IoStore/HttpClient/Private/Peer.inl
+++ b/Engine/Source/Runtime/Experimental/IoStore/HttpClient/Private/Peer.inl
@@ -152,7 +152,7 @@ static int32 Ssl_BioWrite(BIO* Bio, const char* Data, size_t Size, size_t* Bytes
 	BIO_clear_retry_flags(Bio);
 
 	auto* Peer = (FPeer*)BIO_get_data(Bio);
-	FOutcome Outcome = Peer->Send(Data, Size);
+	FOutcome Outcome = Peer->Send(Data, int32(Size));
 	if (Outcome.IsWaiting())
 	{
 		BIO_set_retry_write(Bio);
@@ -175,7 +175,7 @@ static int32 Ssl_BioRead(BIO* Bio, char* Data, size_t Size, size_t* BytesRead)
 	BIO_clear_retry_flags(Bio);
 
 	auto* Peer = (FPeer*)BIO_get_data(Bio);
-	FOutcome Outcome = Peer->Recv(Data, Size);
+	FOutcome Outcome = Peer->Recv(Data, int32(Size));
 	if (Outcome.IsWaiting())
 	{
 		BIO_set_retry_read(Bio);
@@ -289,7 +289,7 @@ static int32 Ssl_GetProtocolVersion(SSL* Ssl)
 	for (int32 i = 0; i < UE_ARRAY_COUNT(Candidates); ++i)
 	{
 		const FAnsiStringView& Candidate = Candidates[i];
-		if (AlpnProtoLen != Candidate.Len())
+		if (AlpnProtoLen != uint32(Candidate.Len()))
 		{
 			continue;
 		}
