Remove unused cvars from OnDemandIoStore.cpp

#rnx

- The following cvars were still registered but the implementation code has long since been deleted:
-- iostore.MinInstallDurationSeconds
-- iostore.InstallRestrictedToKibs
-- iostore.DebugBreakOnRequest
-- iostore.TreatMissingOnDemandChunksAsError

diff --git a/Engine/Source/Runtime/Experimental/IoStore/OnDemand/Private/OnDemandIoStore.cpp b/Engine/Source/Runtime/Experimental/IoStore/OnDemand/Private/OnDemandIoStore.cpp
--- a/Engine/Source/Runtime/Experimental/IoStore/OnDemand/Private/OnDemandIoStore.cpp
+++ b/Engine/Source/Runtime/Experimental/IoStore/OnDemand/Private/OnDemandIoStore.cpp
@@ -59,37 +59,6 @@ static FAutoConsoleVariableRef CVar_HttpIoDispatcherEnabled(
 	ECVF_ReadOnly
 );
 
-bool GIoStoreOnDemandTreatMissingChunksAsError = true;
-static FAutoConsoleVariableRef CVar_IoStoreOnDemandTreatMissingChunksAsError(
-	TEXT("iostore.TreatMissingOnDemandChunksAsError"),
-	GIoStoreOnDemandTreatMissingChunksAsError,
-	TEXT("Whether to treat missing chunks as error when installing on-demand content."),
-	ECVF_ReadOnly
-);
-
-#if !UE_BUILD_SHIPPING
-float GIoStoreOnDemandMinInstallDurationSeconds = 0.0f;
-static FAutoConsoleVariableRef CVar_IoStoreOnDemandMinInstallDurationSeconds(
-	TEXT("iostore.MinInstallDurationSeconds"),
-	GIoStoreOnDemandMinInstallDurationSeconds,
-	TEXT("Sets the minimum install duration in seconds.")
-);
-
-float GIoStoreOnDemandRestrictKibsPerSecond = 0.0f;
-static FAutoConsoleVariableRef CVar_IoStoreOnDemandRestrictKibsPerSecond(
-	TEXT("iostore.InstallRestrictedToKibs"),
-	GIoStoreOnDemandRestrictKibsPerSecond,
-	TEXT("After install delay to simulate a restricted throughput rate.")
-);
-
-
-FString GIoStoreBreakOnRequest = "";
-static FAutoConsoleVariableRef CVar_IoStoreDebugBreakOnRequest(
-		TEXT("iostore.DebugBreakOnRequest"),
-		GIoStoreBreakOnRequest,
-		TEXT("Debug break when we detect this package name."));
-#endif
-
 ///////////////////////////////////////////////////////////////////////////////
 namespace UE::IoStore
 {
