Consolidated some validation into helper functions and moved common has-body check out of transactions and into activities.

#rnx

diff --git a/Engine/Source/Runtime/Experimental/IoStore/HttpClient/Private/Activity.inl b/Engine/Source/Runtime/Experimental/IoStore/HttpClient/Private/Activity.inl
--- a/Engine/Source/Runtime/Experimental/IoStore/HttpClient/Private/Activity.inl
+++ b/Engine/Source/Runtime/Experimental/IoStore/HttpClient/Private/Activity.inl
@@ -497,41 +497,40 @@ FOutcome FActivity::RecvMessage(FHttpPeer& Peer)
 	FIoBuffer* PriorDest = Dest; // to retain unioned Host ptr (redirect uses it in sink)
 	CallSink();
 
-	// Check the user gave us a destination for content
-	if (Transaction->HasBody())
+	// HEAD methods
+	bool bHasBody = !GetMethod().Equals("HEAD", ESearchCase::IgnoreCase);
+	bHasBody &= (Transaction->GetContentLength() > 0) | int32(Transaction->IsChunked());
+	if (!bHasBody)
 	{
-		if (Dest == PriorDest)
-		{
-			Outcome = FOutcome::Error("User did not provide a destination buffer");
-			Fail(Outcome);
-			return Outcome;
-		}
+		ChangeState(EState::RecvDone);
+		return FOutcome::Ok(int32(EStage::Content));
+	}
 
-		// The user seems to have forgotten something. Let's help them along
-		if (int32 DestSize = int32(Dest->GetSize()); DestSize == 0)
-		{
-			static const uint32 DefaultChunkSize = 4 << 10;
-			uint32 Size = bChunked ? DefaultChunkSize : uint32(ContentLength);
-			*Dest = FIoBuffer(Size);
-		}
-		else if (!bChunked && DestSize < ContentLength)
-		{
-			// todo: support piece-wise transfer of content (a la chunked).
-			Outcome = FOutcome::Error("Destination buffer too small");
-			Fail(Outcome);
-			return Outcome;
-		}
-		else if (enum { MinStreamBuf = 256 }; bChunked && DestSize < MinStreamBuf)
-		{
-			*Dest = FIoBuffer(MinStreamBuf);
-		}
+	// Check the user gave us a destination for content
+	if (Dest == PriorDest)
+	{
+		Outcome = FOutcome::Error("User did not provide a destination buffer");
+		Fail(Outcome);
+		return Outcome;
 	}
 
-	// HEAD methods
-	if (!Transaction->HasBody())
+	// The user seems to have forgotten something. Let's help them along
+	if (int32 DestSize = int32(Dest->GetSize()); DestSize == 0)
 	{
-		ChangeState(EState::RecvDone);
-		return FOutcome::Ok(int32(EStage::Content));
+		static const uint32 DefaultChunkSize = 4 << 10;
+		uint32 Size = bChunked ? DefaultChunkSize : uint32(ContentLength);
+		*Dest = FIoBuffer(Size);
+	}
+	else if (!bChunked && DestSize < ContentLength)
+	{
+		// todo: support piece-wise transfer of content (a la chunked).
+		Outcome = FOutcome::Error("Destination buffer too small");
+		Fail(Outcome);
+		return Outcome;
+	}
+	else if (enum { MinStreamBuf = 256 }; bChunked && DestSize < MinStreamBuf)
+	{
+		*Dest = FIoBuffer(MinStreamBuf);
 	}
 
 	// We're all set to go and get content
diff --git a/Engine/Source/Runtime/Experimental/IoStore/HttpClient/Private/Loop.inl b/Engine/Source/Runtime/Experimental/IoStore/HttpClient/Private/Loop.inl
--- a/Engine/Source/Runtime/Experimental/IoStore/HttpClient/Private/Loop.inl
+++ b/Engine/Source/Runtime/Experimental/IoStore/HttpClient/Private/Loop.inl
@@ -1191,7 +1191,8 @@ bool FEventLoop::Redirect(const FTicketStatus& Status, FTicketSink& OuterSink)
 
 	const auto& Activity = (FActivity&)Response; // todo: yuk
 
-	check(!Activity.GetTransaction()->HasBody() || Response.GetContentLength() == 0); // should we ever hit this, we'll fix it
+	// should we ever hit this, we'll fix it
+	check(Activity.GetMethod().Equals("HEAD", ESearchCase::IgnoreCase) || Response.GetContentLength() == 0);
 
 	// Original method should remain unchanged
 	FAnsiStringView Method = Activity.GetMethod();
diff --git a/Engine/Source/Runtime/Experimental/IoStore/HttpClient/Private/TransactionOne.inl b/Engine/Source/Runtime/Experimental/IoStore/HttpClient/Private/TransactionOne.inl
--- a/Engine/Source/Runtime/Experimental/IoStore/HttpClient/Private/TransactionOne.inl
+++ b/Engine/Source/Runtime/Experimental/IoStore/HttpClient/Private/TransactionOne.inl
@@ -132,6 +132,21 @@ static void EnumerateHeaders(FAnsiStringView Headers, LambdaType&& Lambda)
 
 
 
+////////////////////////////////////////////////////////////////////////////////
+bool IsContentless(uint32 StatusCode)
+{
+	return (StatusCode < 200) | (StatusCode == 204) | (StatusCode == 304);
+}
+
+////////////////////////////////////////////////////////////////////////////////
+bool IsStatusCodeOk(uint32 StatusCode)
+{
+	enum { MinCode = 100, MaxCode = 999 };
+	return uint32(StatusCode - MinCode) <= uint32(MaxCode - MinCode);
+}
+
+
+
 namespace DetailOne
 {
 
@@ -293,7 +308,6 @@ public:
 	FAnsiStringView	GetStatusMessage() const;
 	int64			GetContentLength() const;
 	bool			IsChunked() const;
-	bool			HasBody() const;
 	void			ReadHeaders(FResponse::FHeaderSink Sink) const;
 
 protected:
@@ -308,8 +322,7 @@ private:
 	uint16			StatusCode = 0;
 	uint8			bKeepAlive : 1;
 	uint8			bChunked : 1;
-	uint8			bHasBody : 1;
-	uint8			_Unused : 6;
+	uint8			_Unused : 7;
 };
 
 ////////////////////////////////////////////////////////////////////////////////
@@ -347,12 +360,6 @@ bool FStatusHeaders::IsChunked() const
 	return bChunked;
 }
 
-////////////////////////////////////////////////////////////////////////////////
-bool FStatusHeaders::HasBody() const
-{
-	return bHasBody;
-}
-
 ////////////////////////////////////////////////////////////////////////////////
 void FStatusHeaders::ReadHeaders(FResponse::FHeaderSink Sink) const
 {
@@ -451,7 +458,7 @@ FOutcome FStatusHeaders::Parse()
 		Offsets.Message - Offsets.StatusCode
 	);
 	StatusCode = uint16(CrudeToInt(StatusCodeView));
-	if (uint32(StatusCode - 100) > 899)
+	if (!IsStatusCodeOk(StatusCode))
 	{
 		return FOutcome::Error("Invalid status code", StatusCode);
 	}
@@ -489,29 +496,26 @@ FOutcome FStatusHeaders::Parse()
 		}
 	);
 
-	bHasBody = !IsHeadMethod();
-
-	if ((StatusCode < 200) | (StatusCode == 204) | (StatusCode == 304)) // | !connect
-	{	
-		ContentLength = 0;
-		bHasBody = false;
-	}
-	else if (Flags & Flag_XferEnc)
+	if (Flags & Flag_XferEnc)
 	{
 		if (!bChunked)				return FOutcome::Error("Unsupported Transfer-Encoding");
 		if (Flags & Flag_Length)	return FOutcome::Error("Chunked yet with length");
 		ContentLength = -1;
-		bHasBody &= true;
 	}
 	else if (uint32 Overshoot = GetSize() - MessageRight; Flags & Flag_Length)
 	{
 		if (ContentLength < 0)			return FOutcome::Error("Invalid Content-Length field", uint32(ContentLength));
 		if (Overshoot > ContentLength)	return FOutcome::Error("More data received that expected");
+	}
+	else if (Overshoot)
+	{
+		return FOutcome::Error("Received content when none was expected");
+	}
 
-		bHasBody &= (ContentLength > 0);
+	if (IsContentless(StatusCode))
+	{
+		ContentLength = 0;
 	}
-	else if (bHasBody)	return FOutcome::Error("Missing Content-Length field");
-	else if (Overshoot)	return FOutcome::Error("Received content when none was expected");
 
 	return FOutcome::Ok();
 }
@@ -893,7 +897,6 @@ public:
 	TR_METHOD(GetStatusMessage, const)
 	TR_METHOD(GetContentLength, const)
 	TR_METHOD(IsChunked, const)
-	TR_METHOD(HasBody, const)
 	TR_METHOD(ReadHeaders, const)
 	TR_METHOD(GetRemaining, const)
 	TR_METHOD(TryRecv)
