Move IAX debug cvars 'ias.InvokeHttpFailure' and 'ias.RunRequestTest' so they do not rely on FOnDemandIoBackend

#rnx

- Neither of these cvars rely on FOnDemandIoBackend, it just happened to be an easy place to put the code when they were written.
- Since we are moving away from FOnDemandIoBackend we need to move these cvars so that they keep working.
- I didn't want to move them to OnDemandIoStore.cpp as that file is getting quite large already, so instead I've added a new class to host our generic debug commands.
- No functionality should be affected.

diff --git a/Engine/Source/Runtime/Experimental/IoStore/OnDemand/Private/DebugCommands.cpp b/Engine/Source/Runtime/Experimental/IoStore/OnDemand/Private/DebugCommands.cpp
new file mode 100644
--- /dev/null
+++ b/Engine/Source/Runtime/Experimental/IoStore/OnDemand/Private/DebugCommands.cpp
@@ -0,0 +1,108 @@
+// Copyright Epic Games, Inc. All Rights Reserved.
+
+#include "DebugCommands.h"
+
+#include "HAL/IConsoleManager.h"
+#include "HAL/PlatformTime.h"
+#include "IO/IoChunkId.h"
+#include "IO/IoDispatcher.h"
+#include "IO/IoStoreOnDemand.h"
+#include "Misc/AssertionMacros.h"
+#include "OnDemandIoStore.h"
+#include "String/Numeric.h"
+#include "IasHostGroup.h"
+
+#if !UE_BUILD_SHIPPING
+
+namespace UE::IoStore
+{
+
+FOnDemandDebugCommands::FOnDemandDebugCommands(FOnDemandIoStore* InOnDemandIoStore)
+	: OnDemandIoStore(InOnDemandIoStore)
+{
+	check(InOnDemandIoStore != nullptr);
+
+	BindConsoleCommands();
+}
+
+FOnDemandDebugCommands::~FOnDemandDebugCommands()
+{
+	UnbindConsoleCommands();
+}
+
+void FOnDemandDebugCommands::BindConsoleCommands()
+{
+#if !NO_CVARS
+	DynamicConsoleCommands.Emplace(
+		IConsoleManager::Get().RegisterConsoleCommand(
+			TEXT("ias.InvokeHttpFailure"),
+			TEXT("Marks the current ias http connections as failed and forcing them  to try to reconnect"),
+			FConsoleCommandDelegate::CreateLambda([this]()
+				{
+					UE_LOG(LogIas, Display, TEXT("User invoked http error via 'ias.InvokeHttpFailure'"));
+					FHostGroupManager::Get().DisconnectAll();
+				}),
+			ECVF_Default)
+	);
+
+	DynamicConsoleCommands.Emplace(
+		IConsoleManager::Get().RegisterConsoleCommand(
+			TEXT("ias.RunRequestTest"),
+			TEXT("[optional Number of requests to make] Creates a number of requests for chunks that can be found in that IAS system"),
+			FConsoleCommandWithArgsDelegate::CreateRaw(this, &FOnDemandDebugCommands::RunRequestTest),
+			ECVF_Default));
+#endif // !NO_CVARS
+}
+
+void FOnDemandDebugCommands::UnbindConsoleCommands()
+{
+	for (IConsoleCommand* Cmd : DynamicConsoleCommands)
+	{
+		IConsoleManager::Get().UnregisterConsoleObject(Cmd);
+	}
+}
+
+void FOnDemandDebugCommands::RunRequestTest(const TArray<FString>& Args) const
+{
+	int32 NumToRequest = 1000;
+
+	if (Args.Num() > 1)
+	{
+		UE_LOG(LogIas, Error, TEXT("Too many args for 'ias.RequestTestCount'"));
+		return;
+	}
+	else if (Args.Num() == 1)
+	{
+		int32 ArgValue = TCString<TCHAR>::Atoi(*Args[0]);
+
+		if (!UE::String::IsNumericOnlyDigits(Args[0]) || ArgValue <= 0)
+		{
+			UE_LOG(LogIas, Error, TEXT("Invalid arg for 'ias.RequestTestCount', value must be a postive integer"));
+			return;
+		}
+
+		NumToRequest = ArgValue;
+	}
+
+	UE_LOG(LogIas, Display, TEXT("Running IAS Request Test with %d requests..."), NumToRequest);
+
+	TArray<FIoChunkId> RequestIds = OnDemandIoStore->DebugFindStreamingChunkIds(NumToRequest);
+
+	FIoBatch IoBatch;
+	for (const FIoChunkId& Id : RequestIds)
+	{
+		IoBatch.Read(Id, FIoReadOptions(), IoDispatcherPriority_Medium);
+	}
+
+	const double StartTime = FPlatformTime::Seconds();
+
+	IoBatch.IssueWithCallback([StartTime]()
+		{
+			const double TotalTime = FPlatformTime::Seconds() - StartTime;
+			UE_LOG(LogIas, Display, TEXT("IAS Request Test took %.3f(s)"), TotalTime);
+		});
+}
+
+} // namespace UE::IoStore
+
+#endif // !UE_BUILD_SHIPPING
diff --git a/Engine/Source/Runtime/Experimental/IoStore/OnDemand/Private/DebugCommands.h b/Engine/Source/Runtime/Experimental/IoStore/OnDemand/Private/DebugCommands.h
new file mode 100644
--- /dev/null
+++ b/Engine/Source/Runtime/Experimental/IoStore/OnDemand/Private/DebugCommands.h
@@ -0,0 +1,38 @@
+// Copyright Epic Games, Inc. All Rights Reserved.
+
+#pragma once
+
+#if !UE_BUILD_SHIPPING
+
+#include "Containers/Array.h"
+#include "Containers/UnrealString.h"
+
+struct IConsoleCommand;
+
+namespace UE::IoStore
+{
+
+class FOnDemandIoStore;
+
+class FOnDemandDebugCommands
+{
+public:
+	FOnDemandDebugCommands() = delete;
+	FOnDemandDebugCommands(FOnDemandIoStore* OnDemandIoStore);
+	~FOnDemandDebugCommands();
+
+private:
+
+	void BindConsoleCommands();
+	void UnbindConsoleCommands();
+
+	void RunRequestTest(const TArray<FString>& Args) const;
+
+	FOnDemandIoStore* OnDemandIoStore = nullptr;
+
+	TArray<IConsoleCommand*> DynamicConsoleCommands;
+};
+
+} // namespace UE::IoStore
+
+#endif // !UE_BUILD_SHIPPING
diff --git a/Engine/Source/Runtime/Experimental/IoStore/OnDemand/Private/OnDemandIoDispatcherBackend.cpp b/Engine/Source/Runtime/Experimental/IoStore/OnDemand/Private/OnDemandIoDispatcherBackend.cpp
--- a/Engine/Source/Runtime/Experimental/IoStore/OnDemand/Private/OnDemandIoDispatcherBackend.cpp
+++ b/Engine/Source/Runtime/Experimental/IoStore/OnDemand/Private/OnDemandIoDispatcherBackend.cpp
@@ -36,7 +36,6 @@
 #include "OnDemandIoStore.h"
 #include "ProfilingDebugging/IoStoreTrace.h"
 #include "Statistics.h"
-#include "String/Numeric.h"
 #include "Tasks/Task.h"
 #include "ThreadSafeIntrusiveQueue.h"
 
@@ -694,10 +693,6 @@ private:
 	int32 WaitForCompleteRequestTasks(float WaitTimeSeconds, float PollTimeSeconds);
 	void OnThreadTickIdle();
 
-#if UE_IAS_DEBUG_CONSOLE_CMDS
-	void RunRequestTest(const TArray<FString>& Args) const;
-#endif // UE_IAS_DEBUG_CONSOLE_CMDS
-
 	FOnDemandIoStore& IoStore;
 	FOnDemandHttpThread& HttpClient;
 	TUniquePtr<IIasCache> Cache;
@@ -746,18 +741,6 @@ FOnDemandIoBackend::FOnDemandIoBackend(const FOnDemandEndpointConfig& Config, FO
 	BackendStatus.SetCacheEnabled(Cache.IsValid());
 
 #if UE_IAS_DEBUG_CONSOLE_CMDS
-	DynamicConsoleCommands.Emplace(
-	IConsoleManager::Get().RegisterConsoleCommand(
-		TEXT("ias.InvokeHttpFailure"),
-		TEXT("Marks the current ias http connections as failed and forcing them  to try to reconnect"),
-		FConsoleCommandDelegate::CreateLambda([this]()
-			{
-				UE_LOG(LogIas, Display, TEXT("User invoked http error via 'ias.InvokeHttpFailure'"));
-				FHostGroupManager::Get().DisconnectAll();
-			}),
-		ECVF_Default)
-	);
-
 	DynamicConsoleCommands.Emplace(
 		IConsoleManager::Get().RegisterConsoleCommand(
 			TEXT("ias.AbandonCache"),
@@ -769,15 +752,6 @@ FOnDemandIoBackend::FOnDemandIoBackend(const FOnDemandEndpointConfig& Config, FO
 				}),
 			ECVF_Default)
 	);
-
-	DynamicConsoleCommands.Emplace(
-		IConsoleManager::Get().RegisterConsoleCommand(
-			TEXT("ias.RunRequestTest"),
-			TEXT("[optional Number of requests to make] Creates a number of requests for chunks that can be found in that IAS system"),
-			FConsoleCommandWithArgsDelegate::CreateRaw(this, &FOnDemandIoBackend::RunRequestTest),
-			ECVF_Default)
-	);
-
 #endif // UE_IAS_DEBUG_CONSOLE_CMDS
 }
 
@@ -1399,49 +1373,6 @@ int32 FOnDemandIoBackend::WaitForCompleteRequestTasks(float WaitTimeSeconds, flo
 	return ChunkRequests.Num();
 }
 
-#if UE_IAS_DEBUG_CONSOLE_CMDS
-void FOnDemandIoBackend::RunRequestTest(const TArray<FString>& Args) const
-{
-	int32 NumToRequest = 1000;
-
-	if (Args.Num() > 1)
-	{
-		UE_LOG(LogIas, Error, TEXT("Too many args for 'ias.RequestTestCount'"));
-		return;
-	}
-	else if (Args.Num() == 1)
-	{
-		int32 ArgValue = TCString<TCHAR>::Atoi(*Args[0]);
-
-		if (!UE::String::IsNumericOnlyDigits(Args[0]) || ArgValue <= 0)
-		{
-			UE_LOG(LogIas, Error, TEXT("Invalid arg for 'ias.RequestTestCount', value must be a postive integer"));
-			return;
-		}
-
-		NumToRequest = ArgValue;
-	}
-
-	UE_LOG(LogIas, Display, TEXT("Running IAS Request Test with %d requests..."), NumToRequest);
-
-	TArray<FIoChunkId> RequestIds = IoStore.DebugFindStreamingChunkIds(NumToRequest);
-
-	FIoBatch IoBatch;
-	for (const FIoChunkId& Id : RequestIds)
-	{
-		IoBatch.Read(Id, FIoReadOptions(), IoDispatcherPriority_Medium);
-	}
-
-	const double StartTime = FPlatformTime::Seconds();
-
-	IoBatch.IssueWithCallback([StartTime]()
-		{
-			const double TotalTime = FPlatformTime::Seconds() - StartTime;
-			UE_LOG(LogIas, Display, TEXT("IAS Request Test took %.3f(s)"), TotalTime);
-		});
-}
-#endif // UE_IAS_DEBUG_CONSOLE_CMDS
-
 TSharedPtr<IOnDemandIoDispatcherBackend> MakeOnDemandIoDispatcherBackend(
 	const FOnDemandEndpointConfig& Config,
 	FOnDemandIoStore& IoStore,
diff --git a/Engine/Source/Runtime/Experimental/IoStore/OnDemand/Private/OnDemandIoStore.cpp b/Engine/Source/Runtime/Experimental/IoStore/OnDemand/Private/OnDemandIoStore.cpp
--- a/Engine/Source/Runtime/Experimental/IoStore/OnDemand/Private/OnDemandIoStore.cpp
+++ b/Engine/Source/Runtime/Experimental/IoStore/OnDemand/Private/OnDemandIoStore.cpp
@@ -10,6 +10,7 @@
 #include "Containers/RingBuffer.h"
 #include "Containers/StringConv.h"
 #include "Containers/Ticker.h"
+#include "DebugCommands.h"
 #include "HAL/FileManager.h"
 #include "HAL/IConsoleManager.h"
 #include "HAL/LowLevelMemTracker.h"
@@ -194,6 +195,10 @@ FString FOnDemandContainer::UniqueName() const
 ///////////////////////////////////////////////////////////////////////////////
 FOnDemandIoStore::FOnDemandIoStore()
 {
+#if !UE_BUILD_SHIPPING
+	DebugCommands = MakeUnique<FOnDemandDebugCommands>(this);
+#endif // !UE_BUILD_SHIPPING
+
 	HttpClient = MakeUnique<FOnDemandHttpThread>();
 	Installer = MakeUnique<FOnDemandContentInstaller>(*this, *HttpClient);
 	FEncryptionKeyManager::Get().OnKeyAdded().AddRaw(this, &FOnDemandIoStore::OnEncryptionKeyAdded);
diff --git a/Engine/Source/Runtime/Experimental/IoStore/OnDemand/Private/OnDemandIoStore.h b/Engine/Source/Runtime/Experimental/IoStore/OnDemand/Private/OnDemandIoStore.h
--- a/Engine/Source/Runtime/Experimental/IoStore/OnDemand/Private/OnDemandIoStore.h
+++ b/Engine/Source/Runtime/Experimental/IoStore/OnDemand/Private/OnDemandIoStore.h
@@ -27,6 +27,7 @@ using FSharedContainerHeader		= TSharedPtr<FIoContainerHeader>;
 namespace UE::IoStore
 {
 
+class FOnDemandDebugCommands;
 class FOnDemandContentInstaller;
 class FOnDemandHttpThread;
 struct FOnDemandEndpointConfig;
@@ -406,12 +407,17 @@ private:
 											TArray<FSharedOnDemandContainer>& Out);
 	static FIoStatus						SetupHostGroup(const FSharedOnDemandContainer& Container, const FOnDemandMountArgs& MountArgs);
 
+#if !UE_BUILD_SHIPPING
+	TUniquePtr<FOnDemandDebugCommands>		DebugCommands;
+#endif // !UE_BUILD_SHIPPING
+
 	TUniquePtr<FOnDemandHttpThread>			HttpClient;
 	TUniquePtr<FOnDemandContentInstaller>	Installer;
 	FSharedInstallCache						InstallCache;
 	FSharedPackageStoreBackend				PackageStoreBackend;
 	FSharedStreamingBackend					StreamingBackend;
 	FSharedHttpIoBackend					HttpIoBackend;
+
 	FDelegateHandle							OnMountPakHandle;
 	FDelegateHandle							OnServerPostForkHandle;
 	TArray<FSharedOnDemandContainer>		Containers;
