Fix error messages for 'ias.RunRequestTest' using an older cvar name.

#rnx

diff --git a/Engine/Source/Runtime/Experimental/IoStore/OnDemand/Private/DebugCommands.cpp b/Engine/Source/Runtime/Experimental/IoStore/OnDemand/Private/DebugCommands.cpp
--- a/Engine/Source/Runtime/Experimental/IoStore/OnDemand/Private/DebugCommands.cpp
+++ b/Engine/Source/Runtime/Experimental/IoStore/OnDemand/Private/DebugCommands.cpp
@@ -68,7 +68,7 @@ void FOnDemandDebugCommands::RunRequestTest(const TArray<FString>& Args) const
 
 	if (Args.Num() > 1)
 	{
-		UE_LOG(LogIas, Error, TEXT("Too many args for 'ias.RequestTestCount'"));
+		UE_LOG(LogIas, Error, TEXT("Too many args for 'ias.RunRequestTest'"));
 		return;
 	}
 	else if (Args.Num() == 1)
@@ -77,7 +77,7 @@ void FOnDemandDebugCommands::RunRequestTest(const TArray<FString>& Args) const
 
 		if (!UE::String::IsNumericOnlyDigits(Args[0]) || ArgValue <= 0)
 		{
-			UE_LOG(LogIas, Error, TEXT("Invalid arg for 'ias.RequestTestCount', value must be a postive integer"));
+			UE_LOG(LogIas, Error, TEXT("Invalid arg for 'ias.RunRequestTest', value must be a postive integer"));
 			return;
 		}
 
