Remove the debug feature to abandon the IAS cache.

#rnx

- At one point this feature was exposed in the public API but that is no longer the case and it could only be invoked by a development only cvar.

diff --git a/Engine/Source/Runtime/Experimental/IoStore/OnDemand/Private/OnDemandBackendStatus.cpp b/Engine/Source/Runtime/Experimental/IoStore/OnDemand/Private/OnDemandBackendStatus.cpp
--- a/Engine/Source/Runtime/Experimental/IoStore/OnDemand/Private/OnDemandBackendStatus.cpp
+++ b/Engine/Source/Runtime/Experimental/IoStore/OnDemand/Private/OnDemandBackendStatus.cpp
@@ -103,7 +103,6 @@ FStringBuilderBase& operator<<(FStringBuilderBase& Sb, FBackendStatus::EFlags St
 	AppendIf(FBackendStatus::EFlags::CacheEnabled, TEXT("CacheEnabled"));
 	AppendIf(FBackendStatus::EFlags::HttpEnabled, TEXT("HttpEnabled"));
 	AppendIf(FBackendStatus::EFlags::HttpBulkOptionalDisabled, TEXT("HttpBulkOptionalDisabled"));
-	AppendIf(FBackendStatus::EFlags::AbandonCache, TEXT("AbandonCache"));
 
 	return Sb;
 }
diff --git a/Engine/Source/Runtime/Experimental/IoStore/OnDemand/Private/OnDemandBackendStatus.h b/Engine/Source/Runtime/Experimental/IoStore/OnDemand/Private/OnDemandBackendStatus.h
--- a/Engine/Source/Runtime/Experimental/IoStore/OnDemand/Private/OnDemandBackendStatus.h
+++ b/Engine/Source/Runtime/Experimental/IoStore/OnDemand/Private/OnDemandBackendStatus.h
@@ -21,7 +21,6 @@ public:
 		CacheEnabled				= 1 << 0,
 		HttpEnabled					= 1 << 1,
 		HttpBulkOptionalDisabled	= 1 << 2,
-		AbandonCache				= 1 << 3,
 
 		// When adding new values here, remember to update operator<<(FStringBuilderBase& Sb, EFlags StatusFlags) in the cpp!
 	};
@@ -50,11 +49,6 @@ public:
 		return (CurrentFlags & uint8(EFlags::CacheEnabled)) && !IsHttpEnabled(CurrentFlags);
 	}
 
-	bool ShouldAbandonCache() const
-	{
-		return HasAnyFlags(EFlags::AbandonCache);
-	}
-
 	void SetHttpEnabled(bool bEnabled)
 	{
 		AddOrRemoveFlags(EFlags::HttpEnabled, bEnabled, TEXT("HTTP streaming enabled"));
@@ -71,11 +65,6 @@ public:
 		AddOrRemoveFlags(EFlags::CacheEnabled, bEnabled, TEXT("Cache enabled"));
 	}
 
-	void SetAbandonCache(bool bAbandon)
-	{
-		AddOrRemoveFlags(EFlags::AbandonCache, bAbandon, TEXT("Abandon cache"));
-	}
-
 	void ToString(FStringBuilderBase& Builder) const;
 
 private:
diff --git a/Engine/Source/Runtime/Experimental/IoStore/OnDemand/Private/OnDemandIoDispatcherBackend.cpp b/Engine/Source/Runtime/Experimental/IoStore/OnDemand/Private/OnDemandIoDispatcherBackend.cpp
--- a/Engine/Source/Runtime/Experimental/IoStore/OnDemand/Private/OnDemandIoDispatcherBackend.cpp
+++ b/Engine/Source/Runtime/Experimental/IoStore/OnDemand/Private/OnDemandIoDispatcherBackend.cpp
@@ -47,9 +47,6 @@
 #include <catch2/generators/catch_generators.hpp>
 #endif
 
-/** When enabled the IAS system can add additional debug console commands for development use */
-#define UE_IAS_DEBUG_CONSOLE_CMDS (1 && !NO_CVARS && !UE_BUILD_SHIPPING)
-
 /** When enabled we will run some limited testing on start up for issues that are hard to reproduce with normal gameplay */
 #define UE_ENABLE_IAS_TESTING 0 && !UE_BUILD_SHIPPING
 
@@ -673,7 +670,6 @@ public:
 
 	// I/O Http backend
 	virtual void SetBulkOptionalEnabled(bool bEnabled) override;
-	virtual void AbandonCache() override;
 	
 	virtual void ReportAnalytics(TArray<FAnalyticsEventAttribute>& OutAnalyticsArray) const override;
 	virtual TUniquePtr<IAnalyticsRecording> StartAnalyticsRecording() const override;
@@ -691,7 +687,6 @@ private:
 	bool CreateDefaultHostGroup();
 	bool ResolveDistributedEndpoint(const FDistributedEndpointUrl& Url, TArray<FString>& OutUrls);
 	int32 WaitForCompleteRequestTasks(float WaitTimeSeconds, float PollTimeSeconds);
-	void OnThreadTickIdle();
 
 	FOnDemandIoStore& IoStore;
 	FOnDemandHttpThread& HttpClient;
@@ -705,15 +700,10 @@ private:
 	FDistributedEndpointUrl DistributionUrl;
 	FEventRef DistributedEndpointEvent;
 	FIASHostGroup DefaultHostGroup;
-	FDelegateHandle OnThreadTickIdleHandle;
 	mutable FRWLock Lock;
 	UE::Tasks::FTask HostGroupTask;
 	std::atomic_uint32_t InflightCacheRequestCount{0};
 	std::atomic_bool bStopRequested{false};
-
-#if UE_IAS_DEBUG_CONSOLE_CMDS
-	TArray<IConsoleCommand*> DynamicConsoleCommands;
-#endif // UE_IAS_DEBUG_CONSOLE_CMDS
 };
 
 ///////////////////////////////////////////////////////////////////////////////
@@ -739,31 +729,10 @@ FOnDemandIoBackend::FOnDemandIoBackend(const FOnDemandEndpointConfig& Config, FO
 	// Don't enable HTTP until the background thread has been started
 	BackendStatus.SetHttpEnabled(false);
 	BackendStatus.SetCacheEnabled(Cache.IsValid());
-
-#if UE_IAS_DEBUG_CONSOLE_CMDS
-	DynamicConsoleCommands.Emplace(
-		IConsoleManager::Get().RegisterConsoleCommand(
-			TEXT("ias.AbandonCache"),
-			TEXT("Simulate code calling FOnDemandIoBackend::AbandonCache()"),
-			FConsoleCommandDelegate::CreateLambda([this]()
-				{
-					UE_LOG(LogIas, Display, TEXT("User invoked a cache abandon via 'ias.AbandonCache' (disable the IAS cache)"));
-					AbandonCache();
-				}),
-			ECVF_Default)
-	);
-#endif // UE_IAS_DEBUG_CONSOLE_CMDS
 }
 
 FOnDemandIoBackend::~FOnDemandIoBackend()
 {
-#if UE_IAS_DEBUG_CONSOLE_CMDS
-	for (IConsoleCommand* Cmd : DynamicConsoleCommands)
-	{
-		IConsoleManager::Get().UnregisterConsoleObject(Cmd);
-	}
-#endif // UE_IAS_DEBUG_CONSOLE_CMDS
-
 	Shutdown();
 }
 
@@ -774,7 +743,6 @@ void FOnDemandIoBackend::Initialize(TSharedRef<const FIoDispatcherBackendContext
 	UE_LOG(LogIas, Log, TEXT("Initializing on demand I/O dispatcher backend"));
 	BackendContext = Context;
 	BeginCreateDefaultHostGroup();
-	OnThreadTickIdleHandle = HttpClient.OnTickIdle().AddRaw(this, &FOnDemandIoBackend::OnThreadTickIdle);
 }
 
 void FOnDemandIoBackend::Shutdown()
@@ -791,12 +759,6 @@ void FOnDemandIoBackend::Shutdown()
 	bStopRequested.store(true, std::memory_order_relaxed);
 	HostGroupTask.Wait();
 
-	if (OnThreadTickIdleHandle.IsValid())
-	{
-		HttpClient.OnTickIdle().Remove(OnThreadTickIdleHandle);
-		OnThreadTickIdleHandle.Reset();
-	}
-
 	// The CompleteRequest tasks may still be executing a while after the IoDispatcher has been notified about the completed io requests.
 	const int32 NumPending = WaitForCompleteRequestTasks(5.0f, 0.1f);
 	UE_CLOG(NumPending > 0, LogIas, Warning, TEXT("%d request(s) still pending after shutdown"), NumPending);
@@ -804,20 +766,6 @@ void FOnDemandIoBackend::Shutdown()
 	BackendContext.Reset();
 }
 
-void FOnDemandIoBackend::OnThreadTickIdle()
-{
-	if (BackendStatus.ShouldAbandonCache())
-	{
-		BackendStatus.SetAbandonCache(false);
-		check(BackendStatus.IsCacheEnabled() == false);
-		if (Cache.IsValid())
-		{
-			UE_LOG(LogIas, Log, TEXT("Abandoning cache, local file cache is no longer available"));
-			Cache.Release()->Abandon(); // Will delete its self
-		}
-	}
-}
-
 void FOnDemandIoBackend::BeginCreateDefaultHostGroup()
 {
 	HostGroupTask = UE::Tasks::Launch(UE_SOURCE_LOCATION, [this]()
@@ -1317,12 +1265,6 @@ void FOnDemandIoBackend::SetBulkOptionalEnabled(bool bEnabled)
 	BackendStatus.SetHttpOptionalBulkEnabled(bEnabled);
 }
 
-void FOnDemandIoBackend::AbandonCache()
-{
-	BackendStatus.SetCacheEnabled(false);
-	BackendStatus.SetAbandonCache(true);
-}
-
 void FOnDemandIoBackend::ReportAnalytics(TArray<FAnalyticsEventAttribute>& OutAnalyticsArray) const
 {
 	// If we got this far we know that IAS is enabled for the current process as it has a valid backend.
@@ -1829,4 +1771,3 @@ static void Trace(
 
 #undef UE_ALLOW_DISABLE_CANCELLING
 #undef UE_ENABLE_IAS_TESTING
-#undef UE_IAS_DEBUG_CONSOLE_CMDS
diff --git a/Engine/Source/Runtime/Experimental/IoStore/OnDemand/Private/OnDemandIoDispatcherBackend.h b/Engine/Source/Runtime/Experimental/IoStore/OnDemand/Private/OnDemandIoDispatcherBackend.h
--- a/Engine/Source/Runtime/Experimental/IoStore/OnDemand/Private/OnDemandIoDispatcherBackend.h
+++ b/Engine/Source/Runtime/Experimental/IoStore/OnDemand/Private/OnDemandIoDispatcherBackend.h
@@ -80,7 +80,6 @@ public:
 	virtual ~IOnDemandIoDispatcherBackend() = default;
 
 	virtual void SetBulkOptionalEnabled(bool bInEnabled) = 0;
-	virtual void AbandonCache() = 0;
 	virtual void ReportAnalytics(TArray<FAnalyticsEventAttribute>& OutAnalyticsArray) const = 0;
 	virtual TUniquePtr<IAnalyticsRecording> StartAnalyticsRecording() const = 0;
 	virtual FOnDemandStreamingCacheUsage GetCacheUsage() const = 0;
