[AutoRTFM] Optimize transactionally-safe mutex lambda captures.

As we use TSharedPtr for ensuring the state is preserved until the transaction callbacks are called, we do some ref-count twiddling when we capture the state for the on-commit / on-abort lambda callbacks.
As these callbacks need to be run from the closed, we end up taking the closed paths for the TSharedPtr copy, which enqueues more tasks to keep the ref-count preserved in case of abort. For high frequency usage, this can quickly accumulate a vast number of on-commit and on-abort tasks.
We know that these lambda captures will be destructed when the transaction is committed / aborted, so we can avoid the overhead of the TSharedPtr logic by capturing the pointer in the open.

This CL adds a new AutoRTFM::TOpenWrapper helper which is used on the lambda captures to ensure that the shared pointer is copied in the open.
This reduces the number of enqueued tasks for typical lock / unlocks by 2/3rds.
Benchmarking suggests this is about ~10% faster too.

diff --git a/Engine/Source/Programs/AutoRTFMTests/Private/OpenWrapperTests.cpp b/Engine/Source/Programs/AutoRTFMTests/Private/OpenWrapperTests.cpp
new file mode 100644
--- /dev/null
+++ b/Engine/Source/Programs/AutoRTFMTests/Private/OpenWrapperTests.cpp
@@ -0,0 +1,48 @@
+// Copyright Epic Games, Inc. All Rights Reserved.
+
+#include "AutoRTFM.h"
+#include "AutoRTFM/OpenWrapper.h"
+#include "AutoRTFMTesting.h"
+#include "Catch2Includes.h"
+
+TEST_CASE("OpenWrapper")
+{
+	struct FObject
+	{
+		FObject()
+		{
+			REQUIRE(!AutoRTFM::IsClosed());
+		}
+		FObject(const FObject&)
+		{
+			REQUIRE(!AutoRTFM::IsClosed());
+		}
+		FObject(FObject&&)
+		{
+			REQUIRE(!AutoRTFM::IsClosed());
+		}
+		FObject& operator=(const FObject&)
+		{
+			REQUIRE(!AutoRTFM::IsClosed());
+			return *this;
+		}
+		FObject& operator=(FObject&&)
+		{
+			REQUIRE(!AutoRTFM::IsClosed());
+			return *this;
+		}
+		~FObject()
+		{
+			REQUIRE(!AutoRTFM::IsClosed());
+		}
+	};
+
+	FObject Object;
+	AutoRTFM::Testing::Commit([&]
+	{
+		AutoRTFM::TOpenWrapper WrapperA{Object};
+		AutoRTFM::TOpenWrapper WrapperB{std::move(WrapperA)};
+		AutoRTFM::TOpenWrapper WrapperC = WrapperB;
+		AutoRTFM::TOpenWrapper WrapperD = std::move(WrapperC);
+	});
+}
diff --git a/Engine/Source/Programs/AutoRTFMTests/Private/RWLockTests.cpp b/Engine/Source/Programs/AutoRTFMTests/Private/RWLockTests.cpp
--- a/Engine/Source/Programs/AutoRTFMTests/Private/RWLockTests.cpp
+++ b/Engine/Source/Programs/AutoRTFMTests/Private/RWLockTests.cpp
@@ -493,3 +493,29 @@ TEST_CASE("TransactionallySafeRWLock.MultipleTryWriteLocks")
 		});
 	}
 }
+
+TEST_CASE("Benchmarks.TransactionallySafeRWLock")
+{
+	// Disable memory validation as it can affect benchmark results.
+	AUTORTFM_SCOPED_DISABLE_MEMORY_VALIDATION();
+
+	BENCHMARK("Transactional.ReadLockWriteLock")
+	{
+		for (int I = 0; I < 100; I++)
+		{
+			AutoRTFM::Testing::Commit([]
+			{
+				FTransactionallySafeRWLock RWLock;
+				for (int J = 0; J < 100; J++)
+				{
+					{
+						UE::TReadScopeLock ReadLock{RWLock};
+					}
+					{
+						UE::TWriteScopeLock WriteLock{RWLock};
+					}
+				}
+			});
+		}
+	};
+}
diff --git a/Engine/Source/Runtime/AutoRTFM/Public/AutoRTFM/OpenWrapper.h b/Engine/Source/Runtime/AutoRTFM/Public/AutoRTFM/OpenWrapper.h
new file mode 100644
--- /dev/null
+++ b/Engine/Source/Runtime/AutoRTFM/Public/AutoRTFM/OpenWrapper.h
@@ -0,0 +1,41 @@
+// Copyright Epic Games, Inc. All Rights Reserved.
+
+#pragma once
+
+#include "AutoRTFM.h"
+#include <utility>
+
+namespace AutoRTFM
+{
+
+// An object wrapper that will ensure the copy constructor, move constructor,
+// copy assignment, move assignment and destructor are all called in the open.
+// This can be used to wrap a lambda capture if it is known that it safe and
+// more optimal to copy, move & destruct the object in the open.
+//
+// For example - it may be safe to wrap a shared pointer with a TOpenWrapper
+// when capturing for use by AutoRTFM::OnCommit() or AutoRTFM::OnAbort(), as
+// the AutoRTFM::TTask destructor will be called regardless of a transactional
+// failure or commit. This may be preferable to capturing the unwrapped shared
+// pointer as we can avoid costly closed transactional logic to ensure the
+// count is restored on failure:
+//    AutoRTFM::OnCommit([](WrappedSharedPtr = AutoRTFM::TOpenWrapper{MySharedPtr})
+//    {
+//        WrappedSharedPtr.Object->DoSomething();
+//    });
+template<typename ObjectType>
+class AUTORTFM_OPEN TOpenWrapper final
+{
+public:
+	ObjectType Object;
+
+	TOpenWrapper(const TOpenWrapper&) = default;
+	TOpenWrapper(TOpenWrapper&&) = default;
+	TOpenWrapper& operator=(const TOpenWrapper&) = default;
+	TOpenWrapper& operator=(TOpenWrapper&&) = default;
+	TOpenWrapper(const ObjectType& Value) : Object{Value} {}
+	TOpenWrapper(ObjectType&& Value) : Object{std::move(Value)} {}
+	~TOpenWrapper() = default;
+};
+
+}
diff --git a/Engine/Source/Runtime/Core/Public/Async/TransactionallySafeMutex.h b/Engine/Source/Runtime/Core/Public/Async/TransactionallySafeMutex.h
--- a/Engine/Source/Runtime/Core/Public/Async/TransactionallySafeMutex.h
+++ b/Engine/Source/Runtime/Core/Public/Async/TransactionallySafeMutex.h
@@ -4,6 +4,7 @@
 
 #include "Async/Mutex.h"
 #include "AutoRTFM.h"
+#include "AutoRTFM/OpenWrapper.h"
 #include "Templates/SharedPointer.h"
 
 namespace UE
@@ -222,10 +223,10 @@ private:
 
 		// Capture State instead of 'this' as State is a TSharedPtr which can outlive this 
 		// FTransactionallySafeMutex.
-		AutoRTFM::OnAbort([State = this->State]
+		AutoRTFM::OnAbort([State = AutoRTFM::TOpenWrapper{this->State}]
 		{
-			State->MutexUnlock();
-			State->ResetTransactionState();
+			State.Object->MutexUnlock();
+			State.Object->ResetTransactionState();
 		});
 	}
 	
@@ -239,17 +240,17 @@ private:
 
 			// Capture State instead of 'this' as State is a TSharedPtr which can outlive this 
 			// FTransactionallySafeMutex.
-			AutoRTFM::OnCommit([State = this->State]
+			AutoRTFM::OnCommit([State = AutoRTFM::TOpenWrapper{this->State}]
 			{
 				// If the transaction's lock count is negative (more unlocks than locks), or the 
 				// underlying mutex was locked and then re-balanced to an unlocked state, unlock
 				// the underlying mutex.
-				if (State->TransactionalLockCount < 0 ||
-					(State->bTransactionallyLocked && State->TransactionalLockCount == 0))
+				if (State.Object->TransactionalLockCount < 0 ||
+					(State.Object->bTransactionallyLocked && State.Object->TransactionalLockCount == 0))
 				{
-					State->MutexUnlock();
+					State.Object->MutexUnlock();
 				}
-				State->ResetTransactionState();
+				State.Object->ResetTransactionState();
 			});
 		}
 	}
diff --git a/Engine/Source/Runtime/Core/Public/Async/TransactionallySafeRecursiveMutex.h b/Engine/Source/Runtime/Core/Public/Async/TransactionallySafeRecursiveMutex.h
--- a/Engine/Source/Runtime/Core/Public/Async/TransactionallySafeRecursiveMutex.h
+++ b/Engine/Source/Runtime/Core/Public/Async/TransactionallySafeRecursiveMutex.h
@@ -6,6 +6,7 @@
 #include "Async/RecursiveMutex.h"
 
 #if UE_AUTORTFM
+#include "AutoRTFM/OpenWrapper.h"
 #include "Templates/SharedPointer.h"
 #endif
 
@@ -55,15 +56,15 @@ struct FTransactionallySafeRecursiveMutexDefinition final
 
 			// We explicitly copy the state here for the case that `this` was stack
 			// allocated and has already died before the on-abort is hit.
-			AutoRTFM::OnAbort([State = this->State]
+			AutoRTFM::OnAbort([State = AutoRTFM::TOpenWrapper{this->State}]
 			{
-				ensure(0 != State->TransactionalLockCount);
+				ensure(0 != State.Object->TransactionalLockCount);
 
-				State->TransactionalLockCount -= 1;
+				State.Object->TransactionalLockCount -= 1;
 
-				if (0 == State->TransactionalLockCount)
+				if (0 == State.Object->TransactionalLockCount)
 				{
-					State->Mutex.Unlock();
+					State.Object->Mutex.Unlock();
 				}
 			});
 		}
@@ -80,15 +81,15 @@ struct FTransactionallySafeRecursiveMutexDefinition final
 		{
 			// We explicitly copy the state here for the case that `this` was stack
 			// allocated and has already died before the on-commit is hit.
-			AutoRTFM::OnCommit([State = this->State]
+			AutoRTFM::OnCommit([State = AutoRTFM::TOpenWrapper{this->State}]
 			{
-				ensure(0 != State->TransactionalLockCount);
+				ensure(0 != State.Object->TransactionalLockCount);
 
-				State->TransactionalLockCount -= 1;
+				State.Object->TransactionalLockCount -= 1;
 
-				if (0 == State->TransactionalLockCount)
+				if (0 == State.Object->TransactionalLockCount)
 				{
-					State->Mutex.Unlock();
+					State.Object->Mutex.Unlock();
 				}
 			});
 		}
diff --git a/Engine/Source/Runtime/Core/Public/Async/TransactionallySafeSharedMutex.h b/Engine/Source/Runtime/Core/Public/Async/TransactionallySafeSharedMutex.h
--- a/Engine/Source/Runtime/Core/Public/Async/TransactionallySafeSharedMutex.h
+++ b/Engine/Source/Runtime/Core/Public/Async/TransactionallySafeSharedMutex.h
@@ -4,8 +4,10 @@
 
 #include "Async/SharedMutex.h"
 #include "AutoRTFM.h"
+#include "Templates/SharedPointerFwd.h"
 
 #if UE_AUTORTFM
+#include "AutoRTFM/OpenWrapper.h"
 #include "Templates/SharedPointer.h"
 #endif
 
@@ -94,9 +96,9 @@ struct FTransactionallySafeSharedMutexDefinition
 
 			// We explicitly copy the state here for the case that `this` was stack
 			// allocated and has already died before the on-abort is hit.
-			UE_AUTORTFM_ONABORT(State = this->State)
+			UE_AUTORTFM_ONABORT(State = AutoRTFM::TOpenWrapper{this->State})
 			{
-				State->Unlock();
+				State.Object->Unlock();
 			};
 		}
 		else
@@ -112,9 +114,9 @@ struct FTransactionallySafeSharedMutexDefinition
 		{
 			// We explicitly copy the state here for the case that `this` was stack
 			// allocated and has already died before the on-commit is hit.
-			UE_AUTORTFM_ONCOMMIT(State = this->State)
+			UE_AUTORTFM_ONCOMMIT(State = AutoRTFM::TOpenWrapper{this->State})
 			{
-				State->Unlock();
+				State.Object->Unlock();
 			};
 		}
 		else
@@ -152,9 +154,9 @@ struct FTransactionallySafeSharedMutexDefinition
 			{
 				// We explicitly copy the state here for the case that `this` was stack
 				// allocated and has already died before the on-abort is hit.
-				UE_AUTORTFM_ONABORT(State = this->State)
+				UE_AUTORTFM_ONABORT(State = AutoRTFM::TOpenWrapper{this->State})
 				{
-					State->Unlock();
+					State.Object->Unlock();
 				};
 			}
 
diff --git a/Engine/Source/Runtime/Core/Public/Misc/MTTransactionallySafeAccessDetector.h b/Engine/Source/Runtime/Core/Public/Misc/MTTransactionallySafeAccessDetector.h
--- a/Engine/Source/Runtime/Core/Public/Misc/MTTransactionallySafeAccessDetector.h
+++ b/Engine/Source/Runtime/Core/Public/Misc/MTTransactionallySafeAccessDetector.h
@@ -3,6 +3,7 @@
 #pragma once
 
 #include "AutoRTFM.h"
+#include "AutoRTFM/OpenWrapper.h"
 #include "HAL/Platform.h"
 #include "MTAccessDetector.h"
 #include "Misc/AssertionMacros.h"
@@ -237,14 +238,20 @@ private:
 	{
 		// We explicitly copy the state here for the case that `this` was stack
 		// allocated and has already died before the on-abort is hit.
-		AutoRTFM::OnAbort([State = this->State] { ReleaseAccess(State); });
+		AutoRTFM::OnAbort([State = AutoRTFM::TOpenWrapper{this->State}]
+		{
+			ReleaseAccess(State.Object);
+		});
 	}
 
 	UE_FORCEINLINE_HINT void RegisterOnCommitRelease() const
 	{
 		// We explicitly copy the state here for the case that `this` was stack
 		// allocated and has already died before the on-abort is hit.
-		AutoRTFM::OnCommit([State = this->State] { ReleaseAccess(State); });
+		AutoRTFM::OnCommit([State = AutoRTFM::TOpenWrapper{this->State}]
+		{
+			ReleaseAccess(State.Object);
+		});
 	}
 
 	// The state held for calls made when in a transaction.
diff --git a/Engine/Source/Runtime/Experimental/Chaos/Public/Framework/Threading.h b/Engine/Source/Runtime/Experimental/Chaos/Public/Framework/Threading.h
--- a/Engine/Source/Runtime/Experimental/Chaos/Public/Framework/Threading.h
+++ b/Engine/Source/Runtime/Experimental/Chaos/Public/Framework/Threading.h
@@ -2,6 +2,7 @@
 
 #pragma once
 
+#include "AutoRTFM/OpenWrapper.h"
 #include "CoreMinimal.h"
 #include "Framework/ThreadContextEnum.h"
 #include "UObject/ObjectMacros.h"
@@ -1027,9 +1028,9 @@ FORCEINLINE void EnsureIsInGameThreadContext()
 
 						// We explicitly copy the state here for the case that `this` was stack
 						// allocated and has already died before the on-commit is hit.
-						AutoRTFM::OnCommit([State = this->State]
+						AutoRTFM::OnCommit([State = AutoRTFM::TOpenWrapper{this->State}]
 							{
-								ensure(0 == State->TransactionalLockCount);
+								ensure(0 == State.Object->TransactionalLockCount);
 							});
 					});
 
@@ -1086,15 +1087,15 @@ FORCEINLINE void EnsureIsInGameThreadContext()
 
 				// We explicitly copy the state here for the case that `this` was stack
 				// allocated and has already died before the on-abort is hit.
-				AutoRTFM::OnAbort([State = this->State]
+				AutoRTFM::OnAbort([State = AutoRTFM::TOpenWrapper{this->State}]
 					{
-						ensure(0 != State->TransactionalLockCount);
+						ensure(0 != State.Object->TransactionalLockCount);
 
-						State->TransactionalLockCount -= 1;
+						State.Object->TransactionalLockCount -= 1;
 
-						if (0 == State->TransactionalLockCount)
+						if (0 == State.Object->TransactionalLockCount)
 						{
-							State->Lock.WriteUnlock();
+							State.Object->Lock.WriteUnlock();
 						}
 					});
 			}
@@ -1111,15 +1112,15 @@ FORCEINLINE void EnsureIsInGameThreadContext()
 			{
 				// We explicitly copy the state here for the case that `this` was stack
 				// allocated and has already died before the on-commit is hit.
-				AutoRTFM::OnCommit([State = this->State]
+				AutoRTFM::OnCommit([State = AutoRTFM::TOpenWrapper{this->State}]
 					{
-						ensure(0 != State->TransactionalLockCount);
+						ensure(0 != State.Object->TransactionalLockCount);
 
-						State->TransactionalLockCount -= 1;
+						State.Object->TransactionalLockCount -= 1;
 
-						if (0 == State->TransactionalLockCount)
+						if (0 == State.Object->TransactionalLockCount)
 						{
-							State->Lock.WriteUnlock();
+							State.Object->Lock.WriteUnlock();
 						}
 					});
 			}
