IAS Tool - The ChunkPlugin command now supports IAS being enabled independently of IAD.

#rnx

- A new optional cmdline arg '-SkipInstallOnDemand' has been added. When set the command will not chunk any Install On Demand (IAD) containers.
- A new optional cmdline arg '-SkipStreamOnDemand' has been added. When set the command will not chunk any Streaming On Demand (IAS) containers.
- The old hard coded way fo enabling/disabling container types bSupportInstalledContainers/bSupportStreamingContainers are now replaced by these new args.
- We now only write out a uondemand toc for IAD if it has data in the same manner as IAS.
- Cleaned up the logic for writing out dummy toc files. We can just iterate over the files that were queued for deletion (even if they were not deleted) as we know that these are the source utoc files.
- We now always pass a valid array into ::ChunkContainer to collect the files to delete. The code is much cleaner if we just opt to use the array based on bDeleteContainerFiles rather than change the code flow of ::ChunkContainer, this will make further refactors and code clean up easier.

diff --git a/Engine/Source/Runtime/Experimental/IoStore/OnDemand/Private/Tool/ChunkPlugin.cpp b/Engine/Source/Runtime/Experimental/IoStore/OnDemand/Private/Tool/ChunkPlugin.cpp
--- a/Engine/Source/Runtime/Experimental/IoStore/OnDemand/Private/Tool/ChunkPlugin.cpp
+++ b/Engine/Source/Runtime/Experimental/IoStore/OnDemand/Private/Tool/ChunkPlugin.cpp
@@ -233,7 +233,7 @@ static FIoStatus WriteChunk(const FString& Directory, FMemoryView Chunk, const F
 }
 
 static int32 ChunkContainer(const FString& ContainerUTocPath, const FChunkPluginSettings& Settings, const TMap<FGuid, FAES::FAESKey>& EncryptionKeys, FChunkWriterInterface* ChunkWriter, const FString& ChunksRelativeFolder, bool bIgnoreContainerHeader,
-	FOnDemandToc& OutOnDemandToc, TMap<FString, TArray<FOnDemandTocTagSetPackageList>>* OutFoundTagSets, TArray<FString>* OutFilesToDelete, FChunkPluginStats& OutStats)
+	FOnDemandToc& OutOnDemandToc, TMap<FString, TArray<FOnDemandTocTagSetPackageList>>* OutFoundTagSets, TArray<FString>& OutContainerFiles, FChunkPluginStats& OutStats)
 {
 	check(ChunkWriter != nullptr);
 
@@ -462,11 +462,8 @@ static int32 ChunkContainer(const FString& ContainerUTocPath, const FChunkPlugin
 		}
 	}
 
-	if (OutFilesToDelete != nullptr)
-	{
-		OutFilesToDelete->Add(ContainerUTocPath);
-		ContainerFileReader.GetContainerFilePaths(*OutFilesToDelete);
-	}
+	OutContainerFiles.Add(ContainerUTocPath);
+	ContainerFileReader.GetContainerFilePaths(OutContainerFiles);
 
 	return 0;
 }
@@ -539,6 +536,9 @@ static int32 ChunkPluginCommandEntry(const FContext& Context)
 	const bool bStreamOnDemand			= Context.Get<bool>(TEXT("-StreamOnDemand"), false);
 	const bool bInstallOnDemand			= Context.Get<bool>(TEXT("-InstallOnDemand"), false);
 
+	bool bSupportInstalledContainers = !Context.Get<bool>(TEXT("-SkipInstallOnDemand"), false);
+	bool bSupportStreamingContainers = !Context.Get<bool>(TEXT("-SkipStreamOnDemand"), false);
+
 	EOnDemandTocFlags TocFlags = EOnDemandTocFlags::None;
 	if (bStreamOnDemand)
 	{
@@ -577,6 +577,8 @@ static int32 ChunkPluginCommandEntry(const FContext& Context)
 	UE_LOG(LogIoStoreOnDemand, Display, TEXT("\tInputFolder: %s"), *InputFolder);
 	UE_LOG(LogIoStoreOnDemand, Display, TEXT("\tOutputFolder: %s"), *OutputFolder);
 	UE_LOG(LogIoStoreOnDemand, Display, TEXT("\tIntermediateFolder: %s"), *IntermediateFolder);
+	UE_LOG(LogIoStoreOnDemand, Display, TEXT("\tSkipInstallOnDemand: %s"), !bSupportInstalledContainers ? TEXT("true") : TEXT("false"));
+	UE_LOG(LogIoStoreOnDemand, Display, TEXT("\tSkipStreamOnDemand: %s"), !bSupportStreamingContainers ? TEXT("true") : TEXT("false"));
 	UE_LOG(LogIoStoreOnDemand, Display, TEXT("\tSettingsFile: %s"), *SettingsFile);
 	UE_LOG(LogIoStoreOnDemand, Display, TEXT("\tOutputStatsJson: %s"), *OutputStatsJson);
 	UE_LOG(LogIoStoreOnDemand, Display, TEXT("\tIncludeSigPak: %s"), bIncludeSigPak ? TEXT("true") : TEXT("false"));
@@ -673,10 +675,6 @@ static int32 ChunkPluginCommandEntry(const FContext& Context)
 	FileMgr.FindFiles(ContainerFilenames, *ContainerFolder, TEXT("*.utoc"));
 	UE_LOG(LogIoStoreOnDemand, Display, TEXT("Found %d container files(s)"), ContainerFilenames.Num());
 
-	// TODO: This info should be passed into the tool or always check for both
-	bool bSupportInstalledContainers = true; // IAD
-	bool bSupportStreamingContainers = true; // IAS
-
 	if (!bSupportInstalledContainers && !bSupportStreamingContainers)
 	{
 		UE_LOG(LogIoStoreOnDemand, Error, TEXT("Neither Installed or Streaming containers were requested"));
@@ -698,28 +696,28 @@ static int32 ChunkPluginCommandEntry(const FContext& Context)
 	StreamingOnDemandToc.Header.ChunksDirectory = TEXT(""); // TODO?
 	StreamingOnDemandToc.Containers.Reserve(ContainerFilenames.Num());
 
-	TArray<FString> FilesToDelete;
+	TArray<FString> ContainerFiles;
 
 	for (const FString& Filename : ContainerFilenames)
 	{
 		const FString FullPath = ContainerFolder / Filename;
 
-		TArray<FString>* FilesToDeletePtr = bDeleteContainerFiles ? &FilesToDelete : nullptr;
-
 		// TODO: At the moment the only way we can identify an OnDemand container is by the name. Might be better to pass this info into the tool
-		if (bSupportStreamingContainers && FullPath.Contains(TEXT("OnDemand")))
+		const bool bIsStreamOnDemandContainer = FullPath.Contains(TEXT("OnDemand"));
+
+		if (bSupportStreamingContainers && bIsStreamOnDemandContainer)
 		{
 			// Attempt to add to IAS output container
-			int32 Result = ChunkContainer(FullPath, Settings, EncryptionKeys, ChunkWriter.Get(), ChunksRelativeFolder, bIgnoreContainerHeader, StreamingOnDemandToc, nullptr /*OutFoundTagSets*/, FilesToDeletePtr, Stats);
+			int32 Result = ChunkContainer(FullPath, Settings, EncryptionKeys, ChunkWriter.Get(), ChunksRelativeFolder, bIgnoreContainerHeader, StreamingOnDemandToc, nullptr /*OutFoundTagSets*/, ContainerFiles, Stats);
 			if (Result < 0)
 			{
 				return Result;
 			}
 		}
-		else
+		else if(bSupportInstalledContainers && !bIsStreamOnDemandContainer)
 		{
 			// Attempt to add to IAD output container
-			int32 Result = ChunkContainer(FullPath, Settings, EncryptionKeys, ChunkWriter.Get(), ChunksRelativeFolder, bIgnoreContainerHeader, OnDemandToc, &FoundTagSets, FilesToDeletePtr, Stats);
+			int32 Result = ChunkContainer(FullPath, Settings, EncryptionKeys, ChunkWriter.Get(), ChunksRelativeFolder, bIgnoreContainerHeader, OnDemandToc, &FoundTagSets, ContainerFiles, Stats);
 			if (Result < 0)
 			{
 				return Result;
@@ -738,14 +736,19 @@ static int32 ChunkPluginCommandEntry(const FContext& Context)
 		}
 	}
 
-	for (const FString& Path : FilesToDelete)
+	if (bDeleteContainerFiles)
 	{
-		if (FileMgr.FileExists(*Path))
+		UE_LOG(LogIoStoreOnDemand, Display, TEXT("Deleting chunked  containers..."));
+
+		for (const FString& Path : ContainerFiles)
 		{
-			UE_LOG(LogIoStoreOnDemand, Display, TEXT("Deleting '%s'"), *Path);
-			if (!FileMgr.Delete(*Path, /*RequireExists*/true))
+			if (FileMgr.FileExists(*Path))
 			{
-				UE_LOG(LogIoStoreOnDemand, Error, TEXT("Failed to delete '%s'"), *Path);
+				UE_LOG(LogIoStoreOnDemand, Display, TEXT("Deleting '%s'"), *Path);
+				if (!FileMgr.Delete(*Path, /*RequireExists*/true))
+				{
+					UE_LOG(LogIoStoreOnDemand, Error, TEXT("Failed to delete '%s'"), *Path);
+				}
 			}
 		}
 	}
@@ -812,7 +815,7 @@ static int32 ChunkPluginCommandEntry(const FContext& Context)
 	// const FString TocPath = OutputFolder / Filename; // like this yo
 	const FString OnDemandTocOutputFolder = OutputFolder / IoStoreRelativeFolder;
 
-	if (bSupportInstalledContainers)
+	if (bSupportInstalledContainers && !OnDemandToc.Containers.IsEmpty())
 	{
 		OnDemandToc.Meta.EpochTimestamp = Timestamp;
 		OnDemandToc.Meta.TargetPlatform = Platform;
@@ -842,35 +845,23 @@ static int32 ChunkPluginCommandEntry(const FContext& Context)
 	// Write dummy containers if necessary
 	if (bDeleteContainerFiles)
 	{
-		TSet<FString> DummyContainerPaths;
-		FileMgr.IterateDirectory(*ContainerFolder, [&DummyContainerPaths](const TCHAR* FilenameOrDirectory, bool bIsDirectory)
-		{
-			if (!bIsDirectory)
-			{
-				FStringView MaybePak(FilenameOrDirectory);
-				if (FPathViews::GetExtension(MaybePak) == TEXTVIEW("pak"))
-				{
-					DummyContainerPaths.Emplace(FPathViews::GetBaseFilenameWithPath(MaybePak));
-				}
-			}
-			return true;
-		});
+		UE_LOG(LogIoStoreOnDemand, Display, TEXT("Replacing .utoc for chunked  containers..."));
 
-		for (const FString& Filename : ContainerFilenames)
+		for (const FString& SourcePath : ContainerFiles)
 		{
-			const FString FullPath = ContainerFolder / Filename;
-			const FStringView ContainerPath = FPathViews::GetBaseFilenameWithPath(FullPath);
-			if (!DummyContainerPaths.ContainsByHash(GetTypeHash(ContainerPath), ContainerPath))
+			if (!SourcePath.EndsWith(TEXT(".utoc")))
 			{
 				continue;
 			}
 
 			FNameBuilder NameBuilder;
-			NameBuilder << BuildVersion << FPathViews::GetBaseFilename(Filename) << TEXTVIEW("dummy");
+			NameBuilder << BuildVersion << FPathViews::GetBaseFilename(SourcePath) << TEXTVIEW("dummy");
 
 			FIoContainerSettings ContainerSettings;
 			ContainerSettings.ContainerId = FIoContainerId::FromName(FName(NameBuilder));
 
+			const FString FullPath = ContainerFolder / FPaths::GetCleanFilename(SourcePath);
+
 			FIoStoreTocResource Toc;
 
 			TIoStatusOr<uint64> Status = FIoStoreTocResource::Write(*FullPath, Toc, 0, 0, ContainerSettings);
@@ -928,6 +919,8 @@ static FCommand ChunkPluginCommand(
 		TArgument<bool>(TEXT("-KeepContainerFiles"),		TEXT("Should we keep the container files after processing them.")),
 		TArgument<bool>(TEXT("-StreamOnDemand"),			TEXT("Set the content to be streamed on-demand")),
 		TArgument<bool>(TEXT("-InstallOnDemand"),			TEXT("Set the content to be installed on-demand")),
+		TArgument<bool>(TEXT("-SkipStreamOnDemand"),		TEXT("Do not chunk any stream on-demand containers")),
+		TArgument<bool>(TEXT("-SkipInstallOnDemand"),		TEXT("Do not chunk any install on-demand containers")),
 		TArgument<bool>(TEXT("-IgnoreContainerHeader"),		TEXT("Ignore the container header?")),
 		TArgument<FStringView>(TEXT("-BucketPrefix"),		TEXT("Path to prefix to bucket objects")),
 		TArgument<int32>(TEXT("-MaxConcurrentUploads"),		TEXT("Number of simultaneous uploads")),
