IAD cache verification and hardening
- When making new CAS blocks, do not allow appending.
- Add aditional size checks and error logging when purging.
- Add bytes purged from defrag to purge stats
- Add code to handle the case that an already cached block could be downloaded which would require additional space to be freed when purging.  Also add error logging for this case.
- When defragging, add code to verify that the referenced and fragmented bytes count is correct.  Add error logging if its not.

#tests Locally ran IAD stress test

diff --git a/Engine/Source/Runtime/Experimental/IoStore/OnDemand/Private/OnDemandInstallCache.cpp b/Engine/Source/Runtime/Experimental/IoStore/OnDemand/Private/OnDemandInstallCache.cpp
--- a/Engine/Source/Runtime/Experimental/IoStore/OnDemand/Private/OnDemandInstallCache.cpp
+++ b/Engine/Source/Runtime/Experimental/IoStore/OnDemand/Private/OnDemandInstallCache.cpp
@@ -19,6 +19,7 @@
 #include "IO/IoContainerHeader.h"
 #include "IO/IoChunkId.h"
 #include "IO/IoChunkEncoding.h"
+#include "Logging/StructuredLog.h"
 #include "Misc/DateTime.h"
 #include "Misc/Paths.h"
 #include "Misc/PathViews.h"
@@ -200,7 +201,7 @@ struct FCas
 	FSharedFileOpenResult		OpenRead(FCasBlockId BlockId);
 	FSharedFileOpenAsyncResult	OpenAsyncRead(FCasBlockId BlockId);
 	void						OnFileHandleDeleted(FCasBlockId BlockId);
-	FUniqueFileHandle			OpenWrite(FCasBlockId BlockId) const;
+	FUniqueFileHandle			OpenWrite(FCasBlockId BlockId, bool bAppend) const;
 	bool						TrackAccessIf(ECasTrackAccessType Type, FCasBlockId BlockId, int64 UtcTicks, bool bDirty);
 	bool						TrackAccessIf(ECasTrackAccessType Type, FCasBlockId BlockId, bool bDirty)
 	{
@@ -488,11 +489,10 @@ void FCas::OnFileHandleDeleted(FCasBlockId BlockId)
 	}
 }
 
-FUniqueFileHandle FCas::OpenWrite(FCasBlockId BlockId) const
+FUniqueFileHandle FCas::OpenWrite(FCasBlockId BlockId, const bool bAppend) const
 {
 	IPlatformFile&	Ipf = FPlatformFileManager::Get().GetPlatformFile();
 	const FString	Filename = GetBlockFilename(BlockId);
-	const bool		bAppend = true;
 	const bool		bAllowRead = true;
 
 	return FUniqueFileHandle(Ipf.OpenWrite(*Filename, bAppend, bAllowRead));
@@ -1692,7 +1692,8 @@ private:
 													const TArray<FSharedOnDemandContainer>& Containers,
 													const TArray<TBitArray<>>& ChunkEntryIndices,
 													FCasBlockInfoMap& BlockInfo, 
-													const uint64* TotalBytesToFree = nullptr);
+													const uint64* TotalBytesToFree = nullptr,
+													uint64* OutTotalFreedBytes = nullptr);
 	bool										Resolve(FIoRequestImpl* Request);
 	void										CompleteRequest(FIoRequestImpl* Request, EIoErrorCode Status);
 	FIoStatus									FlushPendingChunks(FPendingChunks& Chunks, int64 UtcAccessTicks = 0);
@@ -2234,14 +2235,28 @@ FIoStatus FOnDemandInstallCache::Purge(TSet<FIoHash>&& ChunksToInstall)
 	for (const TPair<FCasBlockId, FCasBlockInfo>& Kv : BlockInfo)
 	{
 		const FCasBlockInfo& Info = Kv.Value;
+
+		if (Info.FileSize > Cas.GetMaxBlockSize())
+		{
+			UE_LOGFMT(LogIoStoreOnDemand, Error, "CAS Block {BlockId} has total size {BlockSize} which is greater than max block size {MaxBlockSize}", Kv.Key.Id, Info.FileSize, Cas.GetMaxBlockSize());
+			ensure(false);
+		}
+
 		if (Info.RefSize < Info.FileSize)
 		{
 			FragmentedBytes += (Info.FileSize - Info.RefSize);
 		}
+		else if (Info.RefSize > Info.FileSize)
+		{
+			UE_LOGFMT(LogIoStoreOnDemand, Error, "CAS Block {BlockId} has RefSize {RefSize} which is greater than total size {BlockSize}", Kv.Key.Id, Info.RefSize, Info.FileSize);
+			ensure(false);
+		}
+
 		if (Info.RefSize > 0)
 		{
 			TotalReferencedBlockBytes += Info.FileSize;
 		}
+
 		if (Info.LastAccess < OldestBlockAccess)
 		{
 			OldestBlockAccess = Info.LastAccess;
@@ -2280,13 +2295,14 @@ FIoStatus FOnDemandInstallCache::Purge(TSet<FIoHash>&& ChunksToInstall)
 		LogIoStoreOnDemand, Warning, TEXT("Max install cache size exceeded by %.2lf MiB (%.2lf%%)"),
 			ToMiB(NewCachedBytes - MaxCacheSize), 100.0 * (double(NewCachedBytes - MaxCacheSize) / double(MaxCacheSize)));
 
+	uint64 DefragPurgedBytes = 0;
 	if (Status.IsOk() && TotalPurgedBytes < TotalBytesToPurge)
 	{
 		if (UE::IoStore::CVars::GIoStoreOnDemandEnableDefrag)
 		{
 			// Attempt to defrag
 			const uint64 DefragBytesToPurge = TotalBytesToPurge - TotalPurgedBytes;
-			FIoStatus DefragStatus = Defrag(Containers, ChunkEntryIndices, BlockInfo, &DefragBytesToPurge);
+			FIoStatus DefragStatus = Defrag(Containers, ChunkEntryIndices, BlockInfo, &DefragBytesToPurge, &DefragPurgedBytes);
 			if (DefragStatus.IsOk() == false)
 			{
 				Status = InstallCacheStatusBuilder(EIoErrorCode::WriteError)
@@ -2299,7 +2315,7 @@ FIoStatus FOnDemandInstallCache::Purge(TSet<FIoHash>&& ChunksToInstall)
 		}
 	}
 
-	FOnDemandInstallCacheStats::OnPurge(Status, MaxCacheSize, NewCachedBytes, TotalBytesToPurge, TotalPurgedBytes);
+	FOnDemandInstallCacheStats::OnPurge(Status, MaxCacheSize, NewCachedBytes, TotalBytesToPurge, TotalPurgedBytes + DefragPurgedBytes);
 
 	return Status;
 }
@@ -2757,13 +2773,22 @@ uint64 FOnDemandInstallCache::AddReferencesToBlocks(
 				}
 			}
 
-			if (BlockInfo)
+			// If a chunk is in the list it WILL be installed even if already cached so must be counted in the uncached size.
+			if (bToInstall)
 			{
-				BlockInfo->RefSize += ChunkDiskSize;
+				TotalUncachedBytes += ChunkDiskSize;
+
+				if (bIsReferenced)
+				{
+					UE_LOGFMT(LogIoStoreOnDemand, Error, "Installing a chunk but its already cached! ChunkId={ChunkId}, Container={Container}}", 
+						LexToString(Container->ChunkIds[EntryIndex]), Container->Name);
+					ensure(false);
+				}
 			}
-			else if (bToInstall)
+			
+			if (bIsReferenced && BlockInfo)
 			{
-				TotalUncachedBytes += ChunkDiskSize;
+				BlockInfo->RefSize += ChunkDiskSize;
 			}
 
 			EntryIndex++;
@@ -2833,7 +2858,8 @@ FIoStatus FOnDemandInstallCache::Defrag(
 	const TArray<FSharedOnDemandContainer>& Containers,
 	const TArray<TBitArray<>>& ChunkEntryIndices,
 	FCasBlockInfoMap& BlockInfo, 
-	const uint64* TotalBytesToFree /*= nullptr*/)
+	const uint64* TotalBytesToFree /*= nullptr*/,
+	uint64* OutTotalFreedBytes /*= nullptr*/)
 {
 	if (TotalBytesToFree && *TotalBytesToFree == 0)
 	{
@@ -2928,7 +2954,7 @@ FIoStatus FOnDemandInstallCache::Defrag(
 				{
 					FragmentedBytes += (Info.FileSize - Info.RefSize);
 				}
-				
+
 				TotalBlockSize += Info.FileSize;
 
 				BlocksToDefrag.Add(FDefragBlock{ .BlockId = BlockId, .LastAccess = Info.LastAccess });
@@ -2989,6 +3015,8 @@ FIoStatus FOnDemandInstallCache::Defrag(
 	// This should be the only thread writing to CurrentBlock.
 	CurrentBlock = FCasBlockId::Invalid;
 
+	uint64 TotalRefBytes = 0;
+
 	// Determine chunks that need to be moved for each defrag block
 	for (int32 Index = 0; FSharedOnDemandContainer Container : Containers)
 	{
@@ -3005,16 +3033,24 @@ FIoStatus FOnDemandInstallCache::Defrag(
 				if (FDefragBlock* DefragBlock = Algo::FindBy(BlocksToDefrag, Loc.BlockId, &FDefragBlock::BlockId))
 				{
 					DefragBlock->ReferencedChunks.Add(FDefragBlockReferencedChunk
-					{
-						.BlockOffset = Loc.BlockOffset,
-						.DiskSize = Entry.GetDiskSize(),
-						.Hash = Entry.Hash,
-					});
+						{
+							.BlockOffset = Loc.BlockOffset,
+							.DiskSize = Entry.GetDiskSize(),
+							.Hash = Entry.Hash,
+						});
+
+					TotalRefBytes += Entry.GetDiskSize();
 				}
 			}
 		}
 	}
 
+	if (TotalBlockSize - FragmentedBytes != TotalRefBytes)
+	{
+		UE_LOGFMT(LogIoStoreOnDemand, Error, "Possibly corrupt CAS blocks - TotalBlockSize {TotalBlockSize}, FragmentedBytes {FragmentedBytes}, TotalRefBytes {TotalRefBytes}", TotalBlockSize,  FragmentedBytes, TotalRefBytes);
+		ensure(false);
+	}
+
 	// Move chunks to new blocks and delete old blocks
 	FPendingChunks DefragPendingChunks;
 	for (const FDefragBlock& DefragBlock : BlocksToDefrag)
@@ -3127,6 +3163,12 @@ FIoStatus FOnDemandInstallCache::Defrag(
 
 	FIoStatus Status = InstallCacheStatusBuilder(EIoErrorCode::Ok);
 	FOnDemandInstallCacheStats::OnDefrag(Status, FragmentedBytes);
+
+	if (OutTotalFreedBytes)
+	{
+		*OutTotalFreedBytes = FragmentedBytes;
+	}
+
 	return Status;
 }
 
@@ -3278,6 +3320,9 @@ FIoStatus FOnDemandInstallCache::FlushPendingChunksImpl(FPendingChunks& Chunks,
 	{
 		FCasJournal::FTransaction Transaction = FCasJournal::Begin(GetJournalFilename());
 		
+		// Only open for append if continuing a block.
+		const bool bAppendToBlock = CurrentBlockId.IsValid();
+
 		if (CurrentBlockId.IsValid() == false)
 		{
 			CurrentBlockId = Cas.CreateBlock();
@@ -3286,7 +3331,7 @@ FIoStatus FOnDemandInstallCache::FlushPendingChunksImpl(FPendingChunks& Chunks,
 			Transaction.BlockCreated(CurrentBlockId);
 		}
 
-		TUniquePtr<IFileHandle>	CasFileHandle = Cas.OpenWrite(CurrentBlockId);
+		TUniquePtr<IFileHandle>	CasFileHandle = Cas.OpenWrite(CurrentBlockId, bAppendToBlock);
 		if (CasFileHandle.IsValid() == false)
 		{
 			FOnDemandInstallCacheStats::OnFlush(EIoErrorCode::WriteError, 0);
