Fix extra brace in FMT string

diff --git a/Engine/Source/Runtime/Experimental/IoStore/OnDemand/Private/OnDemandInstallCache.cpp b/Engine/Source/Runtime/Experimental/IoStore/OnDemand/Private/OnDemandInstallCache.cpp
--- a/Engine/Source/Runtime/Experimental/IoStore/OnDemand/Private/OnDemandInstallCache.cpp
+++ b/Engine/Source/Runtime/Experimental/IoStore/OnDemand/Private/OnDemandInstallCache.cpp
@@ -2780,7 +2780,7 @@ uint64 FOnDemandInstallCache::AddReferencesToBlocks(
 
 				if (bIsReferenced)
 				{
-					UE_LOGFMT(LogIoStoreOnDemand, Error, "Installing a chunk but its already cached! ChunkId={ChunkId}, Container={Container}}", 
+					UE_LOGFMT(LogIoStoreOnDemand, Error, "Installing a chunk but its already cached! ChunkId={ChunkId}, Container={Container}", 
 						LexToString(Container->ChunkIds[EntryIndex]), Container->Name);
 					ensure(false);
 				}
