Don't populate solver constraint manifold points in the case of two non-dynamic contacting bodies.

diff --git a/Engine/Source/Runtime/Experimental/Chaos/Private/Chaos/Collision/PBDCollisionContainerSolver.cpp b/Engine/Source/Runtime/Experimental/Chaos/Private/Chaos/Collision/PBDCollisionContainerSolver.cpp
--- a/Engine/Source/Runtime/Experimental/Chaos/Private/Chaos/Collision/PBDCollisionContainerSolver.cpp
+++ b/Engine/Source/Runtime/Experimental/Chaos/Private/Chaos/Collision/PBDCollisionContainerSolver.cpp
@@ -327,7 +327,11 @@ namespace Chaos
 
 		bOutPerIterationCollision = (!Constraint->GetUseManifold() || Constraint->GetUseIncrementalCollisionDetection());
 
-		UpdateCollisionSolverManifoldFromConstraint(Solver, Constraint, Dt, 0, Constraint->NumManifoldPoints(), SolverSettings);
+		if (Solver.SolverBody0().IsDynamic() ||
+			Solver.SolverBody1().IsDynamic())
+		{
+			UpdateCollisionSolverManifoldFromConstraint(Solver, Constraint, Dt, 0, Constraint->NumManifoldPoints(), SolverSettings);
+		}
 
 		// Convert to a soft collision if the constraint has a nonzero soft separation
 		// NOTE: must come after UpdateCollisionSolverManifoldFromConstraint because we scale the spring forces by the number of manifold points.
