Fixed server crash
Skip pruning if the integration pending threads are not matching the the dynamics particles batches.
Then link the integrations task to the compute spatial acceleration
Server Assert occurs at Chaos::FPBDRigidsEvolutionBase::PruneInternalPendingData

diff --git a/Engine/Source/Runtime/Experimental/Chaos/Public/Chaos/TaskDispatcherEvolution.h b/Engine/Source/Runtime/Experimental/Chaos/Public/Chaos/TaskDispatcherEvolution.h
--- a/Engine/Source/Runtime/Experimental/Chaos/Public/Chaos/TaskDispatcherEvolution.h
+++ b/Engine/Source/Runtime/Experimental/Chaos/Public/Chaos/TaskDispatcherEvolution.h
@@ -193,7 +193,7 @@ public:
 			CurrentPendingTasks.Add(CleanUpTask);
 		}
 		// Prune Dynamics particles
-		if (NumBatchDyn > 0)
+		if (NumBatchDyn > 0 && IntegrationPendingTasks.Num() == NumBatchDyn)
 		{
 			check(IntegrationPendingTasks.Num() == NumBatchDyn);
 			check(NumBatchDyn <= InternalAccelerationQueue.PendingDataArrays.Num());
@@ -232,6 +232,11 @@ public:
 			}, DynamicTasks, LowLevelTasks::ETaskPriority::High);
 			CurrentPendingTasks.Add(CopyDynamicTask);
 		}
+		else if (NumBatchDyn > 0 && IntegrationPendingTasks.Num() != NumBatchDyn)
+		{
+			UE_LOG(LogChaos, Warning, TEXT("No pruning happened in the Spatial Acceleration structure, a slight perfromance hit could occured. IntegrationPendingTasks.Num(): %d !=  NumBatchDyn:  %d"), IntegrationPendingTasks.Num(), NumBatchDyn);
+			CurrentPendingTasks.Append(IntegrationPendingTasks);
+		}
 		int32 NumBatchKin = InternalAccelerationQueue.PendingDataArrays.Num() - InternalAccelerationQueue.KinematicBatchStartIndex;
 		// Prune Kinematics particles
 		if (NumBatchKin > 0 && InternalAccelerationQueue.KinematicBatchStartIndex > 0)
