[UBT][CppCompileWarnings] Engine compliance wundef

diff --git a/Engine/Source/Runtime/Experimental/IoStore/OnDemand/Private/DebugCommands.h b/Engine/Source/Runtime/Experimental/IoStore/OnDemand/Private/DebugCommands.h
--- a/Engine/Source/Runtime/Experimental/IoStore/OnDemand/Private/DebugCommands.h
+++ b/Engine/Source/Runtime/Experimental/IoStore/OnDemand/Private/DebugCommands.h
@@ -2,6 +2,8 @@
 
 #pragma once
 
+#include "Misc/Build.h"
+
 #if !UE_BUILD_SHIPPING
 
 #include "Containers/Array.h"
