Removed unused sink

#rnx

diff --git a/Engine/Source/Runtime/Experimental/IoStore/HttpClient/Private/Test.inl b/Engine/Source/Runtime/Experimental/IoStore/HttpClient/Private/Test.inl
--- a/Engine/Source/Runtime/Experimental/IoStore/HttpClient/Private/Test.inl
+++ b/Engine/Source/Runtime/Experimental/IoStore/HttpClient/Private/Test.inl
@@ -249,10 +249,6 @@ static void TlsTest()
 		check(Status.GetId() == FTicketStatus::EId::Content);
 	};
 
-	auto NotOkSink = [Dest=FIoBuffer()] (const FTicketStatus& Status) mutable {
-		check(Status.GetId() == FTicketStatus::EId::Error);
-	};
-
 	static const ANSICHAR* Url = "http://epicgames.com";
 
 	{
