Fixed brace rules and stayed more true to using/passing the ref instance.

#rnx

diff --git a/Engine/Source/Runtime/Experimental/IoStore/HttpClient/Private/Activity.inl b/Engine/Source/Runtime/Experimental/IoStore/HttpClient/Private/Activity.inl
--- a/Engine/Source/Runtime/Experimental/IoStore/HttpClient/Private/Activity.inl
+++ b/Engine/Source/Runtime/Experimental/IoStore/HttpClient/Private/Activity.inl
@@ -95,7 +95,7 @@ struct FActivity
 	void				Fail(const FOutcome& Outcome);
 	EStage				GetStage() const;
 	FTransactId			GetTransactId() const;
-	const FTransaction*	GetTransaction() const;
+	const FTransactRef&	GetTransaction() const;
 	FAnsiStringView		GetMethod() const;
 	FAnsiStringView		GetPath() const;
 	FHost*				GetHost() const;
@@ -316,11 +316,11 @@ FTransactId FActivity::GetTransactId() const
 }
 
 ////////////////////////////////////////////////////////////////////////////////
-const FTransaction* FActivity::GetTransaction() const
+const FTransactRef& FActivity::GetTransaction() const
 {
 	check(State >= EState::Send && State <= EState::Completed);
 	check(Transaction.IsValid());
-	return Transaction.Get();
+	return Transaction;
 }
 
 ////////////////////////////////////////////////////////////////////////////////
diff --git a/Engine/Source/Runtime/Experimental/IoStore/HttpClient/Private/Peer.inl b/Engine/Source/Runtime/Experimental/IoStore/HttpClient/Private/Peer.inl
--- a/Engine/Source/Runtime/Experimental/IoStore/HttpClient/Private/Peer.inl
+++ b/Engine/Source/Runtime/Experimental/IoStore/HttpClient/Private/Peer.inl
@@ -591,7 +591,7 @@ FTransactRef FHttpPeer::Transact(FBuffer& Buffer)
 {
 	if (Proto == 1)
 	{
-		return FTransactRef(CreateTransactOne(Buffer));
+		return CreateTransactOne(Buffer);
 	}
 
 	check(!"unsupported");
diff --git a/Engine/Source/Runtime/Experimental/IoStore/HttpClient/Private/TransactionOne.inl b/Engine/Source/Runtime/Experimental/IoStore/HttpClient/Private/TransactionOne.inl
--- a/Engine/Source/Runtime/Experimental/IoStore/HttpClient/Private/TransactionOne.inl
+++ b/Engine/Source/Runtime/Experimental/IoStore/HttpClient/Private/TransactionOne.inl
@@ -870,13 +870,6 @@ public:
 	using DetailOne::FBody::FBody;
 };
 
-UPTRINT CreateTransactOne(FBuffer& Buffer)
-{
-	void* Ptr = FMemory::Malloc(sizeof(FTransactionOne), alignof(FTransactionOne));
-	new (Ptr) FTransactionOne(Buffer);
-	return UPTRINT(Ptr);
-}
-
 
 
 ////////////////////////////////////////////////////////////////////////////////
@@ -920,22 +913,65 @@ class FTransactRef
 public:
 						FTransactRef() = default;
 						~FTransactRef();
-						FTransactRef(UPTRINT Ptr) : Self(Ptr) {}
 						FTransactRef(FTransactRef&&) = delete;
 						FTransactRef(const FTransactRef&) = delete;
 	void				operator = (const FTransactRef& Rhs) = delete;
-	void				operator = (FTransactRef&& Rhs) { Swap(Self, Rhs.Self); };
-	const FTransaction*	operator -> () const	{ return Get(); }
-	FTransaction*		operator -> ()			{ return Get(); }
-	bool				IsValid() const			{ return Get() != nullptr; }
-	const FTransaction*	Get() const				{ return (FTransaction*)Self; }
-	FTransaction*		Get()					{ return (FTransaction*)Self; }
-
+	void				operator = (FTransactRef&& Rhs);
+	const FTransaction*	operator -> () const;
+	FTransaction*		operator -> ();
+	bool				IsValid() const;
 
 private:
+	friend FTransactRef CreateTransactOne(FBuffer&);
+
+						FTransactRef(UPTRINT Ptr);
+	const FTransaction*	Get() const;
+	FTransaction*		Get();
 	UPTRINT				Self = 0;
 };
 
+////////////////////////////////////////////////////////////////////////////////
+FTransactRef::FTransactRef(UPTRINT Ptr)
+: Self(Ptr)
+{
+}
+
+////////////////////////////////////////////////////////////////////////////////
+void FTransactRef::operator = (FTransactRef&& Rhs)
+{
+	Swap(Self, Rhs.Self);
+}
+
+////////////////////////////////////////////////////////////////////////////////
+const FTransaction*	FTransactRef::operator -> () const
+{
+	return Get();
+}
+
+////////////////////////////////////////////////////////////////////////////////
+FTransaction* FTransactRef::operator -> ()		
+{
+	return Get();
+}
+
+////////////////////////////////////////////////////////////////////////////////
+bool FTransactRef::IsValid() const		
+{
+	return Get() != nullptr;
+}
+
+////////////////////////////////////////////////////////////////////////////////
+const FTransaction* FTransactRef::Get() const			
+{
+	return (FTransaction*)Self;
+}
+
+////////////////////////////////////////////////////////////////////////////////
+FTransaction* FTransactRef::Get()				
+{
+	return (FTransaction*)Self;
+}
+
 ////////////////////////////////////////////////////////////////////////////////
 FTransactRef::~FTransactRef()
 {
@@ -947,4 +983,14 @@ FTransactRef::~FTransactRef()
 	Get()->Destruct();
 }
 
+
+
+////////////////////////////////////////////////////////////////////////////////
+FTransactRef CreateTransactOne(FBuffer& Buffer)
+{
+	void* Ptr = FMemory::Malloc(sizeof(FTransactionOne), alignof(FTransactionOne));
+	new (Ptr) FTransactionOne(Buffer);
+	return FTransactRef(UPTRINT(Ptr));
+}
+
 } // namespace UE::IoStore::HTTP
