Broader success conditions for redirect test.

#rnx

diff --git a/Engine/Source/Runtime/Experimental/IoStore/HttpClient/Private/Test.inl b/Engine/Source/Runtime/Experimental/IoStore/HttpClient/Private/Test.inl
--- a/Engine/Source/Runtime/Experimental/IoStore/HttpClient/Private/Test.inl
+++ b/Engine/Source/Runtime/Experimental/IoStore/HttpClient/Private/Test.inl
@@ -242,7 +242,7 @@ static void TlsTest()
 		if (Status.GetId() == FTicketStatus::EId::Response)
 		{
 			FResponse& Response = Status.GetResponse();
-			check(Response.GetStatusCode() == 200);
+			check(Response.GetStatus() == EStatusCodeClass::Successful);
 			Response.SetDestination(&Dest);
 			return;
 		}
