clear chunked on content-less

#rnx

diff --git a/Engine/Source/Runtime/Experimental/IoStore/HttpClient/Private/TransactionOne.inl b/Engine/Source/Runtime/Experimental/IoStore/HttpClient/Private/TransactionOne.inl
--- a/Engine/Source/Runtime/Experimental/IoStore/HttpClient/Private/TransactionOne.inl
+++ b/Engine/Source/Runtime/Experimental/IoStore/HttpClient/Private/TransactionOne.inl
@@ -515,6 +515,7 @@ FOutcome FStatusHeaders::Parse()
 	if (IsContentless(StatusCode))
 	{
 		ContentLength = 0;
+		bChunked = false;
 	}
 
 	return FOutcome::Ok();
