Zero out content length just in case.

#rnx

diff --git a/Engine/Source/Runtime/Experimental/IoStore/HttpClient/Private/Activity.inl b/Engine/Source/Runtime/Experimental/IoStore/HttpClient/Private/Activity.inl
--- a/Engine/Source/Runtime/Experimental/IoStore/HttpClient/Private/Activity.inl
+++ b/Engine/Source/Runtime/Experimental/IoStore/HttpClient/Private/Activity.inl
@@ -499,7 +499,7 @@ FOutcome FActivity::RecvMessage(FHttpPeer& Peer)
 
 	// HEAD methods
 	bool bHasBody = !GetMethod().Equals("HEAD", ESearchCase::IgnoreCase);
-	bHasBody &= (Transaction->GetContentLength() > 0) | int32(Transaction->IsChunked());
+	bHasBody &= (ContentLength > 0) | int32(Transaction->IsChunked());
 	if (!bHasBody)
 	{
 		ChangeState(EState::RecvDone);
diff --git a/Engine/Source/Runtime/Experimental/IoStore/HttpClient/Private/TransactionOne.inl b/Engine/Source/Runtime/Experimental/IoStore/HttpClient/Private/TransactionOne.inl
--- a/Engine/Source/Runtime/Experimental/IoStore/HttpClient/Private/TransactionOne.inl
+++ b/Engine/Source/Runtime/Experimental/IoStore/HttpClient/Private/TransactionOne.inl
@@ -511,6 +511,10 @@ FOutcome FStatusHeaders::Parse()
 	{
 		return FOutcome::Error("Received content when none was expected");
 	}
+	else
+	{
+		ContentLength = 0;
+	}
 
 	if (IsContentless(StatusCode))
 	{
