Narrowing warning from Clang fix

#rnx

diff --git a/Engine/Source/Runtime/Experimental/IoStore/HttpClient/Private/Peer.inl b/Engine/Source/Runtime/Experimental/IoStore/HttpClient/Private/Peer.inl
--- a/Engine/Source/Runtime/Experimental/IoStore/HttpClient/Private/Peer.inl
+++ b/Engine/Source/Runtime/Experimental/IoStore/HttpClient/Private/Peer.inl
@@ -286,7 +286,7 @@ static int32 Ssl_GetProtocolVersion(SSL* Ssl)
 		"http/1.1",
 	//	"h2",
 	};
-	for (int32 i = 0; i < UE_ARRAY_COUNT(Candidates); ++i)
+	for (uint32 i = 0; i < UE_ARRAY_COUNT(Candidates); ++i)
 	{
 		const FAnsiStringView& Candidate = Candidates[i];
 		if (AlpnProtoLen != uint32(Candidate.Len()))
