Move-to-head instead of Find-first as the former makes list manipulation more reliable.

#rnx

diff --git a/Engine/Source/Runtime/Experimental/IoStore/HttpClient/Private/Loop.inl b/Engine/Source/Runtime/Experimental/IoStore/HttpClient/Private/Loop.inl
--- a/Engine/Source/Runtime/Experimental/IoStore/HttpClient/Private/Loop.inl
+++ b/Engine/Source/Runtime/Experimental/IoStore/HttpClient/Private/Loop.inl
@@ -179,7 +179,7 @@ public:
 	void			PopAppend(FActivityList& ToList);
 
 	template <typename LAMBDA>
-	FActivityNode*	FindFirst(LAMBDA&& Predicate);
+	FActivityNode*	MoveToHead(LAMBDA&& Predicate);
 
 private:
 	FActivityNode*	PopImpl();
@@ -194,14 +194,27 @@ FActivityList::~FActivityList()
 
 ////////////////////////////////////////////////////////////////////////////////
 template <typename LAMBDA>
-FActivityNode* FActivityList::FindFirst(LAMBDA&& Predicate)
+FActivityNode* FActivityList::MoveToHead(LAMBDA&& Predicate)
 {
-	for (FActivityNode* Node = Head; Node != nullptr; Node = Node->Next)
+	FActivityNode* Node = Head;
+	FActivityNode* Prev = nullptr;
+	for (; Node != nullptr; Prev = Node, Node = Node->Next)
 	{
-		if (Predicate(Node))
+		if (!Predicate(Node))
+		{
+			continue;
+		}
+
+		if (Prev == nullptr)
 		{
+			check(Head == Node);
 			return Node;
 		}
+
+		Prev->Next = Node->Next;
+		Node->Next = Head;
+		Head = Node;
+		return Node;
 	}
 
 	return nullptr;
@@ -251,8 +264,9 @@ void FActivityList::Append(FActivityNode* Node)
 		return;
 	}
 
-	check(Head->Next == nullptr);
-	Head->Next = Node;
+	FActivityNode* Tail = Head;
+	for (; Tail->Next != nullptr; Tail = Tail->Next);
+	Tail->Next = Node;
 }
 
 ////////////////////////////////////////////////////////////////////////////////
@@ -494,7 +508,7 @@ void FPeerGroup::RecvInternal(FTickState& State)
 	while (Outcome.IsOk())
 	{
 		uint32 TransactId = Outcome.GetResult();
-		Activity = Recv.FindFirst(
+		Activity = Recv.MoveToHead(
 			[TransactId] (FActivityNode* Activity)
 			{
 				return Activity->GetTransactId() == TransactId;
@@ -563,6 +577,7 @@ void FPeerGroup::RecvInternal(FTickState& State)
 
 	check(Outcome.GetResult() == uint32(FActivity::EStage::Content));
 
+	check(Recv.GetHead() == Activity);
 	Recv.PopPrepend(State.DoneList);
 
 	// If the server wants the connection closed then we can drop our peer
