Added accessor to query what layers a peer is using.

#rnx

diff --git a/Engine/Source/Runtime/Experimental/IoStore/HttpClient/Private/Peer.inl b/Engine/Source/Runtime/Experimental/IoStore/HttpClient/Private/Peer.inl
--- a/Engine/Source/Runtime/Experimental/IoStore/HttpClient/Private/Peer.inl
+++ b/Engine/Source/Runtime/Experimental/IoStore/HttpClient/Private/Peer.inl
@@ -429,6 +429,7 @@ public:
 				FTlsPeer(FTlsPeer&& Rhs)				{ Move(MoveTemp(Rhs)); }
 				FTlsPeer& operator = (FTlsPeer&& Rhs)	{ return Move(MoveTemp(Rhs)); }
 				FTlsPeer(FSocket InSocket, FCertRootsRef Certs=ECertRootsRefType::None, const char* HostName=nullptr);
+	bool		IsUsingTls() const;
 	FOutcome	Handshake();
 	FOutcome	Send(const char* Data, int32 Size);
 	FOutcome	Recv(char* Out, int32 MaxSize);
@@ -460,6 +461,12 @@ FTlsPeer::~FTlsPeer()
 	}
 }
 
+////////////////////////////////////////////////////////////////////////////////
+bool FTlsPeer::IsUsingTls() const
+{
+	return (Ssl != nullptr);
+}
+
 ////////////////////////////////////////////////////////////////////////////////
 FTlsPeer& FTlsPeer::Move(FTlsPeer&& Rhs)
 {
