Removed an unused method.

#rnx

diff --git a/Engine/Source/Runtime/Experimental/IoStore/HttpClient/Private/TransactionOne.inl b/Engine/Source/Runtime/Experimental/IoStore/HttpClient/Private/TransactionOne.inl
--- a/Engine/Source/Runtime/Experimental/IoStore/HttpClient/Private/TransactionOne.inl
+++ b/Engine/Source/Runtime/Experimental/IoStore/HttpClient/Private/TransactionOne.inl
@@ -212,24 +212,15 @@ public:
 	FTransactId		End();
 	FOutcome		TrySendRequest(FTlsPeer& Peer);
 
-protected:
-	bool			IsHeadMethod() const;
-
 private:
 	FRequest&		operator << (FAnsiStringView Value);
 	uint16			HeaderLeft;
 	uint16			AlreadySent = 0;
 	int16			HeaderRight = -1;
 	int8			MethodLength = -1;
-	bool			IsHead;
+	uint8			_Unnused;
 };
 
-////////////////////////////////////////////////////////////////////////////////
-bool FRequest::IsHeadMethod() const
-{
-	return IsHead;
-}
-
 ////////////////////////////////////////////////////////////////////////////////
 void FRequest::Begin(FAnsiStringView Method, FAnsiStringView Path)
 {
@@ -238,7 +229,6 @@ void FRequest::Begin(FAnsiStringView Method, FAnsiStringView Path)
 	*this << Method << " " << Path << " HTTP/1.1" "\r\n";
 	MethodLength = int8(Method.Len());
 	HeaderLeft = uint16(GetSize());
-	IsHead = Method.Equals("HEAD", ESearchCase::IgnoreCase);
 }
 
 ////////////////////////////////////////////////////////////////////////////////
