Strategic edit to be able to identify the different reasons we wait.

#rnx

diff --git a/Engine/Source/Runtime/Experimental/IoStore/HttpClient/Private/Misc.inl b/Engine/Source/Runtime/Experimental/IoStore/HttpClient/Private/Misc.inl
--- a/Engine/Source/Runtime/Experimental/IoStore/HttpClient/Private/Misc.inl
+++ b/Engine/Source/Runtime/Experimental/IoStore/HttpClient/Private/Misc.inl
@@ -74,10 +74,15 @@ class FOutcome
 public:
 	static FOutcome Ok(int32 Result=0);
 	static FOutcome Waiting(int32 Result=0);
+	static FOutcome WaitBuffer(int32 Result=0);
+	static FOutcome WaitStream(int32 Result=0);
 	static FOutcome Error(const char* Message, int32 Code=-1);
 	static FOutcome None()				{ return Error(""); }
 	bool			IsError() const		{ return Message < 0x8000'0000'0000; }
-	bool			IsWaiting() const	{ return Tag == WaitTag; }
+	bool			IsWaiting() const	{ return (Tag & WaitTag) == WaitTag; }
+	bool			IsWaitData() const	{ check(IsWaiting()); return (Tag & ~WaitTag) == 0; }
+	bool			IsWaitBuffer() const{ check(IsWaiting()); return (Tag & BufferTag) != 0; }
+	bool			IsWaitStream() const{ check(IsWaiting()); return (Tag & StreamTag) != 0; }
 	bool			IsOk() const		{ return Tag == OkTag; }
 	FAnsiStringView GetMessage() const	{ check(IsError()); return (const char*)UPTRINT(Message); }
 	int32			GetErrorCode() const{ check(IsError()); return int32(Code); }
@@ -86,8 +91,10 @@ public:
 private:
 					FOutcome() = default;
 
-	static uint32 const OkTag	= 0x0000'8000;
-	static uint32 const WaitTag = 0x0001'8000;
+	static uint32 const OkTag		= 0x0000'8000;
+	static uint32 const WaitTag		= 0x0001'8000;
+	static uint32 const BufferTag	= 0x0002'0000;
+	static uint32 const StreamTag	= 0x0004'0000;
 
 	union {
 		struct {
@@ -120,6 +127,26 @@ FOutcome FOutcome::Waiting(int32 Result)
 	return Outcome;
 }
 
+////////////////////////////////////////////////////////////////////////////////
+FOutcome FOutcome::WaitBuffer(int32 Result)
+{
+	// Given buffer space exhausted
+	FOutcome Outcome;
+	Outcome.Tag = WaitTag|BufferTag;
+	Outcome.Result = Result;
+	return Outcome;
+}
+
+////////////////////////////////////////////////////////////////////////////////
+FOutcome FOutcome::WaitStream(int32 Result)
+{
+	// Need to switch to another stream
+	FOutcome Outcome;
+	Outcome.Tag = WaitTag|StreamTag;
+	Outcome.Result = Result;
+	return Outcome;
+}
+
 ////////////////////////////////////////////////////////////////////////////////
 FOutcome FOutcome::Error(const char* Message, int32 Code)
 {
diff --git a/Engine/Source/Runtime/Experimental/IoStore/HttpClient/Private/Test.inl b/Engine/Source/Runtime/Experimental/IoStore/HttpClient/Private/Test.inl
--- a/Engine/Source/Runtime/Experimental/IoStore/HttpClient/Private/Test.inl
+++ b/Engine/Source/Runtime/Experimental/IoStore/HttpClient/Private/Test.inl
@@ -154,6 +154,17 @@ static void MiscTest()
 	check(FOutcome::Waiting().IsOk() == false);
 	check(FOutcome::Waiting().IsWaiting());
 	check(FOutcome::Waiting().IsError() == false);
+	check(FOutcome::Waiting().IsWaitData() == true);
+	check(FOutcome::Waiting().IsWaitBuffer() == false);
+	check(FOutcome::Waiting().IsWaitStream() == false);
+
+	check(FOutcome::WaitBuffer().IsWaitData() == false);
+	check(FOutcome::WaitBuffer().IsWaitBuffer() == true);
+	check(FOutcome::WaitBuffer().IsWaitStream() == false);
+
+	check(FOutcome::WaitStream().IsWaitData() == false);
+	check(FOutcome::WaitStream().IsWaitBuffer() == false);
+	check(FOutcome::WaitStream().IsWaitStream() == true);
 }
 
 ////////////////////////////////////////////////////////////////////////////////
