PVS warning maybe

#rnx

diff --git a/Engine/Source/Runtime/Experimental/IoStore/HttpClient/Private/Activity.inl b/Engine/Source/Runtime/Experimental/IoStore/HttpClient/Private/Activity.inl
--- a/Engine/Source/Runtime/Experimental/IoStore/HttpClient/Private/Activity.inl
+++ b/Engine/Source/Runtime/Experimental/IoStore/HttpClient/Private/Activity.inl
@@ -499,7 +499,7 @@ FOutcome FActivity::RecvMessage(FHttpPeer& Peer)
 
 	// HEAD methods
 	bool bHasBody = !GetMethod().Equals("HEAD", ESearchCase::IgnoreCase);
-	bHasBody &= (ContentLength > 0) | int32(Transaction->IsChunked());
+	bHasBody &= ((ContentLength > 0) | int32(Transaction->IsChunked())) == true;
 	if (!bHasBody)
 	{
 		ChangeState(EState::RecvDone);
