Changed I/O store on-demand API to use unified errors.

#rnx

diff --git a/Engine/Plugins/Runtime/GameFeatures/Source/GameFeatures/Private/GameFeaturePluginStateMachine.cpp b/Engine/Plugins/Runtime/GameFeatures/Source/GameFeatures/Private/GameFeaturePluginStateMachine.cpp
--- a/Engine/Plugins/Runtime/GameFeatures/Source/GameFeatures/Private/GameFeaturePluginStateMachine.cpp
+++ b/Engine/Plugins/Runtime/GameFeatures/Source/GameFeatures/Private/GameFeaturePluginStateMachine.cpp
@@ -1776,7 +1776,7 @@ struct FGameFeaturePluginState_Downloading : public FGameFeaturePluginState
 			// This should be pretty small, so not going to worry about progress here.
 			IoStoreOnDemandContext->InstallRequests.Add(
 				IoStore->Install(MoveTemp(InstallArgs),
-				[this, LambdaOnDemandContext = IoStoreOnDemandContext](const UE::IoStore::FOnDemandInstallResult& OnDemandInstallResult)
+				[this, LambdaOnDemandContext = IoStoreOnDemandContext](UE::IoStore::FOnDemandInstallResult&& OnDemandInstallResult)
 				{
 					if (!LambdaOnDemandContext->bStateValid)
 					{
@@ -1784,14 +1784,13 @@ struct FGameFeaturePluginState_Downloading : public FGameFeaturePluginState
 						return;
 					}
 
-					if (!OnDemandInstallResult.Status.IsOk() && !Result.HasError())
+					if (!OnDemandInstallResult.IsOk() && !Result.HasError())
 					{
-						const FString StatusString = OnDemandInstallResult.Status.ToString();
-						FString ErrorCode = GetIoErrorText(OnDemandInstallResult.Status.GetErrorCode());
-						ErrorCode.ReplaceCharInline(TEXT(' '), TEXT('_'), ESearchCase::CaseSensitive);
+						const FText ErrorMessage	= OnDemandInstallResult.Error.GetValue().GetErrorMessage();
+						FString ErrorCode			= FString(OnDemandInstallResult.Error.GetValue().GetModuleIdAndErrorCodeString());
 
-						// IOStore errors aren't localized 
-						Result = GetErrorResult(TEXT("IoStoreOnDemand.OnComplete."), ErrorCode, FText::AsCultureInvariant(StatusString));
+						ErrorCode.ReplaceCharInline(TEXT(' '), TEXT('_'), ESearchCase::CaseSensitive);
+						Result = GetErrorResult(TEXT("IoStoreOnDemand.OnComplete."), ErrorCode, FText::AsCultureInvariant(ErrorMessage));
 						TryCancelState();
 					}
 
@@ -3005,14 +3004,13 @@ struct FGameFeaturePluginState_AssetDependencyStreaming : public FGameFeaturePlu
 						return;
 					}
 
-					if (!OnDemandInstallResult.Status.IsOk() && !Result.HasError())
+					if (!OnDemandInstallResult.IsOk() && !Result.HasError())
 					{
-						const FString StatusString = OnDemandInstallResult.Status.ToString();
-						FString ErrorCode = GetIoErrorText(OnDemandInstallResult.Status.GetErrorCode());
-						ErrorCode.ReplaceCharInline(TEXT(' '), TEXT('_'), ESearchCase::CaseSensitive);
+						const FText ErrorMessage	= OnDemandInstallResult.Error.GetValue().GetErrorMessage();
+						FString ErrorCode			= FString(OnDemandInstallResult.Error.GetValue().GetModuleIdAndErrorCodeString());
 
-						// IOStore errors aren't localized 
-						Result = GetErrorResult(TEXT("IoStoreOnDemand.OnComplete."), ErrorCode, FText::AsCultureInvariant(StatusString));
+						ErrorCode.ReplaceCharInline(TEXT(' '), TEXT('_'), ESearchCase::CaseSensitive);
+						Result = GetErrorResult(TEXT("IoStoreOnDemand.OnComplete."), ErrorCode, ErrorMessage);
 						TryCancelState();
 					}
 
diff --git a/Engine/Plugins/Runtime/GameFeatures/Source/GameFeatures/Private/GameFeaturesSubsystem.cpp b/Engine/Plugins/Runtime/GameFeatures/Source/GameFeatures/Private/GameFeaturesSubsystem.cpp
--- a/Engine/Plugins/Runtime/GameFeatures/Source/GameFeatures/Private/GameFeaturesSubsystem.cpp
+++ b/Engine/Plugins/Runtime/GameFeatures/Source/GameFeatures/Private/GameFeaturesSubsystem.cpp
@@ -9,6 +9,7 @@
 #include "AssetRegistry/AssetData.h"
 #include "AssetRegistry/IAssetRegistry.h"
 #include "BundlePrereqCombinedStatusHelper.h"
+#include "Experimental/UnifiedError/UnifiedError.h"
 #include "Dom/JsonValue.h"
 #include "GameFeaturesSubsystemSettings.h"
 #include "GameFeaturesProjectPolicies.h"
@@ -3018,7 +3019,7 @@ struct FGameFeaturePluginPredownloadContext : public FGameFeaturePluginPredownlo
 	TArray<UE::IoStore::FOnDemandContentHandle> IoStoreContentHandles;
 	TArray<UE::IoStore::FOnDemandInstallRequest> IoStoreInstallRequests;
 	TArray<UE::IoStore::FOnDemandInstallProgress> IoStoreProgress;
-	FIoStatus IoStoreStatus;
+	TOptional<UE::UnifiedError::FError> IoStoreError;
 	int32 IoStorePendingInstalls = 0;
 
 	float Progress = 0.0f;
@@ -3446,14 +3447,13 @@ struct FGameFeaturePluginPredownloadContext : public FGameFeaturePluginPredownlo
 
 			IoStoreInstallRequests[i] = IoStore->Install(MoveTemp(InstallArgs),
 				// On Complete
-				[this, Context = SharedThis(this), i](const UE::IoStore::FOnDemandInstallResult& OnDemandInstallResult)
+				[this, Context = SharedThis(this), i](UE::IoStore::FOnDemandInstallResult&& OnDemandInstallResult)
 				{
 					UE::TUniqueLock Lock(IoStoreMutex);
 
-					if (!OnDemandInstallResult.Status.IsOk() && IoStoreStatus.IsOk())
+					if (!OnDemandInstallResult.IsOk() && !IoStoreError.IsSet())
 					{
-						IoStoreStatus = OnDemandInstallResult.Status;
-
+						IoStoreError = MoveTemp(OnDemandInstallResult.Error);
 						CancelIoStoreRequests();
 					}
 
@@ -3487,17 +3487,17 @@ struct FGameFeaturePluginPredownloadContext : public FGameFeaturePluginPredownlo
 		{
 			SetCompleteCanceled();
 		}
-		else if (!IoStoreStatus.IsOk())
+		else if (IoStoreError.IsSet())
 		{
-			UE_LOGFMT(LogGameFeatures, Error, "GFP Predownload failed to install iostore on demand assets, Error: {Error}", ("Error", IoStoreStatus.ToString()));
+			UE_LOGFMT(LogGameFeatures, Error, "GFP Predownload failed to install iostore on demand assets, Error: {Error}", IoStoreError.GetValue());
+
+			const FText ErrorMessage	= IoStoreError.GetValue().GetErrorMessage();
+			FString ErrorCode			= FString(IoStoreError.GetValue().GetModuleIdAndErrorCodeString());
 
-			const FString StatusString = IoStoreStatus.ToString();
-			FString ErrorCode = GetIoErrorText(IoStoreStatus.GetErrorCode());
 			ErrorCode.ReplaceCharInline(TEXT(' '), TEXT('_'), ESearchCase::CaseSensitive);
-			
 			UE::GameFeatures::FResult ErrorResult(MakeError(
-				FString::Printf(TEXT("%.*s%s"), PredownloadErrorNamespace.Len(), PredownloadErrorNamespace.GetData(), *ErrorCode)), 
-				FText::AsCultureInvariant(StatusString)); // IOStore errors aren't localized 
+				FString::Printf(TEXT("%.*s%s"), PredownloadErrorNamespace.Len(), PredownloadErrorNamespace.GetData(), *ErrorCode)),
+				ErrorMessage);	
 			SetComplete(MoveTemp(ErrorResult));
 		}
 		else
diff --git a/Engine/Source/Runtime/Engine/Private/StreamableManager.cpp b/Engine/Source/Runtime/Engine/Private/StreamableManager.cpp
--- a/Engine/Source/Runtime/Engine/Private/StreamableManager.cpp
+++ b/Engine/Source/Runtime/Engine/Private/StreamableManager.cpp
@@ -1,6 +1,7 @@
 // Copyright Epic Games, Inc. All Rights Reserved.
 
 #include "Engine/StreamableManager.h"
+#include "Experimental/UnifiedError/CoreErrorTypes.h"
 #include "Experimental/UnifiedError/UnifiedError.h"
 #include "Engine/Experimental/StreamableManagerError.h"
 #include "Experimental/IO/IoStatusError.h"
@@ -464,12 +465,12 @@ namespace UE::StreamableManager::Private
 
 			
 			InstallRequest = IoStore->Install(MoveTemp(InstallArgs),
-				[WeakOwningHandle = OwningHandle.ToWeakPtr(), DebugString = OwningHandle->GetDebugName()](const UE::IoStore::FOnDemandInstallResult& OnDemandInstallResult)
+				[WeakOwningHandle = OwningHandle.ToWeakPtr(), DebugString = OwningHandle->GetDebugName()](UE::IoStore::FOnDemandInstallResult&& OnDemandInstallResult)
 				{
 					TSharedPtr<FStreamableHandle> OwningHandle = WeakOwningHandle.Pin();
 					if (OwningHandle && OwningHandle->DownloadContext)
 					{
-						OwningHandle->DownloadContext->OnInstallComplete(OwningHandle.ToSharedRef(), OnDemandInstallResult);
+						OwningHandle->DownloadContext->OnInstallComplete(OwningHandle.ToSharedRef(), MoveTemp(OnDemandInstallResult));
 					}
 					else
 					{
@@ -487,21 +488,21 @@ namespace UE::StreamableManager::Private
 			);
 		}
 
-		void OnInstallComplete(const TSharedRef<FStreamableHandle>& OwningHandle, const UE::IoStore::FOnDemandInstallResult& OnDemandInstallResult)
+		void OnInstallComplete(const TSharedRef<FStreamableHandle>& OwningHandle, UE::IoStore::FOnDemandInstallResult&& OnDemandInstallResult)
 		{
 			// Destroy self at end of scope
 			TUniquePtr<FDownloadContext> Self = MoveTemp(OwningHandle->DownloadContext);
 			OwningHandle->DownloadContext = nullptr;
 
-			if (!OnDemandInstallResult.Status.IsOk())
+			if (!OnDemandInstallResult.IsOk())
 			{
-				if (OwningHandle->WasCanceled() && OnDemandInstallResult.Status.GetErrorCode() == EIoErrorCode::Cancelled)
+				if (OwningHandle->WasCanceled() && UE::UnifiedError::IsCancellationError(OnDemandInstallResult.Error.GetValue()))
 				{
 					CompleteLoad(OwningHandle);
 				}
 				else
 				{
-					CompleteLoad(OwningHandle, UE::UnifiedError::IoStore::ConvertError(OnDemandInstallResult.Status));
+					CompleteLoad(OwningHandle, MoveTemp(OnDemandInstallResult.Error.GetValue()));
 				}
 
 				return;
diff --git a/Engine/Source/Runtime/Experimental/IoStore/OnDemand/Private/OnDemandContentInstaller.cpp b/Engine/Source/Runtime/Experimental/IoStore/OnDemand/Private/OnDemandContentInstaller.cpp
--- a/Engine/Source/Runtime/Experimental/IoStore/OnDemand/Private/OnDemandContentInstaller.cpp
+++ b/Engine/Source/Runtime/Experimental/IoStore/OnDemand/Private/OnDemandContentInstaller.cpp
@@ -3,22 +3,22 @@
 #include "OnDemandContentInstaller.h"
 
 #include "Algo/Accumulate.h"
-#include "Algo/RemoveIf.h"
 #include "Async/UniqueLock.h"
 #include "Containers/RingBuffer.h"
+#include "Experimental/UnifiedError/CoreErrorTypes.h"
 #include "HAL/Platform.h"
 #include "IO/IoBuffer.h"
 #include "IO/IoChunkId.h"
 #include "IO/IoDispatcher.h"
-#include "IO/IoStoreOnDemand.h"
-#include "Serialization/PackageStore.h"
 #include "IO/IoStatus.h"
+#include "IO/IoStoreOnDemand.h"
+#include "IO/OnDemandError.h"
 #include "Misc/Timespan.h"
-#include "Misc/PackageName.h"
 #include "OnDemandHttpThread.h"
 #include "OnDemandInstallCache.h"
 #include "OnDemandIoStore.h"
 #include "OnDemandPackageStoreBackend.h"
+#include "Serialization/PackageStore.h"
 #include "Statistics.h"
 #include <atomic>
 
@@ -38,6 +38,12 @@ namespace UE::IoStore
 namespace Private
 {
 
+////////////////////////////////////////////////////////////////////////////////
+int32 SecondsToMillicseconds(double Seconds)
+{
+	return Seconds > 0.0 ? int32(Seconds / 1000.0) : 0;
+}
+
 ////////////////////////////////////////////////////////////////////////////////
 void ResolvePackageDependencies(
 	const TSet<FPackageId>& PackageIds,
@@ -223,6 +229,20 @@ void ResolveChunksToInstall(
 ////////////////////////////////////////////////////////////////////////////////
 uint32 FOnDemandContentInstaller::FRequest::NextSeqNo = 0;
 
+TOptional<UE::UnifiedError::FError> FOnDemandContentInstaller::FRequest::ConsumeError()
+{
+	if (Result.HasError())
+	{
+		return TOptional<UE::UnifiedError::FError>(Result.StealError());
+	}
+	else if (IsCancelled())
+	{
+		return TOptional<UE::UnifiedError::FError>(UE::UnifiedError::Core::CancellationError::MakeError());
+	}
+
+	return TOptional<UE::UnifiedError::FError>();
+}
+
 ////////////////////////////////////////////////////////////////////////////////
 FOnDemandContentInstaller::FOnDemandContentInstaller(FOnDemandIoStore& InIoStore, FOnDemandHttpThread& InHttpClient)
 	: IoStore(InIoStore)
@@ -317,8 +337,7 @@ void FOnDemandContentInstaller::CancelInstallRequest(FSharedInternalInstallReque
 
 				FRequest* Request = reinterpret_cast<FRequest*>(InstallRequest->InstallerRequest);
 
-				EIoErrorCode Expected = EIoErrorCode::Unknown;
-				if (Request->ErrorCode.compare_exchange_strong(Expected, EIoErrorCode::Cancelled) == false)
+				if (Request->TryCancel() == false)
 				{
 					return;
 				}
@@ -510,9 +529,7 @@ void FOnDemandContentInstaller::ProcessInstallRequest(FRequest& Request)
 
 	if (InstallRequest.Args.ContentHandle.IsValid() == false)
 	{
-		Request.ErrorCode	= EIoErrorCode::InvalidParameter;
-		Request.ErrorReason	= TEXT("Invalid content handle");
-
+		Request.Result = MakeError(UE::UnifiedError::Core::ArgumentError::MakeError(TEXT("ContentHandle"), TEXT("Invalid content handle")));
 		return CompleteInstallRequest(Request);
 	}
 
@@ -532,9 +549,7 @@ void FOnDemandContentInstaller::ProcessInstallRequest(FRequest& Request)
 		ContainersForInstallation,
 		PackageIdsToInstall); !Status.IsOk())
 	{
-		Request.ErrorCode	= Status.GetErrorCode();
-		Request.ErrorReason	= Status.ToString(); 
-
+		Request.Result = MakeError(UE::UnifiedError::Core::ArgumentError::MakeError(TEXT("InstallArgs"), Status.ToString()));
 		return CompleteInstallRequest(Request);
 	}
 
@@ -544,15 +559,16 @@ void FOnDemandContentInstaller::ProcessInstallRequest(FRequest& Request)
 		if (FCString::Strstr(*ContentHandle->DebugName, *UE::IoStore::CVars::IoStoreErrorOnRequest) != nullptr || 
 			FCString::Strstr(*InstallRequest.Args.DebugName, *UE::IoStore::CVars::IoStoreErrorOnRequest) != nullptr)
 		{
-			Request.ErrorCode	= (EIoErrorCode)((FMath::Rand() / (RAND_MAX / ((int)EIoErrorCode::Last-2))) + 2); // ignore unknown and ok states
-			Request.ErrorReason = FString::Printf(TEXT("Debug error requested on debug name %s"), *UE::IoStore::CVars::IoStoreErrorOnRequest);
-
+			Request.Result = MakeError(
+				UE::UnifiedError::Core::ArgumentError::MakeError(
+					TEXT("InstallArgs"),
+					FString::Printf(TEXT("Debug error requested on debug name %s"), *UE::IoStore::CVars::IoStoreErrorOnRequest)));
 			return CompleteInstallRequest(Request);
 		}
 	}
 #endif
 
-	if (Request.ErrorCode.load(std::memory_order_relaxed) == EIoErrorCode::Cancelled)
+	if (Request.IsCancelled())
 	{
 		return CompleteInstallRequest(Request);
 	}
@@ -572,32 +588,32 @@ void FOnDemandContentInstaller::ProcessInstallRequest(FRequest& Request)
 	if (Missing.IsEmpty() == false)
 	{
 		FIoDispatcher& IoDispatcher = FIoDispatcher::Get();
-		uint32 MissingCount = 0;
-
+		TArray<FIoChunkId> MissingChunkIds;
 		for (const FPackageId& PackageId : Missing)
 		{
 			const FIoChunkId ChunkId = CreatePackageDataChunkId(PackageId);
 			if (IoDispatcher.DoesChunkExist(ChunkId) == false)
 			{
-				UE_CLOG(MissingCount == 0, LogIoStoreOnDemand, Warning, TEXT("Failed to resolve the following chunk(s) for content handle '%s':"),
+				UE_CLOG(MissingChunkIds.Num() == 0, LogIoStoreOnDemand, Warning, TEXT("Failed to resolve the following chunk(s) for content handle '%s':"),
 					*LexToString(InstallRequest.Args.ContentHandle));
 
 				UE_LOG(LogIoStoreOnDemand, Warning, TEXT("\tChunkId='%s'"), *LexToString(ChunkId));
-				MissingCount++;
+				MissingChunkIds.Add(ChunkId);
 			}
 		}
 
-		if (MissingCount > 0 && !EnumHasAnyFlags(InstallRequest.Args.Options, EOnDemandInstallOptions::AllowMissingDependencies))
+		if (MissingChunkIds.IsEmpty() == false && !EnumHasAnyFlags(InstallRequest.Args.Options, EOnDemandInstallOptions::AllowMissingDependencies))
 		{
-			Request.ErrorCode	= EIoErrorCode::UnknownChunkID;
-			Request.ErrorReason	= FString::Printf(TEXT("Missing chunk(s), Count=%u, ContentHandle='%s'"),
-				MissingCount, *LexToString(InstallRequest.Args.ContentHandle));
-
+			Request.Result = MakeError(UE::UnifiedError::IoStoreOnDemand::ChunkMissingError::MakeError(
+				UE::UnifiedError::IoStoreOnDemand::FChunkMissingErrorContext
+				{
+					.ChunkIds = MoveTemp(MissingChunkIds)
+				}));
 			return CompleteInstallRequest(Request);
 		}
 	}
 
-	if (Request.ErrorCode.load(std::memory_order_relaxed) == EIoErrorCode::Cancelled)
+	if (Request.IsCancelled())
 	{
 		return CompleteInstallRequest(Request);
 	}
@@ -644,12 +660,10 @@ void FOnDemandContentInstaller::ProcessInstallRequest(FRequest& Request)
 
 	if (TotalInstallSize == 0)
 	{
-		Request.ErrorCode = EIoErrorCode::Ok;
-
 		return CompleteInstallRequest(Request);
 	}
 
-	if (Request.ErrorCode.load(std::memory_order_relaxed) == EIoErrorCode::Cancelled)
+	if (Request.IsCancelled())
 	{
 		return CompleteInstallRequest(Request);
 	}
@@ -671,7 +685,7 @@ void FOnDemandContentInstaller::ExecuteInstallRequest(FRequest& Request)
 	UE_LOG(LogIoStoreOnDemand, Verbose, TEXT("Executing install request, SeqNo=%u, Priority=%d, ContentHandle=(%s)"),
 		Request.SeqNo, Request.Priority, *LexToString(InstallRequest.Args.ContentHandle));
 
-	if (Request.ErrorCode.load(std::memory_order_relaxed) == EIoErrorCode::Cancelled)
+	if (Request.IsCancelled())
 	{
 		return CompleteInstallRequest(Request);
 	}
@@ -705,14 +719,12 @@ void FOnDemandContentInstaller::ExecuteInstallRequest(FRequest& Request)
 
 	if (InstallRequest.HttpRequestHandles.IsEmpty())
 	{
-		Request.ErrorCode = EIoErrorCode::Ok;
 		return CompleteInstallRequest(Request);
 	}
 
 	if (EnumHasAnyFlags(InstallRequest.Args.Options, EOnDemandInstallOptions::DoNotDownload))
 	{
-		Request.ErrorCode = EIoErrorCode::Disabled;
-		Request.ErrorReason = TEXT("DoNotDownload Set");
+		Request.Result = MakeError(UE::UnifiedError::Core::ArgumentError::MakeError(TEXT("Options"), TEXT("DoNotDownload flag specified when download required")));
 		return CompleteInstallRequest(Request);
 	}
 
@@ -727,16 +739,13 @@ void FOnDemandContentInstaller::ExecuteInstallRequest(FRequest& Request)
 			ChunksToInstall.Add(ChunkEntry.Hash);
 		}
 
-		if (FIoStatus Status = IoStore.InstallCache->Purge(MoveTemp(ChunksToInstall)); !Status.IsOk())
+		if (Request.Result = IoStore.InstallCache->Purge(MoveTemp(ChunksToInstall)); Request.Result.HasError())
 		{
-			Request.ErrorCode	= Status.GetErrorCode();
-			Request.ErrorReason	= Status.ToString(); 
-
 			return CompleteInstallRequest(Request);
 		}
 	}
 
-	if (Request.ErrorCode.load(std::memory_order_relaxed) == EIoErrorCode::Cancelled)
+	if (Request.IsCancelled())
 	{
 		return CompleteInstallRequest(Request);
 	}
@@ -833,20 +842,12 @@ void FOnDemandContentInstaller::ExecutePurgeRequest(FRequest& Request)
 	UE_LOG(LogIoStoreOnDemand, Log, TEXT("Executing purge request, BytesToPurge=%llu, Defrag='%s'"),
 		BytesToPurge != nullptr ? *BytesToPurge : -1, bDefrag ? TEXT("True") : TEXT("False"));
 
-	if (Request.ErrorCode.load(std::memory_order_relaxed) == EIoErrorCode::Cancelled)
-	{
-		return CompletePurgeRequest(Request);
-	}
-
-	const FIoStatus Status	= IoStore.InstallCache->PurgeAllUnreferenced(bDefrag, BytesToPurge);
-	Request.ErrorCode		= Status.GetErrorCode();
-
-	if (!Status.IsOk())
+	if (Request.IsCancelled() == false)
 	{
-		Request.ErrorReason = Status.ToString();
+		Request.Result = IoStore.InstallCache->PurgeAllUnreferenced(bDefrag, BytesToPurge);
 	}
 
-	CompletePurgeRequest(Request); 
+	CompletePurgeRequest(Request);
 }
 
 void FOnDemandContentInstaller::ExecuteDefragRequest(FRequest& Request)
@@ -860,20 +861,12 @@ void FOnDemandContentInstaller::ExecuteDefragRequest(FRequest& Request)
 	UE_LOG(LogIoStoreOnDemand, Log, TEXT("Executing defrag request, BytesToFree=%llu"),
 		BytesToFree != nullptr ? *BytesToFree : -1); 
 
-	if (Request.ErrorCode.load(std::memory_order_relaxed) == EIoErrorCode::Cancelled)
-	{
-		return CompleteDefragRequest(Request);
-	}
-
-	const FIoStatus Status	= IoStore.InstallCache->DefragAll(BytesToFree);
-	Request.ErrorCode		= Status.GetErrorCode();
-
-	if (!Status.IsOk())
+	if (Request.IsCancelled() == false)
 	{
-		Request.ErrorReason = Status.ToString();
+		Request.Result = IoStore.InstallCache->DefragAll(BytesToFree);
 	}
 
-	CompleteDefragRequest(Request); 
+	CompleteDefragRequest(Request);
 }
 
 void FOnDemandContentInstaller::ExecuteVerifyRequest(FRequest& Request)
@@ -883,17 +876,9 @@ void FOnDemandContentInstaller::ExecuteVerifyRequest(FRequest& Request)
 
 	UE_LOG(LogIoStoreOnDemand, Log, TEXT("Executing verify cache request"));
 
-	if (Request.ErrorCode.load(std::memory_order_relaxed) == EIoErrorCode::Cancelled)
-	{
-		return CompleteVerifyRequest(Request);
-	}
-
-	const FIoStatus Status	= IoStore.InstallCache->Verify();
-	Request.ErrorCode		= Status.GetErrorCode();
-
-	if (!Status.IsOk())
+	if (Request.IsCancelled() == false)
 	{
-		Request.ErrorReason = Status.ToString();
+		Request.Result = IoStore.InstallCache->Verify();
 	}
 
 	CompleteVerifyRequest(Request); 
@@ -906,18 +891,9 @@ void FOnDemandContentInstaller::ExectuteFlushLastAccessRequest(FRequest& Request
 
 	UE_LOG(LogIoStoreOnDemand, Log, TEXT("Executing flush last access request"));
 
-	if (Request.ErrorCode.load(std::memory_order_relaxed) == EIoErrorCode::Cancelled)
+	if (Request.IsCancelled() == false)
 	{
-		CompleteFlushLastAccessRequest(Request);
-		return;
-	}
-
-	const FIoStatus Status	= IoStore.InstallCache->FlushLastAccess();
-	Request.ErrorCode		= Status.GetErrorCode();
-
-	if (!Status.IsOk())
-	{
-		Request.ErrorReason = Status.ToString();
+		Request.Result = IoStore.InstallCache->FlushLastAccess();
 	}
 
 	CompleteFlushLastAccessRequest(Request);
@@ -943,7 +919,7 @@ void FOnDemandContentInstaller::ProcessDownloadedChunk(
 		double(InstallRequest.Progress.CurrentInstallSize) / 1024.0 / 1024.0, double(InstallRequest.Progress.TotalInstallSize) / 1024.0 / 1024.0,
 		Request.SeqNo, Request.Priority, *LexToString(InstallRequest.Args.ContentHandle), *LexToString(ChunkId), double(ChunkEntry.GetDiskSize()) / 1024.0, HttpStatusCode);
 
-	if (Request.ErrorCode == EIoErrorCode::Unknown)
+	if (Request.IsOk())
 	{
 		const bool bHttpOk = HttpStatusCode > 199 && HttpStatusCode < 300 && Chunk.GetSize() > 0;
 		if (bHttpOk)
@@ -953,47 +929,41 @@ void FOnDemandContentInstaller::ProcessDownloadedChunk(
 			if (ChunkEntry.Hash == ChunkHash)
 			{
 				check(ChunkEntry.GetDiskSize() == Chunk.GetSize());
-
-				if (FIoStatus Status = IoStore.InstallCache->PutChunk(MoveTemp(Chunk), ChunkHash); !Status.IsOk())
-				{
-					Request.ErrorCode	= Status.GetErrorCode();
-					Request.ErrorReason	= Status.ToString();
-				}
+				Request.Result = IoStore.InstallCache->PutChunk(MoveTemp(Chunk), ChunkHash);
 			}
 			else
 			{
-				Request.ErrorCode	= EIoErrorCode::ReadError;
-				Request.ErrorReason	= FString::Printf(TEXT("Hash mismatch, ChunkId='%s', ExpectedHash='%s', ActualHash='%s'"),
-					*LexToString(ChunkId), *LexToString(ChunkEntry.Hash), *LexToString(ChunkHash));
+				Request.Result = MakeError(UE::UnifiedError::IoStoreOnDemand::ChunkHashError::MakeError(
+					UE::UnifiedError::IoStoreOnDemand::FChunkHashMismatchErrorContext
+					{
+						.ChunkId		= ChunkId,
+						.ExpectedHash	= ChunkEntry.Hash,
+						.ActualHash		= ChunkHash
+					}));
 			}
 		}
 		else
 		{
-			Request.ErrorCode	= EIoErrorCode::ReadError;
-			Request.ErrorReason	= FString::Printf(TEXT("Http failure, StatusCode=%u, Reason=%s"), HttpStatusCode, *ErrorReason);
+			Request.Result = MakeError(UE::UnifiedError::IoStoreOnDemand::HttpError::MakeError(HttpStatusCode));
 		}
 	}
 
 	const bool bCompleted = ++InstallRequest.DownloadedChunkCount >= InstallRequest.HttpRequestHandles.Num();
 	if (bCompleted)
 	{
-		if (FIoStatus Status = IoStore.InstallCache->Flush(); !Status.IsOk())
+		//TODO: Don't flush after every install request
+		FResult FlushResult = IoStore.InstallCache->Flush();
+		if (FlushResult.HasError() && !Request.Result.HasError())
 		{
-			Request.ErrorCode	= Status.GetErrorCode();
-			Request.ErrorReason	= Status.ToString();
+			Request.Result = MoveTemp(FlushResult);
 		}
-		if (Request.ErrorCode == EIoErrorCode::Unknown)
-		{
-			Request.ErrorCode = EIoErrorCode::Ok;
-		}
-
 		CompleteInstallRequest(Request);
 	}
 	else 
 	{
 		NotifyInstallProgress(Request);
 
-		if (Request.ErrorCode != EIoErrorCode::Unknown)
+		if (Request.IsOk() == false)
 		{
 			int32 NumCancelled = 0;
 			if (FHttpIoDispatcher::IsInitialized())
@@ -1078,7 +1048,6 @@ void FOnDemandContentInstaller::CompleteInstallRequest(FRequest& Request)
 	using namespace UE::IoStore::Private;
 
 	FRequest::FInstall& InstallRequest = Request.AsInstall();
-	ensure(Request.ErrorCode != EIoErrorCode::Unknown);
 
 #if !UE_BUILD_SHIPPING
 	TArray<FString> DebugInfo;
@@ -1094,7 +1063,7 @@ void FOnDemandContentInstaller::CompleteInstallRequest(FRequest& Request)
 	);
 
 	// Mark all resolved chunk(s) as referenced by the content handle and notify the package store to update
-	if (Request.ErrorCode == EIoErrorCode::Ok && ResolvedChunkCount > 0)
+	if (Request.IsOk() && ResolvedChunkCount > 0)
 	{
 		{
 			FOnDemandContentHandle& ContentHandle = InstallRequest.Args.ContentHandle;
@@ -1120,32 +1089,34 @@ void FOnDemandContentInstaller::CompleteInstallRequest(FRequest& Request)
 		IoStore.PackageStoreBackend->NeedsUpdate(EOnDemandPackageStoreUpdateMode::ReferencedPackages);
 	}
 
-	const uint64 DurationCycles = FPlatformTime::Cycles64() - Request.StartTimeCycles; 
-	const double CacheHitRatio	= InstallRequest.Progress.TotalContentSize > 0
+	const bool bCancelled			= Request.IsCancelled();
+	const uint64 DurationCycles		= FPlatformTime::Cycles64() - Request.StartTimeCycles;
+	const double DurationInSeconds	= FPlatformTime::ToSeconds64(DurationCycles);
+	const double CacheHitRatio		= InstallRequest.Progress.TotalContentSize > 0
 		? double(InstallRequest.Progress.TotalContentSize - InstallRequest.Progress.TotalInstallSize) / double(InstallRequest.Progress.TotalContentSize)
 		: 0.0;
 
-	FOnDemandInstallResult InstallResult;
-	InstallResult.Status				= Request.ErrorCode == EIoErrorCode::Ok ? FIoStatus::Ok : FIoStatus(Request.ErrorCode, Request.ErrorReason);
-	InstallResult.Progress				= InstallRequest.Progress;
-	InstallResult.DurationInSeconds		= FPlatformTime::ToSeconds64(DurationCycles);
+#if !UE_BUILD_SHIPPING
+	FString DebugInfoString = FString::Join(DebugInfo, TEXT(","));
+	UE_LOG(LogIoStoreOnDemand, Verbose, TEXT("ContentHandle=(%s), ResolvedChunks(%s)"), *LexToString(InstallRequest.Args.ContentHandle), *DebugInfoString);
+#endif
+	UE_LOG(LogIoStoreOnDemand, Verbose, TEXT("Install request completed, Result='%s', SeqNo=%u, Priority=%d, ContentHandle=(%s), ContentSize=%.2lf MiB, InstallSize=%.2lf MiB, CacheHitRatio=%d%%, Duration=%dms"),
+		*Request.GetErrorString(), Request.SeqNo, Request.Priority, *LexToString(InstallRequest.Args.ContentHandle), double(InstallRequest.Progress.TotalContentSize) / 1024.0 / 1024.0,
+		double(InstallRequest.Progress.TotalInstallSize) / 1024.0 / 1024.0, int32(CacheHitRatio * 100), Private::SecondsToMillicseconds(DurationInSeconds));
 
 	FOnDemandContentInstallerStats::OnRequestCompleted(
+		Request.Result,
 		ResolvedChunkCount,
-		InstallResult.Progress.TotalContentSize,
+		InstallRequest.Progress.TotalContentSize,
 		static_cast<uint64>(InstallRequest.HttpRequestHandles.Num()),
-		InstallResult.Progress.CurrentInstallSize,
+		InstallRequest.Progress.CurrentInstallSize,
 		CacheHitRatio,
-		DurationCycles,
-		Request.ErrorCode);
+		DurationCycles);
 
-#if !UE_BUILD_SHIPPING
-	FString DebugInfoString = FString::Join(DebugInfo, TEXT(","));
-	UE_LOG(LogIoStoreOnDemand, Verbose, TEXT("ContentHandle=(%s), ResolvedChunks(%s)"), *LexToString(InstallRequest.Args.ContentHandle), *DebugInfoString);
-#endif
-	UE_LOG(LogIoStoreOnDemand, Verbose, TEXT("Install request completed, Result='%s', SeqNo=%u, Priority=%d, ContentHandle=(%s), ContentSize=%.2lf MiB, InstallSize=%.2lf MiB, CacheHitRatio=%d%%, Duration=%dms"),
-		GetIoErrorText(InstallResult.Status.GetErrorCode()), Request.SeqNo, Request.Priority, *LexToString(InstallRequest.Args.ContentHandle), double(InstallResult.Progress.TotalContentSize) / 1024.0 / 1024.0,
-		double(InstallResult.Progress.TotalInstallSize) / 1024.0 / 1024.0, int32(CacheHitRatio * 100), int32(InstallResult.DurationInSeconds * 1000));
+	FOnDemandInstallResult InstallResult;
+	InstallResult.Progress			= InstallRequest.Progress;
+	InstallResult.DurationInSeconds	= DurationInSeconds;
+	InstallResult.Error				= Request.ConsumeError();
 
 	{
 		UE::TUniqueLock Lock(Mutex);
@@ -1159,19 +1130,11 @@ void FOnDemandContentInstaller::CompleteInstallRequest(FRequest& Request)
 
 	TryExecuteNextRequest();
 
-	FOnDemandInstallRequest::EStatus FinalRequestStatus = FOnDemandInstallRequest::EStatus::None;
-	switch (InstallResult.Status.GetErrorCode())
-	{
-		case EIoErrorCode::Ok:
-			FinalRequestStatus = FOnDemandInstallRequest::Ok;
-			break;
-		case EIoErrorCode::Cancelled:
-			FinalRequestStatus = FOnDemandInstallRequest::Cancelled;
-			break;
-		default:
-			FinalRequestStatus = FOnDemandInstallRequest::Error;
-			break;
-	}
+	const FOnDemandInstallRequest::EStatus FinalRequestStatus = bCancelled
+		? FOnDemandInstallRequest::EStatus::Cancelled
+		: InstallResult.Error.IsSet()
+			? FOnDemandInstallRequest::EStatus::Error
+			: FOnDemandInstallRequest::Ok;
 
 	if (!InstallRequest.OnCompleted)
 	{
@@ -1212,15 +1175,15 @@ void FOnDemandContentInstaller::CompletePurgeRequest(FRequest& Request)
 {
 	check(&Request == CurrentRequest);
 
-	const uint64 DurationCycles = FPlatformTime::Cycles64() - Request.StartTimeCycles; 
+	const double DurationInSeconds = FPlatformTime::ToSeconds64(FPlatformTime::Cycles64() - Request.StartTimeCycles);
 	
-	FOnDemandPurgeResult PurgeResult;
-	PurgeResult.Status				= Request.ErrorCode == EIoErrorCode::Ok ? FIoStatus::Ok : FIoStatus(Request.ErrorCode, Request.ErrorReason);
-	PurgeResult.DurationInSeconds	= FPlatformTime::ToSeconds64(DurationCycles);
-
 	UE_LOG(LogIoStoreOnDemand, Log, TEXT("Purge request completed, Result='%s', Duration=%d ms"),
-		GetIoErrorText(PurgeResult.Status.GetErrorCode()), int32(PurgeResult.DurationInSeconds * 1000));
+		*Request.GetErrorString(), Private::SecondsToMillicseconds(DurationInSeconds));
 	
+	FOnDemandPurgeResult PurgeResult;
+	PurgeResult.DurationInSeconds	= DurationInSeconds;
+	PurgeResult.Error				= Request.ConsumeError();
+
 	const bool bCallbackOnGameThread	= EnumHasAnyFlags(Request.AsPurge().Args.Options, EOnDemandPurgeOptions::CallbackOnGameThread);
 	FOnDemandPurgeCompleted OnCompleted = MoveTemp(Request.AsPurge().OnCompleted);
 	{
@@ -1258,14 +1221,14 @@ void FOnDemandContentInstaller::CompleteDefragRequest(FRequest& Request)
 {
 	check(&Request == CurrentRequest);
 
-	const uint64 DurationCycles = FPlatformTime::Cycles64() - Request.StartTimeCycles; 
-
-	FOnDemandDefragResult DefragResult;
-	DefragResult.Status				= Request.ErrorCode == EIoErrorCode::Ok ? FIoStatus::Ok : FIoStatus(Request.ErrorCode, Request.ErrorReason);
-	DefragResult.DurationInSeconds	= FPlatformTime::ToSeconds64(DurationCycles);
+	const double DurationInSeconds = FPlatformTime::ToSeconds64(FPlatformTime::Cycles64() - Request.StartTimeCycles);
 
 	UE_LOG(LogIoStoreOnDemand, Log, TEXT("Defrag request completed, Result='%s', Duration=%d ms"),
-		GetIoErrorText(DefragResult.Status.GetErrorCode()), int32(DefragResult.DurationInSeconds * 1000));
+		*Request.GetErrorString(), Private::SecondsToMillicseconds(DurationInSeconds));
+
+	FOnDemandDefragResult DefragResult;
+	DefragResult.DurationInSeconds	= DurationInSeconds;
+	DefragResult.Error				= Request.ConsumeError();
 
 	const bool bCallbackOnGameThread		= EnumHasAnyFlags(Request.AsDefrag().Args.Options, EOnDemandDefragOptions::CallbackOnGameThread);
 	FOnDemandDefragCompleted OnCompleted	= MoveTemp(Request.AsDefrag().OnCompleted);
@@ -1302,14 +1265,14 @@ void FOnDemandContentInstaller::CompleteDefragRequest(FRequest& Request)
 
 void FOnDemandContentInstaller::CompleteVerifyRequest(FRequest& Request)
 {
-	const uint64 DurationCycles = FPlatformTime::Cycles64() - Request.StartTimeCycles; 
-
-	FOnDemandVerifyCacheResult VerifyResult;
-	VerifyResult.Status = Request.ErrorCode == EIoErrorCode::Ok ? FIoStatus::Ok : FIoStatus(Request.ErrorCode, Request.ErrorReason);
-	VerifyResult.DurationInSeconds	= FPlatformTime::ToSeconds64(DurationCycles);
+	const double DurationInSeconds = FPlatformTime::ToSeconds64(FPlatformTime::Cycles64() - Request.StartTimeCycles);
 
 	UE_LOG(LogIoStoreOnDemand, Log, TEXT("Verify request completed, Result='%s', Duration=%d ms"),
-		GetIoErrorText(VerifyResult.Status.GetErrorCode()), int32(VerifyResult.DurationInSeconds * 1000));
+		*Request.GetErrorString(), Private::SecondsToMillicseconds(DurationInSeconds));
+
+	FOnDemandVerifyCacheResult VerifyResult;
+	VerifyResult.DurationInSeconds	= DurationInSeconds;
+	VerifyResult.Error				= Request.ConsumeError();
 
 	FOnDemandVerifyCacheCompleted OnCompleted = MoveTemp(Request.AsVerify().OnCompleted);
 	{
@@ -1331,14 +1294,14 @@ void FOnDemandContentInstaller::CompleteVerifyRequest(FRequest& Request)
 
 void FOnDemandContentInstaller::CompleteFlushLastAccessRequest(FRequest& Request)
 {
-	const uint64 DurationCycles = FPlatformTime::Cycles64() - Request.StartTimeCycles;
-
-	FOnDemandFlushLastAccessResult Result;
-	Result.Status = Request.ErrorCode == EIoErrorCode::Ok ? FIoStatus::Ok : FIoStatus(Request.ErrorCode, Request.ErrorReason);
-	Result.DurationInSeconds = FPlatformTime::ToSeconds64(DurationCycles);
+	const double DurationInSeconds = FPlatformTime::ToSeconds64(FPlatformTime::Cycles64() - Request.StartTimeCycles);
 
 	UE_LOG(LogIoStoreOnDemand, Log, TEXT("Flush last access request completed, Result='%s', Duration=%d ms"),
-		GetIoErrorText(Result.Status.GetErrorCode()), int32(Result.DurationInSeconds * 1000));
+		*Request.GetErrorString(), Private::SecondsToMillicseconds(DurationInSeconds));
+
+	FOnDemandFlushLastAccessResult FlushResult;
+	FlushResult.DurationInSeconds	= DurationInSeconds;
+	FlushResult.Error				= Request.ConsumeError();
 
 	FOnDemandFlushLastAccessCompleted OnCompleted = MoveTemp(Request.AsFlushLastAccess().OnCompleted);
 	{
@@ -1355,7 +1318,7 @@ void FOnDemandContentInstaller::CompleteFlushLastAccessRequest(FRequest& Request
 		return;
 	}
 
-	OnCompleted(MoveTemp(Result));
+	OnCompleted(MoveTemp(FlushResult));
 }
 
 void FOnDemandContentInstaller::Shutdown()
@@ -1407,7 +1370,7 @@ void FOnDemandContentInstaller::Shutdown()
 			break;
 		}
 
-		NextRequest->ErrorCode = EIoErrorCode::Cancelled;
+		NextRequest->TryCancel();
 		ExecuteRequest(*NextRequest);
 	}
 }
diff --git a/Engine/Source/Runtime/Experimental/IoStore/OnDemand/Private/OnDemandContentInstaller.h b/Engine/Source/Runtime/Experimental/IoStore/OnDemand/Private/OnDemandContentInstaller.h
--- a/Engine/Source/Runtime/Experimental/IoStore/OnDemand/Private/OnDemandContentInstaller.h
+++ b/Engine/Source/Runtime/Experimental/IoStore/OnDemand/Private/OnDemandContentInstaller.h
@@ -4,7 +4,9 @@
 
 #include "Async/Mutex.h"
 #include "Containers/Set.h"
+#include "Experimental/UnifiedError/CoreErrorTypes.h"
 #include "IO/IoStoreOnDemand.h"
+#include "IO/OnDemandError.h"
 #include "IO/PackageId.h"
 #include "IO/IoStoreOnDemandInternals.h"
 #include "IO/IoAllocators.h"
@@ -154,6 +156,11 @@ class FOnDemandContentInstaller
 			Variant.Emplace<FFlushLastAccess>(MoveTemp(OnCompleted));
 		}
 
+		bool				IsCancelled() const			{ return bCancelled.load(std::memory_order_relaxed); }
+		bool				TryCancel()					{ bool bExpected = false; return bCancelled.compare_exchange_strong(bExpected, true); }
+		bool				IsOk() const				{ return !IsCancelled() && !Result.HasError(); }
+		FString				GetErrorString() const		{ return IsOk() ? TEXT("Ok") : IsCancelled() ? TEXT("Cancelled") : LexToString(Result.GetError()); }
+
 		bool				IsInstall() const			{ return Variant.IsType<FInstall>(); }
 		bool				IsPurge() const				{ return Variant.IsType<FPurge>(); }
 		bool				IsDefrag() const			{ return Variant.IsType<FDefrag>(); }
@@ -166,13 +173,15 @@ class FOnDemandContentInstaller
 		FVerify&			AsVerify()			{ return Variant.Get<FVerify>(); }
 		FFlushLastAccess&	AsFlushLastAccess()	{ return Variant.Get<FFlushLastAccess>(); }
 
+		TOptional<UE::UnifiedError::FError> ConsumeError();
+
 		static uint32	NextSeqNo;
 
 		uint32						SeqNo = NextSeqNo++;
 		int32						Priority = 0;
 		uint64						StartTimeCycles = FPlatformTime::Cycles64();
-		FString						ErrorReason;
-		std::atomic<EIoErrorCode>	ErrorCode{EIoErrorCode::Unknown};
+		FResult						Result = MakeValue();
+		std::atomic_bool			bCancelled{false};
 		FRequestVariant				Variant;
 	};
 
diff --git a/Engine/Source/Runtime/Experimental/IoStore/OnDemand/Private/OnDemandInstallCache.cpp b/Engine/Source/Runtime/Experimental/IoStore/OnDemand/Private/OnDemandInstallCache.cpp
--- a/Engine/Source/Runtime/Experimental/IoStore/OnDemand/Private/OnDemandInstallCache.cpp
+++ b/Engine/Source/Runtime/Experimental/IoStore/OnDemand/Private/OnDemandInstallCache.cpp
@@ -20,6 +20,7 @@
 #include "IO/IoChunkId.h"
 #include "IO/IoChunkEncoding.h"
 #include "Logging/StructuredLog.h"
+#include "IO/OnDemandError.h"
 #include "Misc/DateTime.h"
 #include "Misc/Paths.h"
 #include "Misc/PathViews.h"
@@ -55,9 +56,6 @@
 
 namespace UE::IoStore
 {
-///////////////////////////////////////////////////////////////////////////////
-FNamespacedIoStatusBuilderBuilder InstallCacheStatusBuilder(TEXT("IADCache"));
-
 ///////////////////////////////////////////////////////////////////////////////
 namespace CVars
 {
@@ -226,19 +224,21 @@ struct FCas
 	FCas::FLastAccess			GetAndClearDirtyLastAccess();
 
 	// Returns the place holder timestamp to use if a timestamp is not found in the LastAccess table
-	static int64			GetTimestampForMissingLastAccess();
-	static uint32			GetMaxBlockSize()
+	static int64				GetTimestampForMissingLastAccess();
+	static uint32				GetMaxBlockSize()
 	{
 		const uint32 MaxBlockSize = 32 << 20; //TODO: Make configurable	
 		return MaxBlockSize;
 	}
 
-	static uint32			GetMinBlockSize()
+	static uint32				GetMinBlockSize()
 	{
 		const uint32 MinBlockSize = 32 << 19; //TODO: Make configurable	
 		return MinBlockSize;
 	}
 
+	FString						GetRootDirectory() const { return FString(RootDirectory); }
+
 private:
 	FStringView				RootDirectory;
 public: // TODO: FIXME: encapsulate these members
@@ -279,7 +279,7 @@ FIoStatus FCas::Initialize(FStringView Directory, bool bDeleteExisting)
 
 		if (Ifm.DeleteDirectory(Path.ToString(), bRequireExists, bTree) == false)
 		{
-			return InstallCacheStatusBuilder(EIoErrorCode::WriteError)
+			return FIoStatusBuilder(EIoErrorCode::DeleteError, FPlatformMisc::GetLastError())
 				<< TEXT("Failed to delete CAS blocks directory '")
 				<< Path.ToString()
 				<< TEXT("'");
@@ -291,15 +291,14 @@ FIoStatus FCas::Initialize(FStringView Directory, bool bDeleteExisting)
 		const bool bTree = true;
 		if (Ifm.MakeDirectory(Path.ToString(), bTree) == false)
 		{
-			FIoStatus Status = InstallCacheStatusBuilder(EIoErrorCode::WriteError)
+			return FIoStatusBuilder(EIoErrorCode::WriteError, FPlatformMisc::GetLastError())
 				<< TEXT("Failed to create directory '")
 				<< Path.ToString()
 				<< TEXT("'");
-			return Status;
 		}
 	}
 
-	return InstallCacheStatusBuilder(EIoErrorCode::Ok);
+	return FIoStatus::Ok;
 };
 
 FCasLocation FCas::FindChunk(const FIoHash& Hash) const
@@ -368,7 +367,7 @@ FIoStatus FCas::DeleteBlock(FCasBlockId BlockId, TArray<FCasAddr>& OutAddrs)
 
 			if (FPlatformTime::ToMilliseconds(FPlatformTime::Cycles() - StartTimeCycles) > DeleteBlockMaxWaitTimeMs)
 			{
-				return InstallCacheStatusBuilder(EIoErrorCode::Timeout)
+				return FIoStatusBuilder(EIoErrorCode::Timeout)
 					<< TEXT("Timed out waiting for pending read(s) while deleting CAS block"); 
 			}
 
@@ -384,7 +383,7 @@ FIoStatus FCas::DeleteBlock(FCasBlockId BlockId, TArray<FCasAddr>& OutAddrs)
 	UE_LOG(LogIoStoreOnDemand, Log, TEXT("Deleting CAS block '%s'"), *Filename);
 	if (Ipf.DeleteFile(*Filename) == false)
 	{
-		return InstallCacheStatusBuilder(EIoErrorCode::DeleteError) << TEXT("Failed to delete CAS block");
+		return FIoStatusBuilder(EIoErrorCode::DeleteError, FPlatformMisc::GetLastError()) << TEXT("Failed to delete CAS block");
 	}
 
 	BlockIds.Remove(BlockId);
@@ -399,7 +398,7 @@ FIoStatus FCas::DeleteBlock(FCasBlockId BlockId, TArray<FCasAddr>& OutAddrs)
 		}
 	}
 
-	return InstallCacheStatusBuilder(EIoErrorCode::Ok);
+	return FIoStatus::Ok;
 }
 
 FString FCas::GetBlockFilename(FCasBlockId BlockId) const
@@ -643,7 +642,7 @@ FIoStatus FCas::Verify(TArray<FCasAddr>& OutAddrs)
 	FCasBlockInfoMap	BlockInfo;
 	const uint64		TotalSize = GetBlockInfo(BlockInfo);
 	uint64				TotalVerifiedBytes = 0;
-	FIoStatus			Status = InstallCacheStatusBuilder(EIoErrorCode::Ok);
+	FIoStatus			Status = FIoStatus::Ok;
 
 	const int64 Now = FDateTime::UtcNow().GetTicks();
 	const int64 TimestampForMissingLastAccess = FCas::GetTimestampForMissingLastAccess();
@@ -699,11 +698,12 @@ FIoStatus FCas::Verify(TArray<FCasAddr>& OutAddrs)
 		}
 
 		const FString Filename = GetBlockFilename(BlockId);
-		UE_LOG(LogIoStoreOnDemand, Warning, TEXT("Missing CAS block '%s'"), *Filename);
+		const FString ErrroMessage = FString::Printf(TEXT("Missing CAS block '%s'"), *Filename);
+		UE_LOG(LogIoStoreOnDemand, Warning, TEXT("%s"), *ErrroMessage);
 
 		LastAccess.Remove(BlockId);
 		BlockIt.RemoveCurrent();
-		Status = InstallCacheStatusBuilder(EIoErrorCode::NotFound);
+		Status = FIoStatus(EIoErrorCode::NotFound, ErrroMessage);
 	}
 
 	UE_LOG(LogIoStoreOnDemand, Log, TEXT("Verified %d CAS blocks of total %.2lf MiB"),
@@ -734,7 +734,7 @@ FIoStatus FCas::Verify(TArray<FCasAddr>& OutAddrs)
 			OutAddrs.Add(It->Key);
 			It.RemoveCurrent();
 
-			Status = InstallCacheStatusBuilder(EIoErrorCode::NotFound);
+			Status = FIoStatus(EIoErrorCode::NotFound);
 		}
 	}
 
@@ -797,7 +797,8 @@ struct FCasJournal
 	{
 		static const inline uint8 MagicSequence[16] = {'C', 'A', 'S', 'J', 'O', 'U', 'R', 'N', 'A', 'L', 'H', 'E', 'A', 'D', 'E', 'R'};
 
-		bool		IsValid() const;
+		bool			IsValid() const;
+		static int64	Size() { return sizeof(FHeader); }
 
 		uint8		Magic[16] = {0};
 		EVersion	Version = EVersion::Invalid;
@@ -809,7 +810,8 @@ struct FCasJournal
 	{
 		static const inline uint8 MagicSequence[16] = {'C', 'A', 'S', 'J', 'O', 'U', 'R', 'N', 'A', 'L', 'F', 'O', 'O', 'T', 'E', 'R'};
 
-		bool IsValid() const;
+		bool			IsValid() const;
+		static int64	Size() { return sizeof(FFooter); }
 
 		uint8 Magic[16] = {0};
 	};
@@ -860,7 +862,8 @@ struct FCasJournal
 			FCriticalError	CriticalError;
 		};
 
-		EType Type() const { return *reinterpret_cast<const EType*>(this); }
+		EType			Type() const { return *reinterpret_cast<const EType*>(this); }
+		static int64	Size() { return sizeof(FEntry); }
 	};
 	static_assert(sizeof(FEntry) == 24);
 
@@ -943,54 +946,54 @@ FIoStatus FCasJournal::Replay(const FString& JournalFile, FEntryHandler&& Handle
 
 	if (Ipf.FileExists(*JournalFile) == false)
 	{
-		return InstallCacheStatusBuilder(EIoErrorCode::NotFound);
+		return FIoStatus(EIoErrorCode::NotFound);
 	}
 
 	TUniquePtr<IFileHandle> FileHandle(Ipf.OpenRead(*JournalFile));
 	if (FileHandle.IsValid() == false)
 	{
-		return InstallCacheStatusBuilder(EIoErrorCode::FileNotOpen);
+		return FIoStatus(EIoErrorCode::FileOpenFailed, FPlatformMisc::GetLastError());
 	}
 
 	FHeader Header;
-	if ((FileHandle->Read(reinterpret_cast<uint8*>(&Header), sizeof(FHeader)) == false) || (Header.IsValid() == false))
+	if ((FileHandle->Read(reinterpret_cast<uint8*>(&Header), FHeader::Size()) == false) || (Header.IsValid() == false))
 	{
-		return InstallCacheStatusBuilder(EIoErrorCode::ReadError)
+		return FIoStatusBuilder(EIoErrorCode::ReadError, FPlatformMisc::GetLastError())
 			<< TEXT("Failed to validate journal header '")
 			<< JournalFile
 			<< TEXT("'");
 	}
 
 	const int64 FileSize	= FileHandle->Size();
-	const int64 EntryCount	= (FileSize - sizeof(FHeader) - sizeof(FFooter)) / sizeof(FEntry);
+	const int64 EntryCount	= (FileSize - FHeader::Size() - FFooter::Size()) / FEntry::Size();
 
 	if (EntryCount < 0)
 	{
-		return InstallCacheStatusBuilder(EIoErrorCode::ReadError);
+		return FIoStatus(EIoErrorCode::ReadError);
 	}
 
 	if (EntryCount == 0)
 	{
-		return InstallCacheStatusBuilder(EIoErrorCode::Ok);
+		return FIoStatus::Ok;
 	}
 
-	const int64 FooterPos = FileSize - sizeof(FFooter);
+	const int64 FooterPos = FileSize - FFooter::Size();
 	if (FooterPos < 0)
 	{
-		return InstallCacheStatusBuilder(EIoErrorCode::CorruptToc)
+		return FIoStatusBuilder(EIoErrorCode::FileCorrupt)
 			<< TEXT("Invalid journal footer");
 	}
 
 	const int64 EntriesPos = FileHandle->Tell();
 	if (FileHandle->Seek(FooterPos) == false)
 	{
-		return InstallCacheStatusBuilder(EIoErrorCode::ReadError);
+		return FIoStatus(EIoErrorCode::FileSeekFailed, FPlatformMisc::GetLastError());
 	}
 
 	FFooter Footer;
-	if ((FileHandle->Read(reinterpret_cast<uint8*>(&Footer), sizeof(FFooter)) == false) || (Footer.IsValid() == false))
+	if ((FileHandle->Read(reinterpret_cast<uint8*>(&Footer), FFooter::Size()) == false) || (Footer.IsValid() == false))
 	{
-		return InstallCacheStatusBuilder(EIoErrorCode::ReadError)
+		return FIoStatusBuilder(EIoErrorCode::ReadError, FPlatformMisc::GetLastError())
 			<< TEXT("Failed to validate journal footer '")
 			<< JournalFile
 			<< TEXT("'");
@@ -998,19 +1001,19 @@ FIoStatus FCasJournal::Replay(const FString& JournalFile, FEntryHandler&& Handle
 
 	if (FileHandle->Seek(EntriesPos) == false)
 	{
-		return InstallCacheStatusBuilder(EIoErrorCode::ReadError);
+		return FIoStatus(EIoErrorCode::FileSeekFailed, FPlatformMisc::GetLastError());
 	}
 
 	TArray<FEntry> Entries;
 	Entries.SetNumZeroed(IntCastChecked<int32>(EntryCount));
 
-	if (FileHandle->Read(reinterpret_cast<uint8*>(Entries.GetData()), sizeof(FEntry) * EntryCount) == false)
+	if (FileHandle->Read(reinterpret_cast<uint8*>(Entries.GetData()), FEntry::Size() * EntryCount) == false)
 	{
-		return InstallCacheStatusBuilder(EIoErrorCode::ReadError);
+		return FIoStatus(EIoErrorCode::ReadError, FPlatformMisc::GetLastError());
 	}
 
 	UE_LOG(LogIoStoreOnDemand, Log, TEXT("Replaying %" INT64_FMT " CAS journal entries of total %.2lf KiB from '%s'"),
-		EntryCount, ToKiB(sizeof(FEntry) * EntryCount), *JournalFile);
+		EntryCount, ToKiB(FEntry::Size() * EntryCount), *JournalFile);
 
 	for (const FEntry& Entry : Entries)
 	{
@@ -1020,13 +1023,13 @@ FIoStatus FCasJournal::Replay(const FString& JournalFile, FEntryHandler&& Handle
 			UE_LOG(LogIoStoreOnDemand, Warning, TEXT("Found critical error entry '%s' (%d) in journal '%s'"),
 				GetErrorText(Error.ErrorCode), Error.ErrorCode, *JournalFile);
 
-			return InstallCacheStatusBuilder(EIoErrorCode::ReadError) << GetErrorText(Error.ErrorCode);
+			return FIoStatus(EIoErrorCode::InvalidCode);
 		}
 
 		Handler(Entry);
 	}
 
-	return InstallCacheStatusBuilder(EIoErrorCode::Ok);
+	return FIoStatus::Ok;
 }
 
 FIoStatus FCasJournal::Create(const FString& JournalFile)
@@ -1037,26 +1040,26 @@ FIoStatus FCasJournal::Create(const FString& JournalFile)
 	TUniquePtr<IFileHandle> FileHandle(Ipf.OpenWrite(*JournalFile));
 	if (FileHandle.IsValid() == false)
 	{
-		return InstallCacheStatusBuilder(EIoErrorCode::FileNotOpen);
+		return FIoStatus(EIoErrorCode::FileOpenFailed, FPlatformMisc::GetLastError());
 	}
 
 	FHeader Header;
 	FMemory::Memcpy(&Header.Magic, &FHeader::MagicSequence, sizeof(FHeader::MagicSequence));
 	Header.Version = EVersion::Latest;
 
-	if (FileHandle->Write(reinterpret_cast<uint8*>(&Header), sizeof(FHeader)) == false)
+	if (FileHandle->Write(reinterpret_cast<uint8*>(&Header), FHeader::Size()) == false)
 	{
-		return InstallCacheStatusBuilder(EIoErrorCode::WriteError);
+		return FIoStatus(EIoErrorCode::WriteError, FPlatformMisc::GetLastError());
 	}
 
 	FFooter Footer;
 	FMemory::Memcpy(&Footer.Magic, &FFooter::MagicSequence, sizeof(FFooter::MagicSequence));
-	if (FileHandle->Write(reinterpret_cast<uint8*>(&Footer), sizeof(FFooter)) == false)
+	if (FileHandle->Write(reinterpret_cast<uint8*>(&Footer), FFooter::Size()) == false)
 	{
-		return InstallCacheStatusBuilder(EIoErrorCode::WriteError);
+		return FIoStatus(EIoErrorCode::WriteError, FPlatformMisc::GetLastError());
 	}
 
-	return InstallCacheStatusBuilder(EIoErrorCode::Ok);
+	return FIoStatus::Ok;
 }
 
 FCasJournal::FTransaction FCasJournal::Begin(FString&& JournalFile)
@@ -1071,44 +1074,54 @@ FIoStatus FCasJournal::Commit(FTransaction&& Transaction)
 {
 	if (Transaction.Entries.IsEmpty())
 	{
-		return InstallCacheStatusBuilder(EIoErrorCode::Ok);
+		return FIoStatus::Ok;
 	}
 
-	IPlatformFile&	Ipf = FPlatformFileManager::Get().GetPlatformFile();
+	IPlatformFile& Ipf = FPlatformFileManager::Get().GetPlatformFile();
 
 	// Validate header and footer
 	{
 		TUniquePtr<IFileHandle> FileHandle(Ipf.OpenRead(*Transaction.JournalFile));
-		const int64				FileSize = FileHandle.IsValid() ? FileHandle->Size() : -1;
+		if (FileHandle.IsValid() == false)
+		{
+			const uint32 LastError = FPlatformMisc::GetLastError();
+			FOnDemandInstallCacheStats::OnJournalCommit(EIoErrorCode::FileOpenFailed, 0);
+			return FIoStatusBuilder(EIoErrorCode::FileOpenFailed, LastError)
+				<< TEXT("Failed to validate CAS journal file");
+		}
 
-		if (FileSize < sizeof(FHeader))
+		const int64 FileSize = FileHandle->Size();
+		if (FileSize < FHeader::Size())
 		{
-			FOnDemandInstallCacheStats::OnJournalCommit(EIoErrorCode::ReadError, 0);
-			return InstallCacheStatusBuilder(EIoErrorCode::FileOpenFailed)
+			FOnDemandInstallCacheStats::OnJournalCommit(EIoErrorCode::FileCorrupt, 0);
+			return FIoStatusBuilder(EIoErrorCode::FileCorrupt, FPlatformMisc::GetLastError())
 				<< TEXT("Failed to validate CAS journal file");
 		}
 
 		FHeader Header;
-		if ((FileHandle->Read(reinterpret_cast<uint8*>(&Header), sizeof(FHeader)) == false) || (Header.IsValid() == false))
+		if ((FileHandle->Read(reinterpret_cast<uint8*>(&Header), FHeader::Size()) == false) || (Header.IsValid() == false))
 		{
-			FOnDemandInstallCacheStats::OnJournalCommit(EIoErrorCode::SignatureError, 0);
-			return InstallCacheStatusBuilder(EIoErrorCode::ReadError)
+			const uint32 LastError = FPlatformMisc::GetLastError();
+			FOnDemandInstallCacheStats::OnJournalCommit(EIoErrorCode::ReadError, 0);
+			return FIoStatusBuilder(EIoErrorCode::ReadError, LastError)
 				<< TEXT("Failed to validate CAS journal header");
 		}
 
-		const int64 FooterPos = FileSize - sizeof(FFooter);
+		const int64 FooterPos = FileSize - FFooter::Size();
 		if (FileHandle->Seek(FooterPos) == false)
 		{
-			FOnDemandInstallCacheStats::OnJournalCommit(EIoErrorCode::SignatureError, 0);
-			return InstallCacheStatusBuilder(EIoErrorCode::ReadError)
+			const uint32 LastError = FPlatformMisc::GetLastError();
+			FOnDemandInstallCacheStats::OnJournalCommit(EIoErrorCode::FileSeekFailed, 0);
+			return FIoStatusBuilder(EIoErrorCode::FileSeekFailed, LastError)
 				<< TEXT("Failed to validate CAS journal footer");
 		}
 
 		FFooter Footer;
-		if ((FileHandle->Read(reinterpret_cast<uint8*>(&Footer), sizeof(FFooter)) == false) || (Footer.IsValid() == false))
+		if ((FileHandle->Read(reinterpret_cast<uint8*>(&Footer), FFooter::Size()) == false) || (Footer.IsValid() == false))
 		{
-			FOnDemandInstallCacheStats::OnJournalCommit(EIoErrorCode::SignatureError, 0);
-			return InstallCacheStatusBuilder(EIoErrorCode::ReadError)
+			const uint32 LastError = FPlatformMisc::GetLastError();
+			FOnDemandInstallCacheStats::OnJournalCommit(EIoErrorCode::ReadError, 0);
+			return FIoStatusBuilder(EIoErrorCode::ReadError, LastError)
 				<< TEXT("Failed to validate CAS journal footer");
 		}
 	}
@@ -1118,38 +1131,50 @@ FIoStatus FCasJournal::Commit(FTransaction&& Transaction)
 		const bool				bAppend = true;
 		TUniquePtr<IFileHandle> FileHandle(Ipf.OpenWrite(*Transaction.JournalFile, bAppend));
 		const int64				FileSize	= FileHandle.IsValid() ? FileHandle->Size() : -1;
-		const int64				EntriesPos	= FileSize > 0 ? FileSize - sizeof(FFooter) : -1;
+		const int64				EntriesPos	= FileSize > 0 ? FileSize - FFooter::Size() : -1;
 
-		if ((EntriesPos < 0) || (FileHandle->Seek(EntriesPos) == false))
+		if (EntriesPos < 0) 
 		{
-			FOnDemandInstallCacheStats::OnJournalCommit(EIoErrorCode::WriteError, 0);
-			return InstallCacheStatusBuilder(EIoErrorCode::FileOpenFailed)
+			const uint32 LastError = FPlatformMisc::GetLastError();
+			FOnDemandInstallCacheStats::OnJournalCommit(EIoErrorCode::FileOpenFailed, 0);
+			return FIoStatusBuilder(EIoErrorCode::FileOpenFailed, LastError)
 				<< TEXT("Failed to open CAS journal");
 		}
 
-		const int64 TotalEntrySize = Transaction.Entries.Num() * sizeof(FEntry);
+		if (FileHandle->Seek(EntriesPos) == false)
+		{
+			const uint32 LastError = FPlatformMisc::GetLastError();
+			FOnDemandInstallCacheStats::OnJournalCommit(EIoErrorCode::FileSeekFailed, 0);
+			return FIoStatusBuilder(EIoErrorCode::FileSeekFailed, LastError)
+				<< TEXT("Failed to seek to entries in CAS journal");
+		}
+
+		const int64 TotalEntrySize = Transaction.Entries.Num() * FEntry::Size();
 		if (FileHandle->Write(
 			reinterpret_cast<const uint8*>(Transaction.Entries.GetData()),
 			TotalEntrySize) == false)
 		{
+			const uint32 LastError = FPlatformMisc::GetLastError();
 			FOnDemandInstallCacheStats::OnJournalCommit(EIoErrorCode::WriteError, 0);
-			return InstallCacheStatusBuilder(EIoErrorCode::WriteError)
+			return FIoStatusBuilder(EIoErrorCode::WriteError, LastError)
 				<< TEXT("Failed to write CAS journal entries");
 		}
 
 		FFooter Footer;
 		FMemory::Memcpy(&Footer.Magic, &FFooter::MagicSequence, sizeof(FFooter::MagicSequence));
-		if (FileHandle->Write(reinterpret_cast<uint8*>(&Footer), sizeof(FFooter)) == false)
+		if (FileHandle->Write(reinterpret_cast<uint8*>(&Footer), FFooter::Size()) == false)
 		{
+			const uint32 LastError = FPlatformMisc::GetLastError();
 			FOnDemandInstallCacheStats::OnJournalCommit(EIoErrorCode::WriteError, 0);
-			return InstallCacheStatusBuilder(EIoErrorCode::WriteError)
+			return FIoStatusBuilder(EIoErrorCode::WriteError, LastError)
 				<< TEXT("Failed to write CAS journal footer");
 		}
 
 		if (FileHandle->Flush() == false)
 		{
-			FOnDemandInstallCacheStats::OnJournalCommit(EIoErrorCode::WriteError, 0);
-			return InstallCacheStatusBuilder(EIoErrorCode::WriteError)
+			const uint32 LastError = FPlatformMisc::GetLastError();
+			FOnDemandInstallCacheStats::OnJournalCommit(EIoErrorCode::FileFlushFailed, 0);
+			return FIoStatusBuilder(EIoErrorCode::FileFlushFailed, FPlatformMisc::GetLastError())
 				<< TEXT("Failed to flush CAS journal");
 		}
 
@@ -1158,7 +1183,7 @@ FIoStatus FCasJournal::Commit(FTransaction&& Transaction)
 
 		FOnDemandInstallCacheStats::OnJournalCommit(EIoErrorCode::Ok, TotalEntrySize);
 
-		return InstallCacheStatusBuilder(EIoErrorCode::Ok);
+		return FIoStatus::Ok;
 	}
 }
 
@@ -1226,7 +1251,8 @@ struct FCasSnapshot
 	{
 		static const inline uint8 MagicSequence[16] = {'+', 'S', 'N', 'A', 'P', 'S', 'H', 'O', 'T', 'H', 'E', 'A', 'D', 'E', 'R', '+'};
 
-		bool		IsValid() const;
+		bool			IsValid() const;
+		static int64	Size() { return sizeof(FHeader); }
 
 		uint8		Magic[16] = {0};
 		EVersion	Version = EVersion::Invalid;
@@ -1238,7 +1264,8 @@ struct FCasSnapshot
 	{
 		static const inline uint8 MagicSequence[16] = {'+', 'S', 'N', 'A', 'P', 'S', 'H', 'O', 'T', 'F', 'O', 'O', 'T', 'E', 'R', '+'};
 
-		bool IsValid() const;
+		static int64	Size() { return sizeof(FFooter); }
+		bool			IsValid() const;
 
 		uint8 Magic[16] = {0};
 	};
@@ -1377,19 +1404,20 @@ TIoStatusOr<int64> FCasSnapshot::Save(const FCasSnapshot& Snapshot, const FStrin
 	TUniquePtr<FArchive> Ar(Ifm.CreateFileWriter(*TmpSnapshotFile));
 	if (Ar.IsValid() == false)
 	{
-		return InstallCacheStatusBuilder(EIoErrorCode::FileNotOpen);
+		return FIoStatus(EIoErrorCode::FileOpenFailed, FPlatformMisc::GetLastError());
 	}
 
 	FHeader Header;
 	FMemory::Memcpy(&Header.Magic, &FHeader::MagicSequence, sizeof(FHeader::MagicSequence));
 	Header.Version = EVersion::Latest;
 
-	Ar->Serialize(reinterpret_cast<uint8*>(&Header), sizeof(FHeader));
+	Ar->Serialize(reinterpret_cast<uint8*>(&Header), FHeader::Size());
 	if (Ar->IsError())
 	{
+		const uint32 LastError = FPlatformMisc::GetLastError();
 		Ar.Reset();
 		Ifm.Delete(*TmpSnapshotFile);
-		return InstallCacheStatusBuilder(EIoErrorCode::WriteError);
+		return FIoStatus(EIoErrorCode::WriteError, LastError);
 	}
 
 	FCasSnapshot& NonConst = *const_cast<FCasSnapshot*>(&Snapshot);
@@ -1399,33 +1427,37 @@ TIoStatusOr<int64> FCasSnapshot::Save(const FCasSnapshot& Snapshot, const FStrin
 
 	if (Ar->IsError())
 	{
+		const uint32 LastError = FPlatformMisc::GetLastError();
 		Ar.Reset();
 		Ifm.Delete(*TmpSnapshotFile);
-		return InstallCacheStatusBuilder(EIoErrorCode::WriteError);
+		return FIoStatus(EIoErrorCode::WriteError, LastError);
 	}
 
 	FFooter Footer;
 	FMemory::Memcpy(&Footer.Magic, &FFooter::MagicSequence, sizeof(FFooter::MagicSequence));
-	Ar->Serialize(reinterpret_cast<uint8*>(&Footer), sizeof(FFooter));
+	Ar->Serialize(reinterpret_cast<uint8*>(&Footer), FFooter::Size());
 	if (Ar->IsError())
 	{
+		const uint32 LastError = FPlatformMisc::GetLastError();
 		Ar.Reset();
 		Ifm.Delete(*TmpSnapshotFile);
-		return InstallCacheStatusBuilder(EIoErrorCode::WriteError);
+		return FIoStatus(EIoErrorCode::WriteError, LastError);
 	}
 
 	const int64 FileSize = Ar->TotalSize();
 	if (Ar->Close() == false)
 	{
+		const uint32 LastError = FPlatformMisc::GetLastError();
 		Ar.Reset();
 		Ifm.Delete(*TmpSnapshotFile);
-		return InstallCacheStatusBuilder(EIoErrorCode::WriteError);
+		return FIoStatus(EIoErrorCode::FileCloseFailed, LastError);
 	}
 
 	if (Ifm.Move(*SnapshotFile, *TmpSnapshotFile) == false)
 	{
+		const uint32 LastError = FPlatformMisc::GetLastError();
 		Ifm.Delete(*TmpSnapshotFile);
-		return InstallCacheStatusBuilder(EIoErrorCode::WriteError);
+		return FIoStatus(EIoErrorCode::FileMoveFailed, LastError);
 	}
 
 	return FileSize; 
@@ -1438,18 +1470,17 @@ TIoStatusOr<FCasSnapshot> FCasSnapshot::Load(const FString& SnapshotFile, int64*
 	TUniquePtr<FArchive> Ar(Ifm.CreateFileReader(*SnapshotFile));
 	if (Ar.IsValid() == false)
 	{
-		return InstallCacheStatusBuilder(EIoErrorCode::NotFound);
+		return FIoStatus(EIoErrorCode::NotFound, FPlatformMisc::GetLastError());
 	}
 
 	FHeader Header;
-	Ar->Serialize(reinterpret_cast<uint8*>(&Header), sizeof(FHeader));
+	Ar->Serialize(reinterpret_cast<uint8*>(&Header), FHeader::Size());
 	if (Ar->IsError() || Header.IsValid() == false)
 	{
-		FIoStatus Status = InstallCacheStatusBuilder(EIoErrorCode::ReadError)
+		return FIoStatusBuilder(EIoErrorCode::ReadError, FPlatformMisc::GetLastError())
 			<< TEXT("Failed to validate snapshot header '")
 			<< SnapshotFile 
 			<< TEXT("'");
-		return Status;
 	}
 
 	FCasSnapshot Snapshot;
@@ -1458,14 +1489,13 @@ TIoStatusOr<FCasSnapshot> FCasSnapshot::Load(const FString& SnapshotFile, int64*
 	*Ar << Snapshot.CurrentBlockId;
 
 	FFooter Footer;
-	Ar->Serialize(reinterpret_cast<uint8*>(&Footer), sizeof(FFooter));
+	Ar->Serialize(reinterpret_cast<uint8*>(&Footer), FFooter::Size());
 	if (Ar->IsError() || Footer.IsValid() == false)
 	{
-		FIoStatus Status = InstallCacheStatusBuilder(EIoErrorCode::ReadError)
+		return FIoStatusBuilder(EIoErrorCode::ReadError, FPlatformMisc::GetLastError())
 			<< TEXT("Failed to validate snapshot footer '")
 			<< SnapshotFile 
 			<< TEXT("'");
-		return Status;
 	}
 
 	if (OutFileSize != nullptr)
@@ -1482,7 +1512,7 @@ TIoStatusOr<int64> FCasSnapshot::TryCreateAndResetJournal(const FString& Snapsho
 	const int64 JournalFileSize = Ifm.FileSize(*JournalFile);
 	if (JournalFileSize < 0)
 	{
-		return InstallCacheStatusBuilder(EIoErrorCode::NotFound);
+		return FIoStatus(EIoErrorCode::NotFound);
 	}
 
 	// Load the snapshot from the journal
@@ -1518,12 +1548,13 @@ TIoStatusOr<int64> FCasSnapshot::TryCreateAndResetJournal(const FString& Snapsho
 
 	if (Ifm.Move(*JournalFile , *TmpJournalFile) == false)
 	{
+		const uint32 LastError = FPlatformMisc::GetLastError();
 		if (Ifm.Delete(*SnapshotFile) == false)
 		{
 			UE_LOG(LogIoStoreOnDemand, Error, TEXT("Failed to delete CAS snapshot '%s'"), *SnapshotFile); 
 		}
 
-		return InstallCacheStatusBuilder(EIoErrorCode::WriteError);
+		return FIoStatus(EIoErrorCode::FileMoveFailed, LastError);
 	}
 
 	return SnapshotSize;
@@ -1667,13 +1698,13 @@ public:
 													TConstArrayView<int32> EntryIndices,
 													FOnDemandContentHandle ContentHandle,
 													TArray<int32>& OutMissing) override;
-	virtual FIoStatus							PutChunk(FIoBuffer&& Chunk, const FIoHash& ChunkHash) override;
-	virtual FIoStatus							Purge(TSet<FIoHash>&& ChunksToInstall) override;
-	virtual FIoStatus							PurgeAllUnreferenced(bool bDefrag, const uint64* BytesToPurge = nullptr) override;
-	virtual FIoStatus							DefragAll(const uint64* BytesToFree = nullptr) override;
-	virtual FIoStatus							Verify() override;
-	virtual FIoStatus							Flush() override;
-	virtual FIoStatus							FlushLastAccess() override;
+	virtual FResult								PutChunk(FIoBuffer&& Chunk, const FIoHash& ChunkHash) override;
+	virtual FResult								Purge(TSet<FIoHash>&& ChunksToInstall) override;
+	virtual FResult								PurgeAllUnreferenced(bool bDefrag, const uint64* BytesToPurge = nullptr) override;
+	virtual FResult								DefragAll(const uint64* BytesToFree = nullptr) override;
+	virtual FResult								Verify() override;
+	virtual FResult								Flush() override;
+	virtual FResult								FlushLastAccess() override;
 	virtual void								UpdateLastAccess(TConstArrayView<FIoHash> ChunkHashes) override;
 	virtual FOnDemandInstallCacheUsage			GetCacheUsage() override;
 
@@ -1687,8 +1718,8 @@ private:
 													const TSet<FIoHash>& ChunksToInstall,
 													FCasBlockInfoMap& BlockInfoMap,
 													uint64& OutTotalReferencedBytes) const;
-	FIoStatus									Purge(FCasBlockInfoMap& BlockInfo, uint64 TotalBytesToPurge, uint64& OutTotalPurgedBytes);
-	FIoStatus									Defrag(
+	FResult										Purge(FCasBlockInfoMap& BlockInfo, uint64 TotalBytesToPurge, uint64& OutTotalPurgedBytes);
+	FResult										Defrag(
 													const TArray<FSharedOnDemandContainer>& Containers,
 													const TArray<TBitArray<>>& ChunkEntryIndices,
 													FCasBlockInfoMap& BlockInfo, 
@@ -1696,10 +1727,11 @@ private:
 													uint64* OutTotalFreedBytes = nullptr);
 	bool										Resolve(FIoRequestImpl* Request);
 	void										CompleteRequest(FIoRequestImpl* Request, EIoErrorCode Status);
-	FIoStatus									FlushPendingChunks(FPendingChunks& Chunks, int64 UtcAccessTicks = 0);
-	FIoStatus									FlushPendingChunksImpl(FPendingChunks& Chunks, int64 UtcAccessTicks = 0);
+	FResult										FlushPendingChunks(FPendingChunks& Chunks, int64 UtcAccessTicks = 0);
+	FResult										FlushPendingChunksImpl(FPendingChunks& Chunks, int64 UtcAccessTicks = 0);
 	FString										GetJournalFilename() const { return CacheDirectory / TEXT("cas.jrn"); }
 	FString										GetSnapshotFilename() const { return CacheDirectory / TEXT("cas.snp"); }
+	void										CaptureCacheState(UE::UnifiedError::IoStoreOnDemand::FInstallCacheErrorContext& OutCtx);
 
 	FOnDemandIoStore&			IoStore;
 	FString						CacheDirectory;
@@ -1977,12 +2009,12 @@ TIoStatusOr<uint64> FOnDemandInstallCache::GetSizeForChunk(const FIoChunkId& Chu
 		return ChunkInfo.RawSize();
 	}
 
-	return InstallCacheStatusBuilder(ErrorCode);
+	return FIoStatus(ErrorCode);
 }
 
 TIoStatusOr<FIoMappedRegion> FOnDemandInstallCache::OpenMapped(const FIoChunkId& ChunkId, const FIoReadOptions& Options)
 {
-	return InstallCacheStatusBuilder(EIoErrorCode::FileOpenFailed);
+	return FIoStatus(EIoErrorCode::FileOpenFailed, TEXT("On-demand install cache does not support memory mapped files."));
 }
 
 const TCHAR* FOnDemandInstallCache::GetName() const
@@ -1990,7 +2022,6 @@ const TCHAR* FOnDemandInstallCache::GetName() const
 	return TEXT("OnDemandInstallCache");
 }
 
-
 bool FOnDemandInstallCache::Resolve(FIoRequestImpl* Request)
 {
 	EIoErrorCode ErrorCode = EIoErrorCode::UnknownChunkID;
@@ -2193,7 +2224,7 @@ bool FOnDemandInstallCache::TryPinChunks(
 	return true;
 }
 
-FIoStatus FOnDemandInstallCache::PutChunk(FIoBuffer&& Chunk, const FIoHash& ChunkHash)
+FResult FOnDemandInstallCache::PutChunk(FIoBuffer&& Chunk, const FIoHash& ChunkHash)
 {
 	if (PendingChunks.IsValid() == false)
 	{
@@ -2202,18 +2233,18 @@ FIoStatus FOnDemandInstallCache::PutChunk(FIoBuffer&& Chunk, const FIoHash& Chun
 
 	if (PendingChunks->TotalSize > FPendingChunks::MaxPendingBytes)
 	{
-		if (FIoStatus Status = FlushPendingChunks(*PendingChunks); Status.IsOk() == false)
+		if (FResult Result = FlushPendingChunks(*PendingChunks); Result.HasError())
 		{
-			return Status;
+			return Result;
 		}
 		check(PendingChunks->IsEmpty());
 	}
 
 	PendingChunks->Append(MoveTemp(Chunk), ChunkHash);
-	return InstallCacheStatusBuilder(EIoErrorCode::Ok);
+	return MakeValue();
 }
 
-FIoStatus FOnDemandInstallCache::Purge(TSet<FIoHash>&& ChunksToInstall)
+FResult FOnDemandInstallCache::Purge(TSet<FIoHash>&& ChunksToInstall)
 {
 	TUniqueLock Lock(PurgeDefragMutex);
 
@@ -2271,7 +2302,7 @@ FIoStatus FOnDemandInstallCache::Purge(TSet<FIoHash>&& ChunksToInstall)
 	{
 		UE_LOG(LogIoStoreOnDemand, Log, TEXT("Skipping cache purge, MaxCacheSize=%.2lf MiB, CacheSize=%.2lf MiB, ReferencedBlockSize=%.2lf MiB, ReferencedSize=%.2lf MiB, FragmentedBytes=%.2lf MiB, UncachedSize=%.2lf MiB"),
 			ToMiB(MaxCacheSize), ToMiB(TotalCachedBytes), ToMiB(TotalReferencedBlockBytes), ToMiB(ReferencedBytes), ToMiB(FragmentedBytes), ToMiB(TotalUncachedBytes));
-		return InstallCacheStatusBuilder(EIoErrorCode::Ok);
+		return MakeValue();
 	}
 
 	//TODO: Compute fragmentation metric and redownload chunks when this number gets too high
@@ -2281,8 +2312,7 @@ FIoStatus FOnDemandInstallCache::Purge(TSet<FIoHash>&& ChunksToInstall)
 
 	const uint64	TotalBytesToPurge	= TotalRequiredBytes - MaxCacheSize;
 	uint64			TotalPurgedBytes	= 0;
-
-	FIoStatus Status = Purge(BlockInfo, TotalBytesToPurge, TotalPurgedBytes);
+	FResult			Result				= Purge(BlockInfo, TotalBytesToPurge, TotalPurgedBytes);
 
 	if (TotalPurgedBytes > 0)
 	{
@@ -2296,31 +2326,33 @@ FIoStatus FOnDemandInstallCache::Purge(TSet<FIoHash>&& ChunksToInstall)
 			ToMiB(NewCachedBytes - MaxCacheSize), 100.0 * (double(NewCachedBytes - MaxCacheSize) / double(MaxCacheSize)));
 
 	uint64 DefragPurgedBytes = 0;
-	if (Status.IsOk() && TotalPurgedBytes < TotalBytesToPurge)
+	if (Result.HasError() == false && TotalPurgedBytes < TotalBytesToPurge)
 	{
 		if (UE::IoStore::CVars::GIoStoreOnDemandEnableDefrag)
 		{
 			// Attempt to defrag
 			const uint64 DefragBytesToPurge = TotalBytesToPurge - TotalPurgedBytes;
-			FIoStatus DefragStatus = Defrag(Containers, ChunkEntryIndices, BlockInfo, &DefragBytesToPurge, &DefragPurgedBytes);
-			if (DefragStatus.IsOk() == false)
-			{
-				Status = InstallCacheStatusBuilder(EIoErrorCode::WriteError)
-					<< FString::Printf(TEXT("Defrag failed (%s)"), *DefragStatus.ToString());
-			}
+			Result = Defrag(Containers, ChunkEntryIndices, BlockInfo, &DefragBytesToPurge, &DefragPurgedBytes);
 		}
 		else
 		{
-			Status = InstallCacheStatusBuilder(EIoErrorCode::WriteError) << TEXT("Failed to purge required size from install cache");
+			UE::UnifiedError::IoStoreOnDemand::FInstallCacheErrorContext Ctx
+			{
+				.Error			= FIoStatus(EIoErrorCode::InvalidCode, TEXT("Failed to purge required size from install cache")),
+				.CacheSize		= TotalCachedBytes,
+				.ErrorCategory	= UE::UnifiedError::IoStoreOnDemand::EErrorCategory::Logical
+			};
+			CaptureCacheState(Ctx);
+			Result = MakeError(UE::UnifiedError::IoStoreOnDemand::InstallCachePurgeError::MakeError(
+				MoveTemp(Ctx), UE::UnifiedError::EDetailFilter::All));
 		}
 	}
 
-	FOnDemandInstallCacheStats::OnPurge(Status, MaxCacheSize, NewCachedBytes, TotalBytesToPurge, TotalPurgedBytes + DefragPurgedBytes);
-
-	return Status;
+	FOnDemandInstallCacheStats::OnPurge(Result, MaxCacheSize, NewCachedBytes, TotalBytesToPurge, TotalPurgedBytes + DefragPurgedBytes);
+	return Result;
 }
 
-FIoStatus FOnDemandInstallCache::PurgeAllUnreferenced(bool bDefrag, const uint64* BytesToPurge /*= nullptr*/)
+FResult FOnDemandInstallCache::PurgeAllUnreferenced(bool bDefrag, const uint64* BytesToPurge /*= nullptr*/)
 {
 	TUniqueLock Lock(PurgeDefragMutex);
 
@@ -2345,7 +2377,7 @@ FIoStatus FOnDemandInstallCache::PurgeAllUnreferenced(bool bDefrag, const uint64
 
 	const uint64	TotalBytesToPurge	= BytesToPurge ? *BytesToPurge : MaxCacheSize;
 	uint64			TotalPurgedBytes	= 0;
-	FIoStatus Status = Purge(BlockInfo, TotalBytesToPurge, TotalPurgedBytes);
+	FResult			Result				= Purge(BlockInfo, TotalBytesToPurge, TotalPurgedBytes);
 
 	if (TotalPurgedBytes > 0)
 	{
@@ -2364,35 +2396,31 @@ FIoStatus FOnDemandInstallCache::PurgeAllUnreferenced(bool bDefrag, const uint64
 		{
 			// Attempt to defrag
 			const uint64 DefragBytesToPurge = TotalBytesToPurge - TotalPurgedBytes;
-			FIoStatus DefragStatus = Defrag(Containers, ChunkEntryIndices, BlockInfo, &DefragBytesToPurge);
-			if (DefragStatus.IsOk() == false)
-			{
-				Status = InstallCacheStatusBuilder(EIoErrorCode::WriteError)
-					<< FString::Printf(TEXT("Defrag failed (%s)"), *DefragStatus.ToString());
-			}
+			Result = Defrag(Containers, ChunkEntryIndices, BlockInfo, &DefragBytesToPurge);
 		}
 		else
 		{
-			Status = InstallCacheStatusBuilder(EIoErrorCode::WriteError) << TEXT("Failed to purge required size from install cache");
+			UE::UnifiedError::IoStoreOnDemand::FInstallCacheErrorContext Ctx
+			{
+				.Error			= FIoStatus(EIoErrorCode::InvalidCode, TEXT("Failed to purge all unreferenced cache block(s)")),
+				.CacheSize		= TotalCachedBytes,
+				.ErrorCategory	= UE::UnifiedError::IoStoreOnDemand::EErrorCategory::Logical
+			};
+			CaptureCacheState(Ctx);
+			Result = MakeError(UE::UnifiedError::IoStoreOnDemand::InstallCachePurgeError::MakeError(
+				MoveTemp(Ctx), UE::UnifiedError::EDetailFilter::All));
 		}
 	}
 	else if (bDefrag)
 	{
-		// Just do Full Defrag
-		FIoStatus DefragStatus = Defrag(Containers, ChunkEntryIndices, BlockInfo);
-		if (DefragStatus.IsOk() == false)
-		{
-			Status = InstallCacheStatusBuilder(EIoErrorCode::WriteError)
-				<< FString::Printf(TEXT("Defrag failed (%s)"), *DefragStatus.ToString());
-		}
+		Result = Defrag(Containers, ChunkEntryIndices, BlockInfo);
 	}
 
-	FOnDemandInstallCacheStats::OnPurge(Status, MaxCacheSize, NewCachedBytes, TotalBytesToPurge, TotalPurgedBytes);
-
-	return Status;
+	FOnDemandInstallCacheStats::OnPurge(Result, MaxCacheSize, NewCachedBytes, TotalBytesToPurge, TotalPurgedBytes);
+	return Result;
 }
 
-FIoStatus FOnDemandInstallCache::DefragAll(const uint64* BytesToFree /*= nullptr*/)
+FResult FOnDemandInstallCache::DefragAll(const uint64* BytesToFree /*= nullptr*/)
 {
 	TUniqueLock Lock(PurgeDefragMutex);
 
@@ -2418,7 +2446,7 @@ FIoStatus FOnDemandInstallCache::DefragAll(const uint64* BytesToFree /*= nullptr
 	return Defrag(Containers, ChunkEntryIndices, BlockInfo, BytesToFree);
 }
 
-FIoStatus FOnDemandInstallCache::Verify()
+FResult FOnDemandInstallCache::Verify()
 {
 	struct FChunkLookup
 	{
@@ -2498,7 +2526,7 @@ FIoStatus FOnDemandInstallCache::Verify()
 	if (ChunkCount == 0)
 	{
 		UE_LOG(LogIoStoreOnDemand, Log, TEXT("Verify skipped, install cache is empty"));
-		return InstallCacheStatusBuilder(EIoErrorCode::Ok);
+		return MakeValue();
 	}
 
 	UE_LOG(LogIoStoreOnDemand, Log, TEXT("Verifying %d installed chunks..."), ChunkCount);
@@ -2594,23 +2622,19 @@ FIoStatus FOnDemandInstallCache::Verify()
 			UE_LOG(LogIoStoreOnDemand, Warning, TEXT("%s"), *Reason);
 		}
 
-		if (CorruptChunkCount > 0)
-		{
-			return InstallCacheStatusBuilder(EIoErrorCode::SignatureError);
-		}
-
-		if (ReadErrorCount > 0)
-		{
-			return InstallCacheStatusBuilder(EIoErrorCode::ReadError);
-		}
-
-		return InstallCacheStatusBuilder(EIoErrorCode::NotFound);
+		return MakeError(UE::UnifiedError::IoStoreOnDemand::InstallCacheVerificationError::MakeError(
+			UE::UnifiedError::IoStoreOnDemand::FVerificationErrorContext
+			{
+				.CorruptChunkCount	= CorruptChunkCount,
+				.MissingChunkCount	= MissingChunkCount,
+				.ReadErrorCount		= ReadErrorCount
+			}));
 	}
 
 	UE_LOG(LogIoStoreOnDemand, Log, TEXT("Successfully verified %d chunk(s) of total %.2lf MiB"),
 		ChunkCount, ToMiB(TotalVerifiedBytes));
 
-	return InstallCacheStatusBuilder(EIoErrorCode::Ok);
+	return MakeValue();
 }
 
 void FOnDemandInstallCache::RegisterConsoleCommands()
@@ -2646,7 +2670,7 @@ FIoStatus FOnDemandInstallCache::Reset()
 
 	if (Ifm.DeleteDirectory(*CacheDirectory, false, bTree) == false)
 	{
-		return InstallCacheStatusBuilder(EIoErrorCode::WriteError)
+		return FIoStatusBuilder(EIoErrorCode::DeleteError, FPlatformMisc::GetLastError())
 			<< TEXT("Failed to delete directory '")
 			<< CacheDirectory
 			<< TEXT("'");
@@ -2654,7 +2678,7 @@ FIoStatus FOnDemandInstallCache::Reset()
 
 	if (Ifm.MakeDirectory(*CacheDirectory, bTree) == false)
 	{
-		return InstallCacheStatusBuilder(EIoErrorCode::WriteError)
+		return FIoStatusBuilder(EIoErrorCode::WriteError, FPlatformMisc::GetLastError())
 			<< TEXT("Failed to create directory '")
 			<< CacheDirectory
 			<< TEXT("'");
@@ -2672,7 +2696,7 @@ FIoStatus FOnDemandInstallCache::Reset()
 	}
 
 	UE_LOG(LogIoStoreOnDemand, Log, TEXT("Created CAS journal '%s'"), *JournalFile);
-	return InstallCacheStatusBuilder(EIoErrorCode::Ok);
+	return FIoStatus::Ok;
 }
 
 FIoStatus FOnDemandInstallCache::InitialVerify()
@@ -2711,22 +2735,29 @@ FIoStatus FOnDemandInstallCache::InitialVerify()
 			UE_LOG(LogIoStoreOnDemand, Warning, TEXT("Cache size is greater than disk quota - Purging install cache, MaxCacheSize=%.2lf MiB, TotalSize=%.2lf MiB, TotalBytesToPurge=%.2lf MiB"),
 				ToMiB(MaxCacheSize), ToMiB(CacheSize), ToMiB(TotalBytesToPurge));
 
-			FIoStatus PurgeStatus = Purge(BlockInfo, TotalBytesToPurge, TotalPurgedBytes);
-			if (PurgeStatus.IsOk() && TotalPurgedBytes >= TotalBytesToPurge)
+			if (FResult Result = Purge(BlockInfo, TotalBytesToPurge, TotalPurgedBytes); Result.HasError())
 			{
-				UE_LOG(LogIoStoreOnDemand, Warning, TEXT("Successfully purged %.2lf MiB from install cache"), ToMiB(TotalPurgedBytes));
+				const FString ErrorMessage = FString::Printf(TEXT("Failed to purge %.2lf MiB from install cache reason '%s'"),
+					ToMiB(TotalBytesToPurge), *LexToString(Result.GetError()));
+
+				UE_LOG(LogIoStoreOnDemand, Error, TEXT("%s"), *ErrorMessage);
+				return FIoStatus(EIoErrorCode::WriteError, ErrorMessage);
 			}
-			else
+
+			if (TotalPurgedBytes < TotalBytesToPurge)
 			{
-				UE_LOG(LogIoStoreOnDemand, Error, TEXT("Failed to purge %.2lf MiB from install cache. Actually purged %.2lf MiB from install cache"),
+				const FString ErrorMessage = FString::Printf(TEXT("Failed to purge %.2lf MiB from install cache. Actually purged %.2lf MiB from install cache"),
 					ToMiB(TotalBytesToPurge), ToMiB(TotalPurgedBytes));
 
-				return InstallCacheStatusBuilder(EIoErrorCode::WriteError) << FString::Printf(TEXT("Failed to purge overbudget cache (%s)"), *PurgeStatus.ToString());
+				UE_LOG(LogIoStoreOnDemand, Error, TEXT("%s"), *ErrorMessage);
+				return FIoStatus(EIoErrorCode::WriteError, ErrorMessage);
 			}
+
+			UE_LOG(LogIoStoreOnDemand, Log, TEXT("Successfully purged %.2lf MiB from install cache"), ToMiB(TotalPurgedBytes));
 		}
 	}
 
-	return InstallCacheStatusBuilder(EIoErrorCode::Ok);
+	return FIoStatus::Ok;
 }
 
 uint64 FOnDemandInstallCache::AddReferencesToBlocks(
@@ -2780,7 +2811,7 @@ uint64 FOnDemandInstallCache::AddReferencesToBlocks(
 
 				if (bIsReferenced)
 				{
-					UE_LOGFMT(LogIoStoreOnDemand, Error, "Installing a chunk but its already cached! ChunkId={ChunkId}, Container={Container}", 
+					UE_LOGFMT(LogIoStoreOnDemand, Error, "Installing a chunk but its already cached! ChunkId={ChunkId}, Container={Container}",
 						LexToString(Container->ChunkIds[EntryIndex]), Container->Name);
 					ensure(false);
 				}
@@ -2798,7 +2829,7 @@ uint64 FOnDemandInstallCache::AddReferencesToBlocks(
 	return TotalUncachedBytes;
 }
 
-FIoStatus FOnDemandInstallCache::Purge(FCasBlockInfoMap& BlockInfo, const uint64 TotalBytesToPurge, uint64& OutTotalPurgedBytes)
+FResult FOnDemandInstallCache::Purge(FCasBlockInfoMap& BlockInfo, const uint64 TotalBytesToPurge, uint64& OutTotalPurgedBytes)
 {
 	BlockInfo.ValueSort([](const FCasBlockInfo& LHS, const FCasBlockInfo& RHS)
 	{
@@ -2821,7 +2852,14 @@ FIoStatus FOnDemandInstallCache::Purge(FCasBlockInfoMap& BlockInfo, const uint64
 
 		if (FIoStatus Status = Cas.DeleteBlock(BlockId, RemovedChunks); !Status.IsOk())
 		{
-			return Status;
+			UE::UnifiedError::IoStoreOnDemand::FInstallCacheErrorContext Ctx
+			{
+				.Error			= Status,
+				.ErrorCategory	= UE::UnifiedError::IoStoreOnDemand::EErrorCategory::Cas
+			};
+			CaptureCacheState(Ctx);
+			return MakeError(UE::UnifiedError::IoStoreOnDemand::InstallCachePurgeError::MakeError(
+				MoveTemp(Ctx), UE::UnifiedError::EDetailFilter::All));
 		}
 
 		FOnDemandInstallCacheStats::OnBlockDeleted(Info.LastAccess, false /*FromDefrag*/);
@@ -2842,7 +2880,14 @@ FIoStatus FOnDemandInstallCache::Purge(FCasBlockInfoMap& BlockInfo, const uint64
 
 		if (FIoStatus Status = FCasJournal::Commit(MoveTemp(Transaction)); !Status.IsOk())
 		{
-			return Status;
+			UE::UnifiedError::IoStoreOnDemand::FInstallCacheErrorContext Ctx
+			{
+				.Error			= Status,
+				.ErrorCategory	= UE::UnifiedError::IoStoreOnDemand::EErrorCategory::Journal
+			};
+			CaptureCacheState(Ctx);
+			return MakeError(UE::UnifiedError::IoStoreOnDemand::InstallCachePurgeError::MakeError(
+				MoveTemp(Ctx), UE::UnifiedError::EDetailFilter::All));
 		}
 		
 		if (OutTotalPurgedBytes >= TotalBytesToPurge)
@@ -2851,10 +2896,10 @@ FIoStatus FOnDemandInstallCache::Purge(FCasBlockInfoMap& BlockInfo, const uint64
 		}
 	}
 
-	return InstallCacheStatusBuilder(EIoErrorCode::Ok);
+	return MakeValue();
 }
 
-FIoStatus FOnDemandInstallCache::Defrag(
+FResult FOnDemandInstallCache::Defrag(
 	const TArray<FSharedOnDemandContainer>& Containers,
 	const TArray<TBitArray<>>& ChunkEntryIndices,
 	FCasBlockInfoMap& BlockInfo, 
@@ -2863,7 +2908,7 @@ FIoStatus FOnDemandInstallCache::Defrag(
 {
 	if (TotalBytesToFree && *TotalBytesToFree == 0)
 	{
-		return InstallCacheStatusBuilder(EIoErrorCode::Ok);
+		return MakeValue();
 	}
 
 	const uint64 TotalCachedBytes = Algo::TransformAccumulate(BlockInfo,
@@ -2877,8 +2922,16 @@ FIoStatus FOnDemandInstallCache::Defrag(
 			ToMiB(MaxCacheSize));
 		UE_LOG(LogIoStoreOnDemand, Error, TEXT("%s"), *ErrorMsg);
 
-		FIoStatus Status = InstallCacheStatusBuilder(EIoErrorCode::WriteError) << ErrorMsg;
-		FOnDemandInstallCacheStats::OnDefrag(Status, 0);
+		UE::UnifiedError::IoStoreOnDemand::FInstallCacheErrorContext Ctx
+		{
+			.Error			= FIoStatus(EIoErrorCode::OutOfDiskSpace, ErrorMsg),
+			.CacheSize		= TotalCachedBytes,
+			.ErrorCategory	= UE::UnifiedError::IoStoreOnDemand::EErrorCategory::Logical
+		};
+		CaptureCacheState(Ctx);
+		FResult Result = MakeError(UE::UnifiedError::IoStoreOnDemand::InstallCacheDefragError::MakeError(
+			MoveTemp(Ctx), UE::UnifiedError::EDetailFilter::All));
+		FOnDemandInstallCacheStats::OnDefrag(Result, 0);
 
 		// Append a critical error entry to clear the cache at next startup
 		FCasJournal::FTransaction Transaction = FCasJournal::Begin(GetJournalFilename());
@@ -2889,14 +2942,15 @@ FIoStatus FOnDemandInstallCache::Defrag(
 				*JournalStatus.ToString());
 		}
 
-		return Status;
+		return Result;
 	}
 
 	struct FDefragBlockReferencedChunk
 	{
-		uint32 BlockOffset = 0;
-		uint32 DiskSize = 0;
-		FIoHash Hash;
+		FIoHash		Hash;
+		FIoChunkId	ChunkId; 
+		uint32		BlockOffset = 0;
+		uint32		DiskSize = 0;
 	};
 
 	struct FDefragBlock
@@ -2963,9 +3017,16 @@ FIoStatus FOnDemandInstallCache::Defrag(
 
 		if (!bPossibleToFreeBytes)
 		{
-			FIoStatus Status = InstallCacheStatusBuilder(EIoErrorCode::WriteError) << TEXT("Defrag cannot free requested storage space");
-			FOnDemandInstallCacheStats::OnDefrag(Status, 0);
-			return Status;
+			UE::UnifiedError::IoStoreOnDemand::FInstallCacheErrorContext Ctx
+			{
+				.CacheSize		= TotalCachedBytes,
+				.ErrorCategory	= UE::UnifiedError::IoStoreOnDemand::EErrorCategory::Logical
+			};
+			CaptureCacheState(Ctx);
+			FResult Result = MakeError(UE::UnifiedError::IoStoreOnDemand::InstallCacheDefragError::MakeError(
+				MoveTemp(Ctx), UE::UnifiedError::EDetailFilter::All));
+			FOnDemandInstallCacheStats::OnDefrag(Result, 0);
+			return Result;
 		}
 	}
 	else
@@ -3002,7 +3063,7 @@ FIoStatus FOnDemandInstallCache::Defrag(
 		{
 			// Already defragged
 			UE_LOG(LogIoStoreOnDemand, Display, TEXT("Cache not fragmented."));
-			return InstallCacheStatusBuilder(EIoErrorCode::Ok);
+			return MakeValue();
 		}
 	}
 
@@ -3033,12 +3094,12 @@ FIoStatus FOnDemandInstallCache::Defrag(
 				if (FDefragBlock* DefragBlock = Algo::FindBy(BlocksToDefrag, Loc.BlockId, &FDefragBlock::BlockId))
 				{
 					DefragBlock->ReferencedChunks.Add(FDefragBlockReferencedChunk
-						{
-							.BlockOffset = Loc.BlockOffset,
-							.DiskSize = Entry.GetDiskSize(),
-							.Hash = Entry.Hash,
-						});
-
+					{
+						.Hash = Entry.Hash,
+						.ChunkId = Container->ChunkIds[EntryIndex],
+						.BlockOffset = Loc.BlockOffset,
+						.DiskSize = Entry.GetDiskSize(),
+					});
 					TotalRefBytes += Entry.GetDiskSize();
 				}
 			}
@@ -3062,9 +3123,17 @@ FIoStatus FOnDemandInstallCache::Defrag(
 			{
 				const FString ErrorMsg = FileOpenResult.GetError().GetMessage();
 				UE_LOG(LogIoStoreOnDemand, Error, TEXT("Defrag failed to open CAS block for reading, reason '%s'"), *ErrorMsg);
-				FIoStatus Status = InstallCacheStatusBuilder(EIoErrorCode::FileOpenFailed) << ErrorMsg;
-				FOnDemandInstallCacheStats::OnDefrag(Status, 0);
-				return Status;
+				UE::UnifiedError::IoStoreOnDemand::FInstallCacheErrorContext Ctx
+				{
+					.Error			= FIoStatus(EIoErrorCode::FileNotOpen, FPlatformMisc::GetLastError(), ErrorMsg),
+					.CacheSize		= TotalCachedBytes,
+					.ErrorCategory	= UE::UnifiedError::IoStoreOnDemand::EErrorCategory::Cas
+				};
+				CaptureCacheState(Ctx);
+				FResult Result = MakeError(UE::UnifiedError::IoStoreOnDemand::InstallCacheDefragError::MakeError(
+					MoveTemp(Ctx), UE::UnifiedError::EDetailFilter::All));
+				FOnDemandInstallCacheStats::OnDefrag(Result, 0);
+				return Result;
 			}
 
 			FSharedFileHandle FileHandle = FileOpenResult.StealValue();
@@ -3073,22 +3142,50 @@ FIoStatus FOnDemandInstallCache::Defrag(
 
 			for (const FDefragBlockReferencedChunk& ReffedChunk : DefragBlock.ReferencedChunks)
 			{
-				FileHandle->Seek(ReffedChunk.BlockOffset);
+				FIoBuffer		Buffer(ReffedChunk.DiskSize);
+				FIoStatus		ChunkStatus = FIoStatus::Invalid;
 
-				FIoBuffer Buffer(ReffedChunk.DiskSize);
-				const bool bOk = FileHandle->Read(Buffer.GetData(), Buffer.GetSize());
-				if (!bOk)
+				for (int32 Attempt = 0, MaxAttempts = 3; Attempt < MaxAttempts; ++Attempt)
 				{
-					FIoStatus Status = InstallCacheStatusBuilder(EIoErrorCode::ReadError);
-					FOnDemandInstallCacheStats::OnDefrag(Status, 0);
-					return Status;
+					if (FileHandle->Seek(ReffedChunk.BlockOffset) == false)
+					{
+						ChunkStatus = FIoStatusBuilder(EIoErrorCode::FileSeekFailed, FPlatformMisc::GetLastError())
+							<< FString::Printf(TEXT("Failed to seek to offset %d in cache block %d"),
+								ReffedChunk.BlockOffset, DefragBlock.BlockId.Id);
+						continue;
+					}
+					if (FileHandle->Read(Buffer.GetData(), Buffer.GetSize()) == false)
+					{
+						ChunkStatus = FIoStatusBuilder(EIoErrorCode::ReadError, FPlatformMisc::GetLastError())
+							<< FString::Printf(TEXT("Failed to read %ull bytes at offset %d in cache block %d"),
+								Buffer.GetSize(), ReffedChunk.BlockOffset, DefragBlock.BlockId.Id);
+						continue;
+					}
+
+					ChunkStatus = FIoStatus::Ok;
+					break;
+				}
+
+				if (ChunkStatus.IsOk() == false)
+				{
+					UE::UnifiedError::IoStoreOnDemand::FInstallCacheErrorContext Ctx
+					{
+						.Error			= ChunkStatus,
+						.CacheSize		= TotalCachedBytes,
+						.ErrorCategory	= UE::UnifiedError::IoStoreOnDemand::EErrorCategory::Cas
+					};
+					CaptureCacheState(Ctx);
+					FResult Result = MakeError(UE::UnifiedError::IoStoreOnDemand::InstallCacheDefragError::MakeError(
+						MoveTemp(Ctx), UE::UnifiedError::EDetailFilter::All));
+					FOnDemandInstallCacheStats::OnDefrag(Result, 0);
+					return Result;
 				}
 
 				const FIoHash ChunkHash = FIoHash::HashBuffer(Buffer.GetView());
 				if (ChunkHash != ReffedChunk.Hash)
 				{
-					UE_LOG(LogIoStoreOnDemand, Error, TEXT("Found chunk with invalid hash while defragging block, BlockId=%u, BlockOffset=%u"),
-						DefragBlock.BlockId.Id, ReffedChunk.BlockOffset);
+					UE_LOG(LogIoStoreOnDemand, Error, TEXT("Found chunk with invalid hash while defragging block, ChunkId='%s', BlockId=%u, BlockOffset=%u"),
+						*LexToString(ReffedChunk.ChunkId), DefragBlock.BlockId.Id, ReffedChunk.BlockOffset);
 
 					// Append a critical error entry to clear the cache at next startup
 					FCasJournal::FTransaction Transaction = FCasJournal::Begin(GetJournalFilename());
@@ -3099,23 +3196,29 @@ FIoStatus FOnDemandInstallCache::Defrag(
 							*Status.ToString());
 					}
 
-					if (FIoStatus Status = FlushPendingChunks(DefragPendingChunks, DefragBlock.LastAccess); Status.IsOk() == false)
+					if (FResult Result = FlushPendingChunks(DefragPendingChunks, DefragBlock.LastAccess); Result.HasError())
 					{
-						FOnDemandInstallCacheStats::OnDefrag(Status, 0);
-						return Status;
+						FOnDemandInstallCacheStats::OnDefrag(Result, 0);
+						return Result;
 					}
 	
-					FIoStatus Status = InstallCacheStatusBuilder(EIoErrorCode::SignatureError);
-					FOnDemandInstallCacheStats::OnDefrag(Status, 0);
-					return Status;
+					FResult Result = MakeError(UE::UnifiedError::IoStoreOnDemand::InstallCacheDefragError::MakeError(
+						UE::UnifiedError::IoStoreOnDemand::FChunkHashMismatchErrorContext
+						{
+							.ChunkId		= ReffedChunk.ChunkId,
+							.ExpectedHash	= ReffedChunk.Hash,
+							.ActualHash		= ChunkHash
+						}));
+					FOnDemandInstallCacheStats::OnDefrag(Result, 0);
+					return Result;
 				}
 
 				if (DefragPendingChunks.TotalSize > FPendingChunks::MaxPendingBytes)
 				{
-					if (FIoStatus Status = FlushPendingChunks(DefragPendingChunks, DefragBlock.LastAccess); Status.IsOk() == false)
+					if (FResult Result = FlushPendingChunks(DefragPendingChunks, DefragBlock.LastAccess); Result.HasError())
 					{
-						FOnDemandInstallCacheStats::OnDefrag(Status, 0);
-						return Status;
+						FOnDemandInstallCacheStats::OnDefrag(Result, 0);
+						return Result;
 					}
 					check(DefragPendingChunks.IsEmpty());
 				}
@@ -3125,10 +3228,10 @@ FIoStatus FOnDemandInstallCache::Defrag(
 
 			FileHandle.Reset();
 
-			if (FIoStatus Status = FlushPendingChunks(DefragPendingChunks, DefragBlock.LastAccess); Status.IsOk() == false)
+			if (FResult Result = FlushPendingChunks(DefragPendingChunks, DefragBlock.LastAccess); Result.HasError()) 
 			{
-				FOnDemandInstallCacheStats::OnDefrag(Status, 0);
-				return Status;
+				FOnDemandInstallCacheStats::OnDefrag(Result, 0);
+				return Result;
 			}
 			check(DefragPendingChunks.IsEmpty());
 		}
@@ -3140,8 +3243,18 @@ FIoStatus FOnDemandInstallCache::Defrag(
 		TArray<FCasAddr> DeletedChunkAddresses;
 		if (FIoStatus Status = Cas.DeleteBlock(DefragBlock.BlockId, DeletedChunkAddresses); !Status.IsOk())
 		{
-			FOnDemandInstallCacheStats::OnDefrag(Status, 0);
-			return Status;
+			UE::UnifiedError::IoStoreOnDemand::FInstallCacheErrorContext Ctx
+			{
+				.Error			= Status,
+				.CacheSize		= TotalCachedBytes,
+				.ErrorCategory	= UE::UnifiedError::IoStoreOnDemand::EErrorCategory::Cas
+			};
+			CaptureCacheState(Ctx);
+			FResult Result = MakeError(UE::UnifiedError::IoStoreOnDemand::InstallCacheDefragError::MakeError(
+				MoveTemp(Ctx), UE::UnifiedError::EDetailFilter::All));
+			FOnDemandInstallCacheStats::OnDefrag(Result, 0);
+
+			return Result;
 		}
 
 		FOnDemandInstallCacheStats::OnBlockDeleted(DefragBlock.LastAccess, true /*FromDefrag*/);
@@ -3154,25 +3267,34 @@ FIoStatus FOnDemandInstallCache::Defrag(
 
 		if (FIoStatus Status = FCasJournal::Commit(MoveTemp(Transaction)); !Status.IsOk())
 		{
-			FOnDemandInstallCacheStats::OnDefrag(Status, 0);
-			return Status;
+			UE::UnifiedError::IoStoreOnDemand::FInstallCacheErrorContext Ctx
+			{
+				.Error			= Status,
+				.CacheSize		= TotalCachedBytes,
+				.ErrorCategory	= UE::UnifiedError::IoStoreOnDemand::EErrorCategory::Journal
+			};
+			CaptureCacheState(Ctx);
+			FResult Result = MakeError(UE::UnifiedError::IoStoreOnDemand::InstallCacheDefragError::MakeError(
+				MoveTemp(Ctx), UE::UnifiedError::EDetailFilter::All));
+			FOnDemandInstallCacheStats::OnDefrag(Result, 0);
+
+			return Result;
 		}
 	}
 
 	UE_LOG(LogIoStoreOnDemand, Display, TEXT("Defrag removed %" UINT64_FMT " fragmented bytes of %" UINT64_FMT " total bytes in %i blocks."), FragmentedBytes, TotalBlockSize, BlocksToDefrag.Num());
 
-	FIoStatus Status = InstallCacheStatusBuilder(EIoErrorCode::Ok);
-	FOnDemandInstallCacheStats::OnDefrag(Status, FragmentedBytes);
-
 	if (OutTotalFreedBytes)
 	{
 		*OutTotalFreedBytes = FragmentedBytes;
 	}
 
-	return Status;
+	FResult Result = MakeValue();
+	FOnDemandInstallCacheStats::OnDefrag(Result, FragmentedBytes);
+	return Result;
 }
 
-FIoStatus FOnDemandInstallCache::Flush()
+FResult FOnDemandInstallCache::Flush()
 {
 	if (PendingChunks.IsValid())
 	{
@@ -3181,16 +3303,16 @@ FIoStatus FOnDemandInstallCache::Flush()
 	}
 
 	Cas.Compact();
-	return InstallCacheStatusBuilder(EIoErrorCode::Ok);
+	return MakeValue();
 }
 
-FIoStatus FOnDemandInstallCache::FlushLastAccess()
+FResult FOnDemandInstallCache::FlushLastAccess()
 {
 	const FCas::FLastAccess DirtyLastAccess = Cas.GetAndClearDirtyLastAccess();
 
 	if (DirtyLastAccess.IsEmpty())
 	{
-		return InstallCacheStatusBuilder(EIoErrorCode::Ok);
+		return MakeValue();
 	}
 
 	const FString JournalFile = GetJournalFilename();
@@ -3202,10 +3324,17 @@ FIoStatus FOnDemandInstallCache::FlushLastAccess()
 		UE_LOG(LogIoStoreOnDemand, Error, TEXT("Failed to update CAS journal '%s' with block timestamp(s), reason '%s'"),
 			*JournalFile, *Status.ToString());
 
-		return Status;
+		UE::UnifiedError::IoStoreOnDemand::FInstallCacheErrorContext Ctx
+		{
+			.Error			= Status,
+			.ErrorCategory	= UE::UnifiedError::IoStoreOnDemand::EErrorCategory::Journal
+		};
+		CaptureCacheState(Ctx);
+		return MakeError(UE::UnifiedError::IoStoreOnDemand::InstallCacheFlushLastAccessError::MakeError(
+			MoveTemp(Ctx), UE::UnifiedError::EDetailFilter::All));
 	}
 
-	return InstallCacheStatusBuilder(EIoErrorCode::Ok);
+	return MakeValue();
 }
 
 void FOnDemandInstallCache::UpdateLastAccess(TConstArrayView<FIoHash> ChunkHashes)
@@ -3283,15 +3412,15 @@ FOnDemandInstallCacheUsage FOnDemandInstallCache::GetCacheUsage()
 	};
 }
 
-FIoStatus FOnDemandInstallCache::FlushPendingChunks(FPendingChunks& Chunks, int64 UtcAccessTicks)
+FResult FOnDemandInstallCache::FlushPendingChunks(FPendingChunks& Chunks, int64 UtcAccessTicks)
 {
 	if (Chunks.IsEmpty())
 	{
-		return InstallCacheStatusBuilder(EIoErrorCode::Ok);
+		return MakeValue();
 	}
 
 #if UE_ONDEMANDINSTALLCACHE_EXCLUSIVE_WRITE
-	UE::Tasks::TTask<FIoStatus> Task = ExclusivePipe.Launch(UE_SOURCE_LOCATION, [this, &Chunks, UtcAccessTicks]
+	UE::Tasks::TTask<FResult> Task = ExclusivePipe.Launch(UE_SOURCE_LOCATION, [this, &Chunks, UtcAccessTicks]
 	{
 		return FlushPendingChunksImpl(Chunks, UtcAccessTicks);
 	}, UE::Tasks::ETaskPriority::BackgroundHigh);
@@ -3305,7 +3434,7 @@ FIoStatus FOnDemandInstallCache::FlushPendingChunks(FPendingChunks& Chunks, int6
 #endif // UE_ONDEMANDINSTALLCACHE_EXCLUSIVE_WRITE
 }
 
-FIoStatus FOnDemandInstallCache::FlushPendingChunksImpl(FPendingChunks& Chunks, int64 UtcAccessTicks)
+FResult FOnDemandInstallCache::FlushPendingChunksImpl(FPendingChunks& Chunks, int64 UtcAccessTicks)
 {
 	ON_SCOPE_EXIT { Chunks.Reset(); };
 
@@ -3334,8 +3463,16 @@ FIoStatus FOnDemandInstallCache::FlushPendingChunksImpl(FPendingChunks& Chunks,
 		TUniquePtr<IFileHandle>	CasFileHandle = Cas.OpenWrite(CurrentBlockId, bAppendToBlock);
 		if (CasFileHandle.IsValid() == false)
 		{
-			FOnDemandInstallCacheStats::OnFlush(EIoErrorCode::WriteError, 0);
-			return InstallCacheStatusBuilder(EIoErrorCode::FileOpenFailed) << TEXT("Failed to open cache block file");
+			UE::UnifiedError::IoStoreOnDemand::FInstallCacheErrorContext Ctx
+			{
+				.Error			= FIoStatus(EIoErrorCode::FileOpenFailed, FPlatformMisc::GetLastError()),
+				.ErrorCategory	= UE::UnifiedError::IoStoreOnDemand::EErrorCategory::Cas
+			};
+			CaptureCacheState(Ctx);
+			FResult Result = MakeError(UE::UnifiedError::IoStoreOnDemand::InstallCacheFlushError::MakeError(
+				MoveTemp(Ctx), UE::UnifiedError::EDetailFilter::All));
+			FOnDemandInstallCacheStats::OnFlush(Result, 0);
+			return Result;
 		}
 
 		const int64 CasBlockOffset = CasFileHandle->Tell();
@@ -3367,8 +3504,16 @@ FIoStatus FOnDemandInstallCache::FlushPendingChunksImpl(FPendingChunks& Chunks,
 
 			if (CasFileHandle->Write(Ar.GetData(), Ar.Tell()) == false)
 			{
-				return InstallCacheStatusBuilder(EIoErrorCode::WriteError)
-					<< TEXT("Failed to serialize chunks to cache block");
+				UE::UnifiedError::IoStoreOnDemand::FInstallCacheErrorContext Ctx
+				{
+					.Error			= FIoStatus(EIoErrorCode::WriteError, FPlatformMisc::GetLastError()),
+					.ErrorCategory	= UE::UnifiedError::IoStoreOnDemand::EErrorCategory::Cas
+				};
+				CaptureCacheState(Ctx);
+				FResult Result = MakeError(UE::UnifiedError::IoStoreOnDemand::InstallCacheFlushError::MakeError(
+					MoveTemp(Ctx), UE::UnifiedError::EDetailFilter::All));
+				FOnDemandInstallCacheStats::OnFlush(Result, 0);
+				return Result;
 			}
 
 			constexpr const bool bDirty = false;
@@ -3388,12 +3533,19 @@ FIoStatus FOnDemandInstallCache::FlushPendingChunksImpl(FPendingChunks& Chunks,
 
 			if (CasFileHandle->Flush() == false)
 			{
-				FOnDemandInstallCacheStats::OnFlush(EIoErrorCode::WriteError, Ar.Tell());
-				return InstallCacheStatusBuilder(EIoErrorCode::WriteError)
-					<< TEXT("Failed to flush cache block to disk");
+				UE::UnifiedError::IoStoreOnDemand::FInstallCacheErrorContext Ctx
+				{
+					.Error			= FIoStatus(EIoErrorCode::FileFlushFailed, FPlatformMisc::GetLastError()),
+					.ErrorCategory	= UE::UnifiedError::IoStoreOnDemand::EErrorCategory::Cas
+				};
+				CaptureCacheState(Ctx);
+				FResult Result = MakeError(UE::UnifiedError::IoStoreOnDemand::InstallCacheFlushError::MakeError(
+					MoveTemp(Ctx), UE::UnifiedError::EDetailFilter::All));
+				FOnDemandInstallCacheStats::OnFlush(Result, Ar.Tell());
+				return Result;
 			}
 
-			FOnDemandInstallCacheStats::OnFlush(EIoErrorCode::Ok, Ar.Tell());
+			FOnDemandInstallCacheStats::OnFlush(MakeValue(), Ar.Tell());
 
 			check(ChunkHashes.Num() == Offsets.Num());
 			check(CurrentBlockId.IsValid());
@@ -3416,7 +3568,16 @@ FIoStatus FOnDemandInstallCache::FlushPendingChunksImpl(FPendingChunks& Chunks,
 
 		if (FIoStatus Status = FCasJournal::Commit(MoveTemp(Transaction)); Status.IsOk() == false)
 		{
-			return Status;
+			UE::UnifiedError::IoStoreOnDemand::FInstallCacheErrorContext Ctx
+			{
+				.Error			= Status,
+				.ErrorCategory	= UE::UnifiedError::IoStoreOnDemand::EErrorCategory::Journal
+			};
+			CaptureCacheState(Ctx);
+			FResult Result = MakeError(UE::UnifiedError::IoStoreOnDemand::InstallCacheFlushError::MakeError(
+				MoveTemp(Ctx), UE::UnifiedError::EDetailFilter::All));
+			FOnDemandInstallCacheStats::OnFlush(Result, Ar.Tell());
+			return Result;
 		}
 
 		if (ChunkIdx < Chunks.Chunks.Num())
@@ -3425,7 +3586,7 @@ FIoStatus FOnDemandInstallCache::FlushPendingChunksImpl(FPendingChunks& Chunks,
 		}
 	}
 
-	return InstallCacheStatusBuilder(EIoErrorCode::Ok);
+	return MakeValue();
 }
 
 void FOnDemandInstallCache::CompleteRequest(FIoRequestImpl* Request, EIoErrorCode Status)
@@ -3479,6 +3640,13 @@ void FOnDemandInstallCache::CompleteRequest(FIoRequestImpl* Request, EIoErrorCod
 	BackendContext->WakeUpDispatcherThreadDelegate.Execute();
 }
 
+void FOnDemandInstallCache::CaptureCacheState(UE::UnifiedError::IoStoreOnDemand::FInstallCacheErrorContext& OutCtx)
+{
+	const FString CachRootDirectory = Cas.GetRootDirectory();
+	OutCtx.bDiskQuerySucceeded		= FPlatformMisc::GetDiskTotalAndFreeSpace(CachRootDirectory, OutCtx.DiskTotalBytes, OutCtx.DiskFreeBytes);
+	OutCtx.MaxCacheSize				= MaxCacheSize;
+}
+
 ///////////////////////////////////////////////////////////////////////////////
 TSharedPtr<IOnDemandInstallCache> MakeOnDemandInstallCache(
 	FOnDemandIoStore& IoStore,
diff --git a/Engine/Source/Runtime/Experimental/IoStore/OnDemand/Private/OnDemandInstallCache.h b/Engine/Source/Runtime/Experimental/IoStore/OnDemand/Private/OnDemandInstallCache.h
--- a/Engine/Source/Runtime/Experimental/IoStore/OnDemand/Private/OnDemandInstallCache.h
+++ b/Engine/Source/Runtime/Experimental/IoStore/OnDemand/Private/OnDemandInstallCache.h
@@ -3,6 +3,7 @@
 #pragma once
 
 #include "IO/IoDispatcherBackend.h"
+#include "IO/OnDemandError.h"
 #include "Templates/SharedPointer.h"
 #include "Containers/StringFwd.h"
 
@@ -34,13 +35,13 @@ public:
 													TConstArrayView<int32> EntryIndices,
 													FOnDemandContentHandle ContentHandle,
 													TArray<int32>& OutMissing) = 0;
-	virtual FIoStatus							PutChunk(FIoBuffer&& Chunk, const FIoHash& ChunkHash) = 0;
-	virtual FIoStatus							Purge(TSet<FIoHash>&& ChunksToInstall) = 0;
-	virtual FIoStatus							PurgeAllUnreferenced(bool bDefrag, const uint64* BytesToPurge = nullptr) = 0;
-	virtual FIoStatus							DefragAll(const uint64* BytesToFree = nullptr) = 0;
-	virtual FIoStatus							Verify() = 0;
-	virtual FIoStatus							Flush() = 0;
-	virtual FIoStatus							FlushLastAccess() = 0;
+	virtual FResult								PutChunk(FIoBuffer&& Chunk, const FIoHash& ChunkHash) = 0;
+	virtual FResult								Purge(TSet<FIoHash>&& ChunksToInstall) = 0;
+	virtual FResult								PurgeAllUnreferenced(bool bDefrag, const uint64* BytesToPurge = nullptr) = 0;
+	virtual FResult								DefragAll(const uint64* BytesToFree = nullptr) = 0;
+	virtual FResult								Verify() = 0;
+	virtual FResult								Flush() = 0;
+	virtual FResult								FlushLastAccess() = 0;
 	virtual void								UpdateLastAccess(TConstArrayView<FIoHash> ChunkHashes) = 0;
 	virtual FOnDemandInstallCacheUsage			GetCacheUsage() = 0;
 };
diff --git a/Engine/Source/Runtime/Experimental/IoStore/OnDemand/Private/OnDemandIoStore.cpp b/Engine/Source/Runtime/Experimental/IoStore/OnDemand/Private/OnDemandIoStore.cpp
--- a/Engine/Source/Runtime/Experimental/IoStore/OnDemand/Private/OnDemandIoStore.cpp
+++ b/Engine/Source/Runtime/Experimental/IoStore/OnDemand/Private/OnDemandIoStore.cpp
@@ -530,8 +530,8 @@ FIoStatus FOnDemandIoStore::InitializePostHotfix()
 			FEventRef InstallDone;
 			Install(MoveTemp(InstallArgs.GetValue()), [&InstallDone](FOnDemandInstallResult InstallResult)
 			{
-				UE_CLOG(!InstallResult.Status.IsOk(), LogIoStoreOnDemand, Error,
-					TEXT("Failed to install content, reason '%s'"), *InstallResult.Status.ToString());
+				UE_CLOG(!InstallResult.IsOk(), LogIoStoreOnDemand, Error,
+					TEXT("Failed to install content, reason '%s'"), *LexToString(InstallResult.Error.GetValue()));
 				InstallDone->Trigger();
 			});
 			InstallDone->Wait();
@@ -1012,9 +1012,9 @@ FIoStatus FOnDemandIoStore::InitializeInstallCache()
 
 			UE_LOG(LogIoStoreOnDemand, Log, TEXT("Filling install cache with %.2lf MiB of dummy data"), double(FillSize) / 1024.0 / 1024.0);
 
-			FIoStatus	Status = FIoStatus::Ok;
+			FResult		Result = MakeValue();
 			uint64		Seed = 1;
-			while (FillSize >= 0 && Status.IsOk())
+			while (FillSize >= 0 && Result.HasError() == false)
 			{
 				const uint64		ChunkSize = 256 << 10;
 				FIoBuffer			Chunk(ChunkSize);
@@ -1026,17 +1026,17 @@ FIoStatus FOnDemandIoStore::InitializeInstallCache()
 				}
 
 				const FIoHash ChunkHash = FIoHash::HashBuffer(Chunk.GetView());
-				Status = InstallCache->PutChunk(MoveTemp(Chunk), ChunkHash);
+				Result = InstallCache->PutChunk(MoveTemp(Chunk), ChunkHash);
 				Seed++;
 				FillSize -= ChunkSize;
 			}
 
-			if (Status.IsOk())
+			if (Result.HasError() == false)
 			{
-				Status = InstallCache->Flush();
+				Result = InstallCache->Flush();
 			}
 
-			UE_CLOG(!Status.IsOk(), LogIoStoreOnDemand, Warning, TEXT("Failed to fill install cache with dummy data"));
+			UE_CLOG(Result.HasError(), LogIoStoreOnDemand, Warning, TEXT("Failed to fill install cache with dummy data, reason '%s'"), *LexToString(Result.GetError()));
 		}
 	}
 #endif
diff --git a/Engine/Source/Runtime/Experimental/IoStore/OnDemand/Private/OnDemandIoStore.h b/Engine/Source/Runtime/Experimental/IoStore/OnDemand/Private/OnDemandIoStore.h
--- a/Engine/Source/Runtime/Experimental/IoStore/OnDemand/Private/OnDemandIoStore.h
+++ b/Engine/Source/Runtime/Experimental/IoStore/OnDemand/Private/OnDemandIoStore.h
@@ -313,10 +313,12 @@ void FOnDemandChunkInfo::GetUrl(FAnsiStringBuilderBase& Url) const
 /** Result from flushing cache last access times. */
 struct FOnDemandFlushLastAccessResult
 {
-	/** The status of the request. */
-	FIoStatus Status;
+	/** Returns True if the request succeeded. */
+	bool IsOk() const { return Error.IsSet() == false; }
 	/** Duration in seconds. */
 	double DurationInSeconds = 0.0;
+	/** Error information about the request. */ 
+	TOptional<UE::UnifiedError::FError> Error;
 };
 
 /** Flush last access completion callback. */
diff --git a/Engine/Source/Runtime/Experimental/IoStore/OnDemand/Private/Statistics.cpp b/Engine/Source/Runtime/Experimental/IoStore/OnDemand/Private/Statistics.cpp
--- a/Engine/Source/Runtime/Experimental/IoStore/OnDemand/Private/Statistics.cpp
+++ b/Engine/Source/Runtime/Experimental/IoStore/OnDemand/Private/Statistics.cpp
@@ -6,6 +6,7 @@
 #include "Async/UniqueLock.h"
 #include "Algo/MaxElement.h"
 #include "AnalyticsEventAttribute.h"
+#include "Experimental/UnifiedError/CoreErrorTypes.h"
 #include "HAL/IConsoleManager.h"
 #include "HAL/PlatformTime.h"
 #include "IO/IoStatus.h"
@@ -104,32 +105,22 @@ static void SendImmediateAnalytic(FOnDemandImmediateAnalytic&& Event)
 	}
 }
 
-static FString GetAnalyticErrorCode(const FIoStatus& Status)
+FOnDemandImmediateAnalytic MakeAnalyticsEventFromResult(const TCHAR* EventName, const FResult& Result)
 {
-	FString RetCode;
+	FString ErrorCode = TEXT("Success");
+	FString ErrorJson = TEXT("{\"Root\":{\"ErrorCodeString\":\"Success\", \"ModuleIdString\":\"UnifiedError\"}}");
 
-	FStringView ErrorId = Status.GetErrorMessage();
-	if (Status.GetErrorCode() != EIoErrorCode::Ok && ErrorId.StartsWith(TEXT('<')))
+	if (Result.HasError())
 	{
-		ErrorId.RightChopInline(1);
-
-		int32 Idx;
-		if (ErrorId.FindChar(TEXT('>'), Idx))
-		{
-			ErrorId.LeftInline(Idx);
-
-			RetCode = FString::Format(TEXT("{0}:{1}"), { GetIoErrorText(Status.GetErrorCode()), ErrorId });
-		}
+		ErrorCode = FString(Result.GetError().GetModuleIdAndErrorCodeString());
+		ErrorJson = Result.GetError().SerializeToJsonForAnalytics();
 	}
 
-	if (RetCode.IsEmpty())
+	return FOnDemandImmediateAnalytic
 	{
-		RetCode = GetIoErrorText(Status.GetErrorCode());
-	}
-
-	RetCode.ReplaceCharInline(TEXT(' '), TEXT('_'), ESearchCase::CaseSensitive);
-
-	return RetCode;
+		.EventName		= EventName,
+		.AnalyticsArray	= MakeAnalyticsEventAttributeArray(TEXT("ErrorCode"), ErrorCode, TEXT("ErrorJson"), FJsonFragment(MoveTemp(ErrorJson)))
+	};
 }
 
 ////////////////////////////////////////////////////////////////////////////////
@@ -1309,14 +1300,21 @@ void FOnDemandContentInstallerStats::OnRequestEnqueued()
 }
 
 void FOnDemandContentInstallerStats::OnRequestCompleted(
+		const FResult& Result,
 		uint64 RequestedChunkCount,
 		uint64 RequestedBytes,
 		uint64 DownloadedChunkCount,
 		uint64 DownloadedBytes,
 		double CacheHitRatio,
-		uint64 DurationCycles,
-		EIoErrorCode ErrorCode) 
+		uint64 DurationCycles) 
 {
+#if COUNTERSTRACE_ENABLED
+	{
+		TUniqueLock Lock(InstallerTraceCounters);
+		InstallerTraceCounters.InflightInstallCount.Decrement();
+	}
+#endif // COUNTERSTRACE_ENABLED
+
 	if (RequestedChunkCount == 0)
 	{
 		return;
@@ -1326,7 +1324,6 @@ void FOnDemandContentInstallerStats::OnRequestCompleted(
 	{
 		TUniqueLock Lock(InstallerTraceCounters);
 		InstallerTraceCounters.InstallCount.Increment();
-		InstallerTraceCounters.InflightInstallCount.Decrement();
 		InstallerTraceCounters.DownloadedBytes.Add(DownloadedBytes);
 		InstallerTraceCounters.InstallDurationMs.Increment(FPlatformTime::ToMilliseconds64(DurationCycles));
 		InstallerTraceCounters.CacheHitRatio.Increment(CacheHitRatio);
@@ -1337,7 +1334,7 @@ void FOnDemandContentInstallerStats::OnRequestCompleted(
 
 	{
 		TUniqueLock Lock(InstallerAnalytics);
-		if (ErrorCode != EIoErrorCode::Ok && ErrorCode != EIoErrorCode::Cancelled)
+		if (Result.HasError() && !UE::UnifiedError::IsCancellationError(Result.GetError()))
 		{
 			InstallerAnalytics.InstallErrorCount++;
 		}
@@ -1415,15 +1412,15 @@ void FOnDemandInstallCacheStats::OnStartupError(EIoErrorCode ErrorCode)
 	}
 }
 
-void FOnDemandInstallCacheStats::OnFlush(EIoErrorCode ErrorCode, int64 ByteCount)
+void FOnDemandInstallCacheStats::OnFlush(const FResult& Result, int64 ByteCount)
 {
-	if (ErrorCode == EIoErrorCode::Cancelled)
+	if (Result.HasError() && UE::UnifiedError::IsCancellationError(Result.GetError()))
 	{
 		return;
 	}
 
 	UE::TUniqueLock Lock(InstallCacheAnalytics);
-	if (ErrorCode != EIoErrorCode::Ok)
+	if (Result.HasError())
 	{
 		InstallCacheAnalytics.FlushErrorCount++;
 	}
@@ -1453,20 +1450,20 @@ void FOnDemandInstallCacheStats::OnCasVerificationError(int32 RemoveChunks)
 }
 
 void FOnDemandInstallCacheStats::OnPurge(
-	const FIoStatus& Status,
+	const FResult& Result,
 	uint64 MaxCacheSize,
 	uint64 NewCacheSize,
 	uint64 BytesToPurge,
 	uint64 PurgedBytes)
 {
-	if (Status.GetErrorCode() == EIoErrorCode::Cancelled)
+	if (Result.HasError() && UE::UnifiedError::IsCancellationError(Result.GetError()))
 	{
 		return;
 	}
 
 	{
 		UE::TUniqueLock Lock(InstallCacheAnalytics);
-		if (Status.IsOk() == false)
+		if (Result.HasError())
 		{
 			InstallCacheAnalytics.PurgeErrorCount++;
 		}
@@ -1474,27 +1471,21 @@ void FOnDemandInstallCacheStats::OnPurge(
 		InstallCacheAnalytics.PurgedBytes += PurgedBytes;
 	}
 
-	SendImmediateAnalytic(FOnDemandImmediateAnalytic
-	{
-		TEXT("IadCachePurge"),
-		MakeAnalyticsEventAttributeArray(
-			TEXT("Success"), Status.IsOk(),
-			TEXT("ErrorCode"), GetAnalyticErrorCode(Status),
-			TEXT("PurgedBytes"), PurgedBytes
-		)
-	});
+	FOnDemandImmediateAnalytic Event = MakeAnalyticsEventFromResult(TEXT("IadCachePurge"), Result);
+	AppendAnalyticsEventAttributeArray(Event.AnalyticsArray, TEXT("PurgedBytes"), PurgedBytes);
+	SendImmediateAnalytic(MoveTemp(Event));
 }
 
-void FOnDemandInstallCacheStats::OnDefrag(const FIoStatus& Status, uint64 FragmentedBytes)
+void FOnDemandInstallCacheStats::OnDefrag(const FResult& Result, uint64 FragmentedBytes)
 {
-	if (Status.GetErrorCode() == EIoErrorCode::Cancelled)
+	if (Result.HasError() && UE::UnifiedError::IsCancellationError(Result.GetError()))
 	{
 		return;
 	}
 
 	{
 		UE::TUniqueLock Lock(InstallCacheAnalytics);
-		if (Status.IsOk() == false)
+		if (Result.HasError())
 		{
 			InstallCacheAnalytics.DefragErrorCount++;
 		}
@@ -1502,15 +1493,9 @@ void FOnDemandInstallCacheStats::OnDefrag(const FIoStatus& Status, uint64 Fragme
 		InstallCacheAnalytics.FragmentedBytes += FragmentedBytes;
 	}
 
-	SendImmediateAnalytic(FOnDemandImmediateAnalytic
-	{
-		TEXT("IadCacheDefrag"),
-		MakeAnalyticsEventAttributeArray(
-			TEXT("Success"), Status.IsOk(),
-			TEXT("ErrorCode"), GetAnalyticErrorCode(Status),
-			TEXT("FragmentedBytes"), FragmentedBytes
-		)
-	});
+	FOnDemandImmediateAnalytic Event = MakeAnalyticsEventFromResult(TEXT("IadCacheDefrag"), Result);
+	AppendAnalyticsEventAttributeArray(Event.AnalyticsArray, TEXT("FragmentedBytes"), FragmentedBytes);
+	SendImmediateAnalytic(MoveTemp(Event));
 }
 
 void FOnDemandInstallCacheStats::OnBlockDeleted(int64 LastAccessTicks, bool bFromDefrag)
diff --git a/Engine/Source/Runtime/Experimental/IoStore/OnDemand/Private/Statistics.h b/Engine/Source/Runtime/Experimental/IoStore/OnDemand/Private/Statistics.h
--- a/Engine/Source/Runtime/Experimental/IoStore/OnDemand/Private/Statistics.h
+++ b/Engine/Source/Runtime/Experimental/IoStore/OnDemand/Private/Statistics.h
@@ -3,6 +3,7 @@
 #pragma once
 
 #include "HAL/LowLevelMemTracker.h"
+#include "IO/OnDemandError.h"
 #include "Misc/CoreDelegates.h"
 #include "OnDemandBackendStatus.h"
 #include "ProfilingDebugging/CountersTrace.h"
@@ -127,13 +128,13 @@ class FOnDemandContentInstallerStats
 public:
 	static void OnRequestEnqueued() IAS_STATISTICS_IMPL()
 	static void OnRequestCompleted(
+		const FResult& Result,
 		uint64 RequestedChunkCount,
 		uint64 RequestedBytes,
 		uint64 DownloadedChunkCount,
 		uint64 DownloadedBytes,
 		double CacheHitRatio,
-		uint64 DurationCycles,
-		EIoErrorCode ErrorCode) IAS_STATISTICS_IMPL()
+		uint64 DurationCycles) IAS_STATISTICS_IMPL()
 	static void ReportAnalytics(TArray<FAnalyticsEventAttribute>& OutAnalyticsArray) IAS_STATISTICS_IMPL()
 };
 
@@ -142,16 +143,16 @@ class FOnDemandInstallCacheStats
 {
 public:
 	static void OnStartupError(EIoErrorCode ErrorCode) IAS_STATISTICS_IMPL()
-	static void OnFlush(EIoErrorCode ErrorCode, int64 ByteCount) IAS_STATISTICS_IMPL()
+	static void OnFlush(const FResult& Result, int64 ByteCount) IAS_STATISTICS_IMPL()
 	static void OnJournalCommit(EIoErrorCode ErrorCode, int64 ByteCount) IAS_STATISTICS_IMPL()
 	static void OnCasVerificationError(int32 RemoveChunks) IAS_STATISTICS_IMPL()
 	static void OnPurge(
-		const FIoStatus& Status,
+		const FResult& Result,	
 		uint64 MaxCacheSize,
 		uint64 NewCacheSize,
 		uint64 BytesToPurge,
 		uint64 PurgedBytes) IAS_STATISTICS_IMPL()
-	static void OnDefrag(const FIoStatus& Status, uint64 FragmentedBytes) IAS_STATISTICS_IMPL()
+	static void OnDefrag(const FResult& Result, uint64 FragmentedBytes) IAS_STATISTICS_IMPL()
 	static void OnBlockDeleted(int64 LastAccessTicks, bool bFromDefrag) IAS_STATISTICS_IMPL()
 	static void OnCacheUsage(
 		uint64 MaxCacheSize,
diff --git a/Engine/Source/Runtime/Experimental/IoStore/OnDemandCore/Private/IoStoreOnDemand.cpp b/Engine/Source/Runtime/Experimental/IoStore/OnDemandCore/Private/IoStoreOnDemand.cpp
--- a/Engine/Source/Runtime/Experimental/IoStore/OnDemandCore/Private/IoStoreOnDemand.cpp
+++ b/Engine/Source/Runtime/Experimental/IoStore/OnDemandCore/Private/IoStoreOnDemand.cpp
@@ -81,13 +81,13 @@ static void PurgeInstallCache(const TArray<FString>& Args)
 		UE_LOG(LogIoStoreOnDemand, Display, TEXT("Purging on demand install cache"));
 		IoStore->Purge(MoveTemp(PurgeArgs), [](const FOnDemandPurgeResult& Result)
 		{
-			if (Result.Status.IsOk())
+			if (Result.IsOk())
 			{
 				UE_LOG(LogIoStoreOnDemand, Display, TEXT("Purged on demand install cache"));
 			}
 			else
 			{
-				UE_LOG(LogIoStoreOnDemand, Error, TEXT("Failed Purged on demand install cache: %s"), *Result.Status.ToString());
+				UE_LOG(LogIoStoreOnDemand, Error, TEXT("Failed Purged on demand install cache: %s"), *LexToString(Result.Error.GetValue()));
 			}
 		});
 	}
@@ -119,13 +119,13 @@ static void DefragInstallCache(const TArray<FString>& Args)
 		UE_LOG(LogIoStoreOnDemand, Display, TEXT("Defragging on demand install cache"));
 		IoStore->Defrag(MoveTemp(DefragArgs), [](const FOnDemandDefragResult& Result)
 		{
-			if (Result.Status.IsOk())
+			if (Result.IsOk())
 			{
 				UE_LOG(LogIoStoreOnDemand, Display, TEXT("Defragmented on demand install cache"));
 			}
 			else
 			{
-				UE_LOG(LogIoStoreOnDemand, Error, TEXT("Failed to defragment on demand install cache: %s"), *Result.Status.ToString());
+				UE_LOG(LogIoStoreOnDemand, Error, TEXT("Failed to defragment on demand install cache: %s"), *LexToString(Result.Error.GetValue()));
 			}
 		});
 	}
@@ -346,14 +346,14 @@ static void InstallPackage(const TArray<FString>& Args)
 		MoveTemp(InstallArgs),
 		[PackageName = MoveTemp(PackageName)](FOnDemandInstallResult InstallResult)
 		{
-			if (InstallResult.Status.IsOk())
+			if (InstallResult.IsOk())
 			{
 				UE_LOG(LogIoStoreOnDemand, Log, TEXT("Successfully installed package '%s'"), *PackageName);
 			}
 			else
 			{
 				UE_LOG(LogIoStoreOnDemand, Error, TEXT("Failed to install package '%s', reason '%s'"),
-					*PackageName, *InstallResult.Status.ToString());
+					*PackageName, *LexToString(InstallResult.Error.GetValue()));
 			}
 		});
 }
@@ -370,22 +370,14 @@ static void VerifyCache()
 
 	IoStore->Verify([](FOnDemandVerifyCacheResult&& VerifyResult)
 	{
-		if (VerifyResult.Status.IsOk())
+		if (VerifyResult.IsOk())
 		{
 			UE_LOG(LogIoStoreOnDemand, Log, TEXT("Install cache verified OK!"));
 		}
 		else
 		{
-			if (VerifyResult.Status.GetErrorCode() == EIoErrorCode::NotFound)
-			{
-				UE_LOG(LogIoStoreOnDemand, Warning, TEXT("Verify install cache failed, reason '%s'"),
-					*VerifyResult.Status.ToString());
-			}
-			else
-			{
-				UE_LOG(LogIoStoreOnDemand, Error, TEXT("Verify install cache failed, reason '%s'"),
-					*VerifyResult.Status.ToString());
-			}
+			UE_LOG(LogIoStoreOnDemand, Error, TEXT("Verify install cache failed, reason '%s'"),
+				*LexToString(VerifyResult.Error.GetValue()));
 		}
 	});
 }
diff --git a/Engine/Source/Runtime/Experimental/IoStore/OnDemandCore/Private/OnDemandError.cpp b/Engine/Source/Runtime/Experimental/IoStore/OnDemandCore/Private/OnDemandError.cpp
new file mode 100644
--- /dev/null
+++ b/Engine/Source/Runtime/Experimental/IoStore/OnDemandCore/Private/OnDemandError.cpp
@@ -0,0 +1,17 @@
+// Copyright Epic Games, Inc. All Rights Reserved.
+
+#include "IO/OnDemandError.h"
+
+UE_DEFINE_ERROR_MODULE(IoStoreOnDemand);
+UE_DEFINE_ERROR(HttpError, IoStoreOnDemand);
+UE_DEFINE_ERROR(ChunkMissingError, IoStoreOnDemand);
+UE_DEFINE_ERROR(ChunkHashError, IoStoreOnDemand);
+UE_DEFINE_ERROR(InstallCacheFlushError, IoStoreOnDemand);
+UE_DEFINE_ERROR(InstallCacheFlushLastAccessError, IoStoreOnDemand);
+UE_DEFINE_ERROR(InstallCachePurgeError, IoStoreOnDemand);
+UE_DEFINE_ERROR(InstallCacheDefragError, IoStoreOnDemand);
+UE_DEFINE_ERROR(InstallCacheVerificationError, IoStoreOnDemand);
+
+namespace UE::UnifiedError::IoStore
+{
+} // namespace UE::UnifiedError::IoStore
diff --git a/Engine/Source/Runtime/Experimental/IoStore/OnDemandCore/Public/IO/IoStoreOnDemand.h b/Engine/Source/Runtime/Experimental/IoStore/OnDemandCore/Public/IO/IoStoreOnDemand.h
--- a/Engine/Source/Runtime/Experimental/IoStore/OnDemandCore/Public/IO/IoStoreOnDemand.h
+++ b/Engine/Source/Runtime/Experimental/IoStore/OnDemandCore/Public/IO/IoStoreOnDemand.h
@@ -6,11 +6,13 @@
 #include "Containers/UnrealString.h"
 #include "Containers/StringFwd.h"
 #include "Containers/SharedString.h"
+#include "IO/OnDemandError.h"
 #include "IO/OnDemandHostGroup.h"
 #include "IO/OnDemandToc.h"
 #include "IO/PackageId.h"
 #include "Misc/EnumClassFlags.h"
 #include "Features/IModularFeature.h"
+#include "Misc/Optional.h"
 #include "Modules/ModuleInterface.h"
 #include "Templates/SharedPointer.h"
 #include "UObject/NameTypes.h"
@@ -276,12 +278,14 @@ using FOnDemandInstallProgressed = TFunction<void(FOnDemandInstallProgress)>;
 /** Holds information about an install request. */
 struct FOnDemandInstallResult
 {
-	/** The status of the install request. */
-	FIoStatus Status;
+	/** Returns True if the request succeeded. */
+	bool IsOk() const { return Error.IsSet() == false; }
 	/** Duration in seconds. */
 	double DurationInSeconds = 0.0;
 	/** Final progress for the install request. */
 	FOnDemandInstallProgress Progress;
+	/** Error information about the request. */ 
+	TOptional<UE::UnifiedError::FError> Error;
 };
 
 /** Install completion callback. */
@@ -341,10 +345,12 @@ struct FOnDemandPurgeArgs
 /** Holds information about a purge request */
 struct FOnDemandPurgeResult
 {
-	/** The status of the purge request. */
-	FIoStatus Status;
+	/** Returns True if the request succeeded. */
+	bool IsOk() const { return Error.IsSet() == false; }
 	/** Duration in seconds. */
 	double DurationInSeconds = 0.0;
+	/** Error information about the request. */ 
+	TOptional<UE::UnifiedError::FError> Error;
 };
 
 /** Purge completion callback */
@@ -372,10 +378,12 @@ struct FOnDemandDefragArgs
 /** Holds information about a defrag request */
 struct FOnDemandDefragResult
 {
-	/** The status of the defrag request. */
-	FIoStatus Status;
+	/** Returns True if the request succeeded. */
+	bool IsOk() const { return Error.IsSet() == false; }
 	/** Duration in seconds. */
 	double DurationInSeconds = 0.0;
+	/** Error information about the request. */ 
+	TOptional<UE::UnifiedError::FError> Error;
 };
 
 /** Defrag completion callback */
@@ -476,10 +484,12 @@ struct FOnDemandCacheUsage
 /** Result from verifying the install cache. */
 struct FOnDemandVerifyCacheResult
 {
-	/** The status of the cache verification request. */
-	FIoStatus Status;
+	/** Returns True if the request succeeded. */
+	bool IsOk() const { return Error.IsSet() == false; }
 	/** Duration in seconds. */
 	double DurationInSeconds = 0.0;
+	/** Error information about the request. */ 
+	TOptional<UE::UnifiedError::FError> Error;
 };
 
 /** Verify cache completion callback. */
diff --git a/Engine/Source/Runtime/Experimental/IoStore/OnDemandCore/Public/IO/OnDemandError.h b/Engine/Source/Runtime/Experimental/IoStore/OnDemandCore/Public/IO/OnDemandError.h
new file mode 100644
--- /dev/null
+++ b/Engine/Source/Runtime/Experimental/IoStore/OnDemandCore/Public/IO/OnDemandError.h
@@ -0,0 +1,157 @@
+// Copyright Epic Games, Inc. All Rights Reserved.
+
+#pragma once
+
+#include "Experimental/UnifiedError/UnifiedError.h"
+#include "Experimental/UnifiedError/CoreErrorTypes.h"
+#include "IO/IoChunkId.h"
+#include "IO/IoStatus.h"
+#include "Templates/ValueOrError.h"
+
+#define UE_API IOSTOREONDEMANDCORE_API
+
+UE_DECLARE_ERROR_MODULE(UE_API, IoStoreOnDemand);
+
+UE_DECLARE_ERROR_ONEPARAM(UE_API,	HttpError,							2,	IoStoreOnDemand, NSLOCTEXT("IoStoreOnDemand", "HttpError", "HTTP error ({StatusCode})"), uint32, StatusCode, 0);
+UE_DECLARE_ERROR(UE_API,			ChunkMissingError,					3,	IoStoreOnDemand, NSLOCTEXT("IoStoreOnDemand", "ChunkMissingError", "Chunk missing error"));
+UE_DECLARE_ERROR(UE_API,			ChunkHashError,						4,	IoStoreOnDemand, NSLOCTEXT("IoStoreOnDemand", "ChunkHashError", "Chunk hash mismatch error"));
+UE_DECLARE_ERROR(UE_API,			InstallCacheFlushError,				5,	IoStoreOnDemand, NSLOCTEXT("IoStoreOnDemand", "InstallCacheFlushError", "Failed to flush pending data to install cache."));
+UE_DECLARE_ERROR(UE_API,			InstallCacheFlushLastAccessError,	6,	IoStoreOnDemand, NSLOCTEXT("IoStoreOnDemand", "InstallCacheFlushLastAccessError", "Failed to flush last access timestamp(s) to journal."));
+UE_DECLARE_ERROR(UE_API,			InstallCachePurgeError,				7,	IoStoreOnDemand, NSLOCTEXT("IoStoreOnDemand", "InstallCachePurgeError", "Failed to purge unreferenced cache block(s) from the install cache."));
+UE_DECLARE_ERROR(UE_API,			InstallCacheDefragError,			8,	IoStoreOnDemand, NSLOCTEXT("IoStoreOnDemand", "InstallCacheDefragError", "Failed to defrag the install cache."));
+UE_DECLARE_ERROR(UE_API,			InstallCacheVerificationError,		9,	IoStoreOnDemand, NSLOCTEXT("IoStoreOnDemand", "InstallCacheVerificationError", "Verification of installed install cache data failed."));
+
+namespace UE::UnifiedError::IoStoreOnDemand
+{
+
+struct FChunkMissingErrorContext
+{
+	TArray<FIoChunkId> ChunkIds;
+};
+
+struct FChunkHashMismatchErrorContext
+{
+	FIoChunkId	ChunkId;
+	FIoHash		ExpectedHash;
+	FIoHash		ActualHash;
+};
+
+struct FIoErrorContext
+{
+	FIoStatus Status;
+};
+
+enum class EErrorCategory
+{
+	None = 0,
+	Cas,
+	Journal,
+	Logical
+};
+
+inline const TCHAR* ToString(EErrorCategory Category)
+{
+	switch (Category)
+	{
+		case EErrorCategory::None:
+			return TEXT("None");
+		case EErrorCategory::Cas:
+			return TEXT("Cas");
+		case EErrorCategory::Journal:
+			return TEXT("Journal");
+		case EErrorCategory::Logical:
+			return TEXT("Logical");
+		default:
+			return TEXT("Unknown");
+	}
+};
+
+struct FInstallCacheErrorContext
+{
+	FIoStatus		Error = FIoStatus(EIoErrorCode::Unknown);
+	uint64			MaxCacheSize = 0;
+	uint64			CacheSize = 0;
+	uint64			DiskTotalBytes = 0;
+	uint64			DiskFreeBytes = 0;
+	EErrorCategory	ErrorCategory = EErrorCategory::None;
+	uint32			LineNo = uint32(__builtin_LINE());
+	bool			bDiskQuerySucceeded = false;
+};
+
+struct FVerificationErrorContext
+{
+	uint32	CorruptChunkCount = 0;
+	uint32	MissingChunkCount = 0;
+	uint32	ReadErrorCount = 0;
+};
+
+} // namespace UE::UnifiedError::IoStoreOnDemand
+UE_DECLARE_ERRORSTRUCT_FEATURES(IoStoreOnDemand, FChunkMissingErrorContext);
+inline void SerializeForLog(FCbWriter& Writer, const UE::UnifiedError::IoStoreOnDemand::FChunkMissingErrorContext& Ctx)
+{
+	Writer.BeginObject();
+	Writer.AddString(ANSITEXTVIEW("$type"), TErrorStructFeatures<UE::UnifiedError::IoStoreOnDemand::FChunkMissingErrorContext>::GetErrorContextTypeNameAsString());
+	Writer.AddString(ANSITEXTVIEW("$format"), TEXT("(ChunkId: {ChunkId})"));
+	//Writer.BeginArray(TEXT("ChunkIds"));
+	//Writer.EndArray();
+	Writer.EndObject();
+}
+UE_DECLARE_ERRORSTRUCT_FEATURES(IoStoreOnDemand, FChunkHashMismatchErrorContext);
+inline void SerializeForLog(FCbWriter& Writer, const UE::UnifiedError::IoStoreOnDemand::FChunkHashMismatchErrorContext& Ctx)
+{
+	Writer.BeginObject();
+	Writer.AddString(ANSITEXTVIEW("$type"), TErrorStructFeatures<UE::UnifiedError::IoStoreOnDemand::FChunkHashMismatchErrorContext>::GetErrorContextTypeNameAsString());
+	Writer.AddString(ANSITEXTVIEW("$format"), TEXT("(ChunkId: {ChunkId})"));
+	Writer.AddString(ANSITEXTVIEW("ChunkId"), LexToString(Ctx.ChunkId));
+	Writer.AddString(ANSITEXTVIEW("ExpectedHash"), LexToString(Ctx.ExpectedHash));
+	Writer.AddString(ANSITEXTVIEW("ActualHash"), LexToString(Ctx.ActualHash));
+	Writer.EndObject();
+}
+UE_DECLARE_ERRORSTRUCT_FEATURES(IoStoreOnDemand, FIoErrorContext);
+inline void SerializeForLog(FCbWriter& Writer, const UE::UnifiedError::IoStoreOnDemand::FIoErrorContext& Ctx)
+{
+	Writer.BeginObject();
+	Writer.AddString(ANSITEXTVIEW("$type"), TErrorStructFeatures<UE::UnifiedError::IoStoreOnDemand::FIoErrorContext>::GetErrorContextTypeNameAsString());
+	Writer.AddString(ANSITEXTVIEW("$format"), TEXT("(Status: {Status})"));
+	Writer.AddString(ANSITEXTVIEW("Status"), Ctx.Status.ToString());
+	Writer.EndObject();
+}
+UE_DECLARE_ERRORSTRUCT_FEATURES(IoStoreOnDemand, FInstallCacheErrorContext);
+inline void SerializeForLog(FCbWriter& Writer, const UE::UnifiedError::IoStoreOnDemand::FInstallCacheErrorContext& Ctx)
+{
+	Writer.BeginObject();
+	Writer.AddString(ANSITEXTVIEW("$type"), TErrorStructFeatures<UE::UnifiedError::IoStoreOnDemand::FInstallCacheErrorContext>::GetErrorContextTypeNameAsString());
+	Writer.AddString(ANSITEXTVIEW("$format"), TEXT("(Category: {ErrorCategoryText}, ErrorCode: {ErrorCodeText}, SystemErrorCode: {SystemErrorCode}, ErrorMessage: {ErrorMessage})"));
+	Writer.AddInteger(ANSITEXTVIEW("ErrorCategory"), int32(Ctx.ErrorCategory));
+	Writer.AddString(ANSITEXTVIEW("ErrorCategoryText"), ToString(Ctx.ErrorCategory));
+	Writer.AddInteger(ANSITEXTVIEW("ErrorCode"), int32(Ctx.Error.GetErrorCode()));
+	Writer.AddString(ANSITEXTVIEW("ErrorCodeText"), GetIoErrorText(Ctx.Error.GetErrorCode()));
+	Writer.AddString(ANSITEXTVIEW("ErrorMessage"), Ctx.Error.ToString());
+	Writer.AddInteger(ANSITEXTVIEW("SystemErrorCode"), int32(Ctx.Error.GetSystemErrorCode()));
+	Writer.AddInteger(ANSITEXTVIEW("MaxCacheSize"), Ctx.MaxCacheSize);
+	Writer.AddInteger(ANSITEXTVIEW("CacheSize"), Ctx.CacheSize);
+	Writer.AddInteger(ANSITEXTVIEW("DiskTotalBytes"), Ctx.DiskTotalBytes);
+	Writer.AddInteger(ANSITEXTVIEW("DiskFreeBytes"), Ctx.DiskFreeBytes);
+	Writer.AddInteger(ANSITEXTVIEW("LineNo"), uint32(Ctx.LineNo));
+	Writer.EndObject();
+}
+UE_DECLARE_ERRORSTRUCT_FEATURES(IoStoreOnDemand, FVerificationErrorContext);
+inline void SerializeForLog(FCbWriter& Writer, const UE::UnifiedError::IoStoreOnDemand::FVerificationErrorContext& Ctx)
+{
+	Writer.BeginObject();
+	Writer.AddString(ANSITEXTVIEW("$type"), TErrorStructFeatures<UE::UnifiedError::IoStoreOnDemand::FVerificationErrorContext>::GetErrorContextTypeNameAsString());
+	Writer.AddString(ANSITEXTVIEW("$format"), TEXT("(CorruptChunkCount: {CorruptChunkCount}, MissingChunkCount: {MissingChunkCount}, ReadErrorCount: {ReadErrorCount})"));
+	Writer.AddInteger(ANSITEXTVIEW("CorruptChunkCount"), Ctx.CorruptChunkCount);
+	Writer.AddInteger(ANSITEXTVIEW("MissingChunkCount"), Ctx.MissingChunkCount);
+	Writer.AddInteger(ANSITEXTVIEW("ReadErrorCount"), Ctx.ReadErrorCount);
+	Writer.EndObject();
+}
+
+namespace UE::IoStore
+{
+
+using FResult = TValueOrError<void, UE::UnifiedError::FError>;
+
+} // namespace UE::UnifiedError::IoStore
+
+#undef UE_API
