I/O store on-demand bughawk fixes.

#rnx

diff --git a/Engine/Source/Runtime/Experimental/IoStore/OnDemand/Private/OnDemandContentInstaller.cpp b/Engine/Source/Runtime/Experimental/IoStore/OnDemand/Private/OnDemandContentInstaller.cpp
--- a/Engine/Source/Runtime/Experimental/IoStore/OnDemand/Private/OnDemandContentInstaller.cpp
+++ b/Engine/Source/Runtime/Experimental/IoStore/OnDemand/Private/OnDemandContentInstaller.cpp
@@ -41,7 +41,7 @@ namespace Private
 ////////////////////////////////////////////////////////////////////////////////
 int32 SecondsToMillicseconds(double Seconds)
 {
-	return Seconds > 0.0 ? int32(Seconds / 1000.0) : 0;
+	return Seconds > 0.0 ? int32(Seconds * 1000.0) : 0;
 }
 
 ////////////////////////////////////////////////////////////////////////////////
diff --git a/Engine/Source/Runtime/Experimental/IoStore/OnDemand/Private/OnDemandInstallCache.cpp b/Engine/Source/Runtime/Experimental/IoStore/OnDemand/Private/OnDemandInstallCache.cpp
--- a/Engine/Source/Runtime/Experimental/IoStore/OnDemand/Private/OnDemandInstallCache.cpp
+++ b/Engine/Source/Runtime/Experimental/IoStore/OnDemand/Private/OnDemandInstallCache.cpp
@@ -3084,6 +3084,7 @@ FResult FOnDemandInstallCache::Defrag(
 		const TBitArray<>& IsReferenced = ChunkEntryIndices[Index++];
 		for (int32 EntryIndex = 0; const FOnDemandChunkEntry& Entry : Container->ChunkEntries)
 		{
+			const FIoChunkId& ChunkId = Container->ChunkIds[EntryIndex];
 			if (bool bIsReferenced = IsReferenced[EntryIndex++]; bIsReferenced == false)
 			{
 				continue;
@@ -3096,7 +3097,7 @@ FResult FOnDemandInstallCache::Defrag(
 					DefragBlock->ReferencedChunks.Add(FDefragBlockReferencedChunk
 					{
 						.Hash = Entry.Hash,
-						.ChunkId = Container->ChunkIds[EntryIndex],
+						.ChunkId = ChunkId,
 						.BlockOffset = Loc.BlockOffset,
 						.DiskSize = Entry.GetDiskSize(),
 					});
diff --git a/Engine/Source/Runtime/Experimental/IoStore/OnDemandCore/Public/IO/OnDemandError.h b/Engine/Source/Runtime/Experimental/IoStore/OnDemandCore/Public/IO/OnDemandError.h
--- a/Engine/Source/Runtime/Experimental/IoStore/OnDemandCore/Public/IO/OnDemandError.h
+++ b/Engine/Source/Runtime/Experimental/IoStore/OnDemandCore/Public/IO/OnDemandError.h
@@ -91,7 +91,8 @@ inline void SerializeForLog(FCbWriter& Writer, const UE::UnifiedError::IoStoreOn
 {
 	Writer.BeginObject();
 	Writer.AddString(ANSITEXTVIEW("$type"), TErrorStructFeatures<UE::UnifiedError::IoStoreOnDemand::FChunkMissingErrorContext>::GetErrorContextTypeNameAsString());
-	Writer.AddString(ANSITEXTVIEW("$format"), TEXT("(ChunkId: {ChunkId})"));
+	Writer.AddString(ANSITEXTVIEW("$format"), TEXT("(MissingCount: {MissingCount})"));
+	Writer.AddInteger(ANSITEXTVIEW("MissingCount"), Ctx.ChunkIds.Num());
 	//Writer.BeginArray(TEXT("ChunkIds"));
 	//Writer.EndArray();
 	Writer.EndObject();
@@ -152,6 +153,6 @@ namespace UE::IoStore
 
 using FResult = TValueOrError<void, UE::UnifiedError::FError>;
 
-} // namespace UE::UnifiedError::IoStore
+} // namespace UE::IoStore
 
 #undef UE_API
