Add a new cvar allowing users to temporarily pause IAX http requests without rejecting or cancelling them.

#rnx

- Calling iax.httppause will stop the processing of IAX requests and move them to a paused queue. Passing in 'IAD' or 'IAS' as an optional arg will cause requests only with that type to be paused. The effect of the cvar is cumalative, so calling ' iax.httppause IAD' followed by ' iax.httppause IAS' will result in all requests being paused.
- Calling iax.httpunpause will allow requests to be processed once more and as with the pause command 'IAD' and 'IAS' may be passed in as an optional arg to selectively unpause a specific type.
- Requests that are paused will be placed in a separate queue, this allows requests that are not currently paused to be processed as per usual.
- When the pause is initiated we do not remove pending requests that should now be paused from the main queue, instead we move them to the paused queue as and when they are processed.
-- In order to either selectively dequeue from the pending queue or dequeue all requests and replace the ones allowed would require more extensive code changes. Since this code is being deprecated by the new backend it would be wasted work.
- The new pause system is controlled by the define UE_ALLOW_HTTP_PAUSE and will be disabled in shipping.
- Added EHttpRequestTypeFilter as we need a way to identify the requests as a bitfield where as EHttpRequestType identifies the requests as an array index.

diff --git a/Engine/Source/Runtime/Experimental/IoStore/OnDemand/Private/OnDemandHttpThread.cpp b/Engine/Source/Runtime/Experimental/IoStore/OnDemand/Private/OnDemandHttpThread.cpp
--- a/Engine/Source/Runtime/Experimental/IoStore/OnDemand/Private/OnDemandHttpThread.cpp
+++ b/Engine/Source/Runtime/Experimental/IoStore/OnDemand/Private/OnDemandHttpThread.cpp
@@ -48,6 +48,15 @@ static const TCHAR* LexToString(EThreadPriority Priority)
 	};
 }
 
+#if UE_ALLOW_HTTP_PAUSE
+
+void FConsoleCommandUnregister::operator()(IConsoleCommand* ConsoleCommandPtr) const
+{
+	IConsoleManager::Get().UnregisterConsoleObject(ConsoleCommandPtr);
+};
+
+#endif // UE_ALLOW_HTTP_PAUSE
+
 namespace UE::IoStore
 {
 
@@ -56,6 +65,8 @@ namespace UE::IoStore
 DECLARE_MULTICAST_DELEGATE(FOnRecreateHttpClient);
 FOnRecreateHttpClient OnRecreateHttpClient;
 
+ENUM_CLASS_FLAGS(EHttpRequestTypeFilter);
+
 void OnHttpClientCVarChanged(IConsoleVariable* CVar)
 {
 	if (ensure(CVar != nullptr))
@@ -246,6 +257,26 @@ FOnDemandHttpThread::FOnDemandHttpThread()
 	{
 		TickThreadEvent->Trigger();
 	});
+
+#if UE_ALLOW_HTTP_PAUSE
+	PauseCommand = FConsoleCommandPtr(IConsoleManager::Get().RegisterConsoleCommand(
+		TEXT("iax.HttpPause"),
+		TEXT("Pause all http requests. Passing in 'IAD' or 'IAD' as an arg to only pause requests of that type"),
+		FConsoleCommandWithArgsAndOutputDeviceDelegate::CreateLambda([this](const TArray<FString>& Args, FOutputDevice& Ar) -> void
+			{
+				OnPauseCommand(/*bShouldPause*/ true, Args, Ar);
+			}),
+		ECVF_Default));
+
+	UnpauseCommand = FConsoleCommandPtr(IConsoleManager::Get().RegisterConsoleCommand(
+		TEXT("iax.HttpUnpause"),
+		TEXT("Unause all http requests. Passing in 'IAD' or 'IAD' as an arg to only unpause requests of that type"),
+		FConsoleCommandWithArgsAndOutputDeviceDelegate::CreateLambda([this](const TArray<FString>& Args, FOutputDevice& Ar) -> void
+			{
+				OnPauseCommand(/*bShouldPause*/ false, Args, Ar);
+			}),
+		ECVF_Default));
+#endif // UE_ALLOW_HTTP_PAUSE
 }
 
 FOnDemandHttpThread::~FOnDemandHttpThread()
@@ -286,7 +317,15 @@ FOnDemandHttpThread::FRequestHandle FOnDemandHttpThread::IssueRequest(FOnDemandC
 
 	FOnDemandIoBackendStats::Get()->OnHttpEnqueue(Request->Type);
 
+#if UE_ALLOW_HTTP_PAUSE
+	if (!PauseRequest(Request))
+	{
+		HttpRequests.EnqueueByPriority(Request);
+	}
+#else
 	HttpRequests.EnqueueByPriority(Request);
+#endif // UE_ALLOW_HTTP_PAUSE
+
 #if UE_IOSTORE_ONDEMAND_NO_HTTP_THREAD 
 	// When threading is disabled the completion callback has already been triggered and the handle is destroyed before returned to the caller.
 	Tick();
@@ -409,6 +448,13 @@ void FOnDemandHttpThread::TickRequests()
 				NextHttpRequest = HttpRequest->NextRequest;
 				HttpRequest->NextRequest = nullptr;
 
+#if UE_ALLOW_HTTP_PAUSE
+				if (PauseRequest(HttpRequest))
+				{
+					continue;
+				}
+#endif // UE_ALLOW_HTTP_PAUSE
+
 				FOnDemandIoBackendStats::Get()->OnHttpDequeue(HttpRequest->Type);
 
 				if (HttpRequest->bCancelled)
@@ -524,6 +570,10 @@ void FOnDemandHttpThread::TickRequests()
 
 void FOnDemandHttpThread::DrainHttpRequests()
 {
+#if UE_ALLOW_HTTP_PAUSE
+	OnTogglePause(false, EHttpRequestTypeFilter::All);
+#endif // UE_ALLOW_HTTP_PAUSE
+
 	FHttpRequest* Iterator = HttpRequests.Dequeue();
 	while (Iterator != nullptr)
 	{
@@ -582,6 +632,107 @@ void FOnDemandHttpThread::DestroyRequest(FHttpRequest* Request)
 	RequestAllocator.Destroy(Request);
 }
 
+#if UE_ALLOW_HTTP_PAUSE
+
+bool FOnDemandHttpThread::PauseRequest(FHttpRequest* Request)
+{
+	UE::TUniqueLock _(PausedMutex);
+
+	if (IsPaused(Request->Type))
+	{
+		HttpPausedRequests.Enqueue(Request);
+		return true;
+	}
+	else
+	{
+		return false;
+	}
+}
+
+void FOnDemandHttpThread::OnPauseCommand(bool bShouldPause, const TArray<FString>& Args, FOutputDevice& Ar)
+{
+	if (Args.Num() > 1)
+	{
+		Ar.Log(LogIas.GetCategoryName(), ELogVerbosity::Error, TEXT("Too many args for command, either 0 (all) or 1 (IAD/IAS) expected"));
+		return;
+	}
+
+	EHttpRequestTypeFilter Filter = EHttpRequestTypeFilter::All;
+
+	if (!Args.IsEmpty())
+	{
+		if (Args[0] == TEXT("IAD"))
+		{
+			Filter = EHttpRequestTypeFilter::Installed;
+		}
+		else if (Args[0] == TEXT("IAS"))
+		{
+			Filter = EHttpRequestTypeFilter::Streaming;
+		}
+		else
+		{
+			Ar.Log(LogIas.GetCategoryName(), ELogVerbosity::Error, TEXT("Invalid arg for command, expecting 'IAD' or 'IAS'"));
+			return;
+		}
+	}
+
+	OnTogglePause(bShouldPause, Filter);
+}
+
+void FOnDemandHttpThread::OnTogglePause(bool bPause, EHttpRequestTypeFilter Filter)
+{
+	UE::TUniqueLock _(PausedMutex);
+	
+	if (bPause)
+	{
+		EnumAddFlags(PauseFilter, Filter);
+
+
+	}
+	else
+	{
+		EnumRemoveFlags(PauseFilter, Filter);
+
+		FHttpRequest* PendingRequests = HttpPausedRequests.Dequeue();
+		while (PendingRequests != nullptr)
+		{
+			FHttpRequest* RequestToQueue = PendingRequests;
+			PendingRequests = PendingRequests->NextRequest;
+
+			RequestToQueue->NextRequest = nullptr;
+
+			if (IsPaused(RequestToQueue->Type))
+			{
+				HttpPausedRequests.Enqueue(RequestToQueue);
+			}
+			else
+			{
+				HttpRequests.EnqueueByPriority(RequestToQueue);
+			}
+		}
+
+		TickThreadEvent->Trigger();
+	}
+}
+
+bool FOnDemandHttpThread::IsPaused(EHttpRequestType RequestType) const
+{
+	switch (RequestType)
+	{
+		case EHttpRequestType::Installed:
+			return EnumHasAnyFlags(PauseFilter, EHttpRequestTypeFilter::Installed);
+		case EHttpRequestType::Streaming:
+			return EnumHasAnyFlags(PauseFilter, EHttpRequestTypeFilter::Streaming);
+		default:
+		{
+			checkNoEntry();
+			return false;
+		}
+	}
+}
+
+#endif // UE_ALLOW_HTTP_PAUSE
+
 } //namespace UE::IoStore
 
 #undef UE_ALLOW_INVALID_URL_DEBUGGING
diff --git a/Engine/Source/Runtime/Experimental/IoStore/OnDemand/Private/OnDemandHttpThread.h b/Engine/Source/Runtime/Experimental/IoStore/OnDemand/Private/OnDemandHttpThread.h
--- a/Engine/Source/Runtime/Experimental/IoStore/OnDemand/Private/OnDemandHttpThread.h
+++ b/Engine/Source/Runtime/Experimental/IoStore/OnDemand/Private/OnDemandHttpThread.h
@@ -15,8 +15,24 @@
 #include "Templates/UniquePtr.h"
 #include "ThreadSafeIntrusiveQueue.h"
 
+#define UE_ALLOW_HTTP_PAUSE (1 && !NO_CVARS && !UE_BUILD_SHIPPING)
+
 class FRunnableThread;
 
+#if UE_ALLOW_HTTP_PAUSE
+
+struct IConsoleCommand;
+
+struct FConsoleCommandUnregister
+{
+	void operator()(IConsoleCommand* ConsoleCommandPtr) const;
+};
+
+/** Experimental wrapper around a registered IConsoleCommand */
+using FConsoleCommandPtr = TUniquePtr<IConsoleCommand, FConsoleCommandUnregister>;
+
+#endif // UE_ALLOW_HTTP_PAUSE
+
 namespace UE::IoStore
 {
 
@@ -44,6 +60,17 @@ enum class EHttpRequestType : uint8
 	NUM_SOURCES
 };
 
+/** Allows requests to be filtered by type */
+enum class EHttpRequestTypeFilter : uint8
+{
+	None  = 0,
+
+	Streaming = 1 << 0,	/** IAS */
+	Installed = 1 << 1,	/** IAD */
+
+	All = Streaming | Installed
+};
+
 class FOnDemandHttpThread final
 	: public FRunnable
 	, public FSingleThreadRunnable
@@ -114,6 +141,15 @@ private:
 
 	void DestroyRequest(FHttpRequest* Request);
 
+#if UE_ALLOW_HTTP_PAUSE
+	/** Returns true if the request is now paused and has been places in the HttpPausedRequests queue. */
+	bool PauseRequest(FHttpRequest* Request);
+	void OnPauseCommand(bool bPause, const TArray<FString>& Args, FOutputDevice& Ar);
+	void OnTogglePause(bool bPause, EHttpRequestTypeFilter Filter);
+
+	bool IsPaused(EHttpRequestType RequestType) const;
+#endif // UE_ALLOW_HTTP_PAUSE
+
 private:
 	TUniquePtr<FRunnableThread> Thread;
 
@@ -129,6 +165,15 @@ private:
 
 	TThreadSafeIntrusiveQueue<FHttpRequest> HttpRequests;
 
+#if UE_ALLOW_HTTP_PAUSE
+	TThreadSafeIntrusiveQueue<FHttpRequest> HttpPausedRequests;
+	FConsoleCommandPtr PauseCommand;
+	FConsoleCommandPtr UnpauseCommand;
+
+	UE::FMutex PausedMutex;
+	EHttpRequestTypeFilter  PauseFilter = EHttpRequestTypeFilter::None;
+#endif // UE_ALLOW_HTTP_PAUSE
+
 	EThreadPriority ThreadPriority = EThreadPriority::TPri_Normal;
 
 	std::atomic_bool bStopRequested = false;
@@ -140,3 +185,5 @@ private:
 };
 
 } //namespace UE::IoStore
+
+//#undef UE_ALLOW_HTTP_PAUSE
