Case normalisation

#rnx

diff --git a/Engine/Source/Runtime/Experimental/IoStore/HttpClient/Private/Activity.inl b/Engine/Source/Runtime/Experimental/IoStore/HttpClient/Private/Activity.inl
--- a/Engine/Source/Runtime/Experimental/IoStore/HttpClient/Private/Activity.inl
+++ b/Engine/Source/Runtime/Experimental/IoStore/HttpClient/Private/Activity.inl
@@ -405,7 +405,7 @@ FOutcome FActivity::Transact(FHttpPeer& Peer)
 
 	Transaction->Begin(GetMethod(), GetPath());
 
-	Transaction->AddHeader("Host", GetHost()->GetHostName());
+	Transaction->AddHeader("host", GetHost()->GetHostName());
 
 	// HTTP/1.1 is persistent by default thus "Connection" header isn't required
 	if (!GetHost()->IsPooled())
