Additional check() for piece of mind

#rnx

diff --git a/Engine/Source/Runtime/Experimental/IoStore/HttpClient/Private/Loop.inl b/Engine/Source/Runtime/Experimental/IoStore/HttpClient/Private/Loop.inl
--- a/Engine/Source/Runtime/Experimental/IoStore/HttpClient/Private/Loop.inl
+++ b/Engine/Source/Runtime/Experimental/IoStore/HttpClient/Private/Loop.inl
@@ -258,6 +258,7 @@ void FActivityList::Prepend(FActivityNode* Node)
 ////////////////////////////////////////////////////////////////////////////////
 void FActivityList::Append(FActivityNode* Node)
 {
+	check(Node->Next == nullptr);
 	if (Head == nullptr)
 	{
 		Head = Node;
