Error out if length if not as we would expect by this point

#rnx

diff --git a/Engine/Source/Runtime/Experimental/IoStore/HttpClient/Private/Activity.inl b/Engine/Source/Runtime/Experimental/IoStore/HttpClient/Private/Activity.inl
--- a/Engine/Source/Runtime/Experimental/IoStore/HttpClient/Private/Activity.inl
+++ b/Engine/Source/Runtime/Experimental/IoStore/HttpClient/Private/Activity.inl
@@ -492,6 +492,10 @@ FOutcome FActivity::RecvMessage(FHttpPeer& Peer)
 			return Outcome;
 		}
 	}
+	else if (ContentLength < 0)
+	{
+		return FOutcome::Error("Invalid content length");
+	}
 
 	// Call out to the sink to get a content destination
 	FIoBuffer* PriorDest = Dest; // to retain unioned Host ptr (redirect uses it in sink)
