CIS: Fix compile issue when NO_LOGGING is enabled.

#rnx

- When NO_LOGGING is enabled, GetCategoryName is not valid so we need to define around it.

diff --git a/Engine/Source/Runtime/Experimental/IoStore/OnDemand/Private/OnDemandHttpThread.cpp b/Engine/Source/Runtime/Experimental/IoStore/OnDemand/Private/OnDemandHttpThread.cpp
--- a/Engine/Source/Runtime/Experimental/IoStore/OnDemand/Private/OnDemandHttpThread.cpp
+++ b/Engine/Source/Runtime/Experimental/IoStore/OnDemand/Private/OnDemandHttpThread.cpp
@@ -653,7 +653,9 @@ void FOnDemandHttpThread::OnPauseCommand(bool bShouldPause, const TArray<FString
 {
 	if (Args.Num() > 1)
 	{
+#if !NO_LOGGING
 		Ar.Log(LogIas.GetCategoryName(), ELogVerbosity::Error, TEXT("Too many args for command, either 0 (all) or 1 (IAD/IAS) expected"));
+#endif // !NO_LOGGING
 		return;
 	}
 
@@ -671,7 +673,9 @@ void FOnDemandHttpThread::OnPauseCommand(bool bShouldPause, const TArray<FString
 		}
 		else
 		{
+#if !NO_LOGGING
 			Ar.Log(LogIas.GetCategoryName(), ELogVerbosity::Error, TEXT("Invalid arg for command, expecting 'IAD' or 'IAS'"));
+#endif // !NO_LOGGING
 			return;
 		}
 	}
