Set fail state prior to returning the error.

#rnx

diff --git a/Engine/Source/Runtime/Experimental/IoStore/HttpClient/Private/Activity.inl b/Engine/Source/Runtime/Experimental/IoStore/HttpClient/Private/Activity.inl
--- a/Engine/Source/Runtime/Experimental/IoStore/HttpClient/Private/Activity.inl
+++ b/Engine/Source/Runtime/Experimental/IoStore/HttpClient/Private/Activity.inl
@@ -494,7 +494,9 @@ FOutcome FActivity::RecvMessage(FHttpPeer& Peer)
 	}
 	else if (ContentLength < 0)
 	{
-		return FOutcome::Error("Invalid content length");
+		Outcome = FOutcome::Error("Invalid content length");
+		Fail(Outcome);
+		return Outcome;
 	}
 
 	// Call out to the sink to get a content destination
