Lowercase header field names.

#rnx

diff --git a/Engine/Source/Runtime/Experimental/IoStore/OnDemand/Private/DistributionEndpoints.cpp b/Engine/Source/Runtime/Experimental/IoStore/OnDemand/Private/DistributionEndpoints.cpp
--- a/Engine/Source/Runtime/Experimental/IoStore/OnDemand/Private/DistributionEndpoints.cpp
+++ b/Engine/Source/Runtime/Experimental/IoStore/OnDemand/Private/DistributionEndpoints.cpp
@@ -106,7 +106,7 @@ FDistributionEndpoints::EResult FDistributionEndpoints::ResolveEndpoints(const F
 
 	FEventLoop::FRequestParams RequestParams = { .bAllowChunked = true };
 	FRequest Request = Loop.Get(AnsiUrl, &RequestParams);
-	Request.Header("Accept", "application/json");
+	Request.Header("accept", "application/json");
 	Loop.Send(MoveTemp(Request), OnRequestStatus);
 
 	while (Loop.Tick(-1))
diff --git a/Engine/Source/Runtime/Experimental/IoStore/OnDemand/Private/OnDemandHttpClient.cpp b/Engine/Source/Runtime/Experimental/IoStore/OnDemand/Private/OnDemandHttpClient.cpp
--- a/Engine/Source/Runtime/Experimental/IoStore/OnDemand/Private/OnDemandHttpClient.cpp
+++ b/Engine/Source/Runtime/Experimental/IoStore/OnDemand/Private/OnDemandHttpClient.cpp
@@ -308,7 +308,7 @@ void FMultiEndpointHttpClient::IssueRequest(FRequest&& Request)
 	UE::IoStore::HTTP::FRequest HttpRequest = EventLoop.Get(Url, ConnectionPool, &RequestParams);
 
 #if !UE_BUILD_TEST
-	HttpRequest.Header("Pragma", "akamai-x-cache-on");
+	HttpRequest.Header("pragma", "akamai-x-cache-on");
 #endif
 
 	check(Request.Headers.IsEmpty() || ((Request.Headers.Num() % 2) == 0));
@@ -322,7 +322,7 @@ void FMultiEndpointHttpClient::IssueRequest(FRequest&& Request)
 
 	if (GAllowPartialContentRequests && Request.Range.IsValid())
 	{
-		HttpRequest.Header(ANSITEXTVIEW("Range"),
+		HttpRequest.Header(ANSITEXTVIEW("range"),
 			WriteToAnsiString<64>(ANSITEXTVIEW("bytes="), Request.Range.GetOffset(), ANSITEXTVIEW("-"), Request.Range.GetOffset() + Request.Range.GetLength() - 1));
 	}
 
