Hide a few version-specific details from activity instances.

#rnx

diff --git a/Engine/Source/Runtime/Experimental/IoStore/HttpClient/Private/Activity.inl b/Engine/Source/Runtime/Experimental/IoStore/HttpClient/Private/Activity.inl
--- a/Engine/Source/Runtime/Experimental/IoStore/HttpClient/Private/Activity.inl
+++ b/Engine/Source/Runtime/Experimental/IoStore/HttpClient/Private/Activity.inl
@@ -403,9 +403,8 @@ FOutcome FActivity::Transact(FHttpPeer& Peer)
 	Transaction = Peer.Transact(Buffer);
 	check(Transaction.IsValid());
 
-	Transaction->Begin(GetMethod(), GetPath());
-
-	Transaction->AddHeader("host", GetHost()->GetHostName());
+	const FHost* TheHost = GetHost();
+	Transaction->Begin(TheHost->GetHostName(), GetMethod(), GetPath());
 
 	// HTTP/1.1 is persistent by default thus "Connection" header isn't required
 	if (!GetHost()->IsPooled())
diff --git a/Engine/Source/Runtime/Experimental/IoStore/HttpClient/Private/TransactionOne.inl b/Engine/Source/Runtime/Experimental/IoStore/HttpClient/Private/TransactionOne.inl
--- a/Engine/Source/Runtime/Experimental/IoStore/HttpClient/Private/TransactionOne.inl
+++ b/Engine/Source/Runtime/Experimental/IoStore/HttpClient/Private/TransactionOne.inl
@@ -207,7 +207,7 @@ class FRequest
 public:
 	using FBase::FBase;
 
-	void			Begin(FAnsiStringView Method, FAnsiStringView Path);
+	void			Begin(FAnsiStringView Host, FAnsiStringView Method, FAnsiStringView Path);
 	void			AddHeader(FAnsiStringView Key, FAnsiStringView Value);
 	FTransactId		End();
 	FOutcome		TrySendRequest(FTlsPeer& Peer);
@@ -222,13 +222,15 @@ private:
 };
 
 ////////////////////////////////////////////////////////////////////////////////
-void FRequest::Begin(FAnsiStringView Method, FAnsiStringView Path)
+void FRequest::Begin(FAnsiStringView Host, FAnsiStringView Method, FAnsiStringView Path)
 {
 	check(MethodLength < 0);
 	AlreadySent = uint16(GetSize());
 	*this << Method << " " << Path << " HTTP/1.1" "\r\n";
 	MethodLength = int8(Method.Len());
 	HeaderLeft = uint16(GetSize());
+
+	AddHeader("Host", Host);
 }
 
 ////////////////////////////////////////////////////////////////////////////////
