Hide a few version-specific details from activity instances.

#rnx

diff --git a/Engine/Source/Runtime/Experimental/IoStore/HttpClient/Private/Activity.inl b/Engine/Source/Runtime/Experimental/IoStore/HttpClient/Private/Activity.inl
--- a/Engine/Source/Runtime/Experimental/IoStore/HttpClient/Private/Activity.inl
+++ b/Engine/Source/Runtime/Experimental/IoStore/HttpClient/Private/Activity.inl
@@ -404,21 +404,16 @@ FOutcome FActivity::Transact(FHttpPeer& Peer)
 	check(Transaction.IsValid());
 
 	const FHost* TheHost = GetHost();
-	Transaction->Begin(TheHost->GetHostName(), GetMethod(), GetPath());
-
-	// HTTP/1.1 is persistent by default thus "Connection" header isn't required
-	if (!GetHost()->IsPooled())
-	{
-		Transaction->AddHeader("Connection", "close");
-	}
+	bool bKeepAlive = !TheHost->IsPooled();
 
+	Transaction->Begin(TheHost->GetHostName(), GetMethod(), GetPath());
 	EnumerateHeaders([this] (FAnsiStringView Key, FAnsiStringView Value)
 	{
 		Transaction->AddHeader(Key, Value);
 		return true;
 	});
+	TransactId = Transaction->End(bKeepAlive);
 
-	TransactId = Transaction->End();
 	ChangeState(EState::Send);
 	return FOutcome::Ok();
 }
diff --git a/Engine/Source/Runtime/Experimental/IoStore/HttpClient/Private/TransactionOne.inl b/Engine/Source/Runtime/Experimental/IoStore/HttpClient/Private/TransactionOne.inl
--- a/Engine/Source/Runtime/Experimental/IoStore/HttpClient/Private/TransactionOne.inl
+++ b/Engine/Source/Runtime/Experimental/IoStore/HttpClient/Private/TransactionOne.inl
@@ -209,7 +209,7 @@ public:
 
 	void			Begin(FAnsiStringView Host, FAnsiStringView Method, FAnsiStringView Path);
 	void			AddHeader(FAnsiStringView Key, FAnsiStringView Value);
-	FTransactId		End();
+	FTransactId		End(bool bKeepAlive);
 	FOutcome		TrySendRequest(FTlsPeer& Peer);
 
 private:
@@ -241,8 +241,15 @@ void FRequest::AddHeader(FAnsiStringView Key, FAnsiStringView Value)
 }
 
 ////////////////////////////////////////////////////////////////////////////////
-FTransactId FRequest::End()
+FTransactId FRequest::End(bool bKeepAlive)
 {
+	// HTTP/1.1 is persistent by default thus "Connection" header isn't required
+	// unless we want to opt in to a single transaction.
+	if (!bKeepAlive)
+	{
+		AddHeader("Connection", "close");
+	}
+
 	*this << "\r\n";
 	HeaderRight = int16(GetSize());
 	return 1;
