Fix typos in IAX debug cvar help text.

#rnx

diff --git a/Engine/Source/Runtime/Experimental/IoStore/OnDemand/Private/OnDemandHttpThread.cpp b/Engine/Source/Runtime/Experimental/IoStore/OnDemand/Private/OnDemandHttpThread.cpp
--- a/Engine/Source/Runtime/Experimental/IoStore/OnDemand/Private/OnDemandHttpThread.cpp
+++ b/Engine/Source/Runtime/Experimental/IoStore/OnDemand/Private/OnDemandHttpThread.cpp
@@ -261,7 +261,7 @@ FOnDemandHttpThread::FOnDemandHttpThread()
 #if UE_ALLOW_HTTP_PAUSE
 	PauseCommand = FConsoleCommandPtr(IConsoleManager::Get().RegisterConsoleCommand(
 		TEXT("iax.HttpPause"),
-		TEXT("Pause all http requests. Passing in 'IAD' or 'IAD' as an arg to only pause requests of that type"),
+		TEXT("Pause all http requests. Passing in 'IAD' or 'IAS' as an arg to only pause requests of that type"),
 		FConsoleCommandWithArgsAndOutputDeviceDelegate::CreateLambda([this](const TArray<FString>& Args, FOutputDevice& Ar) -> void
 			{
 				OnPauseCommand(/*bShouldPause*/ true, Args, Ar);
@@ -270,7 +270,7 @@ FOnDemandHttpThread::FOnDemandHttpThread()
 
 	UnpauseCommand = FConsoleCommandPtr(IConsoleManager::Get().RegisterConsoleCommand(
 		TEXT("iax.HttpUnpause"),
-		TEXT("Unause all http requests. Passing in 'IAD' or 'IAD' as an arg to only unpause requests of that type"),
+		TEXT("Unause all http requests. Passing in 'IAD' or 'IAS' as an arg to only unpause requests of that type"),
 		FConsoleCommandWithArgsAndOutputDeviceDelegate::CreateLambda([this](const TArray<FString>& Args, FOutputDevice& Ar) -> void
 			{
 				OnPauseCommand(/*bShouldPause*/ false, Args, Ar);
