Expose a choice of version to users.

#rnx

diff --git a/Engine/Source/Runtime/Experimental/IoStore/HttpClient/Internal/IO/Http/Client.h b/Engine/Source/Runtime/Experimental/IoStore/HttpClient/Internal/IO/Http/Client.h
--- a/Engine/Source/Runtime/Experimental/IoStore/HttpClient/Internal/IO/Http/Client.h
+++ b/Engine/Source/Runtime/Experimental/IoStore/HttpClient/Internal/IO/Http/Client.h
@@ -19,6 +19,14 @@ class FIoBuffer;
 namespace UE::IoStore::HTTP
 {
 
+////////////////////////////////////////////////////////////////////////////////
+enum class EHttpVersion : uint8
+{
+	Default,
+	One,		// 1.1 in actuality
+	Two,
+};
+
 ////////////////////////////////////////////////////////////////////////////////
 enum class EMimeType
 {
@@ -86,6 +94,7 @@ public:
 		int32				RecvBufSize = -1;
 		uint32				Port = 0;
 		uint16				ConnectionCount = 1;
+		EHttpVersion		HttpVersion = EHttpVersion::One;
 		/*
 		enum class ProxyType { Http, Socks4 };
 		Proxy = { ip, port, type }
@@ -227,6 +236,7 @@ public:
 		uint16						BufferSize		= 256;
 		bool						bAutoRedirect	= false;
 		bool						bAllowChunked	= true;
+		EHttpVersion				HttpVersion		= EHttpVersion::Default;
 	};
 
 	template <typename... T> [[nodiscard]] FRequest Get(T&&... t)  { return Request("GET",  Forward<T&&>(t)...); }
diff --git a/Engine/Source/Runtime/Experimental/IoStore/HttpClient/Private/ConnectionPool.inl b/Engine/Source/Runtime/Experimental/IoStore/HttpClient/Private/ConnectionPool.inl
--- a/Engine/Source/Runtime/Experimental/IoStore/HttpClient/Private/ConnectionPool.inl
+++ b/Engine/Source/Runtime/Experimental/IoStore/HttpClient/Private/ConnectionPool.inl
@@ -18,6 +18,7 @@ public:
 		uint32			Port = 0;
 		uint16			MaxConnections = 1;
 		bool			bPooled = false;
+		EHttpVersion	HttpVersion = EHttpVersion::One;
 		FCertRootsRef	VerifyCert = {};
 	};
 
@@ -29,6 +30,7 @@ public:
 	FOutcome		ResolveHostName();
 	FCertRootsRef	GetVerifyCert() const		{ return VerifyCert; }
 	uint32			GetMaxConnections() const	{ return MaxConnections; }
+	EHttpVersion	GetHttpVersion() const		{ return HttpVersion; }
 	bool			IsPooled() const			{ return bPooled; }
 	uint32			GetIpAddress() const		{ return IpAddresses[0]; }
 	FAnsiStringView	GetHostName() const			{ return HostName; }
@@ -42,6 +44,7 @@ private:
 	int16			RecvBufKb = -1;
 	uint16			Port;
 	uint8			MaxConnections;
+	EHttpVersion	HttpVersion;
 	bool			bPooled;
 };
 
@@ -51,6 +54,7 @@ FHost::FHost(const FParams& Params)
 , HostName(Params.HostName)
 , Port(uint16(Params.Port))
 , MaxConnections(uint8(Params.MaxConnections))
+, HttpVersion(Params.HttpVersion)
 , bPooled(Params.bPooled)
 {
 	check(MaxConnections && MaxConnections == Params.MaxConnections);
@@ -252,6 +256,7 @@ FConnectionPool::FConnectionPool(const FParams& Params)
 		.Port			= Params.Port,
 		.MaxConnections	= Params.ConnectionCount,
 		.bPooled		= true,
+		.HttpVersion	= Params.HttpVersion,
 		.VerifyCert		= Params.VerifyCert,
 	});
 	Internal->SetBufferSize(FHost::EDirection::Send, Params.SendBufSize);
diff --git a/Engine/Source/Runtime/Experimental/IoStore/HttpClient/Private/Loop.inl b/Engine/Source/Runtime/Experimental/IoStore/HttpClient/Private/Loop.inl
--- a/Engine/Source/Runtime/Experimental/IoStore/HttpClient/Private/Loop.inl
+++ b/Engine/Source/Runtime/Experimental/IoStore/HttpClient/Private/Loop.inl
@@ -32,6 +32,7 @@ struct FActivityNode
 		: public FActivity::FParams
 	{
 		bool				bFollow30x = false;
+		EHttpVersion		HttpVersion = EHttpVersion::Default;
 	};
 
 	static FActivityNode*	Create(FParams& Params, FAnsiStringView Url={}, FCertRootsRef VerifyCert={});
@@ -130,9 +131,13 @@ FActivityNode* FActivityNode::Create(FParams& Params, FAnsiStringView Url, FCert
 	memcpy(HostNamePtr, HostName.GetData(), HostNameLength);
 	HostNamePtr[HostNameLength] = '\0';
 
+	EHttpVersion Version = Params.HttpVersion;
+	Version = (Version != EHttpVersion::Two) ? EHttpVersion::One : Version;
+
 	new (Host) FHost({
 		.HostName	= HostNamePtr,
 		.Port		= Port,
+		.HttpVersion= Version,
 		.VerifyCert = VerifyCert,
 	});
 
@@ -693,7 +698,7 @@ void FPeerGroup::TickSend(FTickState& State, FHost& Host, FPoller& Poller)
 			.Socket = MoveTemp(Socket),
 			.Certs = Host.GetVerifyCert(),
 			.HostName = Host.GetHostName().GetData(),
-			.Version = 1,
+			.HttpVersion = Host.GetHttpVersion(),
 		};
 		Peer = FHttpPeer(MoveTemp(Params));
 	}
@@ -1149,6 +1154,7 @@ FRequest FEventLoop::Request(
 	ActivityParams.bFollow30x = (Params->bAutoRedirect == true);
 	ActivityParams.bAllowChunked = (Params->bAllowChunked == true);
 	ActivityParams.ContentSizeEst = Params->ContentSizeEst;
+	ActivityParams.HttpVersion = Params->HttpVersion;
 
 	FActivityNode* Activity = FActivityNode::Create(ActivityParams, Url, Params->VerifyCert);
 	if (Activity == nullptr)
@@ -1168,6 +1174,7 @@ FRequest FEventLoop::Request(
 {
 	check(Pool.Ptr != nullptr);
 	check(Params == nullptr || Params->VerifyCert == ECertRootsRefType::None); // add cert to FConPool instead
+	check(Params == nullptr || Params->HttpVersion == EHttpVersion::Default); // set protocol version via FConPool
 
 	Params = (Params != nullptr) ? Params : &GDefaultParams;
 
diff --git a/Engine/Source/Runtime/Experimental/IoStore/HttpClient/Private/Peer.inl b/Engine/Source/Runtime/Experimental/IoStore/HttpClient/Private/Peer.inl
--- a/Engine/Source/Runtime/Experimental/IoStore/HttpClient/Private/Peer.inl
+++ b/Engine/Source/Runtime/Experimental/IoStore/HttpClient/Private/Peer.inl
@@ -257,8 +257,8 @@ static void Ssl_SetupAlpn(SSL* Ssl, int32 HttpVersion)
 	FAnsiStringView AlpnProtos;
 	switch (HttpVersion)
 	{
-	case 1:		AlpnProtos = "\x08" "http/1.1";
-//	case 2:		AlpnProtos = "\x02" "h2";
+	case 1:		AlpnProtos = "\x08" "http/1.1";	break;
+//	case 2:		AlpnProtos = "\x02" "h2";		break;
 	default:	break;
 	}
 
@@ -532,7 +532,7 @@ public:
 		FSocket			Socket;
 		FCertRootsRef 	Certs = ECertRootsRefType::None;
 		const char* 	HostName = nullptr;
-		uint8			Version = 1;
+		EHttpVersion	HttpVersion = EHttpVersion::One;
 	};
 
 					FHttpPeer() = default;
@@ -552,7 +552,7 @@ private:
 ////////////////////////////////////////////////////////////////////////////////
 FHttpPeer::FHttpPeer(FParams&& Params)
 : FTlsPeer(MoveTemp(Params.Socket), Params.Certs, Params.HostName)
-, Proto(Params.Version)
+, Proto(uint8(Params.HttpVersion))
 {
 	if (Ssl != nullptr)
 	{
