[NetPhysics] Optimize vector distance comparison operation when comparing rewind history.
And default to allowing resim to incomplete frames if there are no complete frames.

diff --git a/Engine/Source/Runtime/Experimental/Chaos/Private/RewindData.cpp b/Engine/Source/Runtime/Experimental/Chaos/Private/RewindData.cpp
--- a/Engine/Source/Runtime/Experimental/Chaos/Private/RewindData.cpp
+++ b/Engine/Source/Runtime/Experimental/Chaos/Private/RewindData.cpp
@@ -389,7 +389,7 @@ const int32 FRewindData::CompareTargetsToLastFrame()
 bool FRewindData::CheckVectorThreshold(FVec3 A, FVec3 B, float Threshold)
 {
 	const FVector Delta = A - B;
-	return Delta.Size() >= Threshold;
+	return Delta.SizeSquared() >= (Threshold * Threshold);
 }
 
 bool FRewindData::CheckQuaternionThreshold(FQuat A, FQuat B, float ThresholdDegrees)
@@ -1387,7 +1387,7 @@ enum class EResimFrameValidation : int32
 };
 CHAOS_API int32 ResimFrameValidation = static_cast<int32>(EResimFrameValidation::IslandValidation);
 FAutoConsoleVariableRef CVarResimFrameValidationLeniency(TEXT("p.Resim.ResimFrameValidation"), ResimFrameValidation, TEXT("0 = no leniency, all dirty particles need a valid target. 1 = Island leniency, all particles in resim islands need a valid target. 2 = Full leniency, only the particle triggering the resim need a valid target."));
-CHAOS_API bool bResimIncompleteHistory = false;
+CHAOS_API bool bResimIncompleteHistory = true;
 FAutoConsoleVariableRef CVarResimIncompleteHistory(TEXT("p.Resim.IncompleteHistory"), bResimIncompleteHistory, TEXT("If a valid resim frame can't be found, use the requested resim frame and perform a resimulation with incomplete data."));
 CHAOS_API bool bFindValidInputHistory = true;
 FAutoConsoleVariableRef CVarResimFindValidInputHistory(TEXT("p.Resim.FindValidInputHistory"), bFindValidInputHistory, TEXT("If the particle that needs resimulation has custom input history, find a valid resim frame where inputs are available."));
