Differentiate the reasons one might need to wait

#rnx

diff --git a/Engine/Source/Runtime/Experimental/IoStore/HttpClient/Private/Loop.inl b/Engine/Source/Runtime/Experimental/IoStore/HttpClient/Private/Loop.inl
--- a/Engine/Source/Runtime/Experimental/IoStore/HttpClient/Private/Loop.inl
+++ b/Engine/Source/Runtime/Experimental/IoStore/HttpClient/Private/Loop.inl
@@ -542,7 +542,7 @@ void FPeerGroup::RecvInternal(FTickState& State)
 
 	IsKeepAlive &= uint8(Activity->GetTransaction()->IsKeepAlive() == true);
 	LastUseMs = State.NowMs;
-	bWaiting |= Outcome.IsWaiting();
+	bWaiting |= (Outcome.IsWaiting() && Outcome.IsWaitData());
 
 	// If we've only a small amount left to receive we can start more work
 	if (IsKeepAlive & (Activity->Next == nullptr) & Send.IsEmpty())
diff --git a/Engine/Source/Runtime/Experimental/IoStore/HttpClient/Private/TransactionOne.inl b/Engine/Source/Runtime/Experimental/IoStore/HttpClient/Private/TransactionOne.inl
--- a/Engine/Source/Runtime/Experimental/IoStore/HttpClient/Private/TransactionOne.inl
+++ b/Engine/Source/Runtime/Experimental/IoStore/HttpClient/Private/TransactionOne.inl
@@ -675,7 +675,7 @@ FOutcome FBody::Recv(FMutableMemoryView Dest, FTlsPeer& Peer)
 		if (Remaining == 0 || Dest.IsEmpty())
 		{
 			int32 Result = int32(Size);
-			return Remaining ? FOutcome::Waiting(Result) : FOutcome::Ok(Result);
+			return Remaining ? FOutcome::WaitBuffer(Result) : FOutcome::Ok(Result);
 		}
 	}
 
@@ -687,6 +687,12 @@ FOutcome FBody::Recv(FMutableMemoryView Dest, FTlsPeer& Peer)
 		return Outcome;
 	}
 
+	if (Outcome.IsWaiting())
+	{
+		check(Outcome.GetResult() == 0);
+		return FOutcome::Waiting(int32(OverspillSize));
+	}
+
 	int32 Result = Outcome.GetResult();
 	Remaining -= Result;
 	if (Remaining < 0)
@@ -695,7 +701,18 @@ FOutcome FBody::Recv(FMutableMemoryView Dest, FTlsPeer& Peer)
 	}
 
 	Result += int32(OverspillSize);
-	return Remaining ? FOutcome::Waiting(Result) : FOutcome::Ok(Result);
+
+	if (Remaining == 0)
+	{
+		return FOutcome::Ok(Result);
+	}
+
+	if (Result == Dest.GetSize())
+	{
+		return FOutcome::WaitBuffer(Result);
+	}
+
+	return FOutcome::Waiting(Result);
 }
 
 ////////////////////////////////////////////////////////////////////////////////
@@ -823,7 +840,7 @@ FOutcome FBody::Chunk(FMutableMemoryView Dest, FTlsPeer& Peer)
 	State = EState::Epilogue;
 
 	int32 Result = Outcome.GetResult();
-	return FOutcome::Waiting(Result);
+	return FOutcome::WaitBuffer(Result);
 }
 
 ////////////////////////////////////////////////////////////////////////////////
