Fix race condition when Acceleration Structure Split Static And Dynamic is disabled.
Make sure we wait the integration to be completeed before reseting pending data queue

diff --git a/Engine/Source/Runtime/Experimental/Chaos/Private/Chaos/PBDRigidsEvolution.cpp b/Engine/Source/Runtime/Experimental/Chaos/Private/Chaos/PBDRigidsEvolution.cpp
--- a/Engine/Source/Runtime/Experimental/Chaos/Private/Chaos/PBDRigidsEvolution.cpp
+++ b/Engine/Source/Runtime/Experimental/Chaos/Private/Chaos/PBDRigidsEvolution.cpp
@@ -829,7 +829,11 @@ namespace Chaos
 				}
 				else
 				{
-					//other queues are no longer needed since we've flushed all operations and now have a pristine structure
+					// Other queues are no longer needed since we've flushed all operations and now have a pristine structure
+					// Wait for the integration to finish before reseting the internal acceleration queue. 
+					// Note: We could even not fill the Internal Acceleration Queue when only one acceleration structure 
+					// is used (AccelerationStructureSplitStaticAndDynamic is disabled)
+					TaskDispatcher.WaitIntegrationComplete();
 					InternalAccelerationQueue.Reset();
 				}
 				TaskDispatcher.WaitAsyncQueueTask();
