Dataflow context : Fix use of iterator after deletion

diff --git a/Engine/Source/Runtime/Experimental/Dataflow/Core/Private/Dataflow/DataflowContextEvaluator.cpp b/Engine/Source/Runtime/Experimental/Dataflow/Core/Private/Dataflow/DataflowContextEvaluator.cpp
--- a/Engine/Source/Runtime/Experimental/Dataflow/Core/Private/Dataflow/DataflowContextEvaluator.cpp
+++ b/Engine/Source/Runtime/Experimental/Dataflow/Core/Private/Dataflow/DataflowContextEvaluator.cpp
@@ -167,10 +167,11 @@ namespace UE::Dataflow
 		for (auto Iter = RunningTasks.CreateIterator(); Iter; ++Iter)
 		{
 			const FGraphEventRef& Task = (*Iter).Value;
+			const FNodeId TaskId = (*Iter).Key;
 			if (Task->IsCompleted())
 			{
+				CompletedTasks.Add(TaskId);
 				Iter.RemoveCurrent();
-				CompletedTasks.Add((*Iter).Key);
 			}
 		}
 	}
