Added a iax.HttpVersion cvar

#rnx

diff --git a/Engine/Source/Runtime/Experimental/IoStore/OnDemand/Private/DistributionEndpoints.cpp b/Engine/Source/Runtime/Experimental/IoStore/OnDemand/Private/DistributionEndpoints.cpp
--- a/Engine/Source/Runtime/Experimental/IoStore/OnDemand/Private/DistributionEndpoints.cpp
+++ b/Engine/Source/Runtime/Experimental/IoStore/OnDemand/Private/DistributionEndpoints.cpp
@@ -23,6 +23,8 @@
 namespace UE::IoStore
 {
 
+extern int32 GIaxHttpVersion;
+
 static int32 GDistributedEndpointTimeout = 30;
 static FAutoConsoleVariableRef CVar_DistributedEndpointTimeout(
 	TEXT("ias.DistributedEndpointTimeout"),
@@ -104,7 +106,14 @@ FDistributionEndpoints::EResult FDistributionEndpoints::ResolveEndpoints(const F
 		Loop.SetFailTimeout(GDistributedEndpointTimeout * 1000);
 	}
 
-	FEventLoop::FRequestParams RequestParams = { .bAllowChunked = true };
+	FEventLoop::FRequestParams RequestParams = {
+		.bAllowChunked = true,
+		.HttpVersion = (GIaxHttpVersion == 2) ? EHttpVersion::Two : EHttpVersion::One,
+	};
+	if (RequestParams.HttpVersion == EHttpVersion::Two)
+	{
+		RequestParams.VerifyCert = FCertRoots::Default();
+	}
 	FRequest Request = Loop.Get(AnsiUrl, &RequestParams);
 	Request.Header("accept", "application/json");
 	Loop.Send(MoveTemp(Request), OnRequestStatus);
diff --git a/Engine/Source/Runtime/Experimental/IoStore/OnDemand/Private/LatencyTesting.cpp b/Engine/Source/Runtime/Experimental/IoStore/OnDemand/Private/LatencyTesting.cpp
--- a/Engine/Source/Runtime/Experimental/IoStore/OnDemand/Private/LatencyTesting.cpp
+++ b/Engine/Source/Runtime/Experimental/IoStore/OnDemand/Private/LatencyTesting.cpp
@@ -14,6 +14,8 @@
 namespace UE::IoStore
 {
 
+extern int32 GIaxHttpVersion;
+
 void LatencyTest(FAnsiStringView InUrl, FAnsiStringView InPath, uint32 TimeOutMs, TArrayView<int32> OutResults)
 {
 	using namespace HTTP;
@@ -21,6 +23,11 @@ void LatencyTest(FAnsiStringView InUrl, FAnsiStringView InPath, uint32 TimeOutMs
 	FConnectionPool::FParams PoolParams;
 	PoolParams.SetHostFromUrl(InUrl);
 	PoolParams.ConnectionCount = 1;
+	PoolParams.HttpVersion = (GIaxHttpVersion == 2) ? EHttpVersion::Two : EHttpVersion::One;
+	if (PoolParams.HttpVersion == EHttpVersion::Two)
+	{
+		PoolParams.VerifyCert = FCertRoots::Default();
+	}
 	
 	FConnectionPool Pool(PoolParams);
 
diff --git a/Engine/Source/Runtime/Experimental/IoStore/OnDemand/Private/OnDemandHttpClient.cpp b/Engine/Source/Runtime/Experimental/IoStore/OnDemand/Private/OnDemandHttpClient.cpp
--- a/Engine/Source/Runtime/Experimental/IoStore/OnDemand/Private/OnDemandHttpClient.cpp
+++ b/Engine/Source/Runtime/Experimental/IoStore/OnDemand/Private/OnDemandHttpClient.cpp
@@ -23,6 +23,13 @@ static FAutoConsoleVariableRef CVar_IaxWarnOnFailure(
 	TEXT("When enabled FMultiEndpointHttpClient will log failed requests as warnings instead of the config verbosity")
 );
 
+int32 GIaxHttpVersion = 1;
+static FAutoConsoleVariableRef CVar_IaxHttpVersion(
+	TEXT("iax.HttpVersion"),
+	GIaxHttpVersion,
+	TEXT("Protocol version to use, either '1' or '2'")
+);
+
 /**
  * When enabled, http requests will make a partial content request by adding Range: bytes=<offset>-<offset+length> to the request if a valid
  * FIoOffsetAndLength is provided with the request. When disabled we will always request the full chunk and then use the FIoOffsetAndLength to
@@ -151,8 +158,13 @@ TIoStatusOr<FMultiEndpointHttpClientResponse> FMultiEndpointHttpClient::Get(FAns
 
 	FEventLoop::FRequestParams Params = FEventLoop::FRequestParams
 	{
-		.bAutoRedirect = Config.Redirects == EHttpRedirects::Follow
+		.bAutoRedirect = Config.Redirects == EHttpRedirects::Follow,
+		.HttpVersion = (GIaxHttpVersion == 2) ? EHttpVersion::Two : EHttpVersion::One,
 	};
+	if (Params.HttpVersion == EHttpVersion::Two)
+	{
+		Params.VerifyCert = FCertRoots::Default();
+	}
 
 	FEventLoop Loop;
 	FIoBuffer Body;
@@ -463,7 +475,9 @@ void FMultiEndpointHttpClient::RetryRequest(FRequest&& Request)
 
 TUniquePtr<HTTP::FConnectionPool> FMultiEndpointHttpClient::CreateConnection(FAnsiStringView HostUrl) const
 {
-	HTTP::FConnectionPool::FParams Params;
+	using namespace HTTP;
+
+	FConnectionPool::FParams Params;
 	ensure(Params.SetHostFromUrl(HostUrl) >= 0);
 	
 	if (Config.ReceiveBufferSize >= 0)
@@ -477,6 +491,12 @@ TUniquePtr<HTTP::FConnectionPool> FMultiEndpointHttpClient::CreateConnection(FAn
 	}
 
 	Params.ConnectionCount = uint16(Config.MaxConnectionCount);
+	Params.HttpVersion = (GIaxHttpVersion == 2) ? EHttpVersion::Two : EHttpVersion::One;
+	if (Params.HttpVersion == EHttpVersion::Two)
+	{
+		Params.ConnectionCount = 1; // rfc9113 conformance
+		Params.VerifyCert = HTTP::FCertRoots::Default();
+	}
 
 	return MakeUnique<HTTP::FConnectionPool>(Params);
 }
diff --git a/Engine/Source/Runtime/Experimental/IoStore/OnDemand/Private/Statistics.cpp b/Engine/Source/Runtime/Experimental/IoStore/OnDemand/Private/Statistics.cpp
--- a/Engine/Source/Runtime/Experimental/IoStore/OnDemand/Private/Statistics.cpp
+++ b/Engine/Source/Runtime/Experimental/IoStore/OnDemand/Private/Statistics.cpp
@@ -28,6 +28,8 @@ LLM_DEFINE_TAG(Ias);
 namespace UE::IoStore
 {
 
+extern int32 GIaxHttpVersion;
+
 float GIasStatisticsLogInterval = 60.f;
 static FAutoConsoleVariableRef CVar_StatisticsLogInterval(
 	TEXT("ias.StatisticsLogInterval"),
@@ -769,6 +771,7 @@ void FOnDemandIoBackendStats::PrintOnScreenStatistics(FCoreDelegates::FSeverityM
 		{
 			TStringBuilder<256> Builder;
 			Builder << TEXT("IAX Http:");
+			Builder << TEXT(" v") << GIaxHttpVersion;
 			if (GIasHttpRateLimitKiBPerSecond > 0)
 			{
 				Builder << TEXT(" DataRate Cap: ") << GIasHttpRateLimitKiBPerSecond << TEXT("KiB/s");
