Gracefully handle empty transactions

#rnx

diff --git a/Engine/Source/Runtime/Experimental/IoStore/HttpClient/Private/Activity.inl b/Engine/Source/Runtime/Experimental/IoStore/HttpClient/Private/Activity.inl
--- a/Engine/Source/Runtime/Experimental/IoStore/HttpClient/Private/Activity.inl
+++ b/Engine/Source/Runtime/Experimental/IoStore/HttpClient/Private/Activity.inl
@@ -401,7 +401,10 @@ const FStopwatch& FActivity::GetStopwatch() const
 FOutcome FActivity::Transact(FHttpPeer& Peer)
 {
 	Transaction = Peer.Transact(Buffer);
-	check(Transaction.IsValid());
+	if (!Transaction.IsValid())
+	{
+		return FOutcome::Error("Unable to create a suitable transaction");
+	}
 
 	const FHost* TheHost = GetHost();
 	bool bKeepAlive = !TheHost->IsPooled();
