An impl of h2 transactions

#rnx

diff --git a/Engine/Source/Runtime/Experimental/IoStore/HttpClient/IoStoreHttpClient.Build.cs b/Engine/Source/Runtime/Experimental/IoStore/HttpClient/IoStoreHttpClient.Build.cs
--- a/Engine/Source/Runtime/Experimental/IoStore/HttpClient/IoStoreHttpClient.Build.cs
+++ b/Engine/Source/Runtime/Experimental/IoStore/HttpClient/IoStoreHttpClient.Build.cs
@@ -13,5 +13,10 @@ public class IoStoreHttpClient : ModuleRules
 		PrivateDefinitions.Add("IAS_HTTP_HAS_OPENSSL=1");
 
 		CppCompileWarningSettings.UnsafeTypeCastWarningLevel = WarningLevel.Error; 
+
+		if (Target.Platform.IsInGroup(UnrealPlatformGroup.Windows))
+		{
+			PrivateDependencyModuleNames.Add("nghttp2");
+		}
 	}
 }
diff --git a/Engine/Source/Runtime/Experimental/IoStore/HttpClient/Private/ConnectionPool.inl b/Engine/Source/Runtime/Experimental/IoStore/HttpClient/Private/ConnectionPool.inl
--- a/Engine/Source/Runtime/Experimental/IoStore/HttpClient/Private/ConnectionPool.inl
+++ b/Engine/Source/Runtime/Experimental/IoStore/HttpClient/Private/ConnectionPool.inl
@@ -239,6 +239,7 @@ FConnectionPool::FConnectionPool(const FParams& Params)
 {
 	check(Params.ConnectionCount - 1u <= 63u);
 	check(Params.Port <= 0xffffu);
+	check(Params.HttpVersion != EHttpVersion::Two || Params.ConnectionCount == 1); // only one as per rfc9113
 
 	// Alloc a new internal object
 	uint32 HostNameLen = Params.HostName.Len();
diff --git a/Engine/Source/Runtime/Experimental/IoStore/HttpClient/Private/Peer.inl b/Engine/Source/Runtime/Experimental/IoStore/HttpClient/Private/Peer.inl
--- a/Engine/Source/Runtime/Experimental/IoStore/HttpClient/Private/Peer.inl
+++ b/Engine/Source/Runtime/Experimental/IoStore/HttpClient/Private/Peer.inl
@@ -258,7 +258,7 @@ static void Ssl_SetupAlpn(SSL* Ssl, int32 HttpVersion)
 	switch (HttpVersion)
 	{
 	case 1:		AlpnProtos = "\x08" "http/1.1";	break;
-//	case 2:		AlpnProtos = "\x02" "h2";		break;
+	case 2:		AlpnProtos = "\x02" "h2";		break;
 	default:	break;
 	}
 
@@ -284,7 +284,7 @@ static int32 Ssl_GetProtocolVersion(SSL* Ssl)
 	FAnsiStringView Needle(AlpnProto, AlpnProtoLen);
 	FAnsiStringView Candidates[] = {
 		"http/1.1",
-	//	"h2",
+		"h2",
 	};
 	for (uint32 i = 0; i < UE_ARRAY_COUNT(Candidates); ++i)
 	{
@@ -517,6 +517,7 @@ FOutcome FTlsPeer::Recv(char* Out, int32 MaxSize)
 
 
 ////////////////////////////////////////////////////////////////////////////////
+#include "TransactionTwo.inl"
 #include "TransactionOne.inl"
 
 namespace UE::IoStore::HTTP
@@ -536,6 +537,7 @@ public:
 	};
 
 					FHttpPeer() = default;
+					~FHttpPeer();
 					FHttpPeer(FHttpPeer&& Rhs) = delete;
 					FHttpPeer(const FHttpPeer& Rhs) = delete;
 	FHttpPeer&		operator = (FHttpPeer&& Rhs);
@@ -546,6 +548,7 @@ public:
 	FOutcome		GetPendingTransactId();
 
 private:
+	void*			PeerData = nullptr;
 	uint8			Proto;
 };
 
@@ -560,9 +563,21 @@ FHttpPeer::FHttpPeer(FParams&& Params)
 	}
 }
 
+////////////////////////////////////////////////////////////////////////////////
+FHttpPeer::~FHttpPeer()
+{
+	switch (Proto)
+	{
+	case 1:		GoAwayOne(*this, PeerData);
+	case 2:		GoAwayTwo(*this, PeerData);
+	default:	break;
+	}
+}
+
 ////////////////////////////////////////////////////////////////////////////////
 FHttpPeer& FHttpPeer::operator = (FHttpPeer&& Rhs)
 {
+	Swap(PeerData, Rhs.PeerData);
 	Swap(Proto, Rhs.Proto);
 
 	FTlsPeer::operator = (MoveTemp(Rhs));
@@ -585,37 +600,38 @@ FOutcome FHttpPeer::Handshake()
 		{
 			return FOutcome::Error("Failed to negotiate protocol version");
 		}
+
+		/* intentional fallthrough */
 	}
 
-	if (Proto == 1)
+	switch (Proto)
 	{
-		return FOutcome::Ok();
+	case 1:		return HandshakeOne(*this, PeerData);
+	case 2:		return HandshakeTwo(*this, PeerData);
+	default:	return FOutcome::Error("Unsupported");
 	}
-
-	return FOutcome::Error("Unsupported version");
 }
 
 ////////////////////////////////////////////////////////////////////////////////
 FTransactRef FHttpPeer::Transact(FBuffer& Buffer)
 {
-	if (Proto == 1)
+	switch (Proto)
 	{
-		return CreateTransactOne(Buffer);
+	case 1:		return CreateTransactOne(Buffer, PeerData);
+	case 2:		return CreateTransactTwo(Buffer, PeerData);
+	default:	return FTransactRef();
 	}
-
-	check(!"unsupported");
-	return FTransactRef();
 }
 
 ////////////////////////////////////////////////////////////////////////////////
 FOutcome FHttpPeer::GetPendingTransactId()
 {
-	if (Proto == 1)
+	switch (Proto)
 	{
-		return FOutcome::Ok(1);
+	case 1:		return TickOne(*this, PeerData);
+	case 2:		return TickTwo(*this, PeerData);
+	default:	return FOutcome::Error("Unsupported");
 	}
-
-	return FOutcome::Error("unsupported");
 }
 
 } // namespace UE::IoStore::HTTP
diff --git a/Engine/Source/Runtime/Experimental/IoStore/HttpClient/Private/TransactionOne.inl b/Engine/Source/Runtime/Experimental/IoStore/HttpClient/Private/TransactionOne.inl
--- a/Engine/Source/Runtime/Experimental/IoStore/HttpClient/Private/TransactionOne.inl
+++ b/Engine/Source/Runtime/Experimental/IoStore/HttpClient/Private/TransactionOne.inl
@@ -5,9 +5,6 @@
 namespace UE::IoStore::HTTP
 {
 
-////////////////////////////////////////////////////////////////////////////////
-using FTransactId = uint32;
-
 ////////////////////////////////////////////////////////////////////////////////
 struct FMessageOffsets
 {
@@ -132,21 +129,6 @@ static void EnumerateHeaders(FAnsiStringView Headers, LambdaType&& Lambda)
 
 
 
-////////////////////////////////////////////////////////////////////////////////
-bool IsContentless(uint32 StatusCode)
-{
-	return (StatusCode < 200) | (StatusCode == 204) | (StatusCode == 304);
-}
-
-////////////////////////////////////////////////////////////////////////////////
-bool IsStatusCodeOk(uint32 StatusCode)
-{
-	enum { MinCode = 100, MaxCode = 999 };
-	return uint32(StatusCode - MinCode) <= uint32(MaxCode - MinCode);
-}
-
-
-
 namespace DetailOne
 {
 
@@ -893,12 +875,33 @@ public:
 
 
 
+////////////////////////////////////////////////////////////////////////////////
+FOutcome HandshakeOne(FTlsPeer&, void*)
+{
+	return FOutcome::Ok();
+}
+
+////////////////////////////////////////////////////////////////////////////////
+FOutcome TickOne(FTlsPeer&, void*)
+{
+	return FOutcome::Ok(1);
+}
+
+////////////////////////////////////////////////////////////////////////////////
+void GoAwayOne(FTlsPeer&, void*)
+{
+}
+
+
+
 ////////////////////////////////////////////////////////////////////////////////
 class FTransaction
 {
 public:
 #define TR_METHOD(x, ...) \
 	template <typename... Arg> auto x(Arg&&... Args) __VA_ARGS__ { \
+		if (auto Self = UPTRINT(this); Self & 1) \
+			return ((FTransactionTwo*)(Self ^ 1))->x(Forward<Arg>(Args)...); \
 		return ((FTransactionOne*)this)->x(Forward<Arg>(Args)...); \
 	}
 	TR_METHOD(Begin)
@@ -921,6 +924,14 @@ private:
 
 	void Destruct()
 	{
+		if (auto Self = UPTRINT(this); Self & 1)
+		{
+			void* Addr = (void*)(Self ^ 1);
+			((FTransactionTwo*)Addr)->~FTransactionTwo();
+			FMemory::Free(Addr);
+			return;
+		}
+
 		((FTransactionOne*)this)->~FTransactionOne();
 		FMemory::Free(this);
 	}
@@ -934,6 +945,7 @@ class FTransactRef
 public:
 						FTransactRef() = default;
 						~FTransactRef();
+						FTransactRef(UPTRINT Ptr, int32 Ver);
 						FTransactRef(FTransactRef&&) = delete;
 						FTransactRef(const FTransactRef&) = delete;
 	void				operator = (const FTransactRef& Rhs) = delete;
@@ -952,9 +964,10 @@ private:
 };
 
 ////////////////////////////////////////////////////////////////////////////////
-FTransactRef::FTransactRef(UPTRINT Ptr)
-: Self(Ptr)
+FTransactRef::FTransactRef(UPTRINT Ptr, int32 Ver)
+: Self(Ptr | (Ver == 2))
 {
+	check((Ptr & 1) == 0);
 }
 
 ////////////////////////////////////////////////////////////////////////////////
@@ -1007,11 +1020,28 @@ FTransactRef::~FTransactRef()
 
 
 ////////////////////////////////////////////////////////////////////////////////
-FTransactRef CreateTransactOne(FBuffer& Buffer)
+FTransactRef CreateTransactOne(FBuffer& Buffer, void*&)
 {
 	void* Ptr = FMemory::Malloc(sizeof(FTransactionOne), alignof(FTransactionOne));
 	new (Ptr) FTransactionOne(Buffer);
-	return FTransactRef(UPTRINT(Ptr));
+	return FTransactRef(UPTRINT(Ptr), 1);
+}
+
+////////////////////////////////////////////////////////////////////////////////
+FTransactRef CreateTransactTwo(FBuffer& Buffer, void* PeerData)
+{
+#if IAS_HTTP_WITH_TWO
+	using namespace DetailTwo;
+
+	auto* Driver = (FDriverNg*)PeerData;
+
+	void* Ptr = FMemory::Malloc(sizeof(FTransactionTwo), alignof(FTransactionTwo));
+	new (Ptr) FTransactionTwo(*Driver);
+
+	return FTransactRef(UPTRINT(Ptr), 2);
+#else
+	return FTransactRef();
+#endif
 }
 
 } // namespace UE::IoStore::HTTP
diff --git a/Engine/Source/Runtime/Experimental/IoStore/HttpClient/Private/TransactionTwo.inl b/Engine/Source/Runtime/Experimental/IoStore/HttpClient/Private/TransactionTwo.inl
new file mode 100644
--- /dev/null
+++ b/Engine/Source/Runtime/Experimental/IoStore/HttpClient/Private/TransactionTwo.inl
@@ -0,0 +1,934 @@
+// Copyright Epic Games, Inc. All Rights Reserved.
+
+#pragma once
+
+namespace UE::IoStore::HTTP
+{
+
+////////////////////////////////////////////////////////////////////////////////
+using FTransactId = uint32;
+
+////////////////////////////////////////////////////////////////////////////////
+bool IsContentless(uint32 StatusCode)
+{
+	return (StatusCode < 200) | (StatusCode == 204) | (StatusCode == 304);
+}
+
+////////////////////////////////////////////////////////////////////////////////
+bool IsStatusCodeOk(uint32 StatusCode)
+{
+	enum { MinCode = 100, MaxCode = 999 };
+	return uint32(StatusCode - MinCode) <= uint32(MaxCode - MinCode);
+}
+
+} // namespace UE::IoStore::HTTP
+
+
+
+////////////////////////////////////////////////////////////////////////////////
+#if !defined(WITH_NGHTTP2)
+#	define WITH_NGHTTP2 0
+#	if defined(__has_include)
+#		if __has_include(<nghttp2/nghttp2.h>)
+#			undef WITH_NGHTTP2
+#			define WITH_NGHTTP2 1
+#		endif
+#	endif
+#endif
+
+#if !defined(IAS_HTTP_WITH_TWO)
+#	define IAS_HTTP_WITH_TWO WITH_NGHTTP2
+#endif
+
+#if IAS_HTTP_WITH_TWO
+
+////////////////////////////////////////////////////////////////////////////////
+#if defined(_MSC_VER)
+#	define	NGHTTP2_NO_SSIZE_T
+#endif
+#if PLATFORM_MICROSOFT & !defined(NGHTTP2_STATICLIB)
+#	define NGHTTP2_STATICLIB
+#endif
+
+#include <nghttp2/nghttp2.h>
+
+#if PLATFORM_MICROSOFT & !defined(NGHTTP2_STATICLIB)
+#	undef NGHTTP2_STATICLIB
+#endif
+#if defined(_MSC_VER)
+#	undef NGHTTP2_NO_SSIZE_T
+#endif
+
+namespace UE::IoStore::HTTP
+{
+
+namespace DetailTwo
+{
+
+////////////////////////////////////////////////////////////////////////////////
+struct FFrame
+{
+	enum class EType : uint8
+	{
+		Data,			// = 0x00,
+		Headers,		// = 0x01,
+		Priority,		// = 0x02,
+		RstStream,		// = 0x03,
+		Settings,		// = 0x04,
+		PushPromise,	// = 0x05,
+		Ping,			// = 0x06,
+		Goaway,			// = 0x07,
+		WindowUpdate,	// = 0x08,
+		Continuation,	// = 0x09,
+		_Num,
+	};
+
+	enum : uint8
+	{
+		Flag_Padding	= 0x08,
+		Flag_EndHeaders	= 0x04,
+		Flag_EndStream	= 0x01,
+	};
+
+	uint8	Size[3];
+	EType	Type;
+	uint8	Flags;
+	uint8	Id[4];
+
+	int32	IsPertinent() const	{ return (Id[3] & 1) != 0; }
+	uint32	GetSize() const		{ return (Size[0] << 16u) | (Size[1] << 8u) | Size[2]; }
+	uint32	GetId() const		{ return (Id[0] << 24u) | (Id[1] << 16u) | (Id[2] << 8u) | Id[3]; }
+	uint32	GetIdFast() const	{ return Id[3]; }
+};
+static_assert(sizeof(FFrame) == 9);
+
+
+
+////////////////////////////////////////////////////////////////////////////////
+class FBufferNg
+{
+public:
+					FBufferNg() = default;
+					~FBufferNg()							{ if (Inner != nullptr) nghttp2_rcbuf_decref(Inner); }
+					FBufferNg(nghttp2_rcbuf* InRcBuf)		{ nghttp2_rcbuf_incref(Inner = InRcBuf); }
+					FBufferNg(FBufferNg&& Rhs)				{ Swap(Inner, Rhs.Inner); }
+					FBufferNg(const FBufferNg&) = delete;
+	FBufferNg&		operator = (FBufferNg&& Rhs)			{ Swap(Inner, Rhs.Inner); return *this; }
+	FBufferNg&		operator = (const FBufferNg&) = delete;
+
+	FAnsiStringView	Get() const
+	{
+		check(Inner != nullptr);
+		nghttp2_vec Ret = nghttp2_rcbuf_get_buf(Inner);
+		return { (char*)(Ret.base), int32(Ret.len) };
+	}
+
+private:
+	nghttp2_rcbuf*	Inner = nullptr;
+};
+
+////////////////////////////////////////////////////////////////////////////////
+struct FHeaderNg
+{
+	FBufferNg	Key;
+	FBufferNg	Value;
+};
+
+
+
+////////////////////////////////////////////////////////////////////////////////
+class FDriverHeaders
+{
+public:
+						FDriverHeaders() = default;
+	void				Add(FAnsiStringView Key, FAnsiStringView Value);
+
+private:
+	friend class FDriverNg;
+
+	const nghttp2_nv*	GetData() const	{ return Headers.GetData(); }
+	int32 				Num() const		{ return Headers.Num(); }
+	TArray<nghttp2_nv>	Headers;
+
+	UE_NONCOPYABLE(FDriverHeaders);
+};
+
+////////////////////////////////////////////////////////////////////////////////
+void FDriverHeaders::Add(FAnsiStringView Key, FAnsiStringView Value)
+{
+	nghttp2_nv Header{
+		(uint8*)Key.GetData(),
+		(uint8*)Value.GetData(),
+		size_t(Key.Len()),
+		size_t(Value.Len()),
+		NGHTTP2_NV_FLAG_NO_COPY_NAME | NGHTTP2_NV_FLAG_NO_COPY_VALUE,
+	};
+	Headers.Add(MoveTemp(Header));
+}
+
+
+
+////////////////////////////////////////////////////////////////////////////////
+class FDriverNg
+{
+public:
+						FDriverNg(FTlsPeer& Peer);
+	bool				IsUsingTls() const;
+	void				GoAway(FTlsPeer& Peer);
+	FOutcome			Handshake(FTlsPeer& Peer);
+	FOutcome			GetTransactId(FTlsPeer& Peer);
+	FTransactId			Submit(const FDriverHeaders& Headers);
+	FOutcome			Send(FTlsPeer& Peer);
+	FOutcome			RecvResponse(FTlsPeer& Peer, TArray<FHeaderNg>* Sink);
+	FOutcome			RecvBody(FTlsPeer& Peer, FMutableMemoryView Dest);
+
+private:
+	FOutcome			Tick(FTlsPeer& Peer);
+	FOutcome			TickInner(FTlsPeer& Peer);
+	FOutcome			PumpSend(FTlsPeer& Peer);
+	FOutcome			PumpRecv(FTlsPeer& Peer, FMemoryView Data);
+	nghttp2_ssize 		OnSend(const uint8_t* Data, size_t Size);
+	int32 				OnHeader(nghttp2_rcbuf* Name, nghttp2_rcbuf* Value);
+	nghttp2_session*	Session;
+	TArray<FHeaderNg>*	ResponseHeaders = nullptr;
+	FMemoryView			SendPending;
+	int32				RecvRemaining;
+	bool				bWithTls;
+	uint8				RecvBuffer[512 - sizeof(FFrame)];
+	union
+	{
+		FFrame			RecvFrame;
+		uint8			RecvFrameData[sizeof(RecvFrame)];
+	};
+
+	UE_NONCOPYABLE(FDriverNg);
+};
+
+////////////////////////////////////////////////////////////////////////////////
+FDriverNg::FDriverNg(FTlsPeer& Peer)
+: bWithTls(Peer.IsUsingTls())
+{
+	struct FStaticsNg
+	{
+		~FStaticsNg()
+		{
+			nghttp2_session_callbacks_del(Instance);
+		}
+		nghttp2_session_callbacks* Instance = nullptr;
+	};
+
+	static FStaticsNg Statics;
+	if (Statics.Instance == nullptr)
+	{
+		nghttp2_session_callbacks* Callbacks;
+		nghttp2_session_callbacks_new(&Callbacks);
+
+		/*
+		nghttp2_session_callbacks_set_send_callback2(
+			Callbacks,
+			[] (nghttp2_session*, const uint8_t* Data, size_t Size, int32, void* UserData) -> nghttp2_ssize
+			{
+				auto* This = (FDriverNg*)UserData;
+				return This->OnSend(Data, Size);
+			}
+		);
+		*/
+
+		nghttp2_session_callbacks_set_on_header_callback2(
+			Callbacks,
+			[] (nghttp2_session*, const nghttp2_frame*, nghttp2_rcbuf* Name, nghttp2_rcbuf* Value, uint8_t, void* UserData) -> int32
+			{
+				auto* This = (FDriverNg*)UserData;
+				return This->OnHeader(Name, Value);
+			}
+		);
+
+		Statics.Instance = Callbacks;
+	}
+
+	static nghttp2_mem Mem = {
+		.mem_user_data	= nullptr,
+		.malloc			= [] (size_t Size, void*)				{ return FMemory::Malloc(Size); },
+		.free			= [] (void* Ptr, void*)					{ return FMemory::Free(Ptr); },
+		.calloc			= [] (size_t Count, size_t Size, void*) { return FMemory::MallocZeroed(Count * Size); },
+		.realloc		= [] (void* Ptr, size_t Size, void*)	{ return FMemory::Realloc(Ptr, Size); },
+	};
+
+	RecvRemaining = 0 - int32(sizeof(RecvFrame));
+
+	nghttp2_session_client_new3(&Session, Statics.Instance, this, nullptr, &Mem);
+
+	// Perhaps in the future a window size can be picked - or adjusted - by
+	// monitoring latency and throughput.
+	/*
+	nghttp2_session_set_local_window_size(Session, 0, 0, 256 << 10);
+	*/
+}
+
+////////////////////////////////////////////////////////////////////////////////
+void FDriverNg::GoAway(FTlsPeer& Peer)
+{
+	if (Session != nullptr)
+	{
+		nghttp2_session_terminate_session(Session, NGHTTP2_NO_ERROR);
+		PumpSend(Peer);
+		nghttp2_session_del(Session);
+	}
+}
+
+////////////////////////////////////////////////////////////////////////////////
+bool FDriverNg::IsUsingTls() const
+{
+	return bWithTls;
+}
+
+////////////////////////////////////////////////////////////////////////////////
+FOutcome FDriverNg::PumpSend(FTlsPeer& Peer)
+{
+	for (uint32 Sent = 0;;)
+	{
+		if (SendPending.IsEmpty())
+		{
+			const uint8* Data = nullptr;
+			int32 Result = int32(nghttp2_session_mem_send2(Session, &Data));
+			if (Result < 0)
+			{
+				const char* Message = nghttp2_strerror(Result);
+				return FOutcome::Error(Message);
+			}
+
+			if (Result == 0)
+			{
+				return FOutcome::Ok(Sent);
+			}
+
+			SendPending = { Data, uint32(Result) };
+		}
+
+		const char* Data = (char*)(SendPending.GetData());
+		uint32 Size = uint32(SendPending.GetSize());
+		FOutcome Outcome = Peer.Send((char*)Data, Size);
+		if (Outcome.IsError())
+		{
+			return Outcome;
+		}
+
+		if (Outcome.IsOk())
+		{
+			int32 Result = Outcome.GetResult();
+			Sent += Result;
+
+			SendPending = SendPending.Mid(Result);
+			if (SendPending.IsEmpty())
+			{
+				continue;
+			}
+		}
+
+		check(!SendPending.IsEmpty());
+		return FOutcome::Waiting();
+	}
+}
+
+////////////////////////////////////////////////////////////////////////////////
+FOutcome FDriverNg::PumpRecv(FTlsPeer& Peer, FMemoryView Data)
+{
+	int32 Result = int32(nghttp2_session_mem_recv2(Session, (uint8*)Data.GetData(), Data.GetSize()));
+
+	if (Result < 0)
+	{
+		const char* Message = nghttp2_strerror(Result);
+		return FOutcome::Error(Message);
+	}
+
+	if (Result != int32(Data.GetSize()))
+	{
+		return FOutcome::Error("Not enough data consumed");
+	}
+
+	if (nghttp2_session_want_write(Session))
+	{
+		// We'll not care about the outcome here. If an error occurred it will
+		// surface later and data-waits will solved themselves later.
+		PumpSend(Peer);
+	}
+
+	return FOutcome::Ok();
+}
+
+////////////////////////////////////////////////////////////////////////////////
+FOutcome FDriverNg::Handshake(FTlsPeer& Peer)
+{
+	nghttp2_submit_settings(Session, NGHTTP2_FLAG_NONE, nullptr, 0);
+	FOutcome Outcome = PumpSend(Peer);
+	return Outcome.IsError() ? Outcome : FOutcome::Ok();
+}
+
+////////////////////////////////////////////////////////////////////////////////
+FOutcome FDriverNg::GetTransactId(FTlsPeer& Peer)
+{
+	if (RecvRemaining >= 0 && RecvFrame.IsPertinent())
+	{
+		uint32 Id = RecvFrame.GetId();
+		return FOutcome::Ok(Id);
+	}
+
+	return Tick(Peer);
+}
+
+////////////////////////////////////////////////////////////////////////////////
+FOutcome FDriverNg::Tick(FTlsPeer& Peer)
+{
+	while (true)
+	{
+		FOutcome Outcome = TickInner(Peer);
+		if (!Outcome.IsWaiting() || (Outcome.GetResult() == 0))
+		{
+			return Outcome;
+		}
+	}
+}
+
+////////////////////////////////////////////////////////////////////////////////
+FOutcome FDriverNg::TickInner(FTlsPeer& Peer)
+{
+	FMutableMemoryView Dest(RecvBuffer, sizeof(RecvBuffer) + sizeof(RecvFrame));
+	if (RecvRemaining < 0)
+	{
+		Dest = Dest.Right(0 - RecvRemaining);
+	}
+	else if (RecvRemaining <= int32(sizeof(RecvBuffer)))
+	{
+		Dest = Dest.Right(RecvRemaining + sizeof(RecvFrame));
+	}
+	else
+	{
+		Dest = Dest.Left(sizeof(RecvBuffer));
+	}
+
+	FOutcome Outcome = Peer.Recv((char*)(Dest.GetData()), int32(Dest.GetSize()));
+	if (!Outcome.IsOk())
+	{
+		return Outcome;
+	}
+
+	int32 Result = Outcome.GetResult();
+	Dest = Dest.Left(Result);
+
+	FMutableMemoryView FrameView;
+	if (const uint8* End = (uint8*)(Dest.GetDataEnd()); End > RecvFrameData)
+	{
+		FrameView = FMutableMemoryView(RecvFrameData, End - RecvFrameData);
+		Dest = Dest.LeftChop(FrameView.GetSize());
+	}
+
+	if (!Dest.IsEmpty() && (Outcome = PumpRecv(Peer, Dest)).IsError())
+	{
+		return Outcome;
+	}
+
+	if (FrameView.IsEmpty())
+	{
+		check(RecvRemaining >= int32(Dest.GetSize()));
+		if ((RecvRemaining -= int32(Dest.GetSize())) == 0)
+		{
+			RecvRemaining = 0 - int32(sizeof(RecvFrame));
+		}
+		return FOutcome::Waiting(Dest.GetDataEnd() == RecvFrameData);
+	}
+
+	if (Result = int32(FrameView.GetSize() - sizeof(RecvFrame)); Result < 0)
+	{
+		RecvRemaining = Result;
+		return FOutcome::Waiting();
+	}
+
+	if ((Outcome = PumpRecv(Peer, FrameView)).IsError())
+	{
+		return Outcome;
+	}
+
+	RecvRemaining = RecvFrame.GetSize();
+	if (RecvFrame.IsPertinent())
+	{
+		uint32 Id = RecvFrame.GetId();
+		return FOutcome::Ok(Id);
+	}
+
+	if (RecvRemaining == 0)
+	{
+		RecvRemaining = 0 - int32(sizeof(RecvFrame));
+	}
+
+	return FOutcome::Waiting(1);
+}
+
+////////////////////////////////////////////////////////////////////////////////
+FTransactId FDriverNg::Submit(const FDriverHeaders& Headers)
+{
+	return nghttp2_submit_request2(Session, nullptr, Headers.GetData(),
+		Headers.Num(), nullptr, this);
+}
+
+////////////////////////////////////////////////////////////////////////////////
+FOutcome FDriverNg::Send(FTlsPeer& Peer)
+{
+	return PumpSend(Peer);
+}
+
+////////////////////////////////////////////////////////////////////////////////
+FOutcome FDriverNg::RecvResponse(FTlsPeer& Peer, TArray<FHeaderNg>* Sink)
+{
+	if (RecvFrame.Type == FFrame::EType::RstStream)
+	{
+		return FOutcome::Error("Received stream reset");
+	}
+
+	// Assume if we're filling we are expecting something
+	check(RecvRemaining > 0);
+	check(RecvFrame.Type == FFrame::EType::Headers);
+
+	ResponseHeaders = Sink;
+	ON_SCOPE_EXIT { ResponseHeaders = nullptr; };
+
+	while (true)
+	{
+		int32 EndOfHeaders = (RecvFrame.Flags & FFrame::Flag_EndHeaders);
+		uint32 ThisId = RecvFrame.GetId();
+
+		FOutcome Outcome = Tick(Peer); 
+		if (Outcome.IsError())				return Outcome;
+		if (EndOfHeaders)					return FOutcome::Ok();
+		if (Outcome.GetResult() != ThisId)	return FOutcome::WaitStream();
+	}
+}
+
+////////////////////////////////////////////////////////////////////////////////
+FOutcome FDriverNg::RecvBody(FTlsPeer& Peer, FMutableMemoryView Dest)
+{
+	if (RecvFrame.Type == FFrame::EType::RstStream)
+	{
+		return FOutcome::Error("Received stream reset");
+	}
+
+	check(RecvRemaining >= 0);
+	check(RecvFrame.Type == FFrame::EType::Data);
+
+	for (uint32 RecvSize = 0;;)
+	{
+		int32 EndOfStream = (RecvFrame.Flags & FFrame::Flag_EndStream);
+
+		FMutableMemoryView View = Dest.Left(RecvRemaining);
+		if (View.IsEmpty())
+		{
+			return EndOfStream ? FOutcome::Ok(RecvSize) : FOutcome::WaitBuffer(RecvSize);
+		}
+
+		FOutcome Outcome = Peer.Recv((char*)View.GetData(), int32(View.GetSize()));
+		if (!Outcome.IsOk())
+		{
+			if (Outcome.IsError())
+			{
+				return Outcome;
+			}
+			return FOutcome::Waiting(RecvSize);
+		}
+
+		int32 Result = Outcome.GetResult();
+		RecvRemaining -= Result;
+		RecvSize += Result;
+		Dest = Dest.Mid(Result);
+
+		if ((Outcome = PumpRecv(Peer, View.Left(Result))).IsError())
+		{
+			return Outcome;
+		}
+
+		// more of this frame to go?
+		if (RecvRemaining > 0)
+		{
+			continue;
+		}
+
+		RecvRemaining = 0 - int32(sizeof(RecvFrame)) - RecvRemaining;
+
+		// End of the road?
+		if (EndOfStream)
+		{
+			return FOutcome::Ok(RecvSize);
+		}
+
+		uint32 ThisId = RecvFrame.GetId();
+		Outcome = Tick(Peer);
+		if (Outcome.IsError())				return Outcome;
+		if (ThisId != Outcome.GetResult())	return FOutcome::WaitBuffer(RecvSize);
+	}
+}
+
+////////////////////////////////////////////////////////////////////////////////
+nghttp2_ssize FDriverNg::OnSend(const uint8_t* Data, size_t Size)
+{
+	/*
+	FOutcome Outcome = Peer.Send((char*)Data, int32(Size));
+	if (!Outcome.IsOk())
+	{
+		return Outcome.IsWaiting() ? NGHTTP2_ERR_WOULDBLOCK : NGHTTP2_ERR_CALLBACK_FAILURE;
+	}
+
+	return Outcome.GetResult();
+	*/
+	return -1;
+}
+
+////////////////////////////////////////////////////////////////////////////////
+int32 FDriverNg::OnHeader(nghttp2_rcbuf* Name, nghttp2_rcbuf* Value)
+{
+	if (ResponseHeaders == nullptr)
+	{
+		return NGHTTP2_ERR_CALLBACK_FAILURE;
+	}
+
+	FHeaderNg Header(Name, Value);
+	ResponseHeaders->Add(MoveTemp(Header));
+	return 0;
+}
+
+
+
+////////////////////////////////////////////////////////////////////////////////
+class FBase
+{
+public:
+					FBase(FDriverNg& InDriver);
+
+protected:
+	FDriverNg&		GetDriver();
+
+private:
+	FDriverNg&		Driver;
+};
+
+////////////////////////////////////////////////////////////////////////////////
+FBase::FBase(FDriverNg& InDriver)
+: Driver(InDriver)
+{
+}
+
+////////////////////////////////////////////////////////////////////////////////
+FDriverNg& FBase::GetDriver()
+{
+	return Driver;
+}
+
+
+
+////////////////////////////////////////////////////////////////////////////////
+class FRequest
+	: public FBase
+{
+public:
+	using FBase::FBase;
+
+	void			Begin(FAnsiStringView Host, FAnsiStringView Method, FAnsiStringView Path);
+	void			AddHeader(FAnsiStringView Key, FAnsiStringView Value);
+	FTransactId		End(bool bKeepAlive);
+	FOutcome		TrySendRequest(FTlsPeer& Peer);
+
+protected:
+	FTransactId		GetTransactId() const;
+
+private:
+	FDriverHeaders	RequestHeaders;
+	FTransactId		Id;
+};
+
+////////////////////////////////////////////////////////////////////////////////
+void FRequest::Begin(FAnsiStringView Host, FAnsiStringView Method, FAnsiStringView Path)
+{
+	FAnsiStringView Scheme = GetDriver().IsUsingTls() ? "https" : "http";
+	RequestHeaders.Add(":method", Method);
+	RequestHeaders.Add(":path", Path);
+	RequestHeaders.Add(":authority", Host);
+	RequestHeaders.Add(":scheme", Scheme);
+}
+
+////////////////////////////////////////////////////////////////////////////////
+void FRequest::AddHeader(FAnsiStringView Key, FAnsiStringView Value)
+{
+	check(!Key.Equals("connection", ESearchCase::IgnoreCase));
+
+#if UE_BUILD_DEVELOPMENT
+	for (uint32 Char : Key)
+	{
+		if (uint32(Char - 'A') > uint32('Z' - 'A'))
+		{
+			check("header keys must be lowercase");
+		}
+	}
+#endif
+	RequestHeaders.Add(Key, Value);
+}
+
+////////////////////////////////////////////////////////////////////////////////
+FTransactId FRequest::GetTransactId() const
+{
+	return Id;
+}
+
+////////////////////////////////////////////////////////////////////////////////
+FTransactId FRequest::End(bool /*bKeepAlive*/)
+{
+	Id = GetDriver().Submit(RequestHeaders);
+	return Id;
+}
+
+////////////////////////////////////////////////////////////////////////////////
+FOutcome FRequest::TrySendRequest(FTlsPeer& Peer)
+{
+	return GetDriver().Send(Peer);
+}
+
+
+
+////////////////////////////////////////////////////////////////////////////////
+class FStatusHeaders
+	: public FRequest
+{
+public:
+	using FRequest::FRequest;
+
+	FOutcome			TryRecvResponse(FTlsPeer& Peer);
+	bool				IsKeepAlive() const;
+	uint32				GetStatusCode() const;
+	FAnsiStringView		GetStatusMessage() const;
+	int64				GetContentLength() const;
+	bool				IsChunked() const;
+	void				ReadHeaders(FResponse::FHeaderSink Sink) const;
+
+private:
+	TArray<FHeaderNg>	ResponseHeaders;
+	int64				ContentLength = -1;
+	uint32				StatusCode = 0;
+};
+
+////////////////////////////////////////////////////////////////////////////////
+bool FStatusHeaders::IsKeepAlive() const
+{
+	return true;
+}
+
+////////////////////////////////////////////////////////////////////////////////
+uint32 FStatusHeaders::GetStatusCode() const
+{
+	return StatusCode;
+}
+
+////////////////////////////////////////////////////////////////////////////////
+FAnsiStringView	FStatusHeaders::GetStatusMessage() const
+{
+	return "";
+}
+
+////////////////////////////////////////////////////////////////////////////////
+int64 FStatusHeaders::GetContentLength() const
+{
+	return ContentLength;
+}
+
+////////////////////////////////////////////////////////////////////////////////
+bool FStatusHeaders::IsChunked() const
+{
+	return ContentLength < 0;
+}
+
+////////////////////////////////////////////////////////////////////////////////
+void FStatusHeaders::ReadHeaders(FResponse::FHeaderSink Sink) const
+{
+	for (const FHeaderNg& Header : ResponseHeaders)
+	{
+		FAnsiStringView Key = Header.Key.Get();
+		FAnsiStringView Value = Header.Value.Get();
+		if (!Sink(Key, Value))
+		{
+			break;
+		}
+	}
+}
+
+////////////////////////////////////////////////////////////////////////////////
+FOutcome FStatusHeaders::TryRecvResponse(FTlsPeer& Peer)
+{
+	FOutcome Outcome = GetDriver().RecvResponse(Peer, &ResponseHeaders);
+	if (!Outcome.IsOk())
+	{
+		return Outcome;
+	}
+
+	enum : uint32 {
+		Flag_Code	= 1 << 0,
+		Flag_Length = 1 << 1,
+		Flag_All	= Flag_Code | Flag_Length,
+	};
+	uint32 Flags = 0;
+	for (const FHeaderNg& Header : ResponseHeaders)
+	{
+		if (Flags == Flag_All)
+		{
+			break;
+		}
+
+		FAnsiStringView Key = Header.Key.Get();
+		FAnsiStringView Value = Header.Value.Get();
+
+		if ((Flags ^ Flag_Code) && (Key == ":status"))
+		{
+			StatusCode = uint32(CrudeToInt(Value));
+			Flags |= Flag_Code;
+		}
+		else if ((Flags ^ Flag_Length) && (Key == "content-length"))
+		{
+			ContentLength = CrudeToInt(Value);
+			Flags |= Flag_Length;
+		}
+	}
+
+	if ((Flags & Flag_Code) == 0 || !IsStatusCodeOk(StatusCode))
+	{
+		return FOutcome::Error("Invalid status code", StatusCode);
+	}
+
+	if (IsContentless(StatusCode))
+	{
+		ContentLength = 0;
+	}
+
+	return FOutcome::Ok();
+}
+
+
+
+////////////////////////////////////////////////////////////////////////////////
+class FBody
+	: public FStatusHeaders
+{
+public:
+	using FStatusHeaders::FStatusHeaders;
+
+	int64			GetRemaining() const;
+	FOutcome		TryRecv(FMutableMemoryView Dest, FTlsPeer& Peer);
+
+private:
+	int64			Remaining = -1;
+};
+
+////////////////////////////////////////////////////////////////////////////////
+int64 FBody::GetRemaining() const
+{
+	return Remaining;
+}
+
+////////////////////////////////////////////////////////////////////////////////
+FOutcome FBody::TryRecv(FMutableMemoryView Dest, FTlsPeer& Peer)
+{
+	if (!IsChunked())
+	{
+		if (Remaining < 0)
+		{
+			Remaining = GetContentLength();
+		}
+		Dest = Dest.Left(Remaining);
+	}
+
+	FOutcome Outcome = GetDriver().RecvBody(Peer, Dest);
+	if (Outcome.IsError())
+	{
+		return Outcome;
+	}
+
+	Remaining -= Outcome.GetResult();
+	return Outcome;
+}
+
+} // namespace DetailTwo
+
+
+
+////////////////////////////////////////////////////////////////////////////////
+class FTransactionTwo
+	: public DetailTwo::FBody
+{
+public:
+	using DetailTwo::FBody::FBody;
+};
+
+
+
+////////////////////////////////////////////////////////////////////////////////
+FOutcome HandshakeTwo(FTlsPeer& Peer, void*& PeerData)
+{
+	using namespace DetailTwo;
+
+	if (PeerData == nullptr)
+	{
+		PeerData = new FDriverNg(Peer);
+	}
+
+	auto* Driver = (FDriverNg*)PeerData;
+	return Driver->Handshake(Peer);
+}
+
+////////////////////////////////////////////////////////////////////////////////
+FOutcome TickTwo(FTlsPeer& Peer, void* PeerData)
+{
+	using namespace DetailTwo;
+
+	auto* Driver = (FDriverNg*)PeerData;
+	return Driver->GetTransactId(Peer);
+}
+
+////////////////////////////////////////////////////////////////////////////////
+void GoAwayTwo(FTlsPeer& Peer, void* PeerData)
+{
+	using namespace DetailTwo;
+
+	auto* Driver = (FDriverNg*)PeerData;
+	if (Driver == nullptr)
+	{
+		return;
+	}
+
+	Driver->GoAway(Peer);
+	delete Driver;
+}
+
+} // namespace UE::IoStore::HTTP
+
+#else // IAS_HTTP_WITH_TWO
+
+////////////////////////////////////////////////////////////////////////////////
+namespace UE::IoStore::HTTP {
+
+class FTransactionTwo
+{
+public:
+	void			Begin(...)									{}
+	void			AddHeader(...)								{}
+	FTransactId		End(...)									{ return 0; }
+	FOutcome		TrySendRequest(FTlsPeer&)					{ return FOutcome::None(); }
+	FOutcome		TryRecvResponse(FTlsPeer&)					{ return FOutcome::None(); }
+	bool			IsKeepAlive() const							{ return false; }
+	uint32			GetStatusCode() const						{ return 0; }
+	FAnsiStringView	GetStatusMessage() const					{ return {}; }
+	int64			GetContentLength() const					{ return -1; }
+	bool			IsChunked() const							{ return false; }
+	int64			GetRemaining() const						{ return -1; }
+	void			ReadHeaders(FResponse::FHeaderSink) const	{}
+	FOutcome		TryRecv(FMutableMemoryView, FTlsPeer&)		{ return FOutcome::None(); }
+};
+FOutcome HandshakeTwo(FTlsPeer&, void*&)	{ return FOutcome::None(); }
+FOutcome TickTwo(FTlsPeer&, void*)			{ return FOutcome::None(); }
+void GoAwayTwo(FTlsPeer&, void*)			{}
+
+} // namespace UE::IoStore::HTTP
+
+#endif // IAS_HTTP_WITH_TWO
