Missing break statements

#rnx

diff --git a/Engine/Source/Runtime/Experimental/IoStore/HttpClient/Private/Peer.inl b/Engine/Source/Runtime/Experimental/IoStore/HttpClient/Private/Peer.inl
--- a/Engine/Source/Runtime/Experimental/IoStore/HttpClient/Private/Peer.inl
+++ b/Engine/Source/Runtime/Experimental/IoStore/HttpClient/Private/Peer.inl
@@ -568,8 +568,8 @@ FHttpPeer::~FHttpPeer()
 {
 	switch (Proto)
 	{
-	case 1:		GoAwayOne(*this, PeerData);
-	case 2:		GoAwayTwo(*this, PeerData);
+	case 1:		GoAwayOne(*this, PeerData); break;
+	case 2:		GoAwayTwo(*this, PeerData); break;
 	default:	break;
 	}
 }
