Remove unecessary work inside static tree
Updating the static tree with dynamic and kinematic particles is not necessary.
If we configure Chaos to have only one tree for static, dynamic and kinematic particles we still need to do it.

diff --git a/Engine/Source/Runtime/Experimental/Chaos/Private/Chaos/PBDRigidsEvolutionGBF.cpp b/Engine/Source/Runtime/Experimental/Chaos/Private/Chaos/PBDRigidsEvolutionGBF.cpp
--- a/Engine/Source/Runtime/Experimental/Chaos/Private/Chaos/PBDRigidsEvolutionGBF.cpp
+++ b/Engine/Source/Runtime/Experimental/Chaos/Private/Chaos/PBDRigidsEvolutionGBF.cpp
@@ -983,9 +983,14 @@ void FPBDRigidsEvolutionGBF::Integrate(FReal Dt)
 
 	auto DirtyAsync = [this, &ParticlesArray]()
 	{
-		for (FPBDRigidParticleHandle* Particle : ParticlesArray)
+		// The async data structure is for the static particle. Here we update the dynamic particles.
+		// We need to update them only if dynamic and static particles are sharing the same spatial structure
+		if (AccelerationStructureSplitStaticAndDynamic != 1)
 		{
-			Base::DirtyParticleAsync(*Particle);
+			for (FPBDRigidParticleHandle* Particle : ParticlesArray)
+			{
+				Base::DirtyParticleAsync(*Particle);
+			}
 		}
 	};
 	
@@ -1050,15 +1055,25 @@ void FPBDRigidsEvolutionGBF::ApplyKinematicTargets(const FReal Dt, const FReal S
 		}
 	}
 
-	if (GSingleThreadedPhysics)
+	auto DirtyKinAsync = [this]()
 	{
-		for (FTransientPBDRigidParticleHandle& Particle : Particles.GetActiveMovingKinematicParticlesView())
+		// The async data structure is for the static particle. Here we update the kinematic particles.
+		// We need to update them only if kinematic and static particles are sharing the same spatial structure
+		if (AccelerationStructureSplitStaticAndDynamic != 1)
 		{
-			// Moving kinematics need to have SQ updates queued to apply substepped movement to the
-			// physics thread acceleration structure in order to get correct kinematic collisions
-			// and pass a correct structure back to the external thread
-			DirtyParticleAsync(Particle);
+			for (FTransientPBDRigidParticleHandle& Particle : Particles.GetActiveMovingKinematicParticlesView())
+			{
+				// Moving kinematics need to have SQ updates queued to apply substepped movement to the
+				// physics thread acceleration structure in order to get correct kinematic collisions
+				// and pass a correct structure back to the external thread
+				DirtyParticleAsync(Particle);
+			}
 		}
+	};
+
+	if (GSingleThreadedPhysics)
+	{
+		DirtyKinAsync();
 		// done with update, let's clear the tracking structures
 		if (bIsLastStep)
 		{
@@ -1069,17 +1084,6 @@ void FPBDRigidsEvolutionGBF::ApplyKinematicTargets(const FReal Dt, const FReal S
 	}
 	else
 	{
-		auto DirtyKinAsync = [this]()
-		{
-			for (FTransientPBDRigidParticleHandle& Particle : Particles.GetActiveMovingKinematicParticlesView())
-			{
-				// Moving kinematics need to have SQ updates queued to apply substepped movement to the
-				// physics thread acceleration structure in order to get correct kinematic collisions
-				// and pass a correct structure back to the external thread
-				DirtyParticleAsync(Particle);
-			}
-		};
-
 		TaskDispatcher.DispatchKinAsyncDirtyAndUpdateKinematic(DirtyKinAsync, bIsLastStep);
 	}
 }
