Constraints: make sure to avoid cycles when ensuring primary dependencies

#rnx

diff --git a/Engine/Source/Runtime/Experimental/Animation/Constraints/Private/Transform/TransformConstraint.cpp b/Engine/Source/Runtime/Experimental/Animation/Constraints/Private/Transform/TransformConstraint.cpp
--- a/Engine/Source/Runtime/Experimental/Animation/Constraints/Private/Transform/TransformConstraint.cpp
+++ b/Engine/Source/Runtime/Experimental/Animation/Constraints/Private/Transform/TransformConstraint.cpp
@@ -349,15 +349,19 @@ void UTickableTransformConstraint::EnsurePrimaryDependency(const UWorld* InWorld
 			FTickPrerequisite PrimaryPrerex = ParentTRSHandle->GetPrimaryPrerequisite(UE::Private::ShouldForceParentDependency());
 			if (FTickFunction* PotentialFunction = PrimaryPrerex.Get())
 			{
-				if (FDependencyBuilder::LogDependencies())
+				TSet<const FTickFunction*> VisitedFunctions;
+				if (!FConstraintCycleChecker::HasPrerequisiteDependencyWith(&ConstraintTick, PotentialFunction, VisitedFunctions))
 				{
-					UE_LOG(LogTemp, Warning, TEXT("EnsurePrimaryDependency: '%s' must tick before '%s'"),
-						*PotentialFunction->DiagnosticMessage(),
-						*ConstraintTick.DiagnosticMessage() );
+					if (FDependencyBuilder::LogDependencies())
+					{
+						UE_LOG(LogTemp, Warning, TEXT("EnsurePrimaryDependency: '%s' must tick before '%s'"),
+							*PotentialFunction->DiagnosticMessage(),
+							*ConstraintTick.DiagnosticMessage() );
+					}
+					
+					UObject* Target = PrimaryPrerex.PrerequisiteObject.Get();
+					ConstraintTick.AddPrerequisite(Target, *PotentialFunction);
 				}
-				
-				UObject* Target = PrimaryPrerex.PrerequisiteObject.Get();
-				ConstraintTick.AddPrerequisite(Target, *PotentialFunction);
 			}
 		}
 	}
