Allow IAD iochunks to be read out of pending buffer.
- Chunks are no longer fully flushed to disk on every install request.
- Space is still maintained in the cache for any pending chunks that have not yet been written.

diff --git a/Engine/Source/Runtime/Experimental/IoStore/OnDemand/Private/OnDemandContentInstaller.cpp b/Engine/Source/Runtime/Experimental/IoStore/OnDemand/Private/OnDemandContentInstaller.cpp
--- a/Engine/Source/Runtime/Experimental/IoStore/OnDemand/Private/OnDemandContentInstaller.cpp
+++ b/Engine/Source/Runtime/Experimental/IoStore/OnDemand/Private/OnDemandContentInstaller.cpp
@@ -951,12 +951,6 @@ void FOnDemandContentInstaller::ProcessDownloadedChunk(
 	const bool bCompleted = ++InstallRequest.DownloadedChunkCount >= InstallRequest.HttpRequestHandles.Num();
 	if (bCompleted)
 	{
-		//TODO: Don't flush after every install request
-		FResult FlushResult = IoStore.InstallCache->Flush();
-		if (FlushResult.HasError() && !Request.Result.HasError())
-		{
-			Request.Result = MoveTemp(FlushResult);
-		}
 		CompleteInstallRequest(Request);
 	}
 	else 
@@ -1183,7 +1177,7 @@ void FOnDemandContentInstaller::CompletePurgeRequest(FRequest& Request)
 	FOnDemandPurgeResult PurgeResult;
 	PurgeResult.DurationInSeconds	= DurationInSeconds;
 	PurgeResult.Error				= Request.ConsumeError();
-
+	
 	const bool bCallbackOnGameThread	= EnumHasAnyFlags(Request.AsPurge().Args.Options, EOnDemandPurgeOptions::CallbackOnGameThread);
 	FOnDemandPurgeCompleted OnCompleted = MoveTemp(Request.AsPurge().OnCompleted);
 	{
diff --git a/Engine/Source/Runtime/Experimental/IoStore/OnDemand/Private/OnDemandInstallCache.cpp b/Engine/Source/Runtime/Experimental/IoStore/OnDemand/Private/OnDemandInstallCache.cpp
--- a/Engine/Source/Runtime/Experimental/IoStore/OnDemand/Private/OnDemandInstallCache.cpp
+++ b/Engine/Source/Runtime/Experimental/IoStore/OnDemand/Private/OnDemandInstallCache.cpp
@@ -10,6 +10,7 @@
 #include "Async/Mutex.h"
 #include "Async/SharedMutex.h"
 #include "Async/UniqueLock.h"
+#include "Async/SharedLock.h"
 #include "Async/AsyncFileHandle.h"
 #include "Containers/UnrealString.h"
 #include "GenericHash.h"
@@ -1659,6 +1660,7 @@ class FOnDemandInstallCache final
 			TotalSize += Chunk.GetSize();
 			ChunkHashes.Add(ChunkHash);
 			Chunks.Add(MoveTemp(Chunk));
+			LastAccessTime.store(FDateTime::UtcNow().GetTicks(), std::memory_order_relaxed);
 		}
 
 		void Reset()
@@ -1666,14 +1668,26 @@ class FOnDemandInstallCache final
 			Chunks.Reset();
 			ChunkHashes.Reset();
 			TotalSize = 0;
+			LastAccessTime.store(0, std::memory_order_relaxed);
 		}
 
-		TArray<FIoBuffer>	Chunks;
-		TArray<FIoHash>		ChunkHashes;
-		uint64				TotalSize = 0;
-	};
+		const FIoBuffer* FindChunk(const FIoHash& ChunkHash) const
+		{
+			const int32 PendingIdx = ChunkHashes.IndexOfByKey(ChunkHash);
+			return (PendingIdx == INDEX_NONE) ? nullptr : &Chunks[PendingIdx];
+		}
 
-	using FUniquePendingChunks = TUniquePtr<FPendingChunks>;
+		bool ContainsChunk(const FIoHash& ChunkHash) const
+		{
+			return ChunkHashes.Contains(ChunkHash);
+		}
+
+		TArray<FIoBuffer>		Chunks;
+		TArray<FIoHash>			ChunkHashes;
+		uint64					TotalSize = 0;
+		std::atomic<uint64>		LastAccessTime{ 0 };
+		mutable FSharedMutex	SharedMutex;
+	};
 
 public:
 	FOnDemandInstallCache(const FOnDemandInstallCacheConfig& Config, FOnDemandIoStore& IoStore);
@@ -1737,7 +1751,7 @@ private:
 	FString						CacheDirectory;
 	FCas						Cas;
 	std::atomic<FCasBlockId>	CurrentBlock{ FCasBlockId::Invalid };
-	FUniquePendingChunks		PendingChunks;
+	FPendingChunks				PendingChunks;
 	FSharedBackendContext		BackendContext;
 	FIoRequestList				CompletedRequests;
 	FMutex						Mutex;
@@ -1907,6 +1921,11 @@ void FOnDemandInstallCache::Initialize(FSharedBackendContextRef Context)
 
 void FOnDemandInstallCache::Shutdown()
 {
+	if (FResult Result = FlushPendingChunksImpl(PendingChunks, PendingChunks.LastAccessTime.load(std::memory_order_relaxed)); Result.HasError())
+	{
+		UE_LOG(LogIoStoreOnDemand, Error, TEXT("Failed to flush pending chunks on shutdown, reason '%s'"), *LexToString(Result.GetError()));
+	}
+
 	FCas::FLastAccess LastAccess = Cas.ConsumeLastAcccess();
 
 	const FString JournalFile = GetJournalFilename();
@@ -1994,6 +2013,15 @@ bool FOnDemandInstallCache::DoesChunkExist(const FIoChunkId& ChunkId) const
 	EIoErrorCode ErrorCode = EIoErrorCode::UnknownChunkID;
 	if (FOnDemandChunkInfo ChunkInfo = IoStore.GetInstalledChunkInfo(ChunkId, ErrorCode))
 	{
+		{
+			TSharedLock SharedLock(PendingChunks.SharedMutex);
+
+			if (PendingChunks.ContainsChunk(ChunkInfo.Hash()))
+			{
+				return true;
+			}
+		}
+
 		const FCasLocation CasLoc = Cas.FindChunk(ChunkInfo.Hash());
 		return CasLoc.IsValid();
 	}
@@ -2036,13 +2064,6 @@ bool FOnDemandInstallCache::Resolve(FIoRequestImpl* Request)
 		return false;
 	}
 
-	const FCasLocation CasLoc = Cas.FindChunk(ChunkInfo.Hash());
-	if (CasLoc.IsValid() == false)
-	{
-		CompleteRequest(Request, EIoErrorCode::NotInstalled);
-		return true;
-	}
-
 	const uint64 RequestSize = FMath::Min<uint64>(
 		Request->Options.GetSize(),
 		ChunkInfo.RawSize() - Request->Options.GetOffset());
@@ -2061,6 +2082,46 @@ bool FOnDemandInstallCache::Resolve(FIoRequestImpl* Request)
 		return true;
 	}
 
+	// Must check pending chunks before CAS because otherwise
+	// a chunk could be flushed after checking the CAS and before checking pending.
+	// Otherise we'd need take both locks.
+	{
+		TDynamicSharedLock SharedLock(PendingChunks.SharedMutex);
+
+		const FIoBuffer* PendingChunk = PendingChunks.FindChunk(ChunkInfo.Hash());
+		if (PendingChunk != nullptr)
+		{
+			FChunkRequest& ChunkRequest = FChunkRequest::Attach(*Request, new FChunkRequest(
+				FSharedAsyncFileHandle(),
+				Request,
+				MoveTemp(ChunkInfo),
+				ChunkRange.ConsumeValueOrDie(),
+				RequestSize));
+
+			FMemoryView BytesToRead = PendingChunk->GetView().Mid(ChunkRequest.ChunkRange.GetOffset(), ChunkRequest.EncodedChunk.GetSize());
+			check(BytesToRead.GetSize() == ChunkRequest.EncodedChunk.GetSize());
+			FMemory::Memcpy(ChunkRequest.EncodedChunk.GetData(), BytesToRead.GetData(), ChunkRequest.EncodedChunk.GetSize());
+
+			PendingChunks.LastAccessTime.store(FDateTime::UtcNow().GetTicks(), std::memory_order_relaxed);
+
+			SharedLock.Unlock();
+
+			UE::Tasks::Launch(UE_SOURCE_LOCATION, [this, Request]
+			{
+				CompleteRequest(Request, EIoErrorCode::Ok);
+			});
+
+			return true;
+		}
+	}
+
+	const FCasLocation CasLoc = Cas.FindChunk(ChunkInfo.Hash());
+	if (CasLoc.IsValid() == false)
+	{
+		CompleteRequest(Request, EIoErrorCode::NotInstalled);
+		return true;
+	}
+
 	TRACE_IOSTORE_BACKEND_REQUEST_STARTED(Request, this);
 	constexpr const bool bDirty = true;
 	if (Cas.TrackAccessIf(ECasTrackAccessType::Granular, CasLoc.BlockId, bDirty))
@@ -2189,6 +2250,15 @@ bool FOnDemandInstallCache::Resolve(FIoRequestImpl* Request)
 
 bool FOnDemandInstallCache::IsChunkCached(const FIoHash& ChunkHash)
 {
+	{
+		TSharedLock SharedLock(PendingChunks.SharedMutex);
+
+		if (PendingChunks.ContainsChunk(ChunkHash))
+		{
+			return true;
+		}
+	}
+
 	const FCasLocation Loc = Cas.FindChunk(ChunkHash);
 	return Loc.IsValid();
 }
@@ -2199,18 +2269,29 @@ bool FOnDemandInstallCache::TryPinChunks(
 	FOnDemandContentHandle ContentHandle,
 	TArray<int32>& OutMissing)
 {
+	TDynamicSharedLock SharedLock(PurgeDefragMutex, UE::DeferLock);
+
 	// Currently we can only pin chunks if we are not purging/defragging the cache 
-	if (PurgeDefragMutex.TryLockShared() == false)
+	if (SharedLock.TryLock() == false)
 	{
 		OutMissing = EntryIndices;
 		return false;
 	}
 
-	ON_SCOPE_EXIT { PurgeDefragMutex.UnlockShared(); };
-
 	for (int32 EntryIndex : EntryIndices)
 	{
 		const FOnDemandChunkEntry& ChunkEntry = Container->ChunkEntries[EntryIndex];
+
+		{
+			TSharedLock PendingSharedLock(PendingChunks.SharedMutex);
+
+			if (PendingChunks.ContainsChunk(ChunkEntry.Hash))
+			{
+				IoStore.AddReference(Container, EntryIndex, ContentHandle);
+				continue;
+			}
+		}
+
 		if (FCasLocation Loc = Cas.FindChunk(ChunkEntry.Hash); Loc.IsValid())
 		{
 			IoStore.AddReference(Container, EntryIndex, ContentHandle);
@@ -2226,21 +2307,19 @@ bool FOnDemandInstallCache::TryPinChunks(
 
 FResult FOnDemandInstallCache::PutChunk(FIoBuffer&& Chunk, const FIoHash& ChunkHash)
 {
-	if (PendingChunks.IsValid() == false)
-	{
-		PendingChunks = MakeUnique<FPendingChunks>();
-	}
-
-	if (PendingChunks->TotalSize > FPendingChunks::MaxPendingBytes)
+	if (PendingChunks.TotalSize > FPendingChunks::MaxPendingBytes)
 	{
-		if (FResult Result = FlushPendingChunks(*PendingChunks); Result.HasError())
+		if (FResult Result = FlushPendingChunks(PendingChunks, PendingChunks.LastAccessTime.load(std::memory_order_relaxed));
+			Result.HasError())
 		{
 			return Result;
 		}
-		check(PendingChunks->IsEmpty());
+		check(PendingChunks.IsEmpty());
 	}
 
-	PendingChunks->Append(MoveTemp(Chunk), ChunkHash);
+	TUniqueLock Lock(PendingChunks.SharedMutex);
+
+	PendingChunks.Append(MoveTemp(Chunk), ChunkHash);
 	return MakeValue();
 }
 
@@ -2767,7 +2846,7 @@ uint64 FOnDemandInstallCache::AddReferencesToBlocks(
 	FCasBlockInfoMap& BlockInfoMap,
 	uint64& OutTotalReferencedBytes) const
 {
-	uint64 TotalUncachedBytes = 0;
+	uint64 TotalUncachedBytes = PendingChunks.TotalSize;
 	OutTotalReferencedBytes	= 0;
 
 	for (int32 Index = 0; FSharedOnDemandContainer Container : Containers)
@@ -2799,7 +2878,14 @@ uint64 FOnDemandInstallCache::AddReferencesToBlocks(
 				}
 				else
 				{
-					UE_CLOG(bIsReferenced, LogIoStoreOnDemand, Error, TEXT("Failed to find CAS location for referenced chunk, ChunkId='%s', Container='%s'"),
+					// Check pending list
+					bool bChunkIsPending = false;
+					{
+						TSharedLock SharedLock(PendingChunks.SharedMutex);
+						bChunkIsPending = PendingChunks.ContainsChunk(Entry.Hash);
+					}
+
+					UE_CLOG(bIsReferenced && !bChunkIsPending, LogIoStoreOnDemand, Error, TEXT("Failed to find CAS location or pending chunk for chunk reference, ChunkId='%s', Container='%s'"),
 						*LexToString(Container->ChunkIds[EntryIndex]), *Container->Name);
 				}
 			}
@@ -3202,6 +3288,7 @@ FResult FOnDemandInstallCache::Defrag(
 						FOnDemandInstallCacheStats::OnDefrag(Result, 0);
 						return Result;
 					}
+					check(DefragPendingChunks.IsEmpty());
 	
 					FResult Result = MakeError(UE::UnifiedError::IoStoreOnDemand::InstallCacheDefragError::MakeError(
 						UE::UnifiedError::IoStoreOnDemand::FChunkHashMismatchErrorContext
@@ -3297,11 +3384,12 @@ FResult FOnDemandInstallCache::Defrag(
 
 FResult FOnDemandInstallCache::Flush()
 {
-	if (PendingChunks.IsValid())
+	if (FResult Result = FlushPendingChunks(PendingChunks, PendingChunks.LastAccessTime.load(std::memory_order_relaxed));
+		Result.HasError())
 	{
-		FUniquePendingChunks Chunks = MoveTemp(PendingChunks);
-		return FlushPendingChunks(*Chunks);
+		return Result;
 	}
+	check(PendingChunks.IsEmpty());
 
 	Cas.Compact();
 	return MakeValue();
@@ -3437,7 +3525,11 @@ FResult FOnDemandInstallCache::FlushPendingChunks(FPendingChunks& Chunks, int64
 
 FResult FOnDemandInstallCache::FlushPendingChunksImpl(FPendingChunks& Chunks, int64 UtcAccessTicks)
 {
-	ON_SCOPE_EXIT { Chunks.Reset(); };
+	ON_SCOPE_EXIT 
+	{ 
+		TUniqueLock Lock(Chunks.SharedMutex);
+		Chunks.Reset();
+	};
 
 	FLargeMemoryWriter Ar(FMath::Min(Chunks.TotalSize, static_cast<uint64>(Cas.GetMaxBlockSize())));
 
