Register BPS CDN endpoints as the default on-demand host group
- this change fixes an issue where IAS containers not being mounted due to missing Default host group. This
can occur when the IAD bundle is not installed or IAD is disabled.

#rnx

#DisabledValidations Preflight

diff --git a/Engine/Source/Runtime/Experimental/IoStore/OnDemand/Private/OnDemandIoStore.cpp b/Engine/Source/Runtime/Experimental/IoStore/OnDemand/Private/OnDemandIoStore.cpp
--- a/Engine/Source/Runtime/Experimental/IoStore/OnDemand/Private/OnDemandIoStore.cpp
+++ b/Engine/Source/Runtime/Experimental/IoStore/OnDemand/Private/OnDemandIoStore.cpp
@@ -587,6 +587,84 @@ FIoStatus FOnDemandIoStore::InitializeStreamingBackend(const FOnDemandEndpointCo
 	return FIoStatus::Ok;
 }
 
+FOnDemandRegisterHostGroupResult FOnDemandIoStore::RegisterHostGroup(FOnDemandRegisterHostGroupArgs&& Args)
+{
+	using namespace UE::UnifiedError::Core;
+
+	auto SanitizeHostNames = [bUseSecureHttp = Args.bUseSecureHttp](TArrayView<FString> HostNames) -> TArray<FAnsiString>
+	{
+		TArray<FAnsiString> OutHostNames;
+		for (FString& HostName : HostNames)
+		{
+			if (HostName.StartsWith(TEXT("http")) == false)
+			{
+				HostName = TEXT("https://") + HostName;
+			}
+			if (bUseSecureHttp == false)
+			{
+				HostName.ReplaceInline(TEXT("https"), TEXT("http"));
+			}
+			HostName.RemoveFromEnd(TEXT("/"));
+			OutHostNames.Add(FAnsiString(StringCast<ANSICHAR>(*HostName)));
+		}
+
+		return OutHostNames;
+	};
+
+	if (Args.HostGroupName.IsNone())
+	{
+		return FOnDemandRegisterHostGroupResult
+		{
+			.Error = ArgumentError::MakeError(TEXT("HostGroupName"), TEXT("Host group name cannot be empty"))
+		};
+	}
+
+	if (Args.HostNames.IsEmpty())
+	{
+		return FOnDemandRegisterHostGroupResult
+		{
+			.Error = ArgumentError::MakeError(TEXT("HostNames"), TEXT("No host name(s) specified"))
+		};
+	}
+
+	TArray<FAnsiString> HostNames	= SanitizeHostNames(Args.HostNames);
+	FAnsiString AnsiTestUrl			= FAnsiString(StringCast<ANSICHAR>(*Args.TestUrl));
+
+	FIoStatus Status = FIoStatus::Invalid;
+	FIASHostGroup RegisteredHostGroup;
+
+	if (FHttpIoDispatcher::IsInitialized())
+	{
+		Status				= FHttpIoDispatcher::RegisterHostGroup(Args.HostGroupName, HostNames, AnsiTestUrl); 
+		RegisteredHostGroup = FHostGroupManager::Get().Find(Args.HostGroupName); // The HTTP I/O dispatcher currenlty uses the host group manager
+	}
+	else
+	{
+		if (TIoStatusOr<FIASHostGroup> RegisterStatus = FHostGroupManager::Get().Register(Args.HostGroupName, AnsiTestUrl); RegisterStatus.IsOk())
+		{
+			RegisteredHostGroup = RegisterStatus.ConsumeValueOrDie();
+			Status				= RegisteredHostGroup.Resolve(HostNames);
+		}
+		else
+		{
+			Status = RegisterStatus.Status();
+		}
+	}
+
+	if (Status.IsOk() == false)
+	{
+		return FOnDemandRegisterHostGroupResult
+		{
+			.Error = ArgumentError::MakeError(TEXT("Args"), Status.ToString())
+		};
+	}
+
+	return FOnDemandRegisterHostGroupResult
+	{
+		.HostGroup = RegisteredHostGroup.GetUnderlyingHostGroup()
+	};
+}
+
 void FOnDemandIoStore::Mount(FOnDemandMountArgs&& Args, FOnDemandMountCompleted&& OnCompleted)
 {
 	FSharedMountRequest MountRequest = MakeShared<FMountRequest>();
diff --git a/Engine/Source/Runtime/Experimental/IoStore/OnDemand/Private/OnDemandIoStore.h b/Engine/Source/Runtime/Experimental/IoStore/OnDemand/Private/OnDemandIoStore.h
--- a/Engine/Source/Runtime/Experimental/IoStore/OnDemand/Private/OnDemandIoStore.h
+++ b/Engine/Source/Runtime/Experimental/IoStore/OnDemand/Private/OnDemandIoStore.h
@@ -361,6 +361,7 @@ public:
 	// IOnDemandIoStore interface
 	virtual FIoStatus					Initialize() override;
 	virtual FIoStatus					InitializePostHotfix() override;
+	virtual FOnDemandRegisterHostGroupResult RegisterHostGroup(FOnDemandRegisterHostGroupArgs&& Args) override;
 	virtual void						Mount(FOnDemandMountArgs&& Args, FOnDemandMountCompleted&& OnCompleted) override;
 	virtual FOnDemandInstallRequest 	Install(FOnDemandInstallArgs&& Args,
 											FOnDemandInstallCompleted&& OnCompleted,
diff --git a/Engine/Source/Runtime/Experimental/IoStore/OnDemandCore/Public/IO/IoStoreOnDemand.h b/Engine/Source/Runtime/Experimental/IoStore/OnDemandCore/Public/IO/IoStoreOnDemand.h
--- a/Engine/Source/Runtime/Experimental/IoStore/OnDemandCore/Public/IO/IoStoreOnDemand.h
+++ b/Engine/Source/Runtime/Experimental/IoStore/OnDemandCore/Public/IO/IoStoreOnDemand.h
@@ -115,6 +115,30 @@ private:
 	FSharedInternalContentHandle Handle;
 };
 
+/** Arguments for registering a new host group. */ 
+struct FOnDemandRegisterHostGroupArgs
+{
+	/** Name of the host group. */
+	FName HostGroupName;
+	/** List of host name(s). */ 
+	TArray<FString> HostNames;
+	/** Test URL. */
+	FString TestUrl;
+	/** Whether to use secure https. */
+	bool bUseSecureHttp = false;
+};
+
+/** Holds the result when registering a host group. */
+struct FOnDemandRegisterHostGroupResult
+{
+	/** Returns True if the request succeeded. */
+	bool IsOk() const { return Error.IsSet() == false; }
+	/** The registered host group. */
+	FOnDemandHostGroup HostGroup;
+	/** Error information about the request. */ 
+	TOptional<UE::UnifiedError::FError> Error;
+};
+
 /** Options for controlling the behavior of mounted container(s). */
 enum class EOnDemandMountOptions
 {
@@ -538,6 +562,12 @@ public:
 	virtual ~IOnDemandIoStore() = default;
 	/** Initialize the I/O store. Called after the module feature has been registered. */
 	virtual FIoStatus InitializePostHotfix() = 0;
+	/**
+	 * Register a new host group.
+	 * @param	Args	Request arguments.
+	 * @return			Error information about the request.
+	 */
+	virtual FOnDemandRegisterHostGroupResult RegisterHostGroup(FOnDemandRegisterHostGroupArgs&& Args) = 0;
 	/**
 	 * Mount an on-demand container.
 	 * @param	Args		Mount arguments.
