Improve the cancelation of inflight IAX http requests (currently disabled)

#rnx

- This feature can be enabled/disabled by a new cvar 'iax.EnableHttpCancellation' currently defaults to false.
- FMultiEndpointHttpClient::Get now returns a new type FHttpTicketId which represents the current inflight http request and can be used to call the new method FMultiEndpointHttpClient::CancelRequest to attempt to cancel the request early.
-- Note that This identifier remains consistent accross retries, mapping this value to the actual ticket used by HTTP::FEventLoop will be handled internally and should be hidden from the caller.
- The struct FHttpRequest used by the FOnDemandHttpThread will now store this ticket id so it can try and cancel the inflight requests when needed.
- When a new request is made in FMultiEndpointHttpClient we assign it a unique FHttpTicketId which will then be mapped to a FTicketInfo which tracks the current inflight HTTP::FTicket and if any attempt has been made to cancel the overall request.
- When a cancel request is first recieved we will attempt to tell the HTTP::FEventLoop that it should cancel the inflight request. It is possible that the request has recently completed on the main http thread and either just been processed (in which case it should not be in the map) or is waiting to be re-issued (in which case it will not have a valid HTTP::FTicket) so we also mark FTicketInfo::bCancelRequested as true so future operations know that the user intended for this to be canceled.
- When re-issuing a request we first check if cancel has been requested and if so we just invoke the callback of the request and consider the request done. If not we re-issue the request which will update FTicketInfo::HttpTicket with the new requests ticket value once done.
- Calls to FMultiEndpointHttpClient::CancelRequest are allowed on any thread which is why the data structures for cancelation are protected by a recursive mutex.

diff --git a/Engine/Source/Runtime/Experimental/IoStore/OnDemand/Private/OnDemandHttpClient.cpp b/Engine/Source/Runtime/Experimental/IoStore/OnDemand/Private/OnDemandHttpClient.cpp
--- a/Engine/Source/Runtime/Experimental/IoStore/OnDemand/Private/OnDemandHttpClient.cpp
+++ b/Engine/Source/Runtime/Experimental/IoStore/OnDemand/Private/OnDemandHttpClient.cpp
@@ -2,6 +2,7 @@
 
 #include "OnDemandHttpClient.h"
 
+#include "Async/UniqueLock.h"
 #include "Containers/StringConv.h"
 #include "GenericPlatform/GenericPlatformProcess.h"
 #include "HAL/FileManager.h"
@@ -139,6 +140,30 @@ FIoStatus LoadDefaultHttpCertificates(bool& bWasLoaded)
 	return DefaultCerts.Status;
 }
 
+/**
+ * Generates a new FHttpTicketId when called.
+ * 
+ * Calling this should be thread safe and will always return a new unique non zero identifier.
+ * In theory identifiers may be reused once the internal counter has overflowed but in practical
+ * terms this should never cause us a problem as the identifiers are only used for a relatively
+ * short lifespan.
+ */
+FMultiEndpointHttpClient::FHttpTicketId GenerateTicketId()
+{
+	static std::atomic<FMultiEndpointHttpClient::FHttpTicketId> NextFreeId = 1;
+
+	FMultiEndpointHttpClient::FHttpTicketId NewId = 0;
+
+	// To protect against the unlikely scenario where we somehow assign all of the ids and end up
+	// overflowing we need to make sure that we don't return 0 as an Id.
+	do 
+	{
+		NewId = NextFreeId++;
+	} while (NewId == 0);
+	
+	return NewId;
+}
+
 ///////////////////////////////////////////////////////////////////////////////
 
 FMultiEndpointHttpClient::FMultiEndpointHttpClient(const FMultiEndpointHttpClientConfig& InConfig)
@@ -147,6 +172,11 @@ FMultiEndpointHttpClient::FMultiEndpointHttpClient(const FMultiEndpointHttpClien
 	EventLoop.SetFailTimeout(Config.TimeoutMs);
 }
 
+FMultiEndpointHttpClient::~FMultiEndpointHttpClient()
+{
+	check(TicketLookupMap.IsEmpty());
+}
+
 TUniquePtr<FMultiEndpointHttpClient> FMultiEndpointHttpClient::Create(const FMultiEndpointHttpClientConfig& Config)
 {
 	return TUniquePtr<FMultiEndpointHttpClient>(new FMultiEndpointHttpClient(Config));
@@ -219,22 +249,22 @@ TIoStatusOr<FMultiEndpointHttpClientResponse> FMultiEndpointHttpClient::Get(FAns
 	return FIoStatus(EIoErrorCode::ReadError, Reason.ToView());
 }
 
-void FMultiEndpointHttpClient::Get(const FOnDemandHostGroup& HostGroup, FAnsiStringView RelativeUrl, FOnHttpResponse&& OnResponse)
+FMultiEndpointHttpClient::FHttpTicketId FMultiEndpointHttpClient::Get(const FOnDemandHostGroup& HostGroup, FAnsiStringView RelativeUrl, FOnHttpResponse&& OnResponse)
 {
 	FIoOffsetAndLength DefaultChunkRange(0,0);
-	Get(HostGroup, RelativeUrl, DefaultChunkRange, MoveTemp(OnResponse));
+	return Get(HostGroup, RelativeUrl, DefaultChunkRange, MoveTemp(OnResponse));
 }
 
-void FMultiEndpointHttpClient::Get(const FOnDemandHostGroup& HostGroup, FAnsiStringView RelativeUrl, const FIoOffsetAndLength& ChunkRange, FOnHttpResponse&& OnResponse)
+FMultiEndpointHttpClient::FHttpTicketId FMultiEndpointHttpClient::Get(const FOnDemandHostGroup& HostGroup, FAnsiStringView RelativeUrl, const FIoOffsetAndLength& ChunkRange, FOnHttpResponse&& OnResponse)
 {
-	Get(HostGroup, RelativeUrl, ChunkRange, TArray<FAnsiString>(), MoveTemp(OnResponse));
+	return Get(HostGroup, RelativeUrl, ChunkRange, TArray<FAnsiString>(), MoveTemp(OnResponse));
 }
 
-void FMultiEndpointHttpClient::Get(const FOnDemandHostGroup& HostGroup, FAnsiStringView RelativeUrl, const FIoOffsetAndLength& ChunkRange, TArray<FAnsiString>&& Headers, FOnHttpResponse&& OnResponse)
+FMultiEndpointHttpClient::FHttpTicketId FMultiEndpointHttpClient::Get(const FOnDemandHostGroup& HostGroup, FAnsiStringView RelativeUrl, const FIoOffsetAndLength& ChunkRange, TArray<FAnsiString>&& Headers, FOnHttpResponse&& OnResponse)
 {
 	FConnection& Connection = GetConnection(HostGroup);
 
-	IssueRequest(FRequest
+	FHttpTicketId TicketId = IssueRequest(FRequest
 	{
 		.OnResponse		= MoveTemp(OnResponse),
 		.Headers		= MoveTemp(Headers),
@@ -244,22 +274,18 @@ void FMultiEndpointHttpClient::Get(const FOnDemandHostGroup& HostGroup, FAnsiStr
 		.StartTime		= FPlatformTime::Cycles64(),
 		.Host			= Connection.CurrentHost
 	});
+
+	return TicketId;
 }
 
 bool FMultiEndpointHttpClient::Tick(int32 WaitTimeMs, uint32 MaxKiBPerSecond)
 {
+	TRACE_CPUPROFILER_EVENT_SCOPE(FMultiEndpointHttpClient::Tick);
 	EventLoop.Throttle(MaxKiBPerSecond);
+
 	const uint32 TicketCount = EventLoop.Tick(WaitTimeMs);
 
-	if (Retries.IsEmpty() == false)
-	{
-		const int32 RequestCount = FMath::Min(Retries.Num(), int32(HTTP::FEventLoop::MaxActiveTickets - TicketCount));
-		for (int32 Idx = 0; Idx < RequestCount; Idx++)
-		{
-			IssueRequest(MoveTemp(Retries[Idx]));
-		}
-		Retries.RemoveAtSwap(0, RequestCount);
-	}
+	ProcessRetryAttempts(TicketCount);
 
 	const bool bIsIdle = EventLoop.IsIdle();
 	if (bIsIdle)
@@ -280,6 +306,29 @@ bool FMultiEndpointHttpClient::Tick(int32 WaitTimeMs, uint32 MaxKiBPerSecond)
 	return bIsIdle == false;
 }
 
+void FMultiEndpointHttpClient::CancelRequest(FHttpTicketId TicketId)
+{
+	if (TicketId == 0)
+	{
+		return;
+	}
+
+	UE::TUniqueLock _(TicketLookupMutex);
+
+	// Note that normally we would expect to find an entry for a valid FHttpTicketId
+	// but it is possible that the request just completed and was removed from the
+	//  map. See the HttpSink in ::IssueRequest.
+	if (FTicketInfo* Ticket = TicketLookupMap.Find(TicketId))
+	{
+		Ticket->bCancelRequested = true;
+
+		if (Ticket->HttpTicket != 0)
+		{
+			EventLoop.Cancel(Ticket->HttpTicket);
+		}
+	}
+}
+
 void FMultiEndpointHttpClient::UpdateConnections()
 {
 	TRACE_CPUPROFILER_EVENT_SCOPE(FMultiEndpointHttpClient::UpdateConnections);
@@ -300,7 +349,7 @@ void FMultiEndpointHttpClient::UpdateConnections()
 	}
 }
 
-void FMultiEndpointHttpClient::IssueRequest(FRequest&& Request)
+FMultiEndpointHttpClient::FHttpTicketId FMultiEndpointHttpClient::IssueRequest(FRequest&& Request)
 {
 	using namespace UE::IoStore::HTTP;
 
@@ -338,6 +387,14 @@ void FMultiEndpointHttpClient::IssueRequest(FRequest&& Request)
 			WriteToAnsiString<64>(ANSITEXTVIEW("bytes="), Request.Range.GetOffset(), ANSITEXTVIEW("-"), Request.Range.GetOffset() + Request.Range.GetLength() - 1));
 	}
 
+	// If the request is new then we need to assign an identifier
+	if (Request.TicketId == 0)
+	{
+		Request.TicketId = GenerateTicketId();
+	}
+
+	FHttpTicketId RequestId = Request.TicketId;
+
 	auto HttpSink = [this, Request = MoveTemp(Request)](const FTicketStatus& TicketStatus) mutable
 	{
 		switch (TicketStatus.GetId())
@@ -435,13 +492,19 @@ void FMultiEndpointHttpClient::IssueRequest(FRequest&& Request)
 					}
 					else if (TicketStatus.GetId() == FTicketStatus::EId::Cancelled)
 					{
-						Response.Reason = TEXT("Cancelled");
+						Response.Reason = TEXT("Canceled");
+						Response.bCanceled = true;
 					}
 					else if (!Request.ResponseMessage.IsEmpty())
 					{
 						Response.Reason = MoveTemp(Request.ResponseMessage);
 					}
 
+					{
+						UE::TUniqueLock _(TicketLookupMutex);
+						TicketLookupMap.Remove(Request.TicketId);
+					}
+
 					FOnHttpResponse OnResponse = MoveTemp(Request.OnResponse);
 					OnResponse(MoveTemp(Response));
 				}
@@ -450,7 +513,15 @@ void FMultiEndpointHttpClient::IssueRequest(FRequest&& Request)
 		}
 	};
 
-	EventLoop.Send(MoveTemp(HttpRequest), MoveTemp(HttpSink));
+	FTicket RequestTicket = EventLoop.Send(MoveTemp(HttpRequest), MoveTemp(HttpSink));
+	check(RequestTicket != 0);
+
+	{
+		UE::TUniqueLock _(TicketLookupMutex);
+		TicketLookupMap.Add(RequestId, FTicketInfo{ .HttpTicket = RequestTicket});
+	}
+
+	return RequestId;
 }
 
 void FMultiEndpointHttpClient::RetryRequest(FRequest&& Request)
@@ -467,19 +538,87 @@ void FMultiEndpointHttpClient::RetryRequest(FRequest&& Request)
 		}
 	}
 
+	// Clear the ticket association, note that we don't check if the request is canceled as we'd have to
+	// check it again anyway when the retry attempts are being reissued so we might as well wait and do
+	// it in one place.
+	{
+		UE::TUniqueLock _(TicketLookupMutex);
+		if (FTicketInfo* TicketInfo = TicketLookupMap.Find(Request.TicketId))
+		{
+			TicketInfo->HttpTicket = 0;
+		}
+		else
+		{
+			checkNoEntry();
+		}
+	}
+
 	Request.StatusCode = 0;
 	Request.RetryCount++;
 	Request.Host = Connection.CurrentHost;
 	Retries.Emplace(MoveTemp(Request));
 }
 
+void FMultiEndpointHttpClient::ProcessRetryAttempts(uint32 TicketCount)
+{
+	if(Retries.IsEmpty())
+	{
+		return;
+	}
+
+	TRACE_CPUPROFILER_EVENT_SCOPE(FMultiEndpointHttpClient::ProcessRetryAttempts);
+	UE::TUniqueLock _(TicketLookupMutex);
+
+	const int32 RequestCount = FMath::Min(Retries.Num(), int32(HTTP::FEventLoop::MaxActiveTickets - TicketCount));
+	for (int32 Idx = 0; Idx < RequestCount; Idx++)
+	{
+		FRequest& Request = Retries[Idx];
+
+		if (FTicketInfo* TicketInfo = TicketLookupMap.Find(Request.TicketId))
+		{
+			check(TicketInfo->HttpTicket == 0);
+
+			if (!TicketInfo->bCancelRequested)
+			{
+				IssueRequest(MoveTemp(Request));
+			}
+			else
+			{
+				FMultiEndpointHttpClientResponse Response
+				{
+					.Reason = TEXT("Canceled"),
+					.DurationMilliseconds = uint64(FPlatformTime::ToMilliseconds64(FPlatformTime::Cycles64() - Request.StartTime)),
+					.StatusCode = Request.StatusCode,
+					.RetryCount = Request.RetryCount,
+					.HostIndex = Request.Host,
+					.bCanceled = true,
+					.CDNCacheStatus = Request.CDNCacheStatus
+				};
+
+				// The request is completed so we can remove it from the lookup map
+				TicketLookupMap.Remove(Request.TicketId);
+
+				// TODO: Consider moving this outside of the TicketLookupMutex lock?
+				FOnHttpResponse OnResponse = MoveTemp(Request.OnResponse);
+				OnResponse(MoveTemp(Response));	
+			}
+		}
+		else
+		{
+			checkNoEntry(); // If the request is ongoing then we should always have a valid entry
+		}
+	}
+
+	Retries.RemoveAtSwap(0, RequestCount);
+}
+
 TUniquePtr<HTTP::FConnectionPool> FMultiEndpointHttpClient::CreateConnection(FAnsiStringView HostUrl) const
 {
 	using namespace HTTP;
 
 	FConnectionPool::FParams Params;
 	ensure(Params.SetHostFromUrl(HostUrl) >= 0);
-	
+
 	if (Config.ReceiveBufferSize >= 0)
 	{
 		Params.RecvBufSize = Config.ReceiveBufferSize;
@@ -513,9 +652,9 @@ FMultiEndpointHttpClient::FConnection& FMultiEndpointHttpClient::GetConnection(c
 	}
 
 	FConnection& Conn = *Connections.Emplace_GetRef(new FConnection
-	{
-		.HostGroup	= HostGroup,
-	});
+		{
+			.HostGroup = HostGroup,
+		});
 
 	Conn.Pools.SetNum(HostGroup.Hosts().Num());
 	Conn.CurrentHost = HostGroup.PrimaryHostIndex();
@@ -548,8 +687,13 @@ void FMultiEndpointHttpClient::Log(const FRequest& Request, const HTTP::FTicketS
 #if !NO_LOGGING
 
 	using namespace UE::IoStore::HTTP;
-	 
-	const ELogVerbosity::Type Verbosity = !GIaxWarnOnFailure || IsHttpStatusOk(Request.StatusCode) ? Config.LogVerbosity : ELogVerbosity::Type::Warning;
+
+	ELogVerbosity::Type Verbosity = Config.LogVerbosity;
+	// We count  a request as failed if it's status code is bad AND the request was not canceled
+	if (GIaxWarnOnFailure && !IsHttpStatusOk(Request.StatusCode) && TicketStatus.GetId() != FTicketStatus::EId::Cancelled)
+	{
+		Verbosity = ELogVerbosity::Type::Warning;
+	}
 
 	if (Config.LogCategory == nullptr || Config.LogCategory->IsSuppressed(Verbosity))
 	{
@@ -566,7 +710,7 @@ void FMultiEndpointHttpClient::Log(const FRequest& Request, const HTTP::FTicketS
 	}
 	else if (TicketStatus.GetId() == FTicketStatus::EId::Cancelled)
 	{
-		Reason << TEXT(" Cancelled");
+		Reason << TEXT(" Canceled");
 	}
 	else if (!Request.ResponseMessage.IsEmpty())
 	{
diff --git a/Engine/Source/Runtime/Experimental/IoStore/OnDemand/Private/OnDemandHttpClient.h b/Engine/Source/Runtime/Experimental/IoStore/OnDemand/Private/OnDemandHttpClient.h
--- a/Engine/Source/Runtime/Experimental/IoStore/OnDemand/Private/OnDemandHttpClient.h
+++ b/Engine/Source/Runtime/Experimental/IoStore/OnDemand/Private/OnDemandHttpClient.h
@@ -2,6 +2,7 @@
 
 #pragma once
 
+#include "Async/RecursiveMutex.h"
 #include "Containers/AnsiString.h"
 #include "Containers/StringView.h"
 #include "Containers/UnrealString.h"
@@ -73,11 +74,16 @@ struct FMultiEndpointHttpClientConfig
 
 struct FMultiEndpointHttpClientResponse
 {
-	const bool	IsOk() const
+	bool IsOk() const
 	{
 		return IsHttpStatusOk(StatusCode);
 	}
 
+	bool IsCanceled() const
+	{
+		return bCanceled;
+	}
+
 	FIoBuffer					Body;
 	FString						Reason;
 	HTTP::FTicketPerf::FSample	Sample;
@@ -85,6 +91,7 @@ struct FMultiEndpointHttpClientResponse
 	uint32						StatusCode = 0;
 	uint32						RetryCount = 0;
 	int32						HostIndex = INDEX_NONE;
+	bool						bCanceled = false;
 	int8						CDNCacheStatus = -1;
 };
 
@@ -103,19 +110,27 @@ class FMultiEndpointHttpClient
 public:
 	UE_NONCOPYABLE(FMultiEndpointHttpClient);
 
+	using FHttpTicketId = uint32;
 	using FOnHttpResponse = TFunction<void(FMultiEndpointHttpClientResponse&&)>;
 
+	~FMultiEndpointHttpClient();
+
 	[[nodiscard]] static TUniquePtr<FMultiEndpointHttpClient> Create(const FMultiEndpointHttpClientConfig& Config);
 
 	/** Blocking method */
 	[[nodiscard]] static TIoStatusOr<FMultiEndpointHttpClientResponse> Get(FAnsiStringView Url, const FMultiEndpointHttpClientConfig& Config);
 
-	void Get(const FOnDemandHostGroup& HostGroup, FAnsiStringView RelativeUrl, FOnHttpResponse&& OnResponse);
-	void Get(const FOnDemandHostGroup& HostGroup, FAnsiStringView RelativeUrl, const FIoOffsetAndLength& ChunkRange, FOnHttpResponse&& OnResponse);
-	void Get(const FOnDemandHostGroup& HostGroup, FAnsiStringView RelativeUrl, const FIoOffsetAndLength& ChunkRange, TArray<FAnsiString>&& Headers, FOnHttpResponse&& OnResponse);
+	FHttpTicketId Get(const FOnDemandHostGroup& HostGroup, FAnsiStringView RelativeUrl, FOnHttpResponse&& OnResponse);
+	FHttpTicketId Get(const FOnDemandHostGroup& HostGroup, FAnsiStringView RelativeUrl, const FIoOffsetAndLength& ChunkRange, FOnHttpResponse&& OnResponse);
+	FHttpTicketId Get(const FOnDemandHostGroup& HostGroup, FAnsiStringView RelativeUrl, const FIoOffsetAndLength& ChunkRange, TArray<FAnsiString>&& Headers, FOnHttpResponse&& OnResponse);
 
 	bool Tick(int32 WaitTimeMs, uint32 MaxKiBPerSecond);
-	bool Tick() { return Tick(-1, 0); }
+	bool Tick()
+	{
+		return Tick(-1, 0);
+	}
+
+	void CancelRequest(FHttpTicketId TicketId);
 
 	void UpdateConnections();
 
@@ -135,6 +150,7 @@ private:
 		TArray<FAnsiString>	Headers;
 		FAnsiString			RelativeUrl;
 		FIoOffsetAndLength	Range;
+		FHttpTicketId		TicketId;
 		FConnection&		Connection;
 		/** Stored the response body */
 		FIoBuffer			Body;
@@ -154,9 +170,12 @@ private:
 	};
 
 										FMultiEndpointHttpClient(const FMultiEndpointHttpClientConfig& Config);
-	void								IssueRequest(FRequest&& Request);
+
+	FHttpTicketId						IssueRequest(FRequest&& Request);
 	void								RetryRequest(FRequest&& Request);
 
+	void								ProcessRetryAttempts(uint32 TicketCount);
+
 	TUniquePtr<HTTP::FConnectionPool>	CreateConnection(FAnsiStringView HostUrl) const;
 	FConnection&						GetConnection(const FOnDemandHostGroup& HostGroup);
 	FConnection*						FindConnection(const FOnDemandHostGroup& HostGroup);
@@ -169,6 +188,15 @@ private:
 	TArray<TUniquePtr<FConnection>>		Connections;
 	HTTP::FEventLoop					EventLoop;
 	TArray<FRequest>					Retries;
+	struct FTicketInfo
+	{
+		HTTP::FTicket HttpTicket;
+		bool bCancelRequested = false;
+	};
+
+	UE::FRecursiveMutex					TicketLookupMutex;
+	TMap<FHttpTicketId, FTicketInfo>	TicketLookupMap;
+
 };
 
 } // namespace UE::IoStore
diff --git a/Engine/Source/Runtime/Experimental/IoStore/OnDemand/Private/OnDemandHttpThread.cpp b/Engine/Source/Runtime/Experimental/IoStore/OnDemand/Private/OnDemandHttpThread.cpp
--- a/Engine/Source/Runtime/Experimental/IoStore/OnDemand/Private/OnDemandHttpThread.cpp
+++ b/Engine/Source/Runtime/Experimental/IoStore/OnDemand/Private/OnDemandHttpThread.cpp
@@ -163,6 +163,13 @@ static FAutoConsoleVariableRef CVar_GIasHttpPollTimeoutMs(
 	TEXT("Http tick poll timeout in milliseconds")
 );
 
+bool GIaxEnableHttpCancellation = false;
+static FAutoConsoleVariableRef CVar_IaxEnableHttpCancellation(
+	TEXT("iax.EnableHttpCancellation"),
+	GIaxEnableHttpCancellation,
+	TEXT("When enabled the system will attempt to cancel inflight http requests")
+);
+
 #if UE_ALLOW_INVALID_URL_DEBUGGING
 float GIaxInvalidUrlChance = 0.0;
 static FAutoConsoleVariableRef CVar_GIaxInvalidUrlChance(
@@ -218,6 +225,7 @@ struct FHttpRequest
 	}
 
 	FHttpRequest* NextRequest = nullptr;
+	FMultiEndpointHttpClient::FHttpTicketId HttpTicketId = 0;
 	FOnDemandHttpThread::FCompletionCallback CompletionCallback;
 	FOnDemandChunkInfo ChunkInfo;
 	FIoOffsetAndLength ChunkRange;
@@ -344,11 +352,18 @@ void FOnDemandHttpThread::ReprioritizeRequest(FRequestHandle Request, int32 NewP
 	}
 }
 
-void FOnDemandHttpThread::CancelRequest(FRequestHandle Request)
+void FOnDemandHttpThread::CancelRequest(FRequestHandle RequestHandle)
 {
-	if (Request != 0)
+	if (RequestHandle != 0)
 	{
-		static_cast<FHttpRequest*>(Request)->bCancelled = true;
+		FHttpRequest* Request = static_cast<FHttpRequest*>(RequestHandle);
+		
+		Request->bCancelled = true;
+
+		if (GIaxEnableHttpCancellation)
+		{
+			HttpClient->CancelRequest(Request->HttpTicketId);
+		}
 	}
 }
 
@@ -493,7 +508,7 @@ void FOnDemandHttpThread::TickRequests()
 
 					FOnDemandIoBackendStats::Get()->OnHttpRequestStarted();
 
-					HttpClient->Get(HttpRequest->HostGroup().GetUnderlyingHostGroup(), ChunkUrl, HttpRequest->ChunkRange,
+					HttpRequest->HttpTicketId = HttpClient->Get(HttpRequest->HostGroup().GetUnderlyingHostGroup(), ChunkUrl, HttpRequest->ChunkRange,
 						[this, HttpRequest, &NumConcurrentRequests]
 						(FMultiEndpointHttpClientResponse&& HttpResponse)
 						{
@@ -515,7 +530,12 @@ void FOnDemandHttpThread::TickRequests()
 								FOnDemandIoBackendStats::Get()->OnHttpGet(HttpRequest->Type, HttpResponse.Body.DataSize(), HttpResponse.DurationMilliseconds);
 
 								HttpRequest->OnRequestCompleted(HttpResponse.StatusCode, HttpResponse.Reason, MoveTemp(HttpResponse.Body));
-								DestroyRequest(HttpRequest);
+							}
+							else if (HttpResponse.IsCanceled())
+							{
+								FOnDemandIoBackendStats::Get()->OnHttpCancel(HttpRequest->Type);
+
+								HttpRequest->OnRequestCompleted(HttpResponse.StatusCode, HttpResponse.Reason, MoveTemp(HttpResponse.Body));
 							}
 							else
 							{
@@ -528,8 +548,9 @@ void FOnDemandHttpThread::TickRequests()
 								}
 
 								HttpRequest->OnRequestCompleted(HttpResponse.StatusCode, HttpResponse.Reason, FIoBuffer());
-								DestroyRequest(HttpRequest);
 							}
+
+							DestroyRequest(HttpRequest);
 						});
 				}
 			}
diff --git a/Engine/Source/Runtime/Experimental/IoStore/OnDemand/Private/OnDemandIoDispatcherBackend.cpp b/Engine/Source/Runtime/Experimental/IoStore/OnDemand/Private/OnDemandIoDispatcherBackend.cpp
--- a/Engine/Source/Runtime/Experimental/IoStore/OnDemand/Private/OnDemandIoDispatcherBackend.cpp
+++ b/Engine/Source/Runtime/Experimental/IoStore/OnDemand/Private/OnDemandIoDispatcherBackend.cpp
@@ -794,6 +794,7 @@ void FOnDemandIoBackend::IssueHttpRequest(FChunkRequest* ChunkRequest)
 			}
 			else
 			{
+				// Failed or canceled requests
 				CompleteRequest(ChunkRequest);
 			}
 		}, EHttpRequestType::Streaming);
