Add selection cuve + curve facade + better viewport framing + two sided material

diff --git a/Engine/Plugins/Experimental/Dataflow/Source/DataflowEditor/Private/Dataflow/DataflowEditorPreviewSceneBase.cpp b/Engine/Plugins/Experimental/Dataflow/Source/DataflowEditor/Private/Dataflow/DataflowEditorPreviewSceneBase.cpp
--- a/Engine/Plugins/Experimental/Dataflow/Source/DataflowEditor/Private/Dataflow/DataflowEditorPreviewSceneBase.cpp
+++ b/Engine/Plugins/Experimental/Dataflow/Source/DataflowEditor/Private/Dataflow/DataflowEditorPreviewSceneBase.cpp
@@ -16,6 +16,7 @@
 #include "Elements/Columns/TypedElementSelectionColumns.h"
 #include "Elements/Columns/TypedElementViewportColumns.h"
 #include "Elements/Columns/TypedElementVisibilityColumns.h"
+#include "Dataflow/DataflowDebugDrawComponent.h"
 
 #include UE_INLINE_GENERATED_CPP_BY_NAME(DataflowEditorPreviewSceneBase)
 
@@ -149,7 +150,13 @@ FBox FDataflowPreviewSceneBase::GetBoundingBox() const
 			}
 			if(!SceneBounds.IsValid)
 			{
-				SceneBounds = RootSceneActor->GetComponentsBoundingBox(true);
+				RootSceneActor->ForEachComponent<UPrimitiveComponent>(false, [&](const UPrimitiveComponent* InPrimComp)
+				{
+					if (InPrimComp->IsRegistered() && InPrimComp != Cast<UPrimitiveComponent>(DebugDrawComponent.Get()))
+					{
+						SceneBounds += InPrimComp->Bounds.GetBox();
+					}
+				});
 			}
 		}
 	}
diff --git a/Engine/Plugins/Experimental/Dataflow/Source/DataflowEditor/Private/DataflowRendering/DataflowGeometryCollectionRenderableType.cpp b/Engine/Plugins/Experimental/Dataflow/Source/DataflowEditor/Private/DataflowRendering/DataflowGeometryCollectionRenderableType.cpp
--- a/Engine/Plugins/Experimental/Dataflow/Source/DataflowEditor/Private/DataflowRendering/DataflowGeometryCollectionRenderableType.cpp
+++ b/Engine/Plugins/Experimental/Dataflow/Source/DataflowEditor/Private/DataflowRendering/DataflowGeometryCollectionRenderableType.cpp
@@ -20,6 +20,7 @@
 #include "UObject/ObjectPtr.h"
 
 #include "PlanarCut.h" // for ConvertGeometryCollectionToDynamicMesh
+#include "Dataflow/DataflowEditorStyle.h"
 
 namespace UE::Dataflow::Private
 {
@@ -44,6 +45,7 @@ namespace UE::Dataflow::Private
 			// TODO(dataflow) : use a facade or a more generic way to convert a collection to a dynamic mesh
 			const TManagedArray<FTransform3f>* TransformAttribute = Collection.FindAttribute<FTransform3f>(FTransformCollection::TransformAttribute, FTransformCollection::TransformGroup);
 			const TManagedArray<int32>* TransformIndexAttribute = Collection.FindAttribute<int32>(FGeometryCollection::TransformIndexAttribute, FGeometryCollection::GeometryGroup);
+			const TManagedArray<FString>* BoneNameAttribute = Collection.FindAttribute<FString>(FTransformCollection::BoneNameAttribute, FTransformCollection::TransformGroup);
 
 			if (TransformAttribute && TransformIndexAttribute)
 			{
@@ -80,9 +82,13 @@ namespace UE::Dataflow::Private
 				{
 					if (DynamicMeshes[TransformIndex].VertexCount() > 0)
 					{
-						if (UDynamicMeshComponent* Component = OutComponents.AddNewComponent<UDynamicMeshComponent>())
+						const FName ComponentName = BoneNameAttribute ? FName((*BoneNameAttribute)[TransformIndex]) :
+								FName(TEXT("Bone_Dynamic_Mesh_") + FString::FromInt(TransformIndex));
+						if (UDynamicMeshComponent* Component = OutComponents.AddNewComponent<UDynamicMeshComponent>(ComponentName))
 						{
 							Component->SetMesh(MoveTemp(DynamicMeshes[TransformIndex]));
+							Component->SetOverrideRenderMaterial(FDataflowEditorStyle::Get().VertexMaterial);
+							Component->SetShadowsEnabled(false);
 						}
 					}
 				}
@@ -119,4 +125,4 @@ namespace UE::Dataflow::Private
 		UE_DATAFLOW_REGISTER_RENDERABLE_TYPE(FGeometryCollectionSurfaceRenderableType);
 		UE_DATAFLOW_REGISTER_RENDERABLE_TYPE(FGeometryCollectionUvsRenderableType);
 	}
-}
\ No newline at end of file
+}
diff --git a/Engine/Plugins/Experimental/Dataflow/Source/DataflowEditor/Private/DataflowRendering/DataflowRenderableTypeUtils.cpp b/Engine/Plugins/Experimental/Dataflow/Source/DataflowEditor/Private/DataflowRendering/DataflowRenderableTypeUtils.cpp
--- a/Engine/Plugins/Experimental/Dataflow/Source/DataflowEditor/Private/DataflowRendering/DataflowRenderableTypeUtils.cpp
+++ b/Engine/Plugins/Experimental/Dataflow/Source/DataflowEditor/Private/DataflowRendering/DataflowRenderableTypeUtils.cpp
@@ -1,13 +1,12 @@
 // Copyright Epic Games, Inc. All Rights Reserved.
 
 #include "DataflowRendering/DataflowRenderableTypeUtils.h"
-
-#include "Components/DynamicMeshComponent.h"
-
 #include "DataflowRendering/DataflowRenderableComponents.h"
 
+#include "Components/DynamicMeshComponent.h"
 #include "Engine/SkeletalMesh.h"
 #include "Engine/StaticMesh.h"
+#include "Materials/Material.h"
 
 #include "GeometryCollection/Facades/CollectionMeshFacade.h"
 #include "GeometryCollection/GeometryCollectionUtility.h"
@@ -18,6 +17,8 @@
 #include "MeshDescriptionUVsToDynamicMesh.h"
 #include "StaticMeshAttributes.h"
 
+#include "Dataflow/DataflowEditorStyle.h"
+
 namespace UE::Dataflow::Rendering
 {
 	namespace Private
@@ -32,6 +33,13 @@ namespace UE::Dataflow::Rendering
 
 	///////////////////////////////////////////////////////////////////////////////////////////////////////////////
 
+	UMaterial* GetVertexMaterial()
+	{
+		return FDataflowEditorStyle::Get().VertexMaterial;
+	}
+
+	///////////////////////////////////////////////////////////////////////////////////////////////////////////////
+
 	TSharedPtr<UE::Geometry::FDynamicMesh3> GenerateUvMesh(const FMeshDescription& MeshDescription, int32 UvChannel)
 	{
 		UE::Geometry::FMeshDescriptionUVsToDynamicMesh Converter;
diff --git a/Engine/Plugins/Experimental/Dataflow/Source/DataflowEditor/Public/DataflowRendering/DataflowRenderableComponents.h b/Engine/Plugins/Experimental/Dataflow/Source/DataflowEditor/Public/DataflowRendering/DataflowRenderableComponents.h
--- a/Engine/Plugins/Experimental/Dataflow/Source/DataflowEditor/Public/DataflowRendering/DataflowRenderableComponents.h
+++ b/Engine/Plugins/Experimental/Dataflow/Source/DataflowEditor/Public/DataflowRendering/DataflowRenderableComponents.h
@@ -4,6 +4,7 @@
 
 #include "Containers/Array.h"
 #include "Templates/UnrealTypeTraits.h"
+#include "UObject/UnrealNames.h"
 
 class AActor;
 class UPrimitiveComponent;
@@ -27,10 +28,13 @@ namespace UE::Dataflow
 		* returns  the newly added component so that it can be configured by the caller 
 		*/
 		template<typename T UE_REQUIRES(TIsDerivedFrom<T, UPrimitiveComponent>::Value)>
-		T* AddNewComponent()
+		T* AddNewComponent(FName ComponentName = NAME_None, UObject* ComponentParent = nullptr)
 		{
+			UObject* ObjectOwner = ComponentParent ? ComponentParent : RootActor;
+			const FName ObjectName = MakeUniqueObjectName(ObjectOwner, T::StaticClass(), FName(ComponentName));
+			
 			// for now always allocate a new one - in the future we could use a pool 
-			T* NewComponent = NewObject<T>(RootActor);
+			T* NewComponent = NewObject<T>(ObjectOwner/*, ObjectName*/);
 			if (NewComponent)
 			{
 				Components.Add(NewComponent);
diff --git a/Engine/Plugins/Experimental/Dataflow/Source/DataflowEditor/Public/DataflowRendering/DataflowRenderableType.h b/Engine/Plugins/Experimental/Dataflow/Source/DataflowEditor/Public/DataflowRendering/DataflowRenderableType.h
--- a/Engine/Plugins/Experimental/Dataflow/Source/DataflowEditor/Public/DataflowRendering/DataflowRenderableType.h
+++ b/Engine/Plugins/Experimental/Dataflow/Source/DataflowEditor/Public/DataflowRendering/DataflowRenderableType.h
@@ -40,7 +40,7 @@ namespace UE::Dataflow
 	static inline const FName OutputTypeName = TEXT(#Type); \
 	virtual FName GetOutputType() const override \
 		{ return OutputTypeName; } \
-	const Type& Get##Name(const FRenderableTypeInstance& Instance, const Type& Default = Type()) const \
+	const Type& Get##Name(const UE::Dataflow::FRenderableTypeInstance& Instance, const Type& Default = Type()) const \
 		{ return Instance.GetOutputValue<Type>(Default); }
 
 #define UE_DATAFLOW_IRENDERABLE_RENDER_GROUP(Name) \
diff --git a/Engine/Plugins/Experimental/Dataflow/Source/DataflowEditor/Public/DataflowRendering/DataflowRenderableTypeInstance.h b/Engine/Plugins/Experimental/Dataflow/Source/DataflowEditor/Public/DataflowRendering/DataflowRenderableTypeInstance.h
--- a/Engine/Plugins/Experimental/Dataflow/Source/DataflowEditor/Public/DataflowRendering/DataflowRenderableTypeInstance.h
+++ b/Engine/Plugins/Experimental/Dataflow/Source/DataflowEditor/Public/DataflowRendering/DataflowRenderableTypeInstance.h
@@ -11,6 +11,8 @@
 class UPrimitiveComponent;
 class AActor;
 
+#define UE_API DATAFLOWEDITOR_API
+
 namespace UE::Dataflow
 {
 	struct IRenderableType;
@@ -91,7 +93,7 @@ namespace UE::Dataflow
 			return Default;
 		}
 
-		FContextCacheKey ComputeCacheKey() const;
+		UE_API FContextCacheKey ComputeCacheKey() const;
 
 		FRenderableTypeInstance(FContext& InContext, const TSharedPtr<const FDataflowNode>& InNode, const FName InOutputName, const IRenderableType* InRenderableType)
 			: EvaluationContext(InContext)
@@ -109,3 +111,5 @@ namespace UE::Dataflow
 }
 
 
+#undef UE_API
+
diff --git a/Engine/Plugins/Experimental/Dataflow/Source/DataflowEditor/Public/DataflowRendering/DataflowRenderableTypeRegistry.h b/Engine/Plugins/Experimental/Dataflow/Source/DataflowEditor/Public/DataflowRendering/DataflowRenderableTypeRegistry.h
--- a/Engine/Plugins/Experimental/Dataflow/Source/DataflowEditor/Public/DataflowRendering/DataflowRenderableTypeRegistry.h
+++ b/Engine/Plugins/Experimental/Dataflow/Source/DataflowEditor/Public/DataflowRendering/DataflowRenderableTypeRegistry.h
@@ -6,6 +6,8 @@
 #include "Containers/Map.h"
 #include "UObject/NameTypes.h"
 
+#define UE_API DATAFLOWEDITOR_API
+
 namespace UE::Dataflow
 {
 	struct IRenderableType;
@@ -17,13 +19,13 @@ namespace UE::Dataflow
 		using FRenderableTypes = TArray<const IRenderableType*>;
 
 		/** register a renderable type instance - use UE_DATAFLOW_REGISTER_RENDERABLE_TYPE instead of calling this directly for ease of use */
-		void Register(const IRenderableType* RenderableType);
+		UE_API void Register(const IRenderableType* RenderableType);
 
 		/** Get all the available renderable types for a specific primary type */
-		const FRenderableTypes& GetRenderableTypes(FName PrimaryType) const;
+		UE_API const FRenderableTypes& GetRenderableTypes(FName PrimaryType) const;
 
 		/** get the one and only instance of this registry */
-		static FRenderableTypeRegistry& GetInstance();
+		UE_API static FRenderableTypeRegistry& GetInstance();
 
 	private:
 		FRenderableTypeRegistry() = default;
@@ -33,8 +35,10 @@ namespace UE::Dataflow
 	};
 }
 
+#undef UE_API
+
 // macro to use in a module / plugin to register a specific renderable type class
 
 #define UE_DATAFLOW_REGISTER_RENDERABLE_TYPE(Type) \
 	static const Type InstanceOf##Type; \
-	FRenderableTypeRegistry::GetInstance().Register(&InstanceOf##Type);
+	UE::Dataflow::FRenderableTypeRegistry::GetInstance().Register(&InstanceOf##Type);
diff --git a/Engine/Plugins/Experimental/Dataflow/Source/DataflowEditor/Public/DataflowRendering/DataflowRenderableTypeUtils.h b/Engine/Plugins/Experimental/Dataflow/Source/DataflowEditor/Public/DataflowRendering/DataflowRenderableTypeUtils.h
--- a/Engine/Plugins/Experimental/Dataflow/Source/DataflowEditor/Public/DataflowRendering/DataflowRenderableTypeUtils.h
+++ b/Engine/Plugins/Experimental/Dataflow/Source/DataflowEditor/Public/DataflowRendering/DataflowRenderableTypeUtils.h
@@ -10,6 +10,9 @@ struct FMeshDescription;
 class UDynamicMeshComponent;
 class USkeletalMesh;
 class UStaticMesh;
+class UMaterial;
+
+#define UE_API DATAFLOWEDITOR_API
 
 namespace UE::Geometry
 {
@@ -22,6 +25,9 @@ namespace UE::Dataflow
 
 	namespace Rendering
 	{
+		/** Get the default vertex material for dynamic meshes */
+		UE_API UMaterial* GetVertexMaterial();
+		
 		/** Generate a dynamic mesh from a mesh description UV data */
 		TSharedPtr<UE::Geometry::FDynamicMesh3> GenerateUvMesh(const FMeshDescription& MeshDescription, int32 UvChannel);
 
@@ -57,3 +63,5 @@ namespace UE::Dataflow
 		}
 	}
 }
+
+#undef UE_API
diff --git a/Engine/Plugins/Experimental/GeometryCollectionPlugin/Source/GeometryCollectionNodes/Private/Dataflow/GeometryCollectionSelectionNodes.cpp b/Engine/Plugins/Experimental/GeometryCollectionPlugin/Source/GeometryCollectionNodes/Private/Dataflow/GeometryCollectionSelectionNodes.cpp
--- a/Engine/Plugins/Experimental/GeometryCollectionPlugin/Source/GeometryCollectionNodes/Private/Dataflow/GeometryCollectionSelectionNodes.cpp
+++ b/Engine/Plugins/Experimental/GeometryCollectionPlugin/Source/GeometryCollectionNodes/Private/Dataflow/GeometryCollectionSelectionNodes.cpp
@@ -1612,6 +1612,13 @@ void FCollectionSelectionByAttrDataflowNode::Evaluate(UE::Dataflow::FContext& Co
 					OutMaterialSelection.Initialize(NewGenericSelection);
 				}
 				SetValue(Context, MoveTemp(OutMaterialSelection), &MaterialSelection);
+				
+				FDataflowCurveSelection OutCurveSelection;
+				if (GroupName == UE::Dataflow::Private::GetAttributeFromEnumAsName(ESelectionByAttrGroup::Curves))
+				{
+					OutCurveSelection.Initialize(NewGenericSelection);
+				}
+				SetValue(Context, MoveTemp(OutCurveSelection), &CurveSelection);
 
 				return;
 			}
@@ -1622,6 +1629,7 @@ void FCollectionSelectionByAttrDataflowNode::Evaluate(UE::Dataflow::FContext& Co
 		SetValue(Context, FDataflowTransformSelection(), &TransformSelection);
 		SetValue(Context, FDataflowGeometrySelection(), &GeometrySelection);
 		SetValue(Context, FDataflowMaterialSelection(), &MaterialSelection);
+		SetValue(Context, FDataflowCurveSelection(), &CurveSelection);
 	}
 	else if (Out->IsA<FManagedArrayCollection>(&Collection))
 	{
diff --git a/Engine/Plugins/Experimental/GeometryCollectionPlugin/Source/GeometryCollectionNodes/Public/Dataflow/GeometryCollectionSelectionNodes.h b/Engine/Plugins/Experimental/GeometryCollectionPlugin/Source/GeometryCollectionNodes/Public/Dataflow/GeometryCollectionSelectionNodes.h
--- a/Engine/Plugins/Experimental/GeometryCollectionPlugin/Source/GeometryCollectionNodes/Public/Dataflow/GeometryCollectionSelectionNodes.h
+++ b/Engine/Plugins/Experimental/GeometryCollectionPlugin/Source/GeometryCollectionNodes/Public/Dataflow/GeometryCollectionSelectionNodes.h
@@ -1441,7 +1441,8 @@ enum class ESelectionByAttrGroup : uint8
 	Faces UMETA(DisplayName = "Faces"),
 	Transform UMETA(DisplayName = "Transform"),
 	Geometry UMETA(DisplayName = "Geometry"),
-	Material UMETA(DisplayName = "Material")
+	Material UMETA(DisplayName = "Material"),
+	Curves UMETA(DisplayName = "Curves")
 };
 
 namespace UE::Dataflow::Private
@@ -1531,6 +1532,10 @@ public:
 	/** Material selection output */
 	UPROPERTY(meta = (DataflowOutput, DisplayName = "MaterialSelection"))
 	FDataflowMaterialSelection MaterialSelection;
+	
+	/** Curve selection output */
+	UPROPERTY(meta = (DataflowOutput, DisplayName = "CurveSelection"))
+	FDataflowCurveSelection CurveSelection;
 
 	FCollectionSelectionByAttrDataflowNode(const UE::Dataflow::FNodeParameters& InParam, FGuid InGuid = FGuid::NewGuid())
 		: FDataflowNode(InParam, InGuid)
@@ -1543,6 +1548,7 @@ public:
 		RegisterOutputConnection(&TransformSelection);
 		RegisterOutputConnection(&GeometrySelection);
 		RegisterOutputConnection(&MaterialSelection);
+		RegisterOutputConnection(&CurveSelection);
 	}
 
 	virtual void Evaluate(UE::Dataflow::FContext& Context, const FDataflowOutput* Out) const override;
diff --git a/Engine/Plugins/Experimental/PlanarCutPlugin/Source/PlanarCut/Private/GeometryMeshConversion.cpp b/Engine/Plugins/Experimental/PlanarCutPlugin/Source/PlanarCut/Private/GeometryMeshConversion.cpp
--- a/Engine/Plugins/Experimental/PlanarCutPlugin/Source/PlanarCut/Private/GeometryMeshConversion.cpp
+++ b/Engine/Plugins/Experimental/PlanarCutPlugin/Source/PlanarCut/Private/GeometryMeshConversion.cpp
@@ -463,6 +463,28 @@ namespace AugmentedDynamicMesh
 			}
 		}
 	}
+	
+	void InitializeOverlayToPerVertexColors(FDynamicMesh3& Mesh)
+	{
+		Mesh.Attributes()->EnablePrimaryColors();
+		Mesh.Attributes()->PrimaryColors()->CreatePerVertex(0.f);
+		Mesh.EnableVertexColors(FVector3f::Zero());
+
+		if(Mesh.Attributes())
+		{
+			FDynamicMeshColorOverlay* const ColorOverlay = Mesh.Attributes()->PrimaryColors();
+			for (const int TriangleID : Mesh.TriangleIndicesItr())
+			{
+				const FIndex3i Tri = Mesh.GetTriangle(TriangleID);
+				const FIndex3i ColorElementTri = ColorOverlay->GetTriangle(TriangleID);
+				for (int TriVertIndex = 0; TriVertIndex < 3; ++TriVertIndex)
+				{
+					ColorOverlay->SetElement(ColorElementTri[TriVertIndex],
+						AugmentedDynamicMesh::GetVertexColor(Mesh, Tri[TriVertIndex]));
+				}
+			}
+		}
+	}
 
 	void InitializeOverlayToPerVertexUVs(FDynamicMesh3& Mesh, int32 NumUVLayers, int32 FirstUVLayer)
 	{
diff --git a/Engine/Plugins/Experimental/PlanarCutPlugin/Source/PlanarCut/Private/GeometryMeshConversion.h b/Engine/Plugins/Experimental/PlanarCutPlugin/Source/PlanarCut/Private/GeometryMeshConversion.h
--- a/Engine/Plugins/Experimental/PlanarCutPlugin/Source/PlanarCut/Private/GeometryMeshConversion.h
+++ b/Engine/Plugins/Experimental/PlanarCutPlugin/Source/PlanarCut/Private/GeometryMeshConversion.h
@@ -43,6 +43,7 @@ namespace AugmentedDynamicMesh
 	/// Initialize UV overlays based on the custom AugmentedDynamicMesh per-vertex UV attributes.  Optionally use FirstUVLayer to skip layers
 	void InitializeOverlayToPerVertexUVs(UE::Geometry::FDynamicMesh3& Mesh, int32 NumUVLayers, int32 FirstUVLayer = 0);
 	void InitializeOverlayToPerVertexTangents(UE::Geometry::FDynamicMesh3& Mesh);
+	void InitializeOverlayToPerVertexColors(UE::Geometry::FDynamicMesh3& Mesh);
 	void ComputeTangents(UE::Geometry::FDynamicMesh3& Mesh, bool bOnlyInternalSurfaces,
 		bool bRecomputeNormals = true, bool bMakeSharpEdges = false, float SharpAngleDegrees = 60);
 	void AddCollisionSamplesPerComponent(UE::Geometry::FDynamicMesh3& Mesh, double Spacing);
diff --git a/Engine/Plugins/Experimental/PlanarCutPlugin/Source/PlanarCut/Private/PlanarCut.cpp b/Engine/Plugins/Experimental/PlanarCutPlugin/Source/PlanarCut/Private/PlanarCut.cpp
--- a/Engine/Plugins/Experimental/PlanarCutPlugin/Source/PlanarCut/Private/PlanarCut.cpp
+++ b/Engine/Plugins/Experimental/PlanarCutPlugin/Source/PlanarCut/Private/PlanarCut.cpp
@@ -2175,6 +2175,7 @@ void ConvertToDynamicMeshTemplate(
 		FMeshNormals::InitializeOverlayToPerVertexNormals(Mesh.Attributes()->PrimaryNormals(), true);
 		AugmentedDynamicMesh::InitializeOverlayToPerVertexUVs(Mesh, Collection.NumUVLayers());
 		AugmentedDynamicMesh::InitializeOverlayToPerVertexTangents(Mesh);
+		AugmentedDynamicMesh::InitializeOverlayToPerVertexColors(Mesh);
 
 		if (RemapMaterialIDs)
 		{
diff --git a/Engine/Source/Runtime/Experimental/Chaos/Private/GeometryCollection/Facades/CollectionCurveFacade.cpp b/Engine/Source/Runtime/Experimental/Chaos/Private/GeometryCollection/Facades/CollectionCurveFacade.cpp
new file mode 100644
--- /dev/null
+++ b/Engine/Source/Runtime/Experimental/Chaos/Private/GeometryCollection/Facades/CollectionCurveFacade.cpp
@@ -0,0 +1,323 @@
+﻿// Copyright Epic Games, Inc. All Rights Reserved.
+
+#include "GeometryCollection/Facades/CollectionCurveFacade.h"
+#include "GeometryCollection/GeometryCollection.h"
+
+namespace UE::GeometryCollection::Facades
+{
+	// Curves facade  Groups
+	const FName FCollectionCurveGeometryFacade::CurvesGroup("Curves");
+	const FName FCollectionCurveGeometryFacade::PointsGroup("Points");
+	
+	// Curves geometry facade attributes
+	const FName FCollectionCurveGeometryFacade::CurvePointOffsetsAttribute("PointOffset");
+	const FName FCollectionCurveGeometryFacade::PointCurveIndicesAttribute("CurveIndex");
+	const FName FCollectionCurveGeometryFacade::GeometryCurveOffsetsAttribute("CurveOffset");
+	const FName FCollectionCurveGeometryFacade::CurveGeometryIndicesAttribute("GeometryIndex");
+	const FName FCollectionCurveGeometryFacade::PointRestOrientationsAttribute("RestOrientation");
+	const FName FCollectionCurveGeometryFacade::PointRestPositionsAttribute("RestPosition");
+	const FName FCollectionCurveGeometryFacade::GeometryGroupNamesAttribute("GroupName");
+	const FName FCollectionCurveGeometryFacade::GeometryCurveThicknessAttribute("CurveThickness");
+	const FName FCollectionCurveGeometryFacade::CurveSourceIndicesAttribute("SourceIndices");
+	
+	// Curves hierarchy facade attributes
+	const FName FCollectionCurveHierarchyFacade::CurveParentIndicesAttribute("ParentIndices");
+	const FName FCollectionCurveHierarchyFacade::CurveLodIndicesAttribute("LodIndices");
+	
+	FCollectionCurveGeometryFacade::FCollectionCurveGeometryFacade(FManagedArrayCollection& InCollection) :
+		ConstCollection(InCollection), Collection(&InCollection),
+		PointRestOrientations(InCollection, PointRestOrientationsAttribute, PointsGroup),
+		PointRestPositions(InCollection, PointRestPositionsAttribute, PointsGroup),
+		CurvePointOffsets(InCollection, CurvePointOffsetsAttribute, CurvesGroup),
+		GeometryCurveOffsets(InCollection, GeometryCurveOffsetsAttribute, FGeometryCollection::GeometryGroup),
+		PointCurveIndices(InCollection, PointCurveIndicesAttribute, PointsGroup),
+		CurveGeometryIndices(InCollection, CurveGeometryIndicesAttribute, CurvesGroup),
+		GeometryGroupNames(InCollection, GeometryGroupNamesAttribute, FGeometryCollection::GeometryGroup),
+		GeometryCurveThickness(InCollection, GeometryCurveThicknessAttribute, FGeometryCollection::GeometryGroup),
+		CurveSourceIndices(InCollection, CurveSourceIndicesAttribute, CurvesGroup)
+	{
+		DefineSchema();
+	}
+	
+	FCollectionCurveGeometryFacade::FCollectionCurveGeometryFacade(const FManagedArrayCollection& InCollection) :
+		ConstCollection(InCollection), Collection(nullptr),
+		PointRestOrientations(InCollection, PointRestOrientationsAttribute, PointsGroup),
+		PointRestPositions(InCollection, PointRestPositionsAttribute, PointsGroup),
+		CurvePointOffsets(InCollection, CurvePointOffsetsAttribute, CurvesGroup),
+		GeometryCurveOffsets(InCollection, GeometryCurveOffsetsAttribute, FGeometryCollection::GeometryGroup),
+		PointCurveIndices(InCollection, PointCurveIndicesAttribute, PointsGroup),
+		CurveGeometryIndices(InCollection, CurveGeometryIndicesAttribute, CurvesGroup),
+		GeometryGroupNames(InCollection, GeometryGroupNamesAttribute, FGeometryCollection::GeometryGroup),
+		GeometryCurveThickness(InCollection, GeometryCurveThicknessAttribute, FGeometryCollection::GeometryGroup),
+		CurveSourceIndices(InCollection, CurveSourceIndicesAttribute, CurvesGroup)
+	{
+	}
+	
+	bool FCollectionCurveGeometryFacade::IsValid() const
+	{
+		return PointRestOrientations.IsValid() && PointRestPositions.IsValid() && 
+			CurvePointOffsets.IsValid() && GeometryCurveOffsets.IsValid() && PointCurveIndices.IsValid() &&
+					CurveGeometryIndices.IsValid() && PointRestOrientations.IsValid() && GeometryGroupNames.IsValid() &&
+						GeometryCurveThickness.IsValid() && CurveSourceIndices.IsValid();
+	}
+	
+	void FCollectionCurveGeometryFacade::DefineSchema()
+	{
+		check(!IsConst());
+		PointRestOrientations.Add();
+		PointRestPositions.Add();
+		CurvePointOffsets.Add();
+		GeometryCurveOffsets.Add();
+		PointCurveIndices.Add();
+		CurveGeometryIndices.Add();
+		GeometryGroupNames.Add();
+		GeometryCurveThickness.Add();
+		CurveSourceIndices.Add();
+	}
+	
+	void FCollectionCurveGeometryFacade::InitCurvesCollection(const TArray<FVector3f>& InPointRestPositions,
+		const TArray<int32>& InCurvePointOffsets, const TArray<int32>& InGeometryCurveOffsets, const TArray<FString>& InGeometryGroupNames,
+		const TArray<float>& InGeometryCurveThickness, const TArray<int32>& InCurveSourceIndices)
+	{
+		if(Collection)
+		{
+			// Curves group
+			Collection->EmptyGroup(CurvesGroup);
+			Collection->AddElements(InCurvePointOffsets.Num(), CurvesGroup);
+
+			// Geometry group
+			Collection->EmptyGroup(FGeometryCollection::GeometryGroup);
+			Collection->AddElements(InGeometryCurveOffsets.Num(), FGeometryCollection::GeometryGroup);
+
+			// Points group
+			Collection->EmptyGroup(PointsGroup);
+			Collection->AddElements(InPointRestPositions.Num(), PointsGroup);
+
+			// Fill attributes
+			SetGeometryCurveOffsets(InGeometryCurveOffsets);
+			SetCurvePointOffsets(InCurvePointOffsets);
+			SetPointRestPositions(InPointRestPositions);
+			SetGeometryGroupNames(InGeometryGroupNames);
+			SetGeometryCurveThickness(InGeometryCurveThickness);
+			SetCurveSourceIndices(InCurveSourceIndices);
+
+			// Build the geometry collection 
+			BuildGeometryCollection();
+		}
+	}
+
+	void FCollectionCurveGeometryFacade::BuildGeometryCollection()
+	{
+		if(Collection)
+		{
+			FGeometryCollection::DefineGeometrySchema(*Collection);
+
+			// Transform group
+			Collection->EmptyGroup(FGeometryCollection::TransformGroup);
+			Collection->AddElements(GetNumGeometry(), FGeometryCollection::TransformGroup);
+
+			// Vertices group
+			Collection->EmptyGroup(FGeometryCollection::VerticesGroup);
+			Collection->AddElements(GetNumPoints() * 2, FGeometryCollection::VerticesGroup);
+
+			// Faces group
+			Collection->EmptyGroup(FGeometryCollection::FacesGroup);
+			Collection->AddElements((GetNumPoints() - GetNumCurves()) * 2, FGeometryCollection::FacesGroup);
+
+			// Material group
+			Collection->EmptyGroup(FGeometryCollection::MaterialGroup);
+
+			// Vertices Attributes
+			TManagedArray<FVector3f>& VertexPositions = Collection->ModifyAttribute<FVector3f>(FGeometryCollection::VertexPositionAttribute, FGeometryCollection::VerticesGroup);
+			TManagedArray<FVector3f>& VertexNormals = Collection->ModifyAttribute<FVector3f>(FGeometryCollection::VertexNormalAttribute, FGeometryCollection::VerticesGroup);
+			TManagedArray<FVector3f>& VertexTangentU = Collection->ModifyAttribute<FVector3f>(FGeometryCollection::VertexTangentUAttribute, FGeometryCollection::VerticesGroup);
+			TManagedArray<FVector3f>& VertexTangentV = Collection->ModifyAttribute<FVector3f>(FGeometryCollection::VertexTangentVAttribute, FGeometryCollection::VerticesGroup);
+			TManagedArray<FLinearColor>& VertexColors = Collection->ModifyAttribute<FLinearColor>(FGeometryCollection::ColorAttribute, FGeometryCollection::VerticesGroup);
+			TManagedArray<int32>& VertexTransforms = Collection->ModifyAttribute<int32>(FGeometryCollection::VertexBoneMapAttribute, FGeometryCollection::VerticesGroup);
+			
+			// Faces Attributes
+			TManagedArray<FIntVector>& FaceIndices = Collection->ModifyAttribute<FIntVector>(FGeometryCollection::FaceIndicesAttribute, FGeometryCollection::FacesGroup);
+			TManagedArray<bool>& FaceVisibles = Collection->ModifyAttribute<bool>(FGeometryCollection::FaceVisibleAttribute, FGeometryCollection::FacesGroup);
+
+			// Geometry Group
+			TManagedArray<int32>& TransformIndex = Collection->ModifyAttribute<int32>(FGeometryCollection::TransformIndexAttribute, FGeometryCollection::GeometryGroup);
+			TManagedArray<FBox>& BoundingBox =  Collection->ModifyAttribute<FBox>(FGeometryCollection::BoundingBoxAttribute, FGeometryCollection::GeometryGroup);
+			TManagedArray<int32>& VertexStart = Collection->ModifyAttribute<int32>(FGeometryCollection::VertexStartAttribute, FGeometryCollection::GeometryGroup);
+			TManagedArray<int32>& VertexCount = Collection->ModifyAttribute<int32>(FGeometryCollection::VertexCountAttribute, FGeometryCollection::GeometryGroup);
+			TManagedArray<int32>& FaceStart = Collection->ModifyAttribute<int32>(FGeometryCollection::FaceStartAttribute, FGeometryCollection::GeometryGroup);
+			TManagedArray<int32>& FaceCount = Collection->ModifyAttribute<int32>(FGeometryCollection::FaceCountAttribute, FGeometryCollection::GeometryGroup);
+			TManagedArray<float>& InnerRadius =Collection->ModifyAttribute<float>(FGeometryCollection::InnerRadiusAttribute, FGeometryCollection::GeometryGroup);
+			TManagedArray<float>& OuterRadius = Collection->ModifyAttribute<float>(FGeometryCollection::OuterRadiusAttribute, FGeometryCollection::GeometryGroup);
+
+			// Transform group
+			TManagedArray<int32>& BoneGeometry = Collection->ModifyAttribute<int32>(FTransformCollection::GeometryIndexAttribute, FTransformCollection::TransformGroup);
+			TManagedArray<FTransform3f>& BoneTransforms = Collection->ModifyAttribute<FTransform3f>(FTransformCollection::TransformAttribute, FTransformCollection::TransformGroup);
+			TManagedArray<FString>& BoneNames = Collection->ModifyAttribute<FString>(FTransformCollection::BoneNameAttribute, FTransformCollection::TransformGroup);
+			TManagedArray<FLinearColor>& BoneColors = Collection->ModifyAttribute<FLinearColor>(FTransformCollection::BoneColorAttribute , FTransformCollection::TransformGroup);
+			TManagedArray<int32>& BoneParent = Collection->ModifyAttribute<int32>(FTransformCollection::ParentAttribute, FTransformCollection::TransformGroup);
+			
+			const int32 NumGeometry = GetNumGeometry();
+
+			VertexTangentU.Fill(FVector3f::Zero());
+			VertexTangentV.Fill(FVector3f::Zero());
+
+			int32 PointOffset = 0;
+			int32 CurveOffset = 0;
+			for(int32 GeometryIndex = 0; GeometryIndex < NumGeometry; ++GeometryIndex)
+			{
+				const FLinearColor GeometryColor = FLinearColor::IntToDistinctColor(GeometryIndex, 0.75f, 1.0f, 90.0f);
+
+				// Init geometry/transform attributes
+				BoundingBox[GeometryIndex] = FBox(ForceInitToZero);
+				TransformIndex[GeometryIndex] = GeometryIndex;
+				BoneGeometry[GeometryIndex] = GeometryIndex;
+				BoneTransforms[GeometryIndex] = FTransform3f::Identity;
+				BoneParent[GeometryIndex] = INDEX_NONE;
+				BoneColors[GeometryIndex] = GeometryColor;
+				BoneNames[GeometryIndex] = GeometryGroupNames[GeometryIndex];
+				VertexStart[GeometryIndex] = PointOffset * 2;
+				FaceStart[GeometryIndex] = (PointOffset-CurveOffset) * 2;
+				VertexCount[GeometryIndex] = 0;
+				FaceCount[GeometryIndex] = 0;
+				InnerRadius[GeometryIndex] = GeometryCurveThickness[GeometryIndex];
+				OuterRadius[GeometryIndex] = GeometryCurveThickness[GeometryIndex];
+				
+				for(int32 CurveIndex = CurveOffset, CurveEnd = GeometryCurveOffsets[GeometryIndex];
+						  CurveIndex < CurveEnd; ++CurveIndex)
+				{
+					for (int32 PointIndex = PointOffset, PointEnd = CurvePointOffsets[CurveIndex];
+							   PointIndex < PointEnd; ++PointIndex)
+					{
+						// Build 2 vertices for each points
+						const int32 VertexIndex = 2 * PointIndex;
+
+						// Build positions for each curve points
+						const FVector3f SideVector = PointRestOrientations[PointIndex].GetAxisX();
+						VertexPositions[VertexIndex]   = PointRestPositions[PointIndex] + SideVector * GeometryCurveThickness[GeometryIndex];
+						VertexPositions[VertexIndex+1] = PointRestPositions[PointIndex] - SideVector * GeometryCurveThickness[GeometryIndex];
+
+						// Build normals along the up vector
+						const FVector3f NormalVector = PointRestOrientations[PointIndex].GetAxisZ();
+						VertexNormals[VertexIndex]   = NormalVector;
+						VertexNormals[VertexIndex+1] = NormalVector;
+
+						// Build the vertex colors
+						VertexColors[VertexIndex] = GeometryColor;
+						VertexColors[VertexIndex+1] = GeometryColor;
+						
+						// Build the vertex colors
+						VertexTransforms[VertexIndex] = GeometryIndex;
+						VertexTransforms[VertexIndex+1] = GeometryIndex;
+
+						// Expand the bounding box
+						BoundingBox[GeometryIndex] += FVector(VertexPositions[VertexIndex]);
+						BoundingBox[GeometryIndex] += FVector(VertexPositions[VertexIndex+1]);
+
+						VertexCount[GeometryIndex] += 2;
+
+						// Build 2 faces for each edges perpendicular to the up-vector curve
+						if(PointIndex < (PointEnd-1))
+						{
+							const int32 FaceIndex = 2 * (PointIndex - CurveIndex);
+							FaceIndices[FaceIndex] = FIntVector(VertexIndex, VertexIndex + 1, VertexIndex + 3);
+							FaceIndices[FaceIndex+1] = FIntVector(VertexIndex, VertexIndex + 3, VertexIndex + 2);
+
+							FaceVisibles[FaceIndex] = true;
+							FaceVisibles[FaceIndex+1] = true;
+							
+							FaceCount[GeometryIndex] += 2;
+						}
+					}
+					PointOffset = CurvePointOffsets[CurveIndex];
+				}
+				CurveOffset = GeometryCurveOffsets[GeometryIndex];
+			}
+		}
+	}
+	
+	void FCollectionCurveGeometryFacade::UpdateCurveGeometryIndices()
+	{
+		const int32 NumGeometry = GetNumGeometry();
+
+		int32 CurveOffset = 0;
+		for(int32 GeometryIndex = 0; GeometryIndex < NumGeometry; ++GeometryIndex)
+		{
+			for(int32 CurveIndex = CurveOffset, CurveEnd = GeometryCurveOffsets[GeometryIndex];
+					  CurveIndex < CurveEnd; ++CurveIndex)
+			{
+				CurveGeometryIndices.ModifyAt(CurveIndex,GeometryIndex);
+			}
+			CurveOffset = GeometryCurveOffsets[GeometryIndex];
+		}
+	}
+	
+	void FCollectionCurveGeometryFacade::UpdatePointCurveIndices()
+	{
+		const int32 NumCurves = GetNumCurves();
+
+		int32 PointOffset = 0;
+		for(int32 CurveIndex = 0; CurveIndex < NumCurves; ++CurveIndex)
+		{
+			for (int32 PointIndex = PointOffset, PointEnd = CurvePointOffsets[CurveIndex];
+					   PointIndex < PointEnd; ++PointIndex)
+			{
+				PointCurveIndices.ModifyAt(PointIndex, CurveIndex);
+			}
+			PointOffset = CurvePointOffsets[CurveIndex];
+		}
+	}
+	
+	void FCollectionCurveGeometryFacade::UpdatePointRestOrientations()
+	{
+		const int32 NumCurves = GetNumCurves();
+
+		int32 EdgeOffset = 0;
+		int32 PointOffset = 0;
+		for(int32 CurveIndex = 0; CurveIndex < NumCurves; ++CurveIndex)
+		{
+			FVector3f TangentPrev = FVector3f(0.,0.,1.), TangentNext = FVector3f::Zero();
+			FQuat4f EdgeOrientation = FQuat4f::Identity;
+
+			const int32 NumEdges = (CurvePointOffsets[CurveIndex] - PointOffset)-1;
+			for (int32 EdgeIndex = 0; EdgeIndex < NumEdges; ++EdgeIndex)
+			{
+				TangentPrev = TangentNext;
+				TangentNext = (PointRestPositions[PointOffset+EdgeIndex+1]-PointRestPositions[PointOffset+EdgeIndex]).GetSafeNormal();
+
+				EdgeOrientation = (FQuat4f::FindBetweenNormals(TangentPrev,TangentNext) * EdgeOrientation).GetNormalized();
+				PointRestOrientations.ModifyAt(EdgeOffset+EdgeIndex, EdgeOrientation);
+			}
+			PointRestOrientations.ModifyAt(EdgeOffset+NumEdges, EdgeOrientation);
+			EdgeOffset += NumEdges+1;
+			PointOffset = CurvePointOffsets[CurveIndex];
+		}
+	}
+	
+	FCollectionCurveHierarchyFacade::FCollectionCurveHierarchyFacade(FManagedArrayCollection& InCollection) :
+		ConstCollection(InCollection), Collection(&InCollection),
+		CurveParentIndices(InCollection, CurveParentIndicesAttribute, FCollectionCurveGeometryFacade::CurvesGroup),
+		CurveLodIndices(InCollection, CurveLodIndicesAttribute, FCollectionCurveGeometryFacade::CurvesGroup)
+	{
+		DefineSchema();
+	}
+	
+	FCollectionCurveHierarchyFacade::FCollectionCurveHierarchyFacade(const FManagedArrayCollection& InCollection) :
+		ConstCollection(InCollection), Collection(nullptr),
+		CurveParentIndices(InCollection, CurveParentIndicesAttribute, FCollectionCurveGeometryFacade::CurvesGroup),
+		CurveLodIndices(InCollection, CurveLodIndicesAttribute, FCollectionCurveGeometryFacade::CurvesGroup)
+	{
+	}
+	
+	bool FCollectionCurveHierarchyFacade::IsValid() const
+	{
+		return CurveParentIndices.IsValid() && CurveLodIndices.IsValid();
+	}
+	
+	void FCollectionCurveHierarchyFacade::DefineSchema()
+	{
+		check(!IsConst());
+		CurveParentIndices.Add();
+		CurveLodIndices.Add();
+	}
+}
+
diff --git a/Engine/Source/Runtime/Experimental/Chaos/Private/GeometryCollection/Facades/CollectionMeshFacade.cpp b/Engine/Source/Runtime/Experimental/Chaos/Private/GeometryCollection/Facades/CollectionMeshFacade.cpp
--- a/Engine/Source/Runtime/Experimental/Chaos/Private/GeometryCollection/Facades/CollectionMeshFacade.cpp
+++ b/Engine/Source/Runtime/Experimental/Chaos/Private/GeometryCollection/Facades/CollectionMeshFacade.cpp
@@ -123,7 +123,7 @@ namespace GeometryCollection::Facades
 		const TManagedArray<int32>& VertexCounts = VertexCountAttribute.Get();
 		const TManagedArray<FVector3f>& Vertices = VertexAttribute.Get();
 
-		return TArrayView<const FVector3f>(Vertices.GetData() + VertexStarts[TransformToGeometryIndicies[BoneIdx]], VertexCounts[TransformToGeometryIndicies[BoneIdx]] + 1);
+		return TArrayView<const FVector3f>(Vertices.GetData() + VertexStarts[TransformToGeometryIndicies[BoneIdx]], VertexCounts[TransformToGeometryIndicies[BoneIdx]]);
 	}
 
 	const TArray<int32> FCollectionMeshFacade::GetFaceIndices(int32 BoneIdx) const
diff --git a/Engine/Source/Runtime/Experimental/Chaos/Private/GeometryCollection/Facades/CollectionTransformSelectionFacade.cpp b/Engine/Source/Runtime/Experimental/Chaos/Private/GeometryCollection/Facades/CollectionTransformSelectionFacade.cpp
--- a/Engine/Source/Runtime/Experimental/Chaos/Private/GeometryCollection/Facades/CollectionTransformSelectionFacade.cpp
+++ b/Engine/Source/Runtime/Experimental/Chaos/Private/GeometryCollection/Facades/CollectionTransformSelectionFacade.cpp
@@ -7,6 +7,7 @@
 #include "GeometryCollection/Facades/CollectionMeshFacade.h"
 #include "GeometryCollection/Facades/CollectionTransformFacade.h"
 #include "GeometryCollection/Facades/CollectionBoundsFacade.h"
+#include "GeometryCollection/Facades/CollectionCurveFacade.h"
 
 namespace GeometryCollection::Facades
 {
@@ -1604,4 +1605,298 @@ namespace GeometryCollection::Facades
 		return OutSelection;
 	}
 
+	TArray<int32> FCollectionTransformSelectionFacade::ConvertTransformSelectionToCurveSelection(const TArray<int32>& InTransformSelection) const
+	{
+		TArray<int32> OutSelection;
+
+		UE::GeometryCollection::Facades::FCollectionCurveGeometryFacade CurveFacade(ConstCollection);
+		if (TransformToGeometryIndexAttribute.IsValid() && CurveFacade.IsValid())
+		{
+			const TManagedArray<int32>& TransformToGeometryIndices = TransformToGeometryIndexAttribute.Get();
+			const TArray<int32>& CurveOffsets = CurveFacade.GetGeometryCurveOffsets();
+
+			for (int32 TransformIdx : InTransformSelection)
+			{
+				if (TransformToGeometryIndices.IsValidIndex(TransformIdx))
+				{
+					const int32 GeometryIdx = TransformToGeometryIndices[TransformIdx];
+					if (CurveOffsets.IsValidIndex(GeometryIdx))
+					{
+						const int32 CurveIndexStart = GeometryIdx == 0 ? 0 : CurveOffsets[GeometryIdx - 1];
+						const int32 CurveIndexEnd = CurveOffsets[GeometryIdx];
+
+						for (int32 CurveIndex = CurveIndexStart; CurveIndex < CurveIndexEnd; ++CurveIndex)
+						{
+							OutSelection.Add(CurveIndex);
+						}
+					}
+				}
+			}
+		}
+
+		return OutSelection;
+	}
+	
+	TArray<int32> FCollectionTransformSelectionFacade::ConvertGeometrySelectionToCurveSelection(const TArray<int32>& InGeometrySelection) const
+	{
+		TArray<int32> OutSelection;
+
+		UE::GeometryCollection::Facades::FCollectionCurveGeometryFacade CurveFacade(ConstCollection);
+		if (CurveFacade.IsValid())
+		{
+			const TArray<int32>& CurveOffsets = CurveFacade.GetGeometryCurveOffsets();
+
+			for (int32 GeometryIdx : InGeometrySelection)
+			{
+				if (CurveOffsets.IsValidIndex(GeometryIdx))
+				{
+					const int32 CurveIndexStart = GeometryIdx == 0 ? 0 : CurveOffsets[GeometryIdx - 1];
+					const int32 CurveIndexEnd = CurveOffsets[GeometryIdx];
+
+					for (int32 CurveIndex = CurveIndexStart; CurveIndex < CurveIndexEnd; ++CurveIndex)
+					{
+						OutSelection.Add(CurveIndex);
+					}
+				}
+			}
+		}
+		return OutSelection;
+	}
+	
+	TArray<int32> FCollectionTransformSelectionFacade::ConvertFaceSelectionToCurveSelection(const TArray<int32>& InFaceSelection, bool bAllElementsMustBeSelected) const
+	{
+		TArray<int32> OutSelection;
+
+		UE::GeometryCollection::Facades::FCollectionCurveGeometryFacade CurveFacade(ConstCollection);
+		if (CurveFacade.IsValid())
+		{
+			const int32 NumCurves = CurveFacade.GetNumCurves();
+			for (int32 CurveIdx = 0; CurveIdx < NumCurves; ++CurveIdx)
+			{
+				const int32 PointIndexStart = CurveIdx == 0 ? 0 :CurveFacade.GetCurvePointOffsets()[CurveIdx-1];
+				const int32 PointIndexEnd = CurveFacade.GetCurvePointOffsets()[CurveIdx];
+
+				const int32 FaceIndexStart = (PointIndexStart-CurveIdx) * 2;
+				const int32 FaceIndexEnd = (PointIndexEnd-1-CurveIdx) * 2;
+
+				bool bAllIncluded = true;
+				for (int32 FaceIdx = FaceIndexStart; FaceIdx < FaceIndexEnd; ++FaceIdx)
+				{
+					if(InFaceSelection.Contains(FaceIdx))
+					{
+						if(!bAllElementsMustBeSelected)
+						{
+							OutSelection.Add(CurveIdx);
+							break;
+						}
+					}
+					else
+					{
+						bAllIncluded = false;
+						if(bAllElementsMustBeSelected)
+						{
+							break;
+						}
+					}
+				}
+				if(bAllElementsMustBeSelected && bAllIncluded)
+				{
+					OutSelection.Add(CurveIdx);
+				}
+			}
+		}
+
+		return OutSelection;
+	}
+
+	TArray<int32> FCollectionTransformSelectionFacade::ConvertVertexSelectionToCurveSelection(const TArray<int32>& InVertexSelection, bool bAllElementsMustBeSelected) const
+	{
+		TArray<int32> OutSelection;
+
+		UE::GeometryCollection::Facades::FCollectionCurveGeometryFacade CurveFacade(ConstCollection);
+		if (CurveFacade.IsValid())
+		{
+			const int32 NumCurves = CurveFacade.GetNumCurves();
+			for (int32 CurveIdx = 0; CurveIdx < NumCurves; ++CurveIdx)
+			{
+				const int32 PointIndexStart = CurveIdx == 0 ? 0 :CurveFacade.GetCurvePointOffsets()[CurveIdx-1];
+				const int32 PointIndexEnd = CurveFacade.GetCurvePointOffsets()[CurveIdx];
+
+				const int32 VertexIndexStart = PointIndexStart * 2;
+				const int32 VertexIndexEnd = PointIndexEnd * 2;
+
+				bool bAllIncluded = true;
+				for (int32 VertexIdx = VertexIndexStart; VertexIdx < VertexIndexEnd; ++VertexIdx)
+				{
+					if(InVertexSelection.Contains(VertexIdx))
+					{
+						if(!bAllElementsMustBeSelected)
+						{
+							OutSelection.Add(CurveIdx);
+							break;
+						}
+					}
+					else
+					{
+						bAllIncluded = false;
+						if(bAllElementsMustBeSelected)
+						{
+							break;
+						}
+					}
+				}
+				if(bAllElementsMustBeSelected && bAllIncluded)
+				{
+					OutSelection.Add(CurveIdx);
+				}
+			}
+		}
+
+		return OutSelection;
+	}
+
+	TArray<int32> FCollectionTransformSelectionFacade::ConvertCurveSelectionToVertexSelection(const TArray<int32>& InCurveSelection) const
+	{
+		TArray<int32> OutSelection;
+
+		UE::GeometryCollection::Facades::FCollectionCurveGeometryFacade CurveFacade(ConstCollection);
+		if (CurveFacade.IsValid())
+		{
+			const int32 NumCurves = CurveFacade.GetNumCurves();
+			for (const int32& CurveIdx : InCurveSelection)
+			{
+				const int32 PointIndexStart = CurveIdx == 0 ? 0 :CurveFacade.GetCurvePointOffsets()[CurveIdx-1];
+				const int32 PointIndexEnd = CurveFacade.GetCurvePointOffsets()[CurveIdx];
+
+				const int32 VertexIndexStart = PointIndexStart * 2;
+				const int32 VertexIndexEnd = PointIndexEnd * 2;
+
+				for (int32 VertexIdx = VertexIndexStart; VertexIdx < VertexIndexEnd; ++VertexIdx)
+				{
+					OutSelection.Add(VertexIdx);
+				}
+			}
+		}
+
+		return OutSelection;
+	}
+
+	TArray<int32> FCollectionTransformSelectionFacade::ConvertCurveSelectionToFaceSelection(const TArray<int32>& InCurveSelection) const
+	{
+		TArray<int32> OutSelection;
+
+		UE::GeometryCollection::Facades::FCollectionCurveGeometryFacade CurveFacade(ConstCollection);
+		if (CurveFacade.IsValid())
+		{
+			const int32 NumCurves = CurveFacade.GetNumCurves();
+			for (const int32& CurveIdx : InCurveSelection)
+			{
+				const int32 PointIndexStart = CurveIdx == 0 ? 0 :CurveFacade.GetCurvePointOffsets()[CurveIdx-1];
+				const int32 PointIndexEnd = CurveFacade.GetCurvePointOffsets()[CurveIdx];
+
+				const int32 FaceIndexStart = (PointIndexStart-CurveIdx) * 2;
+				const int32 FaceIndexEnd = (PointIndexEnd-1-CurveIdx) * 2;
+				
+				for (int32 FaceIdx = FaceIndexStart; FaceIdx < FaceIndexEnd; ++FaceIdx)
+				{
+					OutSelection.Add(FaceIdx);
+				}
+			}
+		}
+
+		return OutSelection;
+	}
+
+	TArray<int32> FCollectionTransformSelectionFacade::ConvertCurveSelectionToGeometrySelection(const TArray<int32>& InCurveSelection, bool bAllElementsMustBeSelected) const
+	{
+		TArray<int32> OutSelection;
+
+		UE::GeometryCollection::Facades::FCollectionCurveGeometryFacade CurveFacade(ConstCollection);
+		if (CurveFacade.IsValid())
+		{
+			const TArray<int32>& CurveOffsets = CurveFacade.GetGeometryCurveOffsets();
+			
+			for (int32 GeometryIdx = 0; GeometryIdx < CurveFacade.GetNumGeometry(); ++GeometryIdx)
+			{
+				if (CurveOffsets.IsValidIndex(GeometryIdx))
+				{
+					const int32 CurveIndexStart = GeometryIdx == 0 ? 0 : CurveOffsets[GeometryIdx - 1];
+					const int32 CurveIndexEnd = CurveOffsets[GeometryIdx];
+
+					bool bAllIncluded = true;
+					for (int32 CurveIndex = CurveIndexStart; CurveIndex < CurveIndexEnd; ++CurveIndex)
+					{
+						if(InCurveSelection.Contains(CurveIndex))
+						{
+							if(!bAllElementsMustBeSelected)
+							{
+								OutSelection.Add(GeometryIdx);
+								break;
+							}
+						}
+						else
+						{
+							bAllIncluded = false;
+							if(bAllElementsMustBeSelected)
+							{
+								break;
+							}
+						}
+					}
+					if(bAllElementsMustBeSelected && bAllIncluded)
+					{
+						OutSelection.Add(GeometryIdx);
+					}
+				}
+			}
+		}
+		return OutSelection;
+	}
+
+	TArray<int32> FCollectionTransformSelectionFacade::ConvertCurveSelectionToTransformSelection(const TArray<int32>& InCurveSelection, bool bAllElementsMustBeSelected) const
+	{
+		TArray<int32> OutSelection;
+
+		UE::GeometryCollection::Facades::FCollectionCurveGeometryFacade CurveFacade(ConstCollection);
+		if (CurveFacade.IsValid() && TransformIndexAttribute.IsValid())
+		{
+			const TManagedArray<int32>& TransformIndices = TransformIndexAttribute.Get();
+			const TArray<int32>& CurveOffsets = CurveFacade.GetGeometryCurveOffsets();
+			
+			for (int32 GeometryIdx = 0; GeometryIdx < CurveFacade.GetNumGeometry(); ++GeometryIdx)
+			{
+				if (CurveOffsets.IsValidIndex(GeometryIdx))
+				{
+					const int32 CurveIndexStart = GeometryIdx == 0 ? 0 : CurveOffsets[GeometryIdx - 1];
+					const int32 CurveIndexEnd = CurveOffsets[GeometryIdx];
+
+					bool bAllIncluded = true;
+					for (int32 CurveIndex = CurveIndexStart; CurveIndex < CurveIndexEnd; ++CurveIndex)
+					{
+						if(InCurveSelection.Contains(CurveIndex))
+						{
+							if(!bAllElementsMustBeSelected && TransformIndices.IsValidIndex(GeometryIdx))
+							{
+								OutSelection.Add(TransformIndices[GeometryIdx]);
+								break;
+							}
+						}
+						else
+						{
+							bAllIncluded = false;
+							if(bAllElementsMustBeSelected)
+							{
+								break;
+							}
+						}
+					}
+					if(bAllElementsMustBeSelected && bAllIncluded && TransformIndices.IsValidIndex(GeometryIdx))
+					{
+						OutSelection.Add(TransformIndices[GeometryIdx]);
+					}
+				}
+			}
+		}
+		return OutSelection;
+	}
 }; // GeometryCollection::Facades
diff --git a/Engine/Source/Runtime/Experimental/Chaos/Private/GeometryCollection/Facades/CollectionVertexBoneWeightsFacade.cpp b/Engine/Source/Runtime/Experimental/Chaos/Private/GeometryCollection/Facades/CollectionVertexBoneWeightsFacade.cpp
--- a/Engine/Source/Runtime/Experimental/Chaos/Private/GeometryCollection/Facades/CollectionVertexBoneWeightsFacade.cpp
+++ b/Engine/Source/Runtime/Experimental/Chaos/Private/GeometryCollection/Facades/CollectionVertexBoneWeightsFacade.cpp
@@ -7,6 +7,7 @@
 #include "GeometryCollection/Facades/CollectionVertexBoneWeightsFacade.h"
 
 #include "CoreMinimal.h"
+#include "MeshAttributes.h"
 #include "GeometryCollection/Facades/CollectionKinematicBindingFacade.h"
 #include "GeometryCollection/GeometryCollection.h"
 
@@ -16,34 +17,42 @@ namespace GeometryCollection::Facades
 {
 
 	// Attributes
-	const FName FVertexBoneWeightsFacade::BoneWeightAttributeName = "BoneWeights";
-	const FName FVertexBoneWeightsFacade::BoneIndexAttributeName = "BoneWeightsIndex";
-	const FName FVertexBoneWeightsFacade::KinematicAttributeName = "Kinematic";
+	const FName FVertexBoneWeightsFacade::BoneWeightsAttributeName = "BoneWeights";
+	const FName FVertexBoneWeightsFacade::BoneIndicesAttributeName = "BoneIndices";
+	const FName FVertexBoneWeightsFacade::KinematicWeightAttributeName = "KinematicWeight";
+	const FName FVertexBoneWeightsFacade::SkeletalMeshAttributeName = "SkeletalMesh";
+	const FName FVertexBoneWeightsFacade::GeometryLODAttributeName = "GeometryLOD";
+
+	// Deprecated
+	const FName FVertexBoneWeightsFacade::DeprecatedBoneWeightsAttributeName = "BoneWeight";
+	const FName FVertexBoneWeightsFacade::DeprecatedBoneIndicesAttributeName = "BoneWeightIndex";
+	const FName FVertexBoneWeightsFacade::DeprecatedKinematicFlagAttributeName = "Kinematic";
 	
-	FVertexBoneWeightsFacade::FVertexBoneWeightsFacade(FManagedArrayCollection& InCollection)
+	FVertexBoneWeightsFacade::FVertexBoneWeightsFacade(FManagedArrayCollection& InCollection, const bool bInInternalWeights)
 		: ConstCollection(InCollection)
 		, Collection(&InCollection)
-		, BoneIndexAttribute(InCollection, BoneIndexAttributeName, FGeometryCollection::VerticesGroup, FTransformCollection::TransformGroup)
-		, BoneWeightAttribute(InCollection, BoneWeightAttributeName, FGeometryCollection::VerticesGroup, FTransformCollection::TransformGroup)
-		, KinematicAttribute(InCollection, KinematicAttributeName, FGeometryCollection::VerticesGroup)
-		, ParentAttribute(InCollection, FTransformCollection::ParentAttribute, FTransformCollection::TransformGroup)
-		, VerticesAttribute(InCollection, "Vertex", FGeometryCollection::VerticesGroup)
+		, BoneIndicesAttribute(InCollection, BoneIndicesAttributeName, FGeometryCollection::VerticesGroup, bInInternalWeights ? FTransformCollection::TransformGroup : NAME_None)
+		, BoneWeightsAttribute(InCollection, BoneWeightsAttributeName, FGeometryCollection::VerticesGroup, bInInternalWeights ? FTransformCollection::TransformGroup : NAME_None)
+		, KinematicWeightAttribute(InCollection, KinematicWeightAttributeName, FGeometryCollection::VerticesGroup)
+		, GeometryLODAttribute(InCollection, GeometryLODAttributeName, FGeometryCollection::GeometryGroup)
+		, SkeletalMeshAttribute(InCollection, SkeletalMeshAttributeName, FGeometryCollection::GeometryGroup)
+		, bInternalWeights(bInInternalWeights)
 	{
 		DefineSchema();
 	}
 
-	FVertexBoneWeightsFacade::FVertexBoneWeightsFacade(const FManagedArrayCollection& InCollection)
+	FVertexBoneWeightsFacade::FVertexBoneWeightsFacade(const FManagedArrayCollection& InCollection, const bool bInInternalWeights)
 		: ConstCollection(InCollection)
 		, Collection(nullptr)
-		, BoneIndexAttribute(InCollection, BoneIndexAttributeName, FGeometryCollection::VerticesGroup, FTransformCollection::TransformGroup)
-		, BoneWeightAttribute(InCollection, BoneWeightAttributeName, FGeometryCollection::VerticesGroup, FTransformCollection::TransformGroup)
-		, KinematicAttribute(InCollection, KinematicAttributeName, FGeometryCollection::VerticesGroup)
-		, ParentAttribute(InCollection, FTransformCollection::ParentAttribute, FTransformCollection::TransformGroup)
-		, VerticesAttribute(InCollection, "Vertex", FGeometryCollection::VerticesGroup)
+		, BoneIndicesAttribute(InCollection, BoneIndicesAttributeName, FGeometryCollection::VerticesGroup, bInInternalWeights ? FTransformCollection::TransformGroup : NAME_None)
+		, BoneWeightsAttribute(InCollection, BoneWeightsAttributeName, FGeometryCollection::VerticesGroup, bInInternalWeights ? FTransformCollection::TransformGroup : NAME_None)
+		, KinematicWeightAttribute(InCollection, KinematicWeightAttributeName, FGeometryCollection::VerticesGroup)
+		, GeometryLODAttribute(InCollection, GeometryLODAttributeName, FGeometryCollection::GeometryGroup)
+		, SkeletalMeshAttribute(InCollection, SkeletalMeshAttributeName, FGeometryCollection::GeometryGroup)
+		, bInternalWeights(bInInternalWeights)
 	{
 	}
 
-
 	//
 	//  Initialization
 	//
@@ -51,45 +60,62 @@ namespace GeometryCollection::Facades
 	void FVertexBoneWeightsFacade::DefineSchema()
 	{
 		check(!IsConst());
-		BoneIndexAttribute.Add();
-		BoneWeightAttribute.Add();
-		KinematicAttribute.AddAndFill(false);
-		ParentAttribute.Add();
-		if (!BoneIndexAttribute.IsValid())
+		BoneIndicesAttribute.Add();
+		BoneWeightsAttribute.Add();
+		KinematicWeightAttribute.AddAndFill(0.0f);
+		SkeletalMeshAttribute.Add();
+		GeometryLODAttribute.Add();
+		
+		if (!BoneIndicesAttribute.IsValid())
 		{
-			UE_LOG(LogDataflowVertexBoneWeightsFacade, Warning, TEXT("FVertexBoneWeightsFacade failed to initialize because '%s' attribute from '%s' group expected type TArray<int32>."), *FVertexBoneWeightsFacade::BoneIndexAttributeName.ToString(), *FGeometryCollection::VerticesGroup.ToString());
+			UE_LOG(LogDataflowVertexBoneWeightsFacade, Warning, TEXT("FVertexBoneWeightsFacade failed to initialize because '%s' attribute from '%s' group expected type TArray<int32>."), *FVertexBoneWeightsFacade::BoneIndicesAttributeName.ToString(), *FGeometryCollection::VerticesGroup.ToString());
 		}
-		if (!BoneWeightAttribute.IsValid())
+		if (!BoneWeightsAttribute.IsValid())
 		{
-			UE_LOG(LogDataflowVertexBoneWeightsFacade, Warning, TEXT("FVertexBoneWeightsFacade failed to initialize because '%s' attribute from '%s' group expected type TArray<float>."), *FVertexBoneWeightsFacade::BoneWeightAttributeName.ToString(), *FGeometryCollection::VerticesGroup.ToString());
+			UE_LOG(LogDataflowVertexBoneWeightsFacade, Warning, TEXT("FVertexBoneWeightsFacade failed to initialize because '%s' attribute from '%s' group expected type TArray<float>."), *FVertexBoneWeightsFacade::BoneWeightsAttributeName.ToString(), *FGeometryCollection::VerticesGroup.ToString());
 		}
-		if (!KinematicAttribute.IsValid())
+		if (!KinematicWeightAttribute.IsValid())
 		{
-			UE_LOG(LogDataflowVertexBoneWeightsFacade, Warning, TEXT("FVertexBoneWeightsFacade failed to initialize because '%s' attribute from '%s' group expected type <bool>."), *FVertexBoneWeightsFacade::KinematicAttributeName.ToString(), *FGeometryCollection::VerticesGroup.ToString());
+			UE_LOG(LogDataflowVertexBoneWeightsFacade, Warning, TEXT("FVertexBoneWeightsFacade failed to initialize because '%s' attribute from '%s' group expected type <float>."), *FVertexBoneWeightsFacade::KinematicWeightAttributeName.ToString(), *FGeometryCollection::VerticesGroup.ToString());
 		}
-		if (!ParentAttribute.IsValid())
+		if (!SkeletalMeshAttribute.IsValid())
 		{
-			UE_LOG(LogDataflowVertexBoneWeightsFacade, Warning, TEXT("FVertexBoneWeightsFacade failed to initialize because '%s' attribute from '%s' group expected type <int32>."), *FTransformCollection::ParentAttribute.ToString(), *FGeometryCollection::TransformGroup.ToString());
+			UE_LOG(LogDataflowVertexBoneWeightsFacade, Warning, TEXT("FVertexBoneWeightsFacade failed to initialize because '%s' attribute from '%s' group expected type <TObjectPtr<UObject>>."), *FVertexBoneWeightsFacade::SkeletalMeshAttributeName.ToString(), *FGeometryCollection::GeometryGroup.ToString());
+		}
+		if (!GeometryLODAttribute.IsValid())
+		{
+			UE_LOG(LogDataflowVertexBoneWeightsFacade, Warning, TEXT("FVertexBoneWeightsFacade failed to initialize because '%s' attribute from '%s' group expected type <int32>."), *FVertexBoneWeightsFacade::GeometryLODAttributeName.ToString(), *FGeometryCollection::GeometryGroup.ToString());
 		}
 	}
 
 	bool FVertexBoneWeightsFacade::IsValid() const
 	{
-		return BoneIndexAttribute.IsValid() && BoneWeightAttribute.IsValid() && ParentAttribute.IsValid() && VerticesAttribute.IsValid();
+		return BoneIndicesAttribute.IsValid() && BoneWeightsAttribute.IsValid() && KinematicWeightAttribute.IsValid() && SkeletalMeshAttribute.IsValid() && GeometryLODAttribute.IsValid();
 	}
 
+	void FVertexBoneWeightsFacade::ModifyGeometryBinding(const int32 GeometryIndex, const TObjectPtr<UObject>& SkeletalMesh, const int32 GeometryLOD)
+	{
+		TManagedArray< TObjectPtr<UObject> >& SkeletalMeshes = SkeletalMeshAttribute.Modify();
+		TManagedArray< int32 >& GeometryLODs = GeometryLODAttribute.Modify();
+
+		if(0 <= GeometryIndex && GeometryIndex < ConstCollection.NumElements(FGeometryCollection::GeometryGroup))
+		{
+			SkeletalMeshes[GeometryIndex] = SkeletalMesh;
+			GeometryLODs[GeometryIndex] = GeometryLOD;
+		}
+	}
 
 	//
 	//  Add Weights from a bone to a vertex 
 	//
 	void FVertexBoneWeightsFacade::AddBoneWeight(int32 VertexIndex, int32 BoneIndex, float BoneWeight)
 	{
-		TManagedArray< TArray<int32> >& IndicesArray = BoneIndexAttribute.Modify();
-		TManagedArray< TArray<float> >& WeightsArray = BoneWeightAttribute.Modify();
-		const TManagedArray<FVector3f>& Vertices = VerticesAttribute.Modify();
-		if (0 <= VertexIndex && VertexIndex < Vertices.Num())
+		TManagedArray< TArray<int32> >& IndicesArray = BoneIndicesAttribute.Modify();
+		TManagedArray< TArray<float> >& WeightsArray = BoneWeightsAttribute.Modify();
+
+		if (0 <= VertexIndex && VertexIndex < NumVertices())
 		{
-			if (0 <= BoneIndex && BoneIndex < ParentAttribute.Num())
+			if (!bInternalWeights || (0 <= BoneIndex && BoneIndex < NumBones()))
 			{
 				IndicesArray[VertexIndex].Add(BoneIndex);
 				WeightsArray[VertexIndex].Add(BoneWeight);
@@ -97,23 +123,35 @@ namespace GeometryCollection::Facades
 		}
 	}
 
-	void FVertexBoneWeightsFacade::ModifyBoneWeight(int32 VertexIndex, TArray<int32> VertexBoneIndex, TArray<float> VertexBoneWeight)
+	void FVertexBoneWeightsFacade::ModifyKinematicWeight(int32 VertexIndex, const float KinematicWeight)
+	{
+		if (KinematicWeightAttribute.IsValid() && KinematicWeightAttribute.IsValidIndex(VertexIndex))
+		{
+			KinematicWeightAttribute.ModifyAt(VertexIndex, KinematicWeight);
+		}
+	}
+
+	void FVertexBoneWeightsFacade::ModifyBoneWeight(int32 VertexIndex, const TArray<int32>& VertexBoneIndices, const TArray<float>& VertexBoneWeights)
 	{
-		TManagedArray< TArray<int32> >& IndicesArray = BoneIndexAttribute.Modify();
-		TManagedArray< TArray<float> >& WeightsArray = BoneWeightAttribute.Modify();
-		const TManagedArray<FVector3f>& Vertices = VerticesAttribute.Modify();
-		if (VerticesAttribute.IsValidIndex(VertexIndex))
+		TManagedArray< TArray<int32> >& IndicesArray = BoneIndicesAttribute.Modify();
+		TManagedArray< TArray<float> >& WeightsArray = BoneWeightsAttribute.Modify();
+		
+		if (IndicesArray.IsValidIndex(VertexIndex) && WeightsArray.IsValidIndex(VertexIndex))
 		{
 			IndicesArray[VertexIndex].Empty();
 			WeightsArray[VertexIndex].Empty();
 			float TotalWeight = 0.f;
-			for (int32 Idx = 0; Idx < VertexBoneIndex.Num(); ++Idx)
+			int32 LocalIndex = 0;
+			
+			for (const int32& BoneIndex : VertexBoneIndices)
 			{
-				if (ParentAttribute.IsValidIndex(VertexBoneIndex[Idx]))
+				if (!bInternalWeights || (0 <= BoneIndex && BoneIndex < NumBones()))
 				{
-					IndicesArray[VertexIndex].Add(VertexBoneIndex[Idx]);
-					WeightsArray[VertexIndex].Add(VertexBoneWeight[Idx]);
-					TotalWeight += VertexBoneWeight[Idx];
+					const float BoneWeight = VertexBoneWeights[LocalIndex++];
+					
+					IndicesArray[VertexIndex].Add(BoneIndex);
+					WeightsArray[VertexIndex].Add(BoneWeight);
+					TotalWeight += BoneWeight;
 				}
 			}
 			if (TotalWeight < 1.f - UE_KINDA_SMALL_NUMBER || TotalWeight > 1.f + UE_KINDA_SMALL_NUMBER)
@@ -125,22 +163,22 @@ namespace GeometryCollection::Facades
 
 	void FVertexBoneWeightsFacade::SetVertexKinematic(int32 VertexIndex, bool Value)
 	{
-		if (KinematicAttribute.IsValid() && KinematicAttribute.IsValidIndex(VertexIndex))
+		if (KinematicWeightAttribute.IsValid() && KinematicWeightAttribute.IsValidIndex(VertexIndex))
 		{
-			KinematicAttribute.ModifyAt(VertexIndex, Value);
+			KinematicWeightAttribute.ModifyAt(VertexIndex, Value);
 		}
 	}
 
 	void FVertexBoneWeightsFacade::SetVertexArrayKinematic(const TArray<int32>& VertexIndices, bool Value)
 	{
-		if (KinematicAttribute.IsValid())
+		if (KinematicWeightAttribute.IsValid())
 		{
-			TManagedArray<bool>& KinematicArray = KinematicAttribute.Modify();
+			TManagedArray<float>& KinematicWeights = KinematicWeightAttribute.Modify();
 			for (const int32& VertexIndex : VertexIndices)
 			{
-				if (KinematicArray.IsValidIndex(VertexIndex))
+				if (KinematicWeights.IsValidIndex(VertexIndex))
 				{
-					KinematicArray[VertexIndex] = Value;
+					KinematicWeights[VertexIndex] = Value;
 				}
 			}
 		}
@@ -148,89 +186,163 @@ namespace GeometryCollection::Facades
 
 	bool FVertexBoneWeightsFacade::IsKinematicVertex(int32 VertexIndex) const
 	{
-		if (KinematicAttribute.IsValid())
+		if(ConstCollection.HasAttribute(DeprecatedKinematicFlagAttributeName, FGeometryCollection::VerticesGroup))
+    	{
+    		const TManagedArray<bool>& KinematicFlag = ConstCollection.GetAttribute<bool>(DeprecatedKinematicFlagAttributeName, FGeometryCollection::VerticesGroup);
+    		return KinematicFlag.IsValidIndex(VertexIndex) && KinematicFlag.GetConstArray()[VertexIndex];
+    	}
+		else if (KinematicWeightAttribute.IsValid())
+		{
+			return KinematicWeightAttribute.IsValidIndex(VertexIndex) && (KinematicWeightAttribute.Get()[VertexIndex] != 1.0);
+		}
+		else if(ConstCollection.HasAttribute(DeprecatedKinematicFlagAttributeName, FGeometryCollection::VerticesGroup))
 		{
-			return KinematicAttribute.IsValidIndex(VertexIndex) && KinematicAttribute.Get()[VertexIndex];
+			const TManagedArray<bool>& KinematicFlag = ConstCollection.GetAttribute<bool>(DeprecatedKinematicFlagAttributeName, FGeometryCollection::VerticesGroup);
+			return KinematicFlag.IsValidIndex(VertexIndex) && KinematicFlag.GetConstArray()[VertexIndex];
 		}
 		else //backward compatibility for KinematicAttribute added in 5.5
 		{
-			return BoneIndexAttribute.IsValid() && BoneIndexAttribute.IsValidIndex(VertexIndex) && BoneIndexAttribute.Get()[VertexIndex].Num()
-				&& BoneWeightAttribute.IsValid() && BoneWeightAttribute.IsValidIndex(VertexIndex) && BoneWeightAttribute.Get()[VertexIndex].Num();
+			return BoneIndicesAttribute.IsValid() && BoneIndicesAttribute.IsValidIndex(VertexIndex) && BoneIndicesAttribute.Get()[VertexIndex].Num()
+				&& BoneWeightsAttribute.IsValid() && BoneWeightsAttribute.IsValidIndex(VertexIndex) && BoneWeightsAttribute.Get()[VertexIndex].Num();
 		}
 	};
 
+	const TManagedArray< TArray<int32> >* FVertexBoneWeightsFacade::FindBoneIndices()  const
+	{
+		if(ConstCollection.HasAttribute(DeprecatedBoneIndicesAttributeName, FGeometryCollection::VerticesGroup))
+		{
+			return ConstCollection.FindAttribute<TArray<int32> >(DeprecatedBoneIndicesAttributeName, FGeometryCollection::VerticesGroup);
+		}
+		else
+		{
+			return BoneIndicesAttribute.Find();
+		}
+	}
+	
+	const TManagedArray< TArray<int32> >& FVertexBoneWeightsFacade::GetBoneIndices() const
+	{
+		if(ConstCollection.HasAttribute(DeprecatedBoneIndicesAttributeName, FGeometryCollection::VerticesGroup))
+		{
+			return ConstCollection.GetAttribute<TArray<int32> >(DeprecatedBoneIndicesAttributeName, FGeometryCollection::VerticesGroup);
+		}
+		else
+		{
+			return BoneIndicesAttribute.Get();
+		}
+	}
+
+	const TManagedArray< TArray<float> >* FVertexBoneWeightsFacade::FindBoneWeights()  const
+	{
+		if(ConstCollection.HasAttribute(DeprecatedBoneWeightsAttributeName, FGeometryCollection::VerticesGroup))
+		{
+			return ConstCollection.FindAttribute<TArray<float>>(DeprecatedBoneWeightsAttributeName, FGeometryCollection::VerticesGroup);
+		}
+		else
+		{
+			return BoneWeightsAttribute.Find();
+		}
+	}
+	
+	const TManagedArray< TArray<float> >& FVertexBoneWeightsFacade::GetBoneWeights() const
+	{
+		if(ConstCollection.HasAttribute(DeprecatedBoneWeightsAttributeName, FGeometryCollection::VerticesGroup))
+		{
+			return ConstCollection.GetAttribute<TArray<float>>(DeprecatedBoneWeightsAttributeName, FGeometryCollection::VerticesGroup);
+		}
+		else
+		{
+			return BoneWeightsAttribute.Get();
+		}
+	}
+
 	//
 	//  Add Weights from Selection 
 	//
 
-	void FVertexBoneWeightsFacade::AddBoneWeightsFromKinematicBindings() {
+	void FVertexBoneWeightsFacade::NormalizeBoneWeights()
+	{
+		if (IsValid())
+		{
+			TManagedArray< TArray<float> >& WeightsArray = BoneWeightsAttribute.Modify();
+			for(TArray<float>& BoneWeights : WeightsArray)
+			{
+				float TotalWeight = 0.0;
+				for(float& BoneWeight : BoneWeights)
+				{
+					TotalWeight += BoneWeight;
+				}
+				if(TotalWeight <= 1.f + UE_KINDA_SMALL_NUMBER)
+				{
+					for(float& BoneWeight : BoneWeights)
+					{
+						BoneWeight /= TotalWeight;
+					}
+				}
+			}
+		}
+	}
+
+	void FVertexBoneWeightsFacade::AddBoneWeightsFromKinematicBindings()
+	{
 		check(!IsConst());
 		DefineSchema();
 
-		if (IsValid())
+		if (IsValid() && bInternalWeights)
 		{
-			TArray<float> Weights;
-			TArray<int32> Indices;
-
-			int32 NumBones = ParentAttribute.Num(), NumVertices = BoneIndexAttribute.Num();
-			TManagedArray< TArray<int32> >& IndicesArray = BoneIndexAttribute.Modify();
-			TManagedArray< TArray<float> >& WeightsArray = BoneWeightAttribute.Modify();
+			TManagedArray< TArray<int32> >& IndicesArray = BoneIndicesAttribute.Modify();
+			TManagedArray< TArray<float> >& WeightsArray = BoneWeightsAttribute.Modify();
 
 			TArray<float> TotalWeights;
 			TotalWeights.Init(0.f, WeightsArray.Num());
 
-			for (int32 Vert = 0; Vert < WeightsArray.Num(); Vert++)
+			int32 VertexIndex = 0;
+			for(TArray<float>& BoneWeights : WeightsArray)
 			{
-				for (int32 i = 0; i < WeightsArray[Vert].Num(); i++)
+				for(float& BoneWeight : BoneWeights)
 				{
-					TotalWeights[Vert] += WeightsArray[Vert][i];
+					TotalWeights[VertexIndex] += BoneWeight;
 				}
+				++VertexIndex;
 			}
 
 			GeometryCollection::Facades::FKinematicBindingFacade BindingFacade(ConstCollection);
 			for (int32 Kdx = BindingFacade.NumKinematicBindings() - 1; 0 <= Kdx; Kdx--)
 			{
-				int32 Bone;
-				TArray<int32> OutBoneVerts;
-				TArray<float> OutBoneWeights;
+				int32 BoneIndex;
+				TArray<int32> VertexIndices;
+				TArray<float> VertexWeights;
 
-				BindingFacade.GetBoneBindings(BindingFacade.GetKinematicBindingKey(Kdx), Bone, OutBoneVerts, OutBoneWeights);
+				BindingFacade.GetBoneBindings(BindingFacade.GetKinematicBindingKey(Kdx), BoneIndex, VertexIndices, VertexWeights);
 
-				if (0 <= Bone && Bone < NumBones)
+				if (0 <= BoneIndex && BoneIndex < NumBones())
 				{
-					for (int32 Vdx = 0; Vdx < OutBoneVerts.Num(); Vdx++)
+					for (int32 LocalVertex = 0; LocalVertex < VertexIndices.Num(); ++LocalVertex)
 					{
-						int32 Vert = OutBoneVerts[Vdx]; float Weight = OutBoneWeights[Vdx];
-						if (0 <= Vert && Vert < NumVertices && !IndicesArray[Vert].Contains(Bone))
+						VertexIndex = VertexIndices[LocalVertex];
+						const float VertexWeight = VertexWeights[LocalVertex];
+						if (0 <= VertexIndex && VertexIndex < NumVertices() && !IndicesArray[VertexIndex].Contains(BoneIndex))
 						{
-							SetVertexKinematic(Vert);
-							int32 BoneIndex = IndicesArray[Vert].Find(Bone);
-							if (TotalWeights[Vert] + Weight <= 1.f + UE_KINDA_SMALL_NUMBER)
+							SetVertexKinematic(VertexIndex);
+							int32 LocalBone = IndicesArray[VertexIndex].Find(BoneIndex);
+							if (TotalWeights[VertexIndex] + VertexWeight <= 1.f + UE_KINDA_SMALL_NUMBER)
 							{
-								if (BoneIndex == INDEX_NONE)
-								{
-									IndicesArray[Vert].Add(Bone);
-									WeightsArray[Vert].Add(Weight);
-									TotalWeights[Vert] += Weight;
-								}
-								else if (0 <= BoneIndex && BoneIndex < WeightsArray[Vert].Num())
-								{
-									WeightsArray[Vert][BoneIndex] = Weight;
-									TotalWeights[Vert] += Weight;
-								}
+								IndicesArray[VertexIndex].Add(BoneIndex);
+								WeightsArray[VertexIndex].Add(VertexWeight);
+
+								// Maybe we should re-normalize the weights after the kinematic bindings loop
+								TotalWeights[VertexIndex] += VertexWeight;
 							}
 							else
 							{
-								UE_LOG(LogChaos, Warning, TEXT("Bone weight sum %f exceeds 1 on vertex %d"), TotalWeights[Vert] + Weight, Vert);
+								UE_LOG(LogChaos, Warning, TEXT("Bone weight sum %f exceeds 1 on vertex %d"), TotalWeights[VertexIndex] + VertexWeight, VertexIndex);
 							}
 						}
 					}
 				}
 			}
-
+			NormalizeBoneWeights();
 		}
 	}
-
-};
+}
 
 
diff --git a/Engine/Source/Runtime/Experimental/Chaos/Private/GeometryCollection/GeometryCollection.cpp b/Engine/Source/Runtime/Experimental/Chaos/Private/GeometryCollection/GeometryCollection.cpp
--- a/Engine/Source/Runtime/Experimental/Chaos/Private/GeometryCollection/GeometryCollection.cpp
+++ b/Engine/Source/Runtime/Experimental/Chaos/Private/GeometryCollection/GeometryCollection.cpp
@@ -48,6 +48,19 @@ const FName FGeometryCollection::BoundingBoxAttribute("BoundingBox");
 const FName FGeometryCollection::TransformIndexAttribute("TransformIndex");
 const FName FGeometryCollection::VertexStartAttribute("VertexStart");
 const FName FGeometryCollection::VertexCountAttribute("VertexCount");
+const FName FGeometryCollection::InnerRadiusAttribute("InnerRadius");
+const FName FGeometryCollection::OuterRadiusAttribute("OuterRadius");
+const FName FGeometryCollection::FaceStartAttribute("FaceStart");
+const FName FGeometryCollection::FaceCountAttribute("FaceCount");
+const FName FGeometryCollection::FaceIndicesAttribute("Indices");
+const FName FGeometryCollection::FaceVisibleAttribute("Visible");
+const FName FGeometryCollection::MaterialIndexAttribute("MaterialIndex");
+const FName FGeometryCollection::MaterialIDAttribute("MaterialID");
+const FName FGeometryCollection::VertexPositionAttribute("Vertex");
+const FName FGeometryCollection::VertexNormalAttribute("Normal");
+const FName FGeometryCollection::VertexTangentUAttribute("TangentU");
+const FName FGeometryCollection::VertexTangentVAttribute("TangentV");
+const FName FGeometryCollection::VertexBoneMapAttribute("BoneMap");
 
 bool FGeometryCollection::AreCollisionParticlesEnabled()
 {
@@ -84,17 +97,18 @@ void FGeometryCollection::DefineGeometrySchema(FManagedArrayCollection& InCollec
 	// Vertices Group
 	InCollection.AddAttribute<FVector3f>("Vertex", FGeometryCollection::VerticesGroup);
 	InCollection.AddAttribute<FVector3f>("Normal", FGeometryCollection::VerticesGroup);
-	GeometryCollection::UV::DefineUVSchema(InCollection);
 	InCollection.AddAttribute<FLinearColor>(FGeometryCollection::ColorAttribute, FGeometryCollection::VerticesGroup);
 	InCollection.AddAttribute<FVector3f>("TangentU", FGeometryCollection::VerticesGroup);
 	InCollection.AddAttribute<FVector3f>("TangentV", FGeometryCollection::VerticesGroup);
 	InCollection.AddAttribute<int32>("BoneMap", FGeometryCollection::VerticesGroup, TransformDependency);
+	GeometryCollection::UV::DefineUVSchema(InCollection);
 
 	// Faces Group
 	InCollection.AddAttribute<FIntVector>("Indices", FGeometryCollection::FacesGroup, VerticesDependency);
 	InCollection.AddAttribute<bool>("Visible", FGeometryCollection::FacesGroup);
 	InCollection.AddAttribute<int32>("MaterialIndex", FGeometryCollection::FacesGroup);
 	InCollection.AddAttribute<int32>("MaterialID", FGeometryCollection::FacesGroup);
+	InCollection.AddAttribute<bool>("Internal", FGeometryCollection::FacesGroup);
 
 	// Geometry Group
 	InCollection.AddAttribute<int32>("TransformIndex", FGeometryCollection::GeometryGroup, TransformDependency);
diff --git a/Engine/Source/Runtime/Experimental/Chaos/Private/GeometryCollection/ManagedArrayCollection.cpp b/Engine/Source/Runtime/Experimental/Chaos/Private/GeometryCollection/ManagedArrayCollection.cpp
--- a/Engine/Source/Runtime/Experimental/Chaos/Private/GeometryCollection/ManagedArrayCollection.cpp
+++ b/Engine/Source/Runtime/Experimental/Chaos/Private/GeometryCollection/ManagedArrayCollection.cpp
@@ -535,12 +535,17 @@ void FManagedArrayCollection::SetDefaults(FName Group, uint32 StartSize, uint32
 
 void FManagedArrayCollection::SyncGroupSizeFrom(const FManagedArrayCollection& InCollection, FName Group)
 {
-	if (!HasGroup(Group))
+	SyncGroupSizeFrom(InCollection, Group, Group);
+}
+
+void FManagedArrayCollection::SyncGroupSizeFrom(const FManagedArrayCollection& InCollection, FName SrcGroup, FName DstGroup)
+{
+	if (!HasGroup(DstGroup))
 	{
-		AddGroup(Group);
+		AddGroup(DstGroup);
 	}
 
-	Resize(InCollection.GroupInfo[Group].Size, Group);
+	Resize(InCollection.GroupInfo[SrcGroup].Size, DstGroup);
 }
 
 void FManagedArrayCollection::CopyMatchingAttributesFrom(const FManagedArrayCollection& FromCollection, const TArrayView<const FAttributeAndGroupId> SkipList)
@@ -615,16 +620,21 @@ void FManagedArrayCollection::CopyMatchingAttributesFrom(
 
 void FManagedArrayCollection::CopyAttribute(const FManagedArrayCollection& InCollection, FName Name, FName Group)
 {
-	CopyAttribute(InCollection, /*SrcName=*/Name, /*DestName=*/Name, Group);
+	CopyAttribute(InCollection, /*SrcName=*/Name, /*DestName=*/Name, Group, Group);
 }
 
 void FManagedArrayCollection::CopyAttribute(const FManagedArrayCollection& InCollection, FName SrcName, FName DestName, FName Group)
 {
-	SyncGroupSizeFrom(InCollection, Group);
-	FKeyType SrcKey = FManagedArrayCollection::MakeMapKey(SrcName, Group);
-	FKeyType DestKey = FManagedArrayCollection::MakeMapKey(DestName, Group);
+	CopyAttribute(InCollection, SrcName, DestName, Group, Group);
+}
+
+void FManagedArrayCollection::CopyAttribute(const FManagedArrayCollection& InCollection, FName SrcName, FName DestName, FName SrcGroup, FName DstGroup)
+{
+	SyncGroupSizeFrom(InCollection, SrcGroup, DstGroup);
+	FKeyType SrcKey = FManagedArrayCollection::MakeMapKey(SrcName, SrcGroup);
+	FKeyType DestKey = FManagedArrayCollection::MakeMapKey(DestName, DstGroup);
 
-	if (!HasAttribute(DestName, Group))
+	if (!HasAttribute(DestName, DstGroup))
 	{
 		FValueType NewAttribute(InCollection.Map[SrcKey]);
 		Map.Add(DestKey, MoveTemp(NewAttribute));
diff --git a/Engine/Source/Runtime/Experimental/Chaos/Private/GeometryCollection/TransformCollection.cpp b/Engine/Source/Runtime/Experimental/Chaos/Private/GeometryCollection/TransformCollection.cpp
--- a/Engine/Source/Runtime/Experimental/Chaos/Private/GeometryCollection/TransformCollection.cpp
+++ b/Engine/Source/Runtime/Experimental/Chaos/Private/GeometryCollection/TransformCollection.cpp
@@ -15,6 +15,10 @@ const FName FTransformCollection::ParticlesAttribute = "Particles";
 const FName FTransformCollection::LevelAttribute = "Level";
 const FName FTransformCollection::ConvexGroup = "Convex";
 const FName FTransformCollection::ConvexHullAttribute = "ConvexHull";
+const FName FTransformCollection::BoneNameAttribute = "BoneName";
+const FName FTransformCollection::BoneColorAttribute = "BoneColor";
+const FName FTransformCollection::BoneHierarchyAttribute = "BoneHierarchy";
+const FName FTransformCollection::GeometryIndexAttribute = "TransformToGeometryIndex";
 
 FTransformCollection::FTransformCollection()
 	: FManagedArrayCollection()
@@ -28,8 +32,8 @@ void FTransformCollection::DefineTransformSchema(FManagedArrayCollection& InColl
 {
 	// Hierarchy Group
 	InCollection.AddAttribute<FTransform3f>(FTransformCollection::TransformAttribute, FTransformCollection::TransformGroup);
-	InCollection.AddAttribute<FString>("BoneName", FTransformCollection::TransformGroup);
-	InCollection.AddAttribute<FLinearColor>("BoneColor", FTransformCollection::TransformGroup);
+	InCollection.AddAttribute<FString>(FTransformCollection::BoneNameAttribute, FTransformCollection::TransformGroup);
+	InCollection.AddAttribute<FLinearColor>(FTransformCollection::BoneColorAttribute, FTransformCollection::TransformGroup);
 	InCollection.AddAttribute<int32>(FTransformCollection::ParentAttribute, FTransformCollection::TransformGroup);
 	InCollection.AddAttribute<TSet<int32>>(FTransformCollection::ChildrenAttribute, FTransformCollection::TransformGroup);
 }
@@ -38,8 +42,8 @@ void FTransformCollection::Construct()
 {
 	// Hierarchy Group
 	AddExternalAttribute<FTransform3f>(FTransformCollection::TransformAttribute, FTransformCollection::TransformGroup, Transform);
-	AddExternalAttribute<FString>("BoneName", FTransformCollection::TransformGroup, BoneName);
-	AddExternalAttribute<FLinearColor>("BoneColor", FTransformCollection::TransformGroup, BoneColor);
+	AddExternalAttribute<FString>(FTransformCollection::BoneNameAttribute, FTransformCollection::TransformGroup, BoneName);
+	AddExternalAttribute<FLinearColor>(FTransformCollection::BoneColorAttribute, FTransformCollection::TransformGroup, BoneColor);
 	AddExternalAttribute<int32>(FTransformCollection::ParentAttribute, FTransformCollection::TransformGroup, Parent);
 	AddExternalAttribute<TSet<int32>>(FTransformCollection::ChildrenAttribute, FTransformCollection::TransformGroup, Children);
 }
@@ -57,7 +61,7 @@ void FTransformCollection::Serialize(Chaos::FChaosArchive& Ar)
 	if (Ar.IsLoading())
 	{
 		// for backwards compatibility convert old BoneHierarchy struct into split out arrays
-		const TManagedArray<FGeometryCollectionBoneNode>* BoneHierarchyPtr = FindAttribute<FGeometryCollectionBoneNode>("BoneHierarchy", FTransformCollection::TransformGroup);
+		const TManagedArray<FGeometryCollectionBoneNode>* BoneHierarchyPtr = FindAttribute<FGeometryCollectionBoneNode>(FTransformCollection::BoneHierarchyAttribute, FTransformCollection::TransformGroup);
 		if (BoneHierarchyPtr)
 		{
 			const TManagedArray<FGeometryCollectionBoneNode>& BoneHierarchy = *BoneHierarchyPtr;
diff --git a/Engine/Source/Runtime/Experimental/Chaos/Public/GeometryCollection/Facades/CollectionCurveFacade.h b/Engine/Source/Runtime/Experimental/Chaos/Public/GeometryCollection/Facades/CollectionCurveFacade.h
new file mode 100644
--- /dev/null
+++ b/Engine/Source/Runtime/Experimental/Chaos/Public/GeometryCollection/Facades/CollectionCurveFacade.h
@@ -0,0 +1,215 @@
+﻿// Copyright Epic Games, Inc. All Rights Reserved.
+
+#pragma once
+
+#include "GeometryCollection/ManagedArrayAccessor.h"
+#include "GeometryCollection/ManagedArrayCollection.h"
+#include "GeometryCollection/GeometryCollection.h"
+	
+namespace UE::GeometryCollection::Facades
+{
+	class FCollectionCurveGeometryFacade
+	{
+	public:
+		/** Collection group names */
+		static CHAOS_API const FName CurvesGroup;
+		static CHAOS_API const FName PointsGroup;
+		
+		/** Collection attribute names */
+		static CHAOS_API const FName CurvePointOffsetsAttribute;
+		static CHAOS_API const FName GeometryCurveOffsetsAttribute;
+		static CHAOS_API const FName PointRestPositionsAttribute;
+		static CHAOS_API const FName PointRestOrientationsAttribute;
+		static CHAOS_API const FName PointCurveIndicesAttribute;
+		static CHAOS_API const FName CurveGeometryIndicesAttribute;
+		static CHAOS_API const FName GeometryGroupNamesAttribute;
+		static CHAOS_API const FName GeometryCurveThicknessAttribute;
+		static CHAOS_API const FName CurveSourceIndicesAttribute;
+
+		CHAOS_API FCollectionCurveGeometryFacade(FManagedArrayCollection& InCollection);
+		CHAOS_API FCollectionCurveGeometryFacade(const FManagedArrayCollection& InCollection);
+
+		/** Create the facade attributes. */
+		CHAOS_API void DefineSchema();
+
+		/** Is the facade defined constant. */
+		bool IsConst() const { return Collection == nullptr; }
+
+		/** Is the Facade defined on the collection? */
+		CHAOS_API bool IsValid() const;
+		
+		/** Get the number of curves */
+		int32 GetNumCurves() const { return CurvePointOffsets.Num(); }
+
+		/** Get the number of geometry */
+		int32 GetNumGeometry() const { return GeometryCurveOffsets.Num(); }
+
+		/** Get the number of points */
+		int32 GetNumPoints() const { return PointRestPositions.Num(); }
+		
+		/** Get the point rest positions */
+		const TArray<FVector3f>& GetPointRestPositions() const { return PointRestPositions.Get().GetConstArray(); }
+
+		/** Get the point rest orientations */
+		const TArray<FQuat4f>& GetPointRestOrientations() const { return PointRestOrientations.Get().GetConstArray(); }
+
+		/** Get the curve point offsets */
+		const TArray<int32>& GetCurvePointOffsets() const { return CurvePointOffsets.Get().GetConstArray(); }
+
+		/** Get the geometry curve offsets */
+		const TArray<int32>& GetGeometryCurveOffsets() const { return GeometryCurveOffsets.Get().GetConstArray(); }
+
+		/** Get the point curve indices */
+		const TArray<int32>& GetPointCurveIndices() const { return PointCurveIndices.Get().GetConstArray(); }
+
+		/** Get the curve geometry indices */
+		const TArray<int32>& GetCurveGeometryIndices() const { return CurveGeometryIndices.Get().GetConstArray(); }
+
+		/** Get the geometry group names */
+		const TArray<FString>& GetGeometryGroupNames() const { return GeometryGroupNames.Get().GetConstArray(); }
+
+		/** Get the geometry curve thickness */
+		const TArray<float>& GetGeometryCurveThickness() const { return GeometryCurveThickness.Get().GetConstArray(); }
+
+		/** Get the curve source indices*/
+		const TArray<int32>& GetCurveSourceIndices() const { return CurveSourceIndices.Get().GetConstArray(); }
+		
+		/** Set the point rest positions */
+		void SetPointRestPositions(const TArray<FVector3f>& InPointRestPositions) { PointRestPositions.Modify() = InPointRestPositions; UpdatePointRestOrientations(); }	
+
+		/** Set the geometry  group names */
+		void SetGeometryGroupNames(const TArray<FString>& InGeometryGroupNames) { GeometryGroupNames.Modify() = InGeometryGroupNames; }
+		
+		/** Set the geometry curve thickness */
+		void SetGeometryCurveThickness(const TArray<float>& InGeometryCurveThickness) { GeometryCurveThickness.Modify() = InGeometryCurveThickness; }	
+		
+		/** Set the curve point offsets */
+		void SetCurvePointOffsets(const TArray<int32>& InCurvePointOffsets) { CurvePointOffsets.Modify() = InCurvePointOffsets; UpdatePointCurveIndices(); }
+
+		/** Set the geometry curve offsets */
+		void SetGeometryCurveOffsets(const TArray<int32>& InGeometryCurveOffsets) { GeometryCurveOffsets.Modify() = InGeometryCurveOffsets; UpdateCurveGeometryIndices(); }
+
+		/** Set the curve source indices */
+		void SetCurveSourceIndices(const TArray<int32>& InCurveSourceIndices) { CurveSourceIndices.Modify() = InCurveSourceIndices; }
+		
+		/** Initialize the whole curve collection  */
+		CHAOS_API void InitCurvesCollection(const TArray<FVector3f>& InPointRestPositions, const TArray<int32>& InCurvePointOffsets, const TArray<int32>& InGeometryCurveOffsets,
+			const TArray<FString>& InGeometryGroupNames, const TArray<float>& InGeometryCurveThickness, const TArray<int32>& InCurveSourceIndices);
+
+		/** Get the managed array collection */
+		const FManagedArrayCollection& GetManagedArrayCollection() const {return ConstCollection;}
+		
+	protected :
+
+		/** Build geometry collection from the curves one */
+		void BuildGeometryCollection();
+		
+		/** Update the point rest orientations with parallel transport */
+		void UpdatePointRestOrientations();
+		
+		/** Update the point curve indices from the offsets  */
+		void UpdatePointCurveIndices();
+
+		/** Update the curve object indices from the offsets  */
+		void UpdateCurveGeometryIndices();
+		
+		/** Const collection the facade is linked to */
+		const FManagedArrayCollection& ConstCollection;
+
+		/** Non-const collection the facade is linked to */
+		FManagedArrayCollection* Collection = nullptr;
+		
+		/** Points rest orientation */
+		TManagedArrayAccessor<FQuat4f> PointRestOrientations;
+
+		/** Points rest position */
+		TManagedArrayAccessor<FVector3f> PointRestPositions;
+
+		/** Curves point offset */
+		TManagedArrayAccessor<int32> CurvePointOffsets;
+                                                 
+		/** Geometry curve offset */
+		TManagedArrayAccessor<int32> GeometryCurveOffsets;
+
+		/** Points curve index */
+		TManagedArrayAccessor<int32> PointCurveIndices;
+
+		/** Curves object index */
+		TManagedArrayAccessor<int32> CurveGeometryIndices;
+
+		/** Geometry group names */
+		TManagedArrayAccessor<FString> GeometryGroupNames;
+
+		/** Geometry curve thickness */
+		TManagedArrayAccessor<float> GeometryCurveThickness;
+
+		/** Curve source indices */
+		TManagedArrayAccessor<int32> CurveSourceIndices;
+	};
+
+	class FCollectionCurveHierarchyFacade
+	{
+	public:
+
+		/** Collection attribute names */
+		static CHAOS_API const FName CurveParentIndicesAttribute;
+		static CHAOS_API const FName CurveLodIndicesAttribute;
+
+		CHAOS_API FCollectionCurveHierarchyFacade(FManagedArrayCollection& InCollection);
+		CHAOS_API FCollectionCurveHierarchyFacade(const FManagedArrayCollection& InCollection);
+
+		/** Create the facade attributes. */
+		CHAOS_API void DefineSchema();
+
+		/** Is the facade defined constant. */
+		bool IsConst() const { return Collection == nullptr; }
+
+		/** Is the Facade defined on the collection? */
+		CHAOS_API bool IsValid() const;
+		
+		/** Get the number of curves */
+		int32 GetNumCurves() const { return CurveLodIndices.Num(); }
+		
+		/** Get the curve lod indices */
+		const TArray<int32>& GetCurveLodIndices() const { return CurveLodIndices.Get().GetConstArray(); }
+
+		/** Get the curve parent indices */
+		const TArray<int32>& GetCurveParentIndices() const { return CurveParentIndices.Get().GetConstArray(); }
+		
+		/** Set the curve lod indices */
+		void SetCurveLodIndices(const TArray<int32>& InCurveLodIndices) { CurveLodIndices.Modify() = InCurveLodIndices; }
+
+		/** Set the curve parent indices */
+		void SetCurveParentIndices(const TArray<int32>& InCurveParentIndices) { CurveParentIndices.Modify() = InCurveParentIndices; }
+
+		/** Get the managed array collection */
+		const FManagedArrayCollection& GetManagedArrayCollection() const {return ConstCollection;}
+		
+	protected :
+
+		/** Build geometry collection from the curves one */
+		void BuildGeometryCollection();
+		
+		/** Update the point rest orientations with parallel transport */
+		void UpdatePointRestOrientations();
+		
+		/** Update the point curve indices from the offsets  */
+		void UpdatePointCurveIndices();
+
+		/** Update the curve object indices from the offsets  */
+		void UpdateCurveGeometryIndices();
+		
+		/** Const collection the facade is linked to */
+		const FManagedArrayCollection& ConstCollection;
+
+		/** Non-const collection the facade is linked to */
+		FManagedArrayCollection* Collection = nullptr;
+
+		/** Curve parent indices*/
+		TManagedArrayAccessor<int32> CurveParentIndices;
+
+		/** Curve lod indices */
+		TManagedArrayAccessor<int32> CurveLodIndices;
+	};
+}
+
diff --git a/Engine/Source/Runtime/Experimental/Chaos/Public/GeometryCollection/Facades/CollectionTransformSelectionFacade.h b/Engine/Source/Runtime/Experimental/Chaos/Public/GeometryCollection/Facades/CollectionTransformSelectionFacade.h
--- a/Engine/Source/Runtime/Experimental/Chaos/Public/GeometryCollection/Facades/CollectionTransformSelectionFacade.h
+++ b/Engine/Source/Runtime/Experimental/Chaos/Public/GeometryCollection/Facades/CollectionTransformSelectionFacade.h
@@ -229,10 +229,10 @@ namespace GeometryCollection::Facades
 		CHAOS_API TArray<int32> ConvertGeometrySelectionToFaceSelection(const TArray<int32>& InGeometrySelection) const;
 
 		/**
-		* Convert a GeometrySelection to a VartexSelection
+		* Convert a GeometrySelection to a VertexSelection
 		*/
 		CHAOS_API TArray<int32> ConvertGeometrySelectionToVertexSelection(const TArray<int32>& InGeometrySelection) const;
-
+		
 		/**
 		* Convert a FaceSelection to a GeometrySelection
 		*/
@@ -242,6 +242,46 @@ namespace GeometryCollection::Facades
 		* Convert a VertexSelection to a GeometrySelection
 		*/
 		CHAOS_API TArray<int32> ConvertVertexSelectionToGeometrySelection(const TArray<int32>& InVertexSelection, bool bAllElementsMustBeSelected) const;
+		
+		/**
+		* Convert a TransformSelection to a CurveSelection
+		*/
+		CHAOS_API TArray<int32> ConvertTransformSelectionToCurveSelection(const TArray<int32>& InTransformSelection) const;
+
+		/**
+		* Convert a GeometrySelection to a CurveSelection
+		*/
+		CHAOS_API TArray<int32> ConvertGeometrySelectionToCurveSelection(const TArray<int32>& InGeometrySelection) const;
+
+		/**
+		* Convert a VertexSelection to a CurveSelection
+		*/
+		CHAOS_API TArray<int32> ConvertVertexSelectionToCurveSelection(const TArray<int32>& InVertexSelection, bool bAllElementsMustBeSelected) const;
+
+		/**
+		* Convert a FaceSelection to a CurveSelection
+		*/
+		CHAOS_API TArray<int32> ConvertFaceSelectionToCurveSelection(const TArray<int32>& InFaceSelection, bool bAllElementsMustBeSelected) const;
+		
+		/**
+		* Convert a CurveSelection to a VertexSelection
+		*/
+		CHAOS_API TArray<int32> ConvertCurveSelectionToVertexSelection(const TArray<int32>& InCurveSelection) const;
+		
+		/**
+		* Convert a CurveSelection to a FaceSelection
+		*/
+		CHAOS_API TArray<int32> ConvertCurveSelectionToFaceSelection(const TArray<int32>& InCurveSelection) const;
+		
+		/**
+		* Convert a CurveSelection to a GeometrySelection
+		*/
+		CHAOS_API TArray<int32> ConvertCurveSelectionToGeometrySelection(const TArray<int32>& InCurveSelection, bool bAllElementsMustBeSelected) const;
+		
+		/**
+		* Convert a CurveSelection to a TransformSelection
+		*/
+		CHAOS_API TArray<int32> ConvertCurveSelectionToTransformSelection(const TArray<int32>& InCurveSelection, bool bAllElementsMustBeSelected) const;
 
 
 	private:
diff --git a/Engine/Source/Runtime/Experimental/Chaos/Public/GeometryCollection/Facades/CollectionVertexBoneWeightsFacade.h b/Engine/Source/Runtime/Experimental/Chaos/Public/GeometryCollection/Facades/CollectionVertexBoneWeightsFacade.h
--- a/Engine/Source/Runtime/Experimental/Chaos/Public/GeometryCollection/Facades/CollectionVertexBoneWeightsFacade.h
+++ b/Engine/Source/Runtime/Experimental/Chaos/Public/GeometryCollection/Facades/CollectionVertexBoneWeightsFacade.h
@@ -3,6 +3,7 @@
 #pragma once
 #include "Containers/UnrealString.h"
 #include "Containers/Array.h"
+#include "GeometryCollection/GeometryCollection.h"
 #include "GeometryCollection/Facades/CollectionSelectionFacade.h"
 #include "GeometryCollection/ManagedArrayCollection.h"
 #include "GeometryCollection/ManagedArrayAccessor.h"
@@ -32,14 +33,17 @@ namespace GeometryCollection::Facades
 	public:
 
 		// Attributes
-		static CHAOS_API const FName BoneIndexAttributeName;
-		static CHAOS_API const FName BoneWeightAttributeName;
-		static CHAOS_API const FName KinematicAttributeName;
+		static CHAOS_API const FName BoneIndicesAttributeName;
+		static CHAOS_API const FName BoneWeightsAttributeName;
+		static CHAOS_API const FName KinematicWeightAttributeName;
+		static CHAOS_API const FName GeometryLODAttributeName;
+		static CHAOS_API const FName SkeletalMeshAttributeName;
+		
 		/**
 		* FVertexBoneWeightsFacade Constuctor
 		*/
-		CHAOS_API FVertexBoneWeightsFacade(FManagedArrayCollection& InSelf);
-		CHAOS_API FVertexBoneWeightsFacade(const FManagedArrayCollection& InSelf);
+		CHAOS_API FVertexBoneWeightsFacade(FManagedArrayCollection& InSelf, const bool bInternalWeights = true);
+		CHAOS_API FVertexBoneWeightsFacade(const FManagedArrayCollection& InSelf, const bool bInternalWeights = true);
 
 		/** Define the facade */
 		CHAOS_API void DefineSchema();
@@ -57,40 +61,87 @@ namespace GeometryCollection::Facades
 		CHAOS_API void AddBoneWeight(int32 VertexIndex, int32 BoneIndex, float BoneWeight);
 
 		/** Modify bone weight based on the kinematic bindings. */
-		CHAOS_API void ModifyBoneWeight(int32 VertexIndex, TArray<int32> VertexBoneIndex, TArray<float> VertexBoneWeight);
+		CHAOS_API void ModifyBoneWeight(int32 VertexIndex, const TArray<int32>& VertexBoneIndices, const TArray<float>& VertexBoneWeights);
+
+		/** Modify kinematic weight based */
+		CHAOS_API void ModifyKinematicWeight(int32 VertexIndex, const float KinematicWeight);
 
 		/** Set vertex to be kinematic/dynamic */
 		CHAOS_API void SetVertexKinematic(int32 VertexIndex, bool Value = true);
 
 		/** Set vertex to be kinematic/dynamic */
 		CHAOS_API void SetVertexArrayKinematic(const TArray<int32>& VertexIndices, bool Value = true);
-
-		/** Return the vertex bone indices from the collection. Null if not initialized.  */
-		const TManagedArray< TArray<int32> >* FindBoneIndices()  const { return BoneIndexAttribute.Find(); }
-		const TManagedArray< TArray<int32> >& GetBoneIndices() const { return BoneIndexAttribute.Get(); }
-
+		
 		/** Return if the vertex is kinematic
 		Pre 5.5 we did not have per-vertex kinematic attribute. 
 		This supports defining kinematics without per-vertex kinematic attribute. */
 		CHAOS_API bool IsKinematicVertex(int32 VertexIndex) const;
+		
+		/** Modify the geometry skelmesh and LOD. */
+		CHAOS_API void ModifyGeometryBinding(const int32 GeometryIndex, const TObjectPtr<UObject>& SkeletalMesh, const int32 GeometryLOD);
+
+		/** Normalize bone weights */
+		CHAOS_API void NormalizeBoneWeights();
 
 		/** Return number of vertices */
-		int32 NumVertices() const { return VerticesAttribute.Num(); };
+		int32 NumVertices() const { return ConstCollection.NumElements(FGeometryCollection::VerticesGroup); };
+		
+		/** Return number of bones */
+		int32 NumBones() const { return ConstCollection.NumElements(FGeometryCollection::TransformGroup); };
+		
+		/** Return number of geometries */
+		int32 NumGeometry() const { return ConstCollection.NumElements(FGeometryCollection::GeometryGroup); };
+
+		/** Return the vertex bone indices from the collection. Null if not initialized.  */
+		CHAOS_API const TManagedArray< TArray<int32> >* FindBoneIndices()  const;
+		CHAOS_API const TManagedArray< TArray<int32> >& GetBoneIndices() const;
 
 		/** Return the vertex bone weights from the collection. Null if not initialized. */
-		const TManagedArray< TArray<float> >* FindBoneWeights()  const { return BoneWeightAttribute.Find(); }
-		const TManagedArray< TArray<float> >& GetBoneWeights() const { return BoneWeightAttribute.Get(); }
+		CHAOS_API const TManagedArray< TArray<float> >* FindBoneWeights()  const;
+		CHAOS_API const TManagedArray< TArray<float> >& GetBoneWeights() const;
+
+		/** Return the vertex kinematic weights from the collection. Null if not initialized. */
+		const TManagedArray<float>* FindKinematicWeights()  const { return KinematicWeightAttribute.Find(); }
+		const TManagedArray<float>& GetKinematicWeights() const { return KinematicWeightAttribute.Get(); }
+
+		/** Return the skeletal meshes from the collection. Null if not initialized. */
+		const TManagedArray<TObjectPtr<UObject>>* FindSkeletalMeshes()  const { return SkeletalMeshAttribute.Find(); }
+		const TManagedArray<TObjectPtr<UObject>>& GetSkeletalMeshes() const { return SkeletalMeshAttribute.Get(); }
+		
+		/** Return the geometry LODs from the collection. Null if not initialized. */
+		const TManagedArray<int32>* FindGeomketryLODs()  const { return GeometryLODAttribute.Find(); }
+		const TManagedArray<int32>& GetGeometryLODs() const { return GeometryLODAttribute.Get(); }
+		
+		/** Get the managed array collection */
+		const FManagedArrayCollection& GetManagedArrayCollection() const {return ConstCollection;}
 
 	private:
+
+		/** Deprecated names for attributes */
+		static const FName DeprecatedBoneIndicesAttributeName;
+		static const FName DeprecatedBoneWeightsAttributeName;
+		static const FName DeprecatedKinematicFlagAttributeName;
+		
 		const FManagedArrayCollection& ConstCollection;
 		FManagedArrayCollection* Collection = nullptr;
 
-		TManagedArrayAccessor<TArray<int32>> BoneIndexAttribute;
-		TManagedArrayAccessor<TArray<float>> BoneWeightAttribute;
-		TManagedArrayAccessor<bool> KinematicAttribute;
-		TManagedArrayAccessor<int32> ParentAttribute;
-		TManagedArrayAccessor<FVector3f> VerticesAttribute;
+		/** Bone indices for each vertices */
+		TManagedArrayAccessor<TArray<int32>> BoneIndicesAttribute;
+
+		/** Bone weights for each vertices */
+		TManagedArrayAccessor<TArray<float>> BoneWeightsAttribute;
+
+		/** Kinematic weights for each vertices */
+		TManagedArrayAccessor<float> KinematicWeightAttribute;
+		
+		/** Geometry LODSs the skin weights are linked to */
+		TManagedArrayAccessor<int32> GeometryLODAttribute;
+
+		/** Geometry skeletal meshes the skin weights are linked to*/
+		TManagedArrayAccessor<TObjectPtr<UObject>> SkeletalMeshAttribute;
 
+		/** Internal weights flag to specify if the bone indices are relative to the transform group or to an external skelkmesh */
+		const bool bInternalWeights = true;
 	};
 
 }
diff --git a/Engine/Source/Runtime/Experimental/Chaos/Public/GeometryCollection/GeometryCollection.h b/Engine/Source/Runtime/Experimental/Chaos/Public/GeometryCollection/GeometryCollection.h
--- a/Engine/Source/Runtime/Experimental/Chaos/Public/GeometryCollection/GeometryCollection.h
+++ b/Engine/Source/Runtime/Experimental/Chaos/Public/GeometryCollection/GeometryCollection.h
@@ -113,6 +113,19 @@ public:
 	static CHAOS_API const FName TransformIndexAttribute;
 	static CHAOS_API const FName VertexStartAttribute;
 	static CHAOS_API const FName VertexCountAttribute;
+	static CHAOS_API const FName InnerRadiusAttribute;
+	static CHAOS_API const FName OuterRadiusAttribute;
+	static CHAOS_API const FName FaceStartAttribute;
+	static CHAOS_API const FName FaceCountAttribute;
+	static CHAOS_API const FName FaceIndicesAttribute;
+	static CHAOS_API const FName FaceVisibleAttribute;
+	static CHAOS_API const FName MaterialIndexAttribute;
+	static CHAOS_API const FName MaterialIDAttribute;
+	static CHAOS_API const FName VertexPositionAttribute;
+	static CHAOS_API const FName VertexNormalAttribute;
+	static CHAOS_API const FName VertexTangentUAttribute;
+	static CHAOS_API const FName VertexTangentVAttribute;
+	static CHAOS_API const FName VertexBoneMapAttribute;
 
 	enum ESimulationTypes : uint8
 	{
diff --git a/Engine/Source/Runtime/Experimental/Chaos/Public/GeometryCollection/ManagedArrayCollection.h b/Engine/Source/Runtime/Experimental/Chaos/Public/GeometryCollection/ManagedArrayCollection.h
--- a/Engine/Source/Runtime/Experimental/Chaos/Public/GeometryCollection/ManagedArrayCollection.h
+++ b/Engine/Source/Runtime/Experimental/Chaos/Public/GeometryCollection/ManagedArrayCollection.h
@@ -516,6 +516,15 @@ public:
 	*/
 	CHAOS_API void CopyAttribute(const FManagedArrayCollection& InCollection, FName SrcName, FName DestName, FName Group);
 
+	/**
+	* Copy an attribute from one group to another group. Will perform an implicit group sync.
+	* @param SrcName - The name of the attribute being copied from
+	* @param DestName - The name of the attribute being copied to
+	* @param SrcGroup - The group that manages the attribute being copied from 
+	* @param DstGroup - The group that manages the attribute being copied to
+	*/
+	CHAOS_API void CopyAttribute(const FManagedArrayCollection& InCollection, FName SrcName, FName DestName, FName SrcGroup, FName DstGroup);
+
 	/**
 	* Copy attributes that match the input collection. This is a utility to easily sync collections
 	* @param InCollection - All groups from this collection found in the input will be sized accordingly
@@ -832,6 +841,14 @@ protected:
 	* @param Group - The group that manages the attribute
 	*/
 	CHAOS_API void SyncGroupSizeFrom(const FManagedArrayCollection& InCollection, FName Group);
+	
+	/**
+	* Size and order a dstgroup so that it matches the srcgroup found in the input collection.
+	* @param InCollection - The collection we are ordering our group against. 
+	* @param SrcGroup - The group that manages the attribute in the input collection
+	* @param DstGroup - The group that manages the attribute in the output collection
+	*/
+	CHAOS_API void SyncGroupSizeFrom(const FManagedArrayCollection& InCollection, FName SrcGroup, FName DstGroup);
 
 	/**
 	 * Version to indicate need for conditioning to current expected data layout during serialization loading.
diff --git a/Engine/Source/Runtime/Experimental/Chaos/Public/GeometryCollection/TransformCollection.h b/Engine/Source/Runtime/Experimental/Chaos/Public/GeometryCollection/TransformCollection.h
--- a/Engine/Source/Runtime/Experimental/Chaos/Public/GeometryCollection/TransformCollection.h
+++ b/Engine/Source/Runtime/Experimental/Chaos/Public/GeometryCollection/TransformCollection.h
@@ -54,6 +54,10 @@ public:
 	static CHAOS_API const FName LevelAttribute;
 	static CHAOS_API const FName ConvexGroup;
 	static CHAOS_API const FName ConvexHullAttribute;
+	static CHAOS_API const FName BoneNameAttribute;
+	static CHAOS_API const FName BoneColorAttribute;
+	static CHAOS_API const FName BoneHierarchyAttribute;
+	static CHAOS_API const FName GeometryIndexAttribute;
 	
 	/** Schema */
 	static CHAOS_API void DefineTransformSchema(FManagedArrayCollection&);
diff --git a/Engine/Source/Runtime/Experimental/Dataflow/Core/Private/Dataflow/DataflowConvertNodes.cpp b/Engine/Source/Runtime/Experimental/Dataflow/Core/Private/Dataflow/DataflowConvertNodes.cpp
--- a/Engine/Source/Runtime/Experimental/Dataflow/Core/Private/Dataflow/DataflowConvertNodes.cpp
+++ b/Engine/Source/Runtime/Experimental/Dataflow/Core/Private/Dataflow/DataflowConvertNodes.cpp
@@ -214,6 +214,15 @@ void FConvertSelectionTypesDataflowNode::Evaluate(UE::Dataflow::FContext& Contex
 					FDataflowGeometrySelection NewSelection;
 					NewSelection.InitFromArray(InCollection, SelectionArr);
 
+					SetValue(Context, MoveTemp(NewSelection), &Out);
+				}
+				else if (Output && Output->IsType<FDataflowCurveSelection>())
+				{
+					const TArray<int32>& SelectionArr = TransformSelectionFacade.ConvertTransformSelectionToCurveSelection(InValue.AsArray());
+
+					FDataflowCurveSelection NewSelection;
+					NewSelection.InitFromArray(InCollection, SelectionArr);
+
 					SetValue(Context, MoveTemp(NewSelection), &Out);
 				}
 			}
@@ -244,6 +253,15 @@ void FConvertSelectionTypesDataflowNode::Evaluate(UE::Dataflow::FContext& Contex
 					FDataflowGeometrySelection NewSelection;
 					NewSelection.InitFromArray(InCollection, SelectionArr);
 
+					SetValue(Context, MoveTemp(NewSelection), &Out);
+				}
+				else if (Output && Output->IsType<FDataflowCurveSelection>())
+				{
+					const TArray<int32>& SelectionArr = TransformSelectionFacade.ConvertFaceSelectionToCurveSelection(InValue.AsArray(), bAllElementsMustBeSelected);
+
+					FDataflowCurveSelection NewSelection;
+					NewSelection.InitFromArray(InCollection, SelectionArr);
+
 					SetValue(Context, MoveTemp(NewSelection), &Out);
 				}
 			}
@@ -274,6 +292,15 @@ void FConvertSelectionTypesDataflowNode::Evaluate(UE::Dataflow::FContext& Contex
 					FDataflowGeometrySelection NewSelection;
 					NewSelection.InitFromArray(InCollection, SelectionArr);
 
+					SetValue(Context, MoveTemp(NewSelection), &Out);
+				}
+				else if (Output && Output->IsType<FDataflowCurveSelection>())
+				{
+					const TArray<int32>& SelectionArr = TransformSelectionFacade.ConvertVertexSelectionToCurveSelection(InValue.AsArray(), bAllElementsMustBeSelected);
+
+					FDataflowCurveSelection NewSelection;
+					NewSelection.InitFromArray(InCollection, SelectionArr);
+
 					SetValue(Context, MoveTemp(NewSelection), &Out);
 				}
 			}
@@ -304,6 +331,54 @@ void FConvertSelectionTypesDataflowNode::Evaluate(UE::Dataflow::FContext& Contex
 					FDataflowVertexSelection NewSelection;
 					NewSelection.InitFromArray(InCollection, SelectionArr);
 
+					SetValue(Context, MoveTemp(NewSelection), &Out);
+				}
+				else if (Output && Output->IsType<FDataflowCurveSelection>())
+				{
+					const TArray<int32>& SelectionArr = TransformSelectionFacade.ConvertGeometrySelectionToCurveSelection(InValue.AsArray());
+
+					FDataflowCurveSelection NewSelection;
+					NewSelection.InitFromArray(InCollection, SelectionArr);
+
+					SetValue(Context, MoveTemp(NewSelection), &Out);
+				}
+			}
+			else if (SelectionInput->IsType<FDataflowCurveSelection>())
+			{
+				if (Output && Output->IsType<FDataflowTransformSelection>())
+				{
+					const TArray<int32>& SelectionArr = TransformSelectionFacade.ConvertCurveSelectionToTransformSelection(InValue.AsArray(), bAllElementsMustBeSelected);
+
+					FDataflowTransformSelection NewSelection;
+					NewSelection.InitFromArray(InCollection, SelectionArr);
+
+					SetValue(Context, MoveTemp(NewSelection), &Out);
+				}
+				else if (Output && Output->IsType<FDataflowFaceSelection>())
+				{
+					const TArray<int32>& SelectionArr = TransformSelectionFacade.ConvertCurveSelectionToFaceSelection(InValue.AsArray());
+
+					FDataflowFaceSelection NewSelection;
+					NewSelection.InitFromArray(InCollection, SelectionArr);
+
+					SetValue(Context, MoveTemp(NewSelection), &Out);
+				}
+				else if (Output && Output->IsType<FDataflowVertexSelection>())
+				{
+					const TArray<int32>& SelectionArr = TransformSelectionFacade.ConvertCurveSelectionToVertexSelection(InValue.AsArray());
+
+					FDataflowVertexSelection NewSelection;
+					NewSelection.InitFromArray(InCollection, SelectionArr);
+
+					SetValue(Context, MoveTemp(NewSelection), &Out);
+				}
+				else if (Output && Output->IsType<FDataflowGeometrySelection>())
+				{
+					const TArray<int32>& SelectionArr = TransformSelectionFacade.ConvertCurveSelectionToGeometrySelection(InValue.AsArray(), bAllElementsMustBeSelected);
+
+					FDataflowGeometrySelection NewSelection;
+					NewSelection.InitFromArray(InCollection, SelectionArr);
+
 					SetValue(Context, MoveTemp(NewSelection), &Out);
 				}
 			}
diff --git a/Engine/Source/Runtime/Experimental/Dataflow/Core/Private/Dataflow/DataflowCoreNodeAndPinTypeColors.cpp b/Engine/Source/Runtime/Experimental/Dataflow/Core/Private/Dataflow/DataflowCoreNodeAndPinTypeColors.cpp
--- a/Engine/Source/Runtime/Experimental/Dataflow/Core/Private/Dataflow/DataflowCoreNodeAndPinTypeColors.cpp
+++ b/Engine/Source/Runtime/Experimental/Dataflow/Core/Private/Dataflow/DataflowCoreNodeAndPinTypeColors.cpp
@@ -56,6 +56,7 @@ namespace UE::Dataflow
 		DATAFLOW_NODE_REGISTER_CREATION_FACTORY_PIN_SETTINGS_BY_TYPE("FDataflowVertexSelection", FLinearColor(0.83077f, 0.361307f, 0.003035f, 1.0f), CDefaultWireThickness);
 		DATAFLOW_NODE_REGISTER_CREATION_FACTORY_PIN_SETTINGS_BY_TYPE("FDataflowFaceSelection", FLinearColor(0.83077f, 0.361307f, 0.003035f, 1.0f), CDefaultWireThickness);
 		DATAFLOW_NODE_REGISTER_CREATION_FACTORY_PIN_SETTINGS_BY_TYPE("FDataflowGeometrySelection", FLinearColor(0.83077f, 0.361307f, 0.003035f, 1.0f), CDefaultWireThickness);
+		DATAFLOW_NODE_REGISTER_CREATION_FACTORY_PIN_SETTINGS_BY_TYPE("FDataflowCurvesSelection", FLinearColor(0.83077f, 0.361307f, 0.003035f, 1.0f), CDefaultWireThickness);
 		DATAFLOW_NODE_REGISTER_CREATION_FACTORY_PIN_SETTINGS_BY_TYPE("FDataflowMaterialSelection", FLinearColor(0.3f, 0.3f, 0.3f, 1.0f), CDefaultWireThickness);
 
 		DATAFLOW_NODE_REGISTER_CREATION_FACTORY_PIN_SETTINGS_BY_TYPE("FDataflowAnyType", FLinearColor(0.220000f, 0.195800f, 0.195800f, 1.0f), CDefaultWireThickness);
diff --git a/Engine/Source/Runtime/Experimental/Dataflow/Core/Public/Dataflow/DataflowSelection.h b/Engine/Source/Runtime/Experimental/Dataflow/Core/Public/Dataflow/DataflowSelection.h
--- a/Engine/Source/Runtime/Experimental/Dataflow/Core/Public/Dataflow/DataflowSelection.h
+++ b/Engine/Source/Runtime/Experimental/Dataflow/Core/Public/Dataflow/DataflowSelection.h
@@ -142,6 +142,18 @@ struct FDataflowMaterialSelection : public FDataflowSelection
 	{}
 };
 
+USTRUCT()
+struct FDataflowCurveSelection : public FDataflowSelection
+{
+	GENERATED_USTRUCT_BODY()
+
+	inline static const FName CurveGroupName = "Curves";
+
+	FDataflowCurveSelection()
+		: FDataflowSelection(CurveGroupName)
+	{}
+};
+
 UENUM(BlueprintType)
 enum class EDataflowSelectionType : uint8
 {
@@ -149,5 +161,6 @@ enum class EDataflowSelectionType : uint8
 	Vertices UMETA(DisplayName = "Vertices"),
 	Faces UMETA(DisplayName = "Faces"),
 	Geometry UMETA(DisplayName = "Geometry"),
+	Curves UMETA(DisplayName = "Curves"),
 };
 
diff --git a/Engine/Source/Runtime/Experimental/Dataflow/Core/Public/Dataflow/DataflowTypePolicy.h b/Engine/Source/Runtime/Experimental/Dataflow/Core/Public/Dataflow/DataflowTypePolicy.h
--- a/Engine/Source/Runtime/Experimental/Dataflow/Core/Public/Dataflow/DataflowTypePolicy.h
+++ b/Engine/Source/Runtime/Experimental/Dataflow/Core/Public/Dataflow/DataflowTypePolicy.h
@@ -100,7 +100,7 @@ UE_DATAFLOW_POLICY_DECLARE_TYPENAME(FDataflowVertexSelection)
 UE_DATAFLOW_POLICY_DECLARE_TYPENAME(FDataflowFaceSelection)
 UE_DATAFLOW_POLICY_DECLARE_TYPENAME(FDataflowGeometrySelection)
 UE_DATAFLOW_POLICY_DECLARE_TYPENAME(FDataflowMaterialSelection)
-
+UE_DATAFLOW_POLICY_DECLARE_TYPENAME(FDataflowCurveSelection)
 
 struct IDataflowTypePolicy
 {
@@ -287,7 +287,7 @@ struct FDataflowStringTypePolicy :
 };
 
 struct FDataflowSelectionTypePolicy :
-	public TDataflowMultiTypePolicy<FDataflowTransformSelection, FDataflowVertexSelection, FDataflowFaceSelection, FDataflowGeometrySelection, FDataflowMaterialSelection>
+	public TDataflowMultiTypePolicy<FDataflowTransformSelection, FDataflowVertexSelection, FDataflowFaceSelection, FDataflowGeometrySelection, FDataflowMaterialSelection, FDataflowCurveSelection>
 {
 };
 
