Fix CIS for curve/boneweights facades

diff --git a/Engine/Source/Runtime/Experimental/Chaos/Private/GeometryCollection/Facades/CollectionCurveFacade.cpp b/Engine/Source/Runtime/Experimental/Chaos/Private/GeometryCollection/Facades/CollectionCurveFacade.cpp
--- a/Engine/Source/Runtime/Experimental/Chaos/Private/GeometryCollection/Facades/CollectionCurveFacade.cpp
+++ b/Engine/Source/Runtime/Experimental/Chaos/Private/GeometryCollection/Facades/CollectionCurveFacade.cpp
@@ -1,4 +1,4 @@
-﻿// Copyright Epic Games, Inc. All Rights Reserved.
+// Copyright Epic Games, Inc. All Rights Reserved.
 
 #include "GeometryCollection/Facades/CollectionCurveFacade.h"
 #include "GeometryCollection/GeometryCollection.h"
@@ -57,7 +57,7 @@ namespace UE::GeometryCollection::Facades
 	{
 		return PointRestOrientations.IsValid() && PointRestPositions.IsValid() && 
 			CurvePointOffsets.IsValid() && GeometryCurveOffsets.IsValid() && PointCurveIndices.IsValid() &&
-					CurveGeometryIndices.IsValid() && PointRestOrientations.IsValid() && GeometryGroupNames.IsValid() &&
+					CurveGeometryIndices.IsValid() && GeometryGroupNames.IsValid() &&
 						GeometryCurveThickness.IsValid() && CurveSourceIndices.IsValid();
 	}
 	
diff --git a/Engine/Source/Runtime/Experimental/Chaos/Private/GeometryCollection/Facades/CollectionVertexBoneWeightsFacade.cpp b/Engine/Source/Runtime/Experimental/Chaos/Private/GeometryCollection/Facades/CollectionVertexBoneWeightsFacade.cpp
--- a/Engine/Source/Runtime/Experimental/Chaos/Private/GeometryCollection/Facades/CollectionVertexBoneWeightsFacade.cpp
+++ b/Engine/Source/Runtime/Experimental/Chaos/Private/GeometryCollection/Facades/CollectionVertexBoneWeightsFacade.cpp
@@ -195,11 +195,6 @@ namespace GeometryCollection::Facades
 		{
 			return KinematicWeightAttribute.IsValidIndex(VertexIndex) && (KinematicWeightAttribute.Get()[VertexIndex] != 1.0);
 		}
-		else if(ConstCollection.HasAttribute(DeprecatedKinematicFlagAttributeName, FGeometryCollection::VerticesGroup))
-		{
-			const TManagedArray<bool>& KinematicFlag = ConstCollection.GetAttribute<bool>(DeprecatedKinematicFlagAttributeName, FGeometryCollection::VerticesGroup);
-			return KinematicFlag.IsValidIndex(VertexIndex) && KinematicFlag.GetConstArray()[VertexIndex];
-		}
 		else //backward compatibility for KinematicAttribute added in 5.5
 		{
 			return BoneIndicesAttribute.IsValid() && BoneIndicesAttribute.IsValidIndex(VertexIndex) && BoneIndicesAttribute.Get()[VertexIndex].Num()
