[Backout] - CL45497446 - backout due to Main CIS issue

diff --git a/Engine/Plugins/Runtime/HairStrands/Source/HairStrandsDataflow/Public/SmoothGuidesCurvesNode.h b/Engine/Plugins/Runtime/HairStrands/Source/HairStrandsDataflow/Public/SmoothGuidesCurvesNode.h
--- a/Engine/Plugins/Runtime/HairStrands/Source/HairStrandsDataflow/Public/SmoothGuidesCurvesNode.h
+++ b/Engine/Plugins/Runtime/HairStrands/Source/HairStrandsDataflow/Public/SmoothGuidesCurvesNode.h
@@ -12,8 +12,8 @@
 
 
 /** Smooth the guides for more stable simulation */
-USTRUCT(meta = (Experimental, DataflowGroom, Deprecated="5.7"))
-struct UE_DEPRECATED(5.7, "Use the newer version of this node instead.") FSmoothGuidesCurvesDataflowNode : public FDataflowNode
+USTRUCT(meta = (Experimental, DataflowGroom))
+struct FSmoothGuidesCurvesDataflowNode : public FDataflowNode
 {
 	GENERATED_BODY()
 	
@@ -42,34 +42,3 @@ public:
 	float SmoothingFactor = 0.0f;
 };
 
-/** Smooth the guides for more stable simulation */
-USTRUCT(meta = (Experimental, DataflowGroom))
-struct FSmoothCurvePointsDataflowNode : public FDataflowNode
-{
-	GENERATED_BODY()
-	
-	DATAFLOW_NODE_DEFINE_INTERNAL(FSmoothCurvePointsDataflowNode, "SmoothCurvePoints", "Groom", "")
-
-public:
-	
-	FSmoothCurvePointsDataflowNode(const UE::Dataflow::FNodeParameters& InParam, FGuid InGuid = FGuid::NewGuid());
-
-private:
-
-	//~ Begin FDataflowNode interface
-	virtual void Evaluate(UE::Dataflow::FContext& Context, const FDataflowOutput* Out) const override;
-	//~ End FDataflowNode interface
-
-	/** Managed array collection to be used to store data */
-	UPROPERTY(meta = (DataflowInput, DataflowOutput, DisplayName = "Collection", DataflowPassthrough  = "Collection", DataflowRenderGroups = "Surface"))
-	FManagedArrayCollection Collection;
-
-	/** Curve selection to focus the guides smoothing spatially */
-	UPROPERTY(meta = (DataflowInput, DisplayName = "Curve Selection"))
-	FDataflowCurveSelection CurveSelection;
-
-	/** Smoothing factor between 0 and 1 */
-	UPROPERTY(EditAnywhere, Category="Geometry", meta = (DataflowInput, DisplayName = "Smoothing Factor", ClampMin="0", ClampMax="1", UIMin="0", UIMax="1"))
-	float SmoothingFactor = 0.0f;
-};
-
diff --git a/Engine/Plugins/Runtime/HairStrands/Source/HairStrandsEditor/Private/HairStrandsEditor.cpp b/Engine/Plugins/Runtime/HairStrands/Source/HairStrandsEditor/Private/HairStrandsEditor.cpp
--- a/Engine/Plugins/Runtime/HairStrands/Source/HairStrandsEditor/Private/HairStrandsEditor.cpp
+++ b/Engine/Plugins/Runtime/HairStrands/Source/HairStrandsEditor/Private/HairStrandsEditor.cpp
@@ -158,7 +158,7 @@ void FGroomEditor::StartupModule()
 
 	// default no dataflow template
 	UE::Groom::RegisterGroomDataflowTemplate({ "SelectNoDataflowTemplate", "No Dataflow", "Don't add a Dataflow to this Groom Asset. Useful when a Dataflow already exists.", "", false });
-	UE::Groom::RegisterGroomDataflowTemplatePath(TEXT("/HairStrands/Dataflow/Templates"));
+	UE::Groom::RegisterGroomDataflowTemplatePath(TEXT("/HairStrands/Dataflow"));
 }
 
 void FGroomEditor::ShutdownModule()
diff --git a/Engine/Plugins/Runtime/HairStrands/Source/HairStrandsEditor/Private/HairStrandsFactory.cpp b/Engine/Plugins/Runtime/HairStrands/Source/HairStrandsEditor/Private/HairStrandsFactory.cpp
--- a/Engine/Plugins/Runtime/HairStrands/Source/HairStrandsEditor/Private/HairStrandsFactory.cpp
+++ b/Engine/Plugins/Runtime/HairStrands/Source/HairStrandsEditor/Private/HairStrandsFactory.cpp
@@ -141,11 +141,6 @@ namespace UE::Groom
 		TArray<FGroomDataflowTemplateData>	TemplateData;
 		TArray<FString> DynamicallyLoadedTemplatePaths;
 	};
-	
-	TArray<SMessageDialog::FButton> BuildGroomDataflowTemplateButtons(TFunction<void(FString)> OnButtonClicked)
-	{
-		return FGroomDataflowTemplateManager::GetInstance().BuildDataflowTemplateButtons(OnButtonClicked);
-	}
 
 	void RegisterGroomDataflowTemplate(const FGroomDataflowTemplateData& DataflowTemplate)
 	{
diff --git a/Engine/Plugins/Runtime/HairStrands/Source/HairStrandsEditor/Public/HairStrandsFactory.h b/Engine/Plugins/Runtime/HairStrands/Source/HairStrandsEditor/Public/HairStrandsFactory.h
--- a/Engine/Plugins/Runtime/HairStrands/Source/HairStrandsEditor/Public/HairStrandsFactory.h
+++ b/Engine/Plugins/Runtime/HairStrands/Source/HairStrandsEditor/Public/HairStrandsFactory.h
@@ -2,7 +2,6 @@
 
 #pragma once
 
-#include "Dialog/SMessageDialog.h"
 #include "Factories/Factory.h"
 #include "UObject/ObjectMacros.h"
 
@@ -31,9 +30,6 @@ namespace UE::Groom
 		bool bIsPrimaryTemplate = false;
 	};
 
-	/** Build dataflow template buttons */
-	HAIRSTRANDSEDITOR_API TArray<SMessageDialog::FButton> BuildGroomDataflowTemplateButtons(TFunction<void(FString)> OnButtonClicked);
-
 	/** Register dataflow template to the manager */
 	HAIRSTRANDSEDITOR_API void RegisterGroomDataflowTemplate(const FGroomDataflowTemplateData& DataflowTemplate);
 
diff --git a/Engine/Source/Runtime/Experimental/Chaos/Private/GeometryCollection/Facades/CollectionCurveFacade.cpp b/Engine/Source/Runtime/Experimental/Chaos/Private/GeometryCollection/Facades/CollectionCurveFacade.cpp
--- a/Engine/Source/Runtime/Experimental/Chaos/Private/GeometryCollection/Facades/CollectionCurveFacade.cpp
+++ b/Engine/Source/Runtime/Experimental/Chaos/Private/GeometryCollection/Facades/CollectionCurveFacade.cpp
@@ -3,7 +3,7 @@
 #include "GeometryCollection/Facades/CollectionCurveFacade.h"
 #include "GeometryCollection/GeometryCollection.h"
 
-namespace GeometryCollection::Facades
+namespace UE::GeometryCollection::Facades
 {
 	// Curves facade  Groups
 	const FName FCollectionCurveGeometryFacade::CurvesGroup("Curves");
diff --git a/Engine/Source/Runtime/Experimental/Chaos/Private/GeometryCollection/Facades/CollectionTransformSelectionFacade.cpp b/Engine/Source/Runtime/Experimental/Chaos/Private/GeometryCollection/Facades/CollectionTransformSelectionFacade.cpp
--- a/Engine/Source/Runtime/Experimental/Chaos/Private/GeometryCollection/Facades/CollectionTransformSelectionFacade.cpp
+++ b/Engine/Source/Runtime/Experimental/Chaos/Private/GeometryCollection/Facades/CollectionTransformSelectionFacade.cpp
@@ -1609,7 +1609,7 @@ namespace GeometryCollection::Facades
 	{
 		TArray<int32> OutSelection;
 
-		GeometryCollection::Facades::FCollectionCurveGeometryFacade CurveFacade(ConstCollection);
+		UE::GeometryCollection::Facades::FCollectionCurveGeometryFacade CurveFacade(ConstCollection);
 		if (TransformToGeometryIndexAttribute.IsValid() && CurveFacade.IsValid())
 		{
 			const TManagedArray<int32>& TransformToGeometryIndices = TransformToGeometryIndexAttribute.Get();
@@ -1641,7 +1641,7 @@ namespace GeometryCollection::Facades
 	{
 		TArray<int32> OutSelection;
 
-		GeometryCollection::Facades::FCollectionCurveGeometryFacade CurveFacade(ConstCollection);
+		UE::GeometryCollection::Facades::FCollectionCurveGeometryFacade CurveFacade(ConstCollection);
 		if (CurveFacade.IsValid())
 		{
 			const TArray<int32>& CurveOffsets = CurveFacade.GetGeometryCurveOffsets();
@@ -1667,7 +1667,7 @@ namespace GeometryCollection::Facades
 	{
 		TArray<int32> OutSelection;
 
-		GeometryCollection::Facades::FCollectionCurveGeometryFacade CurveFacade(ConstCollection);
+		UE::GeometryCollection::Facades::FCollectionCurveGeometryFacade CurveFacade(ConstCollection);
 		if (CurveFacade.IsValid())
 		{
 			const int32 NumCurves = CurveFacade.GetNumCurves();
@@ -1713,7 +1713,7 @@ namespace GeometryCollection::Facades
 	{
 		TArray<int32> OutSelection;
 
-		GeometryCollection::Facades::FCollectionCurveGeometryFacade CurveFacade(ConstCollection);
+		UE::GeometryCollection::Facades::FCollectionCurveGeometryFacade CurveFacade(ConstCollection);
 		if (CurveFacade.IsValid())
 		{
 			const int32 NumCurves = CurveFacade.GetNumCurves();
@@ -1759,7 +1759,7 @@ namespace GeometryCollection::Facades
 	{
 		TArray<int32> OutSelection;
 
-		GeometryCollection::Facades::FCollectionCurveGeometryFacade CurveFacade(ConstCollection);
+		UE::GeometryCollection::Facades::FCollectionCurveGeometryFacade CurveFacade(ConstCollection);
 		if (CurveFacade.IsValid())
 		{
 			const int32 NumCurves = CurveFacade.GetNumCurves();
@@ -1785,7 +1785,7 @@ namespace GeometryCollection::Facades
 	{
 		TArray<int32> OutSelection;
 
-		GeometryCollection::Facades::FCollectionCurveGeometryFacade CurveFacade(ConstCollection);
+		UE::GeometryCollection::Facades::FCollectionCurveGeometryFacade CurveFacade(ConstCollection);
 		if (CurveFacade.IsValid())
 		{
 			const int32 NumCurves = CurveFacade.GetNumCurves();
@@ -1811,7 +1811,7 @@ namespace GeometryCollection::Facades
 	{
 		TArray<int32> OutSelection;
 
-		GeometryCollection::Facades::FCollectionCurveGeometryFacade CurveFacade(ConstCollection);
+		UE::GeometryCollection::Facades::FCollectionCurveGeometryFacade CurveFacade(ConstCollection);
 		if (CurveFacade.IsValid())
 		{
 			const TArray<int32>& CurveOffsets = CurveFacade.GetGeometryCurveOffsets();
@@ -1857,7 +1857,7 @@ namespace GeometryCollection::Facades
 	{
 		TArray<int32> OutSelection;
 
-		GeometryCollection::Facades::FCollectionCurveGeometryFacade CurveFacade(ConstCollection);
+		UE::GeometryCollection::Facades::FCollectionCurveGeometryFacade CurveFacade(ConstCollection);
 		if (CurveFacade.IsValid() && TransformIndexAttribute.IsValid())
 		{
 			const TManagedArray<int32>& TransformIndices = TransformIndexAttribute.Get();
diff --git a/Engine/Source/Runtime/Experimental/Chaos/Public/GeometryCollection/Facades/CollectionCurveFacade.h b/Engine/Source/Runtime/Experimental/Chaos/Public/GeometryCollection/Facades/CollectionCurveFacade.h
--- a/Engine/Source/Runtime/Experimental/Chaos/Public/GeometryCollection/Facades/CollectionCurveFacade.h
+++ b/Engine/Source/Runtime/Experimental/Chaos/Public/GeometryCollection/Facades/CollectionCurveFacade.h
@@ -1,4 +1,4 @@
-// Copyright Epic Games, Inc. All Rights Reserved.
+﻿// Copyright Epic Games, Inc. All Rights Reserved.
 
 #pragma once
 
@@ -6,7 +6,7 @@
 #include "GeometryCollection/ManagedArrayCollection.h"
 #include "GeometryCollection/GeometryCollection.h"
 	
-namespace GeometryCollection::Facades
+namespace UE::GeometryCollection::Facades
 {
 	class FCollectionCurveGeometryFacade
 	{
