Dataflow : make variable work in SubGraphs
- changed the subgraph context to be an engine context

diff --git a/Engine/Source/Runtime/Experimental/Dataflow/Engine/Private/Dataflow/DataflowSubGraphNodes.cpp b/Engine/Source/Runtime/Experimental/Dataflow/Engine/Private/Dataflow/DataflowSubGraphNodes.cpp
--- a/Engine/Source/Runtime/Experimental/Dataflow/Engine/Private/Dataflow/DataflowSubGraphNodes.cpp
+++ b/Engine/Source/Runtime/Experimental/Dataflow/Engine/Private/Dataflow/DataflowSubGraphNodes.cpp
@@ -5,6 +5,7 @@
 #include "Dataflow/DataflowNode.h"
 #include "Dataflow/DataflowEdNode.h"
 #include "Dataflow/DataflowObject.h"
+#include "Dataflow/DataflowObjectInterface.h"
 #include "Dataflow/DataflowSubGraph.h"
 #include "Dataflow/DataflowTypePolicy.h"
 #include "Dataflow/DataflowNodeFactory.h"
@@ -26,13 +27,14 @@ namespace UE::Dataflow::SubGraphNodes::Private
 		return NewName;
 	}
 
-	class FContextSubGraph: public FContext
+	class FContextSubGraph: public FEngineContext
 	{
 	public:
-		DATAFLOW_CONTEXT_INTERNAL(FContext, FContextSubGraph);
+		DATAFLOW_CONTEXT_INTERNAL(FEngineContext, FContextSubGraph);
 
-		FContextSubGraph(FContext& InParentContext, const ISubGraphContextCallback& InSubGraphCallback)
-			: ParentContext(InParentContext)
+		FContextSubGraph(FContext& InParentContext, UObject* DataflowAsset, const ISubGraphContextCallback& InSubGraphCallback)
+			: FEngineContext(DataflowAsset)
+			, ParentContext(InParentContext)
 			, SubGraphCallback(InSubGraphCallback)
 		{
 			// make sure to match the threaded mode
@@ -696,8 +698,14 @@ void FDataflowCallSubGraphNode::Evaluate(UE::Dataflow::FContext& Context, const
 			{
 				if (const FDataflowSubGraphOutputNode* SubGraphOutputNode = SubGraphToCall->GetOutputNode())
 				{
+					UObject* ContextOwner = nullptr;
+					if (const UE::Dataflow::FEngineContext* EngineContext = Context.AsType<UE::Dataflow::FEngineContext>())
+					{
+						ContextOwner = EngineContext->Owner;
+					}
+
 					// pulling the outputs and when the graph evaluate the remote left end of its graph OnInputValueRequested will be called
-					UE::Dataflow::SubGraphNodes::Private::FContextSubGraph SubGraphContext(Context, *this);
+					UE::Dataflow::SubGraphNodes::Private::FContextSubGraph SubGraphContext(Context, ContextOwner, *this);
 					int32 NumIterations = 1;
 					if (SubGraphToCall->IsForEachSubGraph())
 					{
diff --git a/Engine/Source/Runtime/Experimental/Dataflow/Engine/Private/Dataflow/DataflowVariableNodes.cpp b/Engine/Source/Runtime/Experimental/Dataflow/Engine/Private/Dataflow/DataflowVariableNodes.cpp
--- a/Engine/Source/Runtime/Experimental/Dataflow/Engine/Private/Dataflow/DataflowVariableNodes.cpp
+++ b/Engine/Source/Runtime/Experimental/Dataflow/Engine/Private/Dataflow/DataflowVariableNodes.cpp
@@ -279,6 +279,10 @@ void FGetDataflowVariableNode::Evaluate(UE::Dataflow::FContext& Context, const F
 				}
 			}
 		}
+		
+		const FString Message = FString::Printf(TEXT("Dataflow Variable : Failed to evaluate variable [%s], returning defaut value"), *VariableName.ToString());
+		Context.Warning(Message, this, Out);
+
 		// nothing worked, just write a null value that will read as default value
 		Out->SetNullValue(Context);
 	}
