Improve the 'ias.RunRequestTest' development command.

#rnx

- The command now accepts a 0 and negative value for the number of requests to run. These values will be treated as "attempt to download everything available"
- The command accepts a new arg 'cancel' which will cause the batch of requests to be canceled after 5 seconds. Useful for testing the cancelation code paths.
- The command accepts a new arg 'quit' which will issue a request for the current process to shutdown once the batch has completed. This is useful for testing shutdown code paths.

diff --git a/Engine/Source/Runtime/Experimental/IoStore/OnDemand/Private/DebugCommands.cpp b/Engine/Source/Runtime/Experimental/IoStore/OnDemand/Private/DebugCommands.cpp
--- a/Engine/Source/Runtime/Experimental/IoStore/OnDemand/Private/DebugCommands.cpp
+++ b/Engine/Source/Runtime/Experimental/IoStore/OnDemand/Private/DebugCommands.cpp
@@ -64,43 +64,87 @@ void FOnDemandDebugCommands::UnbindConsoleCommands()
 
 void FOnDemandDebugCommands::RunRequestTest(const TArray<FString>& Args) const
 {
-	int32 NumToRequest = 1000;
+	int32 NumToRequest = -1;
+	const float CancelWaitTime = 5.0f;
+	bool bCancelRequests = false;
+	bool bQuitAfterTest = false;
 
-	if (Args.Num() > 1)
-	{
-		UE_LOG(LogIas, Error, TEXT("Too many args for 'ias.RunRequestTest'"));
-		return;
-	}
-	else if (Args.Num() == 1)
-	{
-		int32 ArgValue = TCString<TCHAR>::Atoi(*Args[0]);
 
-		if (!UE::String::IsNumericOnlyDigits(Args[0]) || ArgValue <= 0)
+	for (const FString& Arg : Args)
+	{
+		if (UE::String::IsNumeric(Arg))
+		{
+			if (NumToRequest != -1)
+			{
+				UE_LOG(LogIas, Error, TEXT("Too many numeric args for 'ias.RunRequestTest'"));
+				return;
+			}
+
+			int32 ArgValue = TCString<TCHAR>::Atoi(*Args[0]);
+			NumToRequest = ArgValue > 0 ? ArgValue : MAX_int32;
+		}
+		else if (Arg == "Cancel")
 		{
-			UE_LOG(LogIas, Error, TEXT("Invalid arg for 'ias.RunRequestTest', value must be a postive integer"));
+			bCancelRequests = true;
+		}
+		else if (Arg == "Quit")
+		{
+			bQuitAfterTest = true;
+		}
+		else
+		{
+			UE_LOG(LogIas, Error, TEXT("Unknown arg '%s' for 'ias.RunRequestTest'"), *Arg);
 			return;
 		}
-
-		NumToRequest = ArgValue;
 	}
 
-	UE_LOG(LogIas, Display, TEXT("Running IAS Request Test with %d requests..."), NumToRequest);
+	// Set number of requests if it was not provided as an arg
+	NumToRequest = NumToRequest == -1 ? 1000 : NumToRequest;
 
 	TArray<FIoChunkId> RequestIds = OnDemandIoStore->DebugFindStreamingChunkIds(NumToRequest);
 
+	UE_LOG(LogIas, Display, TEXT("Running IAS Request Test with %d requests..."), RequestIds.Num());
+
+	TArray<FIoRequest> Requests;
+
 	FIoBatch IoBatch;
 	for (const FIoChunkId& Id : RequestIds)
 	{
-		IoBatch.Read(Id, FIoReadOptions(), IoDispatcherPriority_Medium);
+		FIoRequest Request = IoBatch.Read(Id, FIoReadOptions(), IoDispatcherPriority_Medium);
+		Requests.Emplace(MoveTemp(Request));
 	}
 
 	const double StartTime = FPlatformTime::Seconds();
 
-	IoBatch.IssueWithCallback([StartTime]()
+	IoBatch.IssueWithCallback([StartTime, bQuitAfterTest]()
 		{
 			const double TotalTime = FPlatformTime::Seconds() - StartTime;
 			UE_LOG(LogIas, Display, TEXT("IAS Request Test took %.3f(s)"), TotalTime);
+
+			if (bQuitAfterTest)
+			{
+				UE_LOG(LogIas, Display, TEXT("Quitting now that the request batch has completed"));
+				FPlatformMisc::RequestExitWithStatus(false, 0);
+			}
 		});
+
+	if (bCancelRequests)
+	{
+		UE_LOG(LogIas, Display, TEXT("Requests will be canceled in %f seconds"), CancelWaitTime);
+
+		FTSTicker::GetCoreTicker().AddTicker(TEXT("IASDebugCancel"), CancelWaitTime, [Requests = MoveTemp(Requests)](float DeltaTime) mutable
+			{
+				for (FIoRequest& Request : Requests)
+				{
+					Request.Cancel();
+				}
+
+				UE_LOG(LogIas, Display, TEXT("Canceled all requests!"));
+				return false;
+			});
+	}
+
+	UE_CLOG(bQuitAfterTest, LogIas, Display, TEXT("Process will quit once the requests have been completed"));
 }
 
 } // namespace UE::IoStore
