The command 'ias.RunRequestTest' will now correctly parse the number of requests to test, even if not provided as the first arg.

#rnx

diff --git a/Engine/Source/Runtime/Experimental/IoStore/OnDemand/Private/DebugCommands.cpp b/Engine/Source/Runtime/Experimental/IoStore/OnDemand/Private/DebugCommands.cpp
--- a/Engine/Source/Runtime/Experimental/IoStore/OnDemand/Private/DebugCommands.cpp
+++ b/Engine/Source/Runtime/Experimental/IoStore/OnDemand/Private/DebugCommands.cpp
@@ -80,7 +80,7 @@ void FOnDemandDebugCommands::RunRequestTest(const TArray<FString>& Args) const
 				return;
 			}
 
-			int32 ArgValue = TCString<TCHAR>::Atoi(*Args[0]);
+			int32 ArgValue = TCString<TCHAR>::Atoi(*Arg);
 			NumToRequest = ArgValue > 0 ? ArgValue : MAX_int32;
 		}
 		else if (Arg == "Cancel")
