Fix new IAX backend from potentially reporting more inflight http requests than are actually running.

#rnx

- The new dispatcher could include requests that were never actually sent to the http client (if the required hostgroup is broken for example) so in rare occasions this could cause the statistics system to report more requests currently inflight than GIasHttpConcurrentRequests which would look to the user like the limits were no longer working.
- Changed the storage value in the statistics system (GHttpNumConcurrentRequests) to be atomic since it will be updated on a different thread than it will be read and we don't want to trigger TSAN checks.

diff --git a/Engine/Source/Runtime/Experimental/IoStore/OnDemand/Private/OnDemandHttpIoDispacher.cpp b/Engine/Source/Runtime/Experimental/IoStore/OnDemand/Private/OnDemandHttpIoDispacher.cpp
--- a/Engine/Source/Runtime/Experimental/IoStore/OnDemand/Private/OnDemandHttpIoDispacher.cpp
+++ b/Engine/Source/Runtime/Experimental/IoStore/OnDemand/Private/OnDemandHttpIoDispacher.cpp
@@ -868,7 +868,6 @@ bool FHttpDispatcher::ProcessHttpRequests()
 			NextHttpRequest				= HttpQueue.Dequeue();
 
 			ensure(!HttpRequest->Flags.HasAny(EHttpRequestFlags::HttpQueued | EHttpRequestFlags::CacheInflight));
-			FOnDemandIoBackendStats::Get()->OnHttpRequestStarted();
 
 			FIASHostGroup HostGroup = FHostGroupManager::Get().Find(HttpRequest->HostGroupName);
 			if (HostGroup.IsConnected() == false)
@@ -905,11 +904,16 @@ bool FHttpDispatcher::ProcessHttpRequests()
 			}
 
 			NumConcurrentRequests++;
+			FOnDemandIoBackendStats::Get()->OnHttpRequestStarted();
+
+			// Note that the requests issued here must be completed before this method can return so the capturing of locals is safe.
 			HttpClient->Get(HostGroup.GetUnderlyingHostGroup(), Url, TotalRange.ToOffsetAndLength(), MoveTemp(Headers),
 				[this, HttpRequest, TotalRange, &NumConcurrentRequests, HostGroup = MoveTemp(HostGroup)]
 				(FMultiEndpointHttpClientResponse&& HttpResponse) mutable
 				{
+					FOnDemandIoBackendStats::Get()->OnHttpRequestCompleted();
 					NumConcurrentRequests--;
+
 					UE::Tasks::Launch(
 						UE_SOURCE_LOCATION,
 						[this, HttpRequest, HttpResponse = MoveTemp(HttpResponse), TotalRange, HostGroup = MoveTemp(HostGroup)]() mutable
@@ -951,13 +955,14 @@ bool FHttpDispatcher::ProcessHttpRequests()
 		}
 	} // Outer
 
+	// All requests from HttpClient should be completed at this point and all callbacks processed!.
+	check(NumConcurrentRequests == 0);
+
 	return true;
 }
 
 void FHttpDispatcher::CompleteHttpRequest(FHttpRequest& Request, FMultiEndpointHttpClientResponse&& HttpResponse, FIASHostGroup& HostGroup, const FIoHttpRange& TotalRange)
 {
-	FOnDemandIoBackendStats::Get()->OnHttpRequestCompleted();
-
 	if (HttpResponse.RetryCount > 0)
 	{
 		FOnDemandIoBackendStats::Get()->OnHttpRetry(Request.GetHttpRequestType());
diff --git a/Engine/Source/Runtime/Experimental/IoStore/OnDemand/Private/Statistics.cpp b/Engine/Source/Runtime/Experimental/IoStore/OnDemand/Private/Statistics.cpp
--- a/Engine/Source/Runtime/Experimental/IoStore/OnDemand/Private/Statistics.cpp
+++ b/Engine/Source/Runtime/Experimental/IoStore/OnDemand/Private/Statistics.cpp
@@ -384,7 +384,7 @@ UE_IAX_COUNTER(HttpBandwidthMbps, FCounterFloat);
 double GHttpDurationMsAvg[(uint8)EHttpRequestType::NUM_SOURCES] = { 0.0 };
 int32 GHttpDurationMsMax[(uint8)EHttpRequestType::NUM_SOURCES] = { 0 };
 int64 GHttpDurationMsSum[(uint8)EHttpRequestType::NUM_SOURCES] = { 0 };
-int32 GHttpNumConcurrentRequests = 0;
+std::atomic<int32> GHttpNumConcurrentRequests = 0;
 
 struct FHttpRecentHistoryStatistics
 {
@@ -781,8 +781,9 @@ void FOnDemandIoBackendStats::PrintOnScreenStatistics(FCoreDelegates::FSeverityM
 				Builder << TEXT(" DataRate Cap: None");
 			}
 
-			Builder << TEXT(" ConcurrentRequests: ") << GHttpNumConcurrentRequests << TEXT("/") << GIasHttpConcurrentRequests;
-			if (GHttpNumConcurrentRequests >= GIasHttpConcurrentRequests)
+			const int32 NumConcurrentRequests = GHttpNumConcurrentRequests;
+			Builder << TEXT(" ConcurrentRequests: ") << NumConcurrentRequests << TEXT("/") << GIasHttpConcurrentRequests;
+			if (NumConcurrentRequests >= GIasHttpConcurrentRequests)
 			{
 				Builder << TEXT("(saturated)");
 			}
