Improve code documentation for FMultiEndpointHttpClient::Tick

#rnx

diff --git a/Engine/Source/Runtime/Experimental/IoStore/OnDemand/Private/OnDemandHttpClient.h b/Engine/Source/Runtime/Experimental/IoStore/OnDemand/Private/OnDemandHttpClient.h
--- a/Engine/Source/Runtime/Experimental/IoStore/OnDemand/Private/OnDemandHttpClient.h
+++ b/Engine/Source/Runtime/Experimental/IoStore/OnDemand/Private/OnDemandHttpClient.h
@@ -124,6 +124,15 @@ public:
 	FHttpTicketId Get(const FOnDemandHostGroup& HostGroup, FAnsiStringView RelativeUrl, const FIoOffsetAndLength& ChunkRange, FOnHttpResponse&& OnResponse);
 	FHttpTicketId Get(const FOnDemandHostGroup& HostGroup, FAnsiStringView RelativeUrl, const FIoOffsetAndLength& ChunkRange, TArray<FAnsiString>&& Headers, FOnHttpResponse&& OnResponse);
 
+	/**
+	 * Process the underlying http client.
+	 * 
+	 * @param WaitTimeMs		How long (in milliseconds) should the underlying http client wait for poll events before returning.
+	 * @param MaxKiBPerSecond	The max amount of bandwidth that the underlying http client should use per second. A value of
+	 *							zero indicates that the bandwidth use should not be capped.
+	 * 
+	 * @return True if the client still has work pending and more Tick calls will be needed, false if all work has been completed.
+	 */
 	bool Tick(int32 WaitTimeMs, uint32 MaxKiBPerSecond);
 	bool Tick()
 	{
