The new IAX backend now implements the development cvar 'iax.InvalidUrlChance' to simulate invalid urls.

#rnx

diff --git a/Engine/Source/Runtime/Experimental/IoStore/OnDemand/Private/OnDemandHttpIoDispacher.cpp b/Engine/Source/Runtime/Experimental/IoStore/OnDemand/Private/OnDemandHttpIoDispacher.cpp
--- a/Engine/Source/Runtime/Experimental/IoStore/OnDemand/Private/OnDemandHttpIoDispacher.cpp
+++ b/Engine/Source/Runtime/Experimental/IoStore/OnDemand/Private/OnDemandHttpIoDispacher.cpp
@@ -26,11 +26,17 @@
 
 #include <atomic>
 
+// When enabled the cvar 'iax.InvalidUrlChance' will be enabled and allow us to simulate invalid urls
+#define UE_ALLOW_INVALID_URL_DEBUGGING !UE_BUILD_SHIPPING
+
 namespace UE::IoStore
 {
 ////////////////////////////////////////////////////////////////////////////////
+
+
 extern int32 GIasHttpTimeOutMs;
 extern int32 GOnDemandBackendThreadPriorityIndex; 
+extern float GIaxInvalidUrlChance;
 
 bool GHttpCacheEnabled = true;
 static FAutoConsoleVariableRef CVar_HttpCacheEnabled(
@@ -903,6 +909,14 @@ bool FHttpDispatcher::ProcessHttpRequests()
 				//DebugPrintRange(HttpRequest, TotalRange);
 			}
 
+#if UE_ALLOW_INVALID_URL_DEBUGGING
+			// Avoid the rand call if there is no chance
+			if (GIaxInvalidUrlChance > 0.0 && (FMath::FRand() * 100.0f) < GIaxInvalidUrlChance)
+			{
+				Url << "-DebugInvalidUrl";
+			}
+#endif // UE_ALLOW_INVALID_URL_DEBUGGING
+
 			NumConcurrentRequests++;
 			FOnDemandIoBackendStats::Get()->OnHttpRequestStarted();
 
@@ -1060,3 +1074,5 @@ TSharedPtr<IOnDemandHttpIoDispatcher> MakeOnDemanHttpIoDispatcher(TUniquePtr<cla
 }
 
 } // UE::IoStore
+
+#undef UE_ALLOW_INVALID_URL_DEBUGGING
