The new IAX backend now supports the development cvar 'ias.PoisonCache'

#rnx

- Changed the cvar from int32 to bool which I am fairly sure was the original intent.

diff --git a/Engine/Source/Runtime/Experimental/IoStore/OnDemand/Private/OnDemandHttpIoDispacher.cpp b/Engine/Source/Runtime/Experimental/IoStore/OnDemand/Private/OnDemandHttpIoDispacher.cpp
--- a/Engine/Source/Runtime/Experimental/IoStore/OnDemand/Private/OnDemandHttpIoDispacher.cpp
+++ b/Engine/Source/Runtime/Experimental/IoStore/OnDemand/Private/OnDemandHttpIoDispacher.cpp
@@ -29,6 +29,8 @@
 // When enabled the cvar 'iax.InvalidUrlChance' will be enabled and allow us to simulate invalid urls
 #define UE_ALLOW_INVALID_URL_DEBUGGING !UE_BUILD_SHIPPING
 
+#define UE_ALLOW_CACHE_POISONING !UE_BUILD_SHIPPING
+
 namespace UE::IoStore
 {
 ////////////////////////////////////////////////////////////////////////////////
@@ -37,6 +39,9 @@ namespace UE::IoStore
 extern int32 GIasHttpTimeOutMs;
 extern int32 GOnDemandBackendThreadPriorityIndex; 
 extern float GIaxInvalidUrlChance;
+#if UE_ALLOW_CACHE_POISONING
+extern bool GIasPoisonCache;
+#endif // UE_ALLOW_CACHE_POISONING
 
 bool GHttpCacheEnabled = true;
 static FAutoConsoleVariableRef CVar_HttpCacheEnabled(
@@ -1003,6 +1008,16 @@ void FHttpDispatcher::CompleteHttpRequest(FHttpRequest& Request, FMultiEndpointH
 
 void FHttpDispatcher::CompleteRequest(FHttpRequest& Request, FIoBuffer& Body, const FIoHttpRange& TotalRange, uint32 StatusCode, bool bCached)
 {
+#if UE_ALLOW_CACHE_POISONING
+	if (GIasPoisonCache && bCached)
+	{
+		for (uint64 Index = 0; Index < Body.GetSize(); Index++)
+		{
+			Body.GetData()[Index] = 0x4d;
+		}
+	}
+#endif // UE_ALLOW_CACHE_POISONING
+
 	FHttpRequest* NextRequest = &Request;
 	while (NextRequest != nullptr)
 	{
@@ -1075,4 +1090,6 @@ TSharedPtr<IOnDemandHttpIoDispatcher> MakeOnDemanHttpIoDispatcher(TUniquePtr<cla
 
 } // UE::IoStore
 
+#undef UE_ALLOW_CACHE_POISONING
+
 #undef UE_ALLOW_INVALID_URL_DEBUGGING
diff --git a/Engine/Source/Runtime/Experimental/IoStore/OnDemand/Private/OnDemandIoDispatcherBackend.cpp b/Engine/Source/Runtime/Experimental/IoStore/OnDemand/Private/OnDemandIoDispatcherBackend.cpp
--- a/Engine/Source/Runtime/Experimental/IoStore/OnDemand/Private/OnDemandIoDispatcherBackend.cpp
+++ b/Engine/Source/Runtime/Experimental/IoStore/OnDemand/Private/OnDemandIoDispatcherBackend.cpp
@@ -166,7 +166,7 @@ static FAutoConsoleVariableRef CVar_IasCancelationEnabled(
 #endif //UE_ALLOW_DISABLE_CANCELLING
 
 #if !UE_BUILD_SHIPPING
-static int32 GIasPoisonCache = false;
+bool GIasPoisonCache = false;
 static FAutoConsoleVariableRef CVar_IasPoisonCache(
 	TEXT("ias.PoisonCache"),
 	GIasPoisonCache,
