Dataflow : Add Promote to Variable feature to inputs
- move some useful function  to the dataflow engine util files
- fix logic when operating on anytype types

diff --git a/Engine/Plugins/Experimental/Dataflow/Source/DataflowEditor/Private/Dataflow/DataflowAssetEditUtils.cpp b/Engine/Plugins/Experimental/Dataflow/Source/DataflowEditor/Private/Dataflow/DataflowAssetEditUtils.cpp
--- a/Engine/Plugins/Experimental/Dataflow/Source/DataflowEditor/Private/Dataflow/DataflowAssetEditUtils.cpp
+++ b/Engine/Plugins/Experimental/Dataflow/Source/DataflowEditor/Private/Dataflow/DataflowAssetEditUtils.cpp
@@ -244,7 +244,7 @@ namespace UE::Dataflow
 		static FName GenerateUniqueVariableName(const UDataflow& DataflowAsset, FName BaseName)
 		{
 			int32 Counter = 1;
-			FName UniqueName{ BaseName };
+			FName UniqueName = FInstancedPropertyBag::SanitizePropertyName(BaseName);
 			const FString BasenameStr{ BaseName.ToString() };
 			while (true)
 			{
@@ -855,6 +855,30 @@ namespace UE::Dataflow
 		return UniqueVariableName;
 	}
 
+	FName FEditAssetUtils::AddNewVariable(UDataflow* DataflowAsset, FName BaseName, const FPropertyBagPropertyDesc& TemplateDesc)
+	{
+		using namespace EditAssetUtils::Private;
+
+		const FText TransactionName{ LOCTEXT("AddNewDataflowVariable", "Add New Dataflow Variable") };
+
+		FName UniqueVariableName;
+
+		auto AddNewVariableInternal =
+			[&BaseName, &UniqueVariableName, &TemplateDesc](UDataflow& DataflowAsset) -> EChangeResult
+			{
+				UniqueVariableName = GenerateUniqueVariableName(DataflowAsset, BaseName);
+				FPropertyBagPropertyDesc Desc{ TemplateDesc };
+				Desc.Name = UniqueVariableName;
+				DataflowAsset.Variables.AddProperties({ Desc });
+				FDataflowAssetDelegates::OnVariablesChanged.Broadcast(&DataflowAsset, UniqueVariableName);
+				return EChangeResult::Changed;
+			};
+
+		ChangeDataflowAssetWithTransaction(DataflowAsset, TransactionName, AddNewVariableInternal, GET_MEMBER_NAME_CHECKED(UDataflow, Variables));
+
+		return UniqueVariableName;
+	}
+
 	//////////////////////////////////////////////////////////////////////////////////////////////////////////
 	void FEditAssetUtils::DeleteVariable(UDataflow* DataflowAsset, FName VariableName)
 	{
diff --git a/Engine/Plugins/Experimental/Dataflow/Source/DataflowEditor/Private/Dataflow/DataflowAssetEditUtils.h b/Engine/Plugins/Experimental/Dataflow/Source/DataflowEditor/Private/Dataflow/DataflowAssetEditUtils.h
--- a/Engine/Plugins/Experimental/Dataflow/Source/DataflowEditor/Private/Dataflow/DataflowAssetEditUtils.h
+++ b/Engine/Plugins/Experimental/Dataflow/Source/DataflowEditor/Private/Dataflow/DataflowAssetEditUtils.h
@@ -110,6 +110,12 @@ namespace UE::Dataflow
 		*/
 		static FName AddNewVariable(UDataflow* DataflowAsset, FName BaseName = DefaultNewVariableBaseName);
 
+		/**
+		* Add a new Dataflow asset variable from a property bag description template 
+		* If BaseName already exists, use it to generate a unique name by adding an increment to it in the form of BaseName_##
+		*/
+		static FName AddNewVariable(UDataflow* DataflowAsset, FName BaseName, const FPropertyBagPropertyDesc& TemplateDesc);
+
 		/**
 		* Remove a Dataflow asset variable
 		*/
diff --git a/Engine/Plugins/Experimental/Dataflow/Source/DataflowEditor/Private/Dataflow/DataflowGraphEditor.cpp b/Engine/Plugins/Experimental/Dataflow/Source/DataflowEditor/Private/Dataflow/DataflowGraphEditor.cpp
--- a/Engine/Plugins/Experimental/Dataflow/Source/DataflowEditor/Private/Dataflow/DataflowGraphEditor.cpp
+++ b/Engine/Plugins/Experimental/Dataflow/Source/DataflowEditor/Private/Dataflow/DataflowGraphEditor.cpp
@@ -7,6 +7,7 @@
 #include "Dataflow/DataflowEditor.h"
 #include "Dataflow/DataflowEditorToolkit.h"
 #include "Dataflow/DataflowEngine.h"
+#include "Dataflow/DataflowEngineUtil.h"
 #include "Dataflow/DataflowSEditorInterface.h"
 #include "Dataflow/DataflowSNodeFactories.h"
 #include "Dataflow/DataflowSCommentNode.h"
@@ -16,6 +17,7 @@
 #include "Dataflow/DataflowSubGraphNodes.h"
 #include "Dataflow/DataflowSchema.h"
 #include "Dataflow/DataflowSubGraph.h"
+#include "Dataflow/DataflowVariableNodes.h"
 #include "HAL/PlatformApplicationMisc.h"
 #include "ScopedTransaction.h"
 #include "IStructureDetailsView.h"
@@ -277,6 +279,11 @@ void SDataflowGraphEditor::Construct(const FArguments& InArgs, UObject* InAssetO
 				FCanExecuteAction::CreateSP(this, &SDataflowGraphEditor::CanStopWatchingPin)
 			);
 
+			GraphEditorCommands->MapAction(FGraphEditorCommands::Get().PromoteToVariable,
+				FExecuteAction::CreateSP(this, &SDataflowGraphEditor::OnPromoteToVariable),
+				FCanExecuteAction::CreateSP(this, &SDataflowGraphEditor::CanPromoteToVariable)
+			);
+
 		}
 	}
 
@@ -1057,6 +1064,69 @@ bool SDataflowGraphEditor::CanStopWatchingPin() const
 	return false;
 }
 
+void SDataflowGraphEditor::OnPromoteToVariable()
+{
+	if (UEdGraphPin* Pin = GetGraphPinForMenu())
+	{
+		if (UDataflowEdNode* EdNode = Cast<UDataflowEdNode>(Pin->GetOwningNode()))
+		{
+			if (UDataflow* DataflowAsset = GetDataflowAsset())
+			{
+				FDataflowInput* DataflowInput = nullptr;
+				if (TSharedPtr<FDataflowNode> DataflowNode = EdNode->GetDataflowNode())
+				{
+					DataflowInput = DataflowNode->FindInput(Pin->PinName);
+				}
+				if (DataflowInput)
+				{
+					// create the variable from the input type
+					using namespace UE::Dataflow::Type;
+					const FPropertyBagPropertyDesc TemplateDesc = GetPropertyBagPropertyDescFromDataflowConnection(*DataflowInput);
+					const FName NewVariableName = UE::Dataflow::FEditAssetUtils::AddNewVariable(DataflowAsset, Pin->PinName, TemplateDesc);
+
+					// add the variable node
+					if (UEdGraph* ParentGraph = EdNode->GetGraph())
+					{
+						if (TSharedPtr<FAssetSchemaAction_Dataflow_CreateNode_DataflowEdNode>  VariableNodeAction = FAssetSchemaAction_Dataflow_CreateNode_DataflowEdNode::CreateAction(ParentGraph, "FGetDataflowVariableNode", NewVariableName))
+						{
+							const FVector2f Location = EdNode->GetPosition() - FVector2f(200, 0);
+							UDataflowEdNode* VariableEdNode = Cast<UDataflowEdNode>(VariableNodeAction->PerformAction(ParentGraph, nullptr, Location, false));
+							if (VariableEdNode)
+							{
+								if (TSharedPtr<FGetDataflowVariableNode> VariableNode = StaticCastSharedPtr<FGetDataflowVariableNode>(VariableEdNode->GetDataflowNode()))
+								{
+									VariableNode->SetVariable(DataflowAsset, NewVariableName);
+									VariableEdNode->UpdatePinsFromDataflowNode();
+									ParentGraph->NotifyNodeChanged(VariableEdNode);
+								}
+
+								// now connect the nodes 
+								// there's only pin on a variable node
+								if (UEdGraphPin* VariableNodePin = VariableEdNode->GetPinAt(0))
+								{
+									ParentGraph->GetSchema()->TryCreateConnection(VariableNodePin, Pin);
+								}
+
+								ParentGraph->NotifyGraphChanged();
+							}
+
+						}
+					}
+				}
+			}
+		}
+	}
+}
+
+bool SDataflowGraphEditor::CanPromoteToVariable() const
+{
+	if (UEdGraphPin* Pin = const_cast<SDataflowGraphEditor*>(this)->GetGraphPinForMenu())
+	{
+		return !(Pin->HasAnyConnections());
+	}
+	return false;
+}
+
 void SDataflowGraphEditor::InitGraphEditorMessageBar()
 {
 	SAssignNew(MessageBar, SHorizontalBox)
diff --git a/Engine/Plugins/Experimental/Dataflow/Source/DataflowEditor/Private/Dataflow/DataflowSchema.cpp b/Engine/Plugins/Experimental/Dataflow/Source/DataflowEditor/Private/Dataflow/DataflowSchema.cpp
--- a/Engine/Plugins/Experimental/Dataflow/Source/DataflowEditor/Private/Dataflow/DataflowSchema.cpp
+++ b/Engine/Plugins/Experimental/Dataflow/Source/DataflowEditor/Private/Dataflow/DataflowSchema.cpp
@@ -112,6 +112,14 @@ void UDataflowSchema::GetContextMenuActions(class UToolMenu* Menu, class UGraphN
 				Section.AddMenuEntry(bConnectionWatched? FGraphEditorCommands::Get().StopWatchingPin: FGraphEditorCommands::Get().StartWatchingPin);
 			}
 		}
+		else if (Context->Pin && Context->Pin->Direction == EGPD_Input)
+		{
+			if (FDataflowConnection* Connection = UDataflowEdNode::GetConnectionFromPin(Context->Pin))
+			{
+				FToolMenuSection& Section = Menu->FindOrAddSection("Pin Actions");
+				Section.AddMenuEntry(FGraphEditorCommands::Get().PromoteToVariable);
+			}
+		}
 	}
 }
 bool UDataflowSchema::CanPinBeConnectedToNode(const UEdGraphPin* Pin, const UE::Dataflow::FFactoryParameters& NodeParameters)
diff --git a/Engine/Plugins/Experimental/Dataflow/Source/DataflowEditor/Public/Dataflow/DataflowGraphEditor.h b/Engine/Plugins/Experimental/Dataflow/Source/DataflowEditor/Public/Dataflow/DataflowGraphEditor.h
--- a/Engine/Plugins/Experimental/Dataflow/Source/DataflowEditor/Public/Dataflow/DataflowGraphEditor.h
+++ b/Engine/Plugins/Experimental/Dataflow/Source/DataflowEditor/Public/Dataflow/DataflowGraphEditor.h
@@ -188,13 +188,15 @@ private:
 	UE_API bool CanStartWatchingPin() const;
 	UE_API void OnStopWatchingPin();
 	UE_API bool CanStopWatchingPin() const;
+	UE_API void OnPromoteToVariable();
+	UE_API bool CanPromoteToVariable() const;
 
 	UE_API bool GetPinVisibility(SGraphEditor::EPinVisibility InVisibility) const;
 
 	/** Create a widget to display an overlaid message in the graph editor panel */
 	UE_API void InitGraphEditorMessageBar();
 
-	/** Create a widget to display teh progress of the evaluation of the graph */
+	/** Create a widget to display the progress of the evaluation of the graph */
 	UE_API void InitEvaluationProgressBar();
 
 	/** Text for overlay message in graph editor panel */
diff --git a/Engine/Source/Runtime/Experimental/Dataflow/Engine/Private/Dataflow/DataflowEngineUtil.cpp b/Engine/Source/Runtime/Experimental/Dataflow/Engine/Private/Dataflow/DataflowEngineUtil.cpp
--- a/Engine/Source/Runtime/Experimental/Dataflow/Engine/Private/Dataflow/DataflowEngineUtil.cpp
+++ b/Engine/Source/Runtime/Experimental/Dataflow/Engine/Private/Dataflow/DataflowEngineUtil.cpp
@@ -1,7 +1,12 @@
 // Copyright Epic Games, Inc. All Rights Reserved.
 
 #include "Dataflow/DataflowEngineUtil.h"
+
+#include "Dataflow/DataflowAnyType.h"
+#include "Dataflow/DataflowConnection.h"
+#include "Dataflow/DataflowTypePolicy.h"
 #include "ReferenceSkeleton.h"
+#include "UObject/UnrealType.h"
 
 namespace UE::Dataflow
 {
@@ -57,5 +62,127 @@ namespace UE::Dataflow
 			return FLinearColor(FColor(R, G, B, 255));
 		}
 	}
+
+	namespace Type
+	{
+		FPropertyBagPropertyDesc GetPropertyBagPropertyDescFromDataflowType(const FName Name, const FName Type)
+		{
+			if (Type == UE::Dataflow::GetTypeName<bool>())
+			{
+				return FPropertyBagPropertyDesc(Name, EPropertyBagPropertyType::Bool);
+			}
+			else if (Type == UE::Dataflow::GetTypeName<int32>())
+			{
+				return FPropertyBagPropertyDesc(Name, EPropertyBagPropertyType::Int32);
+			}
+			else if (Type == UE::Dataflow::GetTypeName<int64>())
+			{
+				return FPropertyBagPropertyDesc(Name, EPropertyBagPropertyType::Int64);
+			}
+			else if (Type == UE::Dataflow::GetTypeName<uint32>())
+			{
+				return FPropertyBagPropertyDesc(Name, EPropertyBagPropertyType::UInt32);
+			}
+			else if (Type == UE::Dataflow::GetTypeName<uint64>())
+			{
+				return FPropertyBagPropertyDesc(Name, EPropertyBagPropertyType::UInt64);
+			}
+			else if (Type == UE::Dataflow::GetTypeName<float>())
+			{
+				return FPropertyBagPropertyDesc(Name, EPropertyBagPropertyType::Float);
+			}
+			else if (Type == UE::Dataflow::GetTypeName<double>())
+			{
+				return FPropertyBagPropertyDesc(Name, EPropertyBagPropertyType::Double);
+			}
+			else if (Type == UE::Dataflow::GetTypeName<FName>())
+			{
+				return FPropertyBagPropertyDesc(Name, EPropertyBagPropertyType::Name);
+			}
+			else if (Type == UE::Dataflow::GetTypeName<FString>())
+			{
+				return FPropertyBagPropertyDesc(Name, EPropertyBagPropertyType::String);
+			}
+			else if (Type == UE::Dataflow::GetTypeName<FText>())
+			{
+				return FPropertyBagPropertyDesc(Name, EPropertyBagPropertyType::Text);
+			}
+			else
+			{
+				const FString TypeStr(Type.ToString());
+				if (TypeStr.StartsWith("U"))
+				{
+					// Nothing for now and use thefallback - may need fixing if we find there's an issue with that code path 
+				}
+				else if (TypeStr.StartsWith("F"))
+				{
+					const FString ShortTypeName = TypeStr.RightChop(1);
+					if (const UScriptStruct* ScriptStruct = FindFirstObjectSafe<UScriptStruct>(*ShortTypeName))
+					{
+						return FPropertyBagPropertyDesc(Name, EPropertyBagPropertyType::Struct, ScriptStruct);
+					}
+				}
+				else if (FDataflowArrayTypePolicy::SupportsTypeStatic(Type))
+				{
+					const FName ElementType = FDataflowArrayTypePolicy::GetElementType(Type);
+					FPropertyBagPropertyDesc PropertyDesc = GetPropertyBagPropertyDescFromDataflowType(Name, ElementType);
+					PropertyDesc.ContainerTypes = FPropertyBagContainerTypes{ EPropertyBagContainerType::Array };
+					return PropertyDesc;
+				}
+			}
+			// invalid value 
+			return FPropertyBagPropertyDesc(Name, EPropertyBagPropertyType::Count);
+		}
+
+		FPropertyBagPropertyDesc GetPropertyBagPropertyDescFromDataflowConnection(const FDataflowConnection& Connection)
+		{
+			const FProperty* ConnectionProperty = Connection.GetProperty();
+			check(ConnectionProperty);
+
+			// AnyType require special treatment to make sure we give a concrete type
+			if (Connection.IsAnyType())
+			{
+				if (Connection.HasConcreteType())
+				{
+					const FPropertyBagPropertyDesc Desc = GetPropertyBagPropertyDescFromDataflowType(Connection.GetName(), Connection.GetType());
+					if (Desc.ValueType != EPropertyBagPropertyType::Count)
+					{
+						return Desc;
+					}
+				}
+				// Fallback let's use the AnyType default (Value) type
+				if (ConnectionProperty->GetClass()->IsChildOf(FStructProperty::StaticClass()))
+				{
+					if (const FStructProperty* StructProperty = CastField<FStructProperty>(ConnectionProperty))
+					{
+						if (const bool bInheritsFromAnyType = StructProperty->Struct->IsChildOf<FDataflowAnyType>())
+						{
+							for (TFieldIterator<FProperty> FieldIt(StructProperty->Struct); FieldIt; ++FieldIt)
+							{
+								FProperty* Prop = *FieldIt;
+								if (Prop->GetFName() == "Value")
+								{
+									return FPropertyBagPropertyDesc(Connection.GetName(), Prop);
+								}
+							}
+						}
+					}
+				}
+			}
+			// default simply set property desc from the property
+			return FPropertyBagPropertyDesc(Connection.GetName(), Connection.GetProperty());
+		}
+
+		FName MakeUniqueNameForPropertyBag(FName OriginalName, const FInstancedPropertyBag& PropertyBag)
+		{
+			FName NewName = OriginalName;
+			int32 SuffixNumber = 1;
+			while (PropertyBag.FindPropertyDescByName(NewName))
+			{
+				NewName.SetNumber(SuffixNumber++);
+			}
+			return NewName;
+		}
+	}
 }
 
diff --git a/Engine/Source/Runtime/Experimental/Dataflow/Engine/Private/Dataflow/DataflowSubGraphNodes.cpp b/Engine/Source/Runtime/Experimental/Dataflow/Engine/Private/Dataflow/DataflowSubGraphNodes.cpp
--- a/Engine/Source/Runtime/Experimental/Dataflow/Engine/Private/Dataflow/DataflowSubGraphNodes.cpp
+++ b/Engine/Source/Runtime/Experimental/Dataflow/Engine/Private/Dataflow/DataflowSubGraphNodes.cpp
@@ -4,6 +4,7 @@
 
 #include "Dataflow/DataflowNode.h"
 #include "Dataflow/DataflowEdNode.h"
+#include "Dataflow/DataflowEngineUtil.h"
 #include "Dataflow/DataflowObject.h"
 #include "Dataflow/DataflowObjectInterface.h"
 #include "Dataflow/DataflowSubGraph.h"
@@ -16,17 +17,6 @@
 
 namespace UE::Dataflow::SubGraphNodes::Private
 {
-	FName MakeUniqueNameForPropertyBag(FName OriginalName, const FInstancedPropertyBag& PropertyBag)
-	{
-		FName NewName = OriginalName;
-		int32 SuffixNumber = 1;
-		while (PropertyBag.FindPropertyDescByName(NewName))
-		{
-			NewName.SetNumber(SuffixNumber++);
-		}
-		return NewName;
-	}
-
 	class FContextSubGraph: public FEngineContext
 	{
 	public:
@@ -160,114 +150,6 @@ namespace UE::Dataflow::SubGraphNodes::Private
 		FContext& ParentContext;
 		const UE::Dataflow::ISubGraphContextCallback& SubGraphCallback;
 	};
-
-	FPropertyBagPropertyDesc GetPropertyBagPropertyDescFromDataflowType(const FName Name, const FName Type)
-	{
-		if (Type == UE::Dataflow::GetTypeName<bool>())
-		{
-			return FPropertyBagPropertyDesc(Name, EPropertyBagPropertyType::Bool);
-		}
-		else if (Type == UE::Dataflow::GetTypeName<int32>())
-		{
-			return FPropertyBagPropertyDesc(Name, EPropertyBagPropertyType::Int32);
-		}
-		else if (Type == UE::Dataflow::GetTypeName<int64>())
-		{
-			return FPropertyBagPropertyDesc(Name, EPropertyBagPropertyType::Int64);
-		}
-		else if (Type == UE::Dataflow::GetTypeName<uint32>())
-		{
-			return FPropertyBagPropertyDesc(Name, EPropertyBagPropertyType::UInt32);
-		}
-		else if (Type == UE::Dataflow::GetTypeName<uint64>())
-		{
-			return FPropertyBagPropertyDesc(Name, EPropertyBagPropertyType::UInt64);
-		}
-		else if (Type == UE::Dataflow::GetTypeName<float>())
-		{
-			return FPropertyBagPropertyDesc(Name, EPropertyBagPropertyType::Float);
-		}
-		else if (Type == UE::Dataflow::GetTypeName<double>())
-		{
-			return FPropertyBagPropertyDesc(Name, EPropertyBagPropertyType::Double);
-		}
-		else if (Type == UE::Dataflow::GetTypeName<FName>())
-		{
-			return FPropertyBagPropertyDesc(Name, EPropertyBagPropertyType::Name);
-		}
-		else if (Type == UE::Dataflow::GetTypeName<FString>())
-		{
-			return FPropertyBagPropertyDesc(Name, EPropertyBagPropertyType::String);
-		}
-		else if (Type == UE::Dataflow::GetTypeName<FText>())
-		{
-			return FPropertyBagPropertyDesc(Name, EPropertyBagPropertyType::Text);
-		}
-		else
-		{
-			const FString TypeStr(Type.ToString());
-			if (TypeStr.StartsWith("U"))
-			{
-				// Nothing for now and use thefallback - may need fixing if we find there's an issue with that code path 
-			}
-			else if (TypeStr.StartsWith("F"))
-			{
-				const FString ShortTypeName = TypeStr.RightChop(1);
-				if (const UScriptStruct* ScriptStruct = FindFirstObjectSafe<UScriptStruct>(*ShortTypeName))
-				{
-					return FPropertyBagPropertyDesc(Name, EPropertyBagPropertyType::Struct, ScriptStruct);
-				}
-			}
-			else if (FDataflowArrayTypePolicy::SupportsTypeStatic(Type))
-			{
-				const FName ElementType = FDataflowArrayTypePolicy::GetElementType(Type);
-				FPropertyBagPropertyDesc PropertyDesc = GetPropertyBagPropertyDescFromDataflowType(Name, ElementType);
-				PropertyDesc.ContainerTypes = FPropertyBagContainerTypes{ EPropertyBagContainerType::Array };
-				return PropertyDesc;
-			}
-		}
-		// invalid value 
-		return FPropertyBagPropertyDesc(Name, EPropertyBagPropertyType::Count);
-	}
-
-	FPropertyBagPropertyDesc GetPropertyBagPropertyDescFromDataflowConnection(const FDataflowConnection& Connection)
-	{
-		const FProperty* ConnectionProperty = Connection.GetProperty();
-		check(ConnectionProperty);
-
-		// AnyType require special treatment to make sure we give a concrete type
-		if (Connection.IsAnyType())
-		{
-			const FPropertyBagPropertyDesc Desc = GetPropertyBagPropertyDescFromDataflowType(Connection.GetName(), Connection.GetType());
-			if (Desc.ValueType != EPropertyBagPropertyType::Count)
-			{
-				return Desc;
-			}
-			else
-			{
-				// Fallback let's use the AnyType default (Value) type
-				if (ConnectionProperty->GetClass()->IsChildOf(FStructProperty::StaticClass()))
-				{
-					if (const FStructProperty* StructProperty = CastField<FStructProperty>(ConnectionProperty))
-					{
-						if (const bool bInheritsFromAnyType = StructProperty->Struct->IsChildOf<FDataflowAnyType>())
-						{
-							for (TFieldIterator<FProperty> FieldIt(StructProperty->Struct); FieldIt; ++FieldIt)
-							{
-								FProperty* Prop = *FieldIt;
-								if (Prop->GetFName() == "Value")
-								{
-									return FPropertyBagPropertyDesc(Connection.GetName(), Prop);
-								}
-							}
-						}
-					}
-				}
-			}
-		}
-		// default simply set property desc from the property
-		return FPropertyBagPropertyDesc(Connection.GetName(), Connection.GetProperty());
-	}
 }
 
 //////////////////////////////////////////////////////////////////////////////////////////////////////////
@@ -327,10 +209,11 @@ bool FDataflowSubGraphInputNode::SupportsDropConnectionOnNode(FName TypeName, UE
 const FDataflowConnection* FDataflowSubGraphInputNode::OnDropConnectionOnNode(const FDataflowConnection& DroppedConnection)
 {
 	using namespace  UE::Dataflow::SubGraphNodes;
+	using namespace UE::Dataflow::Type;
 	if (SupportsDropConnectionOnNode(DroppedConnection.GetType(), DroppedConnection.GetDirection()))
 	{
-		FPropertyBagPropertyDesc Desc = Private::GetPropertyBagPropertyDescFromDataflowConnection(DroppedConnection);
-		Desc.Name = Private::MakeUniqueNameForPropertyBag(Desc.Name, PropertyBag);
+		FPropertyBagPropertyDesc Desc = GetPropertyBagPropertyDescFromDataflowConnection(DroppedConnection);
+		Desc.Name = MakeUniqueNameForPropertyBag(Desc.Name, PropertyBag);
 		PropertyBag.AddProperties({ Desc });
 		DynamicConnections.Refresh();
 		return FindOutput(Desc.Name);
@@ -341,12 +224,13 @@ const FDataflowConnection* FDataflowSubGraphInputNode::OnDropConnectionOnNode(co
 void FDataflowSubGraphInputNode::AddConnectionsTo(TArrayView<FDataflowConnection*> Connections)
 {
 	using namespace  UE::Dataflow::SubGraphNodes;
+	using namespace UE::Dataflow::Type;
 	for (const FDataflowConnection* Connection : Connections)
 	{
 		if (Connection && SupportsDropConnectionOnNode(Connection->GetType(), Connection->GetDirection()))
 		{
-			FPropertyBagPropertyDesc Desc = Private::GetPropertyBagPropertyDescFromDataflowConnection(*Connection);
-			Desc.Name = Private::MakeUniqueNameForPropertyBag(Desc.Name, PropertyBag);
+			FPropertyBagPropertyDesc Desc = GetPropertyBagPropertyDescFromDataflowConnection(*Connection);
+			Desc.Name = MakeUniqueNameForPropertyBag(Desc.Name, PropertyBag);
 			PropertyBag.AddProperties({ Desc });
 		}
 	}
@@ -390,10 +274,11 @@ bool FDataflowSubGraphOutputNode::SupportsDropConnectionOnNode(FName TypeName, U
 const FDataflowConnection* FDataflowSubGraphOutputNode::OnDropConnectionOnNode(const FDataflowConnection& DroppedConnection)
 {
 	using namespace  UE::Dataflow::SubGraphNodes;
+	using namespace UE::Dataflow::Type;
 	if (SupportsDropConnectionOnNode(DroppedConnection.GetType(), DroppedConnection.GetDirection()))
 	{
-		FPropertyBagPropertyDesc Desc = Private::GetPropertyBagPropertyDescFromDataflowConnection(DroppedConnection);
-		Desc.Name = Private::MakeUniqueNameForPropertyBag(Desc.Name, PropertyBag);
+		FPropertyBagPropertyDesc Desc = GetPropertyBagPropertyDescFromDataflowConnection(DroppedConnection);
+		Desc.Name = MakeUniqueNameForPropertyBag(Desc.Name, PropertyBag);
 		PropertyBag.AddProperties({ Desc });
 		DynamicConnections.Refresh();
 		return FindInput(Desc.Name);
@@ -404,12 +289,13 @@ const FDataflowConnection* FDataflowSubGraphOutputNode::OnDropConnectionOnNode(c
 void FDataflowSubGraphOutputNode::AddConnectionsTo(TArrayView<FDataflowConnection*> Connections)
 {
 	using namespace  UE::Dataflow::SubGraphNodes;
+	using namespace UE::Dataflow::Type;
 	for (const FDataflowConnection* Connection : Connections)
 	{
 		if (Connection && SupportsDropConnectionOnNode(Connection->GetType(), Connection->GetDirection()))
 		{
-			FPropertyBagPropertyDesc Desc = Private::GetPropertyBagPropertyDescFromDataflowConnection(*Connection);
-			Desc.Name = Private::MakeUniqueNameForPropertyBag(Desc.Name, PropertyBag);
+			FPropertyBagPropertyDesc Desc = GetPropertyBagPropertyDescFromDataflowConnection(*Connection);
+			Desc.Name = MakeUniqueNameForPropertyBag(Desc.Name, PropertyBag);
 			PropertyBag.AddProperties({ Desc });
 		}
 	}
diff --git a/Engine/Source/Runtime/Experimental/Dataflow/Engine/Public/Dataflow/DataflowEngineUtil.h b/Engine/Source/Runtime/Experimental/Dataflow/Engine/Public/Dataflow/DataflowEngineUtil.h
--- a/Engine/Source/Runtime/Experimental/Dataflow/Engine/Public/Dataflow/DataflowEngineUtil.h
+++ b/Engine/Source/Runtime/Experimental/Dataflow/Engine/Public/Dataflow/DataflowEngineUtil.h
@@ -7,6 +7,7 @@
 #include "Dataflow/DataflowEngineTypes.h"
 #include "Math/Transform.h"
 #include "ReferenceSkeleton.h"
+#include "StructUtils/PropertyBag.h"
 #include "Templates/Casts.h"
 #include "UObject/Class.h"
 #include "UObject/Field.h"
@@ -15,6 +16,11 @@
 #include "UObject/UnrealType.h"
 #include "UObject/UnrealTypePrivate.h"
 
+#define UE_API DATAFLOWENGINE_API
+
+struct FDataflowConnection;
+struct FInstancedPropertyBag;
+
 namespace UE::Dataflow
 {
 	namespace Reflection
@@ -72,11 +78,20 @@ namespace UE::Dataflow
 
 	namespace Animation
 	{
-		void DATAFLOWENGINE_API GlobalTransforms(const FReferenceSkeleton&, TArray<FTransform>& GlobalTransforms);
+		UE_API void GlobalTransforms(const FReferenceSkeleton&, TArray<FTransform>& GlobalTransforms);
 	}
 
 	namespace Color
 	{
-		DATAFLOWENGINE_API FLinearColor GetRandomColor(const int32 RandomSeed, int32 Idx);
+		UE_API FLinearColor GetRandomColor(const int32 RandomSeed, int32 Idx);
+	}
+
+	namespace Type
+	{
+		UE_API FPropertyBagPropertyDesc GetPropertyBagPropertyDescFromDataflowType(const FName Name, const FName Type);
+		UE_API FPropertyBagPropertyDesc GetPropertyBagPropertyDescFromDataflowConnection(const FDataflowConnection& Connection);
+		UE_API FName MakeUniqueNameForPropertyBag(FName OriginalName, const FInstancedPropertyBag& PropertyBag);
 	}
-}
\ No newline at end of file
+}
+
+#undef UE_API
