Enable the new inflight request cancelation code path.

#rnx

diff --git a/Engine/Source/Runtime/Experimental/IoStore/OnDemand/Private/OnDemandHttpThread.cpp b/Engine/Source/Runtime/Experimental/IoStore/OnDemand/Private/OnDemandHttpThread.cpp
--- a/Engine/Source/Runtime/Experimental/IoStore/OnDemand/Private/OnDemandHttpThread.cpp
+++ b/Engine/Source/Runtime/Experimental/IoStore/OnDemand/Private/OnDemandHttpThread.cpp
@@ -163,7 +163,7 @@ static FAutoConsoleVariableRef CVar_GIasHttpPollTimeoutMs(
 	TEXT("Http tick poll timeout in milliseconds")
 );
 
-bool GIaxEnableHttpCancellation = false;
+bool GIaxEnableHttpCancellation = true;
 static FAutoConsoleVariableRef CVar_IaxEnableHttpCancellation(
 	TEXT("iax.EnableHttpCancellation"),
 	GIaxEnableHttpCancellation,
