Fixed inverted condition.

#rnx

diff --git a/Engine/Source/Runtime/Experimental/IoStore/HttpClient/Private/Activity.inl b/Engine/Source/Runtime/Experimental/IoStore/HttpClient/Private/Activity.inl
--- a/Engine/Source/Runtime/Experimental/IoStore/HttpClient/Private/Activity.inl
+++ b/Engine/Source/Runtime/Experimental/IoStore/HttpClient/Private/Activity.inl
@@ -407,7 +407,7 @@ FOutcome FActivity::Transact(FHttpPeer& Peer)
 	}
 
 	const FHost* TheHost = GetHost();
-	bool bKeepAlive = !TheHost->IsPooled();
+	bool bKeepAlive = TheHost->IsPooled();
 
 	Transaction->Begin(TheHost->GetHostName(), GetMethod(), GetPath());
 	EnumerateHeaders([this] (FAnsiStringView Key, FAnsiStringView Value)
