The new IAX backend now respects the cvar 'ias.EnableWriteOnlyDecoding'

#rnx

diff --git a/Engine/Source/Runtime/Experimental/IoStore/OnDemand/Private/OnDemandHttpIoDispatcherBackend.cpp b/Engine/Source/Runtime/Experimental/IoStore/OnDemand/Private/OnDemandHttpIoDispatcherBackend.cpp
--- a/Engine/Source/Runtime/Experimental/IoStore/OnDemand/Private/OnDemandHttpIoDispatcherBackend.cpp
+++ b/Engine/Source/Runtime/Experimental/IoStore/OnDemand/Private/OnDemandHttpIoDispatcherBackend.cpp
@@ -17,7 +17,9 @@
 namespace UE::IoStore
 {
 ////////////////////////////////////////////////////////////////////////////////
+
 extern int32 GIasHttpRangeRequestMinSizeKiB;
+extern bool GIasEnableWriteOnlyDecoding;
 
 ////////////////////////////////////////////////////////////////////////////////
 class FHttpIoDispatcherBackend
@@ -268,8 +270,9 @@ void FHttpIoDispatcherBackend::CompleteRequest(FResolvedRequest& ResolvedRequest
 
 		DispatcherRequest.CreateBuffer(DispatcherRequest.Options.GetSize());
 
-		const EIoDecodeFlags Options = EnumHasAnyFlags(DispatcherRequest.Options.GetFlags(), EIoReadOptionsFlags::HardwareTargetBuffer)
+		const EIoDecodeFlags Options = GIasEnableWriteOnlyDecoding && EnumHasAnyFlags(DispatcherRequest.Options.GetFlags(), EIoReadOptionsFlags::HardwareTargetBuffer)
 			? EIoDecodeFlags::WriteOnly : EIoDecodeFlags::None;
+
 		bDecoded = FIoChunkEncoding::Decode(Params, Chunk.GetView(), DispatcherRequest.GetBuffer().GetMutableView(), Options);
 	}
 
