Allow the new IAX backend to recover from corrupted data being stored in it's file cache.

#rnx

### Problem
- If the new backend was unable to decode data loaded from the cache we would only log an error. This meant that any future attempt to load that same data would continue to load it from the cache and encounter the same decode error. The older backend would evict data from the cache if we found that we were unable to decode it so that it would be re-downloaded again in the future which is behavior that the new backend needs to replicate.
- The main problem we face is that with the old backend, decoding was done by the backend which owned the cache, so evicting the data was easy. In the new system 'FHttpIoDispatcherBackend' is responsible for decoding the data but 'FHttpDispatcher' owns the cache and is hidden behind the FHttpIoDispatcher api.
- Note only do we not want to keep corrupted data in the cache, we don't want to write out corrupted data from a http request to the cache.
- We don't want to add another validation pass to FHttpIoDispatcherBackend as that is a waste of the CPU, we also don't want to risk adding the data to the cache and writing it to disk only for the calling code (in this case 'FHttpDispatcher') to later tell us that it should be evicted.

### Fix
- We no longer automatically cache data in 'FHttpDispatcher::CompleteRequest' instead we wait for the calling code to tell us to cache it.
- This is done via a new API method FHttpIoDispatcher::CacheResponse which should be called once we are sure that the data is correct.
-- It takes the original FIoHttpResponse that was passed to the caller via the FIoHttpRequestCompleted callback. This is to make sure that things like the cache key and data are correct and so that we can see if the data was originally loaded from the cache or not so as to avoid adding duplicates to the cache (which seems to be possible.
- In order for this to work, the cache key has been added to FIoHttpResponse.
- A second new API method FHttpIoDispatcher::EvictFromCache was added to allow the calling code to request that it be removed from the cache. This will prevent it from being re-used once we had determined that the data is bad.
- The write options have been removed from 'EIoHttpFlags' as they no longer make sense.
- Fixed tab alignment issues in HttpIoDispatcher.h

diff --git a/Engine/Source/Runtime/Experimental/IoStore/OnDemand/Private/OnDemandHttpIoDispacher.cpp b/Engine/Source/Runtime/Experimental/IoStore/OnDemand/Private/OnDemandHttpIoDispacher.cpp
--- a/Engine/Source/Runtime/Experimental/IoStore/OnDemand/Private/OnDemandHttpIoDispacher.cpp
+++ b/Engine/Source/Runtime/Experimental/IoStore/OnDemand/Private/OnDemandHttpIoDispacher.cpp
@@ -15,6 +15,7 @@
 #include "IO/IoContainers.h"
 #include "IO/IoStoreOnDemand.h"
 #include "IO/HttpIoDispatcher.h"
+#include "Logging/StructuredLog.h"
 #include "Logging/LogVerbosity.h"
 #include "Misc/CommandLine.h"
 #include "Misc/EnumClassFlags.h"
@@ -441,6 +442,9 @@ private:
 	virtual void					UpdateRequestPriority(FIoHttpRequestHandle Handle, int32 NewPriority) override;
 	virtual EIoErrorCode			GetRequestStatus(FIoHttpRequestHandle Handle) override;
 
+	virtual FIoStatus				CacheResponse(const FIoHttpResponse& Response) override;
+	virtual FIoStatus				EvictFromCache(const FIoHttpResponse& Response) override;
+
 	bool							TryReadFromCache(FHttpRequest* Request);
 	bool							ProcessHttpRequests();
 	void							CompleteHttpRequest(FHttpRequest& Request, FMultiEndpointHttpClientResponse&& HttpResponse, FIASHostGroup& HostGroup, const FIoHttpRange& TotalRange);
@@ -779,6 +783,46 @@ EIoErrorCode FHttpDispatcher::GetRequestStatus(FIoHttpRequestHandle Handle)
 	return EIoErrorCode::InvalidCode;
 }
 
+FIoStatus FHttpDispatcher::CacheResponse(const FIoHttpResponse& Response)
+{
+	if (Response.IsCached())
+	{
+		return EIoErrorCode::Ok;
+	}
+
+	// If the cache is disabled we consider the caching to be a success, otherwise the calling code will
+	// have to special case this.
+	if (!GHttpCacheEnabled || !Cache.IsValid())
+	{
+		return EIoErrorCode::Ok;
+	}
+
+	if (Response.GetErrorCode() != EIoErrorCode::Ok || Response.GetCacheKey().IsZero() || Response.GetBody().GetSize() == 0)
+	{
+		return EIoErrorCode::InvalidParameter;
+	}
+
+	FIoBuffer Copy(Response.GetBody());
+	return Cache->Put(Response.GetCacheKey(), Copy);
+}
+
+FIoStatus FHttpDispatcher::EvictFromCache(const FIoHttpResponse& Response)
+{
+	// If the cache is disabled we consider the eviction to be a success, otherwise the calling code will
+	// have to special case this.
+	if (!GHttpCacheEnabled || !Cache.IsValid())
+	{
+		return EIoErrorCode::Ok;
+	}
+
+	if (Response.GetErrorCode() != EIoErrorCode::Ok || Response.GetCacheKey().IsZero())
+	{
+		return EIoErrorCode::InvalidParameter;
+	}
+
+	return Cache->Evict(Response.GetCacheKey());
+}
+
 bool FHttpDispatcher::TryReadFromCache(FHttpRequest* Request)
 {
 	if (!GHttpCacheEnabled || !Cache.IsValid())
@@ -1057,22 +1101,9 @@ void FHttpDispatcher::CompleteRequest(FHttpRequest& Request, FIoBuffer& Body, co
 			PartialBody = FIoBuffer(Body.GetView().Mid(Offset, Size), Body);
 		}
 
-		const bool bWriteCache = EnumHasAnyFlags(ToComplete->Options.GetFlags(), EIoHttpFlags::WriteCache)
-			&& !bCached && CompletionStatus == EIoErrorCode::Ok
-			&& GHttpCacheEnabled && Cache.IsValid()
-			&& PartialBody.GetSize() > 0;
-
-		if (bWriteCache)
-		{
-			if (FIoStatus Status = Cache->Put(ToComplete->CacheKey, PartialBody); !Status.IsOk())
-			{
-				UE_LOG(LogHttpIoDispatcher, Error, TEXT("Put HTTP cache failed, reason '%s'"), *Status.ToString());
-			}
-		}
-
 		const EIoHttpResponseFlags Flags	= bCached ? EIoHttpResponseFlags::Cached : EIoHttpResponseFlags::None;
 		FIoHttpRequestCompleted OnCompleted = MoveTemp(ToComplete->OnCompleted);
-		OnCompleted(FIoHttpResponse(PartialBody, CompletionStatus, StatusCode, Flags));
+		OnCompleted(FIoHttpResponse(PartialBody, ToComplete->CacheKey, CompletionStatus, StatusCode, Flags));
 
 		ToComplete->Flags.Add(EHttpRequestFlags::Completed);
 		ToComplete->CompletionStatus.store(CompletionStatus, std::memory_order_seq_cst);
diff --git a/Engine/Source/Runtime/Experimental/IoStore/OnDemand/Private/OnDemandHttpIoDispatcherBackend.cpp b/Engine/Source/Runtime/Experimental/IoStore/OnDemand/Private/OnDemandHttpIoDispatcherBackend.cpp
--- a/Engine/Source/Runtime/Experimental/IoStore/OnDemand/Private/OnDemandHttpIoDispatcherBackend.cpp
+++ b/Engine/Source/Runtime/Experimental/IoStore/OnDemand/Private/OnDemandHttpIoDispatcherBackend.cpp
@@ -282,6 +282,12 @@ void FHttpIoDispatcherBackend::CompleteRequest(FResolvedRequest& ResolvedRequest
 
 	if (bDecoded)
 	{
+		if (!HttpResponse.IsCached())
+		{
+			FIoStatus Status = FHttpIoDispatcher::CacheResponse(HttpResponse);
+			UE_CLOG(!Status.IsOk(), LogHttpIoDispatcher, Error, TEXT("Put HTTP cache failed, reason '%s'"), *Status.ToString());
+		}
+
 		Stats.OnIoRequestComplete(DispatcherRequest.GetBuffer().GetSize(), DurationMs);
 	}
 	else if (bCancelled)
@@ -292,9 +298,15 @@ void FHttpIoDispatcherBackend::CompleteRequest(FResolvedRequest& ResolvedRequest
 	{
 		DispatcherRequest.SetFailed();
 		Stats.OnIoRequestError();
+
 		if (HttpResponse.IsCached())
 		{
 			Stats.OnCacheDecodeError();
+
+			if (FIoStatus Status = FHttpIoDispatcher::EvictFromCache(HttpResponse); !Status.IsOk())
+			{
+				UE_LOG(LogHttpIoDispatcher, Error, TEXT("Evict HTTP cache failed, reason '%s'"), *Status.ToString());
+			}
 		}
 		else
 		{
diff --git a/Engine/Source/Runtime/Experimental/IoStore/OnDemandCore/Private/HttpIoDispatcher.cpp b/Engine/Source/Runtime/Experimental/IoStore/OnDemandCore/Private/HttpIoDispatcher.cpp
--- a/Engine/Source/Runtime/Experimental/IoStore/OnDemandCore/Private/HttpIoDispatcher.cpp
+++ b/Engine/Source/Runtime/Experimental/IoStore/OnDemandCore/Private/HttpIoDispatcher.cpp
@@ -299,6 +299,17 @@ FIoHttpRequest FHttpIoDispatcher::Get(
 {
 	return Get(HostGroup, RelativeUrl, FIoHttpHeaders(), FIoHttpOptions::Default, MoveTemp(OnCompleted));
 }
+ FIoStatus FHttpIoDispatcher::CacheResponse(const FIoHttpResponse& Response)
+{
+	 ensure(GHttpIoDispatcher.IsValid());
+	 return GHttpIoDispatcher->CacheResponse(Response);
+}
+
+FIoStatus FHttpIoDispatcher::EvictFromCache(const FIoHttpResponse& Response)
+{
+	ensure(GHttpIoDispatcher.IsValid());
+	return GHttpIoDispatcher->EvictFromCache(Response);
+}
 
 FHttpIoDispatcher::FHostGroupRegistered& FHttpIoDispatcher::OnHostGroupRegistered()
 {
diff --git a/Engine/Source/Runtime/Experimental/IoStore/OnDemandCore/Public/IO/HttpIoDispatcher.h b/Engine/Source/Runtime/Experimental/IoStore/OnDemandCore/Public/IO/HttpIoDispatcher.h
--- a/Engine/Source/Runtime/Experimental/IoStore/OnDemandCore/Public/IO/HttpIoDispatcher.h
+++ b/Engine/Source/Runtime/Experimental/IoStore/OnDemandCore/Public/IO/HttpIoDispatcher.h
@@ -30,19 +30,15 @@ enum class EIoHttpFlags
 	None			= 0,
 	/** Whether to read from the HTTP cache. */ 
 	ReadCache		= (1 << 0),
-	/** Whether to write to the HTTP cache. */ 
-	WriteCache		= (1 << 1),
-	/** Whether to read and write to the HTTP cache. */ 
-	ReadWriteCache	= ReadCache | WriteCache,
 	/** Default flags . */
-	Default			= ReadWriteCache
+	Default			= ReadCache
 };
 ENUM_CLASS_FLAGS(EIoHttpFlags);
 
 /** Represents a range within a resource. */
 struct FIoHttpRange
 {
-						FIoHttpRange() = default;
+								FIoHttpRange() = default;
 	explicit inline				FIoHttpRange(uint32 InMin, uint32 InMax);
 
 	bool						IsValid()	const	{ return Min <= Max; }
@@ -111,7 +107,7 @@ public:
 	inline						FIoHttpOptions(int32 InPriority, int32 InRetryCount, EIoHttpFlags InFlags);
 	inline						FIoHttpOptions(int32 InPriority, int32 InRetryCount);
 	inline explicit				FIoHttpOptions(int32 InPriority);
-						FIoHttpOptions() = default;
+								FIoHttpOptions() = default;
 
 	const FIoHttpRange&			GetRange()		const { return Range; }
 	int32						GetPriority()	const { return Priority; }
@@ -164,9 +160,9 @@ FIoHttpOptions::FIoHttpOptions(int32 InPriority)
 class FIoHttpHeaders
 {
 public:
-							FIoHttpHeaders() = default;
-							FIoHttpHeaders(const FIoHttpHeaders&) = default;
-							FIoHttpHeaders(FIoHttpHeaders&&) = default;
+									FIoHttpHeaders() = default;
+									FIoHttpHeaders(const FIoHttpHeaders&) = default;
+									FIoHttpHeaders(FIoHttpHeaders&&) = default;
 	UE_API FIoHttpHeaders&			Add(FAnsiString&& HeaderName, FAnsiString&& HeaderValue);
 	UE_API FIoHttpHeaders&			Add(FAnsiStringView HeaderName, FAnsiStringView HeaderValue);
 	UE_API TArray<FAnsiString>		ToArray() &&;
@@ -187,7 +183,7 @@ private:
 class FIoRelativeUrl
 {
 public:
-							FIoRelativeUrl() = default;
+									FIoRelativeUrl() = default;
 	bool							IsEmpty() const		{ return Url.IsEmpty(); }
 	int32							Len() const			{ return Url.Len(); }
 	const ANSICHAR*					ToString() const	{ return *Url; }
@@ -210,7 +206,7 @@ private:
 class FIoHttpRequest
 {
 public:
-					FIoHttpRequest() = default;
+							FIoHttpRequest() = default;
 	explicit 				FIoHttpRequest(FIoHttpRequestHandle InHandle)
 								: Handle(InHandle) { }
 							FIoHttpRequest(FIoHttpRequest&& Other)
@@ -236,7 +232,7 @@ enum class EIoHttpResponseFlags : uint8
 {
 	/** No additional flags. */
 	None			= 0,
-	/** Whether the response was cached. */ 
+	/** The response was retrieved from the cache. */ 
 	Cached			= (1 << 0),
 };
 ENUM_CLASS_FLAGS(EIoHttpResponseFlags);
@@ -245,13 +241,14 @@ ENUM_CLASS_FLAGS(EIoHttpResponseFlags);
 class FIoHttpResponse
 {
 public:
-					FIoHttpResponse() = default;
-	inline					FIoHttpResponse(const FIoBuffer& InBody, EIoErrorCode InErrorCode, uint32 InStatusCode, EIoHttpResponseFlags InFlags);
+							FIoHttpResponse() = default;
+	inline					FIoHttpResponse(const FIoBuffer& InBody, const FIoHash& InCacheKey, EIoErrorCode InErrorCode, uint32 InStatusCode, EIoHttpResponseFlags InFlags);
 							FIoHttpResponse(EIoErrorCode InErrorCode, uint32 InStatusCode)
 								: ErrorCode(InErrorCode)
 								, StatusCode(InStatusCode) { }
 
 	const FIoBuffer&		GetBody()		const { return Body; }
+	const FIoHash&			GetCacheKey()	const { return CacheKey; }
 	EIoErrorCode			GetErrorCode()	const { return ErrorCode; }
 	uint32					GetStatusCode()	const { return StatusCode; }
 	EIoHttpResponseFlags	GetFlags()		const { return Flags; }
@@ -261,13 +258,15 @@ public:
 
 private:
 	FIoBuffer				Body;
+	FIoHash					CacheKey;
 	EIoErrorCode			ErrorCode;
 	uint32					StatusCode = 0;
 	EIoHttpResponseFlags	Flags = EIoHttpResponseFlags::None;
 };
 
-FIoHttpResponse::FIoHttpResponse(const FIoBuffer& InBody, EIoErrorCode InErrorCode, uint32 InStatusCode, EIoHttpResponseFlags InFlags)
+FIoHttpResponse::FIoHttpResponse(const FIoBuffer& InBody, const FIoHash& InCacheKey, EIoErrorCode InErrorCode, uint32 InStatusCode, EIoHttpResponseFlags InFlags)
 	: Body(InBody)
+	, CacheKey(InCacheKey)
 	, ErrorCode(InErrorCode)
 	, StatusCode(InStatusCode)
 	, Flags(InFlags)
@@ -322,10 +321,12 @@ public:
 	UE_API static bool				IsInitialized();
 	UE_API static FIoStatus			Initialize(TSharedPtr<class IHttpIoDispatcher> Dispatcher);
 	UE_API static FIoStatus			Shutdown();
+
 	UE_API static FIoStatus			RegisterHostGroup(const FName& HostGroup, TConstArrayView<FAnsiString> HostNames, FAnsiStringView TestUrl);
 	UE_API static FIoStatus			RegisterHostGroup(const FName& HostGroup, FAnsiStringView HostName, FAnsiStringView TestUrl);
 	UE_API static bool				IsHostGroupRegistered(const FName& HostGroup);
 	UE_API static bool				IsHostGroupOk(const FName& HostGroup);
+
 	UE_API static FIoHttpBatch		NewBatch();
 	UE_API static FIoHttpRequest	Get(const FName& HostGroup,
 										const FIoRelativeUrl& RelativeUrl,
@@ -340,6 +341,10 @@ public:
 										const FIoRelativeUrl& RelativeUrl,
 										FIoHttpRequestCompleted&& OnCompleted);
 
+	/** */
+	UE_API static FIoStatus			CacheResponse(const FIoHttpResponse& Response);
+	UE_API static FIoStatus			EvictFromCache(const FIoHttpResponse& Response);
+
 	DECLARE_MULTICAST_DELEGATE_OneParam(FHostGroupRegistered, const FName&);
 	UE_API static FHostGroupRegistered& OnHostGroupRegistered();
 };
@@ -351,11 +356,13 @@ public:
 	using FHostGroupRegistered		= FHttpIoDispatcher::FHostGroupRegistered;
 
 	virtual							~IHttpIoDispatcher() = default;
-	virtual void					Shutdown() = 0;				
+	virtual void					Shutdown() = 0;
+	
 	virtual FIoStatus				RegisterHostGroup(const FName& HostGroup, TConstArrayView<FAnsiString> HostNames, FAnsiStringView TestUrl) = 0;
 	virtual bool					IsHostGroupRegistered(const FName& HostGroup) = 0;
 	virtual bool					IsHostGroupOk(const FName& HostGroup) = 0;
 	virtual FHostGroupRegistered&	OnHostGroupRegistered() = 0;
+
 	virtual FIoHttpRequestHandle	CreateRequest(
 										FIoHttpRequestHandle& First,
 										FIoHttpRequestHandle& Last,
@@ -370,6 +377,9 @@ public:
 	virtual void					UpdateRequestPriority(FIoHttpRequestHandle Handle, int32 NewPriority) = 0;
 	virtual EIoErrorCode			GetRequestStatus(FIoHttpRequestHandle Handle) = 0;
 	virtual void					ReleaseRequest(FIoHttpRequestHandle Handle) = 0;
+
+	virtual FIoStatus				CacheResponse(const FIoHttpResponse& Response) = 0;
+	virtual FIoStatus				EvictFromCache(const FIoHttpResponse& Response) = 0;
 };
 
 } // namespace UE
