[ChaosVD] Initial Support for Trace Relay using Message bus as data transport for the generated trace data.

diff --git a/Engine/Plugins/ChaosVD/Source/ChaosVD/Private/ChaosVDEngine.cpp b/Engine/Plugins/ChaosVD/Source/ChaosVD/Private/ChaosVDEngine.cpp
--- a/Engine/Plugins/ChaosVD/Source/ChaosVD/Private/ChaosVDEngine.cpp
+++ b/Engine/Plugins/ChaosVD/Source/ChaosVD/Private/ChaosVDEngine.cpp
@@ -322,6 +322,32 @@ bool FChaosVDEngine::ConnectToLiveSession_Direct(EChaosVDLoadRecordedDataMode Lo
 	return bSuccess;
 }
 
+bool FChaosVDEngine::ConnectToLiveSession_Relay(FGuid RemoteSessionID, EChaosVDLoadRecordedDataMode LoadingMode)
+{
+	FChaosVDTraceSessionDescriptor NewSessionFromFileDescriptor;
+
+	NewSessionFromFileDescriptor.bIsLiveSession = true;
+	
+	TSharedPtr<FChaosVDRecording> ExistingRecordingInstance = nullptr;
+	if (LoadingMode == EChaosVDLoadRecordedDataMode::MultiSource && PlaybackController)
+	{
+		ExistingRecordingInstance = PlaybackController->GetCurrentRecording().Pin();
+	}
+
+	NewSessionFromFileDescriptor.RemoteSessionID = RemoteSessionID;
+	NewSessionFromFileDescriptor.SessionName = FChaosVDModule::Get().GetTraceManager()->ConnectToLiveSession_Relay(RemoteSessionID, ExistingRecordingInstance);
+
+	bool bSuccess = false;
+
+	if (NewSessionFromFileDescriptor.IsValid())
+	{
+		OpenSession(NewSessionFromFileDescriptor, LoadingMode);
+		bSuccess = true;
+	}
+
+	return bSuccess;
+}
+
 void FChaosVDEngine::OpenSession(const FChaosVDTraceSessionDescriptor& SessionDescriptor, EChaosVDLoadRecordedDataMode LoadingMode)
 {
 	if (LoadingMode == EChaosVDLoadRecordedDataMode::SingleSource)
diff --git a/Engine/Plugins/ChaosVD/Source/ChaosVD/Private/Trace/ChaosVDTraceManager.cpp b/Engine/Plugins/ChaosVD/Source/ChaosVD/Private/Trace/ChaosVDTraceManager.cpp
--- a/Engine/Plugins/ChaosVD/Source/ChaosVD/Private/Trace/ChaosVDTraceManager.cpp
+++ b/Engine/Plugins/ChaosVD/Source/ChaosVD/Private/Trace/ChaosVDTraceManager.cpp
@@ -6,6 +6,9 @@
 #include "ChaosVDModule.h"
 #include "ChaosVDRuntimeModule.h"
 #include "ChaosVDSettingsManager.h"
+#include "Chaos/ChaosVDEngineEditorBridge.h"
+#include "Chaos/ChaosVDTraceRelayTransport.h"
+#include "Chaos/Serialization/SerializedDataBuffer.h"
 #include "Features/IModularFeatures.h"
 #include "HAL/PlatformFileManager.h"
 #include "Modules/ModuleManager.h"
@@ -20,6 +23,15 @@
 
 namespace Chaos::VD
 {
+	namespace CVars
+	{
+		static float SecondsToSleepWaitingForRelayData = 0.1f;
+		static FAutoConsoleVariableRef CVarUseCVDDynamicMeshGenerator(
+			TEXT("p.Chaos.VD.Tool.SecondsToSleepWaitingForRelayData"),
+			SecondsToSleepWaitingForRelayData,
+			TEXT("Time in seconds the trace analysis thread should wait to receive new data from a relay connection"));
+	}
+
 	FAsyncProgressNotification::~FAsyncProgressNotification()
 	{
 		if (NotificationHandle.IsValid())
@@ -82,6 +94,139 @@ namespace Chaos::VD
 
 		return true;
 	}
+
+	FRelayDataStream::FRelayDataStream(FGuid InRemoteSessionID)
+		: CancelEvent(FPlatformProcess::GetSynchEventFromPool(false)), RemoteSessionID(InRemoteSessionID)
+	{
+
+#if WITH_CHAOS_VISUAL_DEBUGGER
+		TSharedPtr<ITraceDataRelayTransport> RelayTransportInstance = FChaosVDEngineEditorBridge::Get().GetRelayTransportInstance();
+#else
+		TSharedPtr<ITraceDataRelayTransport> RelayTransportInstance = nullptr;
+#endif
+		if (ensure(RelayTransportInstance))
+		{
+			RelayTransportInstance->RegisterRelayDataReceiverForSessionID(RemoteSessionID, ITraceDataRelayTransport::FProcessReceivedRelayDataDelegate::CreateRaw(this, &FRelayDataStream::EnqueueRelayedData));
+		}
+	}
+
+	FRelayDataStream::~FRelayDataStream()
+	{
+		FPlatformProcess::ReturnSynchEventToPool(CancelEvent);
+
+#if WITH_CHAOS_VISUAL_DEBUGGER
+		TSharedPtr<ITraceDataRelayTransport> RelayTransportInstance = FChaosVDEngineEditorBridge::Get().GetRelayTransportInstance();
+#else
+		TSharedPtr<ITraceDataRelayTransport> RelayTransportInstance = nullptr;
+#endif
+
+		if (RelayTransportInstance)
+		{
+			RelayTransportInstance->UnregisterRelayDataReceiverForSessionID(RemoteSessionID);
+		}
+	}
+
+	void FRelayDataStream::EnqueueRelayedData(TArray<uint8>&& InTraceDataBuffer)
+	{
+		FSerializedDataBuffer DataBuffer(MoveTemp(InTraceDataBuffer));
+		DataQueue.Enqueue(MoveTemp(DataBuffer));
+	}
+
+	void FRelayDataStream::EnqueueRelayedData(const FChaosVDRelayTraceDataMessage& InTraceDataMessage)
+	{
+		check(InTraceDataMessage.InstanceId == RemoteSessionID);
+	
+		FChaosVDRelayTraceDataMessage& MutableMessage = const_cast<FChaosVDRelayTraceDataMessage&>(InTraceDataMessage);
+		EnqueueRelayedData(MoveTemp(MutableMessage.DataBuffer));
+	}
+
+	int32 FRelayDataStream::Read(void* Dest, uint32 DestSize)
+	{
+		if (DestSize == 0)
+		{
+			return 0;
+		}
+
+		if (Dest == nullptr)
+		{
+			return 0;
+		}
+
+		if (!ensure(CancelEvent))
+		{
+			return 0;
+		}
+
+		uint8* DestBuffer = static_cast<uint8*>(Dest);
+		uint32 BytesCopied = 0;
+
+		// We need to wait for some data to be queued. Returning 
+		// zero bytes read is interpreted as eof.
+		while (!CancelEvent->Wait(0, true))
+		{
+			if (!CurrentPacket)
+			{
+				CurrentPacket = DataQueue.Dequeue();
+				CurrentPacketOffset = 0;
+			}
+
+			// Consume as many queued packets as possible and copy into
+			// the destination buffer
+			while (CurrentPacket && BytesCopied < DestSize)
+			{
+				const TArray<uint8>& SrcBuffer = CurrentPacket->GetDataAsByteArrayRef();
+				if (SrcBuffer.IsEmpty())
+				{
+					// We should never get an empty packet. From this point on we can't trust the integrity of the data
+					// so we need to close the stream
+					Close();
+					break;
+				}
+
+				const uint32 ReceivedBufferSize = static_cast<uint32>(SrcBuffer.Num());
+
+				// Copy as much of the current buffer as possible into the remaining space 
+				// in the destination buffer
+				uint32 BytesToCopy = FMath::Min(ReceivedBufferSize - CurrentPacketOffset, DestSize - BytesCopied);
+				FMemory::Memcpy(&DestBuffer[BytesCopied], SrcBuffer.GetData() + CurrentPacketOffset, BytesToCopy);
+				BytesCopied += BytesToCopy;
+				CurrentPacketOffset += BytesToCopy;
+
+				if (CurrentPacketOffset == ReceivedBufferSize)
+				{
+					// Current packet buffer has been fully copied, goto next
+					CurrentPacket = DataQueue.Dequeue();
+					CurrentPacketOffset = 0;
+				}
+			}
+
+			// If at least some bytes have been written return those. Otherwise, 
+			// we'll sleep and wait for some data to arrive.
+			if (BytesCopied > 0)
+			{
+				break;
+			}
+
+			FPlatformProcess::Sleep(CVars::SecondsToSleepWaitingForRelayData);
+		}
+
+		return BytesCopied;
+	}
+
+	void FRelayDataStream::Close()
+	{
+		if (CancelEvent)
+		{
+			CancelEvent->Trigger();
+		}
+
+		bClosed = true;
+	}
+
+	bool FRelayDataStream::WaitUntilReady()
+	{		
+		return DataQueue.Peek() != nullptr; 
+	}
 }
 
 FChaosVDTraceManager::FChaosVDTraceManager() 
@@ -277,6 +422,13 @@ TUniquePtr<UE::Trace::IInDataStream> FChaosVDTraceManager::CreateFileDataStream(
 	return DataStream;
 }
 
+TUniquePtr<Chaos::VD::FRelayDataStream> FChaosVDTraceManager::CreateRelayDataStream(FGuid RemoteSessionID)
+{
+	TUniquePtr<Chaos::VD::FRelayDataStream> NewDataStream = MakeUnique<Chaos::VD::FRelayDataStream>(RemoteSessionID);
+	
+	return MoveTemp(NewDataStream);
+}
+
 void FChaosVDTraceManager::SetPendingExternalRecordingToProcess(const TSharedPtr<FChaosVDRecording>& ExistingRecordingPtr)
 {
 	TWeakPtr<FChaosVDRecording>& PendingRecording = FChaosVDTraceManagerThreadContext::Get().PendingExternalRecordingWeakPtr;
@@ -403,6 +555,23 @@ FString FChaosVDTraceManager::ConnectToLiveSession_Direct(uint16& OutSessionPort
 	}
 }
 
+FString FChaosVDTraceManager::ConnectToLiveSession_Relay(FGuid RemoteSessionID, const TSharedPtr<FChaosVDRecording>& ExistingRecordingPtr)
+{
+	FString RelaySessionName = FString::Printf(TEXT("RelaySession[%s]"), *RemoteSessionID.ToString());
+
+	TUniquePtr<Chaos::VD::FRelayDataStream> RelayDataStream = CreateRelayDataStream(RemoteSessionID);
+
+	if (ConnectToLiveSession_Internal(~0u, ExistingRecordingPtr, RelaySessionName, MoveTemp(RelayDataStream)))
+	{
+		return RelaySessionName;
+	}
+	else
+	{
+		UE_LOG(LogChaosVDEditor, Error, TEXT("[%hs] Failed start trace analysis with relay trace stream | [%s]"), __func__, *RelaySessionName);
+		return FString();
+	}
+}
+
 FString FChaosVDTraceManager::GetLocalTraceStoreDirPath()
 {
 	UE::Trace::FStoreClient* StoreClient = UE::Trace::FStoreClient::Connect(TEXT("localhost"));
diff --git a/Engine/Plugins/ChaosVD/Source/ChaosVD/Private/Widgets/SChaosVDMainTab.cpp b/Engine/Plugins/ChaosVD/Source/ChaosVD/Private/Widgets/SChaosVDMainTab.cpp
--- a/Engine/Plugins/ChaosVD/Source/ChaosVD/Private/Widgets/SChaosVDMainTab.cpp
+++ b/Engine/Plugins/ChaosVD/Source/ChaosVD/Private/Widgets/SChaosVDMainTab.cpp
@@ -914,6 +914,11 @@ bool SChaosVDMainTab::ConnectToLiveSession_Direct(EChaosVDLoadRecordedDataMode L
 	return GetChaosVDEngineInstance()->ConnectToLiveSession_Direct(LoadingMode);
 }
 
+bool SChaosVDMainTab::ConnectToLiveSession_Relay(FGuid RemoteSessionID, EChaosVDLoadRecordedDataMode LoadingMode) const
+{
+	return GetChaosVDEngineInstance()->ConnectToLiveSession_Relay(RemoteSessionID, LoadingMode);
+}
+
 bool SChaosVDMainTab::IsSupportedFile(const FString& InFilename)
 {
 	return InFilename.EndsWith(TEXT(".utrace")) || InFilename.EndsWith(TEXT(".cvdmulti"));
diff --git a/Engine/Plugins/ChaosVD/Source/ChaosVD/Private/Widgets/SChaosVDRecordingControls.cpp b/Engine/Plugins/ChaosVD/Source/ChaosVD/Private/Widgets/SChaosVDRecordingControls.cpp
--- a/Engine/Plugins/ChaosVD/Source/ChaosVD/Private/Widgets/SChaosVDRecordingControls.cpp
+++ b/Engine/Plugins/ChaosVD/Source/ChaosVD/Private/Widgets/SChaosVDRecordingControls.cpp
@@ -744,12 +744,20 @@ void SChaosVDRecordingControls::SetSessionRecordingState(bool bIsRecording, ECha
 	TSharedPtr<FChaosVDRemoteSessionsManager> RemoteSessionManager = FChaosVDEngineEditorBridge::Get().GetRemoteSessionsManager();
 	if (!RemoteSessionManager)
 	{
-		UE_LOG(LogChaosVDEditor, Error, TEXT("[%s] Session Manager is not available"), ANSI_TO_TCHAR(__FUNCTION__));
+		UE_LOG(LogChaosVDEditor, Error, TEXT("[%hs] Session Manager is not available"), __func__);
 		return;
 	}
 
 	if (bIsRecording)
 	{
+		const TSharedPtr<SChaosVDMainTab> MainTabSharedPtr = MainTabWeakPtr.Pin();
+		if (!ensure(MainTabSharedPtr))
+		{
+			// This should not be possible, but if we don't have a valid tab instance it will not be possible to properly handle any start recording attempt
+			UE_LOG(LogChaosVDEditor, Error, TEXT("[%hs] Main Tab instance is not available"), __func__);
+			return;
+		}
+
 		TArray<FString, TInlineAllocator<1>> RecordingArgs;
 
 		int32 RemainingRetries = 4;
@@ -807,18 +815,12 @@ void SChaosVDRecordingControls::SetSessionRecordingState(bool bIsRecording, ECha
 
 			if (RecordingParams.TransportMode == EChaosVDTransportMode::Memory)
 			{
-				// Using the memory transport or direct trace 
-				const TSharedPtr<SChaosVDMainTab> MainTabSharedPtr = MainTabWeakPtr.Pin();
-				if (!MainTabSharedPtr)
-				{
-					return;
-				}
-
 				// Memory Tracing mode requires opening a trace session in the editor side first before starting the recording
 				// so targets have something to connect to (Usually targets connect to a running trace server)
 				if (!MainTabSharedPtr->ConnectToLiveSession_Direct(CurrentLoadingMode))
 				{
 					UE_LOG(LogChaosVDEditor, Error, TEXT("[%hs] Failed to connect to live session using direct trace mode"), __func__);
+					RecordingAttemptFailedCallback();
 					return;
 				}
 
@@ -829,6 +831,19 @@ void SChaosVDRecordingControls::SetSessionRecordingState(bool bIsRecording, ECha
 
 				RemoteSessionManager->SendStartRecordingCommand(SessionInfoRef->Address, RecordingParams);
 			}
+			else if (RecordingParams.TransportMode == EChaosVDTransportMode::Relay)
+			{
+				// Relay Tracing mode requires opening a trace session in the editor side first before starting the recording
+				// so targets have something to connect to (Usually targets connect to a running trace server)
+				if (!MainTabSharedPtr->ConnectToLiveSession_Relay(SessionInfoRef->InstanceId, CurrentLoadingMode))
+				{
+					UE_LOG(LogChaosVDEditor, Error, TEXT("[%hs] Failed to connect to live session using relay trace mode"), __func__);
+					RecordingAttemptFailedCallback();
+					return;
+				}
+
+				RemoteSessionManager->SendStartRecordingCommand(SessionInfoRef->Address, RecordingParams);
+			}
 			else
 			{
 				RemoteSessionManager->SendStartRecordingCommand(SessionInfoRef->Address, RecordingParams);
diff --git a/Engine/Plugins/ChaosVD/Source/ChaosVD/Public/ChaosVDEngine.h b/Engine/Plugins/ChaosVD/Source/ChaosVD/Public/ChaosVDEngine.h
--- a/Engine/Plugins/ChaosVD/Source/ChaosVD/Public/ChaosVDEngine.h
+++ b/Engine/Plugins/ChaosVD/Source/ChaosVD/Public/ChaosVDEngine.h
@@ -76,6 +76,7 @@ public:
 
 	bool ConnectToLiveSession(uint32 SessionID, const FString& InSessionAddress, EChaosVDLoadRecordedDataMode LoadingMode = EChaosVDLoadRecordedDataMode::MultiSource);
 	bool ConnectToLiveSession_Direct(EChaosVDLoadRecordedDataMode LoadingMode = EChaosVDLoadRecordedDataMode::MultiSource);
+	bool ConnectToLiveSession_Relay(FGuid RemoteSessionID, EChaosVDLoadRecordedDataMode LoadingMode = EChaosVDLoadRecordedDataMode::MultiSource);
 
 	void OpenSession(const FChaosVDTraceSessionDescriptor& SessionDescriptor, EChaosVDLoadRecordedDataMode LoadingMode = EChaosVDLoadRecordedDataMode::MultiSource);
 
diff --git a/Engine/Plugins/ChaosVD/Source/ChaosVD/Public/Trace/ChaosVDTraceManager.h b/Engine/Plugins/ChaosVD/Source/ChaosVD/Public/Trace/ChaosVDTraceManager.h
--- a/Engine/Plugins/ChaosVD/Source/ChaosVD/Public/Trace/ChaosVDTraceManager.h
+++ b/Engine/Plugins/ChaosVD/Source/ChaosVD/Public/Trace/ChaosVDTraceManager.h
@@ -3,11 +3,14 @@
 #pragma once
 
 #include "ChaosVDModule.h"
+#include "Chaos/Serialization/SerializedDataBuffer.h"
+#include "Containers/SpscQueue.h"
 #include "Containers/StringFwd.h"
 #include "Framework/Notifications/NotificationManager.h"
 #include "Templates/SharedPointer.h"
 #include "Trace/StoreClient.h"
 
+struct FChaosVDRelayTraceDataMessage;
 enum class EChaosVDLoadRecordedDataMode : uint8;
 class FChaosVDEngine;
 class FChaosVDTraceModule;
@@ -41,6 +44,33 @@ namespace Chaos::VD
 		std::atomic<int32> CurrentProgress = 0;
 		FText NotificationTitle;
 	};
+
+	/** A trace data stream object that listens for data coming from a relay stream
+	 * via CVD's Engine Editor Bridge object */
+	class FRelayDataStream : public UE::Trace::IInDataStream
+	{
+	public:
+		explicit FRelayDataStream(FGuid InRemoteSessionID);
+
+		virtual ~FRelayDataStream() override;
+
+		void EnqueueRelayedData(TArray<uint8>&& InTraceDataBuffer);
+		void EnqueueRelayedData(const FChaosVDRelayTraceDataMessage& InTraceDataMessage);
+
+	private:
+		virtual int32 Read(void* Dest, uint32 DestSize) override;
+
+		virtual void Close() override;
+		virtual bool WaitUntilReady() override;
+
+		TSpscQueue<FSerializedDataBuffer> DataQueue;
+		FEvent*	CancelEvent = nullptr;
+		TOptional<FSerializedDataBuffer> CurrentPacket;
+		uint32 CurrentPacketOffset = 0;
+		bool bClosed = false;
+
+		FGuid RemoteSessionID = FGuid();
+	};
 }
 
 /**
@@ -50,6 +80,7 @@ struct FChaosVDTraceSessionDescriptor
 {
 	FString SessionName;
 	uint16 SessionPort = 0;
+	FGuid RemoteSessionID;
 	bool bIsLiveSession = false;
 
 	bool IsValid() const { return !SessionName.IsEmpty(); }
@@ -97,6 +128,14 @@ public:
 	 */
 	FString ConnectToLiveSession_Direct(uint16& OutSessionPort, const TSharedPtr<FChaosVDRecording>& ExistingRecordingPtr = nullptr);
 
+	/**
+	 * Opens up a relay trace stream and waits for data
+	 * @param RemoteSessionID Session ID to which we want to connect to
+	 * @param ExistingRecordingPtr Ptr to an existing recording instance to append the new CVD data to, if any
+	 * @return Session Name
+	 */
+	FString ConnectToLiveSession_Relay(FGuid RemoteSessionID, const TSharedPtr<FChaosVDRecording>& ExistingRecordingPtr = nullptr);
+
 	/** Returns the path to the local trace store */
 	FString GetLocalTraceStoreDirPath();
 
@@ -137,6 +176,7 @@ private:
 	bool ConnectToLiveSession_Internal(uint32 SessionID, const TSharedPtr<FChaosVDRecording>& ExistingRecordingPtr, const FString& InRequestedSessionName, UE::Trace::FStoreClient::FTraceData&& InTraceDataStream);
 
 	TUniquePtr<UE::Trace::IInDataStream> CreateFileDataStream(TUniquePtr<IFileHandle>&& InFileHandle);
+	TUniquePtr<Chaos::VD::FRelayDataStream> CreateRelayDataStream(FGuid RemoteSessionID);
 
 	/** Temporary workaround method to set an existing recording structure in CVD's trace provider before the trace analysis starts
 	 * in the Trace analysis thread.*/
diff --git a/Engine/Plugins/ChaosVD/Source/ChaosVD/Public/Widgets/SChaosVDMainTab.h b/Engine/Plugins/ChaosVD/Source/ChaosVD/Public/Widgets/SChaosVDMainTab.h
--- a/Engine/Plugins/ChaosVD/Source/ChaosVD/Public/Widgets/SChaosVDMainTab.h
+++ b/Engine/Plugins/ChaosVD/Source/ChaosVD/Public/Widgets/SChaosVDMainTab.h
@@ -108,6 +108,7 @@ public:
 	
 	bool ConnectToLiveSession(int32 SessionID, const FString& InSessionAddress, EChaosVDLoadRecordedDataMode LoadingMode) const;
 	bool ConnectToLiveSession_Direct(EChaosVDLoadRecordedDataMode LoadingMode) const;
+	bool ConnectToLiveSession_Relay(FGuid RemoteSessionID, EChaosVDLoadRecordedDataMode LoadingMode) const;
 
 	/**
 	 * Evaluates a filename to determine if it is a supported filetype
diff --git a/Engine/Source/Runtime/Experimental/Chaos/Public/Chaos/Serialization/SerializedDataBuffer.h b/Engine/Source/Runtime/Experimental/Chaos/Public/Chaos/Serialization/SerializedDataBuffer.h
--- a/Engine/Source/Runtime/Experimental/Chaos/Public/Chaos/Serialization/SerializedDataBuffer.h
+++ b/Engine/Source/Runtime/Experimental/Chaos/Public/Chaos/Serialization/SerializedDataBuffer.h
@@ -10,6 +10,10 @@ namespace Chaos
 	/** Wrapper around a data buffer used for serialization of chaos data */
 	struct FSerializedDataBuffer
 	{
+		FSerializedDataBuffer(TArray<uint8>&& InByteArray) : DataBuffer(MoveTemp(InByteArray))
+		{	
+		}
+
 		FSerializedDataBuffer(const TArray<uint8>& InByteArray) : DataBuffer(InByteArray)
 		{	
 		}
diff --git a/Engine/Source/Runtime/Experimental/ChaosSolverEngine/Private/Chaos/ChaosVDBusTraceRelayTransport.cpp b/Engine/Source/Runtime/Experimental/ChaosSolverEngine/Private/Chaos/ChaosVDBusTraceRelayTransport.cpp
new file mode 100644
--- /dev/null
+++ b/Engine/Source/Runtime/Experimental/ChaosSolverEngine/Private/Chaos/ChaosVDBusTraceRelayTransport.cpp
@@ -0,0 +1,205 @@
+// Copyright Epic Games, Inc. All Rights Reserved.
+
+#include "Chaos/ChaosVDBusTraceRelayTransport.h"
+
+#if WITH_CHAOS_VISUAL_DEBUGGER
+
+#include "MessageEndpoint.h"
+#include "MessageEndpointBuilder.h"
+#include "Chaos/ChaosVDRemoteSessionsManager.h"
+#include "Serialization/BitWriter.h"
+#include "Serialization/BufferArchive.h"
+#include "Trace/RelayTraceWriter.h"
+
+namespace Chaos::VD
+{
+
+void FBusTraceRelayTransport::Initialize()
+{
+	FChaosVDRuntimeModule& CVDRuntimeModule = FChaosVDRuntimeModule::Get();
+	CVDRuntimeModule.RegisterExternalRelayExecutor(FChaosVDRelayTraceDataDelegate::CreateRaw(this, &FBusTraceRelayTransport::RelayTraceDataFromWriter));
+
+	RecordingStartedHandle = CVDRuntimeModule.RegisterRecordingStartedCallback(FChaosVDRecordingStateChangedDelegate::FDelegate::CreateRaw(this, &FBusTraceRelayTransport::HandleCVDRecordingStarted));
+	RecordingStoppedHandle = CVDRuntimeModule.RegisterRecordingStopCallback(FChaosVDRecordingStateChangedDelegate::FDelegate::CreateRaw(this, &FBusTraceRelayTransport::HandleCVDRecordingStopped));
+
+	if (CVDRuntimeModule.IsRecording())
+	{
+		HandleCVDRecordingStarted();
+	}
+}
+
+void FBusTraceRelayTransport::Shutdown()
+{
+	if (MessageEndpoint)
+	{
+		MessageEndpoint->Disable();
+		MessageEndpoint = nullptr;
+	}
+
+	MessageBus = nullptr;
+
+	if (FChaosVDRuntimeModule::IsLoaded())
+	{
+		FChaosVDRuntimeModule& CVDRuntimeModule = FChaosVDRuntimeModule::Get();
+		CVDRuntimeModule.UnregisterCurrentExternalRelayExecutor();
+		CVDRuntimeModule.RemoveRecordingStartedCallback(RecordingStartedHandle);
+		CVDRuntimeModule.RemoveRecordingStopCallback(RecordingStoppedHandle);
+	}
+	
+	FTSTicker::RemoveTicker(PublishStatusTickerDelegate);
+}
+
+void FBusTraceRelayTransport::SetMessageBus(TWeakPtr<IMessageBus> NewMessageBus)
+{
+	if (MessageEndpoint)
+	{
+		MessageEndpoint->Disable();
+		MessageEndpoint = nullptr;
+	}
+
+	MessageBus = NewMessageBus;
+
+	TSharedPtr<IMessageBus> MessageBusPtr = MessageBus.Pin();
+	if (!MessageBusPtr)
+	{
+		return;
+	}
+
+	FMessageEndpointBuilder EndPointBuilder = FMessageEndpoint::Builder(FName("TraceRelayTransportEndPoint"), MessageBusPtr.ToSharedRef());
+	
+#if WITH_EDITOR
+		EndPointBuilder.Handling<FChaosVDRelayTraceDataMessage>(this, &FBusTraceRelayTransport::HandleRelayTraceDataReceivedMessage);
+		EndPointBuilder.Handling<FChaosVDRelayTraceStatusMessage>(this, &FBusTraceRelayTransport::HandleRelayTraceStatusMessage);
+#endif
+
+	MessageEndpoint = EndPointBuilder.Build();
+
+	if (!MessageEndpoint)
+	{
+		return;
+	}
+
+#if WITH_EDITOR
+		MessageEndpoint->Subscribe<FChaosVDRelayTraceDataMessage>();
+		MessageEndpoint->Subscribe<FChaosVDRelayTraceStatusMessage>();
+#endif
+
+	MessageEndpoint->Enable();
+}
+
+void FBusTraceRelayTransport::RegisterRelayDataReceiverForSessionID(FGuid RemoteSessionID, const FProcessReceivedRelayDataDelegate& InProcessDataDelegate)
+{
+	if (!ensure(!ReceivedDataHandlerBySessionID.Contains(RemoteSessionID)))
+	{
+		return;
+	}
+
+	ReceivedDataHandlerBySessionID.Add(RemoteSessionID, InProcessDataDelegate);
+}
+
+void FBusTraceRelayTransport::UnregisterRelayDataReceiverForSessionID(FGuid RemoteSessionID)
+{
+	ReceivedDataHandlerBySessionID.Remove(RemoteSessionID);
+}
+
+bool FBusTraceRelayTransport::PublishRelayStatus(float DeltaTime)
+{
+	constexpr bool bContinueTicking = true;
+	if (!MessageEndpoint)
+	{
+		// The message bus endpoint might be created later, so keep ticking
+		return bContinueTicking;
+	}
+
+	MessageEndpoint->Publish(FMessageEndpoint::MakeMessage<FChaosVDRelayTraceStatusMessage>(LastKnownStatus), EMessageScope::Network);
+
+	return bContinueTicking;
+}
+
+void FBusTraceRelayTransport::UpdateLastKnownRelayStatus(const FRelayTraceWriter& InRelayTraceWriter, ERelayThrottlingState ThrottlingState)
+{
+	LastKnownStatus.InstanceId = FChaosVDRemoteSessionsManager::LocalSessionID;
+	LastKnownStatus.QueuedDataBytesNum = InRelayTraceWriter.GetQueuedBytesNum();
+	LastKnownStatus.ThrottlingState = ThrottlingState;
+}
+
+void FBusTraceRelayTransport::RelayTraceDataFromWriter(FRelayTraceWriter& InRelayTraceWriter)
+{
+	while (InRelayTraceWriter.HasPendingBunches())
+	{
+		if (ExternalDataThrottleDelegate.IsBound())
+		{
+			if (!ExternalDataThrottleDelegate.Execute(InRelayTraceWriter.NextPendingBunchSize()))
+			{
+				UpdateLastKnownRelayStatus(InRelayTraceWriter, ERelayThrottlingState::Active);
+				return;
+			}
+		}
+
+		TOptional<TUniquePtr<FBufferArchive>> PendingBytesBunch = InRelayTraceWriter.DequeuePendingBunch();
+
+		if (PendingBytesBunch.IsSet())
+		{
+			FChaosVDRelayTraceDataMessage DataPacket;
+			DataPacket.InstanceId = FChaosVDRemoteSessionsManager::LocalSessionID;
+			DataPacket.DataBuffer = TArray<uint8>(PendingBytesBunch.GetValue()->GetData(), PendingBytesBunch.GetValue()->Num());
+
+			SendRelayedTraceData(DataPacket);
+		}
+	}
+
+	UpdateLastKnownRelayStatus(InRelayTraceWriter, ERelayThrottlingState::Inactive);
+}
+
+void FBusTraceRelayTransport::SendRelayedTraceData(const FChaosVDRelayTraceDataMessage& InRelayData)
+{
+	if (ensure(MessageEndpoint))
+	{
+		MessageEndpoint->Publish(FMessageEndpoint::MakeMessage<FChaosVDRelayTraceDataMessage>(InRelayData), EMessageScope::Network);
+	}
+}
+
+void FBusTraceRelayTransport::HandleRelayTraceDataReceivedMessage(const FChaosVDRelayTraceDataMessage& InMessage, const TSharedRef<IMessageContext>& InContext)
+{
+	ForwardReceivedTraceData(InMessage);
+}
+
+void FBusTraceRelayTransport::ForwardReceivedTraceData(const FChaosVDRelayTraceDataMessage& InTraceDataPacket)
+{
+	FProcessReceivedRelayDataDelegate* FoundDataHandler = ReceivedDataHandlerBySessionID.Find(InTraceDataPacket.InstanceId);
+
+	if (!ensure(FoundDataHandler))
+	{
+		return;
+	}
+
+	FoundDataHandler->ExecuteIfBound(InTraceDataPacket);
+}
+
+void FBusTraceRelayTransport::HandleRelayTraceStatusMessage(const FChaosVDRelayTraceStatusMessage& InMessage, const TSharedRef<IMessageContext>& InContext)
+{
+	UE_LOG(LogChaosVDRemoteSession, Verbose, TEXT("Remaining Trace Data to be relayed [%lld] | Throttling State [%s]"), InMessage.QueuedDataBytesNum, *UEnum::GetValueAsString(InMessage.ThrottlingState));
+}
+
+void FBusTraceRelayTransport::HandleCVDRecordingStarted()
+{
+	constexpr float DelayBetweenTicks = 2.0f;
+	PublishStatusTickerDelegate = FTSTicker::GetCoreTicker().AddTicker(FTickerDelegate::CreateRaw(this, &FBusTraceRelayTransport::PublishRelayStatus), DelayBetweenTicks);
+}
+
+void FBusTraceRelayTransport::HandleCVDRecordingStopped()
+{
+	FTSTicker::RemoveTicker(PublishStatusTickerDelegate);
+}
+
+void FBusTraceRelayTransport::RegisterExternalDataThrottleHandler(const FCanSendDataDelegate& InExternalThrottleHandler)
+{
+	ExternalDataThrottleDelegate = InExternalThrottleHandler;
+}
+
+void FBusTraceRelayTransport::UnregisterExternalDataThrottleHandler()
+{
+	ExternalDataThrottleDelegate = FCanSendDataDelegate();
+}
+}
+#endif // WITH_CHAOS_VISUAL_DEBUGGER
\ No newline at end of file
diff --git a/Engine/Source/Runtime/Experimental/ChaosSolverEngine/Private/Chaos/ChaosVDEngineEditorBridge.cpp b/Engine/Source/Runtime/Experimental/ChaosSolverEngine/Private/Chaos/ChaosVDEngineEditorBridge.cpp
--- a/Engine/Source/Runtime/Experimental/ChaosSolverEngine/Private/Chaos/ChaosVDEngineEditorBridge.cpp
+++ b/Engine/Source/Runtime/Experimental/ChaosSolverEngine/Private/Chaos/ChaosVDEngineEditorBridge.cpp
@@ -7,6 +7,7 @@
 #if WITH_CHAOS_VISUAL_DEBUGGER
 
 #include "ChaosVDRuntimeModule.h"
+#include "Chaos/ChaosVDBusTraceRelayTransport.h"
 #include "ChaosVisualDebugger/ChaosVDTraceMacros.h"
 #include "ChaosVisualDebugger/ChaosVisualDebuggerTrace.h"
 #include "Components/PrimitiveComponent.h"
@@ -259,18 +260,38 @@ FChaosVDParticleMetadata FChaosVDEngineEditorBridge::GenerateParticleMetadata(co
 	return Metadata;
 }
 
+void FChaosVDEngineEditorBridge::InitializeTraceRelayInstance()
+{
+	// TODO: Add a way for a CVD extension define the relay instance type
+	TSharedPtr<Chaos::VD::FBusTraceRelayTransport> BusRelayDataTransport = MakeShared<Chaos::VD::FBusTraceRelayTransport>();
+	BusRelayDataTransport->SetMessageBus(RemoteSessionsManager->GetMessageBusInstance());
+	RemoteSessionsManager->OnMessagingInitialized().AddLambda([WeakRelayInstance = BusRelayDataTransport.ToWeakPtr()] (TSharedPtr<IMessageBus> NewMessageBus, TSharedPtr<FMessageEndpoint> NewEndpoint)
+		{
+			if (TSharedPtr<Chaos::VD::FBusTraceRelayTransport> RelayInstance = WeakRelayInstance.Pin())
+			{
+				RelayInstance->SetMessageBus(NewMessageBus);
+			}
+		});
+
+	BusRelayDataTransport->Initialize();
+	
+	RelayTraceDataTransportInstance = BusRelayDataTransport;
+}
+
 void FChaosVDEngineEditorBridge::Initialize()
 {
-	FChaosVDRuntimeModule& Module = FChaosVDRuntimeModule::Get();
+	FChaosVDRuntimeModule& CVDRuntimeModule = FChaosVDRuntimeModule::Get();
 
-	FChaosVDTraceDetails TraceDetails = Module.GetCurrentTraceSessionDetails();
+	FChaosVDTraceDetails TraceDetails = CVDRuntimeModule.GetCurrentTraceSessionDetails();
 
 	RemoteSessionsManager->Initialize();
 
-	RecordingStartedHandle = FChaosVDRuntimeModule::Get().RegisterRecordingStartedCallback(FChaosVDRecordingStateChangedDelegate::FDelegate::CreateRaw(this, &FChaosVDEngineEditorBridge::HandleCVDRecordingStarted));
-	PostRecordingStartedHandle = FChaosVDRuntimeModule::Get().RegisterPostRecordingStartedCallback(FChaosVDRecordingStateChangedDelegate::FDelegate::CreateRaw(this, &FChaosVDEngineEditorBridge::HandleCVDPostRecordingStarted));
-	RecordingStoppedHandle = FChaosVDRuntimeModule::Get().RegisterRecordingStopCallback(FChaosVDRecordingStateChangedDelegate::FDelegate::CreateRaw(this, &FChaosVDEngineEditorBridge::HandleCVDRecordingStopped));
-	RecordingStartFailedHandle = FChaosVDRuntimeModule::Get().RegisterRecordingStartFailedCallback(FChaosVDRecordingStartFailedDelegate::FDelegate::CreateRaw(this, &FChaosVDEngineEditorBridge::HandleCVDRecordingStartFailed));
+	InitializeTraceRelayInstance();
+
+	RecordingStartedHandle = CVDRuntimeModule.RegisterRecordingStartedCallback(FChaosVDRecordingStateChangedDelegate::FDelegate::CreateRaw(this, &FChaosVDEngineEditorBridge::HandleCVDRecordingStarted));
+	PostRecordingStartedHandle = CVDRuntimeModule.RegisterPostRecordingStartedCallback(FChaosVDRecordingStateChangedDelegate::FDelegate::CreateRaw(this, &FChaosVDEngineEditorBridge::HandleCVDPostRecordingStarted));
+	RecordingStoppedHandle = CVDRuntimeModule.RegisterRecordingStopCallback(FChaosVDRecordingStateChangedDelegate::FDelegate::CreateRaw(this, &FChaosVDEngineEditorBridge::HandleCVDRecordingStopped));
+	RecordingStartFailedHandle = CVDRuntimeModule.RegisterRecordingStartFailedCallback(FChaosVDRecordingStartFailedDelegate::FDelegate::CreateRaw(this, &FChaosVDEngineEditorBridge::HandleCVDRecordingStartFailed));
 
 	FChaosVisualDebuggerTrace::RegisterExternalParticleDebugNameGenerator(Chaos::VD::FRecordingSessionState::FParticleMetaDataGeneratorDelegate::CreateStatic(&FChaosVDEngineEditorBridge::GenerateParticleMetadata));
 
@@ -282,15 +303,24 @@ void FChaosVDEngineEditorBridge::Initialize()
 
 	// If we were already recording, we need to make sure we run the initialization step to setup the session broadcast ticker
 	// and the collision channel serialization
-	if (FChaosVDRuntimeModule::Get().IsRecording())
+	if (CVDRuntimeModule.IsRecording())
 	{
 		HandleCVDRecordingStarted();
 		HandleCVDPostRecordingStarted();
 	}
+
+	bIsInitialized = true;
+
+	EngineEditorBridgeInitializedDelegate.Broadcast();
 }
 
 void FChaosVDEngineEditorBridge::TearDown()
 {
+	if (RelayTraceDataTransportInstance)
+	{
+		RelayTraceDataTransportInstance->Shutdown();
+	}
+
 #if WITH_EDITOR
 	FWorldDelegates::OnPIEStarted.Remove(PIEStartedHandle);
 	PIEStartedHandle.Reset();
@@ -299,27 +329,34 @@ void FChaosVDEngineEditorBridge::TearDown()
 	DeferredShowMessageOnScreenHandle.Reset();
 	
 	FTSTicker::RemoveTicker(RecordingStatusUpdateHandle);
+	RecordingStatusUpdateHandle.Reset();
 
 	// Note: This works during engine shutdown because the Module Manager doesn't free the dll on module unload to account for use cases like this
 	// If this appears in a callstack crash it means that assumption changed or was not correct to begin with.
 	// A possible solution is just check if the module is loaded querying the module manager just using the module's name
 	if (FChaosVDRuntimeModule::IsLoaded())
 	{
-		FChaosVDRuntimeModule::Get().RemoveRecordingStartedCallback(RecordingStartedHandle);
+
+		FChaosVDRuntimeModule& CVDRuntimeModule = FChaosVDRuntimeModule::Get();
+		CVDRuntimeModule.RemoveRecordingStartedCallback(RecordingStartedHandle);
 		RecordingStartedHandle.Reset();
-		FChaosVDRuntimeModule::Get().RemovePostRecordingStartedCallback(PostRecordingStartedHandle);
+		CVDRuntimeModule.RemovePostRecordingStartedCallback(PostRecordingStartedHandle);
 		PostRecordingStartedHandle.Reset();
-		FChaosVDRuntimeModule::Get().RemoveRecordingStopCallback(RecordingStoppedHandle);
+		CVDRuntimeModule.RemoveRecordingStopCallback(RecordingStoppedHandle);
 		RecordingStoppedHandle.Reset();
+		CVDRuntimeModule.RemoveRecordingStartFailedCallback(RecordingStartFailedHandle);
+		RecordingStartFailedHandle.Reset();
 
 		// Make sure of removing the message from the screen in case the recording didn't quite stopped yet
-		if (FChaosVDRuntimeModule::Get().IsRecording())
+		if (CVDRuntimeModule.IsRecording())
 		{
 			HandleCVDRecordingStopped();
 		}
 
 		RemoteSessionsManager->Shutdown();
 	}
+
+	bIsInitialized = false;
 }
 #else
 
diff --git a/Engine/Source/Runtime/Experimental/ChaosSolverEngine/Private/Chaos/ChaosVDRemoteSessionsManager.cpp b/Engine/Source/Runtime/Experimental/ChaosSolverEngine/Private/Chaos/ChaosVDRemoteSessionsManager.cpp
--- a/Engine/Source/Runtime/Experimental/ChaosSolverEngine/Private/Chaos/ChaosVDRemoteSessionsManager.cpp
+++ b/Engine/Source/Runtime/Experimental/ChaosSolverEngine/Private/Chaos/ChaosVDRemoteSessionsManager.cpp
@@ -19,6 +19,7 @@ const FGuid FChaosVDRemoteSessionsManager::AllRemoteClientsWrapperGUID = FGuid::
 const FGuid FChaosVDRemoteSessionsManager::AllSessionsWrapperGUID = FGuid::NewGuid();
 const FGuid FChaosVDRemoteSessionsManager::CustomSessionsWrapperGUID = FGuid::NewGuid();
 const FGuid FChaosVDRemoteSessionsManager::InvalidSessionGUID = FGuid();
+const FGuid FChaosVDRemoteSessionsManager::LocalSessionID = FApp::GetInstanceId();
 
 const FString FChaosVDRemoteSessionsManager::LocalEditorSessionName = TEXT("Local Editor");
 const FGuid FChaosVDRemoteSessionsManager::LocalEditorSessionID = IsController() ? FApp::GetInstanceId() : InvalidSessionGUID;
@@ -144,7 +145,6 @@ TSharedPtr<FMessageEndpoint> FChaosVDRemoteSessionsManager::CreateEndPoint(const
 		EndPointBuilder.Handling<FChaosVDSessionPong>(this, &FChaosVDRemoteSessionsManager::HandleSessionPongMessage)
 						.Handling<FChaosVDRecordingStatusMessage>(this, &FChaosVDRemoteSessionsManager::HandleRecordingStatusUpdateMessage)
 						.Handling<FChaosVDFullSessionInfoResponseMessage>(this, &FChaosVDRemoteSessionsManager::HandleFullSessionStateResponseMessage)
-						.Handling<FChaosVDChannelStateChangeResponseMessage>(this, &FChaosVDRemoteSessionsManager::HandleChangeDataChannelStateResponseMessage)
 						.Handling<FChaosVDChannelStateChangeResponseMessage>(this, &FChaosVDRemoteSessionsManager::HandleChangeDataChannelStateResponseMessage);
 	}
 
@@ -185,6 +185,17 @@ void FChaosVDRemoteSessionsManager::InitializeMessagingSystem(const TSharedPtr<I
 	MessageEndpoint->Subscribe<FChaosVDStopRecordingCommandMessage>();
 	MessageEndpoint->Subscribe<FChaosVDChannelStateChangeCommandMessage>();
 	MessageEndpoint->Subscribe<FChaosVDFullSessionInfoRequestMessage>();
+
+	if (bInitialized)
+	{
+		MessageEndpoint->Enable();
+	}
+	else
+	{
+		MessageEndpoint->Disable();
+	}
+
+	MessagingInitializedDelegate.Broadcast(InMessageBus, MessageEndpoint);
 }
 
 void FChaosVDRemoteSessionsManager::ShutdownMessagingSystem()
@@ -217,10 +228,18 @@ void FChaosVDRemoteSessionsManager::Initialize()
 	ActiveSessionsByInstanceId.Add(AllRemoteClientsWrapperGUID, CreatedWrapperSessionInfo(AllRemoteClientsWrapperGUID, AllRemoteClientsTargetName));
 	ActiveSessionsByInstanceId.Add(AllSessionsWrapperGUID, CreatedWrapperSessionInfo(AllSessionsWrapperGUID, AllSessionsTargetName));
 	ActiveSessionsByInstanceId.Add(CustomSessionsWrapperGUID, CreatedWrapperSessionInfo(CustomSessionsWrapperGUID, CustomSessionsTargetName));
+
+	if (MessageEndpoint)
+	{
+		MessageEndpoint->Enable();
+	}
+
+	bInitialized = true;
 }
 
 void FChaosVDRemoteSessionsManager::Shutdown()
 {
+	bInitialized = false;
 	ShutdownMessagingSystem();
 }
 
@@ -235,6 +254,8 @@ void FChaosVDRemoteSessionsManager::EnumerateMessageTypes(const FVisitorFunction
 	InVisitor(FChaosVDFullSessionInfoRequestMessage::StaticStruct());
 	InVisitor(FChaosVDFullSessionInfoResponseMessage::StaticStruct());
 	InVisitor(FChaosVDChannelStateChangeResponseMessage::StaticStruct());
+	InVisitor(FChaosVDRelayTraceDataMessage::StaticStruct());
+	InVisitor(FChaosVDRelayTraceStatusMessage::StaticStruct());
 }
 
 void FChaosVDRemoteSessionsManager::StartSessionDiscovery()
@@ -279,7 +300,7 @@ void FChaosVDRemoteSessionsManager::SendStartRecordingCommand(const FMessageAddr
 {
 	if (!MessageEndpoint)
 	{
-		UE_LOG(LogChaosVDRemoteSession, Error, TEXT("[%s] Failed to send command | Invalid endpoint."), ANSI_TO_TCHAR(__FUNCTION__));
+		UE_LOG(LogChaosVDRemoteSession, Error, TEXT("[%hs] Failed to send command | Invalid endpoint."), __func__);
 		return;
 	}
 
@@ -487,10 +508,10 @@ void FChaosVDRemoteSessionsManager::HandleRecordingStatusUpdateMessage(const FCh
 void FChaosVDRemoteSessionsManager::HandleRecordingStartCommandMessage(const FChaosVDStartRecordingCommandMessage& InMessage, const TSharedRef<IMessageContext>& InContext)
 {
 #if WITH_CHAOS_VISUAL_DEBUGGER
-	UE_AUTORTFM_ONCOMMIT(InMessage)
+	UE_AUTORTFM_ONCOMMIT(InMessage, InContext)
 	{
 		FChaosVisualDebuggerTrace::OverrideDefaultEnabledDataChannels(InMessage.DataChannelsEnabledOverrideList);
-		FChaosVDRuntimeModule::Get().StartRecording(InMessage);
+		FChaosVDRuntimeModule::Get().StartRecording(InMessage, InContext->GetSender().ToString());
 	};
 
 #endif
diff --git a/Engine/Source/Runtime/Experimental/ChaosSolverEngine/Public/Chaos/ChaosVDBusTraceRelayTransport.h b/Engine/Source/Runtime/Experimental/ChaosSolverEngine/Public/Chaos/ChaosVDBusTraceRelayTransport.h
new file mode 100644
--- /dev/null
+++ b/Engine/Source/Runtime/Experimental/ChaosSolverEngine/Public/Chaos/ChaosVDBusTraceRelayTransport.h
@@ -0,0 +1,119 @@
+// Copyright Epic Games, Inc. All Rights Reserved.
+
+#pragma once
+#if WITH_CHAOS_VISUAL_DEBUGGER
+#include "ChaosVDRemoteSessionsManager.h"
+#include "ChaosVDTraceRelayTransport.h"
+#include "Containers/Map.h"
+#include "Misc/Guid.h"
+#include "Templates/SharedPointer.h"
+
+enum class ERelayThrottlingState : uint8;
+class IMessageContext;
+class FMessageEndpoint;
+class IMessageBus;
+
+namespace Chaos::VD
+{
+/**
+ * CVD Trace Relay implementation that sends the trace data using Message bus messages
+ */
+class FBusTraceRelayTransport : public ITraceDataRelayTransport
+{
+public:
+
+	virtual ~FBusTraceRelayTransport() override = default;
+
+	virtual void Initialize() override;
+	virtual void Shutdown() override;
+
+	/** Sets the message bus instance to use
+	 * @param NewMessageBus Weakptr to the desired message bus instance
+	 */
+	void SetMessageBus(TWeakPtr<IMessageBus> NewMessageBus);
+
+	/**
+	 * Registers a Relayed data receiver callback, to which any received trace data for a specific session id, will be fowarded to.
+	 * @param RemoteSessionID Session ID of the data that needs to handled
+	 * @param InProcessDataDelegate Callback to forward the data to
+	 */
+	virtual void RegisterRelayDataReceiverForSessionID(FGuid RemoteSessionID, const FProcessReceivedRelayDataDelegate& InProcessDataDelegate) override;
+	virtual void UnregisterRelayDataReceiverForSessionID(FGuid RemoteSessionID) override;
+
+	/**
+	 * Process pending data in the provided relay writer instance and sends it if able
+	 * @param InRelayTraceWriter Reference to the active relay trace writer instance
+	 */
+	virtual void RelayTraceDataFromWriter(FRelayTraceWriter& InRelayTraceWriter) override;
+
+	/**
+	 * Registers an external callback that will be used to know if we need to stop sending data temporarily
+	 * @param InExternalThrottleHandler External callback that determines if we can send data or not
+	 */
+	virtual void RegisterExternalDataThrottleHandler(const FCanSendDataDelegate& InExternalThrottleHandler) override;
+
+	/**
+	 *  Deregisters an external callback that will be used to know if we need to stop sending data temporarily
+	 */
+	virtual void UnregisterExternalDataThrottleHandler() override;
+
+protected:
+
+	TWeakPtr<IMessageBus> MessageBus;
+	TSharedPtr<FMessageEndpoint> MessageEndpoint;
+
+	/**
+	 * Sends the data contained in the message via message bus
+	 * @param InRelayData Message containing the data to send
+	 */
+	virtual void SendRelayedTraceData(const FChaosVDRelayTraceDataMessage& InRelayData) override;
+
+	/**
+	 * Handles any new trace data messaged that was sent via message bus
+	 * @param InMessage Message containing the data
+	 * @param InContext Message bus context
+	 */
+	void HandleRelayTraceDataReceivedMessage(const FChaosVDRelayTraceDataMessage& InMessage, const TSharedRef<IMessageContext>& InContext);
+
+	/**
+	 * Forwards a trace data message
+	 * @param InTraceDataPacket Data to forwards to the correct data handler based on the message's session id
+	 */
+	virtual void ForwardReceivedTraceData(const FChaosVDRelayTraceDataMessage& InTraceDataPacket) override;
+
+	/**
+	 * Handles a received status message
+	 * @param InMessage Received status message
+	 * @param InContext Message bus context
+	 */
+	void HandleRelayTraceStatusMessage(const FChaosVDRelayTraceStatusMessage& InMessage, const TSharedRef<IMessageContext>& InContext);
+
+	void HandleCVDRecordingStarted();
+	void HandleCVDRecordingStopped();
+
+	/**
+	 * Broadcasts the last known relay status
+	 * @param DeltaTime Current delta time
+	 * @return true if we should keep ticking the callback that sends the status
+	 */
+	bool PublishRelayStatus(float DeltaTime);
+
+	/**
+	 * Updates the last known relay status with info from the trace writer and if we are currently throttling the data transfer 
+	 * @param InRelayTraceWriter Relay Writer instance
+	 * @param ThrottlingState Current throttling state
+	 */
+	void UpdateLastKnownRelayStatus(const FRelayTraceWriter& InRelayTraceWriter, ERelayThrottlingState ThrottlingState);
+
+	TMap<FGuid, FProcessReceivedRelayDataDelegate> ReceivedDataHandlerBySessionID;
+
+	FChaosVDRelayTraceStatusMessage LastKnownStatus;
+
+	FCanSendDataDelegate ExternalDataThrottleDelegate;
+	FDelegateHandle RecordingStartedHandle;
+	FDelegateHandle RecordingStoppedHandle;
+	FTSTicker::FDelegateHandle PublishStatusTickerDelegate;
+};
+}
+#endif // WITH_CHAOS_VISUAL_DEBUGGER
+
diff --git a/Engine/Source/Runtime/Experimental/ChaosSolverEngine/Public/Chaos/ChaosVDEngineEditorBridge.h b/Engine/Source/Runtime/Experimental/ChaosSolverEngine/Public/Chaos/ChaosVDEngineEditorBridge.h
--- a/Engine/Source/Runtime/Experimental/ChaosSolverEngine/Public/Chaos/ChaosVDEngineEditorBridge.h
+++ b/Engine/Source/Runtime/Experimental/ChaosSolverEngine/Public/Chaos/ChaosVDEngineEditorBridge.h
@@ -3,6 +3,11 @@
 #pragma once
 #include "Chaos/ChaosVDRemoteSessionsManager.h"
 
+namespace Chaos::VD
+{
+	class ITraceDataRelayTransport;
+}
+
 #if WITH_CHAOS_VISUAL_DEBUGGER
 
 #include "Chaos/ParticleHandleFwd.h"
@@ -10,8 +15,9 @@
 #include "Delegates/IDelegateInstance.h"
 #include "HAL/Platform.h"
 
-class UBodySetup;
 class IPhysicsProxyBase;
+class UBodySetup;
+class UChaosVDTraceDataChannel;
 class UGameInstance;
 
 namespace Chaos
@@ -51,6 +57,21 @@ public:
 
 	using FChaosVDDataDataChannel = Chaos::VisualDebugger::FChaosVDOptionalDataChannel;
 
+	TSharedPtr<Chaos::VD::ITraceDataRelayTransport> GetRelayTransportInstance()
+	{
+		return RelayTraceDataTransportInstance;
+	}
+
+	bool IsInitialized() const
+	{
+		return bIsInitialized;
+	}
+
+	FSimpleMulticastDelegate& OnEngineEditorBridgeInitialized()
+	{
+		return EngineEditorBridgeInitializedDelegate;
+	}
+
 private:
 	bool AddOnScreenRecordingMessage(float DummyDeltaTime = 0.1f);
 	void RemoveOnScreenRecordingMessage();
@@ -67,6 +88,7 @@ private:
 	bool BroadcastSessionStatus(float DeltaTime);
 
 	static FChaosVDParticleMetadata GenerateParticleMetadata(const IPhysicsProxyBase* ParticleProxy, const Chaos::FGeometryParticleHandle* ParticleHandle);
+	void InitializeTraceRelayInstance();
 
 	FDelegateHandle RecordingStartedHandle;
 	FDelegateHandle PostRecordingStartedHandle;
@@ -78,10 +100,16 @@ private:
 	FDelegateHandle PIEStartedHandle;
 #endif
 
+	TSharedPtr<Chaos::VD::ITraceDataRelayTransport> RelayTraceDataTransportInstance;
+
 	TSharedRef<FChaosVDRemoteSessionsManager> RemoteSessionsManager;
 	
 	FTSTicker::FDelegateHandle RecordingStatusUpdateHandle;
 	FTSTicker::FDelegateHandle DeferredShowMessageOnScreenHandle;
+
+	FSimpleMulticastDelegate EngineEditorBridgeInitializedDelegate;
+
+	bool bIsInitialized = false;
 };
 #else
 
@@ -94,6 +122,11 @@ public:
 	{
 		return nullptr;
 	}
+
+	TSharedPtr<Chaos::VD::ITraceDataRelayTransport> GetRelayTransportInstance()
+	{
+		return nullptr;
+	}
 };
 
 #endif
diff --git a/Engine/Source/Runtime/Experimental/ChaosSolverEngine/Public/Chaos/ChaosVDRemoteSessionsManager.h b/Engine/Source/Runtime/Experimental/ChaosSolverEngine/Public/Chaos/ChaosVDRemoteSessionsManager.h
--- a/Engine/Source/Runtime/Experimental/ChaosSolverEngine/Public/Chaos/ChaosVDRemoteSessionsManager.h
+++ b/Engine/Source/Runtime/Experimental/ChaosSolverEngine/Public/Chaos/ChaosVDRemoteSessionsManager.h
@@ -71,6 +71,48 @@ struct FChaosVDRecordingStatusMessage
 	FChaosVDTraceDetails TraceDetails;
 };
 
+USTRUCT()
+struct FChaosVDRelayTraceDataMessage
+{
+	GENERATED_BODY()
+	
+	FChaosVDRelayTraceDataMessage()
+	{
+	}
+
+	UPROPERTY()
+	FGuid InstanceId = FGuid();
+	
+	UPROPERTY()
+	TArray<uint8> DataBuffer;
+};
+
+UENUM()
+enum class ERelayThrottlingState : uint8
+{
+	Inactive,
+	Active
+};
+
+USTRUCT()
+struct FChaosVDRelayTraceStatusMessage
+{
+	GENERATED_BODY()
+	
+	FChaosVDRelayTraceStatusMessage()
+	{
+	}
+
+	UPROPERTY()
+	FGuid InstanceId = FGuid();
+	
+	UPROPERTY()
+	int64 QueuedDataBytesNum = 0;
+
+	UPROPERTY()
+	ERelayThrottlingState ThrottlingState = ERelayThrottlingState::Inactive;
+};
+
 USTRUCT()
 struct FChaosVDDataChannelState
 {
@@ -359,16 +401,33 @@ public:
 	CHAOSSOLVERENGINE_API static const FGuid CustomSessionsWrapperGUID;
 	CHAOSSOLVERENGINE_API static const FGuid InvalidSessionGUID;
 	CHAOSSOLVERENGINE_API static const FGuid LocalEditorSessionID;
+	CHAOSSOLVERENGINE_API static const FGuid LocalSessionID;
 
-private:
+	DECLARE_MULTICAST_DELEGATE_TwoParams(FMessagingInitialized, TSharedPtr<IMessageBus> NewMessageBus, TSharedPtr<FMessageEndpoint> NewEndpoint);
+	FMessagingInitialized& OnMessagingInitialized()
+	{
+		return MessagingInitializedDelegate;
+	}
 
-	void InitializeMessagingSystem(const TSharedPtr<IMessageBus>& InMessageBus);
-	void ShutdownMessagingSystem();
+	TWeakPtr<IMessageBus> GetMessageBusInstance()
+	{
+		return MessageBusPtr;
+	}
 	
+	TWeakPtr<FMessageEndpoint> GetMessageBusEndPoint()
+	{
+		return MessageEndpoint;
+	}
+
 	/**
 	 * Returns true if this instance has controller capabilities (is either an editor or CVD Standalone, which is also an editor)
 	 */
-	static constexpr bool IsController();
+	CHAOSSOLVERENGINE_API static constexpr bool IsController();
+
+private:
+
+	void InitializeMessagingSystem(const TSharedPtr<IMessageBus>& InMessageBus);
+	void ShutdownMessagingSystem();
 
 	/**
 	 * Sends a request to obtain the full session information to the provided message bus address
@@ -451,4 +510,8 @@ private:
 	FChaosVDRecordingStateChangeDelegate RecordingStoppedDelegate;
 	
 	FTSTicker::FDelegateHandle TickHandle;
+
+	FMessagingInitialized MessagingInitializedDelegate;
+
+	bool bInitialized = false;
 };
diff --git a/Engine/Source/Runtime/Experimental/ChaosSolverEngine/Public/Chaos/ChaosVDTraceRelayTransport.h b/Engine/Source/Runtime/Experimental/ChaosSolverEngine/Public/Chaos/ChaosVDTraceRelayTransport.h
new file mode 100644
--- /dev/null
+++ b/Engine/Source/Runtime/Experimental/ChaosSolverEngine/Public/Chaos/ChaosVDTraceRelayTransport.h
@@ -0,0 +1,43 @@
+// Copyright Epic Games, Inc. All Rights Reserved.
+
+#pragma once
+
+#if WITH_CHAOS_VISUAL_DEBUGGER
+#include "Delegates/Delegate.h"
+
+struct FChaosVDRelayTraceDataMessage;
+struct FGuid;
+
+namespace Chaos::VD
+{
+class FRelayTraceWriter;
+
+/**
+ * Interface for any CVD trace relay transport implementation
+ */
+class ITraceDataRelayTransport
+{
+public:
+	virtual ~ITraceDataRelayTransport() = default;
+
+	CHAOSSOLVERENGINE_API virtual void Initialize() = 0;
+	CHAOSSOLVERENGINE_API virtual void Shutdown() = 0;
+
+	DECLARE_DELEGATE_OneParam(FProcessReceivedRelayDataDelegate, const FChaosVDRelayTraceDataMessage& DataPacket);
+	CHAOSSOLVERENGINE_API virtual void RegisterRelayDataReceiverForSessionID(FGuid RemoteSessionID, const FProcessReceivedRelayDataDelegate& InProcessDataDelegate) = 0;
+	CHAOSSOLVERENGINE_API virtual void UnregisterRelayDataReceiverForSessionID(FGuid RemoteSessionID) = 0;
+
+	CHAOSSOLVERENGINE_API virtual void RelayTraceDataFromWriter(FRelayTraceWriter& InRelayTraceWriter) = 0;
+
+	DECLARE_DELEGATE_RetVal_OneParam(bool, FCanSendDataDelegate, int32 BytesSent);
+	CHAOSSOLVERENGINE_API virtual void RegisterExternalDataThrottleHandler(const FCanSendDataDelegate& InExternalThrottleHandler) = 0;
+	CHAOSSOLVERENGINE_API virtual void UnregisterExternalDataThrottleHandler() = 0;
+
+protected:
+	virtual void SendRelayedTraceData(const FChaosVDRelayTraceDataMessage& InRelayData) = 0;
+	virtual void ForwardReceivedTraceData(const FChaosVDRelayTraceDataMessage& InTraceDataPacket) = 0;
+};
+
+}
+
+#endif // WITH_CHAOS_VISUAL_DEBUGGER
diff --git a/Engine/Source/Runtime/Experimental/ChaosVisualDebugger/Private/ChaosVDRuntimeModule.cpp b/Engine/Source/Runtime/Experimental/ChaosVisualDebugger/Private/ChaosVDRuntimeModule.cpp
--- a/Engine/Source/Runtime/Experimental/ChaosVisualDebugger/Private/ChaosVDRuntimeModule.cpp
+++ b/Engine/Source/Runtime/Experimental/ChaosVisualDebugger/Private/ChaosVDRuntimeModule.cpp
@@ -23,10 +23,11 @@
 #include "Misc/App.h"
 #include "Modules/ModuleManager.h"
 #include "ProfilingDebugging/TraceAuxiliary.h"
+#include "Trace/RelayTraceWriter.h"
 
 IMPLEMENT_MODULE(FChaosVDRuntimeModule, ChaosVDRuntime);
 
-DEFINE_LOG_CATEGORY_STATIC( LogChaosVDRuntime, Log, All );
+DEFINE_LOG_CATEGORY(LogChaosVDRuntime);
 
 #define LOCTEXT_NAMESPACE "ChaosVisualDebugger"
 
@@ -64,6 +65,11 @@ static FAutoConsoleVariable CVarChaosVDMaxTimeToWaitForDisconnect(
 	5.0f,
 	TEXT("Max time to wait after attempting to stop an active trace session. After that time has passed if we are still connected, CVD will continue and eventually error out."));
 
+static FAutoConsoleVariable CVarChaosVDMaxBytesPerRelayChunk(
+	TEXT("p.Chaos.VD.MaxBytesPerRelayChunk"),
+	8192,
+	TEXT("How many bytes per relay data message we should send"));
+
 FChaosVDRecordingStateChangedDelegate FChaosVDRuntimeModule::RecordingStartedDelegate = FChaosVDRecordingStateChangedDelegate();
 FChaosVDRecordingStateChangedDelegate FChaosVDRuntimeModule::PostRecordingStartedDelegate = FChaosVDRecordingStateChangedDelegate();
 FChaosVDRecordingStateChangedDelegate FChaosVDRuntimeModule::RecordingStopDelegate = FChaosVDRecordingStateChangedDelegate();
@@ -167,6 +173,11 @@ FChaosVDTraceDetails FChaosVDRuntimeModule::GetCurrentTraceSessionDetails() cons
 	return Details;
 }
 
+const FString& FChaosVDRuntimeModule::GetConnectedToAddress() const
+{
+	return CurrentlyConnectedToAddress;
+}
+
 void FChaosVDRuntimeModule::StopTrace()
 {
 	bRequestedStop = true;
@@ -180,6 +191,31 @@ void FChaosVDRuntimeModule::GenerateRecordingFileName(FString& OutFileName)
 	OutFileName = FString::Format(TEXT("ChaosVD-{0}-{1}-{2}.utrace"), NameArgs);
 }
 
+void FChaosVDRuntimeModule::RegisterExternalRelayExecutor(const FChaosVDRelayTraceDataDelegate& InExternalRelayExecutor)
+{
+	if (!ensureMsgf(!ExternalRelayExecutorDelegate.IsBound(), TEXT("An external relay executor is already registered!")))
+	{
+		return;
+	}
+
+	ExternalRelayExecutorDelegate = InExternalRelayExecutor;
+}
+
+void FChaosVDRuntimeModule::UnregisterCurrentExternalRelayExecutor()
+{
+	ExternalRelayExecutorDelegate.Unbind();
+}
+
+bool FChaosVDRuntimeModule::RelayTraceData(float DeltaTime)
+{
+	if (ensure(ExternalRelayExecutorDelegate.IsBound() && RelayWriter.IsValid()))
+	{
+		ExternalRelayExecutorDelegate.Execute(*RelayWriter);
+	}
+
+	return true;
+}
+
 bool FChaosVDRuntimeModule::RequestFullCapture(float DeltaTime)
 {
 	// Full capture intervals are clamped to be no lower than 1 sec
@@ -237,7 +273,7 @@ void FChaosVDRuntimeModule::HandleTraceConnectionEstablished()
 }
 
 
-void FChaosVDRuntimeModule::StartRecording(const FChaosVDStartRecordingCommandMessage& InRecordingStartCommand)
+void FChaosVDRuntimeModule::StartRecording(const FChaosVDStartRecordingCommandMessage& InRecordingStartCommand, const FString& InConnectedToAddress)
 {
 	if (!ensureAlwaysMsgf(!bIsRecording, TEXT("Received a recording start command while a recording is already active")))
 	{
@@ -307,6 +343,34 @@ void FChaosVDRuntimeModule::StartRecording(const FChaosVDStartRecordingCommandMe
 			bIsRecording = FTraceAuxiliary::Start(FTraceAuxiliary::EConnectionType::Network, *InRecordingStartCommand.Target, nullptr, &TracingOptions);
 			break;
 		case EChaosVDTransportMode::Relay:
+			{
+				if (!ensureMsgf(ExternalRelayExecutorDelegate.IsBound(), TEXT("Cannot start a trace in relay mode without a Relay Executor")))
+				{
+					UE_LOG(LogChaosVDRuntime, Error, TEXT("[%hs] Cannot start a trace in relay mode without a Relay Executor"), __func__);
+					break;
+				}
+				
+				if (RelayWriter && !RelayWriter->IsClosed())
+				{
+					UE_LOG(LogChaosVDRuntime, Warning, TEXT("[%hs] Attempting to relay trace, but there is a relay writer still open"), __func__);
+					RelayWriter->Close();
+				}
+
+				RelayWriter = MakeUnique<Chaos::VD::FRelayTraceWriter>();
+				RelayWriter->SetMaxBytesPerBunch(CVarChaosVDMaxBytesPerRelayChunk->GetInt());
+				
+				bIsRecording = FTraceAuxiliary::Relay(reinterpret_cast<UPTRINT>(RelayWriter.Get()), Chaos::VD::FRelayTraceWriter::WriteHelper, Chaos::VD::FRelayTraceWriter::CloseHelper, nullptr, &TracingOptions);
+
+				if (bIsRecording)
+				{
+					if (!RelayExecutorTickerHandle.IsValid())
+					{
+						RelayExecutorTickerHandle = FTSTicker::GetCoreTicker().AddTicker(FTickerDelegate::CreateRaw(this, &FChaosVDRuntimeModule::RelayTraceData));
+					}
+				}
+
+				break;
+			}
 		case EChaosVDTransportMode::Invalid:
 		case EChaosVDTransportMode::FileSystem:
 		UE_LOG(LogChaosVDRuntime, Error, TEXT("[%hs] Unsupported Transport mode [%s]"), __func__, *UEnum::GetValueAsString(InRecordingStartCommand.TransportMode))
@@ -337,6 +401,8 @@ void FChaosVDRuntimeModule::StartRecording(const FChaosVDStartRecordingCommandMe
 		{
 			OnTraceStoppedDelegateHandle = FTraceAuxiliary::OnTraceStopped.AddRaw(this, &FChaosVDRuntimeModule::HandleTraceStopRequest);
 		}
+
+		CurrentlyConnectedToAddress = InConnectedToAddress;
 	}
 	else
 	{
@@ -352,6 +418,8 @@ void FChaosVDRuntimeModule::StartRecording(const FChaosVDStartRecordingCommandMe
 		}
 
 		CurrentRecordingMode = EChaosVDRecordingMode::Invalid;
+
+		RelayWriter = nullptr;
 	}
 }
 
@@ -368,6 +436,8 @@ void FChaosVDRuntimeModule::StopRecording()
 
 	if (!bIsRecording)
 	{
+		CurrentlyConnectedToAddress = FString();
+
 		UE_LOG(LogChaosVDRuntime, Warning, TEXT("[%s] Attempted to stop recorded when there is no CVD recording active."), ANSI_TO_TCHAR(__FUNCTION__));
 		return;
 	}
@@ -390,6 +460,17 @@ void FChaosVDRuntimeModule::StopRecording()
 		FTSTicker::RemoveTicker(RecordingTimerHandle);
 		RecordingTimerHandle.Reset();
 	}
+
+	if (RelayExecutorTickerHandle.IsValid())
+	{
+		FTSTicker::RemoveTicker(RelayExecutorTickerHandle);
+		RelayExecutorTickerHandle.Reset();
+	}
+
+	if (RelayWriter)
+	{
+		RelayWriter->Close();
+	}
 	
 	bIsRecording = false;
 	AccumulatedRecordingTime = 0.0f;
@@ -398,6 +479,8 @@ void FChaosVDRuntimeModule::StopRecording()
 		UE::TReadScopeLock ReadLock(DelegatesRWLock);
 		RecordingStopDelegate.Broadcast();
 	}
+
+	CurrentlyConnectedToAddress = FString();
 }
 
 void FChaosVDRuntimeModule::HandleTraceStopRequest(FTraceAuxiliary::EConnectionType TraceType, const FString& TraceDestination)
diff --git a/Engine/Source/Runtime/Experimental/ChaosVisualDebugger/Private/Trace/RelayTraceWriter.cpp b/Engine/Source/Runtime/Experimental/ChaosVisualDebugger/Private/Trace/RelayTraceWriter.cpp
new file mode 100644
--- /dev/null
+++ b/Engine/Source/Runtime/Experimental/ChaosVisualDebugger/Private/Trace/RelayTraceWriter.cpp
@@ -0,0 +1,152 @@
+// Copyright Epic Games, Inc. All Rights Reserved.
+
+#include "Trace/RelayTraceWriter.h"
+
+#if WITH_CHAOS_VISUAL_DEBUGGER
+#include "ChaosVDRuntimeModule.h"
+#include "Serialization/BufferArchive.h"
+#include "Serialization/MemoryReader.h"
+
+namespace Chaos::VD
+{
+	FRelayTraceWriter::FRelayTraceWriter()
+	{
+	}
+
+	FRelayTraceWriter::~FRelayTraceWriter()
+	{
+		Close();
+	}
+
+	void FRelayTraceWriter::SetMaxBytesPerBunch(uint32 MaxBytes)
+	{
+		if (MaxBytes == 0)
+		{
+			UE_LOG(LogChaosVDRuntime, Error, TEXT("[%hs] Only values greater than 0 are supported"), __func__);
+			return;
+		}
+
+		if (MaxBytes != MaxPendingBytesPerBunch)
+		{
+			if (!ensure(PendingBunchesToRelay.IsEmpty()))
+			{
+				return;
+			}
+		}
+
+		MaxPendingBytesPerBunch = MaxBytes;
+	}
+
+
+	TOptional<TUniquePtr<FBufferArchive>> FRelayTraceWriter::DequeuePendingBunch()
+	{
+		if (PendingBunchesCount == 0)
+		{
+			return TOptional<TUniquePtr<FBufferArchive>>();
+		}
+
+		QueuedBytesNum -= NextPendingBunchSize();
+	
+		--PendingBunchesCount;
+		return PendingBunchesToRelay.Dequeue();
+	}
+
+	int32 FRelayTraceWriter::NextPendingBunchSize() const
+	{
+		if (PendingBunchesCount == 0)
+		{
+			return 0;
+		}
+
+		TUniquePtr<FBufferArchive>* PendingBunch = PendingBunchesToRelay.Peek();
+		
+		if (PendingBunch && PendingBunch->IsValid())
+		{
+			return (*PendingBunch)->Num();
+		}
+
+		return 0;
+	}
+
+	TUniquePtr<FBufferArchive> FRelayTraceWriter::CreateBunch()
+	{
+		TUniquePtr<FBufferArchive> NewBunch = MakeUnique<FBufferArchive>();
+		return NewBunch;
+	}
+
+	bool FRelayTraceWriter::Write_Internal(const void* Data, uint32 Size)
+	{
+		if (bClosed)
+		{
+			return false;
+		}
+
+		const bool bHasValidLimits = MaxPendingBytesPerBunch != 0;
+		if (!ensure(bHasValidLimits))
+		{
+			return false;
+		}
+
+		if (!CurrentPendingBunch)
+		{
+			CurrentPendingBunch = CreateBunch();
+		}
+
+		uint32 RemainingSizeToCopy = Size;
+
+		while (RemainingSizeToCopy > 0)
+		{
+			int32 AvailableSpaceInCurrentBunch = MaxPendingBytesPerBunch - CurrentPendingBunch->Num();
+
+			if (AvailableSpaceInCurrentBunch <= 0)
+			{
+				QueuedBytesNum += CurrentPendingBunch->Num();
+				PendingBunchesToRelay.Enqueue(MoveTemp(CurrentPendingBunch));
+				++PendingBunchesCount;
+				CurrentPendingBunch = CreateBunch();
+				AvailableSpaceInCurrentBunch = MaxPendingBytesPerBunch;
+			}
+
+			const uint32 BytesToCopyNum = static_cast<uint32>(AvailableSpaceInCurrentBunch) > RemainingSizeToCopy ? RemainingSizeToCopy : static_cast<uint32>(AvailableSpaceInCurrentBunch);
+
+			const uint8* DestBufferPtr = static_cast<const uint8*>(Data);
+			const int32 Offset = Size - RemainingSizeToCopy;
+
+			const bool bIsValidOffset = Offset >= 0 && static_cast<uint32>(Offset) < Size;
+			if (!ensureMsgf(bIsValidOffset, TEXT("Unexpected calculated offset [%d] | Size [%u]"), Offset, Size))
+			{
+				return false;
+			}
+
+			CurrentPendingBunch->Serialize(const_cast<uint8*>(DestBufferPtr + Offset) , BytesToCopyNum);
+
+			if (CurrentPendingBunch->IsError())
+			{
+				return false;
+			}
+
+			RemainingSizeToCopy -= BytesToCopyNum;
+		}
+
+		return !CurrentPendingBunch->IsError();
+	}
+
+	int64 FRelayTraceWriter::GetQueuedBytesNum() const
+	{
+		return QueuedBytesNum;
+	}
+
+	void FRelayTraceWriter::Close()
+	{
+		int64 PendingBytes = QueuedBytesNum;
+		if (PendingBytes > 0)
+		{
+			// TODO: We should handle the case where we close the trace session but we still have data pending in the send queue (UE-320563).
+			// We can ignore it for now as it is not a critical failure
+			UE_LOG(LogChaosVDRuntime, Warning, TEXT("Closing Trace Session with [%lld] bytes in the send queue. This data will be discarded"), PendingBytes);
+		}
+
+		bClosed = true;
+	}
+}
+#endif // WITH_CHAOS_VISUAL_DEBUGGER
\ No newline at end of file
diff --git a/Engine/Source/Runtime/Experimental/ChaosVisualDebugger/Public/ChaosVDRecordingDetails.h b/Engine/Source/Runtime/Experimental/ChaosVisualDebugger/Public/ChaosVDRecordingDetails.h
--- a/Engine/Source/Runtime/Experimental/ChaosVisualDebugger/Public/ChaosVDRecordingDetails.h
+++ b/Engine/Source/Runtime/Experimental/ChaosVisualDebugger/Public/ChaosVDRecordingDetails.h
@@ -69,4 +69,7 @@ struct FChaosVDTraceDetails
 
 	UPROPERTY()
 	EChaosVDTransportMode TransportMode = EChaosVDTransportMode::Invalid;
+
+	UPROPERTY()
+	FString ConnectedToAddress;
 };
\ No newline at end of file
diff --git a/Engine/Source/Runtime/Experimental/ChaosVisualDebugger/Public/ChaosVDRuntimeModule.h b/Engine/Source/Runtime/Experimental/ChaosVisualDebugger/Public/ChaosVDRuntimeModule.h
--- a/Engine/Source/Runtime/Experimental/ChaosVisualDebugger/Public/ChaosVDRuntimeModule.h
+++ b/Engine/Source/Runtime/Experimental/ChaosVisualDebugger/Public/ChaosVDRuntimeModule.h
@@ -14,6 +14,11 @@
 
 #define UE_API CHAOSVDRUNTIME_API
 
+namespace Chaos::VD
+{
+	class FRelayTraceWriter;
+}
+
 struct FChaosVDRecording;
 class FText;
 
@@ -29,6 +34,9 @@ ENUM_CLASS_FLAGS(EChaosVDFullCaptureFlags)
 DECLARE_MULTICAST_DELEGATE(FChaosVDRecordingStateChangedDelegate)
 DECLARE_MULTICAST_DELEGATE_OneParam(FChaosVDCaptureRequestDelegate, EChaosVDFullCaptureFlags)
 DECLARE_MULTICAST_DELEGATE_OneParam(FChaosVDRecordingStartFailedDelegate, const FText&)
+DECLARE_DELEGATE_OneParam(FChaosVDRelayTraceDataDelegate, Chaos::VD::FRelayTraceWriter&)
+
+DECLARE_LOG_CATEGORY_EXTERN(LogChaosVDRuntime, Log, All)
 
 class FChaosVDRuntimeModule : public IModuleInterface
 {
@@ -48,7 +56,7 @@ public:
 	/** Starts a CVD recording by starting a Trace session. It will stop any existing trace session
 	 * @param InRecordingStartCommand : Used to determine if we want to record to file or a local trace server among other settings
 	 */
-	UE_API void StartRecording(const FChaosVDStartRecordingCommandMessage& InRecordingStartCommand);
+	UE_API void StartRecording(const FChaosVDStartRecordingCommandMessage& InRecordingStartCommand, const FString& InConnectedToAddress = FString());
 	
 	/* Stops an active recording */
 	UE_API void StopRecording();
@@ -133,6 +141,7 @@ public:
 	UE_API FString GetLastRecordingFileNamePath() const;
 
 	UE_API FChaosVDTraceDetails GetCurrentTraceSessionDetails() const;
+	UE_API const FString& GetConnectedToAddress() const;
 
 	EChaosVDRecordingMode GetCurrentRecordingMode() const
 	{
@@ -142,8 +151,21 @@ public:
 	/** Finds a valid file name for a new file - Used to generate the file name for the Trace recording */
 	UE_API void GenerateRecordingFileName(FString& OutFileName);
 
+	/**
+	 * Registers an external handler that will be in charge of relaying the traced data to a remote CVD instance if the relay transport mode is selected
+	 * @param InExternalRelayExecutor Delegate that will be executed to relay data to a remote CVD instance, periodically
+	 */
+	UE_API void RegisterExternalRelayExecutor(const FChaosVDRelayTraceDataDelegate& InExternalRelayExecutor);
+
+	/**
+	 * Unbinds the delegate to an external relay executor
+	 */
+	UE_API void UnregisterCurrentExternalRelayExecutor();
+
 private:
 
+	bool RelayTraceData(float DeltaTime);
+
 	/** Stops the current Trace session */
 	UE_API void StopTrace();
 
@@ -173,6 +195,7 @@ private:
 
 	FTSTicker::FDelegateHandle FullCaptureRequesterHandle;
 	FTSTicker::FDelegateHandle RecordingTimerHandle;
+	FTSTicker::FDelegateHandle RelayExecutorTickerHandle;
 
 	FDelegateHandle OnConnectionDelegateHandle;
 	FDelegateHandle OnTraceStoppedDelegateHandle;
@@ -190,6 +213,10 @@ private:
 	EChaosVDRecordingMode CurrentRecordingMode = EChaosVDRecordingMode::Invalid;
 	EChaosVDTransportMode CurrentTransportMode = EChaosVDTransportMode::Invalid;
 
+	TUniquePtr<Chaos::VD::FRelayTraceWriter> RelayWriter;
+	FChaosVDRelayTraceDataDelegate ExternalRelayExecutorDelegate;
+
+	FString CurrentlyConnectedToAddress;
 
 	static UE_API FTransactionallySafeRWLock DelegatesRWLock;
 };
diff --git a/Engine/Source/Runtime/Experimental/ChaosVisualDebugger/Public/Trace/RelayTraceWriter.h b/Engine/Source/Runtime/Experimental/ChaosVisualDebugger/Public/Trace/RelayTraceWriter.h
new file mode 100644
--- /dev/null
+++ b/Engine/Source/Runtime/Experimental/ChaosVisualDebugger/Public/Trace/RelayTraceWriter.h
@@ -0,0 +1,105 @@
+// Copyright Epic Games, Inc. All Rights Reserved.
+
+#pragma once
+
+#if WITH_CHAOS_VISUAL_DEBUGGER
+#include <atomic>
+#include "Containers/Array.h"
+#include "Containers/SpscQueue.h"
+#include "Serialization/BitWriter.h"
+#include "Templates/UniquePtr.h"
+
+#pragma once
+
+class FBufferArchive;
+
+namespace Chaos::VD
+{
+/**
+ * Writer object that chunks trace data into smaller bunches and queues them so a relay instance
+ * can dequeue and send them over the network later
+ */
+class FRelayTraceWriter
+{
+public:
+	FRelayTraceWriter();
+	~FRelayTraceWriter();
+
+	/**
+	 * Sets max number of bytes a queued bunch can have
+	 * @param MaxBytes Greater than zero number of bytes to set as limit
+	 */
+	CHAOSVDRUNTIME_API void SetMaxBytesPerBunch(uint32 MaxBytes);
+
+	static bool WriteHelper(UPTRINT WriterHandle, const void* Data, uint32 Size)
+	{
+		if (!WriterHandle)
+		{
+			return false;
+		}
+
+		return reinterpret_cast<FRelayTraceWriter*>(WriterHandle)->Write_Internal(Data, Size);
+	}
+
+	static void CloseHelper(UPTRINT WriterHandle)
+	{
+		if (!WriterHandle)
+		{
+			return;
+		}
+
+		reinterpret_cast<FRelayTraceWriter*>(WriterHandle)->Close();
+	}
+
+	/**
+	 * Dequeues a bunch ready to send over the network, if any
+	 */
+	CHAOSVDRUNTIME_API TOptional<TUniquePtr<FBufferArchive>> DequeuePendingBunch();
+
+	/**
+	 * Returns the size of the next bunch in the queue
+	 */
+	CHAOSVDRUNTIME_API int32 NextPendingBunchSize() const;
+
+	/**
+	 * Returns true if we have any data pending to be sent over the network
+	 */
+	bool HasPendingBunches() const
+	{
+		return PendingBunchesCount > 0;
+	}
+
+	/**
+	 * returns the number of bytes overall that are in the queue waiting to be sent
+	 */
+	CHAOSVDRUNTIME_API int64 GetQueuedBytesNum() const;
+
+	/**
+	 * Closes this Writer. After this is called, no new data will be processed, but any pending data can still be dequeued and sent over
+	 */
+	void Close();
+
+	/**
+	 * Returns true if this writer is closed and no new data is being processed
+	 */
+	bool IsClosed() const
+	{
+		return bClosed;
+	}
+
+	private:
+
+		TUniquePtr<FBufferArchive> CreateBunch();
+
+		bool Write_Internal(const void* Data, uint32 Size);
+
+		TSpscQueue<TUniquePtr<FBufferArchive>> PendingBunchesToRelay;
+		TUniquePtr<FBufferArchive> CurrentPendingBunch;
+		std::atomic<int32> PendingBunchesCount = 0;
+		std::atomic<bool> bClosed = false;
+		uint32 MaxPendingBytesPerBunch = 0;
+		std::atomic<int64> QueuedBytesNum = 0;
+};
+}
+#endif // WITH_CHAOS_VISUAL_DEBUGGER
+
