IAD Cache
- Remove LastAccessTime from pending chunks.  As implemented, this isn't really an improvement over just using current time at flush.

diff --git a/Engine/Source/Runtime/Experimental/IoStore/OnDemand/Private/OnDemandInstallCache.cpp b/Engine/Source/Runtime/Experimental/IoStore/OnDemand/Private/OnDemandInstallCache.cpp
--- a/Engine/Source/Runtime/Experimental/IoStore/OnDemand/Private/OnDemandInstallCache.cpp
+++ b/Engine/Source/Runtime/Experimental/IoStore/OnDemand/Private/OnDemandInstallCache.cpp
@@ -1660,7 +1660,6 @@ class FOnDemandInstallCache final
 			TotalSize += Chunk.GetSize();
 			ChunkHashes.Add(ChunkHash);
 			Chunks.Add(MoveTemp(Chunk));
-			LastAccessTime.store(FDateTime::UtcNow().GetTicks(), std::memory_order_relaxed);
 		}
 
 		void Reset()
@@ -1668,7 +1667,6 @@ class FOnDemandInstallCache final
 			Chunks.Reset();
 			ChunkHashes.Reset();
 			TotalSize = 0;
-			LastAccessTime.store(0, std::memory_order_relaxed);
 		}
 
 		const FIoBuffer* FindChunk(const FIoHash& ChunkHash) const
@@ -1685,7 +1683,6 @@ class FOnDemandInstallCache final
 		TArray<FIoBuffer>		Chunks;
 		TArray<FIoHash>			ChunkHashes;
 		uint64					TotalSize = 0;
-		std::atomic<uint64>		LastAccessTime{ 0 };
 		mutable FSharedMutex	SharedMutex;
 	};
 
@@ -1742,7 +1739,7 @@ private:
 	bool										Resolve(FIoRequestImpl* Request);
 	void										CompleteRequest(FIoRequestImpl* Request, EIoErrorCode Status);
 	FResult										FlushPendingChunks(FPendingChunks& Chunks, int64 UtcAccessTicks = 0);
-	FResult										FlushPendingChunksImpl(FPendingChunks& Chunks, int64 UtcAccessTicks = 0);
+	FResult										FlushPendingChunksImpl(const FPendingChunks& Chunks, int64 UtcAccessTicks = 0);
 	FString										GetJournalFilename() const { return CacheDirectory / TEXT("cas.jrn"); }
 	FString										GetSnapshotFilename() const { return CacheDirectory / TEXT("cas.snp"); }
 	void										CaptureCacheState(UE::UnifiedError::IoStoreOnDemand::FInstallCacheErrorContext& OutCtx);
@@ -1921,7 +1918,7 @@ void FOnDemandInstallCache::Initialize(FSharedBackendContextRef Context)
 
 void FOnDemandInstallCache::Shutdown()
 {
-	if (FResult Result = FlushPendingChunksImpl(PendingChunks, PendingChunks.LastAccessTime.load(std::memory_order_relaxed)); Result.HasError())
+	if (FResult Result = FlushPendingChunksImpl(PendingChunks); Result.HasError())
 	{
 		UE_LOG(LogIoStoreOnDemand, Error, TEXT("Failed to flush pending chunks on shutdown, reason '%s'"), *LexToString(Result.GetError()));
 	}
@@ -2102,8 +2099,6 @@ bool FOnDemandInstallCache::Resolve(FIoRequestImpl* Request)
 			check(BytesToRead.GetSize() == ChunkRequest.EncodedChunk.GetSize());
 			FMemory::Memcpy(ChunkRequest.EncodedChunk.GetData(), BytesToRead.GetData(), ChunkRequest.EncodedChunk.GetSize());
 
-			PendingChunks.LastAccessTime.store(FDateTime::UtcNow().GetTicks(), std::memory_order_relaxed);
-
 			SharedLock.Unlock();
 
 			UE::Tasks::Launch(UE_SOURCE_LOCATION, [this, Request]
@@ -2309,8 +2304,7 @@ FResult FOnDemandInstallCache::PutChunk(FIoBuffer&& Chunk, const FIoHash& ChunkH
 {
 	if (PendingChunks.TotalSize > FPendingChunks::MaxPendingBytes)
 	{
-		if (FResult Result = FlushPendingChunks(PendingChunks, PendingChunks.LastAccessTime.load(std::memory_order_relaxed));
-			Result.HasError())
+		if (FResult Result = FlushPendingChunks(PendingChunks); Result.HasError())
 		{
 			return Result;
 		}
@@ -3384,8 +3378,7 @@ FResult FOnDemandInstallCache::Defrag(
 
 FResult FOnDemandInstallCache::Flush()
 {
-	if (FResult Result = FlushPendingChunks(PendingChunks, PendingChunks.LastAccessTime.load(std::memory_order_relaxed));
-		Result.HasError())
+	if (FResult Result = FlushPendingChunks(PendingChunks); Result.HasError())
 	{
 		return Result;
 	}
@@ -3511,6 +3504,12 @@ FResult FOnDemandInstallCache::FlushPendingChunks(FPendingChunks& Chunks, int64
 #if UE_ONDEMANDINSTALLCACHE_EXCLUSIVE_WRITE
 	UE::Tasks::TTask<FResult> Task = ExclusivePipe.Launch(UE_SOURCE_LOCATION, [this, &Chunks, UtcAccessTicks]
 	{
+		ON_SCOPE_EXIT 
+		{ 
+			TUniqueLock Lock(Chunks.SharedMutex);
+			Chunks.Reset();
+		};
+
 		return FlushPendingChunksImpl(Chunks, UtcAccessTicks);
 	}, UE::Tasks::ETaskPriority::BackgroundHigh);
 
@@ -3519,18 +3518,18 @@ FResult FOnDemandInstallCache::FlushPendingChunks(FPendingChunks& Chunks, int64
 	return Task.GetResult();
 
 #else
+	ON_SCOPE_EXIT
+	{
+		TUniqueLock Lock(Chunks.SharedMutex);
+		Chunks.Reset();
+	};
+
 	return FlushPendingChunksImpl(Chunks, UtcAccessTicks);
 #endif // UE_ONDEMANDINSTALLCACHE_EXCLUSIVE_WRITE
 }
 
-FResult FOnDemandInstallCache::FlushPendingChunksImpl(FPendingChunks& Chunks, int64 UtcAccessTicks)
+FResult FOnDemandInstallCache::FlushPendingChunksImpl(const FPendingChunks& Chunks, int64 UtcAccessTicks)
 {
-	ON_SCOPE_EXIT 
-	{ 
-		TUniqueLock Lock(Chunks.SharedMutex);
-		Chunks.Reset();
-	};
-
 	FLargeMemoryWriter Ar(FMath::Min(Chunks.TotalSize, static_cast<uint64>(Cas.GetMaxBlockSize())));
 
 	// This should be the only thread writing to CurrentBlock
@@ -3569,9 +3568,9 @@ FResult FOnDemandInstallCache::FlushPendingChunksImpl(FPendingChunks& Chunks, in
 		}
 
 		const int64 CasBlockOffset = CasFileHandle->Tell();
-
-		TArray<FIoHash>		ChunkHashes;
+		
 		TArray<int64>		Offsets;
+		const FIoHash*		ChunkHashBegin = &Chunks.ChunkHashes[ChunkIdx];
 
 		while (ChunkIdx < Chunks.Chunks.Num())
 		{
@@ -3580,14 +3579,15 @@ FResult FOnDemandInstallCache::FlushPendingChunksImpl(FPendingChunks& Chunks, in
 				break;
 			}
 
-			FIoBuffer Chunk = MoveTemp(Chunks.Chunks[ChunkIdx]);
-			ChunkHashes.Add(MoveTemp(Chunks.ChunkHashes[ChunkIdx]));
+			const FIoBuffer& Chunk = Chunks.Chunks[ChunkIdx];
 			Offsets.Add(CasBlockOffset + Ar.Tell());
-			Ar.Serialize(Chunk.GetData(), Chunk.GetSize());
+			Ar.Serialize(const_cast<uint8*>(Chunk.GetData()), Chunk.GetSize());
 
 			++ChunkIdx;
 		}
 
+		TConstArrayView<FIoHash> ChunkHashes = MakeArrayView(ChunkHashBegin, Offsets.Num());
+
 		if (Ar.Tell() > 0)
 		{
 			UE_LOG(LogIoStoreOnDemand, Log, TEXT("Writing %.2lf MiB to CAS block %u"),
