Convert Journal/Cas error handing to FError

diff --git a/Engine/Source/Runtime/Experimental/IoStore/OnDemand/Private/OnDemandInstallCache.cpp b/Engine/Source/Runtime/Experimental/IoStore/OnDemand/Private/OnDemandInstallCache.cpp
--- a/Engine/Source/Runtime/Experimental/IoStore/OnDemand/Private/OnDemandInstallCache.cpp
+++ b/Engine/Source/Runtime/Experimental/IoStore/OnDemand/Private/OnDemandInstallCache.cpp
@@ -83,7 +83,6 @@ double ToMiB(uint64 Value)
 ///////////////////////////////////////////////////////////////////////////////
 using FUniqueFileHandle				= TUniquePtr<IFileHandle>;
 using FSharedFileHandle				= TSharedPtr<IFileHandle>;
-using FSharedFileOpenResult			= TValueOrError<FSharedFileHandle, FFileSystemError>;
 
 using FSharedAsyncFileHandle		= TSharedPtr<IAsyncReadFileHandle>;
 using FWeakAsyncFileHandle			= TWeakPtr<IAsyncReadFileHandle>;
@@ -192,15 +191,15 @@ struct FCas
 	void						Lock()		{ Mutex.Lock(); }
 	void						Unlock()	{ Mutex.Unlock(); }
 
-	FIoStatus					Initialize(FStringView Directory, bool bDeleteExisting = false);
+	FResult						Initialize(FStringView Directory, bool bDeleteExisting = false);
 	FCasLocation				FindChunk(const FIoHash& Hash) const;
 	FCasBlockId					CreateBlock();
-	FIoStatus					DeleteBlock(FCasBlockId BlockId, TArray<FCasAddr>& OutAddrs);
+	FResult						DeleteBlock(FCasBlockId BlockId, TArray<FCasAddr>& OutAddrs);
 	FString						GetBlockFilename(FCasBlockId BlockId) const;
-	FSharedFileOpenResult		OpenRead(FCasBlockId BlockId);
+	TResult<FSharedFileHandle>	OpenRead(FCasBlockId BlockId);
 	FSharedFileOpenAsyncResult	OpenAsyncRead(FCasBlockId BlockId);
 	void						OnFileHandleDeleted(FCasBlockId BlockId);
-	FUniqueFileHandle			OpenWrite(FCasBlockId BlockId, bool bAppend) const;
+	TResult<FUniqueFileHandle>	OpenWrite(FCasBlockId BlockId, bool bAppend) const;
 	bool						TrackAccessIf(ECasTrackAccessType Type, FCasBlockId BlockId, int64 UtcTicks, bool bDirty);
 	bool						TrackAccessIf(ECasTrackAccessType Type, FCasBlockId BlockId, bool bDirty)
 	{
@@ -209,7 +208,7 @@ struct FCas
 	bool						UnlockedTrackAccessIf(ECasTrackAccessType Type, uint32 BlockIdHash, FCasBlockId BlockId, int64 UtcTicks, bool bDirty);
 	uint64						GetBlockInfo(FCasBlockInfoMap& OutBlockInfo);
 	void						Compact();
-	FIoStatus					Verify(TArray<FCasAddr>& OutAddrs);
+	FResult						Verify(TArray<FCasAddr>& OutAddrs);
 	void						LoadSnapshot(FCasSnapshot&& Snapshot);
 	FLastAccess					ConsumeLastAcccess()
 	{
@@ -260,8 +259,10 @@ FCas::FCas(const FOnDemandInstallCacheConfig& Config)
 }
 
 ///////////////////////////////////////////////////////////////////////////////
-FIoStatus FCas::Initialize(FStringView Directory, bool bDeleteExisting)
+FResult FCas::Initialize(FStringView Directory, bool bDeleteExisting)
 {
+	using namespace UE::UnifiedError::IoStoreOnDemand;
+
 	RootDirectory = Directory;
 
 	Lookup.Empty();
@@ -280,10 +281,8 @@ FIoStatus FCas::Initialize(FStringView Directory, bool bDeleteExisting)
 
 		if (Ifm.DeleteDirectory(Path.ToString(), bRequireExists, bTree) == false)
 		{
-			return FIoStatusBuilder(EIoErrorCode::DeleteError, FPlatformMisc::GetLastError())
-				<< TEXT("Failed to delete CAS blocks directory '")
-				<< Path.ToString()
-				<< TEXT("'");
+			return MakeCasError<void>(
+				ECasErrorCode::InitializeFailed, EIoErrorCode::DeleteError, FString::Printf(TEXT("Failed to delete CAS blocks directory '%s'"), Path.ToString()));
 		}
 	}
 
@@ -292,14 +291,12 @@ FIoStatus FCas::Initialize(FStringView Directory, bool bDeleteExisting)
 		const bool bTree = true;
 		if (Ifm.MakeDirectory(Path.ToString(), bTree) == false)
 		{
-			return FIoStatusBuilder(EIoErrorCode::WriteError, FPlatformMisc::GetLastError())
-				<< TEXT("Failed to create directory '")
-				<< Path.ToString()
-				<< TEXT("'");
+			return MakeCasError<void>(
+				ECasErrorCode::InitializeFailed, EIoErrorCode::WriteError, FString::Printf(TEXT("Failed to create CAS blocks directory '%s'"), Path.ToString()));
 		}
 	}
 
-	return FIoStatus::Ok;
+	return MakeValue();
 };
 
 FCasLocation FCas::FindChunk(const FIoHash& Hash) const
@@ -346,8 +343,10 @@ FCasBlockId FCas::CreateBlock()
 	return Out;
 }
 
-FIoStatus FCas::DeleteBlock(FCasBlockId BlockId, TArray<FCasAddr>& OutAddrs)
+FResult FCas::DeleteBlock(FCasBlockId BlockId, TArray<FCasAddr>& OutAddrs)
 {
+	using namespace UE::UnifiedError::IoStoreOnDemand;
+
 	IPlatformFile& Ipf = FPlatformFileManager::Get().GetPlatformFile();
 	const FString	Filename = GetBlockFilename(BlockId);
 
@@ -368,8 +367,8 @@ FIoStatus FCas::DeleteBlock(FCasBlockId BlockId, TArray<FCasAddr>& OutAddrs)
 
 			if (FPlatformTime::ToMilliseconds(FPlatformTime::Cycles() - StartTimeCycles) > DeleteBlockMaxWaitTimeMs)
 			{
-				return FIoStatusBuilder(EIoErrorCode::Timeout)
-					<< TEXT("Timed out waiting for pending read(s) while deleting CAS block"); 
+				return MakeCasError<void>(
+					ECasErrorCode::DeleteBlockFailed, EIoErrorCode::Timeout, FString::Printf(TEXT("Timed out waiting for pending read(s) when deleting CAS block %u"), BlockId.Id));
 			}
 
 			BlockReadsDoneEvent->Wait(WaitTimeMs);
@@ -384,7 +383,8 @@ FIoStatus FCas::DeleteBlock(FCasBlockId BlockId, TArray<FCasAddr>& OutAddrs)
 	UE_LOG(LogIoStoreOnDemand, Log, TEXT("Deleting CAS block '%s'"), *Filename);
 	if (Ipf.DeleteFile(*Filename) == false)
 	{
-		return FIoStatusBuilder(EIoErrorCode::DeleteError, FPlatformMisc::GetLastError()) << TEXT("Failed to delete CAS block");
+		return MakeCasError<void>(
+			ECasErrorCode::DeleteBlockFailed, EIoErrorCode::DeleteError, FString::Printf(TEXT("Failed to delete CAS block %u"), BlockId.Id));
 	}
 
 	BlockIds.Remove(BlockId);
@@ -399,7 +399,7 @@ FIoStatus FCas::DeleteBlock(FCasBlockId BlockId, TArray<FCasAddr>& OutAddrs)
 		}
 	}
 
-	return FIoStatus::Ok;
+	return MakeValue();
 }
 
 FString FCas::GetBlockFilename(FCasBlockId BlockId) const
@@ -415,8 +415,10 @@ FString FCas::GetBlockFilename(FCasBlockId BlockId) const
 	return FString(Path.ToView());
 }
 
-FSharedFileOpenResult FCas::OpenRead(FCasBlockId BlockId)
+TResult<FSharedFileHandle> FCas::OpenRead(FCasBlockId BlockId)
 {
+	using namespace UE::UnifiedError::IoStoreOnDemand;
+
 	const FString	Filename = GetBlockFilename(BlockId);
 	IPlatformFile&	Ipf = FPlatformFileManager::Get().GetPlatformFile();
 
@@ -439,7 +441,8 @@ FSharedFileOpenResult FCas::OpenRead(FCasBlockId BlockId)
 		return MakeValue(MoveTemp(NewHandle));
 	}
 
-	return MakeError(Result.StealError());
+	return MakeCasError<FSharedFileHandle>(
+		ECasErrorCode::ReadBlockFailed, EIoErrorCode::FileOpenFailed, FString::Printf(TEXT("Failed to open CAS block '%s'"), *Filename));
 }
 
 FSharedFileOpenAsyncResult FCas::OpenAsyncRead(FCasBlockId BlockId)
@@ -489,13 +492,22 @@ void FCas::OnFileHandleDeleted(FCasBlockId BlockId)
 	}
 }
 
-FUniqueFileHandle FCas::OpenWrite(FCasBlockId BlockId, const bool bAppend) const
+TResult<FUniqueFileHandle> FCas::OpenWrite(FCasBlockId BlockId, const bool bAppend) const
 {
+	using namespace UE::UnifiedError::IoStoreOnDemand;
+
 	IPlatformFile&	Ipf = FPlatformFileManager::Get().GetPlatformFile();
 	const FString	Filename = GetBlockFilename(BlockId);
 	const bool		bAllowRead = true;
 
-	return FUniqueFileHandle(Ipf.OpenWrite(*Filename, bAppend, bAllowRead));
+	FUniqueFileHandle FileHandle(Ipf.OpenWrite(*Filename, bAppend, bAllowRead));
+	if (FileHandle.IsValid())
+	{
+		return MakeValue(MoveTemp(FileHandle));
+	}
+
+	return MakeCasError<FUniqueFileHandle>(
+		ECasErrorCode::WriteBlockFailed, EIoErrorCode::FileOpenFailed, FString::Printf(TEXT("Failed to open CAS block '%s' for writing"), *Filename));
 }
 
 bool FCas::TrackAccessIf(const ECasTrackAccessType Type, const FCasBlockId BlockId, const int64 UtcTicks, const bool bDirty)
@@ -638,12 +650,14 @@ void FCas::Compact()
 	LastAccess.Compact();
 }
 
-FIoStatus FCas::Verify(TArray<FCasAddr>& OutAddrs)
+FResult FCas::Verify(TArray<FCasAddr>& OutAddrs)
 {
+	using namespace UE::UnifiedError::IoStoreOnDemand;
+
 	FCasBlockInfoMap	BlockInfo;
 	const uint64		TotalSize = GetBlockInfo(BlockInfo);
 	uint64				TotalVerifiedBytes = 0;
-	FIoStatus			Status = FIoStatus::Ok;
+	FResult				Result = MakeValue();
 
 	const int64 Now = FDateTime::UtcNow().GetTicks();
 	const int64 TimestampForMissingLastAccess = FCas::GetTimestampForMissingLastAccess();
@@ -699,12 +713,12 @@ FIoStatus FCas::Verify(TArray<FCasAddr>& OutAddrs)
 		}
 
 		const FString Filename = GetBlockFilename(BlockId);
-		const FString ErrroMessage = FString::Printf(TEXT("Missing CAS block '%s'"), *Filename);
-		UE_LOG(LogIoStoreOnDemand, Warning, TEXT("%s"), *ErrroMessage);
+		FString ErrorMessage = FString::Printf(TEXT("Missing CAS block '%s'"), *Filename);
+		UE_LOG(LogIoStoreOnDemand, Warning, TEXT("%s"), *ErrorMessage);
 
 		LastAccess.Remove(BlockId);
 		BlockIt.RemoveCurrent();
-		Status = FIoStatus(EIoErrorCode::NotFound, ErrroMessage);
+		Result = MakeCasError<void>(ECasErrorCode::VerifyFailed, EIoErrorCode::NotFound, MoveTemp(ErrorMessage));
 	}
 
 	UE_LOG(LogIoStoreOnDemand, Log, TEXT("Verified %d CAS blocks of total %.2lf MiB"),
@@ -730,12 +744,14 @@ FIoStatus FCas::Verify(TArray<FCasAddr>& OutAddrs)
 	{
 		if (!BlockIds.Contains(It->Value.BlockId))
 		{
-			MissingReferencedBlocks.Add(GetBlockFilename(It->Value.BlockId));
+			FString Filename		= GetBlockFilename(It->Value.BlockId);
+			FString ErrorMessage	= FString::Printf(TEXT("Missing CAS block '%s'"), *Filename);
+			MissingReferencedBlocks.Add(MoveTemp(Filename));
 
 			OutAddrs.Add(It->Key);
 			It.RemoveCurrent();
 
-			Status = FIoStatus(EIoErrorCode::NotFound);
+			Result = MakeCasError<void>(ECasErrorCode::VerifyFailed, EIoErrorCode::NotFound, MoveTemp(ErrorMessage));
 		}
 	}
 
@@ -744,7 +760,7 @@ FIoStatus FCas::Verify(TArray<FCasAddr>& OutAddrs)
 		UE_LOG(LogIoStoreOnDemand, Warning, TEXT("Lookup references missing CAS block '%s'"), *Filename);
 	}
 
-	return Status; 
+	return Result;
 }
 
 FCas::FLastAccess FCas::GetAndClearDirtyLastAccess()
@@ -882,11 +898,11 @@ struct FCasJournal
 
 	using FEntryHandler		= TFunction<void(const FEntry&)>;
 
-	static FIoStatus		Replay(const FString& JournalFile, FEntryHandler&& Handler);
-	static FIoStatus		Create(const FString& JournalFile);
+	static FResult			Replay(const FString& JournalFile, FEntryHandler&& Handler);
+	static FResult			Create(const FString& JournalFile);
 	static FTransaction		Begin(FString&& JournalFile);
 	static FTransaction		Begin(const FString& JournalFile) { return Begin(FString(JournalFile)); }
-	static FIoStatus		Commit(FTransaction&& Transaction);
+	static FResult			Commit(FTransaction&& Transaction);
 };
 
 ///////////////////////////////////////////////////////////////////////////////
@@ -941,28 +957,27 @@ bool FCasJournal::FFooter::IsValid() const
 	return FMemory::Memcmp(Magic, FFooter::MagicSequence, sizeof(FFooter::MagicSequence)) == 0;
 }
 
-FIoStatus FCasJournal::Replay(const FString& JournalFile, FEntryHandler&& Handler)
+FResult FCasJournal::Replay(const FString& JournalFile, FEntryHandler&& Handler)
 {
+	using namespace UE::UnifiedError::IoStoreOnDemand;
+
 	IPlatformFile& Ipf = FPlatformFileManager::Get().GetPlatformFile();
 
 	if (Ipf.FileExists(*JournalFile) == false)
 	{
-		return FIoStatus(EIoErrorCode::NotFound);
+		return MakeJournalError(ECasErrorCode::ReplayJournalFailed, EIoErrorCode::NotFound, FString::Printf(TEXT("Failed to find '%s'"), *JournalFile));
 	}
 
 	TUniquePtr<IFileHandle> FileHandle(Ipf.OpenRead(*JournalFile));
 	if (FileHandle.IsValid() == false)
 	{
-		return FIoStatus(EIoErrorCode::FileOpenFailed, FPlatformMisc::GetLastError());
+		return MakeJournalError(ECasErrorCode::ReplayJournalFailed, EIoErrorCode::FileOpenFailed, FString::Printf(TEXT("Failed to open '%s'"), *JournalFile));
 	}
 
 	FHeader Header;
 	if ((FileHandle->Read(reinterpret_cast<uint8*>(&Header), FHeader::Size()) == false) || (Header.IsValid() == false))
 	{
-		return FIoStatusBuilder(EIoErrorCode::ReadError, FPlatformMisc::GetLastError())
-			<< TEXT("Failed to validate journal header '")
-			<< JournalFile
-			<< TEXT("'");
+		return MakeJournalError(ECasErrorCode::ReplayJournalFailed, EIoErrorCode::ReadError, FString::Printf(TEXT("Failed to validate journal header in '%s'"), *JournalFile));
 	}
 
 	const int64 FileSize	= FileHandle->Size();
@@ -970,39 +985,35 @@ FIoStatus FCasJournal::Replay(const FString& JournalFile, FEntryHandler&& Handle
 
 	if (EntryCount < 0)
 	{
-		return FIoStatus(EIoErrorCode::ReadError);
+		return MakeJournalError(ECasErrorCode::ReplayJournalFailed, EIoErrorCode::FileCorrupt, FString::Printf(TEXT("Invalid journal file '%s'"), *JournalFile));
 	}
 
 	if (EntryCount == 0)
 	{
-		return FIoStatus::Ok;
+		return MakeValue();
 	}
 
 	const int64 FooterPos = FileSize - FFooter::Size();
 	if (FooterPos < 0)
 	{
-		return FIoStatusBuilder(EIoErrorCode::FileCorrupt)
-			<< TEXT("Invalid journal footer");
+		return MakeJournalError(ECasErrorCode::ReplayJournalFailed, EIoErrorCode::FileCorrupt, FString::Printf(TEXT("Invalid journal footer in '%s'"), *JournalFile));
 	}
 
 	const int64 EntriesPos = FileHandle->Tell();
 	if (FileHandle->Seek(FooterPos) == false)
 	{
-		return FIoStatus(EIoErrorCode::FileSeekFailed, FPlatformMisc::GetLastError());
+		return MakeJournalError(ECasErrorCode::ReplayJournalFailed, EIoErrorCode::FileSeekFailed, FString::Printf(TEXT("Failed to seek to footer offset %ull in '%s'"), FooterPos, *JournalFile));
 	}
 
 	FFooter Footer;
 	if ((FileHandle->Read(reinterpret_cast<uint8*>(&Footer), FFooter::Size()) == false) || (Footer.IsValid() == false))
 	{
-		return FIoStatusBuilder(EIoErrorCode::ReadError, FPlatformMisc::GetLastError())
-			<< TEXT("Failed to validate journal footer '")
-			<< JournalFile
-			<< TEXT("'");
+		return MakeJournalError(ECasErrorCode::ReplayJournalFailed, EIoErrorCode::ReadError, FString::Printf(TEXT("Failed to validate journal footer in '%s'"), *JournalFile));
 	}
 
 	if (FileHandle->Seek(EntriesPos) == false)
 	{
-		return FIoStatus(EIoErrorCode::FileSeekFailed, FPlatformMisc::GetLastError());
+		return MakeJournalError(ECasErrorCode::ReplayJournalFailed, EIoErrorCode::FileSeekFailed, FString::Printf(TEXT("Failed to seek to entries offset %ull in '%s'"), EntriesPos, *JournalFile));
 	}
 
 	TArray<FEntry> Entries;
@@ -1010,7 +1021,7 @@ FIoStatus FCasJournal::Replay(const FString& JournalFile, FEntryHandler&& Handle
 
 	if (FileHandle->Read(reinterpret_cast<uint8*>(Entries.GetData()), FEntry::Size() * EntryCount) == false)
 	{
-		return FIoStatus(EIoErrorCode::ReadError, FPlatformMisc::GetLastError());
+		return MakeJournalError(ECasErrorCode::ReplayJournalFailed, EIoErrorCode::ReadError, FString::Printf(TEXT("Failed to read journal entries in '%s'"), *JournalFile));
 	}
 
 	UE_LOG(LogIoStoreOnDemand, Log, TEXT("Replaying %" INT64_FMT " CAS journal entries of total %.2lf KiB from '%s'"),
@@ -1024,24 +1035,28 @@ FIoStatus FCasJournal::Replay(const FString& JournalFile, FEntryHandler&& Handle
 			UE_LOG(LogIoStoreOnDemand, Warning, TEXT("Found critical error entry '%s' (%d) in journal '%s'"),
 				GetErrorText(Error.ErrorCode), Error.ErrorCode, *JournalFile);
 
-			return FIoStatus(EIoErrorCode::InvalidCode);
+			// We append "critical error" entries to the journal when we endup in an unrecoverable error state. This will cause the cache to be reset
+			// at startup
+			return MakeJournalError(ECasErrorCode::ReplayJournalFailed, EIoErrorCode::InvalidCode, FString::Printf(TEXT("Found critical error journal entry in '%s'"), *JournalFile));
 		}
 
 		Handler(Entry);
 	}
 
-	return FIoStatus::Ok;
+	return MakeValue();
 }
 
-FIoStatus FCasJournal::Create(const FString& JournalFile)
+FResult FCasJournal::Create(const FString& JournalFile)
 {
+	using namespace UE::UnifiedError::IoStoreOnDemand;
+
 	IPlatformFile& Ipf = FPlatformFileManager::Get().GetPlatformFile();
 	Ipf.DeleteFile(*JournalFile);
 
 	TUniquePtr<IFileHandle> FileHandle(Ipf.OpenWrite(*JournalFile));
 	if (FileHandle.IsValid() == false)
 	{
-		return FIoStatus(EIoErrorCode::FileOpenFailed, FPlatformMisc::GetLastError());
+		return MakeJournalError(ECasErrorCode::CreateJournalFailed, EIoErrorCode::FileOpenFailed, FString::Printf(TEXT("Failed to create journal '%s'"), *JournalFile));
 	}
 
 	FHeader Header;
@@ -1050,17 +1065,17 @@ FIoStatus FCasJournal::Create(const FString& JournalFile)
 
 	if (FileHandle->Write(reinterpret_cast<uint8*>(&Header), FHeader::Size()) == false)
 	{
-		return FIoStatus(EIoErrorCode::WriteError, FPlatformMisc::GetLastError());
+		return MakeJournalError(ECasErrorCode::CreateJournalFailed, EIoErrorCode::WriteError, FString::Printf(TEXT("Failed to write journal header in '%s'"), *JournalFile));
 	}
 
 	FFooter Footer;
 	FMemory::Memcpy(&Footer.Magic, &FFooter::MagicSequence, sizeof(FFooter::MagicSequence));
 	if (FileHandle->Write(reinterpret_cast<uint8*>(&Footer), FFooter::Size()) == false)
 	{
-		return FIoStatus(EIoErrorCode::WriteError, FPlatformMisc::GetLastError());
+		return MakeJournalError(ECasErrorCode::CreateJournalFailed, EIoErrorCode::WriteError, FString::Printf(TEXT("Failed to write journal footer in '%s'"), *JournalFile));
 	}
 
-	return FIoStatus::Ok;
+	return MakeValue();
 }
 
 FCasJournal::FTransaction FCasJournal::Begin(FString&& JournalFile)
@@ -1071,59 +1086,59 @@ FCasJournal::FTransaction FCasJournal::Begin(FString&& JournalFile)
 	};
 }
 
-FIoStatus FCasJournal::Commit(FTransaction&& Transaction)
+FResult FCasJournal::Commit(FTransaction&& Transaction)
 {
+	using namespace UE::UnifiedError::IoStoreOnDemand;
+
 	if (Transaction.Entries.IsEmpty())
 	{
-		return FIoStatus::Ok;
+		return MakeValue();
 	}
 
 	IPlatformFile& Ipf = FPlatformFileManager::Get().GetPlatformFile();
 
+	FResult Result			= MakeValue();
+	uint64 TotalEntrySize	= 0;
+	ON_SCOPE_EXIT
+	{
+		FOnDemandInstallCacheStats::OnJournalCommit(Result, TotalEntrySize);
+	};
+
 	// Validate header and footer
 	{
 		TUniquePtr<IFileHandle> FileHandle(Ipf.OpenRead(*Transaction.JournalFile));
 		if (FileHandle.IsValid() == false)
 		{
-			const uint32 LastError = FPlatformMisc::GetLastError();
-			FOnDemandInstallCacheStats::OnJournalCommit(EIoErrorCode::FileOpenFailed, 0);
-			return FIoStatusBuilder(EIoErrorCode::FileOpenFailed, LastError)
-				<< TEXT("Failed to validate CAS journal file");
+			return Result = MakeJournalError(
+				ECasErrorCode::CommitJournalFailed, EIoErrorCode::FileOpenFailed, FString::Printf(TEXT("Failed to open journal '%s'"), *Transaction.JournalFile));
 		}
 
 		const int64 FileSize = FileHandle->Size();
 		if (FileSize < FHeader::Size())
 		{
-			FOnDemandInstallCacheStats::OnJournalCommit(EIoErrorCode::FileCorrupt, 0);
-			return FIoStatusBuilder(EIoErrorCode::FileCorrupt, FPlatformMisc::GetLastError())
-				<< TEXT("Failed to validate CAS journal file");
+			return Result = MakeJournalError(
+				ECasErrorCode::CommitJournalFailed, EIoErrorCode::FileCorrupt, FString::Printf(TEXT("Failed to validate journal header in '%s'"), *Transaction.JournalFile));
 		}
 
 		FHeader Header;
 		if ((FileHandle->Read(reinterpret_cast<uint8*>(&Header), FHeader::Size()) == false) || (Header.IsValid() == false))
 		{
-			const uint32 LastError = FPlatformMisc::GetLastError();
-			FOnDemandInstallCacheStats::OnJournalCommit(EIoErrorCode::ReadError, 0);
-			return FIoStatusBuilder(EIoErrorCode::ReadError, LastError)
-				<< TEXT("Failed to validate CAS journal header");
+			return Result = MakeJournalError(
+				ECasErrorCode::CommitJournalFailed, EIoErrorCode::ReadError, FString::Printf(TEXT("Failed to validate journal header in '%s'"), *Transaction.JournalFile));
 		}
 
 		const int64 FooterPos = FileSize - FFooter::Size();
 		if (FileHandle->Seek(FooterPos) == false)
 		{
-			const uint32 LastError = FPlatformMisc::GetLastError();
-			FOnDemandInstallCacheStats::OnJournalCommit(EIoErrorCode::FileSeekFailed, 0);
-			return FIoStatusBuilder(EIoErrorCode::FileSeekFailed, LastError)
-				<< TEXT("Failed to validate CAS journal footer");
+			return Result = MakeJournalError(
+				ECasErrorCode::CommitJournalFailed, EIoErrorCode::FileSeekFailed, FString::Printf(TEXT("Failed to validate journal footer in '%s'"), *Transaction.JournalFile));
 		}
 
 		FFooter Footer;
 		if ((FileHandle->Read(reinterpret_cast<uint8*>(&Footer), FFooter::Size()) == false) || (Footer.IsValid() == false))
 		{
-			const uint32 LastError = FPlatformMisc::GetLastError();
-			FOnDemandInstallCacheStats::OnJournalCommit(EIoErrorCode::ReadError, 0);
-			return FIoStatusBuilder(EIoErrorCode::ReadError, LastError)
-				<< TEXT("Failed to validate CAS journal footer");
+			return Result = MakeJournalError(
+				ECasErrorCode::CommitJournalFailed, EIoErrorCode::ReadError, FString::Printf(TEXT("Found to validate journal footer in '%s'"), *Transaction.JournalFile));
 		}
 	}
 
@@ -1136,55 +1151,44 @@ FIoStatus FCasJournal::Commit(FTransaction&& Transaction)
 
 		if (EntriesPos < 0) 
 		{
-			const uint32 LastError = FPlatformMisc::GetLastError();
-			FOnDemandInstallCacheStats::OnJournalCommit(EIoErrorCode::FileOpenFailed, 0);
-			return FIoStatusBuilder(EIoErrorCode::FileOpenFailed, LastError)
-				<< TEXT("Failed to open CAS journal");
+			return Result = MakeJournalError(
+				ECasErrorCode::CommitJournalFailed, EIoErrorCode::FileOpenFailed, FString::Printf(TEXT("Failed to open journal '%s'"), *Transaction.JournalFile));
 		}
 
 		if (FileHandle->Seek(EntriesPos) == false)
 		{
-			const uint32 LastError = FPlatformMisc::GetLastError();
-			FOnDemandInstallCacheStats::OnJournalCommit(EIoErrorCode::FileSeekFailed, 0);
-			return FIoStatusBuilder(EIoErrorCode::FileSeekFailed, LastError)
-				<< TEXT("Failed to seek to entries in CAS journal");
+			return Result = MakeJournalError(
+				ECasErrorCode::CommitJournalFailed, EIoErrorCode::FileSeekFailed, FString::Printf(TEXT("Failed to seek to journal entries offset %ull in '%s'"), EntriesPos, *Transaction.JournalFile));
 		}
 
-		const int64 TotalEntrySize = Transaction.Entries.Num() * FEntry::Size();
+		TotalEntrySize = Transaction.Entries.Num() * FEntry::Size();
 		if (FileHandle->Write(
 			reinterpret_cast<const uint8*>(Transaction.Entries.GetData()),
 			TotalEntrySize) == false)
 		{
-			const uint32 LastError = FPlatformMisc::GetLastError();
-			FOnDemandInstallCacheStats::OnJournalCommit(EIoErrorCode::WriteError, 0);
-			return FIoStatusBuilder(EIoErrorCode::WriteError, LastError)
-				<< TEXT("Failed to write CAS journal entries");
+			return Result = MakeJournalError(
+				ECasErrorCode::CommitJournalFailed, EIoErrorCode::WriteError, FString::Printf(TEXT("Failed to write journal entries to '%s'"), *Transaction.JournalFile));
 		}
 
 		FFooter Footer;
 		FMemory::Memcpy(&Footer.Magic, &FFooter::MagicSequence, sizeof(FFooter::MagicSequence));
 		if (FileHandle->Write(reinterpret_cast<uint8*>(&Footer), FFooter::Size()) == false)
 		{
-			const uint32 LastError = FPlatformMisc::GetLastError();
-			FOnDemandInstallCacheStats::OnJournalCommit(EIoErrorCode::WriteError, 0);
-			return FIoStatusBuilder(EIoErrorCode::WriteError, LastError)
-				<< TEXT("Failed to write CAS journal footer");
+			return Result = MakeJournalError(
+				ECasErrorCode::CommitJournalFailed, EIoErrorCode::WriteError, FString::Printf(TEXT("Failed to write journal footer to '%s'"), *Transaction.JournalFile));
 		}
 
 		if (FileHandle->Flush() == false)
 		{
-			const uint32 LastError = FPlatformMisc::GetLastError();
-			FOnDemandInstallCacheStats::OnJournalCommit(EIoErrorCode::FileFlushFailed, 0);
-			return FIoStatusBuilder(EIoErrorCode::FileFlushFailed, FPlatformMisc::GetLastError())
-				<< TEXT("Failed to flush CAS journal");
+			return Result = MakeJournalError(
+				ECasErrorCode::CommitJournalFailed, EIoErrorCode::FileFlushFailed, FString::Printf(TEXT("Failed to flush journal entries to '%s'"), *Transaction.JournalFile));
 		}
 
 		UE_LOG(LogIoStoreOnDemand, Log, TEXT("Committed %d CAS journal entries of total %.2lf KiB to '%s'"),
 			Transaction.Entries.Num(), ToKiB(TotalEntrySize), *Transaction.JournalFile);
 
-		FOnDemandInstallCacheStats::OnJournalCommit(EIoErrorCode::Ok, TotalEntrySize);
-
-		return FIoStatus::Ok;
+		check(Result.HasValue());
+		return Result;
 	}
 }
 
@@ -1287,10 +1291,10 @@ struct FCasSnapshot
 
 	using FChunkLocation = TPair<FCasAddr, FCasLocation>;
 
-	static TIoStatusOr<FCasSnapshot>	FromJournal(const FString& JournalFile);
-	static TIoStatusOr<FCasSnapshot>	Load(const FString& SnapshotFile, int64* OutFileSize = nullptr);
-	static TIoStatusOr<int64>			Save(const FCasSnapshot& Snapshot, const FString& SnapshotFile);
-	static TIoStatusOr<int64>			TryCreateAndResetJournal(const FString& SnapshotFile, const FString& JournalFile); 
+	static TResult<FCasSnapshot>		FromJournal(const FString& JournalFile);
+	static TResult<FCasSnapshot>		Load(const FString& SnapshotFile, int64* OutFileSize = nullptr);
+	static TResult<int64>				Save(const FCasSnapshot& Snapshot, const FString& SnapshotFile);
+	static TResult<int64>				TryCreateAndResetJournal(const FString& SnapshotFile, const FString& JournalFile);
 
 	TArray<FBlock>						Blocks;
 	TArray<FChunkLocation>				ChunkLocations;
@@ -1318,14 +1322,14 @@ bool FCasSnapshot::FFooter::IsValid() const
 	return FMemory::Memcmp(Magic, FFooter::MagicSequence, sizeof(FFooter::MagicSequence)) == 0;
 }
 
-TIoStatusOr<FCasSnapshot> FCasSnapshot::FromJournal(const FString& JournalFile)
+TResult<FCasSnapshot> FCasSnapshot::FromJournal(const FString& JournalFile)
 {
 	FCas::FLookup		CasLookup;
 	FCas::FLastAccess	LastAccess;
 	TSet<FCasBlockId>	BlockIds;
 	FCasBlockId			CurrentBlockId;
 
-	FIoStatus ReplayStatus = FCasJournal::Replay(
+	FResult ReplayResult = FCasJournal::Replay(
 		JournalFile,
 		[&CasLookup, &LastAccess, &BlockIds, &CurrentBlockId](const FCasJournal::FEntry& JournalEntry)
 		{
@@ -1371,9 +1375,9 @@ TIoStatusOr<FCasSnapshot> FCasSnapshot::FromJournal(const FString& JournalFile)
 			};
 		});
 
-	if (ReplayStatus.IsOk() == false)
+	if (ReplayResult.HasError())
 	{
-		return ReplayStatus;
+		return MakeError(ReplayResult.StealError());
 	}
 
 	const int64 TimestampForMissingLastAccess = FCas::GetTimestampForMissingLastAccess();
@@ -1393,11 +1397,13 @@ TIoStatusOr<FCasSnapshot> FCasSnapshot::FromJournal(const FString& JournalFile)
 	Snapshot.ChunkLocations = CasLookup.Array();
 	Snapshot.CurrentBlockId = CurrentBlockId;
 
-	return Snapshot;
+	return MakeValue(Snapshot);
 }
 
-TIoStatusOr<int64> FCasSnapshot::Save(const FCasSnapshot& Snapshot, const FString& SnapshotFile)
+TResult<int64> FCasSnapshot::Save(const FCasSnapshot& Snapshot, const FString& SnapshotFile)
 {
+	using namespace UE::UnifiedError::IoStoreOnDemand;
+
 	IFileManager& Ifm = IFileManager::Get();
 
 	const FString TmpSnapshotFile = FPaths::ChangeExtension(SnapshotFile, TEXT(".snptmp"));
@@ -1405,7 +1411,7 @@ TIoStatusOr<int64> FCasSnapshot::Save(const FCasSnapshot& Snapshot, const FStrin
 	TUniquePtr<FArchive> Ar(Ifm.CreateFileWriter(*TmpSnapshotFile));
 	if (Ar.IsValid() == false)
 	{
-		return FIoStatus(EIoErrorCode::FileOpenFailed, FPlatformMisc::GetLastError());
+		return MakeSnapshotError<int64>(ECasErrorCode::SaveSnapshotFailed, EIoErrorCode::FileOpenFailed, FString::Printf(TEXT("Failed to open file '%s' for writing"), *SnapshotFile));
 	}
 
 	FHeader Header;
@@ -1418,7 +1424,7 @@ TIoStatusOr<int64> FCasSnapshot::Save(const FCasSnapshot& Snapshot, const FStrin
 		const uint32 LastError = FPlatformMisc::GetLastError();
 		Ar.Reset();
 		Ifm.Delete(*TmpSnapshotFile);
-		return FIoStatus(EIoErrorCode::WriteError, LastError);
+		return MakeSnapshotError<int64>(ECasErrorCode::SaveSnapshotFailed, EIoErrorCode::WriteError, FString::Printf(TEXT("Failed to write snapshot header to '%s'"), *SnapshotFile), LastError);
 	}
 
 	FCasSnapshot& NonConst = *const_cast<FCasSnapshot*>(&Snapshot);
@@ -1431,7 +1437,7 @@ TIoStatusOr<int64> FCasSnapshot::Save(const FCasSnapshot& Snapshot, const FStrin
 		const uint32 LastError = FPlatformMisc::GetLastError();
 		Ar.Reset();
 		Ifm.Delete(*TmpSnapshotFile);
-		return FIoStatus(EIoErrorCode::WriteError, LastError);
+		return MakeSnapshotError<int64>(ECasErrorCode::SaveSnapshotFailed, EIoErrorCode::WriteError, FString::Printf(TEXT("Failed to write snapshot to '%s'"), *SnapshotFile), LastError);
 	}
 
 	FFooter Footer;
@@ -1442,7 +1448,7 @@ TIoStatusOr<int64> FCasSnapshot::Save(const FCasSnapshot& Snapshot, const FStrin
 		const uint32 LastError = FPlatformMisc::GetLastError();
 		Ar.Reset();
 		Ifm.Delete(*TmpSnapshotFile);
-		return FIoStatus(EIoErrorCode::WriteError, LastError);
+		return MakeSnapshotError<int64>(ECasErrorCode::SaveSnapshotFailed, EIoErrorCode::WriteError, FString::Printf(TEXT("Failed to write snapshot footer to '%s'"), *SnapshotFile), LastError);
 	}
 
 	const int64 FileSize = Ar->TotalSize();
@@ -1451,37 +1457,39 @@ TIoStatusOr<int64> FCasSnapshot::Save(const FCasSnapshot& Snapshot, const FStrin
 		const uint32 LastError = FPlatformMisc::GetLastError();
 		Ar.Reset();
 		Ifm.Delete(*TmpSnapshotFile);
-		return FIoStatus(EIoErrorCode::FileCloseFailed, LastError);
+		return MakeSnapshotError<int64>(ECasErrorCode::SaveSnapshotFailed, EIoErrorCode::WriteError, FString::Printf(TEXT("Failed to close snapshot footer to '%s'"), *SnapshotFile), LastError);
 	}
 
 	if (Ifm.Move(*SnapshotFile, *TmpSnapshotFile) == false)
 	{
 		const uint32 LastError = FPlatformMisc::GetLastError();
 		Ifm.Delete(*TmpSnapshotFile);
-		return FIoStatus(EIoErrorCode::FileMoveFailed, LastError);
+		return MakeSnapshotError<int64>(
+			ECasErrorCode::SaveSnapshotFailed,
+			EIoErrorCode::FileMoveFailed,
+			FString::Printf(TEXT("Failed to move snapshot '%s' -> '%s'"), *TmpSnapshotFile, *SnapshotFile), LastError);
 	}
 
-	return FileSize; 
+	return MakeValue(FileSize);
 }
 
-TIoStatusOr<FCasSnapshot> FCasSnapshot::Load(const FString& SnapshotFile, int64* OutFileSize)
+TResult<FCasSnapshot> FCasSnapshot::Load(const FString& SnapshotFile, int64* OutFileSize)
 {
+	using namespace UE::UnifiedError::IoStoreOnDemand;
+
 	IFileManager& Ifm = IFileManager::Get();
 
 	TUniquePtr<FArchive> Ar(Ifm.CreateFileReader(*SnapshotFile));
 	if (Ar.IsValid() == false)
 	{
-		return FIoStatus(EIoErrorCode::NotFound, FPlatformMisc::GetLastError());
+		return MakeSnapshotError<FCasSnapshot>(ECasErrorCode::LoadSnapshotFailed, EIoErrorCode::FileOpenFailed, FString::Printf(TEXT("Failed open snapshot '%s'"), *SnapshotFile));
 	}
 
 	FHeader Header;
 	Ar->Serialize(reinterpret_cast<uint8*>(&Header), FHeader::Size());
 	if (Ar->IsError() || Header.IsValid() == false)
 	{
-		return FIoStatusBuilder(EIoErrorCode::ReadError, FPlatformMisc::GetLastError())
-			<< TEXT("Failed to validate snapshot header '")
-			<< SnapshotFile 
-			<< TEXT("'");
+		return MakeSnapshotError<FCasSnapshot>(ECasErrorCode::LoadSnapshotFailed, EIoErrorCode::ReadError, FString::Printf(TEXT("Failed to validate snapshot header in '%s'"), *SnapshotFile));
 	}
 
 	FCasSnapshot Snapshot;
@@ -1493,58 +1501,59 @@ TIoStatusOr<FCasSnapshot> FCasSnapshot::Load(const FString& SnapshotFile, int64*
 	Ar->Serialize(reinterpret_cast<uint8*>(&Footer), FFooter::Size());
 	if (Ar->IsError() || Footer.IsValid() == false)
 	{
-		return FIoStatusBuilder(EIoErrorCode::ReadError, FPlatformMisc::GetLastError())
-			<< TEXT("Failed to validate snapshot footer '")
-			<< SnapshotFile 
-			<< TEXT("'");
+		return MakeSnapshotError<FCasSnapshot>(ECasErrorCode::LoadSnapshotFailed, EIoErrorCode::ReadError, FString::Printf(TEXT("Failed to validate snapshot footer in '%s'"), *SnapshotFile));
 	}
 
 	if (OutFileSize != nullptr)
 	{
 		*OutFileSize = Ar->Tell();
 	}
-	return Snapshot;
+
+	return MakeValue(Snapshot);
 }
 
-TIoStatusOr<int64> FCasSnapshot::TryCreateAndResetJournal(const FString& SnapshotFile, const FString& JournalFile)
+TResult<int64> FCasSnapshot::TryCreateAndResetJournal(const FString& SnapshotFile, const FString& JournalFile)
 {
+	using namespace UE::UnifiedError::IoStoreOnDemand;
+
 	IFileManager& Ifm = IFileManager::Get();
 
 	const int64 JournalFileSize = Ifm.FileSize(*JournalFile);
 	if (JournalFileSize < 0)
 	{
-		return FIoStatus(EIoErrorCode::NotFound);
+		return MakeSnapshotError<int64>(ECasErrorCode::CreateSnapshotFailed, EIoErrorCode::NotFound, FString::Printf(TEXT("Failed to find snapshot '%s'"), *SnapshotFile));
 	}
 
 	// Load the snapshot from the journal
-	TIoStatusOr<FCasSnapshot> SnapshotStatus = FCasSnapshot::FromJournal(JournalFile);
-	if (SnapshotStatus.IsOk() == false)
+	TResult<FCasSnapshot> SnapshotResult = FCasSnapshot::FromJournal(JournalFile);
+	if (SnapshotResult.HasError())
 	{
-		return SnapshotStatus.Status();
+		return MakeError(SnapshotResult.StealError());
 	}
 
 	// Save the snapshot
-	int64 SnapshotSize		= -1;
-	FCasSnapshot Snapshot	= SnapshotStatus.ConsumeValueOrDie();
-	if (TIoStatusOr<int64> Status = FCasSnapshot::Save(Snapshot, SnapshotFile); Status.IsOk())
+	int64 SnapshotSize			= -1;
+	FCasSnapshot Snapshot		= SnapshotResult.StealValue();
+	TResult<int64> SaveResult	= FCasSnapshot::Save(Snapshot, SnapshotFile);
+	if (SaveResult.HasValue())
 	{
-		SnapshotSize = Status.ConsumeValueOrDie();
+		SnapshotSize = SaveResult.GetValue();
 	}
 	else
 	{
-		return Status.Status();
+		return SaveResult;
 	}
 
 	// Try create a new empty journal 
 	const FString TmpJournalFile = FPaths::ChangeExtension(JournalFile, TEXT(".jrntmp"));
-	if (FIoStatus Status = FCasJournal::Create(TmpJournalFile); Status.IsOk() == false)
+	if (FResult Result = FCasJournal::Create(TmpJournalFile); Result.HasError())
 	{
 		if (Ifm.Delete(*SnapshotFile) == false)
 		{
 			UE_LOG(LogIoStoreOnDemand, Error, TEXT("Failed to delete CAS snapshot '%s'"), *SnapshotFile); 
 		}
 
-		return Status;
+		return MakeError(Result.StealError());
 	}
 
 	if (Ifm.Move(*JournalFile , *TmpJournalFile) == false)
@@ -1555,10 +1564,13 @@ TIoStatusOr<int64> FCasSnapshot::TryCreateAndResetJournal(const FString& Snapsho
 			UE_LOG(LogIoStoreOnDemand, Error, TEXT("Failed to delete CAS snapshot '%s'"), *SnapshotFile); 
 		}
 
-		return FIoStatus(EIoErrorCode::FileMoveFailed, LastError);
+		return MakeSnapshotError<int64>(
+			ECasErrorCode::CreateSnapshotFailed,
+			EIoErrorCode::FileMoveFailed,
+			FString::Printf(TEXT("Failed to move tmp journal file '%s' -> '%s'"), *TmpJournalFile, *JournalFile));
 	}
 
-	return SnapshotSize;
+	return MakeValue(SnapshotSize);
 }
 
 ///////////////////////////////////////////////////////////////////////////////
@@ -1721,8 +1733,8 @@ public:
 
 private:
 	void										RegisterConsoleCommands();
-	FIoStatus									Reset();
-	FIoStatus									InitialVerify();
+	FResult										Reset();
+	FResult										InitialVerify();
 	uint64										AddReferencesToBlocks(
 													const TArray<FSharedOnDemandContainer>& Containers, 
 													const TArray<TBitArray<>>& ChunkEntryIndices,
@@ -1742,7 +1754,7 @@ private:
 	FResult										FlushPendingChunksImpl(const FPendingChunks& Chunks, int64 UtcAccessTicks = 0);
 	FString										GetJournalFilename() const { return CacheDirectory / TEXT("cas.jrn"); }
 	FString										GetSnapshotFilename() const { return CacheDirectory / TEXT("cas.snp"); }
-	void										CaptureCacheState(UE::UnifiedError::IoStoreOnDemand::FInstallCacheErrorContext& OutCtx);
+	UE::UnifiedError::IoStoreOnDemand::FInstallCacheErrorContext MakeInstallCacheErrorContext(uint64 TotalCachedBytes = 0, uint32 LineNo = __builtin_LINE());
 
 	FOnDemandIoStore&			IoStore;
 	FString						CacheDirectory;
@@ -1773,6 +1785,8 @@ FOnDemandInstallCache::FOnDemandInstallCache(const FOnDemandInstallCacheConfig&
 	, MaxCacheSize(Config.DiskQuota)
 	, MaxJournalSize(Config.JournalMaxSize)
 {
+	using namespace UE::UnifiedError::IoStoreOnDemand;
+
 	UE_LOG(LogIoStoreOnDemand, Log, TEXT("Initializing install cache, MaxCacheSize=%.2lf MiB, MaxJournalSize=%.2lf KiB"),
 		ToMiB(MaxCacheSize), ToKiB(MaxJournalSize));
 
@@ -1787,10 +1801,10 @@ FOnDemandInstallCache::FOnDemandInstallCache(const FOnDemandInstallCacheConfig&
 	MaxCacheSize -= Cas.GetMaxBlockSize();
 	UE_LOG(LogIoStoreOnDemand, Log, TEXT("Effective MaxCacheSize without defragmentation space is MaxCacheSize=%.2lf MiB"), ToMiB(MaxCacheSize));
 
-	FIoStatus Status = Cas.Initialize(CacheDirectory);
-	if (Status.IsOk() == false)
+	FResult InitResult = Cas.Initialize(CacheDirectory);
+	if (InitResult.HasError())
 	{
-		UE_LOG(LogIoStoreOnDemand, Error, TEXT("Failed to initialize install cache, reason '%s'"), *Status.ToString());
+		UE_LOGFMT(LogIoStoreOnDemand, Error, "Failed to initialize install cache, error: {Error}", InitResult.GetError());
 		return;
 	}
 
@@ -1798,26 +1812,30 @@ FOnDemandInstallCache::FOnDemandInstallCache(const FOnDemandInstallCacheConfig&
 	{
 		const FString SnapshotFile = GetSnapshotFilename();
 		int64 SnapshotSize = -1;
-		TIoStatusOr<FCasSnapshot> SnapshotStatus = FCasSnapshot::Load(SnapshotFile, &SnapshotSize);
-		if (SnapshotStatus.IsOk())
+		TResult<FCasSnapshot> SnapshotResult = FCasSnapshot::Load(SnapshotFile, &SnapshotSize);
+		if (SnapshotResult.HasValue())
 		{
-			FCasSnapshot Snapshot = SnapshotStatus.ConsumeValueOrDie();
+			FCasSnapshot Snapshot = SnapshotResult.StealValue();
 			UE_LOG(LogIoStoreOnDemand, Log, TEXT("Loaded CAS snapshot '%s' %.2lf KiB with %d blocks and %d chunk locations"),
 					*SnapshotFile, ToKiB(SnapshotSize), Snapshot.Blocks.Num(), Snapshot.ChunkLocations.Num());
 
 			Cas.LoadSnapshot(MoveTemp(Snapshot));
 		}
-		else if (SnapshotStatus.Status().GetErrorCode() != EIoErrorCode::NotFound) // Its ok for the snapshot to not exist, other errors are significant
+		else if (const FCasErrorContext* Ctx = SnapshotResult.GetError().GetErrorContext<FCasErrorContext>())
 		{
-			Status = SnapshotStatus.Status();
+			// It's ok for the snapshot to not exist, other errors are significant
+			if (Ctx->IoErrorCode == EIoErrorCode::NotFound)
+			{
+				InitResult = MakeValue();
+			}
 		}
 	}
 
 	const FString JournalFile = GetJournalFilename();
-	if (Status.IsOk())
+	if (InitResult.HasValue())
 	{
 		// Replay the journal 
-		Status = FCasJournal::Replay(JournalFile, [this](const FCasJournal::FEntry& JournalEntry)
+		InitResult = FCasJournal::Replay(JournalFile, [this](const FCasJournal::FEntry& JournalEntry)
 		{
 			switch(JournalEntry.Type())
 			{
@@ -1861,40 +1879,41 @@ FOnDemandInstallCache::FOnDemandInstallCache(const FOnDemandInstallCacheConfig&
 		});
 	}
 
-	// Initializing the cache for the first time
-	if (Status.GetErrorCode() == EIoErrorCode::NotFound)
+	// If the CAS journal is not found we assume we are initializing the cache for the first time
+	if (InitResult.HasError())
 	{
-		if (Status = FCasJournal::Create(JournalFile); Status.IsOk())
+		if (const FCasErrorContext* Ctx = InitResult.GetError().GetErrorContext<FCasErrorContext>())
 		{
-			UE_LOG(LogIoStoreOnDemand, Log, TEXT("Created CAS journal '%s'"), *JournalFile);
+			if (Ctx->IoErrorCode == EIoErrorCode::NotFound)
+			{
+				if (InitResult = FCasJournal::Create(JournalFile); InitResult.HasValue())
+				{
+					UE_LOG(LogIoStoreOnDemand, Log, TEXT("Created CAS journal '%s'"), *JournalFile);
 
-			// Make sure that there are no existing blocks when starting from an empty cache
-			const bool bDeleteExisting = true;
-			Status = Cas.Initialize(CacheDirectory, bDeleteExisting);
-		}
-		else
-		{
-			UE_LOG(LogIoStoreOnDemand, Error, TEXT("Failed to create CAS journal '%s'"), *JournalFile);
+					// Make sure that there are no existing blocks when starting from an empty cache
+					const bool bDeleteExisting = true;
+					InitResult = Cas.Initialize(CacheDirectory, bDeleteExisting);
+				}
+			}
 		}
 	}
 
 	// Verify the current state of the cache
-	if (Status.IsOk())
+	if (InitResult.HasValue())
 	{
-		Status = InitialVerify();
+		InitResult = InitialVerify();
 	}
 
 	// Try to reset the cache if something has gone wrong 
-	if (Status.IsOk() == false)
+	if (InitResult.HasError())
 	{
-		UE_LOG(LogIoStoreOnDemand, Warning, TEXT("Resetting install cash, reason '%s'"),
-			*Status.ToString());
+		UE_LOGFMT(LogIoStoreOnDemand, Error, "Resetting install cache, reason: {Error}", InitResult.GetError());
 
-		FOnDemandInstallCacheStats::OnStartupError(Status.GetErrorCode());
-		Status = Reset();
+		FOnDemandInstallCacheStats::OnStartupError(InitResult);
+		InitResult = InitialVerify();
 	}
 
-	if (Status.IsOk())
+	if (InitResult.HasValue())
 	{
 		UE_LOG(LogIoStoreOnDemand, Log, TEXT("Install cache Ok!"));
 		RegisterConsoleCommands();
@@ -1902,8 +1921,7 @@ FOnDemandInstallCache::FOnDemandInstallCache(const FOnDemandInstallCacheConfig&
 	}
 	else
 	{
-		UE_LOG(LogIoStoreOnDemand, Error, TEXT("Failed to initialize install cache, reason '%s'"),
-			*Status.ToString());
+		UE_LOGFMT(LogIoStoreOnDemand, Error, "Failed to initialize install cache, reason: {Error}", InitResult.GetError());
 	}
 }
 
@@ -1929,10 +1947,10 @@ void FOnDemandInstallCache::Shutdown()
 	FCasJournal::FTransaction Transaction = FCasJournal::Begin(JournalFile);
 	JournalLastAccess(Transaction, LastAccess);
 
-	if (FIoStatus Status = FCasJournal::Commit(MoveTemp(Transaction)); !Status.IsOk())
+	if (FResult Result = FCasJournal::Commit(MoveTemp(Transaction)); Result.HasError())
 	{
-		UE_LOG(LogIoStoreOnDemand, Error, TEXT("Failed to update CAS journal '%s' with block timestamp(s), reason '%s'"),
-			*JournalFile, *Status.ToString());
+		UE_LOG(LogIoStoreOnDemand, Error, TEXT("Failed to update CAS journal '%s' with block timestamp(s), error '%s'"),
+			*JournalFile, *LexToString(Result.GetError()));
 	}
 
 	// TODO: journal snapshot is also only made on shutdown? Feels like this would be better done on startup as well
@@ -1941,15 +1959,15 @@ void FOnDemandInstallCache::Shutdown()
 	if (Ifm.FileSize(*JournalFile) > int64(MaxJournalSize))
 	{
 		const FString SnapshotFilename = GetSnapshotFilename();
-		TIoStatusOr<int64> SnapshotStatus = FCasSnapshot::TryCreateAndResetJournal(SnapshotFilename, JournalFilename);
-		if (SnapshotStatus.IsOk())
+		TResult<int64> SnapshotResult = FCasSnapshot::TryCreateAndResetJournal(SnapshotFilename, JournalFilename);
+		if (SnapshotResult.HasValue())
 		{
-			UE_LOG(LogIoStoreOnDemand, Log, TEXT("Saved CAS snapshot '%s' %.2lf KiB"), *SnapshotFilename, ToKiB(SnapshotStatus.ValueOrDie()));
+			UE_LOG(LogIoStoreOnDemand, Log, TEXT("Saved CAS snapshot '%s' %.2lf KiB"), *SnapshotFilename, ToKiB(SnapshotResult.GetValue()));
 		}
 		else
 		{
-			UE_LOG(LogIoStoreOnDemand, Error, TEXT("Failed to create CAS snapshot from journal '%s', reason '%s'"),
-				*JournalFile, *SnapshotStatus.Status().ToString());
+			UE_LOG(LogIoStoreOnDemand, Error, TEXT("Failed to create CAS snapshot from journal '%s', error '%s'"),
+				*JournalFile, *LexToString(SnapshotResult.GetError()));
 		}
 	}
 
@@ -2147,8 +2165,8 @@ bool FOnDemandInstallCache::Resolve(FIoRequestImpl* Request)
 
 			const FString Filename = Cas.GetBlockFilename(CasLoc.BlockId);
 
-			FSharedFileOpenResult FileOpenResult = Cas.OpenRead(CasLoc.BlockId);
-			if (FileOpenResult.IsValid())
+			TResult<FSharedFileHandle> FileOpenResult = Cas.OpenRead(CasLoc.BlockId);
+			if (FileOpenResult.HasValue())
 			{
 				Status = EIoErrorCode::ReadError;
 
@@ -2180,8 +2198,7 @@ bool FOnDemandInstallCache::Resolve(FIoRequestImpl* Request)
 			}
 			else
 			{
-				UE_LOG(LogIoStoreOnDemand, Error, TEXT("Failed to open CAS block '%s' for reading, reason '%s'"),
-					*Filename, *FileOpenResult.GetError().GetMessage());
+				UE_LOGFMT(LogIoStoreOnDemand, Error, "Failed to open CAS block '{Filename}' for reading ({Error})", *Filename, FileOpenResult.GetError());
 			}
 
 			UE::Tasks::Launch(UE_SOURCE_LOCATION, [this, Request, Status]
@@ -2319,6 +2336,8 @@ FResult FOnDemandInstallCache::PutChunk(FIoBuffer&& Chunk, const FIoHash& ChunkH
 
 FResult FOnDemandInstallCache::Purge(TSet<FIoHash>&& ChunksToInstall)
 {
+	using namespace UE::UnifiedError::IoStoreOnDemand;
+
 	TUniqueLock Lock(PurgeDefragMutex);
 
 	FCasBlockInfoMap	BlockInfo;
@@ -2409,15 +2428,9 @@ FResult FOnDemandInstallCache::Purge(TSet<FIoHash>&& ChunksToInstall)
 		}
 		else
 		{
-			UE::UnifiedError::IoStoreOnDemand::FInstallCacheErrorContext Ctx
-			{
-				.Error			= FIoStatus(EIoErrorCode::InvalidCode, TEXT("Failed to purge required size from install cache")),
-				.CacheSize		= TotalCachedBytes,
-				.ErrorCategory	= UE::UnifiedError::IoStoreOnDemand::EErrorCategory::Logical
-			};
-			CaptureCacheState(Ctx);
-			Result = MakeError(UE::UnifiedError::IoStoreOnDemand::InstallCachePurgeError::MakeError(
-				MoveTemp(Ctx), UE::UnifiedError::EDetailFilter::All));
+			FInstallCacheErrorContext Ctx = MakeInstallCacheErrorContext(TotalCachedBytes);
+			Ctx.IoError = FIoStatus(EIoErrorCode::InvalidCode, TEXT("Failed to purge required size from install cache"));
+			Result = MakeError(InstallCachePurgeError::MakeError(MoveTemp(Ctx), UE::UnifiedError::EDetailFilter::All));
 		}
 	}
 
@@ -2427,6 +2440,8 @@ FResult FOnDemandInstallCache::Purge(TSet<FIoHash>&& ChunksToInstall)
 
 FResult FOnDemandInstallCache::PurgeAllUnreferenced(bool bDefrag, const uint64* BytesToPurge /*= nullptr*/)
 {
+	using namespace UE::UnifiedError::IoStoreOnDemand;
+
 	TUniqueLock Lock(PurgeDefragMutex);
 
 	FCasBlockInfoMap	BlockInfo;
@@ -2473,15 +2488,8 @@ FResult FOnDemandInstallCache::PurgeAllUnreferenced(bool bDefrag, const uint64*
 		}
 		else
 		{
-			UE::UnifiedError::IoStoreOnDemand::FInstallCacheErrorContext Ctx
-			{
-				.Error			= FIoStatus(EIoErrorCode::InvalidCode, TEXT("Failed to purge all unreferenced cache block(s)")),
-				.CacheSize		= TotalCachedBytes,
-				.ErrorCategory	= UE::UnifiedError::IoStoreOnDemand::EErrorCategory::Logical
-			};
-			CaptureCacheState(Ctx);
-			Result = MakeError(UE::UnifiedError::IoStoreOnDemand::InstallCachePurgeError::MakeError(
-				MoveTemp(Ctx), UE::UnifiedError::EDetailFilter::All));
+			FInstallCacheErrorContext Ctx = MakeInstallCacheErrorContext();
+			Result = MakeError(InstallCachePurgeError::MakeError(MoveTemp(Ctx), UE::UnifiedError::EDetailFilter::All));
 		}
 	}
 	else if (bDefrag)
@@ -2605,10 +2613,10 @@ FResult FOnDemandInstallCache::Verify()
 	UE_LOG(LogIoStoreOnDemand, Log, TEXT("Verifying %d installed chunks..."), ChunkCount);
 	for (int32 ChunkIndex = 0; const FCasAddrLocation& ChunkLocation : ChunkLocations)
 	{
-		FSharedFileOpenResult OpenResult = Cas.OpenRead(ChunkLocation.Location.BlockId);
+		TResult<FSharedFileHandle> OpenResult = Cas.OpenRead(ChunkLocation.Location.BlockId);
 		if (OpenResult.HasError())
 		{
-			UE_LOG(LogIoStoreOnDemand, Error, TEXT("Failed to open block %u for reading"), ChunkLocation.Location.BlockId.Id);
+			UE_LOGFMT(LogIoStoreOnDemand, Error, "Failed to open block {BlockId} for reading ({Error})", ChunkLocation.Location.BlockId.Id, OpenResult.GetError());
 
 			ReadErrorCount++;
 			ChunkIndex++;
@@ -2723,10 +2731,10 @@ void FOnDemandInstallCache::RegisterConsoleCommands()
 
 			FCasJournal::FTransaction Transaction = FCasJournal::Begin(GetJournalFilename());
 			Transaction.CriticalError(FCasJournal::EErrorCode::Simulated);
-			if (FIoStatus Status = FCasJournal::Commit(MoveTemp(Transaction)); Status.IsOk() == false)
+			if (FResult Result = FCasJournal::Commit(MoveTemp(Transaction)); Result.HasError())
 			{
-				UE_LOG(LogIoStoreOnDemand, Warning, TEXT("Failed to append critical error to journal, reason '%s'"),
-					*Status.ToString());
+				UE_LOG(LogIoStoreOnDemand, Warning, TEXT("Failed to append critical error to journal, error '%s'"),
+					*LexToString(Result.GetError()));
 			}
 		}),
 		ECVF_Default)
@@ -2734,8 +2742,10 @@ void FOnDemandInstallCache::RegisterConsoleCommands()
 #endif // UE_IAD_DEBUG_CONSOLE_CMDS
 }
 
-FIoStatus FOnDemandInstallCache::Reset()
+FResult FOnDemandInstallCache::Reset()
 {
+	using namespace UE::UnifiedError::IoStoreOnDemand;
+
 	UE_LOG(LogIoStoreOnDemand, Log, TEXT("Resetting install cache in directory '%s'"), *CacheDirectory);
 
 	IFileManager& Ifm	= IFileManager::Get();
@@ -2743,41 +2753,37 @@ FIoStatus FOnDemandInstallCache::Reset()
 
 	if (Ifm.DeleteDirectory(*CacheDirectory, false, bTree) == false)
 	{
-		return FIoStatusBuilder(EIoErrorCode::DeleteError, FPlatformMisc::GetLastError())
-			<< TEXT("Failed to delete directory '")
-			<< CacheDirectory
-			<< TEXT("'");
+		return MakeCasError<void>(
+			ECasErrorCode::InitializeFailed, EIoErrorCode::DeleteError, FString::Printf(TEXT("Failed to delete cache directory '%s'"), *CacheDirectory));
 	}
 
 	if (Ifm.MakeDirectory(*CacheDirectory, bTree) == false)
 	{
-		return FIoStatusBuilder(EIoErrorCode::WriteError, FPlatformMisc::GetLastError())
-			<< TEXT("Failed to create directory '")
-			<< CacheDirectory
-			<< TEXT("'");
+		return MakeCasError<void>(
+			ECasErrorCode::InitializeFailed, EIoErrorCode::WriteError, FString::Printf(TEXT("Failed to create cache directory '%s'"), *CacheDirectory));
 	}
 
-	if (FIoStatus Status = Cas.Initialize(CacheDirectory); Status.IsOk() == false)
+	if (FResult Result = Cas.Initialize(CacheDirectory); Result.HasError())
 	{
-		return Status;
+		return Result;
 	}
 
 	const FString JournalFile = GetJournalFilename();
-	if (FIoStatus Status = FCasJournal::Create(JournalFile); Status.IsOk() == false)
+	if (FResult Result = FCasJournal::Create(JournalFile); Result.HasError())
 	{
-		return Status;
+		return Result;
 	}
 
 	UE_LOG(LogIoStoreOnDemand, Log, TEXT("Created CAS journal '%s'"), *JournalFile);
-	return FIoStatus::Ok;
+	return MakeValue();
 }
 
-FIoStatus FOnDemandInstallCache::InitialVerify()
+FResult FOnDemandInstallCache::InitialVerify()
 {
 	// Verify the blocks on disk with the current state of the CAS
 	{
 		TArray<FCasAddr> RemovedChunks;
-		if (FIoStatus Verify = Cas.Verify(RemovedChunks); Verify.IsOk() == false)
+		if (FResult Verify = Cas.Verify(RemovedChunks); Verify.HasError())
 		{
 			FOnDemandInstallCacheStats::OnCasVerificationError(RemovedChunks.Num());
 
@@ -2788,9 +2794,9 @@ FIoStatus FOnDemandInstallCache::InitialVerify()
 				Transaction.ChunkLocation(FCasLocation::Invalid, Addr);
 			}
 
-			if (FIoStatus Status = FCasJournal::Commit(MoveTemp(Transaction)); Status.IsOk() == false)
+			if (FResult Result = FCasJournal::Commit(MoveTemp(Transaction)); Result.HasError())
 			{
-				return Status;
+				return Result;
 			}
 		}
 	}
@@ -2814,23 +2820,24 @@ FIoStatus FOnDemandInstallCache::InitialVerify()
 					ToMiB(TotalBytesToPurge), *LexToString(Result.GetError()));
 
 				UE_LOG(LogIoStoreOnDemand, Error, TEXT("%s"), *ErrorMessage);
-				return FIoStatus(EIoErrorCode::WriteError, ErrorMessage);
+				return Result;
 			}
 
 			if (TotalPurgedBytes < TotalBytesToPurge)
 			{
+				// This should never happen since we don't have any referenced cache blocks at startup
 				const FString ErrorMessage = FString::Printf(TEXT("Failed to purge %.2lf MiB from install cache. Actually purged %.2lf MiB from install cache"),
 					ToMiB(TotalBytesToPurge), ToMiB(TotalPurgedBytes));
 
 				UE_LOG(LogIoStoreOnDemand, Error, TEXT("%s"), *ErrorMessage);
-				return FIoStatus(EIoErrorCode::WriteError, ErrorMessage);
+				return MakeError(UE::UnifiedError::IoStoreOnDemand::InstallCachePurgeError::MakeError());
 			}
 
 			UE_LOG(LogIoStoreOnDemand, Log, TEXT("Successfully purged %.2lf MiB from install cache"), ToMiB(TotalPurgedBytes));
 		}
 	}
 
-	return FIoStatus::Ok;
+	return MakeValue();
 }
 
 uint64 FOnDemandInstallCache::AddReferencesToBlocks(
@@ -2930,16 +2937,9 @@ FResult FOnDemandInstallCache::Purge(FCasBlockInfoMap& BlockInfo, const uint64 T
 		FCasJournal::FTransaction	Transaction = FCasJournal::Begin(GetJournalFilename());
 		TArray<FCasAddr>			RemovedChunks;
 
-		if (FIoStatus Status = Cas.DeleteBlock(BlockId, RemovedChunks); !Status.IsOk())
+		if (FResult Result = Cas.DeleteBlock(BlockId, RemovedChunks); Result.HasError())
 		{
-			UE::UnifiedError::IoStoreOnDemand::FInstallCacheErrorContext Ctx
-			{
-				.Error			= Status,
-				.ErrorCategory	= UE::UnifiedError::IoStoreOnDemand::EErrorCategory::Cas
-			};
-			CaptureCacheState(Ctx);
-			return MakeError(UE::UnifiedError::IoStoreOnDemand::InstallCachePurgeError::MakeError(
-				MoveTemp(Ctx), UE::UnifiedError::EDetailFilter::All));
+			return Result;
 		}
 
 		FOnDemandInstallCacheStats::OnBlockDeleted(Info.LastAccess, false /*FromDefrag*/);
@@ -2958,16 +2958,9 @@ FResult FOnDemandInstallCache::Purge(FCasBlockInfoMap& BlockInfo, const uint64 T
 		}
 		Transaction.BlockDeleted(BlockId);
 
-		if (FIoStatus Status = FCasJournal::Commit(MoveTemp(Transaction)); !Status.IsOk())
+		if (FResult Result = FCasJournal::Commit(MoveTemp(Transaction)); Result.HasError())
 		{
-			UE::UnifiedError::IoStoreOnDemand::FInstallCacheErrorContext Ctx
-			{
-				.Error			= Status,
-				.ErrorCategory	= UE::UnifiedError::IoStoreOnDemand::EErrorCategory::Journal
-			};
-			CaptureCacheState(Ctx);
-			return MakeError(UE::UnifiedError::IoStoreOnDemand::InstallCachePurgeError::MakeError(
-				MoveTemp(Ctx), UE::UnifiedError::EDetailFilter::All));
+			return Result;
 		}
 		
 		if (OutTotalPurgedBytes >= TotalBytesToPurge)
@@ -2986,11 +2979,25 @@ FResult FOnDemandInstallCache::Defrag(
 	const uint64* TotalBytesToFree /*= nullptr*/,
 	uint64* OutTotalFreedBytes /*= nullptr*/)
 {
+	using namespace UE::UnifiedError::IoStoreOnDemand;
+
 	if (TotalBytesToFree && *TotalBytesToFree == 0)
 	{
 		return MakeValue();
 	}
 
+	FResult Result			= MakeValue();
+	uint64 FragmentedBytes	= 0;
+
+	ON_SCOPE_EXIT
+	{
+		if (OutTotalFreedBytes != nullptr)
+		{
+			*OutTotalFreedBytes = Result.HasError() ? 0 : FragmentedBytes;
+		}
+		FOnDemandInstallCacheStats::OnDefrag(Result, FragmentedBytes);
+	};
+
 	const uint64 TotalCachedBytes = Algo::TransformAccumulate(BlockInfo,
 		[](const TPair<FCasBlockId, FCasBlockInfo>& Kv) { return Kv.Value.FileSize; },
 		uint64(0));
@@ -3002,24 +3009,16 @@ FResult FOnDemandInstallCache::Defrag(
 			ToMiB(MaxCacheSize));
 		UE_LOG(LogIoStoreOnDemand, Error, TEXT("%s"), *ErrorMsg);
 
-		UE::UnifiedError::IoStoreOnDemand::FInstallCacheErrorContext Ctx
-		{
-			.Error			= FIoStatus(EIoErrorCode::OutOfDiskSpace, ErrorMsg),
-			.CacheSize		= TotalCachedBytes,
-			.ErrorCategory	= UE::UnifiedError::IoStoreOnDemand::EErrorCategory::Logical
-		};
-		CaptureCacheState(Ctx);
-		FResult Result = MakeError(UE::UnifiedError::IoStoreOnDemand::InstallCacheDefragError::MakeError(
-			MoveTemp(Ctx), UE::UnifiedError::EDetailFilter::All));
-		FOnDemandInstallCacheStats::OnDefrag(Result, 0);
+		FInstallCacheErrorContext Ctx = MakeInstallCacheErrorContext(TotalCachedBytes);
+		Ctx.IoError = FIoStatus(EIoErrorCode::OutOfDiskSpace, ErrorMsg),
+		Result = MakeError(InstallCacheDefragError::MakeError(MoveTemp(Ctx), UE::UnifiedError::EDetailFilter::All));
 
 		// Append a critical error entry to clear the cache at next startup
 		FCasJournal::FTransaction Transaction = FCasJournal::Begin(GetJournalFilename());
 		Transaction.CriticalError(FCasJournal::EErrorCode::DefragOutOfDiskSpace);
-		if (FIoStatus JournalStatus = FCasJournal::Commit(MoveTemp(Transaction)); JournalStatus.IsOk() == false)
+		if (FResult CommitResult = FCasJournal::Commit(MoveTemp(Transaction)); CommitResult.HasError())
 		{
-			UE_LOG(LogIoStoreOnDemand, Error, TEXT("Failed to append critical error to journal, reason '%s'"),
-				*JournalStatus.ToString());
+			UE_LOGFMT(LogIoStoreOnDemand, Error, "Failed to commit critical errors to CAS journal ({Error})", CommitResult.GetError());
 		}
 
 		return Result;
@@ -3049,7 +3048,6 @@ FResult FOnDemandInstallCache::Defrag(
 		return LHS.RefSize < RHS.RefSize;
 	});
 
-	uint64 FragmentedBytes = 0;
 	uint64 TotalBlockSize = 0;
 
 	if (TotalBytesToFree)
@@ -3097,16 +3095,11 @@ FResult FOnDemandInstallCache::Defrag(
 
 		if (!bPossibleToFreeBytes)
 		{
-			UE::UnifiedError::IoStoreOnDemand::FInstallCacheErrorContext Ctx
-			{
-				.CacheSize		= TotalCachedBytes,
-				.ErrorCategory	= UE::UnifiedError::IoStoreOnDemand::EErrorCategory::Logical
-			};
-			CaptureCacheState(Ctx);
-			FResult Result = MakeError(UE::UnifiedError::IoStoreOnDemand::InstallCacheDefragError::MakeError(
-				MoveTemp(Ctx), UE::UnifiedError::EDetailFilter::All));
-			FOnDemandInstallCacheStats::OnDefrag(Result, 0);
-			return Result;
+			FString ErrorMessage = TEXT("Failed to defrag the install cache due to all data being referenced by the game");
+			UE_LOG(LogIoStoreOnDemand, Error, TEXT("%s"), *ErrorMessage);
+			FInstallCacheErrorContext Ctx = MakeInstallCacheErrorContext(TotalCachedBytes);
+			Ctx.IoError = FIoStatus(EIoErrorCode::OutOfDiskSpace, ErrorMessage);
+			return Result = MakeError(InstallCacheDefragError::MakeError(MoveTemp(Ctx), UE::UnifiedError::EDetailFilter::All));
 		}
 	}
 	else
@@ -3143,7 +3136,7 @@ FResult FOnDemandInstallCache::Defrag(
 		{
 			// Already defragged
 			UE_LOG(LogIoStoreOnDemand, Display, TEXT("Cache not fragmented."));
-			return MakeValue();
+			return Result = MakeValue();
 		}
 	}
 
@@ -3199,22 +3192,14 @@ FResult FOnDemandInstallCache::Defrag(
 	{
 		if (DefragBlock.ReferencedChunks.IsEmpty() == false)
 		{
-			FSharedFileOpenResult FileOpenResult = Cas.OpenRead(DefragBlock.BlockId);
-			if (!FileOpenResult.IsValid())
+			TResult<FSharedFileHandle> FileOpenResult = Cas.OpenRead(DefragBlock.BlockId);
+			if (FileOpenResult.HasValue())
 			{
-				const FString ErrorMsg = FileOpenResult.GetError().GetMessage();
-				UE_LOG(LogIoStoreOnDemand, Error, TEXT("Defrag failed to open CAS block for reading, reason '%s'"), *ErrorMsg);
-				UE::UnifiedError::IoStoreOnDemand::FInstallCacheErrorContext Ctx
-				{
-					.Error			= FIoStatus(EIoErrorCode::FileNotOpen, FPlatformMisc::GetLastError(), ErrorMsg),
-					.CacheSize		= TotalCachedBytes,
-					.ErrorCategory	= UE::UnifiedError::IoStoreOnDemand::EErrorCategory::Cas
-				};
-				CaptureCacheState(Ctx);
-				FResult Result = MakeError(UE::UnifiedError::IoStoreOnDemand::InstallCacheDefragError::MakeError(
-					MoveTemp(Ctx), UE::UnifiedError::EDetailFilter::All));
-				FOnDemandInstallCacheStats::OnDefrag(Result, 0);
-				return Result;
+				UE_LOGFMT(LogIoStoreOnDemand, Error, "Defrag failed to open CAS block for reading {Error}", FileOpenResult.GetError());
+
+				FInstallCacheErrorContext Ctx = MakeInstallCacheErrorContext(TotalCachedBytes);
+				FileOpenResult.GetError().PushErrorContext(MoveTemp(Ctx));
+				return Result = MakeError(FileOpenResult.StealError());
 			}
 
 			FSharedFileHandle FileHandle = FileOpenResult.StealValue();
@@ -3223,43 +3208,35 @@ FResult FOnDemandInstallCache::Defrag(
 
 			for (const FDefragBlockReferencedChunk& ReffedChunk : DefragBlock.ReferencedChunks)
 			{
-				FIoBuffer		Buffer(ReffedChunk.DiskSize);
-				FIoStatus		ChunkStatus = FIoStatus::Invalid;
+				FIoBuffer	Buffer(ReffedChunk.DiskSize);
+				FResult		ReadResult = MakeValue();
 
 				for (int32 Attempt = 0, MaxAttempts = 3; Attempt < MaxAttempts; ++Attempt)
 				{
 					if (FileHandle->Seek(ReffedChunk.BlockOffset) == false)
 					{
-						ChunkStatus = FIoStatusBuilder(EIoErrorCode::FileSeekFailed, FPlatformMisc::GetLastError())
-							<< FString::Printf(TEXT("Failed to seek to offset %d in cache block %d"),
-								ReffedChunk.BlockOffset, DefragBlock.BlockId.Id);
+						ReadResult = MakeCasError<void>(
+							ECasErrorCode::ReadBlockFailed, EIoErrorCode::FileSeekFailed, FString::Printf(TEXT("Failed to seek to offset %u in cache block %u"),
+								ReffedChunk.BlockOffset, DefragBlock.BlockId.Id));
 						continue;
 					}
 					if (FileHandle->Read(Buffer.GetData(), Buffer.GetSize()) == false)
 					{
-						ChunkStatus = FIoStatusBuilder(EIoErrorCode::ReadError, FPlatformMisc::GetLastError())
-							<< FString::Printf(TEXT("Failed to read %ull bytes at offset %d in cache block %d"),
-								Buffer.GetSize(), ReffedChunk.BlockOffset, DefragBlock.BlockId.Id);
+						ReadResult = MakeCasError<void>(
+							ECasErrorCode::ReadBlockFailed, EIoErrorCode::ReadError, FString::Printf(TEXT("Failed to read %ull bytes at offset %u in cache block %u"),
+								Buffer.GetSize(), ReffedChunk.BlockOffset, DefragBlock.BlockId.Id));
 						continue;
 					}
 
-					ChunkStatus = FIoStatus::Ok;
+					ReadResult = MakeValue();
 					break;
 				}
 
-				if (ChunkStatus.IsOk() == false)
+				if (ReadResult.HasError())
 				{
-					UE::UnifiedError::IoStoreOnDemand::FInstallCacheErrorContext Ctx
-					{
-						.Error			= ChunkStatus,
-						.CacheSize		= TotalCachedBytes,
-						.ErrorCategory	= UE::UnifiedError::IoStoreOnDemand::EErrorCategory::Cas
-					};
-					CaptureCacheState(Ctx);
-					FResult Result = MakeError(UE::UnifiedError::IoStoreOnDemand::InstallCacheDefragError::MakeError(
-						MoveTemp(Ctx), UE::UnifiedError::EDetailFilter::All));
-					FOnDemandInstallCacheStats::OnDefrag(Result, 0);
-					return Result;
+					FInstallCacheErrorContext Ctx = MakeInstallCacheErrorContext(TotalCachedBytes);
+					ReadResult.GetError().PushErrorContext(MoveTemp(Ctx));
+					return Result = MakeError(ReadResult.StealError());
 				}
 
 				const FIoHash ChunkHash = FIoHash::HashBuffer(Buffer.GetView());
@@ -3271,35 +3248,30 @@ FResult FOnDemandInstallCache::Defrag(
 					// Append a critical error entry to clear the cache at next startup
 					FCasJournal::FTransaction Transaction = FCasJournal::Begin(GetJournalFilename());
 					Transaction.CriticalError(FCasJournal::EErrorCode::DefragHashMismatch);
-					if (FIoStatus Status = FCasJournal::Commit(MoveTemp(Transaction)); Status.IsOk() == false)
+					if (FResult CommitResult = FCasJournal::Commit(MoveTemp(Transaction)); CommitResult.HasError())
 					{
-						UE_LOG(LogIoStoreOnDemand, Error, TEXT("Failed to append critical error to journal, reason '%s'"),
-							*Status.ToString());
+						UE_LOGFMT(LogIoStoreOnDemand, Error, "Failed to commit critical errors to CAS journal ({Error})", CommitResult.GetError());
 					}
 
-					if (FResult Result = FlushPendingChunks(DefragPendingChunks, DefragBlock.LastAccess); Result.HasError())
+					if (Result = FlushPendingChunks(DefragPendingChunks, DefragBlock.LastAccess); Result.HasError())
 					{
-						FOnDemandInstallCacheStats::OnDefrag(Result, 0);
 						return Result;
 					}
 					check(DefragPendingChunks.IsEmpty());
 	
-					FResult Result = MakeError(UE::UnifiedError::IoStoreOnDemand::InstallCacheDefragError::MakeError(
-						UE::UnifiedError::IoStoreOnDemand::FChunkHashMismatchErrorContext
+					return Result = MakeError(InstallCacheDefragError::MakeError(
+						FChunkHashMismatchErrorContext
 						{
 							.ChunkId		= ReffedChunk.ChunkId,
 							.ExpectedHash	= ReffedChunk.Hash,
 							.ActualHash		= ChunkHash
 						}));
-					FOnDemandInstallCacheStats::OnDefrag(Result, 0);
-					return Result;
 				}
 
 				if (DefragPendingChunks.TotalSize > FPendingChunks::MaxPendingBytes)
 				{
-					if (FResult Result = FlushPendingChunks(DefragPendingChunks, DefragBlock.LastAccess); Result.HasError())
+					if (Result = FlushPendingChunks(DefragPendingChunks, DefragBlock.LastAccess); Result.HasError())
 					{
-						FOnDemandInstallCacheStats::OnDefrag(Result, 0);
 						return Result;
 					}
 					check(DefragPendingChunks.IsEmpty());
@@ -3310,9 +3282,8 @@ FResult FOnDemandInstallCache::Defrag(
 
 			FileHandle.Reset();
 
-			if (FResult Result = FlushPendingChunks(DefragPendingChunks, DefragBlock.LastAccess); Result.HasError()) 
+			if (Result = FlushPendingChunks(DefragPendingChunks, DefragBlock.LastAccess); Result.HasError()) 
 			{
-				FOnDemandInstallCacheStats::OnDefrag(Result, 0);
 				return Result;
 			}
 			check(DefragPendingChunks.IsEmpty());
@@ -3323,19 +3294,8 @@ FResult FOnDemandInstallCache::Defrag(
 		// Flushing should overwrite the lookup info for the cas addr to point at the new block.
 		// Can now remove the old block
 		TArray<FCasAddr> DeletedChunkAddresses;
-		if (FIoStatus Status = Cas.DeleteBlock(DefragBlock.BlockId, DeletedChunkAddresses); !Status.IsOk())
+		if (Result = Cas.DeleteBlock(DefragBlock.BlockId, DeletedChunkAddresses); Result.HasError())
 		{
-			UE::UnifiedError::IoStoreOnDemand::FInstallCacheErrorContext Ctx
-			{
-				.Error			= Status,
-				.CacheSize		= TotalCachedBytes,
-				.ErrorCategory	= UE::UnifiedError::IoStoreOnDemand::EErrorCategory::Cas
-			};
-			CaptureCacheState(Ctx);
-			FResult Result = MakeError(UE::UnifiedError::IoStoreOnDemand::InstallCacheDefragError::MakeError(
-				MoveTemp(Ctx), UE::UnifiedError::EDetailFilter::All));
-			FOnDemandInstallCacheStats::OnDefrag(Result, 0);
-
 			return Result;
 		}
 
@@ -3347,32 +3307,14 @@ FResult FOnDemandInstallCache::Defrag(
 		}
 		Transaction.BlockDeleted(DefragBlock.BlockId);
 
-		if (FIoStatus Status = FCasJournal::Commit(MoveTemp(Transaction)); !Status.IsOk())
+		if (Result = FCasJournal::Commit(MoveTemp(Transaction)); Result.HasError())
 		{
-			UE::UnifiedError::IoStoreOnDemand::FInstallCacheErrorContext Ctx
-			{
-				.Error			= Status,
-				.CacheSize		= TotalCachedBytes,
-				.ErrorCategory	= UE::UnifiedError::IoStoreOnDemand::EErrorCategory::Journal
-			};
-			CaptureCacheState(Ctx);
-			FResult Result = MakeError(UE::UnifiedError::IoStoreOnDemand::InstallCacheDefragError::MakeError(
-				MoveTemp(Ctx), UE::UnifiedError::EDetailFilter::All));
-			FOnDemandInstallCacheStats::OnDefrag(Result, 0);
-
 			return Result;
 		}
 	}
 
 	UE_LOG(LogIoStoreOnDemand, Display, TEXT("Defrag removed %" UINT64_FMT " fragmented bytes of %" UINT64_FMT " total bytes in %i blocks."), FragmentedBytes, TotalBlockSize, BlocksToDefrag.Num());
 
-	if (OutTotalFreedBytes)
-	{
-		*OutTotalFreedBytes = FragmentedBytes;
-	}
-
-	FResult Result = MakeValue();
-	FOnDemandInstallCacheStats::OnDefrag(Result, FragmentedBytes);
 	return Result;
 }
 
@@ -3401,19 +3343,12 @@ FResult FOnDemandInstallCache::FlushLastAccess()
 	FCasJournal::FTransaction Transaction = FCasJournal::Begin(JournalFile);
 	JournalLastAccess(Transaction, DirtyLastAccess);
 
-	if (FIoStatus Status = FCasJournal::Commit(MoveTemp(Transaction)); !Status.IsOk())
+	if (FResult Result = FCasJournal::Commit(MoveTemp(Transaction)); Result.HasError())
 	{
 		UE_LOG(LogIoStoreOnDemand, Error, TEXT("Failed to update CAS journal '%s' with block timestamp(s), reason '%s'"),
-			*JournalFile, *Status.ToString());
-
-		UE::UnifiedError::IoStoreOnDemand::FInstallCacheErrorContext Ctx
-		{
-			.Error			= Status,
-			.ErrorCategory	= UE::UnifiedError::IoStoreOnDemand::EErrorCategory::Journal
-		};
-		CaptureCacheState(Ctx);
-		return MakeError(UE::UnifiedError::IoStoreOnDemand::InstallCacheFlushLastAccessError::MakeError(
-			MoveTemp(Ctx), UE::UnifiedError::EDetailFilter::All));
+			*JournalFile, *LexToString(Result.GetError()));
+		Result.GetError().PushErrorContext(MakeInstallCacheErrorContext());
+		return Result;
 	}
 
 	return MakeValue();
@@ -3530,6 +3465,16 @@ FResult FOnDemandInstallCache::FlushPendingChunks(FPendingChunks& Chunks, int64
 
 FResult FOnDemandInstallCache::FlushPendingChunksImpl(const FPendingChunks& Chunks, int64 UtcAccessTicks)
 {
+	using namespace UE::UnifiedError::IoStoreOnDemand;
+
+	FResult Result				= MakeValue();
+	uint64 TotalFlushedBytes	= 0;
+
+	ON_SCOPE_EXIT 
+	{ 
+		FOnDemandInstallCacheStats::OnFlush(Result, TotalFlushedBytes);
+	};
+
 	FLargeMemoryWriter Ar(FMath::Min(Chunks.TotalSize, static_cast<uint64>(Cas.GetMaxBlockSize())));
 
 	// This should be the only thread writing to CurrentBlock
@@ -3552,21 +3497,15 @@ FResult FOnDemandInstallCache::FlushPendingChunksImpl(const FPendingChunks& Chun
 			Transaction.BlockCreated(CurrentBlockId);
 		}
 
-		TUniquePtr<IFileHandle>	CasFileHandle = Cas.OpenWrite(CurrentBlockId, bAppendToBlock);
-		if (CasFileHandle.IsValid() == false)
+		TResult<FUniqueFileHandle> OpenWriteResult = Cas.OpenWrite(CurrentBlockId, bAppendToBlock);
+		if (OpenWriteResult.HasError())
 		{
-			UE::UnifiedError::IoStoreOnDemand::FInstallCacheErrorContext Ctx
-			{
-				.Error			= FIoStatus(EIoErrorCode::FileOpenFailed, FPlatformMisc::GetLastError()),
-				.ErrorCategory	= UE::UnifiedError::IoStoreOnDemand::EErrorCategory::Cas
-			};
-			CaptureCacheState(Ctx);
-			FResult Result = MakeError(UE::UnifiedError::IoStoreOnDemand::InstallCacheFlushError::MakeError(
-				MoveTemp(Ctx), UE::UnifiedError::EDetailFilter::All));
-			FOnDemandInstallCacheStats::OnFlush(Result, 0);
-			return Result;
+			FInstallCacheErrorContext Ctx = MakeInstallCacheErrorContext();
+			OpenWriteResult.GetError().PushErrorContext(MoveTemp(Ctx));
+			return Result = MakeError(OpenWriteResult.StealError());
 		}
-
+		
+		FUniqueFileHandle CasFileHandle = OpenWriteResult.StealValue();
 		const int64 CasBlockOffset = CasFileHandle->Tell();
 		
 		TArray<int64>		Offsets;
@@ -3597,18 +3536,24 @@ FResult FOnDemandInstallCache::FlushPendingChunksImpl(const FPendingChunks& Chun
 
 			if (CasFileHandle->Write(Ar.GetData(), Ar.Tell()) == false)
 			{
-				UE::UnifiedError::IoStoreOnDemand::FInstallCacheErrorContext Ctx
-				{
-					.Error			= FIoStatus(EIoErrorCode::WriteError, FPlatformMisc::GetLastError()),
-					.ErrorCategory	= UE::UnifiedError::IoStoreOnDemand::EErrorCategory::Cas
-				};
-				CaptureCacheState(Ctx);
-				FResult Result = MakeError(UE::UnifiedError::IoStoreOnDemand::InstallCacheFlushError::MakeError(
-					MoveTemp(Ctx), UE::UnifiedError::EDetailFilter::All));
-				FOnDemandInstallCacheStats::OnFlush(Result, 0);
+				Result = MakeCasError<void>(
+					ECasErrorCode::WriteBlockFailed, EIoErrorCode::WriteError, FString::Printf(TEXT("Failed to write %lld bytes to CAS block %u"), Ar.Tell(), CurrentBlockId.Id));
+				FInstallCacheErrorContext Ctx = MakeInstallCacheErrorContext();
+				Result.GetError().PushErrorContext(MoveTemp(Ctx));
 				return Result;
 			}
 
+			if (CasFileHandle->Flush() == false)
+			{
+				Result = MakeCasError<void>(
+					ECasErrorCode::WriteBlockFailed, EIoErrorCode::FileFlushFailed, FString::Printf(TEXT("Failed to flush %lld bytes to CAS block %u"), Ar.Tell(), CurrentBlockId.Id));
+				FInstallCacheErrorContext Ctx = MakeInstallCacheErrorContext();
+				Result.GetError().PushErrorContext(MoveTemp(Ctx));
+				return Result;
+			}
+
+			TotalFlushedBytes += Ar.Tell();
+
 			constexpr const bool bDirty = false;
 			if (UtcAccessTicks)
 			{
@@ -3624,22 +3569,6 @@ FResult FOnDemandInstallCache::FlushPendingChunksImpl(const FPendingChunks& Chun
 				Transaction.BlockAccess(CurrentBlockId, Now);
 			}
 
-			if (CasFileHandle->Flush() == false)
-			{
-				UE::UnifiedError::IoStoreOnDemand::FInstallCacheErrorContext Ctx
-				{
-					.Error			= FIoStatus(EIoErrorCode::FileFlushFailed, FPlatformMisc::GetLastError()),
-					.ErrorCategory	= UE::UnifiedError::IoStoreOnDemand::EErrorCategory::Cas
-				};
-				CaptureCacheState(Ctx);
-				FResult Result = MakeError(UE::UnifiedError::IoStoreOnDemand::InstallCacheFlushError::MakeError(
-					MoveTemp(Ctx), UE::UnifiedError::EDetailFilter::All));
-				FOnDemandInstallCacheStats::OnFlush(Result, Ar.Tell());
-				return Result;
-			}
-
-			FOnDemandInstallCacheStats::OnFlush(MakeValue(), Ar.Tell());
-
 			check(ChunkHashes.Num() == Offsets.Num());
 			check(CurrentBlockId.IsValid());
 			{
@@ -3659,17 +3588,10 @@ FResult FOnDemandInstallCache::FlushPendingChunksImpl(const FPendingChunks& Chun
 			Ar.Seek(0);
 		}
 
-		if (FIoStatus Status = FCasJournal::Commit(MoveTemp(Transaction)); Status.IsOk() == false)
+		if (Result = FCasJournal::Commit(MoveTemp(Transaction)); Result.HasError())
 		{
-			UE::UnifiedError::IoStoreOnDemand::FInstallCacheErrorContext Ctx
-			{
-				.Error			= Status,
-				.ErrorCategory	= UE::UnifiedError::IoStoreOnDemand::EErrorCategory::Journal
-			};
-			CaptureCacheState(Ctx);
-			FResult Result = MakeError(UE::UnifiedError::IoStoreOnDemand::InstallCacheFlushError::MakeError(
-				MoveTemp(Ctx), UE::UnifiedError::EDetailFilter::All));
-			FOnDemandInstallCacheStats::OnFlush(Result, Ar.Tell());
+			FInstallCacheErrorContext Ctx = MakeInstallCacheErrorContext();
+			Result.GetError().PushErrorContext(MoveTemp(Ctx));
 			return Result;
 		}
 
@@ -3679,7 +3601,7 @@ FResult FOnDemandInstallCache::FlushPendingChunksImpl(const FPendingChunks& Chun
 		}
 	}
 
-	return MakeValue();
+	return Result;
 }
 
 void FOnDemandInstallCache::CompleteRequest(FIoRequestImpl* Request, EIoErrorCode Status)
@@ -3733,11 +3655,17 @@ void FOnDemandInstallCache::CompleteRequest(FIoRequestImpl* Request, EIoErrorCod
 	BackendContext->WakeUpDispatcherThreadDelegate.Execute();
 }
 
-void FOnDemandInstallCache::CaptureCacheState(UE::UnifiedError::IoStoreOnDemand::FInstallCacheErrorContext& OutCtx)
+UE::UnifiedError::IoStoreOnDemand::FInstallCacheErrorContext FOnDemandInstallCache::MakeInstallCacheErrorContext(uint64 TotalCachedBytes, uint32 LineNo)
 {
+	UE::UnifiedError::IoStoreOnDemand::FInstallCacheErrorContext OutCtx;
+
 	const FString CachRootDirectory = Cas.GetRootDirectory();
 	OutCtx.bDiskQuerySucceeded		= FPlatformMisc::GetDiskTotalAndFreeSpace(CachRootDirectory, OutCtx.DiskTotalBytes, OutCtx.DiskFreeBytes);
 	OutCtx.MaxCacheSize				= MaxCacheSize;
+	OutCtx.CacheSize				= TotalCachedBytes;
+	OutCtx.LineNo					= LineNo;
+
+	return OutCtx;
 }
 
 ///////////////////////////////////////////////////////////////////////////////
@@ -3802,21 +3730,21 @@ TEST_CASE("IoStore::OnDemand::InstallCache::Journal", "[IoStoreOnDemand][Install
 	{
 		FTmpDirectoryScope _(TestBaseDir);
 		const FString JournalFile = TestBaseDir / TEXT("test.jrn");
-		FIoStatus Status = FCasJournal::Create(JournalFile);
-		CHECK(Status.IsOk());
+		FResult Result = FCasJournal::Create(JournalFile);
+		CHECK(Result.HasValue());
 	}
 
 	SECTION("SimpleTransaction")
 	{
 		FTmpDirectoryScope _(TestBaseDir);
 		const FString JournalFile = TestBaseDir / TEXT("test.jrn");
-		FIoStatus Status = FCasJournal::Create(JournalFile);
-		CHECK(Status.IsOk());
+		FResult Result = FCasJournal::Create(JournalFile);
+		CHECK(Result.HasValue());
 
 		FCasJournal::FTransaction Transaction = FCasJournal::Begin(JournalFile);
 		Transaction.BlockCreated(FCasBlockId(1));
-		Status = FCasJournal::Commit(MoveTemp(Transaction));
-		CHECK(Status.IsOk());
+		Result = FCasJournal::Commit(MoveTemp(Transaction));
+		CHECK(Result.HasValue());
 	}
 
 	SECTION("ReplayChunkLocations")
@@ -3835,8 +3763,8 @@ TEST_CASE("IoStore::OnDemand::InstallCache::Journal", "[IoStoreOnDemand][Install
 		// Act
 		FTmpDirectoryScope _(TestBaseDir);
 		const FString JournalFile = TestBaseDir / TEXT("test.jrn");
-		FIoStatus Status = FCasJournal::Create(JournalFile);
-		CHECK(Status.IsOk());
+		FResult Result = FCasJournal::Create(JournalFile);
+		CHECK(Result.HasValue());
 
 		FCasJournal::FTransaction Transaction = FCasJournal::Begin(JournalFile);
 		for (int32 Idx = 0; const FCasAddr& Addr : ExpectedAddresses)
@@ -3850,12 +3778,12 @@ TEST_CASE("IoStore::OnDemand::InstallCache::Journal", "[IoStoreOnDemand][Install
 				Addr);
 		}
 
-		Status = FCasJournal::Commit(MoveTemp(Transaction));
-		CHECK(Status.IsOk());
+		Result = FCasJournal::Commit(MoveTemp(Transaction));
+		CHECK(Result.HasValue());
 
 		// Assert
 		TArray<FCasJournal::FEntry::FChunkLocation> Locs;
-		Status = FCasJournal::Replay(
+		Result = FCasJournal::Replay(
 			JournalFile,
 			[&Locs](const FCasJournal::FEntry& JournalEntry)
 			{
@@ -3871,7 +3799,7 @@ TEST_CASE("IoStore::OnDemand::InstallCache::Journal", "[IoStoreOnDemand][Install
 					break;
 				};
 			});
-		CHECK(Status.IsOk());
+		CHECK(Result.HasValue());
 		CHECK(Locs.Num() == ExpectedAddresses.Num());
 		for (int32 Idx = 0; const FCasJournal::FEntry::FChunkLocation& Loc : Locs)
 		{
@@ -3894,21 +3822,21 @@ TEST_CASE("IoStore::OnDemand::InstallCache::Journal", "[IoStoreOnDemand][Install
 		FTmpDirectoryScope _(TestBaseDir);
 		const FString JournalFile = TestBaseDir / TEXT("test.jrn");
 
-		FIoStatus Status = FCasJournal::Create(JournalFile);
-		CHECK(Status.IsOk());
+		FResult Result = FCasJournal::Create(JournalFile);
+		CHECK(Result.HasValue());
 
 		FCasJournal::FTransaction Tx = FCasJournal::Begin(JournalFile);
 		Tx.BlockCreated(ExpectedBlockId);
 		Tx.BlockDeleted(ExpectedBlockId);
 
-		Status = FCasJournal::Commit(MoveTemp(Tx));
-		CHECK(Status.IsOk());
+		Result = FCasJournal::Commit(MoveTemp(Tx));
+		CHECK(Result.HasValue());
 
 		// Assert
 		FCasBlockId CreatedBlockId;
 		FCasBlockId DeletedBlockId;
 
-		Status = FCasJournal::Replay(
+		Result = FCasJournal::Replay(
 			JournalFile,
 			[&CreatedBlockId, &DeletedBlockId](const FCasJournal::FEntry& JournalEntry)
 			{
@@ -3930,7 +3858,7 @@ TEST_CASE("IoStore::OnDemand::InstallCache::Journal", "[IoStoreOnDemand][Install
 				};
 			});
 
-		CHECK(Status.IsOk());
+		CHECK(Result.HasValue());
 		CHECK(CreatedBlockId == ExpectedBlockId);
 		CHECK(DeletedBlockId == ExpectedBlockId);
 	}
@@ -3945,20 +3873,20 @@ TEST_CASE("IoStore::OnDemand::InstallCache::Journal", "[IoStoreOnDemand][Install
 		FTmpDirectoryScope _(TestBaseDir);
 		const FString JournalFile = TestBaseDir / TEXT("test.jrn");
 
-		FIoStatus Status = FCasJournal::Create(JournalFile);
-		CHECK(Status.IsOk());
+		FResult Result = FCasJournal::Create(JournalFile);
+		CHECK(Result.HasValue());
 
 		FCasJournal::FTransaction Tx = FCasJournal::Begin(JournalFile);
 		Tx.BlockAccess(ExpectedBlockId, ExpectedTicks);
 
-		Status = FCasJournal::Commit(MoveTemp(Tx));
-		CHECK(Status.IsOk());
+		Result = FCasJournal::Commit(MoveTemp(Tx));
+		CHECK(Result.HasValue());
 
 		// Assert
 		FCasBlockId BlockId;
 		uint64 Ticks = 0;
 
-		Status = FCasJournal::Replay(
+		Result = FCasJournal::Replay(
 			JournalFile,
 			[&BlockId, &Ticks](const FCasJournal::FEntry& JournalEntry)
 			{
@@ -3977,7 +3905,7 @@ TEST_CASE("IoStore::OnDemand::InstallCache::Journal", "[IoStoreOnDemand][Install
 				};
 			});
 
-		CHECK(Status.IsOk());
+		CHECK(Result.HasValue());
 		CHECK(BlockId == ExpectedBlockId);
 		CHECK(Ticks == ExpectedTicks);
 	}
@@ -4016,9 +3944,9 @@ TEST_CASE("IoStore::OnDemand::InstallCache::Snapshot", "[IoStoreOnDemand][Instal
 		// Act
 		FTmpDirectoryScope _(TestBaseDir);
 		const FString SnapshotFile = TestBaseDir / TEXT("test.snp");
-		TIoStatusOr<int64> Status = FCasSnapshot::Save(ExpectedSnapshot, SnapshotFile);
-		CHECK(Status.IsOk());
-		const FCasSnapshot Snapshot = FCasSnapshot::Load(SnapshotFile).ConsumeValueOrDie();
+		TResult<int64> Result = FCasSnapshot::Save(ExpectedSnapshot, SnapshotFile);
+		CHECK(Result.HasValue());
+		const FCasSnapshot Snapshot = FCasSnapshot::Load(SnapshotFile).StealValue();
 
 		// Assert
 		CHECK(Snapshot.Blocks.Num() == ExpectedSnapshot.Blocks.Num());
@@ -4043,8 +3971,8 @@ TEST_CASE("IoStore::OnDemand::InstallCache::Snapshot", "[IoStoreOnDemand][Instal
 		const FString JournalFile = TestBaseDir / TEXT("test.jrn");
 		const FCasBlockId ExpectedCurrentBlockId(2);
 
-		FIoStatus Status = FCasJournal::Create(JournalFile);
-		CHECK(Status.IsOk());
+		FResult Result = FCasJournal::Create(JournalFile);
+		CHECK(Result.HasValue());
 
 		FCasJournal::FTransaction Tx = FCasJournal::Begin(JournalFile);
 
@@ -4079,11 +4007,11 @@ TEST_CASE("IoStore::OnDemand::InstallCache::Snapshot", "[IoStoreOnDemand][Instal
 			CreateCasTestAddr(Idx));
 		}
 
-		Status = FCasJournal::Commit(MoveTemp(Tx));
-		CHECK(Status.IsOk());
+		Result = FCasJournal::Commit(MoveTemp(Tx));
+		CHECK(Result.HasValue());
 
 		// Act
-		const FCasSnapshot Snapshot = FCasSnapshot::FromJournal(JournalFile).ConsumeValueOrDie();
+		const FCasSnapshot Snapshot = FCasSnapshot::FromJournal(JournalFile).StealValue();
 
 		// Assert
 		CHECK(Snapshot.CurrentBlockId == ExpectedCurrentBlockId);
@@ -4106,6 +4034,36 @@ TEST_CASE("IoStore::OnDemand::InstallCache::Snapshot", "[IoStoreOnDemand][Instal
 	}
 }
 
+TEST_CASE("IoStore::OnDemand::InstallCache::ErrorHandling", "[IoStoreOnDemand][InstallCache]")
+{
+	using namespace UE::UnifiedError;
+	using namespace UE::UnifiedError::IoStoreOnDemand;
+
+	SECTION("YesWeNeedToTestSerializeErrorDetailsToCbToMakeSureTheGameDoesNotCrashAtRuntime")
+	{
+		FResult Result = MakeCasError<void>(
+			ECasErrorCode::InitializeFailed, EIoErrorCode::DeleteError, TEXT("Test Message"));
+		UE_LOGFMT(LogIoStoreOnDemand, Display, "{Error}", Result.GetError());
+
+		{
+			FError Error = IoStoreOnDemand::InstallCacheFlushError::MakeError(FInstallCacheErrorContext{}, EDetailFilter::All);
+			UE_LOGFMT(LogIoStoreOnDemand, Display, "{Short}", Error.GetModuleIdAndErrorCodeString());
+			UE_LOGFMT(LogIoStoreOnDemand, Display, "{Error}", Error); }
+		{
+			FError Error = IoStoreOnDemand::InstallCacheFlushError::MakeError(FChunkMissingErrorContext{}, EDetailFilter::All);
+			UE_LOGFMT(LogIoStoreOnDemand, Display, "{Error}", Error);
+		}
+		{
+			FError Error = IoStoreOnDemand::InstallCachePurgeError::MakeError(FChunkHashMismatchErrorContext{}, EDetailFilter::All);
+			UE_LOGFMT(LogIoStoreOnDemand, Display, "{Error}", Error);
+		}
+		{
+			FError Error = IoStoreOnDemand::InstallCachePurgeError::MakeError(FVerificationErrorContext{}, EDetailFilter::All);
+			UE_LOGFMT(LogIoStoreOnDemand, Display, "{Error}", Error);
+		}
+	}
+}
+
 #endif // WITH_IOSTORE_ONDEMAND_TESTS
 
 } // namespace UE::IoStore
diff --git a/Engine/Source/Runtime/Experimental/IoStore/OnDemand/Private/Statistics.cpp b/Engine/Source/Runtime/Experimental/IoStore/OnDemand/Private/Statistics.cpp
--- a/Engine/Source/Runtime/Experimental/IoStore/OnDemand/Private/Statistics.cpp
+++ b/Engine/Source/Runtime/Experimental/IoStore/OnDemand/Private/Statistics.cpp
@@ -1407,12 +1407,12 @@ void FOnDemandContentInstallerStats::ReportAnalytics(TArray<FAnalyticsEventAttri
 }
 
 ////////////////////////////////////////////////////////////////////////////////
-void FOnDemandInstallCacheStats::OnStartupError(EIoErrorCode ErrorCode)
+void FOnDemandInstallCacheStats::OnStartupError(const FResult& Result)
 {
-	if (ErrorCode != EIoErrorCode::Ok)
+	if (Result.HasError())
 	{
-		UE::TUniqueLock Lock(InstallCacheAnalytics);
-		InstallCacheAnalytics.StartupErrorCode = static_cast<uint32>(ErrorCode);
+		FOnDemandImmediateAnalytic Event = MakeAnalyticsEventFromResult(TEXT("IadStartupError"), Result);
+		SendImmediateAnalytic(MoveTemp(Event));
 	}
 }
 
@@ -1432,15 +1432,15 @@ void FOnDemandInstallCacheStats::OnFlush(const FResult& Result, int64 ByteCount)
 	InstallCacheAnalytics.FlushedBytes += ByteCount;
 }
 
-void FOnDemandInstallCacheStats::OnJournalCommit(EIoErrorCode ErrorCode, int64 ByteCount)
+void FOnDemandInstallCacheStats::OnJournalCommit(const FResult& Result, int64 ByteCount)
 {
-	if (ErrorCode == EIoErrorCode::Cancelled)
+	if (Result.HasError() && UE::UnifiedError::IsCancellationError(Result.GetError()))
 	{
 		return;
 	}
 
 	UE::TUniqueLock Lock(InstallCacheAnalytics);
-	if (ErrorCode != EIoErrorCode::Ok)
+	if (Result.HasError())
 	{
 		InstallCacheAnalytics.JournalCommitErrorCount++;
 	}
diff --git a/Engine/Source/Runtime/Experimental/IoStore/OnDemand/Private/Statistics.h b/Engine/Source/Runtime/Experimental/IoStore/OnDemand/Private/Statistics.h
--- a/Engine/Source/Runtime/Experimental/IoStore/OnDemand/Private/Statistics.h
+++ b/Engine/Source/Runtime/Experimental/IoStore/OnDemand/Private/Statistics.h
@@ -142,9 +142,10 @@ public:
 class FOnDemandInstallCacheStats
 {
 public:
-	static void OnStartupError(EIoErrorCode ErrorCode) IAS_STATISTICS_IMPL()
+	static void OnStartupError(const FResult& Result) IAS_STATISTICS_IMPL()
 	static void OnFlush(const FResult& Result, int64 ByteCount) IAS_STATISTICS_IMPL()
 	static void OnJournalCommit(EIoErrorCode ErrorCode, int64 ByteCount) IAS_STATISTICS_IMPL()
+	static void OnJournalCommit(const FResult& Result, int64 ByteCount) IAS_STATISTICS_IMPL()
 	static void OnCasVerificationError(int32 RemoveChunks) IAS_STATISTICS_IMPL()
 	static void OnPurge(
 		const FResult& Result,	
diff --git a/Engine/Source/Runtime/Experimental/IoStore/OnDemandCore/Private/OnDemandError.cpp b/Engine/Source/Runtime/Experimental/IoStore/OnDemandCore/Private/OnDemandError.cpp
--- a/Engine/Source/Runtime/Experimental/IoStore/OnDemandCore/Private/OnDemandError.cpp
+++ b/Engine/Source/Runtime/Experimental/IoStore/OnDemandCore/Private/OnDemandError.cpp
@@ -11,6 +11,9 @@ UE_DEFINE_ERROR(InstallCacheFlushLastAccessError, IoStoreOnDemand);
 UE_DEFINE_ERROR(InstallCachePurgeError, IoStoreOnDemand);
 UE_DEFINE_ERROR(InstallCacheDefragError, IoStoreOnDemand);
 UE_DEFINE_ERROR(InstallCacheVerificationError, IoStoreOnDemand);
+UE_DEFINE_ERROR(CasError, IoStoreOnDemand);
+UE_DEFINE_ERROR(CasJournalError, IoStoreOnDemand);
+UE_DEFINE_ERROR(CasSnapshotError, IoStoreOnDemand);
 
 namespace UE::UnifiedError::IoStore
 {
diff --git a/Engine/Source/Runtime/Experimental/IoStore/OnDemandCore/Public/IO/OnDemandError.h b/Engine/Source/Runtime/Experimental/IoStore/OnDemandCore/Public/IO/OnDemandError.h
--- a/Engine/Source/Runtime/Experimental/IoStore/OnDemandCore/Public/IO/OnDemandError.h
+++ b/Engine/Source/Runtime/Experimental/IoStore/OnDemandCore/Public/IO/OnDemandError.h
@@ -10,70 +10,162 @@
 
 #define UE_API IOSTOREONDEMANDCORE_API
 
+////////////////////////////////////////////////////////////////////////////////
 UE_DECLARE_ERROR_MODULE(UE_API, IoStoreOnDemand);
 
-UE_DECLARE_ERROR_ONEPARAM(UE_API,	HttpError,							2,	IoStoreOnDemand, NSLOCTEXT("IoStoreOnDemand", "HttpError", "HTTP error ({StatusCode})"), uint32, StatusCode, 0);
-UE_DECLARE_ERROR(UE_API,			ChunkMissingError,					3,	IoStoreOnDemand, NSLOCTEXT("IoStoreOnDemand", "ChunkMissingError", "Chunk missing error"));
-UE_DECLARE_ERROR(UE_API,			ChunkHashError,						4,	IoStoreOnDemand, NSLOCTEXT("IoStoreOnDemand", "ChunkHashError", "Chunk hash mismatch error"));
-UE_DECLARE_ERROR(UE_API,			InstallCacheFlushError,				5,	IoStoreOnDemand, NSLOCTEXT("IoStoreOnDemand", "InstallCacheFlushError", "Failed to flush pending data to install cache."));
-UE_DECLARE_ERROR(UE_API,			InstallCacheFlushLastAccessError,	6,	IoStoreOnDemand, NSLOCTEXT("IoStoreOnDemand", "InstallCacheFlushLastAccessError", "Failed to flush last access timestamp(s) to journal."));
-UE_DECLARE_ERROR(UE_API,			InstallCachePurgeError,				7,	IoStoreOnDemand, NSLOCTEXT("IoStoreOnDemand", "InstallCachePurgeError", "Failed to purge unreferenced cache block(s) from the install cache."));
-UE_DECLARE_ERROR(UE_API,			InstallCacheDefragError,			8,	IoStoreOnDemand, NSLOCTEXT("IoStoreOnDemand", "InstallCacheDefragError", "Failed to defrag the install cache."));
-UE_DECLARE_ERROR(UE_API,			InstallCacheVerificationError,		9,	IoStoreOnDemand, NSLOCTEXT("IoStoreOnDemand", "InstallCacheVerificationError", "Verification of installed install cache data failed."));
+UE_DECLARE_ERROR_ONEPARAM(UE_API,	HttpError,							1,	IoStoreOnDemand, NSLOCTEXT("IoStoreOnDemand", "HttpError", "HTTP error ({StatusCode})"), uint32, StatusCode, 0);
+UE_DECLARE_ERROR(UE_API,			ChunkMissingError,					2,	IoStoreOnDemand, NSLOCTEXT("IoStoreOnDemand", "ChunkMissingError", "Chunk missing error"));
+UE_DECLARE_ERROR(UE_API,			ChunkHashError,						3,	IoStoreOnDemand, NSLOCTEXT("IoStoreOnDemand", "ChunkHashError", "Chunk hash mismatch error"));
+UE_DECLARE_ERROR(UE_API,			InstallCacheFlushError,				4,	IoStoreOnDemand, NSLOCTEXT("IoStoreOnDemand", "InstallCacheFlushError", "Failed to flush pending data to install cache."));
+UE_DECLARE_ERROR(UE_API,			InstallCacheFlushLastAccessError,	5,	IoStoreOnDemand, NSLOCTEXT("IoStoreOnDemand", "InstallCacheFlushLastAccessError", "Failed to flush last access timestamp(s) to journal."));
+UE_DECLARE_ERROR(UE_API,			InstallCachePurgeError,				6,	IoStoreOnDemand, NSLOCTEXT("IoStoreOnDemand", "InstallCachePurgeError", "Failed to purge unreferenced cache block(s) from the install cache."));
+UE_DECLARE_ERROR(UE_API,			InstallCacheDefragError,			7,	IoStoreOnDemand, NSLOCTEXT("IoStoreOnDemand", "InstallCacheDefragError", "Failed to defrag the install cache."));
+UE_DECLARE_ERROR(UE_API,			InstallCacheVerificationError,		8,	IoStoreOnDemand, NSLOCTEXT("IoStoreOnDemand", "InstallCacheVerificationError", "Verification of installed install cache data failed."));
+UE_DECLARE_ERROR(UE_API,			CasError,							9,	IoStoreOnDemand, NSLOCTEXT("IoStoreOnDemand", "CasError", "Cas error"));
+UE_DECLARE_ERROR(UE_API,			CasJournalError,					10,	IoStoreOnDemand, NSLOCTEXT("IoStoreOnDemand", "CasJournalError", "Cas journal error"));
+UE_DECLARE_ERROR(UE_API,			CasSnapshotError,					11,	IoStoreOnDemand, NSLOCTEXT("IoStoreOnDemand", "CasSnapshotError", "Cas snapshot error"));
+
+////////////////////////////////////////////////////////////////////////////////
+namespace UE::IoStore
+{
+
+using FResult = TValueOrError<void, UE::UnifiedError::FError>;
+
+template <typename ResultType>
+using TResult = TValueOrError<ResultType, UE::UnifiedError::FError>;
+
+} // namespace UE::IoStore
 
+////////////////////////////////////////////////////////////////////////////////
 namespace UE::UnifiedError::IoStoreOnDemand
 {
 
-struct FChunkMissingErrorContext
+enum class ECasErrorCode : uint32
 {
-	TArray<FIoChunkId> ChunkIds;
+	None,
+	InitializeFailed,
+	VerifyFailed,
+	ReadBlockFailed,
+	WriteBlockFailed,
+	DeleteBlockFailed,
+	CreateJournalFailed,
+	ReplayJournalFailed,
+	CommitJournalFailed,
+	CreateSnapshotFailed,
+	LoadSnapshotFailed,
+	SaveSnapshotFailed
 };
 
-struct FChunkHashMismatchErrorContext
+inline const TCHAR* ToString(ECasErrorCode Code)
 {
-	FIoChunkId	ChunkId;
-	FIoHash		ExpectedHash;
-	FIoHash		ActualHash;
-};
+	switch (Code)
+	{
+		case ECasErrorCode::None:
+			return TEXT("None");
+		case ECasErrorCode::InitializeFailed:
+			return TEXT("InitializeFailed");
+		case ECasErrorCode::VerifyFailed:
+			return TEXT("VerifyFailed");
+		case ECasErrorCode::ReadBlockFailed:
+			return TEXT("ReadBlockFailed");
+		case ECasErrorCode::WriteBlockFailed:
+			return TEXT("WriteBlockFailed");
+		case ECasErrorCode::DeleteBlockFailed:
+			return TEXT("DeleteBlockFailed");
+		case ECasErrorCode::CreateJournalFailed:
+			return TEXT("CreateJournalFailed");
+		case ECasErrorCode::ReplayJournalFailed:
+			return TEXT("ReplayJournalFailed");
+		case ECasErrorCode::CommitJournalFailed:
+			return TEXT("CommitJournalFailed");
+		case ECasErrorCode::CreateSnapshotFailed:
+			return TEXT("CreateSnapshotFailed");
+		case ECasErrorCode::LoadSnapshotFailed:
+			return TEXT("LoadSnapshotFailed");
+		case ECasErrorCode::SaveSnapshotFailed:
+			return TEXT("SaveSnapshotFailed");
+		default:
+			return TEXT("<InvalidErrorCode>");
+	};
+}
 
-struct FIoErrorContext
+struct FCasErrorContext
 {
-	FIoStatus Status;
+	ECasErrorCode	ErrorCode = ECasErrorCode::None;
+	EIoErrorCode	IoErrorCode = EIoErrorCode::Unknown;
+	uint32			SystemErrorCode = 0;
+	FString			ErrorMessage;
 };
 
-enum class EErrorCategory
+template <typename ResultType>
+inline UE::IoStore::TResult<ResultType> MakeCasError(
+	ECasErrorCode ErrorCode,
+	EIoErrorCode IoErrorCode,
+	FString&& ErrorMessage,
+	uint32 SystemErrorCode = FPlatformMisc::GetLastError())
 {
-	None = 0,
-	Cas,
-	Journal,
-	Logical
-};
+	return MakeError(CasError::MakeError(FCasErrorContext
+	{
+		.ErrorCode			= ErrorCode,
+		.IoErrorCode		= IoErrorCode,
+		.SystemErrorCode	= SystemErrorCode,
+		.ErrorMessage		= MoveTemp(ErrorMessage)
+	}, UE::UnifiedError::EDetailFilter::All));
+}
 
-inline const TCHAR* ToString(EErrorCategory Category)
+inline UE::IoStore::FResult MakeJournalError(
+	ECasErrorCode ErrorCode,
+	EIoErrorCode IoErrorCode,
+	FString&& ErrorMessage,
+	uint32 SystemErrorCode = FPlatformMisc::GetLastError())
 {
-	switch (Category)
+	return MakeError(CasJournalError::MakeError(FCasErrorContext
 	{
-		case EErrorCategory::None:
-			return TEXT("None");
-		case EErrorCategory::Cas:
-			return TEXT("Cas");
-		case EErrorCategory::Journal:
-			return TEXT("Journal");
-		case EErrorCategory::Logical:
-			return TEXT("Logical");
-		default:
-			return TEXT("Unknown");
-	}
+		.ErrorCode			= ErrorCode,
+		.IoErrorCode		= IoErrorCode,
+		.SystemErrorCode	= SystemErrorCode,
+		.ErrorMessage		= MoveTemp(ErrorMessage)
+	}));
+}
+
+template <typename ResultType>
+inline UE::IoStore::TResult<ResultType> MakeSnapshotError(
+	ECasErrorCode ErrorCode,
+	EIoErrorCode IoErrorCode,
+	FString&& ErrorMessage,
+	uint32 SystemErrorCode = FPlatformMisc::GetLastError())
+{
+	using namespace UE::UnifiedError;
+	using namespace UE::UnifiedError::IoStoreOnDemand;
+
+	return MakeError(CasSnapshotError::MakeError(FCasErrorContext
+	{
+		.ErrorCode			= ErrorCode,
+		.IoErrorCode		= IoErrorCode,
+		.SystemErrorCode	= SystemErrorCode,
+		.ErrorMessage		= MoveTemp(ErrorMessage)
+	}));
+}
+
+struct FChunkMissingErrorContext
+{
+	TArray<FIoChunkId> ChunkIds;
+};
+
+struct FChunkHashMismatchErrorContext
+{
+	FIoChunkId	ChunkId;
+	FIoHash		ExpectedHash;
+	FIoHash		ActualHash;
 };
 
 struct FInstallCacheErrorContext
 {
-	FIoStatus		Error = FIoStatus(EIoErrorCode::Unknown);
+	FIoStatus		IoError = FIoStatus(EIoErrorCode::Unknown);
 	uint64			MaxCacheSize = 0;
 	uint64			CacheSize = 0;
 	uint64			DiskTotalBytes = 0;
 	uint64			DiskFreeBytes = 0;
-	EErrorCategory	ErrorCategory = EErrorCategory::None;
 	uint32			LineNo = uint32(__builtin_LINE());
 	bool			bDiskQuerySucceeded = false;
 };
@@ -86,6 +178,22 @@ struct FVerificationErrorContext
 };
 
 } // namespace UE::UnifiedError::IoStoreOnDemand
+
+////////////////////////////////////////////////////////////////////////////////
+UE_DECLARE_ERRORSTRUCT_FEATURES(IoStoreOnDemand, FCasErrorContext);
+inline void SerializeForLog(FCbWriter& Writer, const UE::UnifiedError::IoStoreOnDemand::FCasErrorContext& Ctx)
+{
+	Writer.BeginObject();
+	Writer.AddString(ANSITEXTVIEW("$type"), TErrorStructFeatures<UE::UnifiedError::IoStoreOnDemand::FCasErrorContext>::GetErrorContextTypeNameAsString());
+	Writer.AddString(ANSITEXTVIEW("$format"), TEXT("({ErrorCodeText}: {ErrorMessage})"));
+	Writer.AddInteger(ANSITEXTVIEW("ErrorCode"), uint32(Ctx.ErrorCode));
+	Writer.AddString(ANSITEXTVIEW("ErrorCodeText"), ToString(Ctx.ErrorCode));
+	Writer.AddInteger(ANSITEXTVIEW("IoErrorCode"), uint32(Ctx.IoErrorCode));
+	Writer.AddString(ANSITEXTVIEW("IoErrorCodeText"), GetIoErrorText(Ctx.IoErrorCode));
+	Writer.AddString(ANSITEXTVIEW("ErrorMessage"), Ctx.ErrorMessage); 
+	Writer.EndObject();
+}
+////////////////////////////////////////////////////////////////////////////////
 UE_DECLARE_ERRORSTRUCT_FEATURES(IoStoreOnDemand, FChunkMissingErrorContext);
 inline void SerializeForLog(FCbWriter& Writer, const UE::UnifiedError::IoStoreOnDemand::FChunkMissingErrorContext& Ctx)
 {
@@ -97,38 +205,28 @@ inline void SerializeForLog(FCbWriter& Writer, const UE::UnifiedError::IoStoreOn
 	//Writer.EndArray();
 	Writer.EndObject();
 }
+////////////////////////////////////////////////////////////////////////////////
 UE_DECLARE_ERRORSTRUCT_FEATURES(IoStoreOnDemand, FChunkHashMismatchErrorContext);
 inline void SerializeForLog(FCbWriter& Writer, const UE::UnifiedError::IoStoreOnDemand::FChunkHashMismatchErrorContext& Ctx)
 {
 	Writer.BeginObject();
 	Writer.AddString(ANSITEXTVIEW("$type"), TErrorStructFeatures<UE::UnifiedError::IoStoreOnDemand::FChunkHashMismatchErrorContext>::GetErrorContextTypeNameAsString());
-	Writer.AddString(ANSITEXTVIEW("$format"), TEXT("(ChunkId: {ChunkId})"));
+	Writer.AddString(ANSITEXTVIEW("$format"), TEXT("(ChunkId: {ChunkId}, ExpectedHash: {ExpectedHash}, ActualHash: {ActualHash})"));
 	Writer.AddString(ANSITEXTVIEW("ChunkId"), LexToString(Ctx.ChunkId));
 	Writer.AddString(ANSITEXTVIEW("ExpectedHash"), LexToString(Ctx.ExpectedHash));
 	Writer.AddString(ANSITEXTVIEW("ActualHash"), LexToString(Ctx.ActualHash));
 	Writer.EndObject();
 }
-UE_DECLARE_ERRORSTRUCT_FEATURES(IoStoreOnDemand, FIoErrorContext);
-inline void SerializeForLog(FCbWriter& Writer, const UE::UnifiedError::IoStoreOnDemand::FIoErrorContext& Ctx)
-{
-	Writer.BeginObject();
-	Writer.AddString(ANSITEXTVIEW("$type"), TErrorStructFeatures<UE::UnifiedError::IoStoreOnDemand::FIoErrorContext>::GetErrorContextTypeNameAsString());
-	Writer.AddString(ANSITEXTVIEW("$format"), TEXT("(Status: {Status})"));
-	Writer.AddString(ANSITEXTVIEW("Status"), Ctx.Status.ToString());
-	Writer.EndObject();
-}
+////////////////////////////////////////////////////////////////////////////////
 UE_DECLARE_ERRORSTRUCT_FEATURES(IoStoreOnDemand, FInstallCacheErrorContext);
 inline void SerializeForLog(FCbWriter& Writer, const UE::UnifiedError::IoStoreOnDemand::FInstallCacheErrorContext& Ctx)
 {
 	Writer.BeginObject();
 	Writer.AddString(ANSITEXTVIEW("$type"), TErrorStructFeatures<UE::UnifiedError::IoStoreOnDemand::FInstallCacheErrorContext>::GetErrorContextTypeNameAsString());
-	Writer.AddString(ANSITEXTVIEW("$format"), TEXT("(Category: {ErrorCategoryText}, ErrorCode: {ErrorCodeText}, SystemErrorCode: {SystemErrorCode}, ErrorMessage: {ErrorMessage})"));
-	Writer.AddInteger(ANSITEXTVIEW("ErrorCategory"), int32(Ctx.ErrorCategory));
-	Writer.AddString(ANSITEXTVIEW("ErrorCategoryText"), ToString(Ctx.ErrorCategory));
-	Writer.AddInteger(ANSITEXTVIEW("ErrorCode"), int32(Ctx.Error.GetErrorCode()));
-	Writer.AddString(ANSITEXTVIEW("ErrorCodeText"), GetIoErrorText(Ctx.Error.GetErrorCode()));
-	Writer.AddString(ANSITEXTVIEW("ErrorMessage"), Ctx.Error.ToString());
-	Writer.AddInteger(ANSITEXTVIEW("SystemErrorCode"), int32(Ctx.Error.GetSystemErrorCode()));
+	Writer.AddString(ANSITEXTVIEW("$format"), TEXT("({ErrorMessage})"));
+	Writer.AddInteger(ANSITEXTVIEW("ErrorCode"), int32(Ctx.IoError.GetErrorCode()));
+	Writer.AddString(ANSITEXTVIEW("ErrorMessage"), Ctx.IoError.ToString());
+	Writer.AddInteger(ANSITEXTVIEW("SystemErrorCode"), int32(Ctx.IoError.GetSystemErrorCode()));
 	Writer.AddInteger(ANSITEXTVIEW("MaxCacheSize"), Ctx.MaxCacheSize);
 	Writer.AddInteger(ANSITEXTVIEW("CacheSize"), Ctx.CacheSize);
 	Writer.AddInteger(ANSITEXTVIEW("DiskTotalBytes"), Ctx.DiskTotalBytes);
@@ -136,6 +234,7 @@ inline void SerializeForLog(FCbWriter& Writer, const UE::UnifiedError::IoStoreOn
 	Writer.AddInteger(ANSITEXTVIEW("LineNo"), uint32(Ctx.LineNo));
 	Writer.EndObject();
 }
+////////////////////////////////////////////////////////////////////////////////
 UE_DECLARE_ERRORSTRUCT_FEATURES(IoStoreOnDemand, FVerificationErrorContext);
 inline void SerializeForLog(FCbWriter& Writer, const UE::UnifiedError::IoStoreOnDemand::FVerificationErrorContext& Ctx)
 {
@@ -148,11 +247,4 @@ inline void SerializeForLog(FCbWriter& Writer, const UE::UnifiedError::IoStoreOn
 	Writer.EndObject();
 }
 
-namespace UE::IoStore
-{
-
-using FResult = TValueOrError<void, UE::UnifiedError::FError>;
-
-} // namespace UE::IoStore
-
 #undef UE_API
