Fixed bug when handling open file result in on-demand install cache.

diff --git a/Engine/Source/Runtime/Experimental/IoStore/OnDemand/Private/OnDemandInstallCache.cpp b/Engine/Source/Runtime/Experimental/IoStore/OnDemand/Private/OnDemandInstallCache.cpp
--- a/Engine/Source/Runtime/Experimental/IoStore/OnDemand/Private/OnDemandInstallCache.cpp
+++ b/Engine/Source/Runtime/Experimental/IoStore/OnDemand/Private/OnDemandInstallCache.cpp
@@ -3193,7 +3193,7 @@ FResult FOnDemandInstallCache::Defrag(
 		if (DefragBlock.ReferencedChunks.IsEmpty() == false)
 		{
 			TResult<FSharedFileHandle> FileOpenResult = Cas.OpenRead(DefragBlock.BlockId);
-			if (FileOpenResult.HasValue())
+			if (FileOpenResult.HasError())
 			{
 				UE_LOGFMT(LogIoStoreOnDemand, Error, "Defrag failed to open CAS block for reading {Error}", FileOpenResult.GetError());
 
