Changed format specifiers to use UE macros in on-demand install cache.

#rnx

diff --git a/Engine/Source/Runtime/Experimental/IoStore/OnDemand/Private/OnDemandInstallCache.cpp b/Engine/Source/Runtime/Experimental/IoStore/OnDemand/Private/OnDemandInstallCache.cpp
--- a/Engine/Source/Runtime/Experimental/IoStore/OnDemand/Private/OnDemandInstallCache.cpp
+++ b/Engine/Source/Runtime/Experimental/IoStore/OnDemand/Private/OnDemandInstallCache.cpp
@@ -599,7 +599,7 @@ uint64 FCas::GetBlockInfo(FCasBlockInfoMap& OutBlockInfo)
 
 			if (BlockId.IsValid() == false || FileSize < 0)
 			{
-				UE_LOG(LogIoStoreOnDemand, Warning, TEXT("Found invalid CAS block '%s', FileSize=%lld"),
+				UE_LOG(LogIoStoreOnDemand, Warning, TEXT("Found invalid CAS block '%s', FileSize=%" INT64_FMT),
 					FilenameOrDirectory, FileSize);
 				return true;
 			}
@@ -1002,7 +1002,7 @@ FResult FCasJournal::Replay(const FString& JournalFile, FEntryHandler&& Handler)
 	const int64 EntriesPos = FileHandle->Tell();
 	if (FileHandle->Seek(FooterPos) == false)
 	{
-		return MakeJournalError(ECasErrorCode::ReplayJournalFailed, EIoErrorCode::FileSeekFailed, FString::Printf(TEXT("Failed to seek to footer offset %ull in '%s'"), FooterPos, *JournalFile));
+		return MakeJournalError(ECasErrorCode::ReplayJournalFailed, EIoErrorCode::FileSeekFailed, FString::Printf(TEXT("Failed to seek to footer offset %" INT64_FMT " in '%s'"), FooterPos, *JournalFile));
 	}
 
 	FFooter Footer;
@@ -1013,7 +1013,7 @@ FResult FCasJournal::Replay(const FString& JournalFile, FEntryHandler&& Handler)
 
 	if (FileHandle->Seek(EntriesPos) == false)
 	{
-		return MakeJournalError(ECasErrorCode::ReplayJournalFailed, EIoErrorCode::FileSeekFailed, FString::Printf(TEXT("Failed to seek to entries offset %ull in '%s'"), EntriesPos, *JournalFile));
+		return MakeJournalError(ECasErrorCode::ReplayJournalFailed, EIoErrorCode::FileSeekFailed, FString::Printf(TEXT("Failed to seek to entries offset %" INT64_FMT " in '%s'"), EntriesPos, *JournalFile));
 	}
 
 	TArray<FEntry> Entries;
@@ -1158,7 +1158,7 @@ FResult FCasJournal::Commit(FTransaction&& Transaction)
 		if (FileHandle->Seek(EntriesPos) == false)
 		{
 			return Result = MakeJournalError(
-				ECasErrorCode::CommitJournalFailed, EIoErrorCode::FileSeekFailed, FString::Printf(TEXT("Failed to seek to journal entries offset %ull in '%s'"), EntriesPos, *Transaction.JournalFile));
+				ECasErrorCode::CommitJournalFailed, EIoErrorCode::FileSeekFailed, FString::Printf(TEXT("Failed to seek to journal entries offset %" INT64_FMT " in '%s'"), EntriesPos, *Transaction.JournalFile));
 		}
 
 		TotalEntrySize = Transaction.Entries.Num() * FEntry::Size();
@@ -2174,7 +2174,7 @@ bool FOnDemandInstallCache::Resolve(FIoRequestImpl* Request)
 				const int64 CasBlockOffset = CasLoc.BlockOffset + ChunkRequest.ChunkRange.GetOffset();
 				if (Request->IsCancelled())
 				{
-					UE_LOG(LogIoStoreOnDemand, Verbose, TEXT("Cancelled request - skipped seek to offset %lld in CAS block '%s'"), CasBlockOffset, *Filename);
+					UE_LOG(LogIoStoreOnDemand, Verbose, TEXT("Cancelled request - skipped seek to offset %" INT64_FMT " in CAS block '%s'"), CasBlockOffset, *Filename);
 				}
 				else if (FileHandle->Seek(CasBlockOffset))
 				{
@@ -2185,7 +2185,7 @@ bool FOnDemandInstallCache::Resolve(FIoRequestImpl* Request)
 					}
 					else
 					{
-						UE_LOG(LogIoStoreOnDemand, Error, TEXT("Failed to read %llu bytes at offset %lld in CAS block '%s'"),
+						UE_LOG(LogIoStoreOnDemand, Error, TEXT("Failed to read %" UINT64_FMT " bytes at offset %" INT64_FMT " in CAS block '%s'"),
 							ChunkRequest.EncodedChunk.GetSize(),
 							CasBlockOffset,
 							*Filename);
@@ -2193,7 +2193,7 @@ bool FOnDemandInstallCache::Resolve(FIoRequestImpl* Request)
 				}
 				else
 				{
-					UE_LOG(LogIoStoreOnDemand, Error, TEXT("Failed to seek to offset %lld in CAS block '%s'"), CasBlockOffset, *Filename);
+					UE_LOG(LogIoStoreOnDemand, Error, TEXT("Failed to seek to offset %" INT64_FMT " in CAS block '%s'"), CasBlockOffset, *Filename);
 				}
 			}
 			else
@@ -2649,7 +2649,7 @@ FResult FOnDemandInstallCache::Verify()
 
 		if (FileHandle->Seek(int64(ChunkLocation.Location.BlockOffset)) == false)
 		{
-			UE_LOG(LogIoStoreOnDemand, Error, TEXT("Chunk %d/%d SEEK FAILED, Container='%s', ChunkId='%s', ChunkSize=%lld, Hash='%s', Block=%u, BlockOffset=%u"),
+			UE_LOG(LogIoStoreOnDemand, Error, TEXT("Chunk %d/%d SEEK FAILED, Container='%s', ChunkId='%s', ChunkSize=%" INT64_FMT ", Hash='%s', Block=%u, BlockOffset=%u"),
 				ChunkIndex + 1, ChunkCount, *Container->Name, *LexToString(ChunkId), ChunkSize, *LexToString(ChunkEntry.Hash),
 				ChunkLocation.Location.BlockId.Id, ChunkLocation.Location.BlockOffset);
 
@@ -2660,7 +2660,7 @@ FResult FOnDemandInstallCache::Verify()
 
 		if (FileHandle->Read(reinterpret_cast<uint8*>(Chunk.GetData()), ChunkSize) == false)
 		{
-			UE_LOG(LogIoStoreOnDemand, Error, TEXT("Chunk %d/%d READ FAILED, Container='%s', ChunkId='%s', ChunkSize=%lld, Hash='%s', Block=%u, BlockOffset=%u"),
+			UE_LOG(LogIoStoreOnDemand, Error, TEXT("Chunk %d/%d READ FAILED, Container='%s', ChunkId='%s', ChunkSize=%" INT64_FMT ", Hash='%s', Block=%u, BlockOffset=%u"),
 				ChunkIndex + 1, ChunkCount, *Container->Name, *LexToString(ChunkId), ChunkSize, *LexToString(ChunkEntry.Hash),
 				ChunkLocation.Location.BlockId.Id, ChunkLocation.Location.BlockOffset);
 
@@ -2673,13 +2673,13 @@ FResult FOnDemandInstallCache::Verify()
 
 		if (ChunkHash == ChunkEntry.Hash)
 		{
-			UE_LOG(LogIoStoreOnDemand, VeryVerbose, TEXT("Chunk %d/%d OK, Container='%s', ChunkId='%s', ChunkSize=%lld, Hash='%s', Block=%u, BlockOffset=%u"),
+			UE_LOG(LogIoStoreOnDemand, VeryVerbose, TEXT("Chunk %d/%d OK, Container='%s', ChunkId='%s', ChunkSize=%" INT64_FMT ", Hash='%s', Block=%u, BlockOffset=%u"),
 				ChunkIndex + 1, ChunkCount, *Container->Name, *LexToString(ChunkId), ChunkSize, *LexToString(ChunkEntry.Hash),
 				ChunkLocation.Location.BlockId.Id, ChunkLocation.Location.BlockOffset);
 		}
 		else
 		{
-			UE_LOG(LogIoStoreOnDemand, Error, TEXT("Chunk %d/%d CORRUPT, Container='%s', ChunkId='%s', ChunkSize=%lld, Hash='%s', ActualHash='%s', Block=%u, BlockOffset=%u"),
+			UE_LOG(LogIoStoreOnDemand, Error, TEXT("Chunk %d/%d CORRUPT, Container='%s', ChunkId='%s', ChunkSize=%" INT64_FMT ", Hash='%s', ActualHash='%s', Block=%u, BlockOffset=%u"),
 				ChunkIndex + 1, ChunkCount, *Container->Name, *LexToString(ChunkId), ChunkSize, *LexToString(ChunkEntry.Hash), *LexToString(ChunkHash),
 				ChunkLocation.Location.BlockId.Id, ChunkLocation.Location.BlockOffset);
 
@@ -3223,7 +3223,7 @@ FResult FOnDemandInstallCache::Defrag(
 					if (FileHandle->Read(Buffer.GetData(), Buffer.GetSize()) == false)
 					{
 						ReadResult = MakeCasError<void>(
-							ECasErrorCode::ReadBlockFailed, EIoErrorCode::ReadError, FString::Printf(TEXT("Failed to read %ull bytes at offset %u in cache block %u"),
+							ECasErrorCode::ReadBlockFailed, EIoErrorCode::ReadError, FString::Printf(TEXT("Failed to read %" UINT64_FMT " bytes at offset %u in cache block %u"),
 								Buffer.GetSize(), ReffedChunk.BlockOffset, DefragBlock.BlockId.Id));
 						continue;
 					}
@@ -3537,7 +3537,7 @@ FResult FOnDemandInstallCache::FlushPendingChunksImpl(const FPendingChunks& Chun
 			if (CasFileHandle->Write(Ar.GetData(), Ar.Tell()) == false)
 			{
 				Result = MakeCasError<void>(
-					ECasErrorCode::WriteBlockFailed, EIoErrorCode::WriteError, FString::Printf(TEXT("Failed to write %lld bytes to CAS block %u"), Ar.Tell(), CurrentBlockId.Id));
+					ECasErrorCode::WriteBlockFailed, EIoErrorCode::WriteError, FString::Printf(TEXT("Failed to write %" INT64_FMT " bytes to CAS block %u"), Ar.Tell(), CurrentBlockId.Id));
 				FInstallCacheErrorContext Ctx = MakeInstallCacheErrorContext();
 				Result.GetError().PushErrorContext(MoveTemp(Ctx));
 				return Result;
@@ -3546,7 +3546,7 @@ FResult FOnDemandInstallCache::FlushPendingChunksImpl(const FPendingChunks& Chun
 			if (CasFileHandle->Flush() == false)
 			{
 				Result = MakeCasError<void>(
-					ECasErrorCode::WriteBlockFailed, EIoErrorCode::FileFlushFailed, FString::Printf(TEXT("Failed to flush %lld bytes to CAS block %u"), Ar.Tell(), CurrentBlockId.Id));
+					ECasErrorCode::WriteBlockFailed, EIoErrorCode::FileFlushFailed, FString::Printf(TEXT("Failed to flush %" INT64_FMT " bytes to CAS block %u"), Ar.Tell(), CurrentBlockId.Id));
 				FInstallCacheErrorContext Ctx = MakeInstallCacheErrorContext();
 				Result.GetError().PushErrorContext(MoveTemp(Ctx));
 				return Result;
