Geometry collection : make sure that when the physical material override is not specified and there's no associated physics material with the rendering materials that the Geometry collection defined material is being used

diff --git a/Engine/Source/Runtime/Experimental/GeometryCollectionEngine/Private/GeometryCollection/GeometryCollectionComponent.cpp b/Engine/Source/Runtime/Experimental/GeometryCollectionEngine/Private/GeometryCollection/GeometryCollectionComponent.cpp
--- a/Engine/Source/Runtime/Experimental/GeometryCollectionEngine/Private/GeometryCollection/GeometryCollectionComponent.cpp
+++ b/Engine/Source/Runtime/Experimental/GeometryCollectionEngine/Private/GeometryCollection/GeometryCollectionComponent.cpp
@@ -1649,7 +1649,7 @@ UPhysicalMaterial* UGeometryCollectionComponent::GetPhysicalMaterial() const
 	// Pull material from first mesh element to grab physical material. Prefer an override if one exists
 	UPhysicalMaterial* PhysMatToUse = BodyInstance.GetSimplePhysicalMaterial();
 
-	if(!PhysMatToUse || PhysMatToUse->GetFName() == "DefaultPhysicalMaterial")
+	if(!PhysMatToUse || PhysMatToUse->GetFName() == "DefaultPhysicalMaterial" || PhysMatToUse == GEngine->DefaultPhysMaterial)
 	{
 		// No override, try render materials
 		const int32 NumMaterials = GetNumMaterials();
@@ -1665,6 +1665,11 @@ UPhysicalMaterial* UGeometryCollectionComponent::GetPhysicalMaterial() const
 		}
 	}
 
+	if ((!PhysMatToUse || PhysMatToUse == GEngine->DefaultPhysMaterial) && RestCollection)
+	{
+		PhysMatToUse = RestCollection->PhysicsMaterial;
+	}
+
 	if(!PhysMatToUse)
 	{
 		// Still no material, fallback on default
