Do not log an error when ondemand toc mounting is delayed pending a valid hostgroup.

#rnx
#preflgiht

- Previously we were logging all failed FOnDemandMountResult as errors but one of the failure conditions (EIoErrorCode::PendingHostGroup) is something that is valid to happen. The toc will be mounted at a later point when the hostgroup becomes available. This could result in an error being logged very early on that gives the impression to the user that there is a problem with an IAX system, even though everything is fine.
- We now check for that specific return value and print a normal log message instead to avoid showing a false error.
- Since this logic needed to be repeated in several places, it has been moved to a single method FOnDemandMountResult::LogResult to avoid code duplication and make it easier to update in the future.
- Switched from old style logging to structured logs.

diff --git a/Engine/Source/Runtime/Experimental/IoStore/OnDemand/Private/OnDemandIoStore.cpp b/Engine/Source/Runtime/Experimental/IoStore/OnDemand/Private/OnDemandIoStore.cpp
--- a/Engine/Source/Runtime/Experimental/IoStore/OnDemand/Private/OnDemandIoStore.cpp
+++ b/Engine/Source/Runtime/Experimental/IoStore/OnDemand/Private/OnDemandIoStore.cpp
@@ -521,8 +521,7 @@ FIoStatus FOnDemandIoStore::InitializePostHotfix()
 	{
 		Mount(MoveTemp(MountArgs.GetValue()), [](FOnDemandMountResult MountResult)
 		{
-			UE_CLOG(!MountResult.Status.IsOk(), LogIas, Error,
-				TEXT("Failed to mount TOC for '%s', reason '%s'"), *MountResult.MountId, *MountResult.Status.ToString());
+			MountResult.LogResult();
 		});
 #if !UE_BUILD_SHIPPING
 		if (InstallArgs)
diff --git a/Engine/Source/Runtime/Experimental/IoStore/OnDemandCore/Private/IoStoreOnDemand.cpp b/Engine/Source/Runtime/Experimental/IoStore/OnDemandCore/Private/IoStoreOnDemand.cpp
--- a/Engine/Source/Runtime/Experimental/IoStore/OnDemandCore/Private/IoStoreOnDemand.cpp
+++ b/Engine/Source/Runtime/Experimental/IoStore/OnDemandCore/Private/IoStoreOnDemand.cpp
@@ -207,8 +207,7 @@ static void MountUrl(const TArray<FString>& Args)
 		},
 		[](FOnDemandMountResult MountResult)
 		{
-			UE_CLOG(!MountResult.Status.IsOk(), LogIas, Error,
-				TEXT("Failed to mount TOC for '%s', reason '%s'"), *MountResult.MountId, *MountResult.Status.ToString());
+			MountResult.LogResult();
 		});
 }
 
@@ -299,8 +298,7 @@ static void MountFile(const TArray<FString>& Args)
 			MoveTemp(MountArgs),
 			[](FOnDemandMountResult MountResult)
 			{
-				UE_CLOG(!MountResult.Status.IsOk(), LogIas, Error,
-					TEXT("Failed to mount TOC for '%s', reason '%s'"), *MountResult.MountId, *MountResult.Status.ToString());
+				MountResult.LogResult();
 			});
 	}
 }
diff --git a/Engine/Source/Runtime/Experimental/IoStore/OnDemandCore/Private/IoStoreOnDemandInternals.cpp b/Engine/Source/Runtime/Experimental/IoStore/OnDemandCore/Private/IoStoreOnDemandInternals.cpp
--- a/Engine/Source/Runtime/Experimental/IoStore/OnDemandCore/Private/IoStoreOnDemandInternals.cpp
+++ b/Engine/Source/Runtime/Experimental/IoStore/OnDemandCore/Private/IoStoreOnDemandInternals.cpp
@@ -2,9 +2,21 @@
 
 #include "IO/IoStoreOnDemandInternals.h"
 
+#include "Logging/StructuredLog.h"
+
 namespace UE::IoStore
 {
 
+////////////////////////////////////////////////////////////////////////////////
+void FOnDemandMountResult::LogResult()
+{
+	if (!Status.IsOk())
+	{
+		UE_CLOGFMT(Status.GetErrorCode() != EIoErrorCode::PendingHostGroup, LogIas, Error, "Failed to mount TOC for '{0}', reason '{1}'", MountId, Status);
+		UE_CLOGFMT(Status.GetErrorCode() == EIoErrorCode::PendingHostGroup, LogIas, Log, "Failed to mount TOC for '{0}', reason '{1}'", MountId, Status);
+	}
+}
+
 ///////////////////////////////////////////////////////////////////////////////
 FOnDemandWeakContentHandle FOnDemandWeakContentHandle::FromUnsafeHandle(UPTRINT HandleId)
 {
diff --git a/Engine/Source/Runtime/Experimental/IoStore/OnDemandCore/Public/IO/IoStoreOnDemand.h b/Engine/Source/Runtime/Experimental/IoStore/OnDemandCore/Public/IO/IoStoreOnDemand.h
--- a/Engine/Source/Runtime/Experimental/IoStore/OnDemandCore/Public/IO/IoStoreOnDemand.h
+++ b/Engine/Source/Runtime/Experimental/IoStore/OnDemandCore/Public/IO/IoStoreOnDemand.h
@@ -197,6 +197,14 @@ struct FOnDemandMountResult
 	FIoStatus Status;
 	/** Duration in seconds. */
 	double DurationInSeconds = 0.0;
+
+	/** 
+	 * Prints the result to the log.
+	 * If the result was a success then nothing will be logged.
+	 * If the mount failed on EIoErrorCode::PendingHostGroup then the result will be printed with 'log' verbosity.
+	 * All other results will be printed with 'error' verbosity.
+	 */
+	UE_API void LogResult();
 };
 
 /** Mount completion callback. */
