The new IAX backend now respects the 'ias.HttpRecvBufKiB"' cvar.

#rnx

diff --git a/Engine/Source/Runtime/Experimental/IoStore/OnDemand/Private/OnDemandHttpIoDispacher.cpp b/Engine/Source/Runtime/Experimental/IoStore/OnDemand/Private/OnDemandHttpIoDispacher.cpp
--- a/Engine/Source/Runtime/Experimental/IoStore/OnDemand/Private/OnDemandHttpIoDispacher.cpp
+++ b/Engine/Source/Runtime/Experimental/IoStore/OnDemand/Private/OnDemandHttpIoDispacher.cpp
@@ -497,6 +497,7 @@ bool FHttpDispatcher::Init()
 	{
 		.MaxConnectionCount		= GIasHttpConnectionCount,
 		.ReceiveBufferSize		= GIasHttpRecvBufKiB >= 0 ? GIasHttpRecvBufKiB << 10 : -1,
+		.SendBufferSize			= GIasHttpSendBufKiB >= 0 ? GIasHttpSendBufKiB << 10 : -1,
 		.MaxRetryCount			= GIasHttpRetryCount,
 		.TimeoutMs				= GIasHttpFailTimeOutMs,
 		.Redirects				= EHttpRedirects::Follow,
diff --git a/Engine/Source/Runtime/Experimental/IoStore/OnDemand/Private/OnDemandHttpThread.h b/Engine/Source/Runtime/Experimental/IoStore/OnDemand/Private/OnDemandHttpThread.h
--- a/Engine/Source/Runtime/Experimental/IoStore/OnDemand/Private/OnDemandHttpThread.h
+++ b/Engine/Source/Runtime/Experimental/IoStore/OnDemand/Private/OnDemandHttpThread.h
@@ -44,6 +44,7 @@ extern int32 GIasHttpHealthCheckWaitTime;
 extern int32 GOnDemandBackendThreadPriorityIndex;
 extern int32 GIasHttpConnectionCount;
 extern int32 GIasHttpRecvBufKiB;
+extern int32 GIasHttpSendBufKiB;
 extern int32 GIasHttpRetryCount;
 extern int32 GIasHttpFailTimeOutMs; 
 extern int32 GIasHttpConcurrentRequests; 
