Fix server crash in Task Dispatcher Evolution
CHeck was triggered when kinemnatic batches was not matching kinemtatics thread number

diff --git a/Engine/Source/Runtime/Experimental/Chaos/Public/Chaos/TaskDispatcherEvolution.h b/Engine/Source/Runtime/Experimental/Chaos/Public/Chaos/TaskDispatcherEvolution.h
--- a/Engine/Source/Runtime/Experimental/Chaos/Public/Chaos/TaskDispatcherEvolution.h
+++ b/Engine/Source/Runtime/Experimental/Chaos/Public/Chaos/TaskDispatcherEvolution.h
@@ -239,7 +239,7 @@ public:
 		}
 		int32 NumBatchKin = InternalAccelerationQueue.PendingDataArrays.Num() - InternalAccelerationQueue.KinematicBatchStartIndex;
 		// Prune Kinematics particles
-		if (NumBatchKin > 0 && InternalAccelerationQueue.KinematicBatchStartIndex > 0)
+		if (NumBatchKin > 0 && InternalAccelerationQueue.KinematicBatchStartIndex > 0 && KinematicPendingTasks.Num() == NumBatchKin)
 		{
 			check(KinematicPendingTasks.Num() == NumBatchKin);
 			NewKinPendingDataArray.SetNum(NumBatchKin);
@@ -274,6 +274,11 @@ public:
 			}, KinematicTasks, LowLevelTasks::ETaskPriority::High);
 			CurrentPendingTasks.Add(CopyKinematicTask);
 		}
+		else if (NumBatchKin > 0 && InternalAccelerationQueue.KinematicBatchStartIndex > 0 && KinematicPendingTasks.Num() != NumBatchKin)
+		{
+			UE_LOG(LogChaos, Warning, TEXT("No pruning happened in the Spatial Acceleration structure, a slight perfromance hit could occured. KinematicPendingTasks.Num(): %d !=  NumBatchKin:  %d"), KinematicPendingTasks.Num(), NumBatchKin);
+			CurrentPendingTasks.Append(KinematicPendingTasks);
+		}
 		UE::Tasks::Wait(CurrentPendingTasks);
 		IntegrationPendingTasks.Reset();
 		KinematicPendingTasks.Reset();
