Constraints: make sure that the task runs its evaluation function even if the evaluator is valid to cache the transform

#rnx

diff --git a/Engine/Source/Runtime/Experimental/Animation/Constraints/Private/Transform/AnimationEvaluation.cpp b/Engine/Source/Runtime/Experimental/Animation/Constraints/Private/Transform/AnimationEvaluation.cpp
--- a/Engine/Source/Runtime/Experimental/Animation/Constraints/Private/Transform/AnimationEvaluation.cpp
+++ b/Engine/Source/Runtime/Experimental/Animation/Constraints/Private/Transform/AnimationEvaluation.cpp
@@ -399,6 +399,10 @@ const FAnimationEvaluator& FAnimationEvaluationCache::GetEvaluator(USkeletalMesh
 			}
 			Found->Update(bRefreshBones);
 		}
+		else if (InTask.IsValid(InSkeletalMeshComponent))
+		{
+			InTask.PostEvaluationFunction();
+		}
 		return *Found;
 	}
 	
