Add a function to dump mounted IoStoreOnDemand containers to the log.

diff --git a/Engine/Source/Runtime/Experimental/IoStore/OnDemand/Private/OnDemandIoStore.cpp b/Engine/Source/Runtime/Experimental/IoStore/OnDemand/Private/OnDemandIoStore.cpp
--- a/Engine/Source/Runtime/Experimental/IoStore/OnDemand/Private/OnDemandIoStore.cpp
+++ b/Engine/Source/Runtime/Experimental/IoStore/OnDemand/Private/OnDemandIoStore.cpp
@@ -25,6 +25,7 @@
 #include "IO/IoStoreOnDemand.h"
 #include "Serialization/PackageStore.h"
 #include "IasCache.h"
+#include "Logging/StructuredLog.h"
 #include "Misc/CommandLine.h"
 #include "Misc/ConfigCacheIni.h"
 #include "Misc/CoreDelegates.h"
@@ -2124,7 +2125,7 @@ FOnDemandCacheUsage FOnDemandIoStore::GetCacheUsage(const FOnDemandGetCacheUsage
 				{
 					if (EnumHasAnyFlags(Args.Options, EOnDemandGetCacheUsageOptions::DumpHandlesToLog))
 					{
-						UE_LOG(LogIoStoreOnDemand, Log, TEXT("HandleId=0x%llX, DebugName='%s', ReferencedBytes=%.2lf KiB"),
+						UE_LOG(LogIoStoreOnDemand, Display, TEXT("HandleId=0x%llX, DebugName='%s', ReferencedBytes=%.2lf KiB"),
 							WeakHandle.HandleId, *WeakHandle.DebugName, double(ReferencedBytesByHandle) / 1024.0);
 					}
 
@@ -2152,6 +2153,68 @@ FOnDemandCacheUsage FOnDemandIoStore::GetCacheUsage(const FOnDemandGetCacheUsage
 	return CacheUsage;
 }
 
+void FOnDemandIoStore::DumpMountedContainersToLog() const
+{
+	TStringBuilder<128> FlagsSb;
+
+	TUniqueLock Lock(ContainerMutex);
+
+	UE_LOG(LogIoStoreOnDemand, Display, TEXT("Containers:"));
+	for (const FSharedOnDemandContainer& Container : Containers)
+	{
+		FlagsSb.Reset();
+		bool bFirst = true;
+
+		if (EnumHasAnyFlags(Container->Flags, EOnDemandContainerFlags::StreamOnDemand))
+		{
+			FlagsSb << TEXT("StreamOnDemand");
+			bFirst = false;
+		}
+		if (EnumHasAnyFlags(Container->Flags, EOnDemandContainerFlags::InstallOnDemand))
+		{
+			if (!bFirst)
+			{
+				FlagsSb << TEXT(", ");
+			}
+
+			FlagsSb << TEXT("InstallOnDemand");
+			bFirst = false;
+		}
+		if (EnumHasAnyFlags(Container->Flags, EOnDemandContainerFlags::Mounted))
+		{
+			if (!bFirst)
+			{
+				FlagsSb << TEXT(", ");
+			}
+
+			FlagsSb << TEXT("Mounted");
+			bFirst = false;
+		}
+		if (EnumHasAnyFlags(Container->Flags, EOnDemandContainerFlags::PendingEncryptionKey))
+		{
+			if (!bFirst)
+			{
+				FlagsSb << TEXT(", ");
+			}
+
+			FlagsSb << TEXT("PendingEncryptionKey");
+			bFirst = false;
+		}
+		if (EnumHasAnyFlags(Container->Flags, EOnDemandContainerFlags::PendingHostGroup))
+		{
+			if (!bFirst)
+			{
+				FlagsSb << TEXT(", ");
+			}
+
+			FlagsSb << TEXT("PendingHostGroup");
+			bFirst = false;
+		}
+
+		UE_LOGFMT(LogIoStoreOnDemand, Display, "\t Container: {Name}, Flags: {Flags}", *Container->UniqueName(), FlagsSb);
+	}
+}
+
 bool FOnDemandIoStore::IsOnDemandStreamingEnabled() const
 {
 	return StreamingBackend != nullptr;
diff --git a/Engine/Source/Runtime/Experimental/IoStore/OnDemand/Private/OnDemandIoStore.h b/Engine/Source/Runtime/Experimental/IoStore/OnDemand/Private/OnDemandIoStore.h
--- a/Engine/Source/Runtime/Experimental/IoStore/OnDemand/Private/OnDemandIoStore.h
+++ b/Engine/Source/Runtime/Experimental/IoStore/OnDemand/Private/OnDemandIoStore.h
@@ -373,6 +373,7 @@ public:
 	virtual TIoStatusOr<FOnDemandInstallSizeResult> GetInstallSize(const FOnDemandGetInstallSizeArgs& Args) const override;
 	virtual FIoStatus					GetInstallSizesByMountId(const FOnDemandGetInstallSizeArgs& Args, TMap<FString, uint64>& OutSizesByMountId) const override;
 	virtual FOnDemandCacheUsage			GetCacheUsage(const FOnDemandGetCacheUsageArgs& Args) const override;
+	virtual void						DumpMountedContainersToLog() const override;
 	virtual bool						IsOnDemandStreamingEnabled() const override;
 	virtual void						SetStreamingOptions(EOnDemandStreamingOptions Options) override;
 	virtual void						ReportAnalytics(TArray<FAnalyticsEventAttribute>& OutAnalyticsArray) const override;
diff --git a/Engine/Source/Runtime/Experimental/IoStore/OnDemandCore/Private/IoStoreOnDemand.cpp b/Engine/Source/Runtime/Experimental/IoStore/OnDemandCore/Private/IoStoreOnDemand.cpp
--- a/Engine/Source/Runtime/Experimental/IoStore/OnDemandCore/Private/IoStoreOnDemand.cpp
+++ b/Engine/Source/Runtime/Experimental/IoStore/OnDemandCore/Private/IoStoreOnDemand.cpp
@@ -155,6 +155,8 @@ static void PrintCacheUsage()
 		TStringBuilder<512> Sb;
 		Sb << TEXTVIEW("InstallCache: ") << CacheUsage.InstallCache << TEXTVIEW(", StreamingCache: ") << CacheUsage.StreamingCache;
 		UE_LOG(LogIoStoreOnDemand, Display, TEXT("%s"), Sb.ToString());
+
+		IoStore->DumpMountedContainersToLog();
 	}
 	else
 	{
diff --git a/Engine/Source/Runtime/Experimental/IoStore/OnDemandCore/Public/IO/IoStoreOnDemand.h b/Engine/Source/Runtime/Experimental/IoStore/OnDemandCore/Public/IO/IoStoreOnDemand.h
--- a/Engine/Source/Runtime/Experimental/IoStore/OnDemandCore/Public/IO/IoStoreOnDemand.h
+++ b/Engine/Source/Runtime/Experimental/IoStore/OnDemandCore/Public/IO/IoStoreOnDemand.h
@@ -628,6 +628,8 @@ public:
 	virtual FIoStatus GetInstallSizesByMountId(const FOnDemandGetInstallSizeArgs& Args, TMap<FString, uint64>& OutSizesByMountId) const = 0;
 	/** Returns the total cache size in bytes. */
 	virtual FOnDemandCacheUsage GetCacheUsage(const FOnDemandGetCacheUsageArgs& Args) const = 0;
+	/** Dump information about mounted containers to the log */
+	virtual void DumpMountedContainersToLog() const = 0;
 	/** Return if the OnDemand streaming system is enabled. */
 	virtual bool IsOnDemandStreamingEnabled() const = 0;
 	/** Set streaming options.*/
