Dataflow : fix issue with drag and dropping generic arrays on an subgraph input/output node

diff --git a/Engine/Source/Runtime/Experimental/Dataflow/Engine/Private/Dataflow/DataflowEngineUtil.cpp b/Engine/Source/Runtime/Experimental/Dataflow/Engine/Private/Dataflow/DataflowEngineUtil.cpp
--- a/Engine/Source/Runtime/Experimental/Dataflow/Engine/Private/Dataflow/DataflowEngineUtil.cpp
+++ b/Engine/Source/Runtime/Experimental/Dataflow/Engine/Private/Dataflow/DataflowEngineUtil.cpp
@@ -110,9 +110,23 @@ namespace UE::Dataflow
 			else
 			{
 				const FString TypeStr(Type.ToString());
+
+				// Make sure to remove the TObjectPtr
+				FString ObjectPtrInnerType;
+				if (FDataflowUObjectConvertibleTypePolicy::GetObjectPtrInnerType(TypeStr, ObjectPtrInnerType))
+				{
+					if (const UObject* Object = FindFirstObjectSafe<UObject>(*ObjectPtrInnerType))
+					{
+						return FPropertyBagPropertyDesc(Name, EPropertyBagPropertyType::Object, Object);
+					}
+				}
 				if (TypeStr.StartsWith("U"))
 				{
-					// Nothing for now and use thefallback - may need fixing if we find there's an issue with that code path 
+					const FString ShortTypeName = TypeStr.RightChop(1);
+					if (const UObject* Object = FindFirstObjectSafe<UObject>(*ShortTypeName))
+					{
+						return FPropertyBagPropertyDesc(Name, EPropertyBagPropertyType::Object, Object);
+					}
 				}
 				else if (TypeStr.StartsWith("F"))
 				{
