Fixed game thread assertion in dataflow

diff --git a/Engine/Source/Runtime/Experimental/Dataflow/Engine/Private/Dataflow/DataflowEdNode.cpp b/Engine/Source/Runtime/Experimental/Dataflow/Engine/Private/Dataflow/DataflowEdNode.cpp
--- a/Engine/Source/Runtime/Experimental/Dataflow/Engine/Private/Dataflow/DataflowEdNode.cpp
+++ b/Engine/Source/Runtime/Experimental/Dataflow/Engine/Private/Dataflow/DataflowEdNode.cpp
@@ -1204,12 +1204,15 @@ void UDataflowEdNode::RegisterDelegateHandle()
 			{
 				if (UEdGraph* EdGraph = GetGraph())
 				{
-					GetGraph()->NotifyNodeChanged(this);
-					if (UDataflow* DataflowAsset = UDataflow::GetDataflowAssetFromEdGraph(EdGraph))
+					if (IsInGameThread() || IsInSlateThread())
 					{
-						if (InDataflowNode)
+						GetGraph()->NotifyNodeChanged(this);
+						if (UDataflow* DataflowAsset = UDataflow::GetDataflowAssetFromEdGraph(EdGraph))
 						{
-							FDataflowAssetDelegates::OnNodeInvalidated.Broadcast(*DataflowAsset, *InDataflowNode);
+							if (InDataflowNode)
+							{
+								FDataflowAssetDelegates::OnNodeInvalidated.Broadcast(*DataflowAsset, *InDataflowNode);
+							}
 						}
 					}
 				}
