The new IAX backend now respects hotfixes for it's http client.

#rnx

- The http client cannot have it's values changed on the fly, they only are taken into account when it is created.
- FHttpDispatcher hooks into the OnRecreateHttpClient which is broadcast when ever one of the cvars affecting the http client changes and if so we make a note of it.
- At the start of each ::ProcessHttpRequests loop we check if any of the cvars have changed and if so we try to create a new http client. If that succeeds we destroy the old one and swap the new one in that is running with the new values.
-- Note that we wait until the start of each ::ProcessHttpRequests to make sure that the client is not processing any requests as it is destroyed.
- Moved the logic for poping new requests to be processed to it's own method ::GetNextRequest
-- This method will stop returning new requests if we know that we want to recreate the client. This will let us get to it sooner.
- Cleaned up the final loop of ::ProcessHttpRequests where we are waiting for all requests to complete or for new work to be made.
- Add an assert to the FMultiEndpointHttpClient destructor to make sure that it's event loop is not processing anything before it is also destroyed.

diff --git a/Engine/Source/Runtime/Experimental/IoStore/OnDemand/Private/OnDemandHttpClient.cpp b/Engine/Source/Runtime/Experimental/IoStore/OnDemand/Private/OnDemandHttpClient.cpp
--- a/Engine/Source/Runtime/Experimental/IoStore/OnDemand/Private/OnDemandHttpClient.cpp
+++ b/Engine/Source/Runtime/Experimental/IoStore/OnDemand/Private/OnDemandHttpClient.cpp
@@ -174,6 +174,7 @@ FMultiEndpointHttpClient::FMultiEndpointHttpClient(const FMultiEndpointHttpClien
 
 FMultiEndpointHttpClient::~FMultiEndpointHttpClient()
 {
+	checkf(EventLoop.IsIdle(), TEXT("FMultiEndpointHttpClient still has active requests on shutdown"));
 	check(TicketLookupMap.IsEmpty());
 }
 
diff --git a/Engine/Source/Runtime/Experimental/IoStore/OnDemand/Private/OnDemandHttpIoDispacher.cpp b/Engine/Source/Runtime/Experimental/IoStore/OnDemand/Private/OnDemandHttpIoDispacher.cpp
--- a/Engine/Source/Runtime/Experimental/IoStore/OnDemand/Private/OnDemandHttpIoDispacher.cpp
+++ b/Engine/Source/Runtime/Experimental/IoStore/OnDemand/Private/OnDemandHttpIoDispacher.cpp
@@ -36,6 +36,8 @@ namespace UE::IoStore
 {
 ////////////////////////////////////////////////////////////////////////////////
 
+DECLARE_MULTICAST_DELEGATE(FOnRecreateHttpClient);
+extern FOnRecreateHttpClient OnRecreateHttpClient;
 
 extern int32 GIasHttpTimeOutMs;
 extern int32 GOnDemandBackendThreadPriorityIndex; 
@@ -445,8 +447,12 @@ private:
 	virtual FIoStatus				CacheResponse(const FIoHttpResponse& Response) override;
 	virtual FIoStatus				EvictFromCache(const FIoHttpResponse& Response) override;
 
+	bool							TryCreateHttpClient();
+	void							RecreateHttpClientIfNeeded();
+
 	bool							TryReadFromCache(FHttpRequest* Request);
 	bool							ProcessHttpRequests();
+	FHttpRequest*					GetNextRequest();
 	void							CompleteHttpRequest(FHttpRequest& Request, FMultiEndpointHttpClientResponse&& HttpResponse, FIASHostGroup& HostGroup, const FIoHttpRange& TotalRange);
 	void							CompleteRequest(FHttpRequest& Request, FIoBuffer& Body, const FIoHttpRange& TotalRange, uint32 StatusCode, bool bCached);
 
@@ -457,7 +463,9 @@ private:
 	FHostGroupRegistered					HostGroupRegisterd;
 	FEventRef								WakeUp;
 	FMutex									Mutex;
-	std::atomic_bool						bStopRequested{false};
+	FDelegateHandle							OnRecreateHttpClientHandle;
+	std::atomic_bool						bRecreateHttpClient = false;
+	std::atomic_bool						bStopRequested = false;
 };
 
 static TUniquePtr<FHttpDispatcher> HttpDispatcher;
@@ -488,25 +496,24 @@ FHttpDispatcher::FHttpDispatcher(TUniquePtr<IIasCache>&& InCache)
 
 FHttpDispatcher::~FHttpDispatcher()
 {
+	if (OnRecreateHttpClientHandle.IsValid())
+	{
+		OnRecreateHttpClient.Remove(OnRecreateHttpClientHandle);
+		OnRecreateHttpClientHandle.Reset();
+	}
+
 	Thread.Reset();
 }
 
 bool FHttpDispatcher::Init()
 {
-	HttpClient = FMultiEndpointHttpClient::Create(FMultiEndpointHttpClientConfig
-	{
-		.MaxConnectionCount		= GIasHttpConnectionCount,
-		.ReceiveBufferSize		= GIasHttpRecvBufKiB >= 0 ? GIasHttpRecvBufKiB << 10 : -1,
-		.SendBufferSize			= GIasHttpSendBufKiB >= 0 ? GIasHttpSendBufKiB << 10 : -1,
-		.MaxRetryCount			= GIasHttpRetryCount,
-		.TimeoutMs				= GIasHttpFailTimeOutMs,
-		.Redirects				= EHttpRedirects::Follow,
-		.bAllowChunkedTransfer	= true,
-		.LogCategory			= &LogHttpIoDispatcher,
-		.LogVerbosity			= ELogVerbosity::VeryVerbose
-	});
+	OnRecreateHttpClientHandle = OnRecreateHttpClient.AddLambda([this]()
+		{
+			bRecreateHttpClient = true;
+			WakeUp->Trigger();
+		});
 
-	return HttpClient.IsValid();
+	return TryCreateHttpClient();
 }
 
 uint32 FHttpDispatcher::Run()
@@ -548,6 +555,7 @@ void FHttpDispatcher::Stop()
 
 void FHttpDispatcher::Exit()
 {
+	
 }
 
 void FHttpDispatcher::Tick()
@@ -807,6 +815,46 @@ FIoStatus FHttpDispatcher::CacheResponse(const FIoHttpResponse& Response)
 	return Cache->Put(Response.GetCacheKey(), Copy);
 }
 
+bool FHttpDispatcher::TryCreateHttpClient()
+{
+	HttpClient = FMultiEndpointHttpClient::Create(FMultiEndpointHttpClientConfig
+		{
+			.MaxConnectionCount = GIasHttpConnectionCount,
+			.ReceiveBufferSize = GIasHttpRecvBufKiB >= 0 ? GIasHttpRecvBufKiB << 10 : -1,
+			.SendBufferSize = GIasHttpSendBufKiB >= 0 ? GIasHttpSendBufKiB << 10 : -1,
+			.MaxRetryCount = GIasHttpRetryCount,
+			.TimeoutMs = GIasHttpFailTimeOutMs,
+			.Redirects = EHttpRedirects::Follow,
+			.bAllowChunkedTransfer = true,
+			.LogCategory = &LogHttpIoDispatcher,
+			.LogVerbosity = ELogVerbosity::VeryVerbose
+		});
+
+	return HttpClient.IsValid();
+}
+
+void FHttpDispatcher::RecreateHttpClientIfNeeded()
+{
+	if (bRecreateHttpClient.exchange(false))
+	{
+		TRACE_CPUPROFILER_EVENT_SCOPE(FHttpDispatcher::RecreateHttpClientIfNeeded);
+
+		// Keep the old client at first in case we fail to create the new one so that we can restore it
+		// this way the game can keep on functioning correctly.
+		TUniquePtr<FMultiEndpointHttpClient> OldHttpClient = MoveTemp(HttpClient);
+		if (TryCreateHttpClient())
+		{
+			OldHttpClient.Reset();
+			UE_LOGFMT(LogIas, Display, "FHttpDispatcher: Successfully created a new http client for use");
+		}
+		else
+		{
+			HttpClient = MoveTemp(OldHttpClient);
+			UE_LOGFMT(LogIas, Warning, "FHttpDispatcher: Failed to create a new http client, the existing client will continue to be used");
+		}
+	}
+}
+
 FIoStatus FHttpDispatcher::EvictFromCache(const FIoHttpResponse& Response)
 {
 	// If the cache is disabled we consider the eviction to be a success, otherwise the calling code will
@@ -904,6 +952,10 @@ void FHttpDispatcher::ReleaseRequest(FIoHttpRequestHandle Handle)
 
 bool FHttpDispatcher::ProcessHttpRequests()
 {
+	TRACE_CPUPROFILER_EVENT_SCOPE(FHttpDispatcher::ProcessHttpRequests);
+
+	RecreateHttpClientIfNeeded();
+
 	const int32 MaxConcurrentRequests = FMath::Min(GIasHttpConcurrentRequests, 32);
 	int32		NumConcurrentRequests = 0;
 
@@ -921,7 +973,7 @@ bool FHttpDispatcher::ProcessHttpRequests()
 		while (NextHttpRequest)
 		{
 			FHttpRequest* HttpRequest	= NextHttpRequest;
-			NextHttpRequest				= HttpQueue.Dequeue();
+			NextHttpRequest				= GetNextRequest();
 
 			ensure(!HttpRequest->Flags.HasAny(EHttpRequestFlags::HttpQueued | EHttpRequestFlags::CacheInflight));
 
@@ -988,7 +1040,8 @@ bool FHttpDispatcher::ProcessHttpRequests()
 
 			if (NumConcurrentRequests >= MaxConcurrentRequests)
 			{
-				TRACE_CPUPROFILER_EVENT_SCOPE(FOnDemandHttpThread::TickHttpSaturated);
+				TRACE_CPUPROFILER_EVENT_SCOPE(FHttpDispatcher::TickHttpSaturated);
+
 				while (NumConcurrentRequests >= MaxConcurrentRequests) //-V654
 				{
 					HttpClient->Tick(MAX_uint32, GIasHttpRateLimitKiBPerSecond); 
@@ -997,24 +1050,18 @@ bool FHttpDispatcher::ProcessHttpRequests()
 
 			if (!NextHttpRequest)
 			{
-				NextHttpRequest = HttpQueue.Dequeue();
+				NextHttpRequest = GetNextRequest();
 			}
 		} // Inner
-		
+
 		{
+			check(NextHttpRequest == nullptr);
+
 			// Keep processing pending connections until all requests are completed or a new one is issued
-			TRACE_CPUPROFILER_EVENT_SCOPE(FOnDemandHttpThread::TickHttp);
-			while (HttpClient->Tick(GIasHttpPollTimeoutMs, GIasHttpRateLimitKiBPerSecond))
+			TRACE_CPUPROFILER_EVENT_SCOPE(FHttpDispatcher::TickHttp);
+			while (NextHttpRequest == nullptr && HttpClient->Tick(GIasHttpPollTimeoutMs, GIasHttpRateLimitKiBPerSecond))
 			{
-				if (!NextHttpRequest)
-				{
-					NextHttpRequest = HttpQueue.Dequeue();
-				}
-
-				if (NextHttpRequest)
-				{
-					break;
-				}
+				NextHttpRequest = GetNextRequest();
 			}
 		}
 	} // Outer
@@ -1025,6 +1072,11 @@ bool FHttpDispatcher::ProcessHttpRequests()
 	return true;
 }
 
+FHttpRequest* FHttpDispatcher::GetNextRequest()
+{
+	return bRecreateHttpClient == false ? HttpQueue.Dequeue() : nullptr;
+}
+
 void FHttpDispatcher::CompleteHttpRequest(FHttpRequest& Request, FMultiEndpointHttpClientResponse&& HttpResponse, FIASHostGroup& HostGroup, const FIoHttpRange& TotalRange)
 {
 	if (HttpResponse.RetryCount > 0)
