CIS: Fix code static analysis warning

#rnx

- This code is  functionally equivalent to the previous version but should silence what I believe to be a false positive with a static analysis tool.

diff --git a/Engine/Source/Runtime/Experimental/IoStore/OnDemand/Private/OnDemandHttpIoDispacher.cpp b/Engine/Source/Runtime/Experimental/IoStore/OnDemand/Private/OnDemandHttpIoDispacher.cpp
--- a/Engine/Source/Runtime/Experimental/IoStore/OnDemand/Private/OnDemandHttpIoDispacher.cpp
+++ b/Engine/Source/Runtime/Experimental/IoStore/OnDemand/Private/OnDemandHttpIoDispacher.cpp
@@ -1074,7 +1074,14 @@ bool FHttpDispatcher::ProcessHttpRequests()
 
 FHttpRequest* FHttpDispatcher::GetNextRequest()
 {
-	return bRecreateHttpClient == false ? HttpQueue.Dequeue() : nullptr;
+	if (!bRecreateHttpClient)
+	{
+		return HttpQueue.Dequeue();
+	}
+	else
+	{
+		return nullptr;
+	}
 }
 
 void FHttpDispatcher::CompleteHttpRequest(FHttpRequest& Request, FMultiEndpointHttpClientResponse&& HttpResponse, FIASHostGroup& HostGroup, const FIoHttpRange& TotalRange)
