Protect the new IAX backend against invalid values from the 'ias.HttpConnectionCount' cvar.

#rnx

- FEventLoop will cast the max concurrent connections to a uint16 so that is the max value we can support, although it is not advisable to have that many.
- FEventLoop will assert/crash on 0 connections so we must have at least 1.

diff --git a/Engine/Source/Runtime/Experimental/IoStore/OnDemand/Private/OnDemandHttpIoDispacher.cpp b/Engine/Source/Runtime/Experimental/IoStore/OnDemand/Private/OnDemandHttpIoDispacher.cpp
--- a/Engine/Source/Runtime/Experimental/IoStore/OnDemand/Private/OnDemandHttpIoDispacher.cpp
+++ b/Engine/Source/Runtime/Experimental/IoStore/OnDemand/Private/OnDemandHttpIoDispacher.cpp
@@ -817,9 +817,13 @@ FIoStatus FHttpDispatcher::CacheResponse(const FIoHttpResponse& Response)
 
 bool FHttpDispatcher::TryCreateHttpClient()
 {
+	// Make sure that 'ias.HttpConnectionCount' is within a range that the client will accept
+	const int32 MaxNumberOfConnections = FMath::Clamp(GIasHttpConnectionCount, 1, (int32)MAX_uint16);
+	UE_CLOG(MaxNumberOfConnections != GIasHttpConnectionCount, LogHttpIoDispatcher, Error, TEXT("ias.HttpConnectionCount (%d) outside of valid range 1-%d"), GIasHttpConnectionCount, MAX_uint16);
+
 	HttpClient = FMultiEndpointHttpClient::Create(FMultiEndpointHttpClientConfig
 		{
-			.MaxConnectionCount = GIasHttpConnectionCount,
+			.MaxConnectionCount = MaxNumberOfConnections,
 			.ReceiveBufferSize = GIasHttpRecvBufKiB >= 0 ? GIasHttpRecvBufKiB << 10 : -1,
 			.SendBufferSize = GIasHttpSendBufKiB >= 0 ? GIasHttpSendBufKiB << 10 : -1,
 			.MaxRetryCount = GIasHttpRetryCount,
