The new IAX backend now correctly supports the FSingleThreadRunnable interface.

#rnx

- FHttpDispatcher was inheriting from FSingleThreadRunnable but did not implement FRunnable::GetSingleThreadInterface meaning that it would never be considered when creating the thread.
- FHttpDispatcher::Run and FHttpDispatcher::Tick were pretty much the same code so now ::Run just calls ::Tick
-- ::Tick no longer loops until ::ProcessHttpRequests returns false. In all code paths ::ProcessHttpRequests should only return when all requests are processed so there is no real gain in doing multiple loops when running in singlethreaded mode.
-- Note that ::ProcessHttpRequests is probably processing far too much when dealing with a large number of requests in single threaded mode. At some point in the future we should limit how many requests are processed.

diff --git a/Engine/Source/Runtime/Experimental/IoStore/OnDemand/Private/OnDemandHttpIoDispacher.cpp b/Engine/Source/Runtime/Experimental/IoStore/OnDemand/Private/OnDemandHttpIoDispacher.cpp
--- a/Engine/Source/Runtime/Experimental/IoStore/OnDemand/Private/OnDemandHttpIoDispacher.cpp
+++ b/Engine/Source/Runtime/Experimental/IoStore/OnDemand/Private/OnDemandHttpIoDispacher.cpp
@@ -421,6 +421,11 @@ private:
 	virtual uint32					Run() override;
 	virtual void					Stop() override;
 	virtual void					Exit() override;
+	virtual FSingleThreadRunnable*	GetSingleThreadInterface() override
+	{
+		return this;
+	}
+
 	//FSingleThreadRunnable
 	virtual void					Tick() override;
 	// IHttpIoDispatcher
@@ -522,16 +527,7 @@ uint32 FHttpDispatcher::Run()
 
 	for (;;)
 	{
-		FHostGroupManager::Get().Tick(GIasHttpTimeOutMs, bStopRequested);
-
-		// TODO: It would be better to only update connections as they need it, consider doing this on
-		// hostgroup connect/disconnect events.
-		HttpClient->UpdateConnections();
-
-		if (ProcessHttpRequests())
-		{
-			continue;
-		}
+		Tick();
 
 		if (bStopRequested.load(std::memory_order_relaxed))
 		{
@@ -560,19 +556,13 @@ void FHttpDispatcher::Exit()
 
 void FHttpDispatcher::Tick()
 {
-	for (;;)
-	{
-		FHostGroupManager::Get().Tick(GIasHttpTimeOutMs, bStopRequested);
+	FHostGroupManager::Get().Tick(GIasHttpTimeOutMs, bStopRequested);
 
-		// TODO: It would be better to only update connections as they need it, consider doing this on
-		// hostgroup connect/disconnect events.
-		HttpClient->UpdateConnections();
+	// TODO: It would be better to only update connections as they need it, consider doing this on
+	// hostgroup connect/disconnect events.
+	HttpClient->UpdateConnections();
 
-		if (ProcessHttpRequests() == false)
-		{
-			break;
-		}
-	}
+	ProcessHttpRequests();
 }
 
 void FHttpDispatcher::Shutdown()
