The new IAX backend now respects the 'ias.onDemandBackendThreadPriority' cvar for setting thread priority.

#rnx

- When the thread is created, the prriority from the cvar was accidently being ignored and we were always applying EThreadPriority::TPri_Normal.
-Added new method ::UpdateThreadPriorityIfNeeded which is called once per thread loop.
- This method checks if the thread priority has changed and if so will attempt to apply the new priority to the thread, this will allow the priority to be hotfixed in the future.

diff --git a/Engine/Source/Runtime/Experimental/IoStore/OnDemand/Private/OnDemandHttpIoDispacher.cpp b/Engine/Source/Runtime/Experimental/IoStore/OnDemand/Private/OnDemandHttpIoDispacher.cpp
--- a/Engine/Source/Runtime/Experimental/IoStore/OnDemand/Private/OnDemandHttpIoDispacher.cpp
+++ b/Engine/Source/Runtime/Experimental/IoStore/OnDemand/Private/OnDemandHttpIoDispacher.cpp
@@ -32,6 +32,9 @@
 
 #define UE_ALLOW_CACHE_POISONING !UE_BUILD_SHIPPING
 
+// TODO - Maybe do LexToString/LexFromString (dupe code in UnrealEngine.cpp ThreadPriorityToString)
+
+extern const TCHAR* LexToString(EThreadPriority Priority);
 namespace UE::IoStore
 {
 ////////////////////////////////////////////////////////////////////////////////
@@ -46,6 +49,16 @@ extern float GIaxInvalidUrlChance;
 extern bool GIasPoisonCache;
 #endif // UE_ALLOW_CACHE_POISONING
 
+/** Matches the priority set up found in Engine/Source/Runtime/Online/HTTP/Private/HttpThread.cpp */
+static constexpr int32 GThreadPriorities[5] =
+{
+	EThreadPriority::TPri_Lowest,
+	EThreadPriority::TPri_BelowNormal,
+	EThreadPriority::TPri_SlightlyBelowNormal,
+	EThreadPriority::TPri_Normal,
+	EThreadPriority::TPri_AboveNormal
+};
+
 bool GHttpCacheEnabled = true;
 static FAutoConsoleVariableRef CVar_HttpCacheEnabled(
 	TEXT("iax.HttpCacheEnabled"),
@@ -452,6 +465,8 @@ private:
 	virtual FIoStatus				CacheResponse(const FIoHttpResponse& Response) override;
 	virtual FIoStatus				EvictFromCache(const FIoHttpResponse& Response) override;
 
+	void							UpdateThreadPriorityIfNeeded();
+
 	bool							TryCreateHttpClient();
 	void							RecreateHttpClientIfNeeded();
 
@@ -464,6 +479,7 @@ private:
 	TUniquePtr<IIasCache>					Cache;
 	TUniquePtr<FMultiEndpointHttpClient>	HttpClient;
 	TUniquePtr<FRunnableThread>				Thread;
+	EThreadPriority							ThreadPriority = EThreadPriority::TPri_Normal;
 	FHttpQueue								HttpQueue;
 	FHostGroupRegistered					HostGroupRegisterd;
 	FEventRef								WakeUp;
@@ -479,18 +495,10 @@ static TUniquePtr<FHttpDispatcher> HttpDispatcher;
 FHttpDispatcher::FHttpDispatcher(TUniquePtr<IIasCache>&& InCache)
 	: Cache(MoveTemp(InCache))
 {
-	static constexpr int32 ThreadPriorities[5] =
-	{
-		EThreadPriority::TPri_Lowest,
-		EThreadPriority::TPri_BelowNormal,
-		EThreadPriority::TPri_SlightlyBelowNormal,
-		EThreadPriority::TPri_Normal,
-		EThreadPriority::TPri_AboveNormal
-	};
-
-	const int32 ThreadPriorityIndex			= FMath::Clamp(GOnDemandBackendThreadPriorityIndex, 0, (int32)UE_ARRAY_COUNT(ThreadPriorities) - 1);
-	EThreadPriority DesiredThreadPriority	= (EThreadPriority)ThreadPriorities[ThreadPriorityIndex];
-	EThreadPriority ThreadPriority			= EThreadPriority::TPri_Normal;
+	const int32 ThreadPriorityIndex			= FMath::Clamp(GOnDemandBackendThreadPriorityIndex, 0, (int32)UE_ARRAY_COUNT(GThreadPriorities) - 1);
+	EThreadPriority DesiredThreadPriority	= (EThreadPriority)GThreadPriorities[ThreadPriorityIndex];
+	ThreadPriority							= DesiredThreadPriority;
+
 	const uint32 StackSize					= 0; // Use default stack size
 	const uint64 ThreadAffinityMask			= FGenericPlatformAffinity::GetNoAffinityMask();
 	const EThreadCreateFlags CreateFlags	= EThreadCreateFlags::None;
@@ -527,6 +535,8 @@ uint32 FHttpDispatcher::Run()
 
 	for (;;)
 	{
+		UpdateThreadPriorityIfNeeded();
+
 		Tick();
 
 		if (bStopRequested.load(std::memory_order_relaxed))
@@ -805,6 +815,20 @@ FIoStatus FHttpDispatcher::CacheResponse(const FIoHttpResponse& Response)
 	return Cache->Put(Response.GetCacheKey(), Copy);
 }
 
+void FHttpDispatcher::UpdateThreadPriorityIfNeeded()
+{
+	// TODO - We should wake up the thread when 'ias.onDemandBackendThreadPriority' is updated
+	int32 ThreadPriorityIndex = FMath::Clamp(GOnDemandBackendThreadPriorityIndex, 0, (int32)UE_ARRAY_COUNT(GThreadPriorities) - 1);
+	EThreadPriority DesiredThreadPriority = (EThreadPriority)GThreadPriorities[ThreadPriorityIndex];
+	if (DesiredThreadPriority != ThreadPriority)
+	{
+		UE_LOGFMT(LogHttpIoDispatcher, Log, "Updated IoService.Http thread priority to '{Priority}'", LexToString(DesiredThreadPriority));
+
+		FPlatformProcess::SetThreadPriority(DesiredThreadPriority);
+		ThreadPriority = DesiredThreadPriority;
+	}
+}
+
 bool FHttpDispatcher::TryCreateHttpClient()
 {
 	// Make sure that 'ias.HttpConnectionCount' is within a range that the client will accept
diff --git a/Engine/Source/Runtime/Experimental/IoStore/OnDemand/Private/OnDemandHttpThread.cpp b/Engine/Source/Runtime/Experimental/IoStore/OnDemand/Private/OnDemandHttpThread.cpp
--- a/Engine/Source/Runtime/Experimental/IoStore/OnDemand/Private/OnDemandHttpThread.cpp
+++ b/Engine/Source/Runtime/Experimental/IoStore/OnDemand/Private/OnDemandHttpThread.cpp
@@ -23,7 +23,7 @@
 #define UE_ALLOW_INVALID_URL_DEBUGGING !UE_BUILD_SHIPPING
 
 // TODO - Maybe do LexToString/LexFromString (dupe code in UnrealEngine.cpp ThreadPriorityToString)
-static const TCHAR* LexToString(EThreadPriority Priority)
+const TCHAR* LexToString(EThreadPriority Priority)
 {
 	switch (Priority)
 	{
