Move out FGovernor from the IAS cache to be shared by the IAD cache.

#rnx

diff --git a/Engine/Source/Runtime/Experimental/IoStore/OnDemand/Private/DiskCacheGovernor.cpp b/Engine/Source/Runtime/Experimental/IoStore/OnDemand/Private/DiskCacheGovernor.cpp
new file mode 100644
--- /dev/null
+++ b/Engine/Source/Runtime/Experimental/IoStore/OnDemand/Private/DiskCacheGovernor.cpp
@@ -0,0 +1,144 @@
+// Copyright Epic Games, Inc. All Rights Reserved.
+
+#include "DiskCacheGovernor.h"
+#include "HAL/PlatformTime.h"
+
+namespace UE::IoStore
+{
+
+////////////////////////////////////////////////////////////////////////////////
+FDiskCacheGovernor::FDiskCacheGovernor()
+{
+	FConfig Cfg;
+	ConfigureStreamingCache(Cfg);
+}
+
+////////////////////////////////////////////////////////////////////////////////
+void FDiskCacheGovernor::ConfigureStreamingCache(const FDiskCacheGovernor::FConfig& Config)
+{
+	Configure(StreamingCacheState, Config);
+}
+
+////////////////////////////////////////////////////////////////////////////////
+int32 FDiskCacheGovernor::TryBeginWriteToStreamingCache(uint32 DemandPercent)
+{
+	return TryBeginWrite(StreamingCacheState, DemandPercent);
+}
+
+////////////////////////////////////////////////////////////////////////////////
+int32 FDiskCacheGovernor::EndWriteToStreamingCache(uint32 UnusedAllowance)
+{
+	return EndWrite(StreamingCacheState, UnusedAllowance);
+}
+
+////////////////////////////////////////////////////////////////////////////////
+void FDiskCacheGovernor::Configure(FDiskCacheGovernor::FCacheState& CacheState, const FDiskCacheGovernor::FConfig& Config)
+{
+	const int64 CycleFreq		= int64(1.0 / FPlatformTime::GetSecondsPerCycle64());
+	const FAllowance& Allowance	= Config.Allowance;
+	CacheState.Demand			= Config.Demand;
+
+	if (Config.CommitBufferSize == 0)
+	{
+		CacheState.OpAllowance	= Allowance.Bytes / Allowance.Ops;
+		CacheState.OpInterval	= (CycleFreq * Allowance.Seconds) / Allowance.Ops;
+		CacheState.MaxOpCount	= 4;
+		return;
+	}
+
+	// A small commit buffer will make us op bound
+	if (uint32 CommitWidth = Config.CommitBufferSize * Allowance.Ops; CommitWidth < Allowance.Bytes)
+	{
+		CacheState.OpAllowance	= Config.CommitBufferSize;
+		CacheState.OpInterval	= (Allowance.Seconds * CycleFreq * 3) / Allowance.Ops;
+		CacheState.MaxOpCount	= 1;
+		return;
+	}
+
+	// Allowance bound.
+	const int32 BlockCount		= FMath::Max<int32>(1, Allowance.Bytes / Config.CommitBufferSize);
+	const int32 CommitOpCost	= BlockCount * 3;
+	CacheState.MaxOpCount		= (Allowance.Ops - CommitOpCost) / BlockCount;
+	CacheState.OpAllowance		= Config.CommitBufferSize / CacheState.MaxOpCount;
+	CacheState.OpInterval		= (Allowance.Seconds * CycleFreq) / (BlockCount * (CacheState.MaxOpCount - 1));
+	CacheState.State			= FCacheState::EState::Waiting;
+}
+
+////////////////////////////////////////////////////////////////////////////////
+int32 FDiskCacheGovernor::TryBeginWrite(FDiskCacheGovernor::FCacheState& CacheState, uint32 DemandPercent, int64 Cycles)
+{
+	int64 Interval = CacheState.OpInterval;
+	Interval >>= int32(DemandPercent >= CacheState.Demand.Boost);
+	Interval >>= int32(DemandPercent >= CacheState.Demand.SuperBoost);
+	Interval <<= int32(DemandPercent <= CacheState.Demand.Threshold);
+
+	const int64 Delta = Cycles - CacheState.PrevCycles;
+	bool bNotYet = (Delta < Interval);
+
+	// Calculate how much time we are into the shortest poll interval
+	int64 Remainder = Delta;
+	Interval = GetMaxWaitCycles(CacheState);
+	for (; Remainder > Interval; Remainder -= Interval);
+
+	if (bNotYet)
+	{
+		// We haven't hit the current interval length but might be drawn in if
+		// demand increases. So we return a wait that takes us to that.
+		return int32(Remainder - Interval);
+	}
+
+	// PrevCycles is adjusted so we do not lose any left over time
+	CacheState.PrevCycles = Cycles - Remainder;
+	CacheState.OpCount++;
+
+	return CacheState.OpAllowance + CacheState.RunOff;
+}
+
+////////////////////////////////////////////////////////////////////////////////
+int32 FDiskCacheGovernor::TryBeginWrite(FDiskCacheGovernor::FCacheState& CacheState, uint32 DemandPercent)
+{
+	// A return of >=0 is the allowance of bytes that can be read. Otherwise the
+	// value is number of cycles to wait until allowance may be ready.
+
+	if (CacheState.State == FCacheState::EState::Rolling)
+	{
+		const int64 Cycles = FPlatformTime::Cycles64();
+		return TryBeginWrite(CacheState, DemandPercent, Cycles);
+	}
+
+	// Wait until the demand gets above this threshold
+	if (DemandPercent < CacheState.Demand.Threshold)
+	{
+		return 0 - GetMaxWaitCycles(CacheState);
+	}
+
+	CacheState.State		= FCacheState::EState::Rolling;
+	CacheState.PrevCycles	= FPlatformTime::Cycles64();
+	CacheState.OpCount		= 1;
+	CacheState.RunOff		= 0;
+
+	return CacheState.OpAllowance;
+}
+
+////////////////////////////////////////////////////////////////////////////////
+int32 FDiskCacheGovernor::EndWrite(FDiskCacheGovernor::FCacheState& CacheState, uint32 UnusedAllowance)
+{
+	CacheState.RunOff = UnusedAllowance;
+
+	if (CacheState.OpCount >= CacheState.MaxOpCount)
+	{
+		CacheState.State = FCacheState::EState::Waiting;
+		return 0 - GetMaxWaitCycles(CacheState);
+	}
+
+	return GetMaxWaitCycles(CacheState);
+}
+
+////////////////////////////////////////////////////////////////////////////////
+int32 FDiskCacheGovernor::GetMaxWaitCycles(FDiskCacheGovernor::FCacheState& CacheState) const
+{
+	// ">> 2" so we check at four times the speed in case of a S.U.P.E.R BOOST
+	return int32(CacheState.OpInterval >> 2);
+}
+
+} // namespace UE::IoStore
diff --git a/Engine/Source/Runtime/Experimental/IoStore/OnDemand/Private/DiskCacheGovernor.h b/Engine/Source/Runtime/Experimental/IoStore/OnDemand/Private/DiskCacheGovernor.h
new file mode 100644
--- /dev/null
+++ b/Engine/Source/Runtime/Experimental/IoStore/OnDemand/Private/DiskCacheGovernor.h
@@ -0,0 +1,72 @@
+// Copyright Epic Games, Inc. All Rights Reserved.
+
+#pragma once
+
+#include "Misc/CoreMiscDefines.h"
+
+namespace UE::IoStore
+{
+
+////////////////////////////////////////////////////////////////////////////////
+class FDiskCacheGovernor
+{
+public:
+	struct FAllowance
+	{
+		uint32	Bytes = 16 << 20;
+		uint32	Ops = 32;
+		uint32	Seconds = 60;
+	};
+
+	struct FDemand
+	{
+		uint8	Threshold = 30;
+		uint8	Boost = 60;
+		uint8	SuperBoost = 87;
+	};
+
+	struct FConfig
+	{
+		FAllowance	Allowance;
+		FDemand		Demand;
+		uint32		CommitBufferSize = 0;
+	};
+
+			FDiskCacheGovernor();
+			UE_NONCOPYABLE(FDiskCacheGovernor);
+
+	void	ConfigureStreamingCache(const FConfig& Config);
+	int32	TryBeginWriteToStreamingCache(uint32 DemandPercent);
+	int32	EndWriteToStreamingCache(uint32 UnusedAllowance);
+
+private:
+	struct FCacheState;
+
+	void	Configure(FCacheState& CacheState, const FConfig& Config);
+	int32	TryBeginWrite(FCacheState& CacheState, uint32 DemandPercent, int64 Cycles);
+	int32	TryBeginWrite(FCacheState& CacheState, uint32 DemandPercent);
+	int32	EndWrite(struct FCacheState& CacheState, uint32 UnusedAllowance);
+	int32	GetMaxWaitCycles(FCacheState& CacheState) const;
+
+	struct FCacheState
+	{
+		enum class EState : uint8
+		{
+			Waiting,
+			Rolling,
+		};
+
+		int64	OpInterval = 0;
+		int64	PrevCycles = 0;
+		uint32	RunOff = 0;
+		uint32	OpCount = 0;
+		uint32	MaxOpCount = 0;
+		uint32	OpAllowance = 0;
+		FDemand	Demand;
+		EState	State = EState::Waiting;
+	};
+
+	FCacheState StreamingCacheState;
+};
+
+} // namespace UE::IoStore
diff --git a/Engine/Source/Runtime/Experimental/IoStore/OnDemand/Private/IasCache.h b/Engine/Source/Runtime/Experimental/IoStore/OnDemand/Private/IasCache.h
--- a/Engine/Source/Runtime/Experimental/IoStore/OnDemand/Private/IasCache.h
+++ b/Engine/Source/Runtime/Experimental/IoStore/OnDemand/Private/IasCache.h
@@ -79,6 +79,6 @@ struct FIasCacheConfig
 	bool		DropCache = false;
 };
 
-TUniquePtr<IIasCache> MakeIasCache(const TCHAR* RootDir, const FIasCacheConfig& Config);
+TUniquePtr<IIasCache> MakeIasCache(const TCHAR* RootDir, const FIasCacheConfig& Config, class FDiskCacheGovernor& Governor);
 
 } // namespace UE::IoStore
diff --git a/Engine/Source/Runtime/Experimental/IoStore/OnDemand/Private/JournaledCache.cpp b/Engine/Source/Runtime/Experimental/IoStore/OnDemand/Private/JournaledCache.cpp
--- a/Engine/Source/Runtime/Experimental/IoStore/OnDemand/Private/JournaledCache.cpp
+++ b/Engine/Source/Runtime/Experimental/IoStore/OnDemand/Private/JournaledCache.cpp
@@ -1,6 +1,7 @@
 // Copyright Epic Games, Inc. All Rights Reserved.
 
 #include "IasCache.h"
+#include "DiskCacheGovernor.h"
 #include "IO/IoStoreOnDemand.h"
 #include "Statistics.h"
 
@@ -1598,169 +1599,6 @@ uint32 FCache::DebugVisit(void* Param, FDebugCacheEntry::Callback* Callback)
 }
 
 
-
-// {{{1 governor ...............................................................
-
-////////////////////////////////////////////////////////////////////////////////
-class FGovernor
-{
-public:
-			FGovernor();
-	void	Set(uint32 Allowance, uint32 Ops, uint32 Seconds);
-	void	SetDemands(uint32 Threshold, uint32 Boost, uint32 SuperBoost);
-	int32	BeginAllowance(uint32 DemandPercent);
-	int32	EndAllowance(uint32 UnusedAllowance);
-
-private:
-	enum class EState : uint8
-	{
-		Waiting,
-		Rolling,
-	};
-
-	int32	GetMaxWaitCycles() const;
-	void	Set(uint32 Allowance, uint32 Ops, uint32 Seconds, int64 CycleFreq);
-	int32	BeginInternal(uint32 Demand, int64 Cycle);
-	int64	OpInterval;
-	int64	PrevCycles;
-	uint32	RunOff = 0;
-	uint32	OpCount = 0;
-	uint32	MaxOpCount;
-	uint32	OpAllowance;
-	uint8	DemandThreshold = 30;
-	uint8	DemandBoost = 60;
-	uint8	DemandSuperBoost = 87;
-	EState	State = EState::Waiting;
-};
-
-////////////////////////////////////////////////////////////////////////////////
-FGovernor::FGovernor()
-{
-	Set(1, 1, 86400);
-}
-
-////////////////////////////////////////////////////////////////////////////////
-void FGovernor::Set(uint32 Allowance, uint32 Ops, uint32 Seconds)
-{
-	int64 CycleFreq	= int64(1.0 / FPlatformTime::GetSecondsPerCycle64());
-	return Set(Allowance, Ops, Seconds, CycleFreq);
-}
-
-////////////////////////////////////////////////////////////////////////////////
-void FGovernor::Set(uint32 Allowance, uint32 Ops, uint32 Seconds, int64 CycleFreq)
-{
-	int32 CommitBufferSize = JournaledCache::GetWriteCommitThreshold();
-	if (CommitBufferSize == 0)
-	{
-		OpAllowance = Allowance / Ops;
-		OpInterval = (CycleFreq * Seconds) / Ops;
-		MaxOpCount = 4;
-		return;
-	}
-
-	// A small commit buffer will make us op bound
-	if (uint32 CommitWidth = CommitBufferSize * Ops; CommitWidth < Allowance)
-	{
-		OpAllowance = CommitBufferSize;
-		OpInterval = (Seconds * CycleFreq * 3) / Ops;
-		MaxOpCount = 1;
-		return;
-	}
-
-	// Allowance bound.
-	int32 BlockCount = FMath::Max<int32>(1, Allowance / CommitBufferSize);
-	int32 CommitOpCost = BlockCount * 3;
-	MaxOpCount = (Ops - CommitOpCost) / BlockCount;
-
-	OpAllowance = CommitBufferSize / MaxOpCount;
-	OpInterval = (Seconds * CycleFreq) / (BlockCount * (MaxOpCount - 1));
-}
-
-////////////////////////////////////////////////////////////////////////////////
-void FGovernor::SetDemands(uint32 Threshold, uint32 Boost, uint32 SuperBoost)
-{
-	DemandThreshold = uint8(Threshold);
-	DemandBoost = uint8(Boost);
-	DemandSuperBoost = uint8(SuperBoost);
-}
-
-////////////////////////////////////////////////////////////////////////////////
-int32 FGovernor::BeginInternal(uint32 Demand, int64 Cycles)
-{
-	int64 Interval = OpInterval;
-	Interval >>= int32(Demand >= DemandBoost);
-	Interval >>= int32(Demand >= DemandSuperBoost);
-	Interval <<= int32(Demand <= DemandThreshold);
-
-	int64 Delta = Cycles - PrevCycles;
-	bool bNotYet = (Delta < Interval);
-
-	// Calculate how much time we are into the shortest poll interval
-	int64 Remainder = Delta;
-	Interval = GetMaxWaitCycles();
-	for (; Remainder > Interval; Remainder -= Interval);
-
-	if (bNotYet)
-	{
-		// We haven't hit the current interval length but might be drawn in if
-		// demand increases. So we return a wait that takes us to that.
-		return int32(Remainder - Interval);
-	}
-
-	// PrevCycles is adjusted so we do not lose any left over time
-	PrevCycles = Cycles - Remainder;
-
-	OpCount++;
-	return OpAllowance + RunOff;
-}
-
-////////////////////////////////////////////////////////////////////////////////
-int32 FGovernor::BeginAllowance(uint32 DemandPercent)
-{
-	// A return of >=0 is the allowance of bytes that can be read. Otherwise the
-	// value is number of cycles to wait until allowance may be ready.
-
-	if (State == EState::Rolling)
-	{
-		int64 Cycles = FPlatformTime::Cycles64();
-		return BeginInternal(DemandPercent, Cycles);
-	}
-
-	if (DemandPercent < DemandThreshold)
-	{
-		return 0 - GetMaxWaitCycles();
-	}
-
-	State = EState::Rolling;
-	PrevCycles = FPlatformTime::Cycles64();
-	OpCount = 1;
-	RunOff = 0;
-	return OpAllowance;
-}
-
-////////////////////////////////////////////////////////////////////////////////
-int32 FGovernor::EndAllowance(uint32 UnusedAllowance)
-{
-	RunOff = UnusedAllowance;
-
-	if (OpCount >= MaxOpCount)
-	{
-		State = EState::Waiting;
-		return 0 - GetMaxWaitCycles();
-	}
-
-	return GetMaxWaitCycles();
-}
-
-////////////////////////////////////////////////////////////////////////////////
-int32 FGovernor::GetMaxWaitCycles() const
-{
-	// ">> 2" so we check at four times the speed in case of a S.U.P.E.R BOOST
-	return int32(OpInterval >> 2);
-}
-
-
-
 // {{{1 service-thread .........................................................
 
 ////////////////////////////////////////////////////////////////////////////////
@@ -1787,8 +1625,7 @@ public:
 							~FServiceThread();
 	void					RegisterCache(TUniquePtr<FCache> Cache);
 	void					UnregisterCache(FCache* Cache);
-	void					SetGovernorRate(uint32 Allowance, uint32 Ops, uint32 Seconds);
-	void					SetGovernorDemand(uint32 Threshold, uint32 Boost, uint32 SuperBoost);
+	void					SetGovernor(FDiskCacheGovernor* Governor);
 	uint32					ClaimReadId();
 	void					BeginRead(const FCache* Cache, const FReadRequest& Request);
 	void					CancelRead(const FIoBuffer* GivenDest);
@@ -1799,8 +1636,7 @@ private:
 		enum {
 			Work_Register,
 			Work_Unregister,
-			Work_GovDemand,
-			Work_GovRate,
+			Work_Gov,
 			Work_Read,
 			Work_Cancel,
 		};
@@ -1823,17 +1659,9 @@ private:
 			TField<FReadSink*>	Sink;
 			TField<FIoBuffer*>	Dest;
 		};
-		struct {				// rate
+		struct {				// governor
 			uint16				_What0;
-			uint16				Ops;
-			uint16				Seconds;
-			TField<uint32>		Allowance;
-		};
-		struct {				// demand
-			uint16				_What1;
-			uint16				Threshold;
-			uint16				Boost;
-			uint16				SuperBoost;
+			TField<FDiskCacheGovernor*> Gov;
 		};
 	};
 	static_assert(sizeof(FWork) == sizeof(uint16) * 9);
@@ -1847,7 +1675,7 @@ private:
 	void						ReceiveWork();
 	TUniquePtr<FRunnableThread> Thread;
 	FEventRef					WakeEvent;
-	FGovernor					Governor;
+	FDiskCacheGovernor*			Governor = nullptr;
 	TArray<TUniquePtr<FCache>>	Caches;
 	std::atomic_int				RunCount = 0;
 	FCriticalSection			Lock;
@@ -1923,30 +1751,12 @@ void FServiceThread::UnregisterCache(FCache* Cache)
 	}
 }
 
-////////////////////////////////////////////////////////////////////////////////
-void FServiceThread::SetGovernorRate(uint32 Allowance, uint32 Ops, uint32 Seconds)
+void FServiceThread::SetGovernor(FDiskCacheGovernor* Gov)
 {
 	FWork Work;
-	Work.What = FWork::Work_GovRate;
-	Work.Allowance = Allowance;
-	Work.Ops = uint16(Ops);
-	Work.Seconds = uint16(Seconds);
+	Work.What = FWork::Work_Gov;
+	Work.Gov = Gov;
 	SubmitWork(&Work, 1);
-
-	check(Work.Ops == Ops && Work.Seconds == Seconds);
-}
-
-////////////////////////////////////////////////////////////////////////////////
-void FServiceThread::SetGovernorDemand(uint32 Threshold, uint32 Boost, uint32 SuperBoost)
-{
-	FWork Work;
-	Work.What = FWork::Work_GovDemand;
-	Work.Threshold = uint16(Threshold);
-	Work.Boost = uint16(Boost);
-	Work.SuperBoost = uint16(SuperBoost);
-	SubmitWork(&Work, 1);
-
-	check(Work.Threshold == Threshold && Work.Boost == Boost && Work.SuperBoost == SuperBoost);
 }
 
 ////////////////////////////////////////////////////////////////////////////////
@@ -2040,11 +1850,14 @@ int32 FServiceThread::Update()
 
 	// Update caches
 	uint32 CycleSlice = MAX_uint32;
-	for (TUniquePtr<FCache>& Item : Caches)
+	if (ensureMsgf(Governor != nullptr, TEXT("Trying to update the streaming cache without a valid disk governor")))
 	{
-		FCache* Cache = Item.Get();
-		uint32 CyclesTillActive = UpdateCache(Cache);
-		CycleSlice = FMath::Min(CyclesTillActive, CycleSlice);
+		for (TUniquePtr<FCache>& Item : Caches)
+		{
+			FCache* Cache = Item.Get();
+			uint32 CyclesTillActive = UpdateCache(Cache);
+			CycleSlice = FMath::Min(CyclesTillActive, CycleSlice);
+		}
 	}
 
 	// Early out
@@ -2167,15 +1980,9 @@ void FServiceThread::ReceiveWork()
 			continue;
 		}
 
-		if (Work.What == FWork::Work_GovRate)
-		{
-			Governor.Set(Work.Allowance, Work.Ops, Work.Seconds);
-			continue;
-		}
-
-		if (Work.What == FWork::Work_GovDemand)
+		if (Work.What == FWork::Work_Gov)
 		{
-			Governor.SetDemands(Work.Threshold, Work.Boost, Work.SuperBoost);
+			Governor = Work.Gov;
 			continue;
 		}
 
@@ -2213,7 +2020,7 @@ void FServiceThread::ReceiveWork()
 uint32 FServiceThread::UpdateCache(FCache* Cache)
 {
 	uint32 Demand = Cache->GetDemand();
-	int32 Allowance = Governor.BeginAllowance(Demand);
+	int32 Allowance = Governor->TryBeginWriteToStreamingCache(Demand);
 	if (Allowance <= 0)
 	{
 		return -Allowance;
@@ -2230,7 +2037,7 @@ uint32 FServiceThread::UpdateCache(FCache* Cache)
 	TRACE_COUNTER_SET(IasAllowance, Unused);
 	TRACE_COUNTER_ADD(IasOpCount, 1);
 
-	int32 WaitCycles = Governor.EndAllowance(Unused);
+	int32 WaitCycles = Governor->EndWriteToStreamingCache(Unused);
 	if (WaitCycles < 0)
 	{
 		WaitCycles = -WaitCycles;
@@ -2261,7 +2068,7 @@ class FJournaledCache final
 	, public JournaledCache::FServiceThread::FReadSink
 {
 public:
-								FJournaledCache() = default;
+								FJournaledCache(FDiskCacheGovernor& DiskCacheGovernor);
 								~FJournaledCache();
 	bool						Initialize(const TCHAR* RootDir, const FIasCacheConfig& Config);
 	virtual void				Abandon() override;
@@ -2281,12 +2088,19 @@ private:
 
 	virtual void				OnRead(const FReadResult* Results, int32 Num) override;
 	static uint64				ReduceKey(const FIoHash& Key);
+	FDiskCacheGovernor&			DiskCacheGovernor;
 	JournaledCache::FCache*		Cache = nullptr;
 	UE::Tasks::FPipe			GetPipe = UE::Tasks::FPipe(TEXT("IasCacheGetPipe"));
 	FCriticalSection			Lock;
 	TMap<int32, FMaterialOp>	PendingMaterializes;
 };
 
+////////////////////////////////////////////////////////////////////////////////
+FJournaledCache::FJournaledCache(FDiskCacheGovernor& InDiskCacheGovernor)
+	: DiskCacheGovernor(InDiskCacheGovernor)
+{
+}
+
 ////////////////////////////////////////////////////////////////////////////////
 FJournaledCache::~FJournaledCache()
 {
@@ -2294,6 +2108,7 @@ FJournaledCache::~FJournaledCache()
 	{
 		JournaledCache::FServiceThread::Get().UnregisterCache(Cache);
 	}
+	JournaledCache::FServiceThread::Get().SetGovernor(nullptr);
 }
 
 ////////////////////////////////////////////////////////////////////////////////
@@ -2323,6 +2138,21 @@ bool FJournaledCache::Initialize(const TCHAR* RootDir, const FIasCacheConfig& Co
 	FCache::FConfig EventualConfig;
 	static_cast<FIasCacheConfig&>(EventualConfig) = Config;
 	EventualConfig.Path = CachePath;
+
+	const FIasCacheConfig::FRate& WriteRate = Config.WriteRate;
+	const FIasCacheConfig::FDemand& Demand = Config.Demand;
+
+	// Configure the governor
+	FDiskCacheGovernor::FConfig Cfg;
+	Cfg.Allowance.Bytes		= WriteRate.Allowance;
+	Cfg.Allowance.Ops		= WriteRate.Ops;
+	Cfg.Allowance.Seconds	= WriteRate.Seconds;
+	Cfg.Demand.Threshold	= Demand.Threshold;
+	Cfg.Demand.Boost		= Demand.Boost;
+	Cfg.Demand.SuperBoost	= Demand.SuperBoost;
+	Cfg.CommitBufferSize	= JournaledCache::GetWriteCommitThreshold();
+	DiskCacheGovernor.ConfigureStreamingCache(Cfg);
+
 	TUniquePtr<FCache> NewCache = MakeUnique<FCache>(MoveTemp(EventualConfig));
 
 	if (uint32 Ailments = NewCache->GetAilments(); Ailments != 0)
@@ -2333,13 +2163,9 @@ bool FJournaledCache::Initialize(const TCHAR* RootDir, const FIasCacheConfig& Co
 
 	Cache = NewCache.Get();
 
-	const FIasCacheConfig::FRate& WriteRate = Config.WriteRate;
-	const FIasCacheConfig::FDemand& Demand = Config.Demand;
-
 	FServiceThread& ServiceThread = FServiceThread::Get();
+	ServiceThread.SetGovernor(&DiskCacheGovernor);
 	ServiceThread.RegisterCache(MoveTemp(NewCache));
-	ServiceThread.SetGovernorRate(WriteRate.Allowance, WriteRate.Ops, WriteRate.Seconds);
-	ServiceThread.SetGovernorDemand(Demand.Threshold, Demand.Boost, Demand.SuperBoost);
 
 	return true;
 }
@@ -2461,11 +2287,11 @@ uint64 FJournaledCache::ReduceKey(const FIoHash& Key)
 
 
 ////////////////////////////////////////////////////////////////////////////////
-TUniquePtr<IIasCache> MakeIasCache(const TCHAR* RootPath, const FIasCacheConfig& Config)
+TUniquePtr<IIasCache> MakeIasCache(const TCHAR* RootPath, const FIasCacheConfig& Config, FDiskCacheGovernor& Governor)
 {
 	LLM_SCOPE_BYTAG(Ias);
 
-	FJournaledCache* Cache = new FJournaledCache();
+	FJournaledCache* Cache = new FJournaledCache(Governor);
 	if (Cache->Initialize(RootPath, Config))
 	{
 		return TUniquePtr<IIasCache>(Cache);
@@ -3175,7 +3001,8 @@ static void MiscTests(FSupport& Support)
 
 	{ // Benign
 		Support.CleanFs();
-		auto Jc = MakeIasCache(*Support.TestDir, Config);
+		FDiskCacheGovernor Gov;
+		auto Jc = MakeIasCache(*Support.TestDir, Config, Gov);
 		check(Jc.IsValid());
 		Jc.Reset();
 	}
@@ -3186,7 +3013,8 @@ static void MiscTests(FSupport& Support)
 		{
 			const TCHAR* TestName = TEXT("m/i/s/c");
 			Config.Name = TestName + i;
-			auto Jc = MakeIasCache(*Support.TestDir, Config);
+			FDiskCacheGovernor Gov;
+			auto Jc = MakeIasCache(*Support.TestDir, Config, Gov);
 			check(Jc.IsValid());
 			Jc.Reset();
 		}
@@ -3212,7 +3040,8 @@ static void MiscTests(FSupport& Support)
 			IFileManager& Ifm = IFileManager::Get();
 			check(Ifm.MakeDirectory(*Blocker, true));
 
-			auto Jc = MakeIasCache(*Support.TestDir, Config);
+			FDiskCacheGovernor Gov;
+			auto Jc = MakeIasCache(*Support.TestDir, Config, Gov);
 			check(!Jc.IsValid())
 		}
 	}
diff --git a/Engine/Source/Runtime/Experimental/IoStore/OnDemand/Private/OnDemandIoStore.cpp b/Engine/Source/Runtime/Experimental/IoStore/OnDemand/Private/OnDemandIoStore.cpp
--- a/Engine/Source/Runtime/Experimental/IoStore/OnDemand/Private/OnDemandIoStore.cpp
+++ b/Engine/Source/Runtime/Experimental/IoStore/OnDemand/Private/OnDemandIoStore.cpp
@@ -357,7 +357,7 @@ FIoStatus FOnDemandIoStore::Initialize()
 			if (FPaths::HasProjectPersistentDownloadDir())
 			{
 				FString CacheDir = FPaths::ProjectPersistentDownloadDir();
-				HttpCache = MakeIasCache(*CacheDir, CacheConfig);
+				HttpCache = MakeIasCache(*CacheDir, CacheConfig, DiskCacheGovernor);
 
 				UE_CLOG(!HttpCache.IsValid(), LogIoStoreOnDemand, Warning, TEXT("HTTP cache disabled"));
 			}
@@ -552,7 +552,7 @@ FIoStatus FOnDemandIoStore::InitializeStreamingBackend(const FOnDemandEndpointCo
 		if (FPaths::HasProjectPersistentDownloadDir())
 		{
 			FString CacheDir = FPaths::ProjectPersistentDownloadDir();
-			StreamingCache = MakeIasCache(*CacheDir, CacheConfig);
+			StreamingCache = MakeIasCache(*CacheDir, CacheConfig, DiskCacheGovernor);
 
 			UE_CLOG(!StreamingCache.IsValid(), LogIas, Warning, TEXT("File cache disabled - streaming only (init-fail)"));
 		}
diff --git a/Engine/Source/Runtime/Experimental/IoStore/OnDemand/Private/OnDemandIoStore.h b/Engine/Source/Runtime/Experimental/IoStore/OnDemand/Private/OnDemandIoStore.h
--- a/Engine/Source/Runtime/Experimental/IoStore/OnDemand/Private/OnDemandIoStore.h
+++ b/Engine/Source/Runtime/Experimental/IoStore/OnDemand/Private/OnDemandIoStore.h
@@ -7,6 +7,7 @@
 #include "Async/Mutex.h"
 #include "Containers/AnsiString.h"
 #include "Containers/BitArray.h"
+#include "DiskCacheGovernor.h"
 #include "IO/HttpIoDispatcher.h"
 #include "IO/IoBuffer.h"
 #include "IO/IoChunkEncoding.h"
@@ -421,6 +422,7 @@ private:
 	FSharedPackageStoreBackend				PackageStoreBackend;
 	FSharedStreamingBackend					StreamingBackend;
 	FSharedHttpIoBackend					HttpIoBackend;
+	FDiskCacheGovernor						DiskCacheGovernor;
 
 	FDelegateHandle							OnMountPakHandle;
 	FDelegateHandle							OnServerPostForkHandle;
