[ChaosVD] Fix for a crash that happens when the post-integrate solver stage data channel is enabled.

A few weeks ago a refactor to uptimize chaos threading performance broke some assumptions CVD made two years ago. Now when we attempt to trace the results in the place when Integrate used to be completed, we might still be integrating particles.

The fix is just add a call to make the execution wait for all integrate tasks to be completed before attempting to trace.

#okforgitgub public

diff --git a/Engine/Source/Runtime/Experimental/Chaos/Private/Chaos/PBDRigidsEvolutionGBF.cpp b/Engine/Source/Runtime/Experimental/Chaos/Private/Chaos/PBDRigidsEvolutionGBF.cpp
--- a/Engine/Source/Runtime/Experimental/Chaos/Private/Chaos/PBDRigidsEvolutionGBF.cpp
+++ b/Engine/Source/Runtime/Experimental/Chaos/Private/Chaos/PBDRigidsEvolutionGBF.cpp
@@ -552,11 +552,6 @@ void FPBDRigidsEvolutionGBF::AdvanceOneTimeStepImpl(const FReal Dt, const FSubSt
 			ApplyKinematicTargets(Dt, SubStepInfo.PseudoFraction);
 		}
 
-		{
-			CVD_SCOPE_TRACE_SOLVER_STEP(CVDDC_PostIntegrate, TEXT("Post Integrate"));
-			CVD_TRACE_PARTICLES_SOA(Particles, &Clustering);
-		}
-
 		if (PostIntegrateCallback != nullptr)
 		{
 			SCOPE_CYCLE_COUNTER(STAT_Evolution_PostIntegrateCallback);
@@ -576,6 +571,15 @@ void FPBDRigidsEvolutionGBF::AdvanceOneTimeStepImpl(const FReal Dt, const FSubSt
 		TaskDispatcher.WaitTaskEndSpatial(); 
 	}
 
+#if WITH_CHAOS_VISUAL_DEBUGGER
+		if (FChaosVisualDebuggerTrace::IsTracing())
+		{
+			WaitIntegrationComplete();
+			CVD_SCOPE_TRACE_SOLVER_STEP(CVDDC_PostIntegrate, TEXT("Post Integrate"));
+			CVD_TRACE_PARTICLES_SOA(Particles, &Clustering);
+		}
+#endif
+
 	// Collision detection is sensitive to duplication bugs in the particle views
 	// so this is here to help us track them down when they happen
 	CheckParticleViewsForErrors(Particles);
