Chaos
- fix crash when changing joint solver type (per joint) after its creation
- also clean up joint proxy to improve testability of evolution independently of the proxies

#rnx

diff --git a/Engine/Source/Programs/HeadlessChaos/Private/HeadlessChaosTestJoints.cpp b/Engine/Source/Programs/HeadlessChaos/Private/HeadlessChaosTestJoints.cpp
--- a/Engine/Source/Programs/HeadlessChaos/Private/HeadlessChaosTestJoints.cpp
+++ b/Engine/Source/Programs/HeadlessChaos/Private/HeadlessChaosTestJoints.cpp
@@ -4,6 +4,7 @@
 #include "HeadlessChaosTestConstraints.h"
 
 #include "Modules/ModuleManager.h"
+#include "Chaos/ChaosEngineInterface.h"
 #include "Chaos/ParticleHandle.h"
 #include "Chaos/PBDRigidsEvolution.h"
 #include "Chaos/PBDRigidsEvolutionGBF.h"
@@ -20,6 +21,9 @@
 #include "Chaos/Levelset.h"
 #include "Chaos/UniformGrid.h"
 #include "Chaos/Utilities.h"
+#include "PBDRigidsSolver.h"
+#include "PhysicsProxy/JointConstraintProxy.h"
+#include "PhysicsProxy/SingleParticlePhysicsProxy.h"
 
 namespace ChaosTest {
 
@@ -83,10 +87,21 @@ namespace ChaosTest {
 			for (int32 JointIndex = 0; JointIndex < JointPositions.Num(); ++JointIndex)
 			{
 				const TVec2<FGeometryParticleHandle*> ConstraintedParticleIds(GetParticle(JointParticleIndices[JointIndex][0]), GetParticle(JointParticleIndices[JointIndex][1]));
-				AddJoint(ConstraintedParticleIds, JointPositions[JointIndex]);
+				FPBDJointConstraintHandle*  Joint = AddJoint(ConstraintedParticleIds, JointPositions[JointIndex]);
+				Joints.Add(Joint);
 			}
 		}
 
+		int32 GetNumJoints() const
+		{
+			return Joints.Num();
+		}
+
+		FPBDJointConstraintHandle* GetJoint(int32 Index)
+		{
+			return Joints[Index];
+		}
+
 		TFunction<void(FPBDJointSettings&)>& SetSetTestJointSetting() { return SetTestJointSetting; }
 
 		// Initial particles setup
@@ -99,6 +114,9 @@ namespace ChaosTest {
 		TArray<TVec2<int32>> JointParticleIndices;
 
 		TFunction<void(FPBDJointSettings&)> SetTestJointSetting = [](FPBDJointSettings& InputSetting) {};
+
+		// Created joints
+		TArray<FPBDJointConstraintHandle*> Joints;
 	};
 
 	/**
@@ -1066,5 +1084,56 @@ namespace ChaosTest {
 		EXPECT_EQ(JointContainer.GetConstraintLevel(1), 0);
 	}
 
+	// Change the linear/non-linear solver mode of a constraint after creation using the GT Proxy API.
+	// Bug Fix: changing the solver type after creation would result in an attempt to remove the constraint
+	// from the wrong Container on destruction, leading to an OOB array access, or destruction of the wrong
+	// constraint and a later OOB access attempting to destroy that constraint.
+	GTEST_TEST(JointTests, TestSetUseLinearJointSolver)
+	{
+		using namespace Chaos;
+		auto Sphere = Chaos::FImplicitObjectPtr(new Chaos::FSphere(FVec3(0), 10));
+		FChaosSolversModule* Module = FChaosSolversModule::GetModule();
+
+		// Make a solver
+		auto* Solver = Module->CreateSolver(nullptr, /*AsyncDt=*/-1);
+
+		// Make 2 particles for our constraint
+		auto Proxy1 = FSingleParticlePhysicsProxy::Create(Chaos::FPBDRigidParticle::CreateParticle());
+		auto& Particle1 = Proxy1->GetGameThreadAPI();
+		Particle1.SetGeometry(Sphere);
+		Solver->RegisterObject(Proxy1);
+
+		auto Proxy2 = FSingleParticlePhysicsProxy::Create(Chaos::FPBDRigidParticle::CreateParticle());
+		auto& Particle2 = Proxy2->GetGameThreadAPI();
+		Particle2.SetGeometry(Sphere);
+		Solver->RegisterObject(Proxy2);
+
+		// Confirm solver doesn't have any joint constraint proxies yet.
+		EXPECT_EQ(Solver->GetJointConstraintPhysicsProxies_Internal().Num(), 0);
+
+		FTransform IdentityTransform(FTransform::Identity);
+		FPhysicsConstraintHandle Constraint = FChaosEngineInterface::CreateConstraint(Proxy1, Proxy2, IdentityTransform, IdentityTransform);
+		FChaosEngineInterface::SetUseLinearJointSolver_AssumesLocked(Constraint, false);
+
+		// Tick physics thread to create the PT constraint
+		Solver->AdvanceAndDispatch_External(/*Dt=*/1.0f);
+
+		// Change the solver type for the joint
+		FChaosEngineInterface::SetUseLinearJointSolver_AssumesLocked(Constraint, true);
+
+		// Tick physics thread
+		Solver->AdvanceAndDispatch_External(/*Dt=*/1.0f);
+
+		// Release constraint
+		FChaosEngineInterface::ReleaseConstraint(Constraint);
+
+		// Tick physics thread
+		// Bug Fix: This would crash in FPBDJointConstraints::RemoveConstraint
+		Solver->AdvanceAndDispatch_External(/*Dt=*/1.0f);
+
+		Solver->UnregisterObject(Proxy1);
+		Solver->UnregisterObject(Proxy2);
+		Module->DestroySolver(Solver);
+	}
 }
 
diff --git a/Engine/Source/Runtime/Experimental/Chaos/Private/Chaos/PBDRigidsEvolutionGBF.cpp b/Engine/Source/Runtime/Experimental/Chaos/Private/Chaos/PBDRigidsEvolutionGBF.cpp
--- a/Engine/Source/Runtime/Experimental/Chaos/Private/Chaos/PBDRigidsEvolutionGBF.cpp
+++ b/Engine/Source/Runtime/Experimental/Chaos/Private/Chaos/PBDRigidsEvolutionGBF.cpp
@@ -1370,6 +1370,109 @@ void FPBDRigidsEvolutionGBF::SetCurrentStepResimCache(IResimCacheBase* InCurrent
 	CurrentStepResimCacheImp = static_cast<FEvolutionResimCache*>(InCurrentStepResimCache);
 }
 
+FPBDJointConstraintHandle* FPBDRigidsEvolutionGBF::CreateJointConstraint(const FParticlePair& InConstrainedParticles, const FPBDJointSettings& InJointSettings)
+{
+	FPBDJointConstraintHandle* JointHandle = nullptr;
+
+	if (InJointSettings.bUseLinearSolver)
+	{
+		JointHandle = JointConstraints.LinearConstraints.AddConstraint(InConstrainedParticles, InJointSettings.ConnectorTransforms);
+	}
+	else
+	{
+		JointHandle = JointConstraints.NonLinearConstraints.AddConstraint(InConstrainedParticles, InJointSettings.ConnectorTransforms);
+	}
+
+	if (!InJointSettings.bCollisionEnabled)
+	{
+		GetBroadPhase().GetIgnoreCollisionManager().AddIgnoreCollisions(InConstrainedParticles[0], InConstrainedParticles[1]);
+	}
+
+	JointHandle->SetSettings(InJointSettings);
+
+	// Register the constraint with each particle
+	InConstrainedParticles[0]->AddConstraintHandle(JointHandle);
+	InConstrainedParticles[1]->AddConstraintHandle(JointHandle);
+
+	// We added a joint to the particles, so we need to modify the inertia (See UpdateInertiaConditioning)
+	FGenericParticleHandle(InConstrainedParticles[0])->SetInertiaConditioningDirty();
+	FGenericParticleHandle(InConstrainedParticles[1])->SetInertiaConditioningDirty();
+
+	return JointHandle;
+}
+
+void FPBDRigidsEvolutionGBF::DestroyJointConstraint(FPBDJointConstraintHandle* InJointHandle)
+{
+	FParticlePair ConstrainedParticles = InJointHandle->GetConstrainedParticles();
+
+	const FPBDJointSettings& Settings = InJointHandle->GetSettings();
+	const bool bValidPair = ConstrainedParticles[0] && ConstrainedParticles[1];
+
+	// If this constraint disables collisions - need to restore the collisions after we destroy the constraint
+	if (bValidPair && !Settings.bCollisionEnabled)
+	{
+		GetBroadPhase().GetIgnoreCollisionManager().RemoveIgnoreCollisions(ConstrainedParticles[0], ConstrainedParticles[1]);
+	}
+
+	// Unegister the constraint from each particle
+	if (ConstrainedParticles[0])
+	{
+		ConstrainedParticles[0]->RemoveConstraintHandle(InJointHandle);
+	}
+	if (ConstrainedParticles[1])
+	{
+		ConstrainedParticles[1]->RemoveConstraintHandle(InJointHandle);
+	}
+
+	// Remove from the constraint graph
+	if (InJointHandle->IsInConstraintGraph())
+	{
+		RemoveConstraintFromConstraintGraph(InJointHandle);
+	}
+
+	// Destroy the constraint
+	if (Settings.bUseLinearSolver)
+	{
+		JointConstraints.LinearConstraints.RemoveConstraint(InJointHandle->GetConstraintIndex());
+	}
+	else
+	{
+		JointConstraints.NonLinearConstraints.RemoveConstraint(InJointHandle->GetConstraintIndex());
+	}
+}
+
+FPBDJointConstraintHandle* FPBDRigidsEvolutionGBF::SetJointConstraintSettings(FPBDJointConstraintHandle* InJointHandle, const FPBDJointSettings& InJointSettings)
+{
+	const FPBDJointSettings& CurrentJointSettings = InJointHandle->GetSettings();
+	FParticlePair ConstrainedParticles = InJointHandle->GetConstrainedParticles();
+
+	// Handle changes to the CollisionEnabled flag
+	if (CurrentJointSettings.bCollisionEnabled != InJointSettings.bCollisionEnabled)
+	{
+		if (InJointSettings.bCollisionEnabled)
+		{
+			GetBroadPhase().GetIgnoreCollisionManager().RemoveIgnoreCollisions(ConstrainedParticles[0], ConstrainedParticles[1]);
+		}
+		else
+		{
+			GetBroadPhase().GetIgnoreCollisionManager().AddIgnoreCollisions(ConstrainedParticles[0], ConstrainedParticles[1]);
+		}
+	}
+
+	// NOTE: Changing the solver type requires us to destroy and recreate the joint, which results in
+	// relocating the joint handle. This also resets any state in the joint, e.g., whether it is broken.
+	// @todo(chaos): Improve suport for this feature, but for now, just make it not crash!
+	if (CurrentJointSettings.bUseLinearSolver != InJointSettings.bUseLinearSolver)
+	{
+		DestroyJointConstraint(InJointHandle);
+		return CreateJointConstraint(ConstrainedParticles, InJointSettings);
+	}
+
+	// Update the joint settings
+	InJointHandle->SetSettings(InJointSettings);
+	return InJointHandle;
+}
+
 void FPBDRigidsEvolutionGBF::TransferJointConstraintCollisions()
 {
 	//
diff --git a/Engine/Source/Runtime/Experimental/Chaos/Private/PhysicsProxy/JointConstraintProxy.cpp b/Engine/Source/Runtime/Experimental/Chaos/Private/PhysicsProxy/JointConstraintProxy.cpp
--- a/Engine/Source/Runtime/Experimental/Chaos/Private/PhysicsProxy/JointConstraintProxy.cpp
+++ b/Engine/Source/Runtime/Experimental/Chaos/Private/PhysicsProxy/JointConstraintProxy.cpp
@@ -114,9 +114,9 @@ void FJointConstraintPhysicsProxy::InitializeOnPhysicsThread(FPBDRigidsSolver* I
 
 			if (RootHandles[0] && RootHandles[1])
 			{
-				if (const FPBDJointSettings* JointSettings = RemoteData.FindJointSettings(Manager, DataIdx))
+				if (const FPBDJointSettings* SrcJointSettings = RemoteData.FindJointSettings(Manager, DataIdx))
 				{
-					FTransformPair ConnectorTransforms{ JointSettings->ConnectorTransforms };
+					FPBDJointSettings JointSettings = *SrcJointSettings;
 
 					// if the root particles do not match the actual particles
 					// we need to adjust the frames to be in the root particle space 
@@ -125,24 +125,9 @@ void FJointConstraintPhysicsProxy::InitializeOnPhysicsThread(FPBDRigidsSolver* I
 						BodyPairs->PhysicsBodies[0]->GetParticle<Chaos::EThreadContext::Internal>(),
 						BodyPairs->PhysicsBodies[1]->GetParticle<Chaos::EThreadContext::Internal>(),
 					};
+					FixConnectorTransformsForRoot(OriginalParticleHandles_PT, RootHandles, JointSettings.ConnectorTransforms);
 
-					FixConnectorTransformsForRoot(OriginalParticleHandles_PT, RootHandles, ConnectorTransforms);
-
-					if (JointSettings->bUseLinearSolver)
-					{
-						Constraint_PT = JointConstraints.LinearConstraints.AddConstraint(RootHandles, ConnectorTransforms);
-					}
-					else
-					{
-						Constraint_PT = JointConstraints.NonLinearConstraints.AddConstraint(RootHandles, ConnectorTransforms);
-					}
-
-					RootHandles[0]->AddConstraintHandle(Constraint_PT);
-					RootHandles[1]->AddConstraintHandle(Constraint_PT);
-
-					// We added a joint to the particles, so we need to modify the inertia (See FPBDRigidsEvolutionGBF::UpdateInertiaConditioning)
-					FGenericParticleHandle(RootHandles[0])->SetInertiaConditioningDirty();
-					FGenericParticleHandle(RootHandles[1])->SetInertiaConditioningDirty();
+					Constraint_PT = InSolver->GetEvolution()->CreateJointConstraint(RootHandles, JointSettings);
 				}
 			}
 		}
@@ -153,43 +138,8 @@ void FJointConstraintPhysicsProxy::DestroyOnPhysicsThread(FPBDRigidsSolver* InSo
 {
 	if (Constraint_PT)
 	{
-		TVec2<FGeometryParticleHandle*> Particles = Constraint_PT->GetConstrainedParticles();
-		
-		const FPBDJointSettings Settings = Constraint_PT->GetSettings();
-		const bool bValidPair = InSolver && Particles[0] && Particles[1];
-
-		// If this constraint disables collisions - need to restore the collisions after we destroy the constraint
-		if(bValidPair && !Settings.bCollisionEnabled)
-		{
-			InSolver->GetEvolution()->GetBroadPhase().GetIgnoreCollisionManager().RemoveIgnoreCollisions(Particles[0], Particles[1]);
-		}
+		InSolver->GetEvolution()->DestroyJointConstraint(Constraint_PT);
 
-		// Ensure that our connected particles are aware that this constraint no longer exists
-		if(Particles[0])
-		{
-			Particles[0]->RemoveConstraintHandle(Constraint_PT);
-		}
-
-		if(Particles[1])
-		{
-			Particles[1]->RemoveConstraintHandle(Constraint_PT);
-		}
-
-		// @todo(chaos): clean up constraint management
-		if(Constraint_PT->IsInConstraintGraph())
-		{
-			InSolver->GetEvolution()->RemoveConstraintFromConstraintGraph(Constraint_PT);
-		}
-
-		FPBDRigidsSolver::FJointConstraints& JointConstraints = InSolver->GetJointCombinedConstraints();
-		if (Settings.bUseLinearSolver)
-		{
-			JointConstraints.LinearConstraints.RemoveConstraint(Constraint_PT->GetConstraintIndex());
-		}
-		else
-		{
-			JointConstraints.NonLinearConstraints.RemoveConstraint(Constraint_PT->GetConstraintIndex());
-		}
 		Constraint_PT = nullptr;
 		OriginalParticleHandles_PT = { nullptr, nullptr };
 	}
@@ -217,31 +167,14 @@ void FJointConstraintPhysicsProxy::PushStateOnPhysicsThread(FPBDRigidsSolver* In
 	{
 		if (const FPBDJointSettings* Data = RemoteData.FindJointSettings(Manager, DataIdx))
 		{
-			const FPBDJointSettings& JointSettingsBuffer = *Data;
-			const FPBDJointSettings& CurrentConstraintSettings = Constraint_PT->GetSettings();
-
-			// Handle changes to the CollisionEnabled flag
-			if (CurrentConstraintSettings.bCollisionEnabled != JointSettingsBuffer.bCollisionEnabled)
-			{
-				const TVector<FGeometryParticleHandle*, 2>& Particles = Constraint_PT->GetConstrainedParticles();
-				if (JointSettingsBuffer.bCollisionEnabled)
-				{
-					InSolver->GetEvolution()->GetBroadPhase().GetIgnoreCollisionManager().RemoveIgnoreCollisions(Particles[0], Particles[1]);
-				}
-				else
-				{
-					InSolver->GetEvolution()->GetBroadPhase().GetIgnoreCollisionManager().AddIgnoreCollisions(Particles[0], Particles[1]);
-				}
-			}
-
-			FPBDJointSettings JointSettingsToSet{ JointSettingsBuffer };
+			FPBDJointSettings JointSettingsToSet = *Data;
 
 			// PushStateOnPhysicsThread is always called right after InitializeOnPhysicsThread
 			// so we need to make sure the root space correction logic is also run here to avoid resetting the transform in the wrong space
 			FixConnectorTransformsForRoot(OriginalParticleHandles_PT, Constraint_PT->GetConstrainedParticles(), JointSettingsToSet.ConnectorTransforms);
 
-			// Update the joint settings
-			Constraint_PT->SetSettings(JointSettingsToSet);
+			// NOTE: If bUseLinearSolver changes, the handle get re-allocated
+			Constraint_PT = InSolver->GetEvolution()->SetJointConstraintSettings(Constraint_PT, JointSettingsToSet);
 		}
 	}
 }
diff --git a/Engine/Source/Runtime/Experimental/Chaos/Public/Chaos/PBDRigidsEvolutionGBF.h b/Engine/Source/Runtime/Experimental/Chaos/Public/Chaos/PBDRigidsEvolutionGBF.h
--- a/Engine/Source/Runtime/Experimental/Chaos/Public/Chaos/PBDRigidsEvolutionGBF.h
+++ b/Engine/Source/Runtime/Experimental/Chaos/Public/Chaos/PBDRigidsEvolutionGBF.h
@@ -261,6 +261,16 @@ namespace Chaos
 
 		FSpatialAccelerationBroadPhase& GetBroadPhase() { return BroadPhase; }
 
+		// Create a joint constraint connecting two particles
+		CHAOS_API FPBDJointConstraintHandle* CreateJointConstraint(const FParticlePair& InParticlePair, const FPBDJointSettings& InJointSettings);
+
+		// Destroy a joint constraint
+		CHAOS_API void DestroyJointConstraint(FPBDJointConstraintHandle* InJointHandle);
+
+		// Change the settings of a joint constraint.
+		// NOTE: May recreate the joint which would change the handle pointer. (Currently only when changing linear/non-linear solver type)
+		CHAOS_API FPBDJointConstraintHandle* SetJointConstraintSettings(FPBDJointConstraintHandle* InJointHandle, const FPBDJointSettings& InJointSettings);
+
 		CHAOS_API void TransferJointConstraintCollisions();
 
 		// Resets VSmooth value to something plausible based on external forces to prevent object from going back to sleep if it was just impulsed.
