Dataflow : Add performance data to subgraphs

diff --git a/Engine/Plugins/Experimental/Dataflow/Source/DataflowEditor/Private/Dataflow/DataflowEditorToolkit.cpp b/Engine/Plugins/Experimental/Dataflow/Source/DataflowEditor/Private/Dataflow/DataflowEditorToolkit.cpp
--- a/Engine/Plugins/Experimental/Dataflow/Source/DataflowEditor/Private/Dataflow/DataflowEditorToolkit.cpp
+++ b/Engine/Plugins/Experimental/Dataflow/Source/DataflowEditor/Private/Dataflow/DataflowEditorToolkit.cpp
@@ -842,6 +842,7 @@ void FDataflowEditorToolkit::ClearGraphCache()
 		if (TSharedPtr<UE::Dataflow::FEngineContext> DataflowContext = EditorContent->GetDataflowContext())
 		{
 			DataflowContext->ClearAllData();
+			DataflowContext->ClearAllPerfData();
 		}
 	}
 }
diff --git a/Engine/Plugins/Experimental/Dataflow/Source/DataflowEditor/Private/Dataflow/DataflowSNode.cpp b/Engine/Plugins/Experimental/Dataflow/Source/DataflowEditor/Private/Dataflow/DataflowSNode.cpp
--- a/Engine/Plugins/Experimental/Dataflow/Source/DataflowEditor/Private/Dataflow/DataflowSNode.cpp
+++ b/Engine/Plugins/Experimental/Dataflow/Source/DataflowEditor/Private/Dataflow/DataflowSNode.cpp
@@ -87,6 +87,8 @@ void SDataflowEdNode::Construct(const FArguments& InArgs, UDataflowEdNode* InNod
 	DataflowGraphNode = Cast<UDataflowEdNode>(InNode);
 	DataflowInterface = InArgs._DataflowInterface;
 
+	this->TitleBorderMargin = FMargin(5.f, 5.f, 5.f, 5.f);
+
 	UpdateGraphNode();
 
 	//
@@ -266,7 +268,17 @@ TArray<FOverlayWidgetInfo> SDataflowEdNode::GetOverlayWidgets(bool bSelected, co
 					FmtOptions.MaximumFractionalDigits = 2;
 					FmtOptions.MinimumFractionalDigits = 2;
 
-					PerfWidget->SetText(FText::Format(LOCTEXT("DataflowNodePerfDataFormat", "Time: {0} ms"), FText::AsNumber(PerfData.ExclusiveTimeMs, &FmtOptions)));
+					FText PerfText;
+					if (PerfData.NumCalls > 1)
+					{
+						PerfText = FText::Format(LOCTEXT("DataflowNodePerfDataFormat", "Time: {0} ms - {1} calls"), FText::AsNumber(PerfData.ExclusiveTimeMs, &FmtOptions), FText::AsNumber(PerfData.NumCalls));
+					}
+					else
+					{
+						PerfText = FText::Format(LOCTEXT("DataflowNodePerfDataFormat", "Time: {0} ms"), FText::AsNumber(PerfData.ExclusiveTimeMs, &FmtOptions));
+					}
+					
+					PerfWidget->SetText(PerfText);
 				}
 				else
 				{
diff --git a/Engine/Source/Runtime/Experimental/Dataflow/Core/Private/Dataflow/DataflowNodeParameters.cpp b/Engine/Source/Runtime/Experimental/Dataflow/Core/Private/Dataflow/DataflowNodeParameters.cpp
--- a/Engine/Source/Runtime/Experimental/Dataflow/Core/Private/Dataflow/DataflowNodeParameters.cpp
+++ b/Engine/Source/Runtime/Experimental/Dataflow/Core/Private/Dataflow/DataflowNodeParameters.cpp
@@ -27,9 +27,24 @@ namespace UE::Dataflow
 			FData& Data = DataPerOutput.FindOrAdd(Connection->GetGuid());
 			Data.ExclusiveTimeMs = FMath::Max(Data.ExclusiveTimeMs, (float)FPlatformTime::ToMilliseconds64(TotalTime - ExternalTime));
 			Data.InclusiveTimeMs = FMath::Max(Data.InclusiveTimeMs, (float)FPlatformTime::ToMilliseconds64(TotalTime));
+			Data.NumCalls++;
 			Data.LastTimestamp = Connection->GetOwningNodeTimestamp();
 		}
 	}
+
+	void FContextPerfData::Append(const FContextPerfData& InPerfData)
+	{
+		for (const TPair<FGuid, FData>& InEntry: InPerfData.DataPerOutput)
+		{
+			const FData& InData = InEntry.Value;
+			FData& Data = DataPerOutput.FindOrAdd(InEntry.Key);
+			Data.ExclusiveTimeMs = FMath::Max(Data.ExclusiveTimeMs, InData.ExclusiveTimeMs);
+			Data.InclusiveTimeMs = FMath::Max(Data.InclusiveTimeMs, InData.InclusiveTimeMs);
+			Data.NumCalls = FMath::Max(Data.NumCalls, InData.NumCalls);
+			Data.LastTimestamp = FMath::Max(Data.LastTimestamp, InData.LastTimestamp);
+		}
+	}
+
 	////////////////////////////////////////////////////////////////////////////////////////////////////////////////////////////////////////////
 
 	void FContextCallstack::Push(const FDataflowConnection* Connection)
@@ -240,6 +255,16 @@ namespace UE::Dataflow
 		return FTimestamp::Invalid;
 	}
 
+	const FContextPerfData& FContext::GetPerfData() const 
+	{ 
+#if DATAFLOW_EDITOR_EVALUATION
+		return PerfData; 
+#else
+		static const FContextPerfData Default;
+		return Default;
+#endif
+	}
+
 	FContextPerfData::FData FContext::GetPerfDataForNode(const FDataflowNode& Node) const
 	{
 		FContextPerfData::FData NodeData;
@@ -252,6 +277,7 @@ namespace UE::Dataflow
 				{
 					NodeData.ExclusiveTimeMs += OutputData->ExclusiveTimeMs;
 					NodeData.InclusiveTimeMs += OutputData->InclusiveTimeMs;
+					NodeData.NumCalls = FMath::Max(NodeData.NumCalls, OutputData->NumCalls);
 				}
 			}
 		}
@@ -279,6 +305,20 @@ namespace UE::Dataflow
 #endif
 	}
 
+	void FContext::AddExternalPerfData(const FContextPerfData& InPerfData)
+	{
+#if DATAFLOW_EDITOR_EVALUATION
+		PerfData.Append(InPerfData);
+#endif
+	}
+
+	void FContext::ClearAllPerfData()
+	{
+#if DATAFLOW_EDITOR_EVALUATION
+		PerfData.Reset();
+#endif
+	}
+
 	void FContext::EnablePerfData(bool bEnable)
 	{
 #if DATAFLOW_EDITOR_EVALUATION
diff --git a/Engine/Source/Runtime/Experimental/Dataflow/Core/Public/Dataflow/DataflowNodeParameters.h b/Engine/Source/Runtime/Experimental/Dataflow/Core/Public/Dataflow/DataflowNodeParameters.h
--- a/Engine/Source/Runtime/Experimental/Dataflow/Core/Public/Dataflow/DataflowNodeParameters.h
+++ b/Engine/Source/Runtime/Experimental/Dataflow/Core/Public/Dataflow/DataflowNodeParameters.h
@@ -46,11 +46,13 @@ namespace UE::Dataflow
 	{
 		void Reset();
 		void Accumulate(const FDataflowConnection* Connection, uint64 TotalTime, uint64 ExternalTime);
+		void Append(const FContextPerfData& InPerfData);
 
 		struct FData
 		{
 			float InclusiveTimeMs = 0;
 			float ExclusiveTimeMs = 0;
+			int32 NumCalls = 0;
 			FTimestamp LastTimestamp = FTimestamp(0); // timestamp of the owner node when data was last collected
 		};
 		TMap<FGuid, FData> DataPerOutput;
@@ -304,6 +306,8 @@ namespace UE::Dataflow
 
 		DATAFLOWCORE_API FContextPerfData::FData GetPerfDataForNode(const FDataflowNode& Node) const;
 		DATAFLOWCORE_API void ResetPerfDataForNode(const FDataflowNode& Node);
+		DATAFLOWCORE_API void AddExternalPerfData(const FContextPerfData& InPerfData);
+		DATAFLOWCORE_API void ClearAllPerfData();
 		DATAFLOWCORE_API void EnablePerfData(bool bEnable);
 		DATAFLOWCORE_API bool IsPerfDataEnabled() const;
 
@@ -341,6 +345,8 @@ namespace UE::Dataflow
 	protected:
 		DATAFLOWCORE_API void BeginContextEvaluation(const FDataflowNode* Node, const FDataflowOutput* Output);
 
+		DATAFLOWCORE_API const FContextPerfData& GetPerfData() const;
+
 	private:
 
 		void CheckIntrinsicInputs(const FDataflowOutput& Connection);
diff --git a/Engine/Source/Runtime/Experimental/Dataflow/Engine/Private/Dataflow/DataflowSubGraphNodes.cpp b/Engine/Source/Runtime/Experimental/Dataflow/Engine/Private/Dataflow/DataflowSubGraphNodes.cpp
--- a/Engine/Source/Runtime/Experimental/Dataflow/Engine/Private/Dataflow/DataflowSubGraphNodes.cpp
+++ b/Engine/Source/Runtime/Experimental/Dataflow/Engine/Private/Dataflow/DataflowSubGraphNodes.cpp
@@ -130,6 +130,13 @@ namespace UE::Dataflow::SubGraphNodes::Private
 		void SetCurrentArrayIndex(int32 Index) { CurrentArrayIndex = Index; }
 		int32 GetCurrentArrayIndex() const { return CurrentArrayIndex; }
 
+		void TransferPerfData()
+		{
+#if DATAFLOW_EDITOR_EVALUATION
+			ParentContext.AddExternalPerfData(GetPerfData());
+#endif
+		}
+
 		virtual UObject* AddAsset(const FString& AssetPath, const UClass* AssetClass) override
 		{
 			return ParentContext.AddAsset(AssetPath, AssetClass);
@@ -642,6 +649,7 @@ void FDataflowCallSubGraphNode::Evaluate(UE::Dataflow::FContext& Context, const
 							}
 							SubGraphContext.ClearAllData();
 						}
+						SubGraphContext.TransferPerfData();
 					}
 					else
 					{
