Do not create the old IAX http thread if we are using the new backend.

#rnx

- FOnDemandContentInstaller now takes an optional pointer to FOnDemandHttpThread rather than a reference.
-- We assert that the thread is valid before all use cases.
- Moved the creation of FOnDemandHttpThread and FOnDemandContentInstaller to FOnDemandIoStore::Initialize which is slightly later in the start up process, but not by much.
- We no only create FOnDemandHttpThread if the old backend is being used as we should not be using it all when the new backend is active.

diff --git a/Engine/Source/Runtime/Experimental/IoStore/OnDemand/Private/OnDemandContentInstaller.cpp b/Engine/Source/Runtime/Experimental/IoStore/OnDemand/Private/OnDemandContentInstaller.cpp
--- a/Engine/Source/Runtime/Experimental/IoStore/OnDemand/Private/OnDemandContentInstaller.cpp
+++ b/Engine/Source/Runtime/Experimental/IoStore/OnDemand/Private/OnDemandContentInstaller.cpp
@@ -244,7 +244,7 @@ TOptional<UE::UnifiedError::FError> FOnDemandContentInstaller::FRequest::Consume
 }
 
 ////////////////////////////////////////////////////////////////////////////////
-FOnDemandContentInstaller::FOnDemandContentInstaller(FOnDemandIoStore& InIoStore, FOnDemandHttpThread& InHttpClient)
+FOnDemandContentInstaller::FOnDemandContentInstaller(FOnDemandIoStore& InIoStore, FOnDemandHttpThread* InHttpClient)
 	: IoStore(InIoStore)
 	, HttpClient(InHttpClient)
 	, InstallerPipe(TEXT("IoStoreOnDemandInstallerPipe"))
@@ -390,10 +390,12 @@ void FOnDemandContentInstaller::UpdateInstallRequestPriority(FSharedInternalInst
 					}
 					else 
 					{
+						check(HttpClient != nullptr);
+
 						TUniqueLock HttpLock(HttpMutex);
 						if (PendingHttpRequest.Handle != nullptr)
 						{
-							HttpClient.ReprioritizeRequest(PendingHttpRequest.Handle, Request.Priority);
+							HttpClient->ReprioritizeRequest(PendingHttpRequest.Handle, Request.Priority);
 						}
 					}
 				}
@@ -794,13 +796,15 @@ void FOnDemandContentInstaller::ExecuteInstallRequest(FRequest& Request)
 	}
 	else
 	{
+		check(HttpClient != nullptr);
+
 		for (FChunkHttpRequestHandle& HttpRequest : InstallRequest.HttpRequestHandles)
 		{
 			TRACE_CPUPROFILER_EVENT_SCOPE(FOnDemandContentInstaller::IssueRequest);
 			FSharedOnDemandContainer& Container		= InstallRequest.ResolvedChunks[HttpRequest.ContainerIndex].Container;
 			const FOnDemandChunkEntry& ChunkEntry	= Container->ChunkEntries[HttpRequest.EntryIndex];
 
-			HttpRequest.Handle = HttpClient.IssueRequest(
+			HttpRequest.Handle = HttpClient->IssueRequest(
 				FOnDemandChunkInfo(Container, ChunkEntry),
 				FIoOffsetAndLength(),
 				Request.Priority,
@@ -973,12 +977,14 @@ void FOnDemandContentInstaller::ProcessDownloadedChunk(
 			}
 			else
 			{
+				check(HttpClient != nullptr);
+
 				TUniqueLock Lock(HttpMutex);
 				for (FChunkHttpRequestHandle& PendingHttpRequest : InstallRequest.HttpRequestHandles)
 				{
 					if (PendingHttpRequest.Handle != nullptr)
 					{
-						HttpClient.CancelRequest(PendingHttpRequest.Handle);
+						HttpClient->CancelRequest(PendingHttpRequest.Handle);
 						PendingHttpRequest.Handle = nullptr;
 						++NumCancelled;
 					}
diff --git a/Engine/Source/Runtime/Experimental/IoStore/OnDemand/Private/OnDemandContentInstaller.h b/Engine/Source/Runtime/Experimental/IoStore/OnDemand/Private/OnDemandContentInstaller.h
--- a/Engine/Source/Runtime/Experimental/IoStore/OnDemand/Private/OnDemandContentInstaller.h
+++ b/Engine/Source/Runtime/Experimental/IoStore/OnDemand/Private/OnDemandContentInstaller.h
@@ -203,7 +203,7 @@ class FOnDemandContentInstaller
 	using FRequestAllocator	= TSingleThreadedSlabAllocator<FRequest, 32>;
 
 public:
-									FOnDemandContentInstaller(FOnDemandIoStore& IoStore, FOnDemandHttpThread& HttpClient);
+									FOnDemandContentInstaller(FOnDemandIoStore& IoStore, FOnDemandHttpThread* HttpClient);
 									~FOnDemandContentInstaller();
 
 	FSharedInternalInstallRequest	EnqueueInstallRequest(
@@ -244,7 +244,7 @@ private:
 	void							Shutdown();
 
 	FOnDemandIoStore&				IoStore;
-	FOnDemandHttpThread&			HttpClient;
+	FOnDemandHttpThread*			HttpClient;
 	UE::Tasks::FPipe				InstallerPipe;
 
 	FMutex							Mutex;
diff --git a/Engine/Source/Runtime/Experimental/IoStore/OnDemand/Private/OnDemandIoStore.cpp b/Engine/Source/Runtime/Experimental/IoStore/OnDemand/Private/OnDemandIoStore.cpp
--- a/Engine/Source/Runtime/Experimental/IoStore/OnDemand/Private/OnDemandIoStore.cpp
+++ b/Engine/Source/Runtime/Experimental/IoStore/OnDemand/Private/OnDemandIoStore.cpp
@@ -200,8 +200,6 @@ FOnDemandIoStore::FOnDemandIoStore()
 	DebugCommands = MakeUnique<FOnDemandDebugCommands>(this);
 #endif // !UE_BUILD_SHIPPING
 
-	HttpClient = MakeUnique<FOnDemandHttpThread>();
-	Installer = MakeUnique<FOnDemandContentInstaller>(*this, *HttpClient);
 	FEncryptionKeyManager::Get().OnKeyAdded().AddRaw(this, &FOnDemandIoStore::OnEncryptionKeyAdded);
 }
 
@@ -389,6 +387,12 @@ FIoStatus FOnDemandIoStore::Initialize()
 		FIoDispatcher::Get().Mount(HttpIoBackend.ToSharedRef(), HttpBackendPriority);
 		FHttpIoDispatcher::OnHostGroupRegistered().AddRaw(this, &FOnDemandIoStore::OnHostGroupRegistered);
 	}
+	else
+	{
+		HttpClient = MakeUnique<FOnDemandHttpThread>();
+	}
+
+	Installer = MakeUnique<FOnDemandContentInstaller>(*this, HttpClient.Get());
 
 	return InstallCacheStatus;
 }
