When using the old IAX backend make sure that the http thread does exist before use.

#rnx

- This code hopefully won't exist too long, but still best to protect against future programming errors just in case.

diff --git a/Engine/Source/Runtime/Experimental/IoStore/OnDemand/Private/OnDemandIoStore.cpp b/Engine/Source/Runtime/Experimental/IoStore/OnDemand/Private/OnDemandIoStore.cpp
--- a/Engine/Source/Runtime/Experimental/IoStore/OnDemand/Private/OnDemandIoStore.cpp
+++ b/Engine/Source/Runtime/Experimental/IoStore/OnDemand/Private/OnDemandIoStore.cpp
@@ -570,6 +570,7 @@ FIoStatus FOnDemandIoStore::InitializeStreamingBackend(const FOnDemandEndpointCo
 		UE_LOG(LogIas, Log, TEXT("File cache disabled - streaming only (zero-quota)"));
 	}
 
+	check(HttpClient.IsValid());
 	StreamingBackend = MakeOnDemandIoDispatcherBackend(EndpointConfig, *this, *HttpClient, MoveTemp(StreamingCache));
 	if (StreamingBackend.IsValid() == false)
 	{
