Chaos - fix comments

#rnx

diff --git a/Engine/Source/Runtime/Experimental/Chaos/Private/PhysicsProxy/JointConstraintProxy.cpp b/Engine/Source/Runtime/Experimental/Chaos/Private/PhysicsProxy/JointConstraintProxy.cpp
--- a/Engine/Source/Runtime/Experimental/Chaos/Private/PhysicsProxy/JointConstraintProxy.cpp
+++ b/Engine/Source/Runtime/Experimental/Chaos/Private/PhysicsProxy/JointConstraintProxy.cpp
@@ -21,8 +21,8 @@ namespace Chaos
 
 FJointConstraintPhysicsProxy::FJointConstraintPhysicsProxy(FJointConstraint* InConstraint, FPBDJointConstraintHandle* InHandle, UObject* InOwner)
 	: Base(EPhysicsProxyType::JointConstraintType, InOwner, MakeShared<FProxyTimestampBase>())
-	, Constraint_GT(InConstraint) // This proxy assumes ownership of the Constraint, and will free it during DestroyOnPhysicsThread
-	, Constraint_PT(InHandle)
+	, Constraint_GT(InConstraint)	// This proxy assumes ownership of the Constraint, and will free it in DestroyOnGameThread
+	, Constraint_PT(InHandle)		// This proxy assumes ownership of the Handle, and will free it in DestroyOnPhysicsThread
 	, OriginalParticleHandles_PT(nullptr, nullptr)
 {
 	check(Constraint_GT !=nullptr);
diff --git a/Engine/Source/Runtime/Experimental/Chaos/Public/Chaos/PBDRigidsEvolutionGBF.h b/Engine/Source/Runtime/Experimental/Chaos/Public/Chaos/PBDRigidsEvolutionGBF.h
--- a/Engine/Source/Runtime/Experimental/Chaos/Public/Chaos/PBDRigidsEvolutionGBF.h
+++ b/Engine/Source/Runtime/Experimental/Chaos/Public/Chaos/PBDRigidsEvolutionGBF.h
@@ -194,7 +194,7 @@ namespace Chaos
 		CHAOS_API void ApplyParticleTransformCorrectionDelta(FGeometryParticleHandle* InParticle, const FVec3& InPosDelta, const FVec3& InRotDelta, const bool bApplyToConnectedBodies, const bool bInRecalculateFrictionOnConnectedBodies = true, const TArray<FParticleID>& ExcludeConnections = TArray<FParticleID>());
 
 		/*
-		 * [EXPERIMENTAL] Similar to SetParticleTransformCorrectionDelta, but supplied an absolute transform to jump to. This is used for snaps.
+		 * [EXPERIMENTAL] Similar to ApplyParticleTransformCorrectionDelta, but supplied an absolute transform to jump to. This is used for snaps.
 		 */
 		CHAOS_API void ApplyParticleTransformCorrection(FGeometryParticleHandle* InParticle, const FVec3& InPos, const FRotation3& InRot, const bool bApplyToConnectedBodies, const bool bInRecalculateFrictionOnConnectedBodies = true, const TArray<FParticleID>& ExcludeConnections = TArray<FParticleID>());
 
