Chaos - change inertia conditioning inertia thresholds back to their original values.

History:
Inertia conditioning is disabled when inverse inertia is "small" - this was intended to prevent divide by zero, but the "small" threshold defaults to 1e-4 which is a very large number compared to typical inverse inertias. Therefore Inertia Conditioning was effectively disabled for medium or large objects. This threshold was fixed to 1e-30, which allowed larger objects to take advantage of inertia conditioning. However, Inertia Conditioning on/off is a significant behavioural change, and specifically affects vehicles if the chassis is thin on one axis. So, we are returning the thresholds back to their original values for now, pending a better idea.
#rnx

diff --git a/Engine/Source/Runtime/Experimental/Chaos/Private/Chaos/MassConditioning.cpp b/Engine/Source/Runtime/Experimental/Chaos/Private/Chaos/MassConditioning.cpp
--- a/Engine/Source/Runtime/Experimental/Chaos/Private/Chaos/MassConditioning.cpp
+++ b/Engine/Source/Runtime/Experimental/Chaos/Private/Chaos/MassConditioning.cpp
@@ -8,12 +8,23 @@
 
 namespace Chaos::CVars
 {
+	// NOTE: ChaosInertiaConditioningInvInertiaTolerance of 1e-4 effectively disables inertia conditioning for medium and larger object.
+	// If you want inertia conditioning on large objects, make ChaosInertiaConditioningInvInertiaTolerance very small (say, 1e-30).
 	float ChaosInertiaConditioningInvMassTolerance = 1.e-20f;
-	float ChaosInertiaConditioningInvInertiaTolerance = 1.e-30f;
+	float ChaosInertiaConditioningInvInertiaTolerance = 1.e-4f;	
 	float ChaosInertiaConditioningExtentTolerance = 1.e-8f;
-	FAutoConsoleVariableRef  CVarChaosInertiaConditioningInvMassTolerance(TEXT("p.Chaos.InertiaConditioning.InvMassTolerance"), ChaosInertiaConditioningInvMassTolerance, TEXT(""));
-	FAutoConsoleVariableRef  CVarChaosInertiaConditioningInvInertiaTolerance(TEXT("p.Chaos.InertiaConditioning.InvInertiaTolerance"), ChaosInertiaConditioningInvInertiaTolerance, TEXT(""));
-	FAutoConsoleVariableRef  CVarChaosInertiaConditioningExtentTolerance(TEXT("p.Chaos.InertiaConditioning.ExtentTolerance"), ChaosInertiaConditioningExtentTolerance, TEXT(""));
+	FAutoConsoleVariableRef  CVarChaosInertiaConditioningInvMassTolerance(
+		TEXT("p.Chaos.InertiaConditioning.InvMassTolerance"), 
+		ChaosInertiaConditioningInvMassTolerance, 
+		TEXT("Inertia conditioning is disabled for objects with inverse mass less than this."));
+	FAutoConsoleVariableRef  CVarChaosInertiaConditioningInvInertiaTolerance(
+		TEXT("p.Chaos.InertiaConditioning.InvInertiaTolerance"), 
+		ChaosInertiaConditioningInvInertiaTolerance, 
+		TEXT("Inertia conditioning is disabled for objects with inverse inertia less than this. The default value disables Inertia Conditioing for mediums to large objects."));
+	FAutoConsoleVariableRef  CVarChaosInertiaConditioningExtentTolerance(
+		TEXT("p.Chaos.InertiaConditioning.ExtentTolerance"), 
+		ChaosInertiaConditioningExtentTolerance, 
+		TEXT("Inertia conditioning is disabled for objects with dimensions smaller than this."));
 }
 
 namespace Chaos
diff --git a/Engine/Source/Runtime/Experimental/Chaos/Private/Chaos/PBDRigidsEvolutionGBF.cpp b/Engine/Source/Runtime/Experimental/Chaos/Private/Chaos/PBDRigidsEvolutionGBF.cpp
--- a/Engine/Source/Runtime/Experimental/Chaos/Private/Chaos/PBDRigidsEvolutionGBF.cpp
+++ b/Engine/Source/Runtime/Experimental/Chaos/Private/Chaos/PBDRigidsEvolutionGBF.cpp
@@ -32,8 +32,6 @@
 #include "ChaosDebugDraw/ChaosDDScene.h"
 #include "ChaosDebugDraw/ChaosDDTimeline.h"
 
-//UE_DISABLE_OPTIMIZATION
-
 namespace Chaos
 {
 #if !UE_BUILD_SHIPPING
@@ -113,8 +111,9 @@ namespace Chaos
 		FRealSingle ChaosSolverInertiaConditioningDistance = 20;
 		FAutoConsoleVariableRef  CVarChaosSolverInertiaConditioningDistance(TEXT("p.Chaos.Solver.InertiaConditioning.Distance"), ChaosSolverInertiaConditioningDistance, TEXT("An input to inertia conditioning system. The joint distance error which needs to be stable (generate a low rotation)."));
 
-		// The ratio of joint error correction that comes from particle rotation versus translation
-		FRealSingle ChaosSolverInertiaConditioningRotationRatio = 2;
+		// The ratio of joint error correction that comes from particle rotation versus translation.
+		// NOTE: A value of 2.5 means that spheres (and cubes) are not effected by inertia conditioning (unless they have constraints outside their bounds).
+		FRealSingle ChaosSolverInertiaConditioningRotationRatio = 2.5f;
 		FAutoConsoleVariableRef  CVarChaosSolverInertiaConditioningRotationRatio(TEXT("p.Chaos.Solver.InertiaConditioning.RotationRatio"), ChaosSolverInertiaConditioningRotationRatio, TEXT("An input to inertia conditioning system. The maximum ratio of joint correction from rotation versus translation"));
 
 		// If > 1, limits the inverse inertia components to be mo more than this multiple of the smallest component. Makes the objects more round. Set to 0 to disable.
