Added the possibility to suspend the IAS cache for a period of time after the IAD memory buffer has been flushed to disk.
- the time to wait is computed from number of bytes written from the IAD cache and the IAS cache write limitation settings
- the time to wait can be adjusted with the cvar ias.CacheSuspendTimeoutFactor=<0.0-1.0> where the default value of zero
means disabled

#rnx

diff --git a/Engine/Source/Runtime/Experimental/IoStore/OnDemand/Private/DiskCacheGovernor.cpp b/Engine/Source/Runtime/Experimental/IoStore/OnDemand/Private/DiskCacheGovernor.cpp
--- a/Engine/Source/Runtime/Experimental/IoStore/OnDemand/Private/DiskCacheGovernor.cpp
+++ b/Engine/Source/Runtime/Experimental/IoStore/OnDemand/Private/DiskCacheGovernor.cpp
@@ -1,7 +1,17 @@
 // Copyright Epic Games, Inc. All Rights Reserved.
 
 #include "DiskCacheGovernor.h"
+#include "HAL/IConsoleManager.h"
 #include "HAL/PlatformTime.h"
+#include "IO/IoStoreOnDemand.h" 
+#include "Misc/ScopeExit.h"
+#include "Statistics.h"
+
+float GIasCacheSuspendTimeoutFactor = 0.0f;
+static FAutoConsoleVariableRef CVar_IasCacheSuspendTimeoutFactor(
+	TEXT("ias.CacheSuspendTimeoutFactor"),
+	GIasCacheSuspendTimeoutFactor,
+	TEXT("Factor between 0 - 1 for adjusting the streaming cache suspend timeout"));
 
 namespace UE::IoStore
 {
@@ -38,6 +48,13 @@ void FDiskCacheGovernor::Configure(FDiskCacheGovernor::FCacheState& CacheState,
 	const FAllowance& Allowance	= Config.Allowance;
 	CacheState.Demand			= Config.Demand;
 
+	ON_SCOPE_EXIT
+	{
+		UE_LOG(LogIoStoreOnDemand, Log, TEXT("Streaming cache write limits, BytesPerPeriod=%u, OpsPerPeriod=%u, Period=%us, Factor=%.2f, BytesPerOp=%u, OpInterval=%ums, MaxOps=%u"),
+			Allowance.Bytes, Allowance.Ops, Allowance.Seconds, GIasCacheSuspendTimeoutFactor,
+			CacheState.OpAllowance, uint32(FPlatformTime::ToMilliseconds64(CacheState.OpInterval)), CacheState.MaxOpCount);
+	};
+
 	if (Config.CommitBufferSize == 0)
 	{
 		CacheState.OpAllowance	= Allowance.Bytes / Allowance.Ops;
@@ -106,6 +123,26 @@ int32 FDiskCacheGovernor::TryBeginWrite(FDiskCacheGovernor::FCacheState& CacheSt
 		return TryBeginWrite(CacheState, DemandPercent, Cycles);
 	}
 
+	// Check if the cache has been suspended
+	if (int64 CyclesSuspended = StreamingCacheState.CyclesSuspended.load(std::memory_order_relaxed); CyclesSuspended > 0)
+	{
+		const int64 Cycles		= FPlatformTime::Cycles64();
+		const int64 Delta		= Cycles - CacheState.PrevCycles;
+		const int64 Remainder	= CyclesSuspended - Delta;
+		const int64 Consumed	= Remainder >= 0 ? Delta : CyclesSuspended;
+		CacheState.PrevCycles	= Cycles;
+
+		check(Consumed >= 0);
+		StreamingCacheState.CyclesSuspended.fetch_sub(Consumed);
+		FOnDemandIoBackendStats::Get()->OnCacheSuspended(
+			FPlatformTime::ToSeconds64(StreamingCacheState.CyclesSuspended.load(std::memory_order_relaxed)));
+
+		if (Remainder > 0)
+		{
+			return 0 - FMath::Min(int32(Remainder), GetMaxWaitCycles(CacheState));
+		}
+	}
+
 	// Wait until the demand gets above this threshold
 	if (DemandPercent < CacheState.Demand.Threshold)
 	{
@@ -125,7 +162,8 @@ int32 FDiskCacheGovernor::EndWrite(FDiskCacheGovernor::FCacheState& CacheState,
 {
 	CacheState.RunOff = UnusedAllowance;
 
-	if (CacheState.OpCount >= CacheState.MaxOpCount)
+	// Suspend writing if we have exceeded the max op count or if the install cache has been flushed to disk
+	if (CacheState.OpCount >= CacheState.MaxOpCount || CacheState.CyclesSuspended.load(std::memory_order_relaxed) > 0)
 	{
 		CacheState.State = FCacheState::EState::Waiting;
 		return 0 - GetMaxWaitCycles(CacheState);
@@ -141,4 +179,23 @@ int32 FDiskCacheGovernor::GetMaxWaitCycles(FDiskCacheGovernor::FCacheState& Cach
 	return int32(CacheState.OpInterval >> 2);
 }
 
+////////////////////////////////////////////////////////////////////////////////
+void FDiskCacheGovernor::OnInstallCacheFlushed(uint64 ByteCount, uint32 OpCount_Unused)
+{
+	const double Factor = FMath::Max(0.0, FMath::Min(1.0, double(GIasCacheSuspendTimeoutFactor)));
+	if (ByteCount == 0 || Factor <= 0.0 || StreamingCacheState.OpAllowance == 0)
+	{
+		return;
+	}
+
+	// Suspend the streaming cache by computing the number of cycles to wait before the next write
+	const double Ops			= double(ByteCount) / double(StreamingCacheState.OpAllowance);
+	const double CyclesPerOp	= double(StreamingCacheState.OpInterval);
+
+	if (int64 Cycles = int64(Ops * CyclesPerOp * Factor); Cycles > 0)
+	{
+		StreamingCacheState.CyclesSuspended.fetch_add(Cycles);
+	}
+}
+
 } // namespace UE::IoStore
diff --git a/Engine/Source/Runtime/Experimental/IoStore/OnDemand/Private/DiskCacheGovernor.h b/Engine/Source/Runtime/Experimental/IoStore/OnDemand/Private/DiskCacheGovernor.h
--- a/Engine/Source/Runtime/Experimental/IoStore/OnDemand/Private/DiskCacheGovernor.h
+++ b/Engine/Source/Runtime/Experimental/IoStore/OnDemand/Private/DiskCacheGovernor.h
@@ -3,6 +3,7 @@
 #pragma once
 
 #include "Misc/CoreMiscDefines.h"
+#include <atomic>
 
 namespace UE::IoStore
 {
@@ -38,6 +39,7 @@ public:
 	void	ConfigureStreamingCache(const FConfig& Config);
 	int32	TryBeginWriteToStreamingCache(uint32 DemandPercent);
 	int32	EndWriteToStreamingCache(uint32 UnusedAllowance);
+	void	OnInstallCacheFlushed(uint64 ByteCount, uint32 OpCount);
 
 private:
 	struct FCacheState;
@@ -56,14 +58,15 @@ private:
 			Rolling,
 		};
 
-		int64	OpInterval = 0;
-		int64	PrevCycles = 0;
-		uint32	RunOff = 0;
-		uint32	OpCount = 0;
-		uint32	MaxOpCount = 0;
-		uint32	OpAllowance = 0;
-		FDemand	Demand;
-		EState	State = EState::Waiting;
+		int64				OpInterval = 0;
+		int64				PrevCycles = 0;
+		uint32				RunOff = 0;
+		uint32				OpCount = 0;
+		uint32				MaxOpCount = 0;
+		uint32				OpAllowance = 0;
+		std::atomic_int64_t	CyclesSuspended{0};
+		FDemand				Demand;
+		EState				State = EState::Waiting;
 	};
 
 	FCacheState StreamingCacheState;
diff --git a/Engine/Source/Runtime/Experimental/IoStore/OnDemand/Private/OnDemandInstallCache.cpp b/Engine/Source/Runtime/Experimental/IoStore/OnDemand/Private/OnDemandInstallCache.cpp
--- a/Engine/Source/Runtime/Experimental/IoStore/OnDemand/Private/OnDemandInstallCache.cpp
+++ b/Engine/Source/Runtime/Experimental/IoStore/OnDemand/Private/OnDemandInstallCache.cpp
@@ -1,6 +1,7 @@
 // Copyright Epic Games, Inc. All Rights Reserved.
 
 #include "OnDemandInstallCache.h"
+#include "DiskCacheGovernor.h"
 #include "OnDemandHttpClient.h"
 #include "OnDemandIoStore.h"
 #include "Statistics.h"
@@ -903,6 +904,7 @@ struct FCasJournal
 	static FTransaction		Begin(FString&& JournalFile);
 	static FTransaction		Begin(const FString& JournalFile) { return Begin(FString(JournalFile)); }
 	static FResult			Commit(FTransaction&& Transaction);
+	static FResult			Commit(FTransaction&& Transaction, uint64& OutByteCount, uint32& OutOpCount);
 };
 
 ///////////////////////////////////////////////////////////////////////////////
@@ -1086,10 +1088,13 @@ FCasJournal::FTransaction FCasJournal::Begin(FString&& JournalFile)
 	};
 }
 
-FResult FCasJournal::Commit(FTransaction&& Transaction)
+FResult FCasJournal::Commit(FTransaction&& Transaction, uint64& OutByteCount, uint32& OutOpCount)
 {
 	using namespace UE::UnifiedError::IoStoreOnDemand;
 
+	OutByteCount	= 0;
+	OutOpCount		= 0;
+
 	if (Transaction.Entries.IsEmpty())
 	{
 		return MakeValue();
@@ -1169,6 +1174,8 @@ FResult FCasJournal::Commit(FTransaction&& Transaction)
 			return Result = MakeJournalError(
 				ECasErrorCode::CommitJournalFailed, EIoErrorCode::WriteError, FString::Printf(TEXT("Failed to write journal entries to '%s'"), *Transaction.JournalFile));
 		}
+		OutOpCount++;
+		OutByteCount += TotalEntrySize;
 
 		FFooter Footer;
 		FMemory::Memcpy(&Footer.Magic, &FFooter::MagicSequence, sizeof(FFooter::MagicSequence));
@@ -1177,12 +1184,15 @@ FResult FCasJournal::Commit(FTransaction&& Transaction)
 			return Result = MakeJournalError(
 				ECasErrorCode::CommitJournalFailed, EIoErrorCode::WriteError, FString::Printf(TEXT("Failed to write journal footer to '%s'"), *Transaction.JournalFile));
 		}
+		OutOpCount++;
+		OutByteCount += uint64(FFooter::Size());
 
 		if (FileHandle->Flush() == false)
 		{
 			return Result = MakeJournalError(
 				ECasErrorCode::CommitJournalFailed, EIoErrorCode::FileFlushFailed, FString::Printf(TEXT("Failed to flush journal entries to '%s'"), *Transaction.JournalFile));
 		}
+		OutOpCount++;
 
 		UE_LOG(LogIoStoreOnDemand, Log, TEXT("Committed %d CAS journal entries of total %.2lf KiB to '%s'"),
 			Transaction.Entries.Num(), ToKiB(TotalEntrySize), *Transaction.JournalFile);
@@ -1192,6 +1202,13 @@ FResult FCasJournal::Commit(FTransaction&& Transaction)
 	}
 }
 
+FResult FCasJournal::Commit(FTransaction&& Transaction)
+{
+	uint64 ByteCount	= 0;
+	uint32 OpCount		= 0;
+	return Commit(MoveTemp(Transaction), ByteCount, OpCount);
+}
+
 void FCasJournal::FTransaction::ChunkLocation(const FCasLocation& Location, const FCasAddr& Addr)
 {
 	Entries.AddZeroed_GetRef().ChunkLocation = FEntry::FChunkLocation
@@ -1699,7 +1716,7 @@ class FOnDemandInstallCache final
 	};
 
 public:
-	FOnDemandInstallCache(const FOnDemandInstallCacheConfig& Config, FOnDemandIoStore& IoStore);
+	FOnDemandInstallCache(const FOnDemandInstallCacheConfig& Config, FOnDemandIoStore& IoStore, FDiskCacheGovernor& Governor);
 	virtual ~FOnDemandInstallCache();
 
 	// IIoDispatcherBackend
@@ -1757,6 +1774,7 @@ private:
 	UE::UnifiedError::IoStoreOnDemand::FInstallCacheErrorContext MakeInstallCacheErrorContext(uint64 TotalCachedBytes = 0, uint32 LineNo = __builtin_LINE());
 
 	FOnDemandIoStore&			IoStore;
+	FDiskCacheGovernor&			Governor;
 	FString						CacheDirectory;
 	FCas						Cas;
 	std::atomic<FCasBlockId>	CurrentBlock{ FCasBlockId::Invalid };
@@ -1778,8 +1796,9 @@ private:
 };
 
 ///////////////////////////////////////////////////////////////////////////////
-FOnDemandInstallCache::FOnDemandInstallCache(const FOnDemandInstallCacheConfig& Config, FOnDemandIoStore& InIoStore)
+FOnDemandInstallCache::FOnDemandInstallCache(const FOnDemandInstallCacheConfig& Config, FOnDemandIoStore& InIoStore, FDiskCacheGovernor& InGovernor)
 	: IoStore(InIoStore)
+	, Governor(InGovernor)
 	, CacheDirectory(Config.RootDirectory)
 	, Cas(Config)
 	, MaxCacheSize(Config.DiskQuota)
@@ -2628,7 +2647,7 @@ FResult FOnDemandInstallCache::Verify()
 
 		if (EntryIndex == INDEX_NONE)
 		{
-			UE_LOG(LogIoStoreOnDemand, Error, TEXT("Failed to find chunk entry for CAS adress '%s'"), *LexToString(ChunkLocation.Addr));
+			UE_LOG(LogIoStoreOnDemand, Error, TEXT("Failed to find chunk entry for CAS address '%s'"), *LexToString(ChunkLocation.Addr));
 
 			MissingChunkCount++;
 			ChunkIndex++;
@@ -3468,11 +3487,14 @@ FResult FOnDemandInstallCache::FlushPendingChunksImpl(const FPendingChunks& Chun
 	using namespace UE::UnifiedError::IoStoreOnDemand;
 
 	FResult Result				= MakeValue();
-	uint64 TotalFlushedBytes	= 0;
+	uint64 TotalCasBytes		= 0;
+	uint64 TotalJournalBytes	= 0;
+	uint32 TotalOpCount			= 0;
 
 	ON_SCOPE_EXIT 
 	{ 
-		FOnDemandInstallCacheStats::OnFlush(Result, TotalFlushedBytes);
+		FOnDemandInstallCacheStats::OnFlush(Result, TotalCasBytes);
+		Governor.OnInstallCacheFlushed(TotalCasBytes + TotalJournalBytes, TotalOpCount);
 	};
 
 	FLargeMemoryWriter Ar(FMath::Min(Chunks.TotalSize, static_cast<uint64>(Cas.GetMaxBlockSize())));
@@ -3526,33 +3548,35 @@ FResult FOnDemandInstallCache::FlushPendingChunksImpl(const FPendingChunks& Chun
 		}
 
 		TConstArrayView<FIoHash> ChunkHashes = MakeArrayView(ChunkHashBegin, Offsets.Num());
+		const int64 BytesToWrite = uint64(Ar.Tell());
 
-		if (Ar.Tell() > 0)
+		if (BytesToWrite > 0)
 		{
 			UE_LOG(LogIoStoreOnDemand, Log, TEXT("Writing %.2lf MiB to CAS block %u"),
 				ToMiB(Ar.Tell()), CurrentBlockId.Id);
 
-			check(CasBlockOffset + Ar.Tell() <= Cas.GetMaxBlockSize());
+			check(CasBlockOffset + BytesToWrite <= Cas.GetMaxBlockSize());
 
-			if (CasFileHandle->Write(Ar.GetData(), Ar.Tell()) == false)
+			if (CasFileHandle->Write(Ar.GetData(), BytesToWrite) == false)
 			{
 				Result = MakeCasError<void>(
-					ECasErrorCode::WriteBlockFailed, EIoErrorCode::WriteError, FString::Printf(TEXT("Failed to write %" INT64_FMT " bytes to CAS block %u"), Ar.Tell(), CurrentBlockId.Id));
+					ECasErrorCode::WriteBlockFailed, EIoErrorCode::WriteError, FString::Printf(TEXT("Failed to write %" INT64_FMT " bytes to CAS block %u"), BytesToWrite, CurrentBlockId.Id));
 				FInstallCacheErrorContext Ctx = MakeInstallCacheErrorContext();
 				Result.GetError().PushErrorContext(MoveTemp(Ctx));
 				return Result;
 			}
+			TotalOpCount++;
+			TotalCasBytes += uint64(BytesToWrite);
 
 			if (CasFileHandle->Flush() == false)
 			{
 				Result = MakeCasError<void>(
-					ECasErrorCode::WriteBlockFailed, EIoErrorCode::FileFlushFailed, FString::Printf(TEXT("Failed to flush %" INT64_FMT " bytes to CAS block %u"), Ar.Tell(), CurrentBlockId.Id));
+					ECasErrorCode::WriteBlockFailed, EIoErrorCode::FileFlushFailed, FString::Printf(TEXT("Failed to flush %" INT64_FMT " bytes to CAS block %u"), BytesToWrite, CurrentBlockId.Id));
 				FInstallCacheErrorContext Ctx = MakeInstallCacheErrorContext();
 				Result.GetError().PushErrorContext(MoveTemp(Ctx));
 				return Result;
 			}
-
-			TotalFlushedBytes += Ar.Tell();
+			TotalOpCount++;
 
 			constexpr const bool bDirty = false;
 			if (UtcAccessTicks)
@@ -3588,13 +3612,18 @@ FResult FOnDemandInstallCache::FlushPendingChunksImpl(const FPendingChunks& Chun
 			Ar.Seek(0);
 		}
 
-		if (Result = FCasJournal::Commit(MoveTemp(Transaction)); Result.HasError())
+		uint64 JrnlBytesWritten	= 0;
+		uint32 JrnlOps			= 0;
+		if (Result = FCasJournal::Commit(MoveTemp(Transaction), JrnlBytesWritten, JrnlOps); Result.HasError())
 		{
 			FInstallCacheErrorContext Ctx = MakeInstallCacheErrorContext();
 			Result.GetError().PushErrorContext(MoveTemp(Ctx));
 			return Result;
 		}
 
+		TotalJournalBytes	+= JrnlBytesWritten;
+		TotalOpCount		+= JrnlOps;
+
 		if (ChunkIdx < Chunks.Chunks.Num())
 		{
 			CurrentBlockId = FCasBlockId::Invalid;
@@ -3671,7 +3700,8 @@ UE::UnifiedError::IoStoreOnDemand::FInstallCacheErrorContext FOnDemandInstallCac
 ///////////////////////////////////////////////////////////////////////////////
 TSharedPtr<IOnDemandInstallCache> MakeOnDemandInstallCache(
 	FOnDemandIoStore& IoStore,
-	const FOnDemandInstallCacheConfig& Config)
+	const FOnDemandInstallCacheConfig& Config,
+	FDiskCacheGovernor& Governor)
 {
 	IFileManager& Ifm = IFileManager::Get();
 	if (Config.bDropCache)
@@ -3687,7 +3717,7 @@ TSharedPtr<IOnDemandInstallCache> MakeOnDemandInstallCache(
 		return TSharedPtr<IOnDemandInstallCache>();
 	}
 
-	return MakeShared<FOnDemandInstallCache>(Config, IoStore);
+	return MakeShared<FOnDemandInstallCache>(Config, IoStore, Governor);
 }
 
 ///////////////////////////////////////////////////////////////////////////////
diff --git a/Engine/Source/Runtime/Experimental/IoStore/OnDemand/Private/OnDemandInstallCache.h b/Engine/Source/Runtime/Experimental/IoStore/OnDemand/Private/OnDemandInstallCache.h
--- a/Engine/Source/Runtime/Experimental/IoStore/OnDemand/Private/OnDemandInstallCache.h
+++ b/Engine/Source/Runtime/Experimental/IoStore/OnDemand/Private/OnDemandInstallCache.h
@@ -57,6 +57,7 @@ struct FOnDemandInstallCacheConfig
 
 TSharedPtr<IOnDemandInstallCache> MakeOnDemandInstallCache(
 	FOnDemandIoStore& IoStore,
-	const FOnDemandInstallCacheConfig& Config);
+	const FOnDemandInstallCacheConfig& Config,
+	class FDiskCacheGovernor& Governor);
 
 } // namespace UE::IoStore
diff --git a/Engine/Source/Runtime/Experimental/IoStore/OnDemand/Private/OnDemandIoStore.cpp b/Engine/Source/Runtime/Experimental/IoStore/OnDemand/Private/OnDemandIoStore.cpp
--- a/Engine/Source/Runtime/Experimental/IoStore/OnDemand/Private/OnDemandIoStore.cpp
+++ b/Engine/Source/Runtime/Experimental/IoStore/OnDemand/Private/OnDemandIoStore.cpp
@@ -1047,7 +1047,7 @@ FIoStatus FOnDemandIoStore::InitializeInstallCache()
 		}
 	}
 
-	InstallCache = MakeOnDemandInstallCache(*this, OnDemandInstallCacheConfig);
+	InstallCache = MakeOnDemandInstallCache(*this, OnDemandInstallCacheConfig, DiskCacheGovernor);
 	if (InstallCache.IsValid())
 	{
 		int32 BackendPriority = -5; // Lower than file (zero) but higher than streaming backend (-10)
diff --git a/Engine/Source/Runtime/Experimental/IoStore/OnDemand/Private/Statistics.cpp b/Engine/Source/Runtime/Experimental/IoStore/OnDemand/Private/Statistics.cpp
--- a/Engine/Source/Runtime/Experimental/IoStore/OnDemand/Private/Statistics.cpp
+++ b/Engine/Source/Runtime/Experimental/IoStore/OnDemand/Private/Statistics.cpp
@@ -361,6 +361,7 @@ FCounterAtomicInt		GCachePutExistingCount(TEXT("Ias/CachePutExistingCount"), Tra
 FCounterAtomicInt		GCachePutRejectCount(TEXT("Ias/CachePutRejectCount"), TraceCounterDisplayHint_None);
 FCounterAtomicInt		GCacheCachedBytes(TEXT("Ias/CacheCachedBytes"), TraceCounterDisplayHint_Memory);
 FCounterAtomicInt		GCacheWrittenBytes(TEXT("Ias/CacheWrittenBytes"), TraceCounterDisplayHint_Memory);
+FCounterFloat			GCacheSuspendedSeconds(TEXT("Ias/CacheSuspendedSeconds"), TraceCounterDisplayHint_None);
 int64					GCacheMaxBytes = 0;
 FCounterAtomicInt		GCachePendingBytes(TEXT("Ias/CachePendingBytes"), TraceCounterDisplayHint_Memory);
 FCounterAtomicInt		GCacheReadBytes(TEXT("Ias/CacheReadBytes"), TraceCounterDisplayHint_Memory);
@@ -829,6 +830,11 @@ void FOnDemandIoBackendStats::PrintOnScreenStatistics(FCoreDelegates::FSeverityM
 			CachingText << TEXT(" | Read: ") << FText::AsMemory(GCacheReadBytes.Get()).ToString();
 			CachingText << TEXT(" (") << GCacheGetCount.Get() << TEXT(")");
 			CachingText << TEXT(" (Boot: ") << GCacheBootMs << TEXT("ms)");
+			const double SuspendedSeconds = GCacheSuspendedSeconds.Get();
+			if (SuspendedSeconds > 0.0)
+			{
+				CachingText << FString::Printf(TEXT(" (Suspended: %.2lfs)"), SuspendedSeconds);
+			}
 		}
 		else
 		{
@@ -1204,6 +1210,11 @@ void FOnDemandIoBackendStats::OnCacheSetMaxBytes(uint64 TotalSize)
 	GCacheMaxBytes = TotalSize;
 }
 
+void FOnDemandIoBackendStats::OnCacheSuspended(double Seconds)
+{
+	GCacheSuspendedSeconds.Set(Seconds);
+}
+
 void FOnDemandIoBackendStats::OnHttpDistributedEndpointResolved()
 {
 	GHttpDistributedEndpointResolved = true;
diff --git a/Engine/Source/Runtime/Experimental/IoStore/OnDemand/Private/Statistics.h b/Engine/Source/Runtime/Experimental/IoStore/OnDemand/Private/Statistics.h
--- a/Engine/Source/Runtime/Experimental/IoStore/OnDemand/Private/Statistics.h
+++ b/Engine/Source/Runtime/Experimental/IoStore/OnDemand/Private/Statistics.h
@@ -83,6 +83,7 @@ public:
 	void OnCachePersistedBytes(uint64 TotalSize) IAS_STATISTICS_IMPL()
 	void OnCacheWriteBytes(uint64 WriteSize) IAS_STATISTICS_IMPL()
 	void OnCacheSetMaxBytes(uint64 TotalSize) IAS_STATISTICS_IMPL()
+	void OnCacheSuspended(double Seconds) IAS_STATISTICS_IMPL()
 
 	// Per session
 	void OnHttpDistributedEndpointResolved() IAS_STATISTICS_IMPL()
