Fix use-after-free.

#rnx

diff --git a/Engine/Source/Runtime/Experimental/IoStore/HttpClient/Private/TransactionOne.inl b/Engine/Source/Runtime/Experimental/IoStore/HttpClient/Private/TransactionOne.inl
--- a/Engine/Source/Runtime/Experimental/IoStore/HttpClient/Private/TransactionOne.inl
+++ b/Engine/Source/Runtime/Experimental/IoStore/HttpClient/Private/TransactionOne.inl
@@ -136,7 +136,7 @@ namespace DetailOne
 class FBase
 {
 public:
-					FBase(FBuffer& InBuffer);
+					FBase(FBuffer&);
 
 protected:
 	using FMutableSection = FBuffer::FMutableSection;
@@ -147,12 +147,13 @@ protected:
 	void			AdvanceUsed(uint32 Delta);
 
 private:
-	FBuffer&		Buffer;
+	FBuffer			Buffer;
+	char			Data[256];
 };
 
 ////////////////////////////////////////////////////////////////////////////////
-FBase::FBase(FBuffer& InBuffer)
-: Buffer(InBuffer)
+FBase::FBase(FBuffer&)
+: Buffer(Data, sizeof(Data))
 {
 }
 
