Appease PVS

#rnx

diff --git a/Engine/Source/Runtime/Experimental/IoStore/HttpClient/Private/TransactionOne.inl b/Engine/Source/Runtime/Experimental/IoStore/HttpClient/Private/TransactionOne.inl
--- a/Engine/Source/Runtime/Experimental/IoStore/HttpClient/Private/TransactionOne.inl
+++ b/Engine/Source/Runtime/Experimental/IoStore/HttpClient/Private/TransactionOne.inl
@@ -153,7 +153,7 @@ private:
 
 ////////////////////////////////////////////////////////////////////////////////
 FBase::FBase(FBuffer&)
-: Buffer(Data, sizeof(Data))
+: Buffer(Data, sizeof(Data)) // -V670
 {
 }
 
