Bughawk fixes.

#rnx

diff --git a/Engine/Source/Runtime/Experimental/IoStore/OnDemand/Private/DiskCacheGovernor.cpp b/Engine/Source/Runtime/Experimental/IoStore/OnDemand/Private/DiskCacheGovernor.cpp
--- a/Engine/Source/Runtime/Experimental/IoStore/OnDemand/Private/DiskCacheGovernor.cpp
+++ b/Engine/Source/Runtime/Experimental/IoStore/OnDemand/Private/DiskCacheGovernor.cpp
@@ -124,7 +124,7 @@ int32 FDiskCacheGovernor::TryBeginWrite(FDiskCacheGovernor::FCacheState& CacheSt
 	}
 
 	// Check if the cache has been suspended
-	if (int64 CyclesSuspended = StreamingCacheState.CyclesSuspended.load(std::memory_order_relaxed); CyclesSuspended > 0)
+	if (int64 CyclesSuspended = CacheState.CyclesSuspended.load(std::memory_order_relaxed); CyclesSuspended > 0)
 	{
 		const int64 Cycles		= FPlatformTime::Cycles64();
 		const int64 Delta		= Cycles - CacheState.PrevCycles;
@@ -133,9 +133,9 @@ int32 FDiskCacheGovernor::TryBeginWrite(FDiskCacheGovernor::FCacheState& CacheSt
 		CacheState.PrevCycles	= Cycles;
 
 		check(Consumed >= 0);
-		StreamingCacheState.CyclesSuspended.fetch_sub(Consumed);
+		CacheState.CyclesSuspended.fetch_sub(Consumed);
 		FOnDemandIoBackendStats::Get()->OnCacheSuspended(
-			FPlatformTime::ToSeconds64(StreamingCacheState.CyclesSuspended.load(std::memory_order_relaxed)));
+			FPlatformTime::ToSeconds64(CacheState.CyclesSuspended.load(std::memory_order_relaxed)));
 
 		if (Remainder > 0)
 		{
@@ -173,7 +173,7 @@ int32 FDiskCacheGovernor::EndWrite(FDiskCacheGovernor::FCacheState& CacheState,
 }
 
 ////////////////////////////////////////////////////////////////////////////////
-int32 FDiskCacheGovernor::GetMaxWaitCycles(FDiskCacheGovernor::FCacheState& CacheState) const
+int32 FDiskCacheGovernor::GetMaxWaitCycles(FDiskCacheGovernor::FCacheState& CacheState)
 {
 	// ">> 2" so we check at four times the speed in case of a S.U.P.E.R BOOST
 	return int32(CacheState.OpInterval >> 2);
diff --git a/Engine/Source/Runtime/Experimental/IoStore/OnDemand/Private/DiskCacheGovernor.h b/Engine/Source/Runtime/Experimental/IoStore/OnDemand/Private/DiskCacheGovernor.h
--- a/Engine/Source/Runtime/Experimental/IoStore/OnDemand/Private/DiskCacheGovernor.h
+++ b/Engine/Source/Runtime/Experimental/IoStore/OnDemand/Private/DiskCacheGovernor.h
@@ -44,11 +44,11 @@ public:
 private:
 	struct FCacheState;
 
-	void	Configure(FCacheState& CacheState, const FConfig& Config);
-	int32	TryBeginWrite(FCacheState& CacheState, uint32 DemandPercent, int64 Cycles);
-	int32	TryBeginWrite(FCacheState& CacheState, uint32 DemandPercent);
-	int32	EndWrite(struct FCacheState& CacheState, uint32 UnusedAllowance);
-	int32	GetMaxWaitCycles(FCacheState& CacheState) const;
+	static void		Configure(FCacheState& CacheState, const FConfig& Config);
+	static int32	TryBeginWrite(FCacheState& CacheState, uint32 DemandPercent, int64 Cycles);
+	static int32	TryBeginWrite(FCacheState& CacheState, uint32 DemandPercent);
+	static int32	EndWrite(struct FCacheState& CacheState, uint32 UnusedAllowance);
+	static int32	GetMaxWaitCycles(FCacheState& CacheState);
 
 	struct FCacheState
 	{
