Workaround to fix a crash in TSpatialHashGridPoints (used by cloth self collision spheres) where the compiler is not correctly aligning the CellSize vector register.

diff --git a/Engine/Source/Runtime/Experimental/Chaos/Public/Chaos/HierarchicalSpatialHash.h b/Engine/Source/Runtime/Experimental/Chaos/Public/Chaos/HierarchicalSpatialHash.h
--- a/Engine/Source/Runtime/Experimental/Chaos/Public/Chaos/HierarchicalSpatialHash.h
+++ b/Engine/Source/Runtime/Experimental/Chaos/Public/Chaos/HierarchicalSpatialHash.h
@@ -11,67 +11,6 @@
 
 namespace Chaos
 {
-struct UE_DEPRECATED(5.4, "This is deprecated and its replacement is private.") FHierarchicalSpatialHashCellIdx
-{
-	TVec3<int32> Cell;
-	int32 Lod;
-
-	bool operator==(const FHierarchicalSpatialHashCellIdx& OtherIdx) const
-	{
-		return Cell[0] == OtherIdx.Cell[0] && Cell[1] == OtherIdx.Cell[1] && Cell[2] == OtherIdx.Cell[2] && Lod == OtherIdx.Lod;
-	}
-	
-	template<typename T>
-	static int32 LodForCellSize(T CellSize)
-	{
-		// want Ceil( log2(CellSize) )
-		if (CellSize <= (T).5)
-		{
-			return -(int32)FMath::FloorLog2((uint32)FMath::FloorToInt32((T)1 / CellSize));
-		}
-		else
-		{
-			return (int32)FMath::CeilLogTwo((uint32)FMath::CeilToInt32(CellSize));
-		}
-	}
-
-	template<typename T>
-	static int32 LodForAABB(const TAABB<T, 3>& Box)
-	{
-		const T LongestSideLen = Box.Extents().Max();
-		return LodForCellSize(LongestSideLen);
-	}
-
-	template<typename T>
-	static T CellSizeForLod(const int32 Lod)
-	{
-		if (Lod >= 0)
-		{
-			return (T)(1ULL << Lod);
-		}
-		return (T)1 / (T)(1ULL << (-Lod));
-	}
-
-	template<typename T>
-	static TVec3<int32> CellIdxForPoint(const TVec3<T>& Point, const T CellSize)
-	{
-		return { FMath::FloorToInt32(Point[0] / CellSize), FMath::FloorToInt32(Point[1] / CellSize), FMath::FloorToInt32(Point[2] / CellSize) };
-	}
-};
-
-UE_DEPRECATED(5.4, "This is deprecated and its replacement is private.")
-PRAGMA_DISABLE_DEPRECATION_WARNINGS
-inline uint32 GetTypeHash(const Chaos::FHierarchicalSpatialHashCellIdx& Idx)
-{
-	// xorHash
-	int32 Hash = Idx.Cell[0] * 73856093;
-	Hash = Hash ^ (Idx.Cell[1] * 19349663);
-	Hash = Hash ^ (Idx.Cell[2] * 83492791);
-	Hash = Hash ^ (Idx.Lod * 67867979);
-	return (uint32)Hash;
-}
-PRAGMA_ENABLE_DEPRECATION_WARNINGS
-
 template<typename PayloadType>
 class TSpatialHashGridBase
 {
@@ -706,7 +645,7 @@ public:
 		return ConcurrentResults.Result;
 	}
 private:
-	const TVectorRegisterType<T> CellSize;
+	alignas(32) const TVectorRegisterType<T> CellSize; // alignas is a workaround to a crash where clang isn't automatically doing the aligning.
 
 	// Particle Index is dense array index passed in at build time. 
 	TArray<int32> ParticleIndexToElement;
