The new IAX backend now supports 'iax.HttpPause' and 'iax.HttpUnpause' cvars

#rnx

- The cvars work in exactly the same way that they did for the old backend.
- Added statistic tracking (new backend only) for paused requests which is printed to the 'ias.DisplayOnScreenStatistics' display
-- Note that when pausing requests the number of paused requests might be higher than the number of pending http requests that you were seeing. This is because every request within the same chunk can be combined into a single http request but each one will be shown as being a paused request.
- Implementation for pausing requests is entirely hidden in the FHttpQueue rather than the backend itself.
- FHttpQueue::Enqueue logic is moved to ::EnqueueInternal so that it can be called when unpausing requests without risking a recursive lock on the mutex.
- Moved the amount of slack to leave in the containers to a constant FHttpQueue::ContainerSlackSize to make it easier to modify.
-- Also reserve that amount of space in each container.
- Added documentation to FHttpQueue::Dequeue to clear up a misconception of what is actually returned.

diff --git a/Engine/Source/Runtime/Experimental/IoStore/OnDemand/Private/OnDemandHttpIoDispacher.cpp b/Engine/Source/Runtime/Experimental/IoStore/OnDemand/Private/OnDemandHttpIoDispacher.cpp
--- a/Engine/Source/Runtime/Experimental/IoStore/OnDemand/Private/OnDemandHttpIoDispacher.cpp
+++ b/Engine/Source/Runtime/Experimental/IoStore/OnDemand/Private/OnDemandHttpIoDispacher.cpp
@@ -223,23 +223,60 @@ void DebugPrintRange(FHttpRequest* Requests, const UE::FIoHttpRange& TotalRange)
 class FHttpQueue
 {
 public:
+	FHttpQueue();
+	~FHttpQueue() = default;
+
 	void			Reprioritize(FHttpRequest* Request, int32 NewPriority);
 	bool			Cancel(FHttpRequest* Request);
 	void			Enqueue(FHttpRequest* Request);
+	/**
+	* Returns one or more requests. Multiple requests will returned as a linked list (FHttpRequest::Next)
+	* and will all be for the same chunk.
+	*/
 	FHttpRequest*	Dequeue();
 
+#if UE_ALLOW_HTTP_PAUSE
+	void			OnTogglePause(bool bPause, EHttpRequestTypeFilter Filter);
+	bool			IsPaused(FHttpRequest* Request) const;
+#endif // UE_ALLOW_HTTP_PAUSE
+
 private:
+
+	void			EnqueueInternal(FHttpRequest* Request);
+
 	TArray<FHttpRequest*>			Heap;
 	TMap<FIoHash, FHttpRequestList>	ByChunkKey;
 	FMutex							Mutex;
+#if UE_ALLOW_HTTP_PAUSE
+	EHttpRequestTypeFilter			PauseFilter = EHttpRequestTypeFilter::None;
+	TArray<FHttpRequest*>			PausedRequests;
+#endif // UE_ALLOW_HTTP_PAUSE
 	bool							bReprioritize = false;
+
+	/** Slack to leave in containers when emptying them */
+	static constexpr int32			ContainerSlackSize = 16;
 };
 
 ////////////////////////////////////////////////////////////////////////////////
+
+FHttpQueue::FHttpQueue()
+{
+	Heap.Reserve(ContainerSlackSize);
+	ByChunkKey.Reserve(ContainerSlackSize);
+}
+
 void FHttpQueue::Reprioritize(FHttpRequest* Request, int32 NewPriority)
 {
 	TUniqueLock Lock(Mutex);
 
+#if UE_ALLOW_HTTP_PAUSE
+	if (IsPaused(Request))
+	{
+		Request->Options.SetPriority(NewPriority);
+		return;
+	}
+#endif //UE_ALLOW_HTTP_PAUSE
+
 	const bool bGenericHttpRequest = Request->ChunkHash.IsZero();
 	if (bGenericHttpRequest)
 	{
@@ -272,6 +309,21 @@ bool FHttpQueue::Cancel(FHttpRequest* Request)
 {
 	TUniqueLock Lock(Mutex);
 
+#if UE_ALLOW_HTTP_PAUSE
+	if (IsPaused(Request))
+	{
+		if(PausedRequests.Remove(Request) != 0)
+		{
+			// Cancel stats assume that the request was enqueued, so we have to add that back
+			FOnDemandIoBackendStats::Get()->OnHttpEnqueue(Request->GetHttpRequestType());
+
+			FOnDemandIoBackendStats::Get()->OnHttpCancel(Request->GetHttpRequestType());
+			FOnDemandIoBackendStats::Get()->OnHttpUnpaused(Request->GetHttpRequestType());
+			return true;
+		}
+	}
+#endif //UE_ALLOW_HTTP_PAUSE
+
 	const bool bGenericHttpRequest = Request->ChunkHash.IsZero();
 	if (bGenericHttpRequest)
 	{
@@ -325,6 +377,23 @@ void FHttpQueue::Enqueue(FHttpRequest* Request)
 {
 	TUniqueLock Lock(Mutex);
 
+#if UE_ALLOW_HTTP_PAUSE
+	if (IsPaused(Request))
+	{
+		PausedRequests.Add(Request);
+		FOnDemandIoBackendStats::Get()->OnHttpPaused(Request->GetHttpRequestType());
+
+		return;
+	}
+#endif //UE_ALLOW_HTTP_PAUSE
+
+	EnqueueInternal(Request);
+}
+
+void FHttpQueue::EnqueueInternal(FHttpRequest* Request)
+{
+	check(Request->Next == nullptr);
+
 	Request->Flags.Remove(EHttpRequestFlags::CacheInflight);
 
 	const bool bGenericHttpRequest = Request->ChunkHash.IsZero();
@@ -384,7 +453,7 @@ FHttpRequest* FHttpQueue::Dequeue()
 	Heap.HeapPop(Next, HttpRequestSortPredicate, EAllowShrinking::No);
 	if (Heap.IsEmpty())
 	{
-		Heap.Empty(16);
+		Heap.Empty(ContainerSlackSize);
 	}
 
 	if (Next == nullptr)
@@ -400,7 +469,7 @@ FHttpRequest* FHttpQueue::Dequeue()
 		Next = List->PeekHead();
 		if (ByChunkKey.IsEmpty())
 		{
-			ByChunkKey.Empty(16);
+			ByChunkKey.Empty(ContainerSlackSize);
 		}
 	}
 
@@ -418,6 +487,123 @@ FHttpRequest* FHttpQueue::Dequeue()
 	return Next;
 }
 
+#if UE_ALLOW_HTTP_PAUSE
+
+void FHttpQueue::OnTogglePause(bool bPause, EHttpRequestTypeFilter Filter)
+{
+	UE::TUniqueLock _(Mutex);
+
+	if (bPause)
+	{
+		EnumAddFlags(PauseFilter, Filter);
+
+		// Remove any queued requests which now be paused
+		TArray<FHttpRequest*> NewHeap;
+		for (FHttpRequest* HeadRequest :Heap)
+		{
+			if (FHttpRequestList* List = ByChunkKey.Find(HeadRequest->ChunkHash))
+			{
+				bool bRemovedHeadRequest = false;
+
+				// Strip paused requests from the list
+				FHttpRequest* Iterator = List->PeekHead();
+				while (Iterator != nullptr)
+				{
+					FHttpRequest* ListRequest = Iterator;
+					Iterator = ListRequest->Next;
+
+					if (IsPaused(ListRequest))
+					{
+						if (ListRequest == HeadRequest)
+						{
+							bRemovedHeadRequest = true;
+						}
+
+						List->Remove(ListRequest);
+
+						check(ListRequest->Next == nullptr);
+						PausedRequests.Add(ListRequest);
+						FOnDemandIoBackendStats::Get()->OnHttpPaused(ListRequest->GetHttpRequestType());
+					}
+				}
+
+				// If the list is now empty we can remove the entire request from Heap
+				if (List->IsEmpty())
+				{
+					FOnDemandIoBackendStats::Get()->OnHttpOnRemovedPending(HeadRequest->GetHttpRequestType());
+					ByChunkKey.Remove(HeadRequest->ChunkHash);
+				}
+				// If we removed the head request we need to replace it in Heap with a non paused request
+				else if (bRemovedHeadRequest)
+				{
+					FOnDemandIoBackendStats::Get()->OnHttpOnRemovedPending(HeadRequest->GetHttpRequestType());
+
+					FHttpRequest* NewHeadRequest = List->PeekHead();
+					NewHeap.HeapPush(NewHeadRequest, HttpRequestSortPredicate);
+					FOnDemandIoBackendStats::Get()->OnHttpEnqueue(NewHeadRequest->GetHttpRequestType());
+				}
+			}
+			else if (IsPaused(HeadRequest))
+			{
+				check(HeadRequest->Next == nullptr);
+				PausedRequests.Add(HeadRequest);
+
+				FOnDemandIoBackendStats::Get()->OnHttpOnRemovedPending(HeadRequest->GetHttpRequestType());
+				FOnDemandIoBackendStats::Get()->OnHttpPaused(HeadRequest->GetHttpRequestType());
+			}
+			else
+			{
+				NewHeap.HeapPush(HeadRequest, HttpRequestSortPredicate);
+			}
+		}
+
+		Heap = MoveTemp(NewHeap);
+	}
+	else
+	{
+		EnumRemoveFlags(PauseFilter, Filter);
+
+		for (int32 Index = 0; Index < PausedRequests.Num(); ++Index)
+		{
+			FHttpRequest* Request = PausedRequests[Index];
+
+			if (IsPaused(Request) == false)
+			{
+				// Add the newly unpaused request back to the queue
+				EnqueueInternal(Request);
+				FOnDemandIoBackendStats::Get()->OnHttpUnpaused(Request->GetHttpRequestType());
+
+				PausedRequests.RemoveAt(Index);
+				--Index;
+			}
+		}
+	}
+}
+
+bool FHttpQueue::IsPaused(FHttpRequest* Request) const
+{
+	check(Request != nullptr);
+
+	switch (Request->GetHttpRequestType())
+	{
+		case EHttpRequestType::Installed:
+		{
+			return EnumHasAnyFlags(PauseFilter, EHttpRequestTypeFilter::Installed);
+		}
+		case EHttpRequestType::Streaming:
+		{
+			return EnumHasAnyFlags(PauseFilter, EHttpRequestTypeFilter::Streaming);
+		}
+		default:
+		{
+			checkNoEntry();
+			return false;
+		}
+	}
+}
+
+#endif // UE_ALLOW_HTTP_PAUSE
+
 ////////////////////////////////////////////////////////////////////////////////
 class FHttpDispatcher final
 	: public FRunnable
@@ -476,6 +662,10 @@ private:
 	void							CompleteHttpRequest(FHttpRequest& Request, FMultiEndpointHttpClientResponse&& HttpResponse, FIASHostGroup& HostGroup, const FIoHttpRange& TotalRange);
 	void							CompleteRequest(FHttpRequest& Request, FIoBuffer& Body, const FIoHttpRange& TotalRange, uint32 StatusCode, bool bCached);
 
+#if UE_ALLOW_HTTP_PAUSE
+	void OnPauseCommand(bool bShouldPause, const TArray<FString>& Args, FOutputDevice& Ar);
+#endif //UE_ALLOW_HTTP_PAUSE
+
 	TUniquePtr<IIasCache>					Cache;
 	TUniquePtr<FMultiEndpointHttpClient>	HttpClient;
 	TUniquePtr<FRunnableThread>				Thread;
@@ -484,6 +674,10 @@ private:
 	FHostGroupRegistered					HostGroupRegisterd;
 	FEventRef								WakeUp;
 	FMutex									Mutex;
+#if UE_ALLOW_HTTP_PAUSE
+	FConsoleCommandPtr						PauseCommand;
+	FConsoleCommandPtr						UnpauseCommand;
+#endif // UE_ALLOW_HTTP_PAUSE
 	FDelegateHandle							OnRecreateHttpClientHandle;
 	std::atomic_bool						bRecreateHttpClient = false;
 	std::atomic_bool						bStopRequested = false;
@@ -495,6 +689,26 @@ static TUniquePtr<FHttpDispatcher> HttpDispatcher;
 FHttpDispatcher::FHttpDispatcher(TUniquePtr<IIasCache>&& InCache)
 	: Cache(MoveTemp(InCache))
 {
+#if UE_ALLOW_HTTP_PAUSE
+	PauseCommand = FConsoleCommandPtr(IConsoleManager::Get().RegisterConsoleCommand(
+		TEXT("iax.HttpPause"),
+		TEXT("Pause all http requests. Passing in 'IAD' or 'IAD' as an arg to only pause requests of that type"),
+		FConsoleCommandWithArgsAndOutputDeviceDelegate::CreateLambda([this](const TArray<FString>& Args, FOutputDevice& Ar) -> void
+			{
+				OnPauseCommand(/*bShouldPause*/ true, Args, Ar);
+			}),
+		ECVF_Default));
+
+	UnpauseCommand = FConsoleCommandPtr(IConsoleManager::Get().RegisterConsoleCommand(
+		TEXT("iax.HttpUnpause"),
+		TEXT("Unause all http requests. Passing in 'IAD' or 'IAD' as an arg to only unpause requests of that type"),
+		FConsoleCommandWithArgsAndOutputDeviceDelegate::CreateLambda([this](const TArray<FString>& Args, FOutputDevice& Ar) -> void
+			{
+				OnPauseCommand(/*bShouldPause*/ false, Args, Ar);
+			}),
+		ECVF_Default));
+#endif // UE_ALLOW_HTTP_PAUSE
+
 	const int32 ThreadPriorityIndex			= FMath::Clamp(GOnDemandBackendThreadPriorityIndex, 0, (int32)UE_ARRAY_COUNT(GThreadPriorities) - 1);
 	EThreadPriority DesiredThreadPriority	= (EThreadPriority)GThreadPriorities[ThreadPriorityIndex];
 	ThreadPriority							= DesiredThreadPriority;
@@ -1189,6 +1403,51 @@ void FHttpDispatcher::CompleteRequest(FHttpRequest& Request, FIoBuffer& Body, co
 	}
 }
 
+#if UE_ALLOW_HTTP_PAUSE
+
+void FHttpDispatcher::OnPauseCommand(bool bShouldPause, const TArray<FString>& Args, FOutputDevice& Ar)
+{
+	if (Args.Num() > 1)
+	{
+#if !NO_LOGGING
+		Ar.Log(LogHttpIoDispatcher.GetCategoryName(), ELogVerbosity::Error, TEXT("Too many args for command, either 0 (all) or 1 (IAD/IAS) expected"));
+#endif // !NO_LOGGING
+		return;
+	}
+
+	EHttpRequestTypeFilter Filter = EHttpRequestTypeFilter::All;
+
+	if (!Args.IsEmpty())
+	{
+		if (Args[0] == TEXT("IAD"))
+		{
+			Filter = EHttpRequestTypeFilter::Installed;
+		}
+		else if (Args[0] == TEXT("IAS"))
+		{
+			Filter = EHttpRequestTypeFilter::Streaming;
+		}
+		else
+		{
+#if !NO_LOGGING
+			Ar.Log(LogHttpIoDispatcher.GetCategoryName(), ELogVerbosity::Error, TEXT("Invalid arg for command, expecting 'IAD' or 'IAS'"));
+#endif // !NO_LOGGING
+			return;
+		}
+	}
+
+	HttpQueue.OnTogglePause(bShouldPause, Filter);
+
+	if (!bShouldPause)
+	{
+		// Since we are unpausing we might have moved previously paused requests back to the active queue
+		// so we need the thread to start runnning again.
+		WakeUp->Trigger();
+	}
+}
+
+#endif // UE_ALLOW_HTTP_PAUSE
+
 } // namespace UE::IoStore::HttpIoDispatcher
 
 ///////////////////////////////////////////////////////////////////////////////
diff --git a/Engine/Source/Runtime/Experimental/IoStore/OnDemand/Private/Statistics.cpp b/Engine/Source/Runtime/Experimental/IoStore/OnDemand/Private/Statistics.cpp
--- a/Engine/Source/Runtime/Experimental/IoStore/OnDemand/Private/Statistics.cpp
+++ b/Engine/Source/Runtime/Experimental/IoStore/OnDemand/Private/Statistics.cpp
@@ -385,7 +385,8 @@ UE_IAX_COUNTER(HttpBandwidthMbps, FCounterFloat);
 double GHttpDurationMsAvg[(uint8)EHttpRequestType::NUM_SOURCES] = { 0.0 };
 int32 GHttpDurationMsMax[(uint8)EHttpRequestType::NUM_SOURCES] = { 0 };
 int64 GHttpDurationMsSum[(uint8)EHttpRequestType::NUM_SOURCES] = { 0 };
-std::atomic<int32> GHttpNumConcurrentRequests = 0;
+std::atomic<uint32> GHttpNumPausedRequests[(uint8)EHttpRequestType::NUM_SOURCES] = {};
+std::atomic<uint32> GHttpNumConcurrentRequests = 0;
 
 struct FHttpRecentHistoryStatistics
 {
@@ -766,6 +767,12 @@ void FOnDemandIoBackendStats::PrintOnScreenStatistics(FCoreDelegates::FSeverityM
 					GHttpPendingCount[(uint8)Type].Get()
 				);
 
+				uint32 NumPaused = GHttpNumPausedRequests[(uint8)Type];
+				if (NumPaused > 0)
+				{
+					Builder.Appendf(TEXT(" | Paused: %d"), NumPaused);
+				}
+
 				OutMessages.Add(FCoreDelegates::EOnScreenMessageSeverity::Info, FText::FromStringView(Builder.ToView()));
 			};
 
@@ -1290,6 +1297,21 @@ void FOnDemandIoBackendStats::OnHttpCdnCacheReply(EHttpRequestType Type, int32 R
 	}
 }
 
+void FOnDemandIoBackendStats::OnHttpOnRemovedPending(EHttpRequestType Type)
+{
+	GHttpPendingCount[(uint8)Type].Decrement();
+}
+
+void FOnDemandIoBackendStats::OnHttpPaused(EHttpRequestType Type)
+{
+	++GHttpNumPausedRequests[(uint8)Type];
+}
+
+void FOnDemandIoBackendStats::OnHttpUnpaused(EHttpRequestType Type)
+{
+	--GHttpNumPausedRequests[(uint8)Type];
+}
+
 void FOnDemandIoBackendStats::OnHttpRequestStarted()
 {
 	++GHttpNumConcurrentRequests;
diff --git a/Engine/Source/Runtime/Experimental/IoStore/OnDemand/Private/Statistics.h b/Engine/Source/Runtime/Experimental/IoStore/OnDemand/Private/Statistics.h
--- a/Engine/Source/Runtime/Experimental/IoStore/OnDemand/Private/Statistics.h
+++ b/Engine/Source/Runtime/Experimental/IoStore/OnDemand/Private/Statistics.h
@@ -101,9 +101,15 @@ public:
 	void OnHttpError(EHttpRequestType Type) IAS_STATISTICS_IMPL()
 	void OnHttpDecodeError(EHttpRequestType Type) IAS_STATISTICS_IMPL()
 	void OnHttpCdnCacheReply(EHttpRequestType Type,int32 Reply) IAS_STATISTICS_IMPL()
+
+	void OnHttpOnRemovedPending(EHttpRequestType Type) IAS_STATISTICS_IMPL()
+	void OnHttpPaused(EHttpRequestType Type) IAS_STATISTICS_IMPL()
+	void OnHttpUnpaused(EHttpRequestType Type) IAS_STATISTICS_IMPL()
+
 	void OnHttpRequestStarted() IAS_STATISTICS_IMPL()
 	void OnHttpRequestCompleted() IAS_STATISTICS_IMPL()
 
+
 private:
 
 #if IAS_WITH_STATISTICS
