Inflight parameter

#rnx

diff --git a/Engine/Source/Runtime/Experimental/IoStore/HttpClient/Internal/IO/Http/Client.h b/Engine/Source/Runtime/Experimental/IoStore/HttpClient/Internal/IO/Http/Client.h
--- a/Engine/Source/Runtime/Experimental/IoStore/HttpClient/Internal/IO/Http/Client.h
+++ b/Engine/Source/Runtime/Experimental/IoStore/HttpClient/Internal/IO/Http/Client.h
@@ -95,6 +95,7 @@ public:
 		uint32				Port = 0;
 		uint16				ConnectionCount = 1;
 		EHttpVersion		HttpVersion = EHttpVersion::One;
+		uint8				MaxInflight = 1;
 		/*
 		enum class ProxyType { Http, Socks4 };
 		Proxy = { ip, port, type }
diff --git a/Engine/Source/Runtime/Experimental/IoStore/HttpClient/Private/ConnectionPool.inl b/Engine/Source/Runtime/Experimental/IoStore/HttpClient/Private/ConnectionPool.inl
--- a/Engine/Source/Runtime/Experimental/IoStore/HttpClient/Private/ConnectionPool.inl
+++ b/Engine/Source/Runtime/Experimental/IoStore/HttpClient/Private/ConnectionPool.inl
@@ -18,6 +18,7 @@ public:
 		uint32			Port = 0;
 		uint16			MaxConnections = 1;
 		bool			bPooled = false;
+		uint8			MaxInflight = 1;
 		EHttpVersion	HttpVersion = EHttpVersion::One;
 		FCertRootsRef	VerifyCert = {};
 	};
@@ -30,6 +31,7 @@ public:
 	FOutcome		ResolveHostName();
 	FCertRootsRef	GetVerifyCert() const		{ return VerifyCert; }
 	uint32			GetMaxConnections() const	{ return MaxConnections; }
+	uint32			GetMaxInflight() const		{ return MaxInflight; }
 	EHttpVersion	GetHttpVersion() const		{ return HttpVersion; }
 	bool			IsPooled() const			{ return bPooled; }
 	uint32			GetIpAddress() const		{ return IpAddresses[0]; }
@@ -44,6 +46,7 @@ private:
 	int16			RecvBufKb = -1;
 	uint16			Port;
 	uint8			MaxConnections;
+	uint8			MaxInflight;
 	EHttpVersion	HttpVersion;
 	bool			bPooled;
 };
@@ -54,6 +57,7 @@ FHost::FHost(const FParams& Params)
 , HostName(Params.HostName)
 , Port(uint16(Params.Port))
 , MaxConnections(uint8(Params.MaxConnections))
+, MaxInflight(Params.MaxInflight)
 , HttpVersion(Params.HttpVersion)
 , bPooled(Params.bPooled)
 {
@@ -257,6 +261,7 @@ FConnectionPool::FConnectionPool(const FParams& Params)
 		.Port			= Params.Port,
 		.MaxConnections	= Params.ConnectionCount,
 		.bPooled		= true,
+		.MaxInflight	= Params.MaxInflight,
 		.HttpVersion	= Params.HttpVersion,
 		.VerifyCert		= Params.VerifyCert,
 	});
