Fix typos in cvar help documentation.

#rnx

diff --git a/Engine/Source/Runtime/Experimental/IoStore/OnDemand/Private/OnDemandHttpIoDispacher.cpp b/Engine/Source/Runtime/Experimental/IoStore/OnDemand/Private/OnDemandHttpIoDispacher.cpp
--- a/Engine/Source/Runtime/Experimental/IoStore/OnDemand/Private/OnDemandHttpIoDispacher.cpp
+++ b/Engine/Source/Runtime/Experimental/IoStore/OnDemand/Private/OnDemandHttpIoDispacher.cpp
@@ -692,7 +692,7 @@ FHttpDispatcher::FHttpDispatcher(TUniquePtr<IIasCache>&& InCache)
 #if UE_ALLOW_HTTP_PAUSE
 	PauseCommand = FConsoleCommandPtr(IConsoleManager::Get().RegisterConsoleCommand(
 		TEXT("iax.HttpPause"),
-		TEXT("Pause all http requests. Passing in 'IAD' or 'IAD' as an arg to only pause requests of that type"),
+		TEXT("Pause all http requests. Passing in 'IAD' or 'IAS' as an arg to only pause requests of that type"),
 		FConsoleCommandWithArgsAndOutputDeviceDelegate::CreateLambda([this](const TArray<FString>& Args, FOutputDevice& Ar) -> void
 			{
 				OnPauseCommand(/*bShouldPause*/ true, Args, Ar);
@@ -701,7 +701,7 @@ FHttpDispatcher::FHttpDispatcher(TUniquePtr<IIasCache>&& InCache)
 
 	UnpauseCommand = FConsoleCommandPtr(IConsoleManager::Get().RegisterConsoleCommand(
 		TEXT("iax.HttpUnpause"),
-		TEXT("Unause all http requests. Passing in 'IAD' or 'IAD' as an arg to only unpause requests of that type"),
+		TEXT("Unpause all http requests. Passing in 'IAD' or 'IAS' as an arg to only unpause requests of that type"),
 		FConsoleCommandWithArgsAndOutputDeviceDelegate::CreateLambda([this](const TArray<FString>& Args, FOutputDevice& Ar) -> void
 			{
 				OnPauseCommand(/*bShouldPause*/ false, Args, Ar);
diff --git a/Engine/Source/Runtime/Experimental/IoStore/OnDemand/Private/OnDemandHttpThread.cpp b/Engine/Source/Runtime/Experimental/IoStore/OnDemand/Private/OnDemandHttpThread.cpp
--- a/Engine/Source/Runtime/Experimental/IoStore/OnDemand/Private/OnDemandHttpThread.cpp
+++ b/Engine/Source/Runtime/Experimental/IoStore/OnDemand/Private/OnDemandHttpThread.cpp
@@ -278,7 +278,7 @@ FOnDemandHttpThread::FOnDemandHttpThread()
 
 	UnpauseCommand = FConsoleCommandPtr(IConsoleManager::Get().RegisterConsoleCommand(
 		TEXT("iax.HttpUnpause"),
-		TEXT("Unause all http requests. Passing in 'IAD' or 'IAS' as an arg to only unpause requests of that type"),
+		TEXT("Unpause all http requests. Passing in 'IAD' or 'IAS' as an arg to only unpause requests of that type"),
 		FConsoleCommandWithArgsAndOutputDeviceDelegate::CreateLambda([this](const TArray<FString>& Args, FOutputDevice& Ar) -> void
 			{
 				OnPauseCommand(/*bShouldPause*/ false, Args, Ar);
