Break out the logic that actions new work

#rnx

diff --git a/Engine/Source/Runtime/Experimental/IoStore/HttpClient/Private/Loop.inl b/Engine/Source/Runtime/Experimental/IoStore/HttpClient/Private/Loop.inl
--- a/Engine/Source/Runtime/Experimental/IoStore/HttpClient/Private/Loop.inl
+++ b/Engine/Source/Runtime/Experimental/IoStore/HttpClient/Private/Loop.inl
@@ -421,6 +421,7 @@ public:
 	void				Fail(FTickState& State, FOutcome Reason);
 
 private:
+	void				MaybeStartNewWork(FTickState& State);
 	void				Negotiate(FTickState& State);
 	void				RecvInternal(FTickState& State);
 	void				SendInternal(FTickState& State);
@@ -472,6 +473,44 @@ void FPeerGroup::Fail(FTickState& State, FOutcome Reason)
 	bNegotiating = false;
 }
 
+////////////////////////////////////////////////////////////////////////////////
+void FPeerGroup::MaybeStartNewWork(FTickState& State)
+{
+	if (IsKeepAlive == 0)
+	{
+		return;
+	}
+
+	FActivityNode* Activity = Recv.GetHead();
+	check(Activity != nullptr);
+	if (Activity->Next != nullptr)
+	{
+		return;
+	}
+
+	if (!Send.IsEmpty())
+	{
+		return;
+	}
+
+	uint32 Remaining = Activity->GetRemainingKiB();
+	if (Remaining >= uint32(GRecvWorkThresholdKiB))
+	{
+		return;
+	}
+
+	Activity = State.Work->PopActivity();
+	if (Activity == nullptr)
+	{
+		return;
+	}
+	check(Activity->GetHost()->IsPooled());
+
+	Trace(Activity, ETrace::StartWork);
+	Send.Append(Activity);
+	SendInternal(State);
+}
+
 ////////////////////////////////////////////////////////////////////////////////
 void FPeerGroup::Negotiate(FTickState& State)
 {
@@ -544,28 +583,12 @@ void FPeerGroup::RecvInternal(FTickState& State)
 	LastUseMs = State.NowMs;
 	bWaiting |= (Outcome.IsWaiting() && Outcome.IsWaitData());
 
-	// If we've only a small amount left to receive we can start more work
-	if (IsKeepAlive & (Activity->Next == nullptr) & Send.IsEmpty())
-	{
-		uint32 Remaining = Activity->GetRemainingKiB();
-		if (Remaining < uint32(GRecvWorkThresholdKiB))
-		{
-			if (FActivityNode* Next = State.Work->PopActivity(); Next != nullptr)
-			{
-				Trace(Activity, ETrace::StartWork);
-	
-				check(Next->GetHost()->IsPooled());
-				check(Send.IsEmpty());
+	MaybeStartNewWork(State);
 
-				Send.Append(Next);
-				SendInternal(State);
-
-				if (!Peer.IsValid())
-				{
-					return;
-				}
-			}
-		}
+	// New work may have fell foul of peer issues
+	if (!Peer.IsValid())
+	{
+		return;
 	}
 
 	// If there was no data available this is far as receiving can go
