Removed the now unused buffer

#rnx

diff --git a/Engine/Source/Runtime/Experimental/IoStore/HttpClient/Private/Activity.inl b/Engine/Source/Runtime/Experimental/IoStore/HttpClient/Private/Activity.inl
--- a/Engine/Source/Runtime/Experimental/IoStore/HttpClient/Private/Activity.inl
+++ b/Engine/Source/Runtime/Experimental/IoStore/HttpClient/Private/Activity.inl
@@ -400,7 +400,7 @@ const FStopwatch& FActivity::GetStopwatch() const
 ////////////////////////////////////////////////////////////////////////////////
 FOutcome FActivity::Transact(FHttpPeer& Peer)
 {
-	Transaction = Peer.Transact(Buffer);
+	Transaction = Peer.Transact();
 	if (!Transaction.IsValid())
 	{
 		return FOutcome::Error("Unable to create a suitable transaction");
diff --git a/Engine/Source/Runtime/Experimental/IoStore/HttpClient/Private/Peer.inl b/Engine/Source/Runtime/Experimental/IoStore/HttpClient/Private/Peer.inl
--- a/Engine/Source/Runtime/Experimental/IoStore/HttpClient/Private/Peer.inl
+++ b/Engine/Source/Runtime/Experimental/IoStore/HttpClient/Private/Peer.inl
@@ -544,7 +544,7 @@ public:
 	FHttpPeer&		operator = (const FHttpPeer& Rhs) = delete;
 					FHttpPeer(FParams&& Params);
 	FOutcome		Handshake();
-	FTransactRef	Transact(FBuffer& Buffer);
+	FTransactRef	Transact();
 	FOutcome		GetPendingTransactId();
 
 private:
@@ -613,12 +613,12 @@ FOutcome FHttpPeer::Handshake()
 }
 
 ////////////////////////////////////////////////////////////////////////////////
-FTransactRef FHttpPeer::Transact(FBuffer& Buffer)
+FTransactRef FHttpPeer::Transact()
 {
 	switch (Proto)
 	{
-	case 1:		return CreateTransactOne(Buffer, PeerData);
-	case 2:		return CreateTransactTwo(Buffer, PeerData);
+	case 1:		return CreateTransactOne(PeerData);
+	case 2:		return CreateTransactTwo(PeerData);
 	default:	return FTransactRef();
 	}
 }
diff --git a/Engine/Source/Runtime/Experimental/IoStore/HttpClient/Private/TransactionOne.inl b/Engine/Source/Runtime/Experimental/IoStore/HttpClient/Private/TransactionOne.inl
--- a/Engine/Source/Runtime/Experimental/IoStore/HttpClient/Private/TransactionOne.inl
+++ b/Engine/Source/Runtime/Experimental/IoStore/HttpClient/Private/TransactionOne.inl
@@ -136,7 +136,7 @@ namespace DetailOne
 class FBase
 {
 public:
-					FBase(FBuffer&);
+					FBase();
 
 protected:
 	using FMutableSection = FBuffer::FMutableSection;
@@ -152,7 +152,7 @@ private:
 };
 
 ////////////////////////////////////////////////////////////////////////////////
-FBase::FBase(FBuffer&)
+FBase::FBase()
 : Buffer(Data, sizeof(Data)) // -V670
 {
 }
@@ -1021,15 +1021,15 @@ FTransactRef::~FTransactRef()
 
 
 ////////////////////////////////////////////////////////////////////////////////
-FTransactRef CreateTransactOne(FBuffer& Buffer, void*&)
+FTransactRef CreateTransactOne(void*&)
 {
 	void* Ptr = FMemory::Malloc(sizeof(FTransactionOne), alignof(FTransactionOne));
-	new (Ptr) FTransactionOne(Buffer);
+	new (Ptr) FTransactionOne();
 	return FTransactRef(UPTRINT(Ptr), 1);
 }
 
 ////////////////////////////////////////////////////////////////////////////////
-FTransactRef CreateTransactTwo(FBuffer& Buffer, void* PeerData)
+FTransactRef CreateTransactTwo(void* PeerData)
 {
 #if IAS_HTTP_WITH_TWO
 	using namespace DetailTwo;
