Missing member initialisation.

#rnx

diff --git a/Engine/Source/Runtime/Experimental/IoStore/HttpClient/Private/Peer.inl b/Engine/Source/Runtime/Experimental/IoStore/HttpClient/Private/Peer.inl
--- a/Engine/Source/Runtime/Experimental/IoStore/HttpClient/Private/Peer.inl
+++ b/Engine/Source/Runtime/Experimental/IoStore/HttpClient/Private/Peer.inl
@@ -549,7 +549,7 @@ public:
 
 private:
 	void*			PeerData = nullptr;
-	uint8			Proto;
+	uint8			Proto = 0;
 };
 
 ////////////////////////////////////////////////////////////////////////////////
