Set up and track the number of inflight work.

#rnx

diff --git a/Engine/Source/Runtime/Experimental/IoStore/HttpClient/Private/Loop.inl b/Engine/Source/Runtime/Experimental/IoStore/HttpClient/Private/Loop.inl
--- a/Engine/Source/Runtime/Experimental/IoStore/HttpClient/Private/Loop.inl
+++ b/Engine/Source/Runtime/Experimental/IoStore/HttpClient/Private/Loop.inl
@@ -414,6 +414,7 @@ class FPeerGroup
 {
 public:
 						FPeerGroup() = default;
+	void				SetMaxInflight(uint32 Maximum);
 	void				Unwait()			{ check(bWaiting); bWaiting = false; }
 	FWaiter				GetWaiter() const;
 	bool				Tick(FTickState& State);
@@ -430,12 +431,20 @@ private:
 	FHttpPeer			Peer;
 	uint32				LastUseMs = 0;
 	uint8				IsKeepAlive = 0;
+	uint8				InflightNum = 0;
+	uint8				InflightMax = 1;
 	bool				bNegotiating = false;
 	bool				bWaiting = false;
 
 	UE_NONCOPYABLE(FPeerGroup);
 };
 
+////////////////////////////////////////////////////////////////////////////////
+void FPeerGroup::SetMaxInflight(uint32 Maximum)
+{
+	InflightMax = uint8(FMath::Min(uint32(Maximum), 127u));
+}
+
 ////////////////////////////////////////////////////////////////////////////////
 FWaiter FPeerGroup::GetWaiter() const
 {
@@ -508,6 +517,7 @@ void FPeerGroup::MaybeStartNewWork(FTickState& State)
 
 	Trace(Activity, ETrace::StartWork);
 	Send.Append(Activity);
+	InflightNum++;
 	SendInternal(State);
 }
 
@@ -608,6 +618,7 @@ void FPeerGroup::RecvInternal(FTickState& State)
 
 	check(Recv.GetHead() == Activity);
 	Recv.PopPrepend(State.DoneList);
+	InflightNum--;
 
 	// If the server wants the connection closed then we can drop our peer
 	if (IsKeepAlive == 0)
@@ -716,6 +727,7 @@ void FPeerGroup::TickSend(FTickState& State, FHost& Host, FPoller& Poller)
 		IsKeepAlive = 1;
 		bNegotiating = true;
 		bWillBlock = Outcome.IsWaiting();
+		InflightNum = 0;
 
 		FHttpPeer::FParams Params = {
 			.Socket = MoveTemp(Socket),
@@ -727,6 +739,7 @@ void FPeerGroup::TickSend(FTickState& State, FHost& Host, FPoller& Poller)
 	}
 
 	Send.Append(Pending);
+	InflightNum++;
 
 	if (!bWillBlock)
 	{
@@ -773,6 +786,11 @@ FHostGroup::FHostGroup(FHost& InHost)
 {
 	uint32 Num = InHost.GetMaxConnections();
 	PeerGroups.SetNum(Num);
+
+	for (uint32 Inflight = Host.GetMaxInflight(); FPeerGroup& Group : PeerGroups)
+	{
+		Group.SetMaxInflight(Inflight);
+	}
 }
 
 ////////////////////////////////////////////////////////////////////////////////
