Use inflight count as condition for new work and try to dispatch more regularly

#rnx

diff --git a/Engine/Source/Runtime/Experimental/IoStore/HttpClient/Private/Loop.inl b/Engine/Source/Runtime/Experimental/IoStore/HttpClient/Private/Loop.inl
--- a/Engine/Source/Runtime/Experimental/IoStore/HttpClient/Private/Loop.inl
+++ b/Engine/Source/Runtime/Experimental/IoStore/HttpClient/Private/Loop.inl
@@ -422,7 +422,7 @@ public:
 	void				Fail(FTickState& State, FOutcome Reason);
 
 private:
-	void				MaybeStartNewWork(FTickState& State);
+	bool				MaybeStartNewWork(FTickState& State);
 	void				Negotiate(FTickState& State);
 	void				RecvInternal(FTickState& State);
 	void				SendInternal(FTickState& State);
@@ -479,46 +479,36 @@ void FPeerGroup::Fail(FTickState& State, FOutcome Reason)
 	Recv = FActivityList();
 	bWaiting = false;
 	IsKeepAlive = 0;
+	InflightNum = 0;
 	bNegotiating = false;
 }
 
 ////////////////////////////////////////////////////////////////////////////////
-void FPeerGroup::MaybeStartNewWork(FTickState& State)
+bool FPeerGroup::MaybeStartNewWork(FTickState& State)
 {
-	if (IsKeepAlive == 0)
-	{
-		return;
-	}
+	if (InflightNum >= InflightMax)		return false;
+	if (IsKeepAlive == 0)				return false;
+	if (!State.Work->HasWork())			return false;
 
-	FActivityNode* Activity = Recv.GetHead();
-	check(Activity != nullptr);
-	if (Activity->Next != nullptr)
+	FActivityNode* Activity = State.Work->PopActivity();
+	do
 	{
-		return;
-	}
+		Trace(Activity, ETrace::StartWork);
 
-	if (!Send.IsEmpty())
-	{
-		return;
-	}
+		check(Activity->GetHost()->IsPooled());
 
-	uint32 Remaining = Activity->GetRemainingKiB();
-	if (Remaining >= uint32(GRecvWorkThresholdKiB))
-	{
-		return;
-	}
+		Send.Append(Activity);
+		InflightNum++;
+		if (InflightNum == InflightMax)
+		{
+			break;
+		}
 
-	Activity = State.Work->PopActivity();
-	if (Activity == nullptr)
-	{
-		return;
+		Activity = State.Work->PopActivity();
 	}
-	check(Activity->GetHost()->IsPooled());
+	while (Activity != nullptr);
 
-	Trace(Activity, ETrace::StartWork);
-	Send.Append(Activity);
-	InflightNum++;
-	SendInternal(State);
+	return true;
 }
 
 ////////////////////////////////////////////////////////////////////////////////
@@ -593,8 +583,6 @@ void FPeerGroup::RecvInternal(FTickState& State)
 	LastUseMs = State.NowMs;
 	bWaiting |= (Outcome.IsWaiting() && Outcome.IsWaitData());
 
-	MaybeStartNewWork(State);
-
 	// New work may have fell foul of peer issues
 	if (!Peer.IsValid())
 	{
@@ -620,6 +608,11 @@ void FPeerGroup::RecvInternal(FTickState& State)
 	Recv.PopPrepend(State.DoneList);
 	InflightNum--;
 
+	if (MaybeStartNewWork(State))
+	{
+		SendInternal(State);
+	}
+
 	// If the server wants the connection closed then we can drop our peer
 	if (IsKeepAlive == 0)
 	{
@@ -653,6 +646,11 @@ void FPeerGroup::SendInternal(FTickState& State)
 	check(Outcome.GetResult() == uint32(FActivity::EStage::Request));
 
 	Send.PopAppend(Recv);
+
+	if (!Send.IsEmpty())
+	{
+		return SendInternal(State);
+	}
 }
 
 ////////////////////////////////////////////////////////////////////////////////
@@ -679,8 +677,8 @@ bool FPeerGroup::Tick(FTickState& State)
 ////////////////////////////////////////////////////////////////////////////////
 void FPeerGroup::TickSend(FTickState& State, FHost& Host, FPoller& Poller)
 {
-	// This path is only for those that are idle and have nothing to do
-	if (!(Send.IsEmpty() & Recv.IsEmpty()))
+	// This path is only for those that have spare capacity
+	if (InflightNum >= InflightMax)
 	{
 		return;
 	}
@@ -741,6 +739,8 @@ void FPeerGroup::TickSend(FTickState& State, FHost& Host, FPoller& Poller)
 	Send.Append(Pending);
 	InflightNum++;
 
+	MaybeStartNewWork(State);
+
 	if (!bWillBlock)
 	{
 		if (bNegotiating)
