Fixed an occasional crash when unwelding rigid bodies.
The crash can happen when a shape in the middle of the shapes array gets removed.

diff --git a/Engine/Source/Runtime/Experimental/Chaos/Private/Chaos/ParticleHandle.cpp b/Engine/Source/Runtime/Experimental/Chaos/Private/Chaos/ParticleHandle.cpp
--- a/Engine/Source/Runtime/Experimental/Chaos/Private/Chaos/ParticleHandle.cpp
+++ b/Engine/Source/Runtime/Experimental/Chaos/Private/Chaos/ParticleHandle.cpp
@@ -164,6 +164,19 @@ namespace Chaos
 					RemoveArrayItemsAtSortedIndices(MShapesArray, InIndices);
 
 					Union->RemoveAtSortedIndices(InIndices);
+
+					// If shapes are removed in the middle of the array, the logic in UpdateShapesArrayFromGeometry does 
+					// not actually pack down the indices correctly as we've already removed the shape. 
+					// Just force the indices to be correct here.
+					if (!InIndices.IsEmpty())
+					{
+						// These indices are sorted, we can start the fix-up at the first index.
+						const int32 FirstIndex = InIndices[0];
+						for (int32 I = FirstIndex; I < MShapesArray.Num(); ++I)
+						{
+							MShapesArray[I]->ModifyShapeIndex(I);
+						}
+					}
 				}
 			});
 	}
