Chaos - handle runtime changes to InertiaConditioing cvars for easier tuning

#rnx

diff --git a/Engine/Source/Runtime/Experimental/Chaos/Private/Chaos/MassConditioning.cpp b/Engine/Source/Runtime/Experimental/Chaos/Private/Chaos/MassConditioning.cpp
--- a/Engine/Source/Runtime/Experimental/Chaos/Private/Chaos/MassConditioning.cpp
+++ b/Engine/Source/Runtime/Experimental/Chaos/Private/Chaos/MassConditioning.cpp
@@ -6,39 +6,18 @@
 #include "Chaos/PBDJointConstraints.h"
 #include "Chaos/PBDJointConstraintUtilities.h"
 
-namespace Chaos::CVars
-{
-	// NOTE: ChaosInertiaConditioningInvInertiaTolerance of 1e-4 effectively disables inertia conditioning for medium and larger object.
-	// If you want inertia conditioning on large objects, make ChaosInertiaConditioningInvInertiaTolerance very small (say, 1e-30).
-	float ChaosInertiaConditioningInvMassTolerance = 1.e-20f;
-	float ChaosInertiaConditioningInvInertiaTolerance = 1.e-4f;	
-	float ChaosInertiaConditioningExtentTolerance = 1.e-8f;
-	FAutoConsoleVariableRef  CVarChaosInertiaConditioningInvMassTolerance(
-		TEXT("p.Chaos.InertiaConditioning.InvMassTolerance"), 
-		ChaosInertiaConditioningInvMassTolerance, 
-		TEXT("Inertia conditioning is disabled for objects with inverse mass less than this."));
-	FAutoConsoleVariableRef  CVarChaosInertiaConditioningInvInertiaTolerance(
-		TEXT("p.Chaos.InertiaConditioning.InvInertiaTolerance"), 
-		ChaosInertiaConditioningInvInertiaTolerance, 
-		TEXT("Inertia conditioning is disabled for objects with inverse inertia less than this. The default value disables Inertia Conditioing for mediums to large objects."));
-	FAutoConsoleVariableRef  CVarChaosInertiaConditioningExtentTolerance(
-		TEXT("p.Chaos.InertiaConditioning.ExtentTolerance"), 
-		ChaosInertiaConditioningExtentTolerance, 
-		TEXT("Inertia conditioning is disabled for objects with dimensions smaller than this."));
-}
-
 namespace Chaos
 {
-	FVec3f CalculateInertiaConditioning(const FRealSingle InvM, const FVec3f& InvI, const FVec3f& ConstraintExtents, const FRealSingle MaxDistance, const FRealSingle MaxRotationRatio, const FRealSingle MaxInvInertiaComponentRatio)
+	FVec3f CalculateInertiaConditioning(const FRealSingle InvM, const FVec3f& InvI, const FVec3f& ConstraintExtents, const FRealSingle MaxDistance, const FRealSingle MaxRotationRatio, const FRealSingle MaxInvInertiaComponentRatio, const FInertiaConditioningTolerances& Tolerances)
 	{
 		FVec3f InvInertiaScale = FVec3f(1);
 
 		// Sanity check for zero mass or zero dimensions along at least one axis
 		// We should never get this unless there are asset errors
 		bool bIsValid = 
-			(InvM > CVars::ChaosInertiaConditioningInvMassTolerance) 
-			&& (InvI.GetMin() > CVars::ChaosInertiaConditioningInvInertiaTolerance)
-			&& (ConstraintExtents.GetMin() > CVars::ChaosInertiaConditioningExtentTolerance);
+			(InvM > Tolerances.InvMassTolerance)
+			&& (InvI.GetMin() > Tolerances.InvInertiaTolerance)
+			&& (ConstraintExtents.GetMin() > Tolerances.ExtentTolerance);
 
 		if (bIsValid)
 		{
@@ -134,12 +113,12 @@ namespace Chaos
 		return FVec3::Max(ColisionExtents, ConstraintExtents);
 	}
 
-	FVec3f CalculateParticleInertiaConditioning(const FPBDRigidParticleHandle* Rigid, const FRealSingle MaxDistance, const FRealSingle MaxRotationRatio, const FRealSingle MaxInvInertiaComponentRatio)
+	FVec3f CalculateParticleInertiaConditioning(const FPBDRigidParticleHandle* Rigid, const FRealSingle MaxDistance, const FRealSingle MaxRotationRatio, const FRealSingle MaxInvInertiaComponentRatio, const FInertiaConditioningTolerances& Tolerances)
 	{
 		if (Rigid->IsDynamic())
 		{
 			const FVec3 ConstraintExtents = CalculateCoMConstraintExtents(Rigid);
-			return CalculateInertiaConditioning(FRealSingle(Rigid->InvM()), Rigid->InvI(), ConstraintExtents, MaxDistance, MaxRotationRatio, MaxInvInertiaComponentRatio);
+			return CalculateInertiaConditioning(FRealSingle(Rigid->InvM()), Rigid->InvI(), ConstraintExtents, MaxDistance, MaxRotationRatio, MaxInvInertiaComponentRatio, Tolerances);
 		}
 		return FVec3f(1);
 	}
diff --git a/Engine/Source/Runtime/Experimental/Chaos/Private/Chaos/PBDRigidsEvolutionGBF.cpp b/Engine/Source/Runtime/Experimental/Chaos/Private/Chaos/PBDRigidsEvolutionGBF.cpp
--- a/Engine/Source/Runtime/Experimental/Chaos/Private/Chaos/PBDRigidsEvolutionGBF.cpp
+++ b/Engine/Source/Runtime/Experimental/Chaos/Private/Chaos/PBDRigidsEvolutionGBF.cpp
@@ -103,22 +103,61 @@ namespace Chaos
 		bool DoFinalProbeNarrowPhase = true;
 		FAutoConsoleVariableRef CVarDoFinalProbeNarrowPhase(TEXT("p.Chaos.Solver.DoFinalProbeNarrowPhase"), DoFinalProbeNarrowPhase, TEXT(""));
 
+		// Gets set when one of the Inertia Conditioning CVars changes so we can refresh the system
+		bool bChaosSolverRefreshInertiaConditioning = false;
+		FConsoleVariableDelegate OnInertiaConditioningCVarChangedDelegate = FConsoleVariableDelegate::CreateLambda([](IConsoleVariable*) { bChaosSolverRefreshInertiaConditioning = true; });
+
 		// Enable inertia conditioning for joint and collisions stabilization for small and thin objects
 		bool bChaosSolverInertiaConditioningEnabled = true;
-		FAutoConsoleVariableRef  CVarChaosSolverInertiaConditioningEnabled(TEXT("p.Chaos.Solver.InertiaConditioning.Enabled"), bChaosSolverInertiaConditioningEnabled, TEXT("Enable/Disable constraint stabilization through inertia conditioning"));
+		FAutoConsoleVariableRef  CVarChaosSolverInertiaConditioningEnabled(
+			TEXT("p.Chaos.Solver.InertiaConditioning.Enabled"), 
+			bChaosSolverInertiaConditioningEnabled, 
+			TEXT("Enable/Disable constraint stabilization through inertia conditioning"), 
+			OnInertiaConditioningCVarChangedDelegate);
 
 		// The largest joint error we expect to resolve in a moderately stable way
 		FRealSingle ChaosSolverInertiaConditioningDistance = 20;
-		FAutoConsoleVariableRef  CVarChaosSolverInertiaConditioningDistance(TEXT("p.Chaos.Solver.InertiaConditioning.Distance"), ChaosSolverInertiaConditioningDistance, TEXT("An input to inertia conditioning system. The joint distance error which needs to be stable (generate a low rotation)."));
+		FAutoConsoleVariableRef  CVarChaosSolverInertiaConditioningDistance(
+			TEXT("p.Chaos.Solver.InertiaConditioning.Distance"), 
+			ChaosSolverInertiaConditioningDistance, 
+			TEXT("An input to inertia conditioning system. The joint distance error which needs to be stable (generate a low rotation)."), 
+			OnInertiaConditioningCVarChangedDelegate);
 
 		// The ratio of joint error correction that comes from particle rotation versus translation.
 		// NOTE: A value of 2.5 means that spheres (and cubes) are not effected by inertia conditioning (unless they have constraints outside their bounds).
 		FRealSingle ChaosSolverInertiaConditioningRotationRatio = 2.5f;
-		FAutoConsoleVariableRef  CVarChaosSolverInertiaConditioningRotationRatio(TEXT("p.Chaos.Solver.InertiaConditioning.RotationRatio"), ChaosSolverInertiaConditioningRotationRatio, TEXT("An input to inertia conditioning system. The maximum ratio of joint correction from rotation versus translation"));
+		FAutoConsoleVariableRef  CVarChaosSolverInertiaConditioningRotationRatio(
+			TEXT("p.Chaos.Solver.InertiaConditioning.RotationRatio"), 
+			ChaosSolverInertiaConditioningRotationRatio, 
+			TEXT("An input to inertia conditioning system. The maximum ratio of joint correction from rotation versus translation"), 
+			OnInertiaConditioningCVarChangedDelegate);
 
 		// If > 1, limits the inverse inertia components to be mo more than this multiple of the smallest component. Makes the objects more round. Set to 0 to disable.
 		FRealSingle ChaosSolverMaxInvInertiaComponentRatio = 0;
-		FAutoConsoleVariableRef  CVarChaosSolverInertiaConditioningMaxInvInertiaComponentRatio(TEXT("p.Chaos.Solver.InertiaConditioning.MaxInvInertiaComponentRatio"), ChaosSolverMaxInvInertiaComponentRatio, TEXT("An input to inertia conditioning system. The largest inertia component must be at least least multiple of the smallest component"));
+		FAutoConsoleVariableRef  CVarChaosSolverInertiaConditioningMaxInvInertiaComponentRatio(
+			TEXT("p.Chaos.Solver.InertiaConditioning.MaxInvInertiaComponentRatio"),
+			ChaosSolverMaxInvInertiaComponentRatio,
+			TEXT("An input to inertia conditioning system. The largest inertia component must be at least least multiple of the smallest component"), 
+			OnInertiaConditioningCVarChangedDelegate);
+
+		// NOTE: ChaosInertiaConditioningInvInertiaTolerance of 1e-4 effectively disables inertia conditioning for medium and larger object.
+		// If you want inertia conditioning on large objects, make ChaosInertiaConditioningInvInertiaTolerance very small (say, 1e-30).
+		FInertiaConditioningTolerances ChaosInertiaConditioningTolerances = FInertiaConditioningTolerances();
+		FAutoConsoleVariableRef  CVarChaosInertiaConditioningInvMassTolerance(
+			TEXT("p.Chaos.InertiaConditioning.InvMassTolerance"),
+			ChaosInertiaConditioningTolerances.InvMassTolerance,
+			TEXT("Inertia conditioning is disabled for objects with inverse mass less than this."),
+			OnInertiaConditioningCVarChangedDelegate);
+		FAutoConsoleVariableRef  CVarChaosInertiaConditioningInvInertiaTolerance(
+			TEXT("p.Chaos.InertiaConditioning.InvInertiaTolerance"),
+			ChaosInertiaConditioningTolerances.InvInertiaTolerance,
+			TEXT("Inertia conditioning is disabled for objects with inverse inertia less than this. The default value disables Inertia Conditioing for mediums to large objects."),
+			OnInertiaConditioningCVarChangedDelegate);
+		FAutoConsoleVariableRef  CVarChaosInertiaConditioningExtentTolerance(
+			TEXT("p.Chaos.InertiaConditioning.ExtentTolerance"),
+			ChaosInertiaConditioningTolerances.ExtentTolerance,
+			TEXT("Inertia conditioning is disabled for objects with dimensions smaller than this."),
+			OnInertiaConditioningCVarChangedDelegate);
 
 		// Collision modifiers are run after the CCD rewind and manifold regeneration because the CCD "contact" data is only used to rewind the particle 
 		// and we regenerate the manifold at that new position. Running the modifier callback before CCD rewind would mean that the user does not have 
@@ -2042,15 +2081,21 @@ void FPBDRigidsEvolutionGBF::UpdateInertiaConditioning()
 	// A limit on the relative sizes of the inertia components (inverse)
 	const FRealSingle MaxInvInertiaComponentRatio = ChaosSolverMaxInvInertiaComponentRatio;
 
+	// Tolerances (currently controlled by cvar)
+	const FInertiaConditioningTolerances& Tolerances = ChaosInertiaConditioningTolerances;
+
 	// @todo(chaos): we should only visit particles that are dirty
 	for (auto& Rigid : Particles.GetNonDisabledDynamicView())
 	{
-		if (Rigid.InertiaConditioningDirty())
+		// If any cvars changed, update the inertia conditioning of all particles that use it
+		const bool bForceUpdate = (bChaosSolverRefreshInertiaConditioning && Rigid.InertiaConditioningEnabled());
+
+		if (Rigid.InertiaConditioningDirty() || bForceUpdate)
 		{
 			const bool bWantInertiaConditioning = Rigid.IsDynamic() && Rigid.InertiaConditioningEnabled() && bChaosSolverInertiaConditioningEnabled;
 			if (bWantInertiaConditioning)
 			{
-				const FVec3f InvInertiaScale = CalculateParticleInertiaConditioning(Rigid.Handle(), MaxDistance, MaxRotationRatio, MaxInvInertiaComponentRatio);
+				const FVec3f InvInertiaScale = CalculateParticleInertiaConditioning(Rigid.Handle(), MaxDistance, MaxRotationRatio, MaxInvInertiaComponentRatio, Tolerances);
 				Rigid.SetInvIConditioning(InvInertiaScale);
 			}
 			else
@@ -2060,6 +2105,8 @@ void FPBDRigidsEvolutionGBF::UpdateInertiaConditioning()
 			Rigid.ClearInertiaConditioningDirty();
 		}
 	}
+
+	bChaosSolverRefreshInertiaConditioning = false;
 }
 
 #if CHAOS_EVOLUTION_COLLISION_TESTMODE
diff --git a/Engine/Source/Runtime/Experimental/Chaos/Public/Chaos/MassConditioning.h b/Engine/Source/Runtime/Experimental/Chaos/Public/Chaos/MassConditioning.h
--- a/Engine/Source/Runtime/Experimental/Chaos/Public/Chaos/MassConditioning.h
+++ b/Engine/Source/Runtime/Experimental/Chaos/Public/Chaos/MassConditioning.h
@@ -8,6 +8,12 @@
 
 namespace Chaos
 {
+	struct FInertiaConditioningTolerances
+	{
+		float InvMassTolerance = 1.e-20f;
+		float InvInertiaTolerance = 1.e-4f;
+		float ExtentTolerance = 1.e-8f;
+	};
 
 
 	/**
@@ -18,7 +24,7 @@ namespace Chaos
 	 * @param MaxDistance the constraint error that we want to be stable. Corrections above this may still contain large rotation components.
 	 * @param MaxRotationRatio the contribution to the constraint correction from rotation of the object will be less that this fraction of the total error.
 	*/
-	FVec3f CHAOS_API CalculateInertiaConditioning(const FRealSingle InvM, const FVec3f& InvI, const FVec3f& ConstraintExtents, const FRealSingle MaxDistance, const FRealSingle MaxRotationRatio, const FRealSingle MaxInvInertiaComponentRatio);
+	FVec3f CHAOS_API CalculateInertiaConditioning(const FRealSingle InvM, const FVec3f& InvI, const FVec3f& ConstraintExtents, const FRealSingle MaxDistance, const FRealSingle MaxRotationRatio, const FRealSingle MaxInvInertiaComponentRatio, const FInertiaConditioningTolerances& Tolerances = FInertiaConditioningTolerances());
 
 	/**
 	 * Generate a CoM-relative extent that encompasses all constraints (collision and joints for now) on the particle. For collisions, we approximate the max extend
@@ -30,6 +36,6 @@ namespace Chaos
 	 * Calculate an inertia scale for a particle based on object size and all joint connector offsets that should stabilize the constraints.
 	 * @see CalculateInertiaConditioning()
 	*/
-	FVec3f CHAOS_API CalculateParticleInertiaConditioning(const FPBDRigidParticleHandle* Rigid, const FRealSingle MaxDistance, const FRealSingle MaxRotationRatio, const FRealSingle MaxInvInertiaComponentRatio);
+	FVec3f CHAOS_API CalculateParticleInertiaConditioning(const FPBDRigidParticleHandle* Rigid, const FRealSingle MaxDistance, const FRealSingle MaxRotationRatio, const FRealSingle MaxInvInertiaComponentRatio, const FInertiaConditioningTolerances& Tolerances = FInertiaConditioningTolerances());
 
 }
\ No newline at end of file
