Chaos - fix InertiaConditioning cvar change handles to support multiple solvers (e.g., client/server PIE)

#rnx

diff --git a/Engine/Source/Runtime/Experimental/Chaos/Private/Chaos/PBDRigidsEvolutionGBF.cpp b/Engine/Source/Runtime/Experimental/Chaos/Private/Chaos/PBDRigidsEvolutionGBF.cpp
--- a/Engine/Source/Runtime/Experimental/Chaos/Private/Chaos/PBDRigidsEvolutionGBF.cpp
+++ b/Engine/Source/Runtime/Experimental/Chaos/Private/Chaos/PBDRigidsEvolutionGBF.cpp
@@ -103,9 +103,10 @@ namespace Chaos
 		bool DoFinalProbeNarrowPhase = true;
 		FAutoConsoleVariableRef CVarDoFinalProbeNarrowPhase(TEXT("p.Chaos.Solver.DoFinalProbeNarrowPhase"), DoFinalProbeNarrowPhase, TEXT(""));
 
-		// Gets set when one of the Inertia Conditioning CVars changes so we can refresh the system
-		bool bChaosSolverRefreshInertiaConditioning = false;
-		FConsoleVariableDelegate OnInertiaConditioningCVarChangedDelegate = FConsoleVariableDelegate::CreateLambda([](IConsoleVariable*) { bChaosSolverRefreshInertiaConditioning = true; });
+		// Gets incremented when one of the Inertia Conditioning CVars changes so we can refresh the system.
+		// Uses a change counter rather than a boolean to handle multiple Evolutions (e.g., Client+Server PIE)
+		int32 ChaosSolverInertiaConditioningEpoch = 0;
+		FConsoleVariableDelegate OnInertiaConditioningCVarChangedDelegate = FConsoleVariableDelegate::CreateLambda([](IConsoleVariable*) { ++ChaosSolverInertiaConditioningEpoch; });
 
 		// Enable inertia conditioning for joint and collisions stabilization for small and thin objects
 		bool bChaosSolverInertiaConditioningEnabled = true;
@@ -1345,6 +1346,7 @@ FPBDRigidsEvolutionGBF::FPBDRigidsEvolutionGBF(
 	, CollisionModifiers(InCollisionModifiers)
 	, CCDManager()
 	, bIsDeterministic(false)
+	, InertiaConditioningEpoch(-1)
 {
 	SetNumPositionIterations(DefaultNumPositionIterations);
 	SetNumVelocityIterations(DefaultNumVelocityIterations);
@@ -2084,13 +2086,15 @@ void FPBDRigidsEvolutionGBF::UpdateInertiaConditioning()
 	// Tolerances (currently controlled by cvar)
 	const FInertiaConditioningTolerances& Tolerances = ChaosInertiaConditioningTolerances;
 
+	// Did the settings change?
+	const bool bSettingsChanged = (InertiaConditioningEpoch != ChaosSolverInertiaConditioningEpoch);
+	InertiaConditioningEpoch = ChaosSolverInertiaConditioningEpoch;
+
 	// @todo(chaos): we should only visit particles that are dirty
 	for (auto& Rigid : Particles.GetNonDisabledDynamicView())
 	{
-		// If any cvars changed, update the inertia conditioning of all particles that use it
-		const bool bForceUpdate = (bChaosSolverRefreshInertiaConditioning && Rigid.InertiaConditioningEnabled());
-
-		if (Rigid.InertiaConditioningDirty() || bForceUpdate)
+		// Update the inertia conditioning if some property of the body, or the system settings have changed
+		if (Rigid.InertiaConditioningDirty() || (bSettingsChanged && Rigid.InertiaConditioningEnabled()))
 		{
 			const bool bWantInertiaConditioning = Rigid.IsDynamic() && Rigid.InertiaConditioningEnabled() && bChaosSolverInertiaConditioningEnabled;
 			if (bWantInertiaConditioning)
@@ -2105,8 +2109,6 @@ void FPBDRigidsEvolutionGBF::UpdateInertiaConditioning()
 			Rigid.ClearInertiaConditioningDirty();
 		}
 	}
-
-	bChaosSolverRefreshInertiaConditioning = false;
 }
 
 #if CHAOS_EVOLUTION_COLLISION_TESTMODE
diff --git a/Engine/Source/Runtime/Experimental/Chaos/Public/Chaos/MassConditioning.h b/Engine/Source/Runtime/Experimental/Chaos/Public/Chaos/MassConditioning.h
--- a/Engine/Source/Runtime/Experimental/Chaos/Public/Chaos/MassConditioning.h
+++ b/Engine/Source/Runtime/Experimental/Chaos/Public/Chaos/MassConditioning.h
@@ -20,21 +20,23 @@ namespace Chaos
 	 * Calculate an inertia scale so that rotation corrections when applying a constraint are no larger than some fraction of the total correction.
 	 * It has the net effect of making the object "more round" and also increasing the inertia relative to the mass when the object is not of "regular" proportions.
 	 * 
+	 * @param InvM the inverse mass
+	 * @param InvI the inverse inertia
 	 * @param ConstraintExtents the Maximum constraint arm distance along each axis. This should be the extents of the object, including all shapes and joint connectors.
 	 * @param MaxDistance the constraint error that we want to be stable. Corrections above this may still contain large rotation components.
 	 * @param MaxRotationRatio the contribution to the constraint correction from rotation of the object will be less that this fraction of the total error.
+	 * @param MaxInvInertiaComponentRatio the maximum allowed ratio between the smallest and largest components of the inverse inertia.
+	 * @param Tolerances the tolerances to use
 	*/
 	FVec3f CHAOS_API CalculateInertiaConditioning(const FRealSingle InvM, const FVec3f& InvI, const FVec3f& ConstraintExtents, const FRealSingle MaxDistance, const FRealSingle MaxRotationRatio, const FRealSingle MaxInvInertiaComponentRatio, const FInertiaConditioningTolerances& Tolerances = FInertiaConditioningTolerances());
 
-	/**
-	 * Generate a CoM-relative extent that encompasses all constraints (collision and joints for now) on the particle. For collisions, we approximate the max extend
-	 * using the CoM-space bounds. For joints we use the actual joint attachment positions.
-	*/
-	FVec3f CHAOS_API CalculateParticleConstraintExtents(const FPBDRigidParticleHandle* Rigid);
-
 	/**
 	 * Calculate an inertia scale for a particle based on object size and all joint connector offsets that should stabilize the constraints.
 	 * @see CalculateInertiaConditioning()
+	 * @param Rigid the particle to calculate inertia conditioning for
+	 * @param MaxDistance the constraint error that we want to be stable. Corrections above this may still contain large rotation components.
+	 * @param MaxRotationRatio the contribution to the constraint correction from rotation of the object will be less that this fraction of the total error.
+	 * @param MaxInvInertiaComponentRatio the maximum allowed ratio between the smallest and largest components of the inverse inertia.
 	*/
 	FVec3f CHAOS_API CalculateParticleInertiaConditioning(const FPBDRigidParticleHandle* Rigid, const FRealSingle MaxDistance, const FRealSingle MaxRotationRatio, const FRealSingle MaxInvInertiaComponentRatio, const FInertiaConditioningTolerances& Tolerances = FInertiaConditioningTolerances());
 
diff --git a/Engine/Source/Runtime/Experimental/Chaos/Public/Chaos/PBDRigidsEvolutionGBF.h b/Engine/Source/Runtime/Experimental/Chaos/Public/Chaos/PBDRigidsEvolutionGBF.h
--- a/Engine/Source/Runtime/Experimental/Chaos/Public/Chaos/PBDRigidsEvolutionGBF.h
+++ b/Engine/Source/Runtime/Experimental/Chaos/Public/Chaos/PBDRigidsEvolutionGBF.h
@@ -354,6 +354,9 @@ namespace Chaos
 
 		bool bIsDeterministic;
 
+		// Used to track whether we have applied the latest Inertia Conditioning Cvars
+		int32 InertiaConditioningEpoch;
+
 #if CHAOS_EVOLUTION_COLLISION_TESTMODE
 		void TestModeResetCollisions();
 #endif
