Address code analysis issue in recent changes in dataflow util

diff --git a/Engine/Source/Runtime/Experimental/Dataflow/Engine/Private/Dataflow/DataflowEngineUtil.cpp b/Engine/Source/Runtime/Experimental/Dataflow/Engine/Private/Dataflow/DataflowEngineUtil.cpp
--- a/Engine/Source/Runtime/Experimental/Dataflow/Engine/Private/Dataflow/DataflowEngineUtil.cpp
+++ b/Engine/Source/Runtime/Experimental/Dataflow/Engine/Private/Dataflow/DataflowEngineUtil.cpp
@@ -115,17 +115,17 @@ namespace UE::Dataflow
 				FString ObjectPtrInnerType;
 				if (FDataflowUObjectConvertibleTypePolicy::GetObjectPtrInnerType(TypeStr, ObjectPtrInnerType))
 				{
-					if (const UObject* Object = FindFirstObjectSafe<UObject>(*ObjectPtrInnerType))
+					if (const UClass* ObjectClass = FindFirstObjectSafe<UClass>(*ObjectPtrInnerType))
 					{
-						return FPropertyBagPropertyDesc(Name, EPropertyBagPropertyType::Object, Object);
+						return FPropertyBagPropertyDesc(Name, EPropertyBagPropertyType::Object, ObjectClass);
 					}
 				}
 				if (TypeStr.StartsWith("U"))
 				{
 					const FString ShortTypeName = TypeStr.RightChop(1);
-					if (const UObject* Object = FindFirstObjectSafe<UObject>(*ShortTypeName))
+					if (const UClass* ObjectClass = FindFirstObjectSafe<UClass>(*ShortTypeName))
 					{
-						return FPropertyBagPropertyDesc(Name, EPropertyBagPropertyType::Object, Object);
+						return FPropertyBagPropertyDesc(Name, EPropertyBagPropertyType::Object, ObjectClass);
 					}
 				}
 				else if (TypeStr.StartsWith("F"))
