Remove alignment hack for TSpatialHashGridPoints::CellSize

Not needed for Clang for Windows anymore now that TVectorRegisterType<T> is aligning as expected. Fix was integrated into 38.00 in changelist 45973269.

diff --git a/Engine/Source/Runtime/Experimental/Chaos/Public/Chaos/HierarchicalSpatialHash.h b/Engine/Source/Runtime/Experimental/Chaos/Public/Chaos/HierarchicalSpatialHash.h
--- a/Engine/Source/Runtime/Experimental/Chaos/Public/Chaos/HierarchicalSpatialHash.h
+++ b/Engine/Source/Runtime/Experimental/Chaos/Public/Chaos/HierarchicalSpatialHash.h
@@ -645,7 +645,7 @@ public:
 		return ConcurrentResults.Result;
 	}
 private:
-	alignas(32) const TVectorRegisterType<T> CellSize; // alignas is a workaround to a crash where clang isn't automatically doing the aligning.
+	const TVectorRegisterType<T> CellSize;
 
 	// Particle Index is dense array index passed in at build time. 
 	TArray<int32> ParticleIndexToElement;
