Reset frame once we're done.

#rnx

diff --git a/Engine/Source/Runtime/Experimental/IoStore/HttpClient/Private/TransactionTwo.inl b/Engine/Source/Runtime/Experimental/IoStore/HttpClient/Private/TransactionTwo.inl
--- a/Engine/Source/Runtime/Experimental/IoStore/HttpClient/Private/TransactionTwo.inl
+++ b/Engine/Source/Runtime/Experimental/IoStore/HttpClient/Private/TransactionTwo.inl
@@ -496,6 +496,11 @@ FOutcome FDriverNg::RecvResponse(FTlsPeer& Peer, TArray<FHeaderNg>* Sink)
 		int32 EndOfHeaders = (RecvFrame.Flags & FFrame::Flag_EndHeaders);
 		uint32 ThisId = RecvFrame.GetId();
 
+		if (EndOfHeaders)
+		{
+			RecvFrame = {};
+		}
+
 		FOutcome Outcome = Tick(Peer); 
 		if (Outcome.IsError())				return Outcome;
 		if (EndOfHeaders)					return FOutcome::Ok();
@@ -555,6 +560,7 @@ FOutcome FDriverNg::RecvBody(FTlsPeer& Peer, FMutableMemoryView Dest)
 		// End of the road?
 		if (EndOfStream)
 		{
+			RecvFrame = {};
 			return FOutcome::Ok(RecvSize);
 		}
 
