Delete unused code

#rnx

diff --git a/Engine/Source/Runtime/Experimental/IoStore/HttpClient/Private/TransactionTwo.inl b/Engine/Source/Runtime/Experimental/IoStore/HttpClient/Private/TransactionTwo.inl
--- a/Engine/Source/Runtime/Experimental/IoStore/HttpClient/Private/TransactionTwo.inl
+++ b/Engine/Source/Runtime/Experimental/IoStore/HttpClient/Private/TransactionTwo.inl
@@ -187,7 +187,6 @@ private:
 	FOutcome			TickInner(FTlsPeer& Peer);
 	FOutcome			PumpSend(FTlsPeer& Peer);
 	FOutcome			PumpRecv(FTlsPeer& Peer, FMemoryView Data);
-	nghttp2_ssize 		OnSend(const uint8_t* Data, size_t Size);
 	int32 				OnHeader(nghttp2_rcbuf* Name, nghttp2_rcbuf* Value);
 	nghttp2_session*	Session;
 	TArray<FHeaderNg>*	ResponseHeaders = nullptr;
@@ -223,17 +222,6 @@ FDriverNg::FDriverNg(FTlsPeer& Peer)
 		nghttp2_session_callbacks* Callbacks;
 		nghttp2_session_callbacks_new(&Callbacks);
 
-		/*
-		nghttp2_session_callbacks_set_send_callback2(
-			Callbacks,
-			[] (nghttp2_session*, const uint8_t* Data, size_t Size, int32, void* UserData) -> nghttp2_ssize
-			{
-				auto* This = (FDriverNg*)UserData;
-				return This->OnSend(Data, Size);
-			}
-		);
-		*/
-
 		nghttp2_session_callbacks_set_on_header_callback2(
 			Callbacks,
 			[] (nghttp2_session*, const nghttp2_frame*, nghttp2_rcbuf* Name, nghttp2_rcbuf* Value, uint8_t, void* UserData) -> int32
@@ -571,21 +559,6 @@ FOutcome FDriverNg::RecvBody(FTlsPeer& Peer, FMutableMemoryView Dest)
 	}
 }
 
-////////////////////////////////////////////////////////////////////////////////
-nghttp2_ssize FDriverNg::OnSend(const uint8_t* Data, size_t Size)
-{
-	/*
-	FOutcome Outcome = Peer.Send((char*)Data, int32(Size));
-	if (!Outcome.IsOk())
-	{
-		return Outcome.IsWaiting() ? NGHTTP2_ERR_WOULDBLOCK : NGHTTP2_ERR_CALLBACK_FAILURE;
-	}
-
-	return Outcome.GetResult();
-	*/
-	return -1;
-}
-
 ////////////////////////////////////////////////////////////////////////////////
 int32 FDriverNg::OnHeader(nghttp2_rcbuf* Name, nghttp2_rcbuf* Value)
 {
