cvar for inflight

#rnx

diff --git a/Engine/Source/Runtime/Experimental/IoStore/OnDemand/Private/OnDemandHttpClient.cpp b/Engine/Source/Runtime/Experimental/IoStore/OnDemand/Private/OnDemandHttpClient.cpp
--- a/Engine/Source/Runtime/Experimental/IoStore/OnDemand/Private/OnDemandHttpClient.cpp
+++ b/Engine/Source/Runtime/Experimental/IoStore/OnDemand/Private/OnDemandHttpClient.cpp
@@ -16,6 +16,9 @@
 namespace UE::IoStore
 {
 
+extern int32 GIaxHttpMaxInflight;
+extern int32 GIaxHttpVersion;
+
 /** Not intended as a long term cvar, only including so that we can revert the new behavior easily if needed */
 bool GIaxWarnOnFailure = true;
 static FAutoConsoleVariableRef CVar_IaxWarnOnFailure(
@@ -24,13 +27,6 @@ static FAutoConsoleVariableRef CVar_IaxWarnOnFailure(
 	TEXT("When enabled FMultiEndpointHttpClient will log failed requests as warnings instead of the config verbosity")
 );
 
-int32 GIaxHttpVersion = 1;
-static FAutoConsoleVariableRef CVar_IaxHttpVersion(
-	TEXT("iax.HttpVersion"),
-	GIaxHttpVersion,
-	TEXT("Protocol version to use, either '1' or '2'")
-);
-
 /**
  * When enabled, http requests will make a partial content request by adding Range: bytes=<offset>-<offset+length> to the request if a valid
  * FIoOffsetAndLength is provided with the request. When disabled we will always request the full chunk and then use the FIoOffsetAndLength to
@@ -632,6 +628,7 @@ TUniquePtr<HTTP::FConnectionPool> FMultiEndpointHttpClient::CreateConnection(FAn
 
 	Params.ConnectionCount = uint16(Config.MaxConnectionCount);
 	Params.HttpVersion = (GIaxHttpVersion == 2) ? EHttpVersion::Two : EHttpVersion::One;
+	Params.MaxInflight = uint8(FMath::Clamp(1, 64, GIaxHttpMaxInflight));
 	if (Params.HttpVersion == EHttpVersion::Two)
 	{
 		Params.ConnectionCount = 1; // rfc9113 conformance
diff --git a/Engine/Source/Runtime/Experimental/IoStore/OnDemand/Private/OnDemandHttpThread.cpp b/Engine/Source/Runtime/Experimental/IoStore/OnDemand/Private/OnDemandHttpThread.cpp
--- a/Engine/Source/Runtime/Experimental/IoStore/OnDemand/Private/OnDemandHttpThread.cpp
+++ b/Engine/Source/Runtime/Experimental/IoStore/OnDemand/Private/OnDemandHttpThread.cpp
@@ -94,6 +94,13 @@ FAutoConsoleVariableRef CVarOnDemandBackendThreadPriority(
 	TEXT("Note that this is switchable at runtime"),
 	ECVF_Default);
 
+int32 GIaxHttpVersion = 1;
+static FAutoConsoleVariableRef CVar_IaxHttpVersion(
+	TEXT("iax.HttpVersion"),
+	GIaxHttpVersion,
+	TEXT("Protocol version to use, either '1' or '2'")
+);
+
 int32 GIasHttpConnectionCount = 4;
 static FAutoConsoleVariableRef CVar_IasHttpConnectionCount(
 	TEXT("ias.HttpConnectionCount"),
@@ -102,6 +109,13 @@ static FAutoConsoleVariableRef CVar_IasHttpConnectionCount(
 	FConsoleVariableDelegate::CreateStatic(&OnHttpClientCVarChanged)
 );
 
+int32 GIaxHttpMaxInflight = 1;
+static FAutoConsoleVariableRef CVar_IaxHttpMaxInflight(
+	TEXT("iax.HttpMaxInflight"),
+	GIaxHttpMaxInflight,
+	TEXT("Number of requests to issue at once per connection")
+);
+
 int32 GIasHttpRecvBufKiB = -1;
 static FAutoConsoleVariableRef CVar_GIasHttpRecvBufKiB(
 	TEXT("ias.HttpRecvBufKiB"),
