Clamp parameters were back-to-front.

#rnx

#FailedValidations Preflight

diff --git a/Engine/Source/Runtime/Experimental/IoStore/OnDemand/Private/OnDemandHttpClient.cpp b/Engine/Source/Runtime/Experimental/IoStore/OnDemand/Private/OnDemandHttpClient.cpp
--- a/Engine/Source/Runtime/Experimental/IoStore/OnDemand/Private/OnDemandHttpClient.cpp
+++ b/Engine/Source/Runtime/Experimental/IoStore/OnDemand/Private/OnDemandHttpClient.cpp
@@ -628,7 +628,7 @@ TUniquePtr<HTTP::FConnectionPool> FMultiEndpointHttpClient::CreateConnection(FAn
 
 	Params.ConnectionCount = uint16(Config.MaxConnectionCount);
 	Params.HttpVersion = (GIaxHttpVersion == 2) ? EHttpVersion::Two : EHttpVersion::One;
-	Params.MaxInflight = uint8(FMath::Clamp(1, 64, GIaxHttpMaxInflight));
+	Params.MaxInflight = uint8(FMath::Clamp(GIaxHttpMaxInflight, 1, 64));
 	if (Params.HttpVersion == EHttpVersion::Two)
 	{
 		Params.ConnectionCount = 1; // rfc9113 conformance
