Fixup API macros

#rnx

diff --git a/Engine/Plugins/Runtime/StateTree/Source/StateTreeModule/Public/StateTreeInstanceData.h b/Engine/Plugins/Runtime/StateTree/Source/StateTreeModule/Public/StateTreeInstanceData.h
--- a/Engine/Plugins/Runtime/StateTree/Source/StateTreeModule/Public/StateTreeInstanceData.h
+++ b/Engine/Plugins/Runtime/StateTree/Source/StateTreeModule/Public/StateTreeInstanceData.h
@@ -253,7 +253,7 @@ struct FStateTreeInstanceStorage
 	}
 
 	/** @return the storage's execution runtime data for the state tree. */
-	[[nodiscard]] UE_API UE::StateTree::InstanceData::FInstanceContainer& GetExecutionRuntimeData()
+	[[nodiscard]] UE::StateTree::InstanceData::FInstanceContainer& GetExecutionRuntimeData()
 	{
 		return ExecutionRuntimeData;
 	}
diff --git a/Engine/Source/Runtime/Engine/Public/WorldPartition/HLOD/HLODSourceActors.h b/Engine/Source/Runtime/Engine/Public/WorldPartition/HLOD/HLODSourceActors.h
--- a/Engine/Source/Runtime/Engine/Public/WorldPartition/HLOD/HLODSourceActors.h
+++ b/Engine/Source/Runtime/Engine/Public/WorldPartition/HLOD/HLODSourceActors.h
@@ -24,7 +24,7 @@ public:
 	ENGINE_API const UHLODLayer* GetHLODLayer() const;
 
 	UE_DEPRECATED(5.7, "GetHLODHash() has been replaced by ComputeHLODHash()")
-	ENGINE_API virtual uint32 GetHLODHash() const final { return 0; }
+	virtual uint32 GetHLODHash() const final { return 0; }
 #endif
 
 private:
diff --git a/Engine/Source/Runtime/Experimental/ChaosSolverEngine/Public/Chaos/ChaosVDTraceRelayTransport.h b/Engine/Source/Runtime/Experimental/ChaosSolverEngine/Public/Chaos/ChaosVDTraceRelayTransport.h
--- a/Engine/Source/Runtime/Experimental/ChaosSolverEngine/Public/Chaos/ChaosVDTraceRelayTransport.h
+++ b/Engine/Source/Runtime/Experimental/ChaosSolverEngine/Public/Chaos/ChaosVDTraceRelayTransport.h
@@ -20,18 +20,18 @@ class ITraceDataRelayTransport
 public:
 	virtual ~ITraceDataRelayTransport() = default;
 
-	CHAOSSOLVERENGINE_API virtual void Initialize() = 0;
-	CHAOSSOLVERENGINE_API virtual void Shutdown() = 0;
+	virtual void Initialize() = 0;
+	virtual void Shutdown() = 0;
 
 	DECLARE_DELEGATE_OneParam(FProcessReceivedRelayDataDelegate, const FChaosVDRelayTraceDataMessage& DataPacket);
-	CHAOSSOLVERENGINE_API virtual void RegisterRelayDataReceiverForSessionID(FGuid RemoteSessionID, const FProcessReceivedRelayDataDelegate& InProcessDataDelegate) = 0;
-	CHAOSSOLVERENGINE_API virtual void UnregisterRelayDataReceiverForSessionID(FGuid RemoteSessionID) = 0;
+	virtual void RegisterRelayDataReceiverForSessionID(FGuid RemoteSessionID, const FProcessReceivedRelayDataDelegate& InProcessDataDelegate) = 0;
+	virtual void UnregisterRelayDataReceiverForSessionID(FGuid RemoteSessionID) = 0;
 
-	CHAOSSOLVERENGINE_API virtual void RelayTraceDataFromWriter(FRelayTraceWriter& InRelayTraceWriter) = 0;
+	virtual void RelayTraceDataFromWriter(FRelayTraceWriter& InRelayTraceWriter) = 0;
 
 	DECLARE_DELEGATE_RetVal_OneParam(bool, FCanSendDataDelegate, int32 BytesSent);
-	CHAOSSOLVERENGINE_API virtual void RegisterExternalDataThrottleHandler(const FCanSendDataDelegate& InExternalThrottleHandler) = 0;
-	CHAOSSOLVERENGINE_API virtual void UnregisterExternalDataThrottleHandler() = 0;
+	virtual void RegisterExternalDataThrottleHandler(const FCanSendDataDelegate& InExternalThrottleHandler) = 0;
+	virtual void UnregisterExternalDataThrottleHandler() = 0;
 
 protected:
 	virtual void SendRelayedTraceData(const FChaosVDRelayTraceDataMessage& InRelayData) = 0;
diff --git a/Engine/Source/Runtime/SlateCore/Public/Widgets/SWindow.h b/Engine/Source/Runtime/SlateCore/Public/Widgets/SWindow.h
--- a/Engine/Source/Runtime/SlateCore/Public/Widgets/SWindow.h
+++ b/Engine/Source/Runtime/SlateCore/Public/Widgets/SWindow.h
@@ -804,7 +804,7 @@ public:
 	SLATECORE_API bool AcceptsInput() const;
 
 	/** If true the window will accept input if window type supports it*/
-	SLATECORE_API void SetAcceptsInput(bool bEnabled) { bAcceptsInput = bEnabled; }
+	void SetAcceptsInput(bool bEnabled) { bAcceptsInput = bEnabled; }
 
 	/** @return true if the user decides the size of the window */
 	SLATECORE_API bool IsUserSized() const;
