Make GeometryCollectionComponent destroy its owned chaos event dispatcher on destruction.
Prevents the dispatcher holding a garbage reference to the destroyed GC component.

diff --git a/Engine/Source/Runtime/Experimental/GeometryCollectionEngine/Private/GeometryCollection/GeometryCollectionComponent.cpp b/Engine/Source/Runtime/Experimental/GeometryCollectionEngine/Private/GeometryCollection/GeometryCollectionComponent.cpp
--- a/Engine/Source/Runtime/Experimental/GeometryCollectionEngine/Private/GeometryCollection/GeometryCollectionComponent.cpp
+++ b/Engine/Source/Runtime/Experimental/GeometryCollectionEngine/Private/GeometryCollection/GeometryCollectionComponent.cpp
@@ -809,6 +809,17 @@ void UGeometryCollectionComponent::GetLifetimeReplicatedProps(TArray<FLifetimePr
 	DOREPLIFETIME_WITH_PARAMS_FAST(UGeometryCollectionComponent, RepDynamicData, Params);
 }
 
+void UGeometryCollectionComponent::OnComponentDestroyed(bool bDestroyingHierarchy)
+{
+	if (EventDispatcher)
+	{
+		EventDispatcher->DestroyComponent();
+		EventDispatcher = nullptr;
+	}
+
+	Super::OnComponentDestroyed(bDestroyingHierarchy);
+}
+
 namespace
 {
 	template <typename TTransformType>
diff --git a/Engine/Source/Runtime/Experimental/GeometryCollectionEngine/Public/GeometryCollection/GeometryCollectionComponent.h b/Engine/Source/Runtime/Experimental/GeometryCollectionEngine/Public/GeometryCollection/GeometryCollectionComponent.h
--- a/Engine/Source/Runtime/Experimental/GeometryCollectionEngine/Public/GeometryCollection/GeometryCollectionComponent.h
+++ b/Engine/Source/Runtime/Experimental/GeometryCollectionEngine/Public/GeometryCollection/GeometryCollectionComponent.h
@@ -612,6 +612,7 @@ public:
 	GEOMETRYCOLLECTIONENGINE_API virtual void OnHiddenInGameChanged() override;
 	GEOMETRYCOLLECTIONENGINE_API virtual void GetLifetimeReplicatedProps(TArray<FLifetimeProperty>& OutLifetimeProps) const override;
 	GEOMETRYCOLLECTIONENGINE_API virtual void GetResourceSizeEx(FResourceSizeEx& CumulativeResourceSize) override;
+	GEOMETRYCOLLECTIONENGINE_API virtual void OnComponentDestroyed(bool bDestroyingHierarchy) override;
 
 #if WITH_EDITOR
 	GEOMETRYCOLLECTIONENGINE_API virtual void PostEditChangeChainProperty(FPropertyChangedChainEvent& PropertyChangedEvent) override;
