Optimize Pull Data
Multithread pull data for rigid particles
Up to 15% improvement in some scenarios

diff --git a/Engine/Source/Runtime/Experimental/Chaos/Private/PBDRigidsSolver.cpp b/Engine/Source/Runtime/Experimental/Chaos/Private/PBDRigidsSolver.cpp
--- a/Engine/Source/Runtime/Experimental/Chaos/Private/PBDRigidsSolver.cpp
+++ b/Engine/Source/Runtime/Experimental/Chaos/Private/PBDRigidsSolver.cpp
@@ -609,6 +609,7 @@ namespace Chaos
 			//TODO: handle this more gracefully
 			if(MDeltaTime > 0)
 			{
+				CSV_SCOPED_TIMING_STAT(PhysicsVerbose, StepSolver_PullData);
 				MSolver->CompleteSceneSimulation();
 			}
 
@@ -1688,7 +1689,7 @@ namespace Chaos
 				{
 					case EPhysicsProxyType::SingleParticleProxy:
 					{
-						auto Proxy = static_cast<FSingleParticlePhysicsProxy*>(Dirty.Proxy);
+						FSingleParticlePhysicsProxy* Proxy = static_cast<FSingleParticlePhysicsProxy*>(Dirty.Proxy);
 						ProcessProxyPT(Proxy, DataIdx, Dirty, [this, &Dirty](const FUniqueIdx* UniqueIdx) -> TGeometryParticleHandle<FReal,3>*
 						{
 							switch (Dirty.PropertyData.GetParticleBufferType())
@@ -1716,7 +1717,7 @@ namespace Chaos
 			
 					case EPhysicsProxyType::GeometryCollectionType:
 					{
-						auto Proxy = static_cast<FGeometryCollectionPhysicsProxy*>(Dirty.Proxy);
+						FGeometryCollectionPhysicsProxy* Proxy = static_cast<FGeometryCollectionPhysicsProxy*>(Dirty.Proxy);
 						if (!Proxy->IsInitializedOnPhysicsThread())
 						{
 							// Finish registration on the physics thread...
@@ -2390,7 +2391,8 @@ namespace Chaos
 				}
 			}
 		}
-
+		TArray<UE::Tasks::FTask> RigidsPendingTasks;
+		if (GSingleThreadedPhysics)
 		{
 			//we only fill this once per frame
 			ensure(PullData->DirtyRigids.Num() == 0);
@@ -2402,6 +2404,36 @@ namespace Chaos
 				ActiveRigid[Idx]->BufferPhysicsResults(PullData->DirtyRigids.Last());
 			}
 		}
+		else
+		{
+			constexpr int32 MinRigidPerTask = 50;
+			const int32 NumRigids = ActiveRigid.Num();
+			const int32 NumTasks = FMath::Max(FMath::Min(FTaskGraphInterface::Get().GetNumWorkerThreads()-1, Chaos::MaxNumWorkers), 1);
+			const int32 NumRigidTask = FMath::Max(FMath::Min(NumTasks, NumRigids), 1);
+			const int32 RigidsByTask = FMath::Max(FMath::DivideAndRoundUp(NumRigids, NumRigidTask), MinRigidPerTask);
+			const int32 NumBatches = FMath::DivideAndRoundUp(NumRigids, RigidsByTask);
+
+			//we only fill this once per frame
+			ensure(PullData->DirtyRigids.Num() == 0);
+			PullData->DirtyRigids.SetNum(ActiveRigid.Num());
+			RigidsPendingTasks.Reserve(NumBatches);
+			for (int32 BatchIndex = 0; BatchIndex < NumBatches; BatchIndex++)
+			{
+				const int32 StartIndex = BatchIndex * RigidsByTask;
+				int32 EndIndex = (BatchIndex + 1) * RigidsByTask;
+				EndIndex = FMath::Min(NumRigids, EndIndex);
+
+				UE::Tasks::FTask PendingTask = UE::Tasks::Launch(UE_SOURCE_LOCATION, [StartIndex, EndIndex, &ActiveRigid , &PullData]()
+				{
+					QUICK_SCOPE_CYCLE_COUNTER(STAT_BufferPhysicsResultsTasks);
+					for (int32 Idx = StartIndex; Idx < EndIndex; ++Idx)
+					{
+						ActiveRigid[Idx]->BufferPhysicsResults(PullData->DirtyRigids[Idx]);
+					}
+				});
+				RigidsPendingTasks.Add(PendingTask);
+			}
+		}
 
 		{
 			ensure(PullData->DirtyGeometryCollections.Num() == 0);	//we only fill this once per frame
@@ -2453,6 +2485,11 @@ namespace Chaos
 			RewindData->BufferPhysicsResults(PullData->DirtyRigidErrors);
 		}
 
+		if (!GSingleThreadedPhysics)
+		{
+			UE::Tasks::Wait(RigidsPendingTasks);
+		}
+
 		// Now that results have been buffered we have completed a solve step so we can broadcast that event
 		EventPostSolve.Broadcast(MLastDt);
 
