Optimize PullData
Add a task dispatcher for the pull data to keep the threading logic outside of the physics logic.
Keep a single threaded version in the physics code to see what the logic is.
Keep merging data in Pull Data in a deterministic order, to conserve determinism

Up to 7 times faster improvement in some scenarios

#FailedValidations Preflight

diff --git a/Engine/Source/Runtime/Experimental/Chaos/Private/PBDRigidsSolver.cpp b/Engine/Source/Runtime/Experimental/Chaos/Private/PBDRigidsSolver.cpp
--- a/Engine/Source/Runtime/Experimental/Chaos/Private/PBDRigidsSolver.cpp
+++ b/Engine/Source/Runtime/Experimental/Chaos/Private/PBDRigidsSolver.cpp
@@ -37,6 +37,7 @@
 
 #include "ProfilingDebugging/CountersTrace.h"
 #include "ProfilingDebugging/CsvProfiler.h"
+#include "TaskDispatcherSolver.h"
 #include "ChaosLog.h"
 
 //PRAGMA_DISABLE_OPTIMIZATION
@@ -658,6 +659,7 @@ namespace Chaos
 		, bIsFloorAnalytic(false)
 		, FloorHeight(0.f)
 		, bIsDeterministic(false)
+		, TaskDispatcher(MakeUnique<Private::FTaskDispatcherSolver>())
 		, QueryMaterials_External()
 		, QueryMaterialMasks_External()
 		, SimMaterials()
@@ -2257,47 +2259,141 @@ namespace Chaos
 		BufferPhysicsResults();
 	}
 
-	void FPBDRigidsSolver::BufferPhysicsResults()
-	{
-		//ensure(IsInPhysicsThread());
-
-		TArray<FGeometryCollectionPhysicsProxy*> ActiveGC;
-		ActiveGC.Reserve(GeometryCollectionPhysicsProxies_Internal.Num());
-
-		FPullPhysicsData* PullData = MarshallingManager.GetCurrentPullData_Internal();
-
-		TParticleView<FPBDRigidParticles>& DirtyParticles = GetParticles().GetDirtyParticlesView();
-		const int32 NumDirty = DirtyParticles.Num();
-
-		// We can have at most NumDirty overall active particles across all dirty proxy lists
-		// and in each list at most NumProxies. So we take the min, otherwise for example in a world
-		// with a few million static particles and one dynamic we'll allocate enough storage for the
-		// millions of particles despite only ever syncing one
-		TArray<FSingleParticlePhysicsProxy*> ActiveRigid;
-		ActiveRigid.Reserve(FMath::Min(SingleParticlePhysicsProxies_PT.Num(), NumDirty));
 
-		TArray<FClusterUnionPhysicsProxy*> ActiveClusterUnions;
-		ActiveClusterUnions.Reserve(FMath::Min(ClusterUnionPhysicsProxies_Internal.Num(), NumDirty));
+	template <typename LambdaType>
+	int32 FPBDRigidsSolver::DispatchTasksBufferPhysicsResults(int32 DispatchBatchIndex, int32 NumDirty, LambdaType ReturnHandle)
+	{
+		using TSOA = typename TParticleView<FPBDRigidParticles>::TSOA;
+		using THandle = typename TSOA::THandleType;
+		using THandleBase = typename THandle::THandleBase;
 
 		const bool bIsResim = GetEvolution()->IsResimming();
 
-		//todo: should be able to go wide just add defaulted etc...
+		// For performance reason collect on one threads the rigid proxies and on another one the geometry collection and the cluster union.
+		// That way buffering results from rigid proxy could start after the first batch of rigid proxy collected, while geometry collection, cluster union 
+		// and other batches of rigid cluster will be still collected.
+		auto CollectRigidResults = [bIsResim, ReturnHandle](int32 StartIndex, int32 EndIndex, TArray<FSingleParticlePhysicsProxy*>& ActiveRigid)
 		{
-			ensure(PullData->DirtyRigids.Num() == 0);	//we only fill this once per frame
-			int32 BufferIdx = 0;
-
-			for (TPBDRigidParticleHandleImp<FReal, 3, false>& DirtyParticle : DirtyParticles)
+			for (int32 Idx = StartIndex; Idx < EndIndex; ++Idx)
 			{
+				THandleBase Handle = ReturnHandle(Idx);
+				FTransientPBDRigidParticleHandle& DirtyParticle = static_cast<FTransientPBDRigidParticleHandle&>(Handle);
+
 				const EParticleType ParticleType = DirtyParticle.GetParticleType();
 				if ((ParticleType == EParticleType::Kinematic) || (ParticleType == EParticleType::Static))
 				{
 					ensure(false);
 				}
 
-				if(IPhysicsProxyBase* Proxy = DirtyParticle.Handle()->PhysicsProxy())
+				if (IPhysicsProxyBase* Proxy = DirtyParticle.Handle()->PhysicsProxy())
 				{
 					switch (Proxy->GetType())
 					{
+					case EPhysicsProxyType::GeometryCollectionType:
+					case EPhysicsProxyType::ClusterUnionProxy:
+					{
+						if (auto ClusterParticle = DirtyParticle.CastToClustered())
+						{
+							if (ClusterParticle->InternalCluster())
+							{
+
+								const bool bActivateChildren =
+									ClusterParticle->ObjectState() == EObjectStateType::Dynamic ||
+									Proxy->GetType() != EPhysicsProxyType::ClusterUnionProxy;
+
+								// If the cluster is dynamic, all children will likely need a transform update
+								// otherwise we should not need a full update (if the cluster is still we only
+								// need to consider de-clustered objects which should be in other parts of the 
+								// dirty view)
+								if (bActivateChildren)
+								{
+									const TSet<IPhysicsProxyBase*> Proxies = ClusterParticle->PhysicsProxies();
+									for (IPhysicsProxyBase* ClusterProxy : Proxies)
+									{
+										if (!ClusterProxy)
+										{
+											continue;
+										}
+										switch (ClusterProxy->GetType())
+										{
+										case EPhysicsProxyType::SingleParticleProxy:
+											ActiveRigid.Add((FSingleParticlePhysicsProxy*)ClusterProxy);
+											break;
+										case EPhysicsProxyType::GeometryCollectionType:
+											break;
+										default:
+											ensure(false);
+											break;
+										}
+									}
+								}
+							}
+							else
+							{
+								switch (Proxy->GetType())
+								{
+								case EPhysicsProxyType::SingleParticleProxy:
+									ActiveRigid.Add((FSingleParticlePhysicsProxy*)Proxy);
+									break;
+								case EPhysicsProxyType::GeometryCollectionType:
+									break;
+								default:
+									ensure(false);
+									break;
+								}
+							}
+						}
+						else
+						{
+							ensure(false);
+						}
+					}
+					break;
+					case EPhysicsProxyType::SingleParticleProxy:
+					{
+						if (!bIsResim || DirtyParticle.SyncState() == ESyncState::HardDesync)
+						{
+							// Although per-particle we can control the syncing of target positions (see ShouldUpdateFromSimulation)
+							// we cannot avoid marking kinematics as dirty (unless the global config Chaos::SyncKinematicOnGameThread
+							// forces it) because we always need the correct velocities/dynamics to be synced back from the kinematic
+							// target.
+							// FSingleParticlePhysicsProxy ::PullFromPhysicsState will use the proper checks to see whether it needs
+							// to sync the particle positions when the dirty particle is processed.
+							if (!(Chaos::SyncKinematicOnGameThread == 0 && DirtyParticle.ObjectState() == EObjectStateType::Kinematic))
+							{
+								ActiveRigid.Add((FSingleParticlePhysicsProxy*)Proxy);
+							}
+						}
+					}
+					break;
+
+					default:
+						check(false);
+						break;
+					}
+				}
+			}
+		};
+
+		auto CollectGCAndClustersResults = [bIsResim, ReturnHandle](int32 StartIndex, int32 EndIndex,
+			TArray<FGeometryCollectionPhysicsProxy*>& ActiveGC,
+			TArray<FClusterUnionPhysicsProxy*>& ActiveClusterUnions)
+			{
+				for (int32 Idx = StartIndex; Idx < EndIndex; ++Idx)
+				{
+					THandleBase Handle = ReturnHandle(Idx);
+					FTransientPBDRigidParticleHandle& DirtyParticle = static_cast<FTransientPBDRigidParticleHandle&>(Handle);
+
+					const EParticleType ParticleType = DirtyParticle.GetParticleType();
+					if ((ParticleType == EParticleType::Kinematic) || (ParticleType == EParticleType::Static))
+					{
+						ensure(false);
+					}
+
+					if (IPhysicsProxyBase* Proxy = DirtyParticle.Handle()->PhysicsProxy())
+					{
+						switch (Proxy->GetType())
+						{
 						case EPhysicsProxyType::GeometryCollectionType:
 						case EPhysicsProxyType::ClusterUnionProxy:
 						{
@@ -2331,7 +2427,6 @@ namespace Chaos
 											switch (ClusterProxy->GetType())
 											{
 											case EPhysicsProxyType::SingleParticleProxy:
-												ActiveRigid.AddUnique((FSingleParticlePhysicsProxy*)ClusterProxy);
 												break;
 											case EPhysicsProxyType::GeometryCollectionType:
 												ActiveGC.AddUnique((FGeometryCollectionPhysicsProxy*)(ClusterProxy));
@@ -2348,7 +2443,6 @@ namespace Chaos
 									switch (Proxy->GetType())
 									{
 									case EPhysicsProxyType::SingleParticleProxy:
-										ActiveRigid.AddUnique((FSingleParticlePhysicsProxy*)Proxy);
 										break;
 									case EPhysicsProxyType::GeometryCollectionType:
 										ActiveGC.AddUnique((FGeometryCollectionPhysicsProxy*)(Proxy));
@@ -2366,94 +2460,240 @@ namespace Chaos
 						}
 						break;
 						case EPhysicsProxyType::SingleParticleProxy:
+							break;
+
+						default:
+							check(false);
+							break;
+						}
+					}
+				}
+			};
+
+		return TaskDispatcher->CollectPhysicsResults(CollectRigidResults, CollectGCAndClustersResults, NumDirty, DispatchBatchIndex);
+	}
+
+	void FPBDRigidsSolver::BufferPhysicsResultsSingleThreaded(TArray<FSingleParticlePhysicsProxy*>& ActiveRigid, TArray<FGeometryCollectionPhysicsProxy*>& ActiveGC, TArray<FClusterUnionPhysicsProxy*>& ActiveClusterUnions)
+	{
+		ActiveGC.Reserve(GeometryCollectionPhysicsProxies_Internal.Num());
+
+		FPullPhysicsData* PullData = MarshallingManager.GetCurrentPullData_Internal();
+
+		TParticleView<FPBDRigidParticles>& DirtyParticles = GetParticles().GetDirtyParticlesView();
+		const int32 NumDirty = DirtyParticles.Num();
+
+		// We can have at most NumDirty overall active particles across all dirty proxy lists
+		// and in each list at most NumProxies. So we take the min, otherwise for example in a world
+		// with a few million static particles and one dynamic we'll allocate enough storage for the
+		// millions of particles despite only ever syncing one
+		
+		ActiveRigid.Reserve(FMath::Min(SingleParticlePhysicsProxies_PT.Num(), NumDirty));
+
+		ActiveClusterUnions.Reserve(FMath::Min(ClusterUnionPhysicsProxies_Internal.Num(), NumDirty));
+
+		const bool bIsResim = GetEvolution()->IsResimming();
+
+		{
+			ensure(PullData->DirtyRigids.Num() == 0);	//we only fill this once per frame
+			int32 BufferIdx = 0;
+
+			for (TPBDRigidParticleHandleImp<FReal, 3, false>& DirtyParticle : DirtyParticles)
+			{
+				const EParticleType ParticleType = DirtyParticle.GetParticleType();
+				if ((ParticleType == EParticleType::Kinematic) || (ParticleType == EParticleType::Static))
+				{
+					ensure(false);
+				}
+
+				if (IPhysicsProxyBase* Proxy = DirtyParticle.Handle()->PhysicsProxy())
+				{
+					switch (Proxy->GetType())
+					{
+					case EPhysicsProxyType::GeometryCollectionType:
+					case EPhysicsProxyType::ClusterUnionProxy:
+					{
+						if (auto ClusterParticle = DirtyParticle.CastToClustered())
 						{
-							if (!bIsResim || DirtyParticle.SyncState() == ESyncState::HardDesync)
+							if (ClusterParticle->InternalCluster())
+							{
+								if (Proxy->GetType() == EPhysicsProxyType::ClusterUnionProxy)
+								{
+									ActiveClusterUnions.AddUnique((FClusterUnionPhysicsProxy*)(Proxy));
+								}
+
+								const bool bActivateChildren =
+									ClusterParticle->ObjectState() == EObjectStateType::Dynamic ||
+									Proxy->GetType() != EPhysicsProxyType::ClusterUnionProxy;
+
+								// If the cluster is dynamic, all children will likely need a transform update
+								// otherwise we should not need a full update (if the cluster is still we only
+								// need to consider de-clustered objects which should be in other parts of the 
+								// dirty view)
+								if (bActivateChildren)
+								{
+									const TSet<IPhysicsProxyBase*> Proxies = ClusterParticle->PhysicsProxies();
+									for (IPhysicsProxyBase* ClusterProxy : Proxies)
+									{
+										if (!ClusterProxy)
+										{
+											continue;
+										}
+
+										switch (ClusterProxy->GetType())
+										{
+										case EPhysicsProxyType::SingleParticleProxy:
+											ActiveRigid.AddUnique((FSingleParticlePhysicsProxy*)ClusterProxy);
+											break;
+										case EPhysicsProxyType::GeometryCollectionType:
+											ActiveGC.AddUnique((FGeometryCollectionPhysicsProxy*)(ClusterProxy));
+											break;
+										default:
+											ensure(false);
+											break;
+										}
+									}
+								}
+							}
+							else
 							{
-								// Although per-particle we can control the syncing of target positions (see ShouldUpdateFromSimulation)
-								// we cannot avoid marking kinematics as dirty (unless the global config Chaos::SyncKinematicOnGameThread
-								// forces it) because we always need the correct velocities/dynamics to be synced back from the kinematic
-								// target.
-								// FSingleParticlePhysicsProxy::PullFromPhysicsState will use the proper checks to see whether it needs
-								// to sync the particle positions when the dirty particle is processed.
-								if (!(Chaos::SyncKinematicOnGameThread == 0 && DirtyParticle.ObjectState() == EObjectStateType::Kinematic))
+								switch (Proxy->GetType())
 								{
+								case EPhysicsProxyType::SingleParticleProxy:
 									ActiveRigid.AddUnique((FSingleParticlePhysicsProxy*)Proxy);
+									break;
+								case EPhysicsProxyType::GeometryCollectionType:
+									ActiveGC.AddUnique((FGeometryCollectionPhysicsProxy*)(Proxy));
+									break;
+								default:
+									ensure(false);
+									break;
 								}
 							}
 						}
-						break;
+						else
+						{
+							ensure(false);
+						}
+					}
+					break;
+					case EPhysicsProxyType::SingleParticleProxy:
+					{
+						if (!bIsResim || DirtyParticle.SyncState() == ESyncState::HardDesync)
+						{
+							// Although per-particle we can control the syncing of target positions (see ShouldUpdateFromSimulation)
+							// we cannot avoid marking kinematics as dirty (unless the global config Chaos::SyncKinematicOnGameThread
+							// forces it) because we always need the correct velocities/dynamics to be synced back from the kinematic
+							// target.
+							// FSingleParticlePhysicsProxy::PullFromPhysicsState will use the proper checks to see whether it needs
+							// to sync the particle positions when the dirty particle is processed.
+							if (!(Chaos::SyncKinematicOnGameThread == 0 && DirtyParticle.ObjectState() == EObjectStateType::Kinematic))
+							{
+								ActiveRigid.AddUnique((FSingleParticlePhysicsProxy*)Proxy);
+							}
+						}
+					}
+					break;
 
-						default:
-							check(false);
-							break;
+					default:
+						check(false);
+						break;
 					}
 
 				}
 			}
 		}
-		TArray<UE::Tasks::FTask> RigidsPendingTasks;
-		if (GSingleThreadedPhysics)
+	}
+
+	void FPBDRigidsSolver::BufferPhysicsResults()
+	{
+		FPullPhysicsData* PullData = MarshallingManager.GetCurrentPullData_Internal();
+
+		if (!GSingleThreadedPhysics)
 		{
-			//we only fill this once per frame
-			ensure(PullData->DirtyRigids.Num() == 0);
-			PullData->DirtyRigids.Reserve(ActiveRigid.Num());
+			using TSOA = typename TParticleView<FPBDRigidParticles>::TSOA;
+			using THandle = typename TSOA::THandleType;
+			using THandleBase = typename THandle::THandleBase;
+
+			TParticleView<FPBDRigidParticles>& DirtyParticles = GetParticles().GetDirtyParticlesView();
+			const int32 NumDirty = DirtyParticles.Num();
+			int32 DispatchBatchIndex = 0;
+			TaskDispatcher->Initialize(DirtyParticles);
 
-			for (int32 Idx = 0; Idx < ActiveRigid.Num(); ++Idx)
+			for (int32 ViewIndex = 0; ViewIndex < DirtyParticles.SOAViews.Num(); ++ViewIndex)
 			{
-				PullData->DirtyRigids.AddDefaulted();
-				ActiveRigid[Idx]->BufferPhysicsResults(PullData->DirtyRigids.Last());
+				const TSOAView<TSOA>& SOAView = DirtyParticles.SOAViews[ViewIndex];
+				const int32 NumParticles = SOAView.Size();
+				if (NumParticles == 0)
+				{
+					continue;
+				}
+				// Dispatch tasks according if it is an array of handles or a SOA
+				if (const TArray<THandle*>* CurHandlesArray = SOAView.HandlesArray)
+				{
+					auto ReturnHandle = [CurHandlesArray](int32 Index) -> TParticleView<FPBDRigidParticles>::THandle::THandleBase
+					{
+						THandle* HandlePtr = (*CurHandlesArray)[Index];
+						return THandleBase(HandlePtr->GeometryParticles, HandlePtr->ParticleIdx);
+					};
+					using LambdaType = decltype(ReturnHandle);
+					const int32 NumHandles = CurHandlesArray->Num();
+					check(NumHandles == NumParticles);
+					DispatchBatchIndex += DispatchTasksBufferPhysicsResults<LambdaType>(DispatchBatchIndex, NumHandles, ReturnHandle);
+				}
+				else if (TSOA* Soa = SOAView.SOA)
+				{
+					auto ReturnHandle = [Soa](int32 Index) { return THandleBase(Soa, Index); };
+					using LambdaType = decltype(ReturnHandle);
+					DispatchBatchIndex += DispatchTasksBufferPhysicsResults<LambdaType>(DispatchBatchIndex, NumParticles, ReturnHandle);
+				}
 			}
-		}
-		else
-		{
-			constexpr int32 MinRigidPerTask = 50;
-			const int32 NumRigids = ActiveRigid.Num();
-			const int32 NumTasks = FMath::Max(FMath::Min(FTaskGraphInterface::Get().GetNumWorkerThreads()-1, Chaos::MaxNumWorkers), 1);
-			const int32 NumRigidTask = FMath::Max(FMath::Min(NumTasks, NumRigids), 1);
-			const int32 RigidsByTask = FMath::Max(FMath::DivideAndRoundUp(NumRigids, NumRigidTask), MinRigidPerTask);
-			const int32 NumBatches = FMath::DivideAndRoundUp(NumRigids, RigidsByTask);
 
 			//we only fill this once per frame
 			ensure(PullData->DirtyRigids.Num() == 0);
-			PullData->DirtyRigids.SetNum(ActiveRigid.Num());
-			RigidsPendingTasks.Reserve(NumBatches);
-			for (int32 BatchIndex = 0; BatchIndex < NumBatches; BatchIndex++)
+			ensure(PullData->DirtyGeometryCollections.Num() == 0);
+			ensure(PullData->DirtyClusterUnions.IsEmpty());
+			check(TaskDispatcher->GetNumRigidBatches() == DispatchBatchIndex);
+			TaskDispatcher->BufferRigidResults(PullData);
+			TaskDispatcher->BufferGCResults(PullData, this);
+			TaskDispatcher->BufferClusterResults(PullData);
+		}
+		else
+		{
+			// Those arrays are here only for the single threaded version
+			TArray<FSingleParticlePhysicsProxy*> ActiveRigid;
+			TArray<FGeometryCollectionPhysicsProxy*> ActiveGC;
+			TArray<FClusterUnionPhysicsProxy*> ActiveClusterUnions;
+			BufferPhysicsResultsSingleThreaded(ActiveRigid, ActiveGC, ActiveClusterUnions);
 			{
-				const int32 StartIndex = BatchIndex * RigidsByTask;
-				int32 EndIndex = (BatchIndex + 1) * RigidsByTask;
-				EndIndex = FMath::Min(NumRigids, EndIndex);
+				//we only fill this once per frame
+				ensure(PullData->DirtyRigids.Num() == 0);
+				PullData->DirtyRigids.Reserve(ActiveRigid.Num());
 
-				UE::Tasks::FTask PendingTask = UE::Tasks::Launch(UE_SOURCE_LOCATION, [StartIndex, EndIndex, &ActiveRigid , &PullData]()
+				for (int32 Idx = 0; Idx < ActiveRigid.Num(); ++Idx)
 				{
-					QUICK_SCOPE_CYCLE_COUNTER(STAT_BufferPhysicsResultsTasks);
-					for (int32 Idx = StartIndex; Idx < EndIndex; ++Idx)
-					{
-						ActiveRigid[Idx]->BufferPhysicsResults(PullData->DirtyRigids[Idx]);
-					}
-				});
-				RigidsPendingTasks.Add(PendingTask);
+					PullData->DirtyRigids.AddDefaulted();
+					ActiveRigid[Idx]->BufferPhysicsResults(PullData->DirtyRigids.Last());
+				}
 			}
-		}
-
-		{
-			ensure(PullData->DirtyGeometryCollections.Num() == 0);	//we only fill this once per frame
-			PullData->DirtyGeometryCollections.Reserve(ActiveGC.Num());
-
-			for (int32 Idx = 0; Idx < ActiveGC.Num(); ++Idx)
 			{
-				PullData->DirtyGeometryCollections.AddDefaulted();
-				ActiveGC[Idx]->BufferPhysicsResults_Internal(this, PullData->DirtyGeometryCollections.Last());
-			}
-		}
-
-		{
-			ensure(PullData->DirtyClusterUnions.IsEmpty());
-			PullData->DirtyClusterUnions.Reserve(ActiveClusterUnions.Num());
+				ensure(PullData->DirtyGeometryCollections.Num() == 0);	//we only fill this once per frame
+				PullData->DirtyGeometryCollections.Reserve(ActiveGC.Num());
 
-			for (int32 Idx = 0; Idx < ActiveClusterUnions.Num(); ++Idx)
+				for (int32 Idx = 0; Idx < ActiveGC.Num(); ++Idx)
+				{
+					PullData->DirtyGeometryCollections.AddDefaulted();
+					ActiveGC[Idx]->BufferPhysicsResults_Internal(this, PullData->DirtyGeometryCollections.Last());
+				}
+			}
 			{
-				PullData->DirtyClusterUnions.AddDefaulted();
-				ActiveClusterUnions[Idx]->BufferPhysicsResults_Internal(PullData->DirtyClusterUnions.Last());
+				ensure(PullData->DirtyClusterUnions.IsEmpty());
+				PullData->DirtyClusterUnions.Reserve(ActiveClusterUnions.Num());
+
+				for (int32 Idx = 0; Idx < ActiveClusterUnions.Num(); ++Idx)
+				{
+					PullData->DirtyClusterUnions.AddDefaulted();
+					ActiveClusterUnions[Idx]->BufferPhysicsResults_Internal(PullData->DirtyClusterUnions.Last());
+				}
 			}
 		}
 
@@ -2487,7 +2727,7 @@ namespace Chaos
 
 		if (!GSingleThreadedPhysics)
 		{
-			UE::Tasks::Wait(RigidsPendingTasks);
+			TaskDispatcher->WaitTaskEndBufferResults();
 		}
 
 		// Now that results have been buffered we have completed a solve step so we can broadcast that event
diff --git a/Engine/Source/Runtime/Experimental/Chaos/Private/PhysicsProxy/GeometryCollectionPhysicsProxy.cpp b/Engine/Source/Runtime/Experimental/Chaos/Private/PhysicsProxy/GeometryCollectionPhysicsProxy.cpp
--- a/Engine/Source/Runtime/Experimental/Chaos/Private/PhysicsProxy/GeometryCollectionPhysicsProxy.cpp
+++ b/Engine/Source/Runtime/Experimental/Chaos/Private/PhysicsProxy/GeometryCollectionPhysicsProxy.cpp
@@ -3839,7 +3839,8 @@ static void UpdateParticleHandleTransform(Chaos::FPBDRigidsSolver& CurrentSolver
 	Handle.SetX(NewTransform.GetTranslation());
 	Handle.SetR(NewTransform.GetRotation());
 	Handle.UpdateWorldSpaceState(NewTransform, Chaos::FVec3{0});
-	CurrentSolver.GetEvolution()->DirtyParticle(Handle);
+	// This current function is called by BufferPhysicsResults_Internal which is called by PullData in multi thread
+	CurrentSolver.GetEvolution()->DirtyParticle_ThreadSafe(Handle);
 }
 
 static void UpdateParticleHandleTransformIfNeeded(Chaos::FPBDRigidsSolver& CurrentSolver, Chaos::FPBDRigidClusteredParticleHandle& Handle, const Chaos::FRigidTransform3& NewTransform)
diff --git a/Engine/Source/Runtime/Experimental/Chaos/Private/TaskDispatcherSolver.cpp b/Engine/Source/Runtime/Experimental/Chaos/Private/TaskDispatcherSolver.cpp
new file mode 100644
--- /dev/null
+++ b/Engine/Source/Runtime/Experimental/Chaos/Private/TaskDispatcherSolver.cpp
@@ -0,0 +1,200 @@
+// Copyright Epic Games, Inc. All Rights Reserved.
+
+#include "TaskDispatcherSolver.h"
+
+#include "PhysicsProxy/ClusterUnionPhysicsProxy.h"
+#include "PhysicsProxy/GeometryCollectionPhysicsProxy.h"
+
+
+namespace Chaos::Private
+{
+
+void FTaskDispatcherSolver::Initialize(TParticleView<FPBDRigidParticles>& DirtyParticles)
+{
+	using TSOA = typename TParticleView<FPBDRigidParticles>::TSOA;
+	using THandle = typename TSOA::THandleType;
+	using THandleBase = typename THandle::THandleBase;
+	using TTransientHandle = typename THandle::TTransientHandle;
+	int32 DispatchBatchIndex = 0;
+
+	// Precompute NumRigidBatches 
+	NumRigidBatches = 0;
+	for (int32 ViewIndex = 0; ViewIndex < DirtyParticles.SOAViews.Num(); ++ViewIndex)
+	{
+		const TSOAView<TSOA>& SOAView = DirtyParticles.SOAViews[ViewIndex];
+		const int32 NumParticles = SOAView.Size();
+		if (NumParticles == 0)
+		{
+			continue;
+		}
+		constexpr int32 MinRigidPerTask = 50;
+		const int32 NumTasks = FMath::Max(FMath::Min(FTaskGraphInterface::Get().GetNumWorkerThreads() - 1, Chaos::MaxNumWorkers), 1);
+		const int32 NumRigidTask = FMath::Max(FMath::Min(NumTasks, NumParticles), 1);
+		const int32 DirtyByTask = FMath::Max(FMath::DivideAndRoundUp(NumParticles, NumRigidTask), MinRigidPerTask);
+		NumRigidBatches += FMath::DivideAndRoundUp(NumParticles, DirtyByTask);
+	}
+
+	DirtyRigidPendingTasks.Reset(NumRigidBatches);
+	DirtyGCAndClusterPendingTask = UE::Tasks::FTask();
+
+	ActiveRigidArray.SetNum(NumRigidBatches);
+
+	PullDataDirtyRigidsArray.Reset();
+	PullDataDirtyGCArray.Reset();
+	PullDataDirtyClusterArray.Reset();
+
+	PreviousCopyRigidTask = UE::Tasks::FTask();
+	PreviousCopyGCTask = UE::Tasks::FTask();
+	PreviousCopyClusterTask = UE::Tasks::FTask();
+
+}
+
+void FTaskDispatcherSolver::BufferRigidResults(FPullPhysicsData* PullData)
+{
+	// Num Batch should be the same for the three arrays
+	check(ActiveRigidArray.Num() == NumRigidBatches);
+	check(ActiveRigidArray.Num() == DirtyRigidPendingTasks.Num());
+
+	PullDataDirtyRigidsArray.SetNum(NumRigidBatches);
+	//we only fill this once per frame
+	ensure(PullData->DirtyRigids.Num() == 0);
+
+	for (int32 BatchIndex = 0; BatchIndex < NumRigidBatches; BatchIndex++)
+	{
+		TArray<FDirtyRigidParticleData>& PullDataDirtyRigids = PullDataDirtyRigidsArray[BatchIndex];
+		PullDataDirtyRigids.Reset();
+		TArray<FSingleParticlePhysicsProxy*>& ActiveRigid = ActiveRigidArray[BatchIndex];
+		UE::Tasks::FTask RigidPendingTask = UE::Tasks::Launch(UE_SOURCE_LOCATION, [&ActiveRigid, &PullDataDirtyRigids]()
+		{
+			QUICK_SCOPE_CYCLE_COUNTER(STAT_BufferRigidResultsTasks);
+			for (int32 Idx = 0; Idx < ActiveRigid.Num(); ++Idx)
+			{
+				PullDataDirtyRigids.AddDefaulted();
+				ActiveRigid[Idx]->BufferPhysicsResults(PullDataDirtyRigids.Last());
+			}
+		}, DirtyRigidPendingTasks[BatchIndex]);
+
+		// Link the following task with the one from the previous iteration to keep determinism
+		UE::Tasks::FTask RigidCopyPendingTask = UE::Tasks::Launch(UE_SOURCE_LOCATION, [PullData, &PullDataDirtyRigids]()
+		{
+			QUICK_SCOPE_CYCLE_COUNTER(STAT_BufferRigidResultsCopyTasks);
+			PullData->DirtyRigids.Append(PullDataDirtyRigids);
+		}, TArray<UE::Tasks::FTask>{RigidPendingTask, PreviousCopyRigidTask});
+		PreviousCopyRigidTask = RigidCopyPendingTask;
+	}
+}
+
+void FTaskDispatcherSolver::BufferGCResults(FPullPhysicsData* PullData, FPBDRigidsSolver* RigidsSolver)
+{
+	GCPendingTask = UE::Tasks::Launch(UE_SOURCE_LOCATION, [this, PullData, RigidsSolver]()
+	{
+		QUICK_SCOPE_CYCLE_COUNTER(STAT_BufferGCResultsTask);
+
+		// This ActiveGC cannot be read while the last DirtyGCAndClusterPendingTask is not finished
+		// that why we need to have this code in a dependent task
+		int32 NumGC = ActiveGC.Num();
+		constexpr int32 MinGCPerTask = 1;
+		const int32 NumTasks = FMath::Max(FMath::Min(FTaskGraphInterface::Get().GetNumWorkerThreads() - 1, Chaos::MaxNumWorkers), 1);
+		const int32 NumGCTask = FMath::Max(FMath::Min(NumTasks, NumGC), 1);
+		const int32 GCByTask = FMath::Max(FMath::DivideAndRoundUp(NumGC, NumGCTask), MinGCPerTask);
+		const int32 NumGCBatches = FMath::DivideAndRoundUp(NumGC, GCByTask);
+
+		PullDataDirtyGCArray.SetNum(NumGCBatches);
+		//we only fill this once per frame
+		ensure(PullData->DirtyGeometryCollections.Num() == 0);
+
+		for (int32 BatchIndex = 0; BatchIndex < NumGCBatches; BatchIndex++)
+		{
+			TArray<FDirtyGeometryCollectionData>& PullDataDirtyGC = PullDataDirtyGCArray[BatchIndex];
+			PullDataDirtyGC.Reset();
+
+			const int32 StartIndex = BatchIndex * GCByTask;
+			int32 EndIndex = (BatchIndex + 1) * GCByTask;
+			EndIndex = FMath::Min(NumGC, EndIndex);
+
+			UE::Tasks::FTask GCPendingBufferTask = UE::Tasks::Launch(UE_SOURCE_LOCATION, [StartIndex, EndIndex, RigidsSolver, &ActiveGC = ActiveGC, &PullDataDirtyGC]()
+			{
+				QUICK_SCOPE_CYCLE_COUNTER(STAT_BufferGCResultsTasks);
+				for (int32 Idx = StartIndex; Idx < EndIndex; ++Idx)
+				{
+					PullDataDirtyGC.AddDefaulted();
+					ActiveGC[Idx]->BufferPhysicsResults_Internal(RigidsSolver, PullDataDirtyGC.Last());
+				}
+			});
+
+			// Link the following task with the one from the previous iteration to keep determinism
+			UE::Tasks::FTask GCCopyPendingTask = UE::Tasks::Launch(UE_SOURCE_LOCATION, [PullData, &PullDataDirtyGC]()
+			{
+				QUICK_SCOPE_CYCLE_COUNTER(STAT_BufferGCResultsCopyTasks);
+				PullData->DirtyGeometryCollections.Append(PullDataDirtyGC);
+			}, TArray<UE::Tasks::FTask>{GCPendingBufferTask, PreviousCopyGCTask});
+			PreviousCopyGCTask = GCCopyPendingTask;
+		}
+	}, DirtyGCAndClusterPendingTask);
+}
+
+
+void FTaskDispatcherSolver::BufferClusterResults(FPullPhysicsData* PullData)
+{
+	ClusterPendingTask = UE::Tasks::Launch(UE_SOURCE_LOCATION, [this, PullData]()
+	{
+		QUICK_SCOPE_CYCLE_COUNTER(STAT_BufferClusterResultsTask);
+
+		// This value cannot be read while the last DirtyGCAndClusterPendingTask is not finished
+		// that why we need to have this code in a dependent task
+		int32 NumClusters = ActiveClusterUnions.Num();
+		constexpr int32 MinClusterPerTask = 1;
+		const int32 NumTasks = FMath::Max(FMath::Min(FTaskGraphInterface::Get().GetNumWorkerThreads() - 1, Chaos::MaxNumWorkers), 1);
+		const int32 NumClusterTask = FMath::Max(FMath::Min(NumTasks, NumClusters), 1);
+		const int32 ClusterByTask = FMath::Max(FMath::DivideAndRoundUp(NumClusters, NumClusterTask), MinClusterPerTask);
+		const int32 NumClusterBatches = FMath::DivideAndRoundUp(NumClusters, ClusterByTask);
+
+		PullDataDirtyClusterArray.SetNum(NumClusterBatches);
+		//we only fill this once per frame
+		ensure(PullData->DirtyClusterUnions.Num() == 0);
+
+		for (int32 BatchIndex = 0; BatchIndex < NumClusterBatches; BatchIndex++)
+		{
+			const int32 StartIndex = BatchIndex * ClusterByTask;
+			int32 EndIndex = (BatchIndex + 1) * ClusterByTask;
+			EndIndex = FMath::Min(NumClusters, EndIndex);
+
+			TArray<FDirtyClusterUnionData>& PullDataDirtyCluster = PullDataDirtyClusterArray[BatchIndex];
+			PullDataDirtyCluster.Reset();
+
+			UE::Tasks::FTask ClusterBufferTask = UE::Tasks::Launch(UE_SOURCE_LOCATION, [StartIndex, EndIndex, &ActiveClusterUnions = ActiveClusterUnions, &PullDataDirtyCluster]()
+			{
+				QUICK_SCOPE_CYCLE_COUNTER(STAT_BufferClusterResultsTasks);
+				for (int32 Idx = StartIndex; Idx < EndIndex; ++Idx)
+				{
+					PullDataDirtyCluster.AddDefaulted();
+					ActiveClusterUnions[Idx]->BufferPhysicsResults_Internal(PullDataDirtyCluster.Last());
+				}
+			});
+
+			// Link the following task with the one from the previous iteration to keep determinism
+			UE::Tasks::FTask ClusterCopyPendingTask = UE::Tasks::Launch(UE_SOURCE_LOCATION, [PullData, &PullDataDirtyCluster]()
+			{
+				QUICK_SCOPE_CYCLE_COUNTER(STAT_BufferClusterResultsCopyTasks);
+				PullData->DirtyClusterUnions.Append(PullDataDirtyCluster);
+			}, TArray<UE::Tasks::FTask>{ClusterBufferTask, PreviousCopyClusterTask});
+			PreviousCopyClusterTask = ClusterCopyPendingTask;
+		}
+	}, PreviousCopyGCTask);
+}
+
+void FTaskDispatcherSolver::WaitTaskEndBufferResults() const
+{
+	check(ActiveRigidArray.Num() == NumRigidBatches);
+	check(ActiveRigidArray.Num() == DirtyRigidPendingTasks.Num());
+
+	// First we need to wait for the outer task which will create inner task
+	UE::Tasks::Wait(TArray<UE::Tasks::FTask>{GCPendingTask});
+	// We are creating a dependency between GC and CU threads to be sure if any dependency is created
+	// a CU and its GC the program doesn't create any race condition
+	UE::Tasks::Wait(TArray<UE::Tasks::FTask>{ClusterPendingTask});
+	// Only when those outer tasks are finished and so those inner task created we can wait for those inner tasks
+	UE::Tasks::Wait(TArray<UE::Tasks::FTask>{PreviousCopyRigidTask, PreviousCopyGCTask, PreviousCopyClusterTask});
+}
+
+}
\ No newline at end of file
diff --git a/Engine/Source/Runtime/Experimental/Chaos/Private/TaskDispatcherSolver.h b/Engine/Source/Runtime/Experimental/Chaos/Private/TaskDispatcherSolver.h
new file mode 100644
--- /dev/null
+++ b/Engine/Source/Runtime/Experimental/Chaos/Private/TaskDispatcherSolver.h
@@ -0,0 +1,116 @@
+// Copyright Epic Games, Inc. All Rights Reserved.
+#pragma once
+
+#include "PhysicsProxy/SingleParticlePhysicsProxy.h"
+#include "Tasks/Task.h"
+
+
+namespace Chaos
+{
+	class FClusterUnionPhysicsProxy;
+}
+class FGeometryCollectionPhysicsProxy;
+
+
+namespace Chaos::Private
+{
+
+// This class is responsible to dispatch the solver tasks
+// It orchestrates the task dependencies according to data flow.
+// This class shouldn't do any simulation logic and code.
+class FTaskDispatcherSolver
+{
+private:
+
+	// Array of batch of array of rigid particles
+	TArray<TArray<FSingleParticlePhysicsProxy*>> ActiveRigidArray;
+	// Geometry Collection array
+	TArray<FGeometryCollectionPhysicsProxy*> ActiveGC;
+	TArray<FClusterUnionPhysicsProxy*> ActiveClusterUnions;
+
+	// Tasks that collect rigid particle
+	TArray<UE::Tasks::FTask> DirtyRigidPendingTasks;
+	// Tasks that collect Geometry collection and cluster union
+	UE::Tasks::FTask DirtyGCAndClusterPendingTask;
+
+	// Outer tasks which will create sub task once the collect task will be finish
+	// collecting Geometry collection and cluster union proxies 
+	UE::Tasks::FTask GCPendingTask;
+	UE::Tasks::FTask ClusterPendingTask;
+
+	// Last task which copy in PullData the proxy for those three types of proxy
+	UE::Tasks::FTask PreviousCopyRigidTask;
+	UE::Tasks::FTask PreviousCopyGCTask;
+	UE::Tasks::FTask PreviousCopyClusterTask;
+
+	// PullData copies to buffer physics results by batch
+	TArray<TArray<FDirtyRigidParticleData>> PullDataDirtyRigidsArray;
+	TArray<TArray<FDirtyGeometryCollectionData>> PullDataDirtyGCArray;
+	TArray<TArray<FDirtyClusterUnionData>> PullDataDirtyClusterArray;
+	
+	int32 NumRigidBatches;
+
+public:
+	
+	int32 GetNumRigidBatches() const { return NumRigidBatches; }
+
+	FTaskDispatcherSolver()
+	{
+	}
+
+	void Initialize(TParticleView<FPBDRigidParticles>& DirtyParticles);
+
+	template<typename LambdaRigid, typename LambdaGCAndCluster>
+	int32 CollectPhysicsResults(LambdaRigid CollectRigidResults, LambdaGCAndCluster CollectGCAndClustersResults, int32 NumDirty, int32 DispatchBatchIndex)
+	{
+		constexpr int32 MinRigidPerTask = 50;
+		const int32 NumTasks = FMath::Max(FMath::Min(FTaskGraphInterface::Get().GetNumWorkerThreads() - 1, Chaos::MaxNumWorkers), 1);
+		const int32 NumDirtyTask = FMath::Max(FMath::Min(NumTasks, NumDirty), 1);
+		const int32 DirtyByTask = FMath::Max(FMath::DivideAndRoundUp(NumDirty, NumDirtyTask), MinRigidPerTask);
+		const int32 LocalNumRigidBatches = FMath::DivideAndRoundUp(NumDirty, DirtyByTask);
+
+		// We can have at most NumDirty overall active particles across all dirty proxy lists
+		// and in each list at most NumProxies. So we take the min, otherwise for example in a world
+		// with a few million static particles and one dynamic we'll allocate enough storage for the
+		// millions of particles despite only ever syncing one
+		for (int32 BatchIndex = 0; BatchIndex < LocalNumRigidBatches; BatchIndex++)
+		{
+			const int32 StartIndex = BatchIndex * DirtyByTask;
+			int32 EndIndex = (BatchIndex + 1) * DirtyByTask;
+			EndIndex = FMath::Min(NumDirty, EndIndex);
+
+			TArray<FSingleParticlePhysicsProxy*>& ActiveRigid = ActiveRigidArray[DispatchBatchIndex + BatchIndex];
+
+			ActiveRigid.Reset();
+			ActiveGC.Reset();
+			ActiveClusterUnions.Reset();
+
+			// Create Task here 
+			UE::Tasks::FTask DirtyPendingTask = UE::Tasks::Launch(UE_SOURCE_LOCATION, [StartIndex, EndIndex, &ActiveRigid, CollectRigidResults]()
+			{
+				QUICK_SCOPE_CYCLE_COUNTER(STAT_CollectRigidsTasks);
+				CollectRigidResults(StartIndex, EndIndex, ActiveRigid);
+			});
+			DirtyRigidPendingTasks.Add(DirtyPendingTask);
+		}
+
+		// The GC and cluster proxies are shared among several particles
+		// all those following task will be computed sequentially to be sure not to add twice the same proxy. 
+		// A fill all then merge approach could be potentially more efficient (not sure)
+		UE::Tasks::FTask DirtyPendingTask = UE::Tasks::Launch(UE_SOURCE_LOCATION, [this, NumDirty, CollectGCAndClustersResults]()
+		{
+			QUICK_SCOPE_CYCLE_COUNTER(STAT_CollectGCAndClustersTasks);
+			CollectGCAndClustersResults(0, NumDirty, ActiveGC, ActiveClusterUnions);
+		}, DirtyGCAndClusterPendingTask);
+		DirtyGCAndClusterPendingTask = DirtyPendingTask;
+
+		return LocalNumRigidBatches;
+	}
+
+	void BufferRigidResults(FPullPhysicsData* PullData);
+	void BufferGCResults(FPullPhysicsData* PullData, FPBDRigidsSolver* RigidsSolver);
+	void BufferClusterResults(FPullPhysicsData* PullData);
+	void WaitTaskEndBufferResults() const;
+	
+};
+}
diff --git a/Engine/Source/Runtime/Experimental/Chaos/Public/Chaos/PBDRigidsEvolution.h b/Engine/Source/Runtime/Experimental/Chaos/Public/Chaos/PBDRigidsEvolution.h
--- a/Engine/Source/Runtime/Experimental/Chaos/Public/Chaos/PBDRigidsEvolution.h
+++ b/Engine/Source/Runtime/Experimental/Chaos/Public/Chaos/PBDRigidsEvolution.h
@@ -549,6 +549,13 @@ public:
 		return false;
 	}
 
+	template <bool bPersistent>
+	FORCEINLINE_DEBUGGABLE void DirtyParticle_ThreadSafe(TGeometryParticleHandleImp<FReal, 3, bPersistent>& Particle, const EPendingSpatialDataOperation Op = EPendingSpatialDataOperation::Update)
+	{
+		FScopeLock Lock(&MutexDirtyParticles);
+		DirtyParticle(Particle, Op);
+	}
+
 	template <bool bPersistent>
 	FORCEINLINE_DEBUGGABLE void DirtyParticle(TGeometryParticleHandleImp<FReal, 3, bPersistent>& Particle, const EPendingSpatialDataOperation Op = EPendingSpatialDataOperation::Update)
 	{
@@ -1001,6 +1008,10 @@ protected:
 	FPendingSpatialDataQueue AsyncAccelerationQueue;
 
 	Private::FTaskDispatcherEvolution TaskDispatcher;
+	
+	// Protect access to the FPendingSpatialData list when a particle is dirty 
+	// The above task dispatcher Buffer Geometry Collection result in parallel
+	FCriticalSection MutexDirtyParticles;
 
 	/*void SerializePendingMap(FChaosArchive& Ar, TMap<FGeometryParticleHandle*, FPendingSpatialData>& Map)
 	{
diff --git a/Engine/Source/Runtime/Experimental/Chaos/Public/PBDRigidsSolver.h b/Engine/Source/Runtime/Experimental/Chaos/Public/PBDRigidsSolver.h
--- a/Engine/Source/Runtime/Experimental/Chaos/Public/PBDRigidsSolver.h
+++ b/Engine/Source/Runtime/Experimental/Chaos/Public/PBDRigidsSolver.h
@@ -10,12 +10,15 @@
 #include "Chaos/PBDCollisionConstraints.h"
 #include "Chaos/PBDRigidDynamicSpringConstraints.h"
 #include "Chaos/PBDPositionConstraints.h"
+#include "Chaos/PBDRigidsEvolutionFwd.h"
 #include "Chaos/PBDSuspensionConstraints.h"
 #include "Chaos/PBDJointConstraints.h"
 #include "Chaos/PerParticleGravity.h"
 #include "Chaos/ParticleHandle.h"
+#include "Chaos/Serialization/SolverSerializer.h"
 #include "Chaos/Transform.h"
 #include "Chaos/Framework/PhysicsProxy.h"
+#include "ChaosSolversModule.h"
 #include "CoreMinimal.h"
 #include "Containers/Queue.h"
 #include "EventManager.h"
@@ -25,9 +28,6 @@
 #include "PhysicsProxy/JointConstraintProxy.h"
 #include "PhysicsProxy/SuspensionConstraintProxy.h"
 #include "SolverEventFilters.h"
-#include "Chaos/PBDRigidsEvolutionFwd.h"
-#include "ChaosSolversModule.h"
-#include "Chaos/Serialization/SolverSerializer.h"
 
 class FPhysInterface_Chaos;
 struct FChaosSolverConfiguration;
@@ -47,6 +47,10 @@ namespace ChaosTest
 	void AdvanceSolverNoPushHelper(TSolver* Solver, Chaos::FReal Dt);
 }
 
+namespace Chaos::Private
+{
+	class FTaskDispatcherSolver;
+}
 /**
 *
 */
@@ -366,7 +370,14 @@ namespace Chaos
 	private:
 
 		/**/
+		template <typename LambdaType>
+		int32 DispatchTasksBufferPhysicsResults(int32 DispatchBatchIndex, int32 NumDirty, LambdaType ReturnHandle);
+
 		CHAOS_API void BufferPhysicsResults();
+		CHAOS_API void BufferPhysicsResultsSingleThreaded(
+			TArray<FSingleParticlePhysicsProxy*>& ActiveRigid, 
+			TArray<FGeometryCollectionPhysicsProxy*>& ActiveGC, 
+			TArray<FClusterUnionPhysicsProxy*>& ActiveClusterUnions);
 	
 		/**/
 		CHAOS_API virtual void PrepareAdvanceBy(const FReal DeltaTime) override;
@@ -391,6 +402,9 @@ namespace Chaos
 		FReal FloorHeight;
 		bool bIsDeterministic;
 
+		// Task dispatcher manage the task graph logic
+		TUniquePtr<Private::FTaskDispatcherSolver> TaskDispatcher;
+
 		// Physics material mirrors for the solver. These should generally stay in sync with the global material list from
 		// the game thread. This data is read only in the solver as we should never need to update it here. External threads can
 		// Enqueue commands to change parameters.
