Optimize PullData by removing DirtyParticle
In Pull Data at the end of the physics pipeline all particle that have been released from a GeometryCollection are added to the Pending Spatial Queue in order to update the Particle in the Broad Phase.
Those particles have been released and are all already in the Active list in the Rigid SOA.
In the fist step of the pipeline all particles in the active list will be added to this Dirty particles and will be added to this Pending List.

diff --git a/Engine/Source/Runtime/Experimental/Chaos/Private/PhysicsProxy/GeometryCollectionPhysicsProxy.cpp b/Engine/Source/Runtime/Experimental/Chaos/Private/PhysicsProxy/GeometryCollectionPhysicsProxy.cpp
--- a/Engine/Source/Runtime/Experimental/Chaos/Private/PhysicsProxy/GeometryCollectionPhysicsProxy.cpp
+++ b/Engine/Source/Runtime/Experimental/Chaos/Private/PhysicsProxy/GeometryCollectionPhysicsProxy.cpp
@@ -3839,8 +3839,6 @@ static void UpdateParticleHandleTransform(Chaos::FPBDRigidsSolver& CurrentSolver
 	Handle.SetX(NewTransform.GetTranslation());
 	Handle.SetR(NewTransform.GetRotation());
 	Handle.UpdateWorldSpaceState(NewTransform, Chaos::FVec3{0});
-	// This current function is called by BufferPhysicsResults_Internal which is called by PullData in multi thread
-	CurrentSolver.GetEvolution()->DirtyParticle_ThreadSafe(Handle);
 }
 
 static void UpdateParticleHandleTransformIfNeeded(Chaos::FPBDRigidsSolver& CurrentSolver, Chaos::FPBDRigidClusteredParticleHandle& Handle, const Chaos::FRigidTransform3& NewTransform)
diff --git a/Engine/Source/Runtime/Experimental/Chaos/Public/Chaos/PBDRigidsEvolution.h b/Engine/Source/Runtime/Experimental/Chaos/Public/Chaos/PBDRigidsEvolution.h
--- a/Engine/Source/Runtime/Experimental/Chaos/Public/Chaos/PBDRigidsEvolution.h
+++ b/Engine/Source/Runtime/Experimental/Chaos/Public/Chaos/PBDRigidsEvolution.h
@@ -549,13 +549,6 @@ public:
 		return false;
 	}
 
-	template <bool bPersistent>
-	FORCEINLINE_DEBUGGABLE void DirtyParticle_ThreadSafe(TGeometryParticleHandleImp<FReal, 3, bPersistent>& Particle, const EPendingSpatialDataOperation Op = EPendingSpatialDataOperation::Update)
-	{
-		FScopeLock Lock(&MutexDirtyParticles);
-		DirtyParticle(Particle, Op);
-	}
-
 	template <bool bPersistent>
 	FORCEINLINE_DEBUGGABLE void DirtyParticle(TGeometryParticleHandleImp<FReal, 3, bPersistent>& Particle, const EPendingSpatialDataOperation Op = EPendingSpatialDataOperation::Update)
 	{
