Clean up : Remove useless mutex in header

#FailedValidations Preflight

diff --git a/Engine/Source/Runtime/Experimental/Chaos/Public/Chaos/PBDRigidsEvolution.h b/Engine/Source/Runtime/Experimental/Chaos/Public/Chaos/PBDRigidsEvolution.h
--- a/Engine/Source/Runtime/Experimental/Chaos/Public/Chaos/PBDRigidsEvolution.h
+++ b/Engine/Source/Runtime/Experimental/Chaos/Public/Chaos/PBDRigidsEvolution.h
@@ -1002,9 +1002,6 @@ protected:
 
 	Private::FTaskDispatcherEvolution TaskDispatcher;
 	
-	// Protect access to the FPendingSpatialData list when a particle is dirty 
-	// The above task dispatcher Buffer Geometry Collection result in parallel
-	FCriticalSection MutexDirtyParticles;
 
 	/*void SerializePendingMap(FChaosArchive& Ar, TMap<FGeometryParticleHandle*, FPendingSpatialData>& Map)
 	{
