Fixed Failing Unit test
The GC and the CU array were reset  at every iteration instead of once by step.

#FailedValidations Preflight

diff --git a/Engine/Source/Runtime/Experimental/Chaos/Private/TaskDispatcherSolver.cpp b/Engine/Source/Runtime/Experimental/Chaos/Private/TaskDispatcherSolver.cpp
--- a/Engine/Source/Runtime/Experimental/Chaos/Private/TaskDispatcherSolver.cpp
+++ b/Engine/Source/Runtime/Experimental/Chaos/Private/TaskDispatcherSolver.cpp
@@ -39,6 +39,9 @@ void FTaskDispatcherSolver::Initialize(TParticleView<FPBDRigidParticles>& DirtyP
 
 	ActiveRigidArray.SetNum(NumRigidBatches);
 
+	ActiveGC.Reset();
+	ActiveClusterUnions.Reset();
+
 	PullDataDirtyRigidsArray.Reset();
 	PullDataDirtyGCArray.Reset();
 	PullDataDirtyClusterArray.Reset();
diff --git a/Engine/Source/Runtime/Experimental/Chaos/Private/TaskDispatcherSolver.h b/Engine/Source/Runtime/Experimental/Chaos/Private/TaskDispatcherSolver.h
--- a/Engine/Source/Runtime/Experimental/Chaos/Private/TaskDispatcherSolver.h
+++ b/Engine/Source/Runtime/Experimental/Chaos/Private/TaskDispatcherSolver.h
@@ -80,10 +80,7 @@ public:
 			EndIndex = FMath::Min(NumDirty, EndIndex);
 
 			TArray<FSingleParticlePhysicsProxy*>& ActiveRigid = ActiveRigidArray[DispatchBatchIndex + BatchIndex];
-
 			ActiveRigid.Reset();
-			ActiveGC.Reset();
-			ActiveClusterUnions.Reset();
 
 			// Create Task here 
 			UE::Tasks::FTask DirtyPendingTask = UE::Tasks::Launch(UE_SOURCE_LOCATION, [StartIndex, EndIndex, &ActiveRigid, CollectRigidResults]()
