Constraints: fixed bMaintainOffset not taken into account
- regression from 45690970

#rnx

diff --git a/Engine/Source/Runtime/Experimental/Animation/Constraints/Private/Transform/TransformConstraintUtil.cpp b/Engine/Source/Runtime/Experimental/Animation/Constraints/Private/Transform/TransformConstraintUtil.cpp
--- a/Engine/Source/Runtime/Experimental/Animation/Constraints/Private/Transform/TransformConstraintUtil.cpp
+++ b/Engine/Source/Runtime/Experimental/Animation/Constraints/Private/Transform/TransformConstraintUtil.cpp
@@ -416,10 +416,13 @@ bool AddConstraint(UWorld* InWorld, UTransformableHandle* InParentHandle, UTrans
 			bSetup = CDO->bUseCurrentOffset;
 		}
 	}
+	else
+	{
+		InNewConstraint->bMaintainOffset = bMaintainOffset;
+	}
 
 	if (bSetup)
 	{
-		InNewConstraint->bMaintainOffset = bMaintainOffset;
 		InNewConstraint->Setup();
 	}
 	InNewConstraint->InitConstraint(InWorld);
