NAMING AND DOCUMENTATION CHANGES ONLY

Update CVars for partial island sleep
- Create hierarchy to make it clear what CVars belong to partial island sleep
- Update default values and description
- CVars which should generally not be changed are labeled "safeguard only"
- Fix incorrect naming convention for momentum propagation

diff --git a/Engine/Source/Programs/HeadlessChaos/Private/HeadlessChaosTestGraphEvolution.cpp b/Engine/Source/Programs/HeadlessChaos/Private/HeadlessChaosTestGraphEvolution.cpp
--- a/Engine/Source/Programs/HeadlessChaos/Private/HeadlessChaosTestGraphEvolution.cpp
+++ b/Engine/Source/Programs/HeadlessChaos/Private/HeadlessChaosTestGraphEvolution.cpp
@@ -29,32 +29,32 @@ namespace ChaosTest
 
 		void SetPartialIslandSleepUnitTestDefaults()
 		{
-			CVarPostStepWake = IConsoleManager::Get().FindConsoleVariable(TEXT("p.Chaos.Solver.Sleep.PostStepWakeThreshold"), false);
+			CVarPostStepWake = IConsoleManager::Get().FindConsoleVariable(TEXT("p.Chaos.Solver.Sleep.PartialIslandSleep.PostStepWakeThreshold"), false);
 			check(CVarPostStepWake && CVarPostStepWake->IsVariableFloat());
 			PostStepWake_EditorDefault = CVarPostStepWake->GetFloat();
 			CVarPostStepWake->Set(PostStepWake_UnitTestDefault);
 
-			CVarMomentumPropagation = IConsoleManager::Get().FindConsoleVariable(TEXT("p.Chaos.Solver.Sleep.MomentumPropagation"), false);
+			CVarMomentumPropagation = IConsoleManager::Get().FindConsoleVariable(TEXT("p.Chaos.Solver.Sleep.PartialIslandSleep.MomentumPropagation"), false);
 			check(CVarMomentumPropagation && CVarMomentumPropagation->IsVariableFloat());
 			MomentumPropagation_EditorDefault = CVarMomentumPropagation->GetFloat();
 			CVarMomentumPropagation->Set(MomentumPropagation_UnitTestDefault);
 
-			CVarPartialSleepCollisionConstraintsOnly = IConsoleManager::Get().FindConsoleVariable(TEXT("p.Chaos.Solver.Sleep.PartialSleepCollisionConstraintsOnly"), false);
+			CVarPartialSleepCollisionConstraintsOnly = IConsoleManager::Get().FindConsoleVariable(TEXT("p.Chaos.Solver.Sleep.PartialIslandSleep.CollisionConstraintsOnly"), false);
 			check(CVarPartialSleepCollisionConstraintsOnly && CVarPartialSleepCollisionConstraintsOnly->IsVariableBool());
 			bPartialSleepCollisionConstraintsOnly_EditorDefault = CVarPartialSleepCollisionConstraintsOnly->GetBool();
 			CVarPartialSleepCollisionConstraintsOnly->Set(bPartialSleepCollisionConstraintsOnly_UnitTestDefault);
 
-			CVarPartialWakePreIntegration = IConsoleManager::Get().FindConsoleVariable(TEXT("p.Chaos.Solver.Sleep.PartialWakePreIntegration"), false);
+			CVarPartialWakePreIntegration = IConsoleManager::Get().FindConsoleVariable(TEXT("p.Chaos.Solver.Sleep.PartialIslandSleep.PartialWakePreIntegration"), false);
 			check(CVarPartialWakePreIntegration && CVarPartialWakePreIntegration->IsVariableBool());
 			bPartialWakePreIntegration_EditorDefault = CVarPartialWakePreIntegration->GetBool();
 			CVarPartialWakePreIntegration->Set(bPartialWakePreIntegration_UnitTestDefault);
 
-			CVarLinearWakeThresholdMultiplier = IConsoleManager::Get().FindConsoleVariable(TEXT("p.Chaos.Solver.Sleep.Defaults.LinearWakeThresholdMultiplier"), false);
+			CVarLinearWakeThresholdMultiplier = IConsoleManager::Get().FindConsoleVariable(TEXT("p.Chaos.Solver.Sleep.PartialIslandSleep.LinearWakeThresholdMultiplier"), false);
 			check(CVarLinearWakeThresholdMultiplier && CVarLinearWakeThresholdMultiplier->IsVariableFloat());
 			LinearWakeThresholdMultiplier_EditorDefault = CVarLinearWakeThresholdMultiplier->GetFloat();
 			CVarLinearWakeThresholdMultiplier->Set(LinearWakeThresholdMultiplier_UnitTestDefault);
 
-			CVarAngularWakeThresholdMultiplier = IConsoleManager::Get().FindConsoleVariable(TEXT("p.Chaos.Solver.Sleep.Defaults.AngularWakeThresholdMultiplier"), false);
+			CVarAngularWakeThresholdMultiplier = IConsoleManager::Get().FindConsoleVariable(TEXT("p.Chaos.Solver.Sleep.PartialIslandSleep.AngularWakeThresholdMultiplier"), false);
 			check(CVarAngularWakeThresholdMultiplier && CVarAngularWakeThresholdMultiplier->IsVariableFloat());
 			AngularWakeThresholdMultiplier_EditorDefault = CVarAngularWakeThresholdMultiplier->GetFloat();
 			CVarAngularWakeThresholdMultiplier->Set(AngularWakeThresholdMultiplier_UnitTestDefault);
diff --git a/Engine/Source/Programs/HeadlessChaos/Private/HeadlessChaosTestSleeping.cpp b/Engine/Source/Programs/HeadlessChaos/Private/HeadlessChaosTestSleeping.cpp
--- a/Engine/Source/Programs/HeadlessChaos/Private/HeadlessChaosTestSleeping.cpp
+++ b/Engine/Source/Programs/HeadlessChaos/Private/HeadlessChaosTestSleeping.cpp
@@ -26,32 +26,32 @@ namespace ChaosTest
 
 		void SetPartialIslandSleepUnitTestDefaults()
 		{
-			CVarPostStepWake = IConsoleManager::Get().FindConsoleVariable(TEXT("p.Chaos.Solver.Sleep.PostStepWakeThreshold"), false);
+			CVarPostStepWake = IConsoleManager::Get().FindConsoleVariable(TEXT("p.Chaos.Solver.Sleep.PartialIslandSleep.PostStepWakeThreshold"), false);
 			check(CVarPostStepWake && CVarPostStepWake->IsVariableFloat());
 			PostStepWake_EditorDefault = CVarPostStepWake->GetFloat();
 			CVarPostStepWake->Set(PostStepWake_UnitTestDefault);
 
-			CVarMomentumPropagation = IConsoleManager::Get().FindConsoleVariable(TEXT("p.Chaos.Solver.Sleep.MomentumPropagation"), false);
+			CVarMomentumPropagation = IConsoleManager::Get().FindConsoleVariable(TEXT("p.Chaos.Solver.Sleep.PartialIslandSleep.MomentumPropagation"), false);
 			check(CVarMomentumPropagation && CVarMomentumPropagation->IsVariableFloat());
 			MomentumPropagation_EditorDefault = CVarMomentumPropagation->GetFloat();
 			CVarMomentumPropagation->Set(MomentumPropagation_UnitTestDefault);
 
-			CVarPartialSleepCollisionConstraintsOnly = IConsoleManager::Get().FindConsoleVariable(TEXT("p.Chaos.Solver.Sleep.PartialSleepCollisionConstraintsOnly"), false);
+			CVarPartialSleepCollisionConstraintsOnly = IConsoleManager::Get().FindConsoleVariable(TEXT("p.Chaos.Solver.Sleep.PartialIslandSleep.CollisionConstraintsOnly"), false);
 			check(CVarPartialSleepCollisionConstraintsOnly && CVarPartialSleepCollisionConstraintsOnly->IsVariableBool());
 			bPartialSleepCollisionConstraintsOnly_EditorDefault = CVarPartialSleepCollisionConstraintsOnly->GetBool();
 			CVarPartialSleepCollisionConstraintsOnly->Set(bPartialSleepCollisionConstraintsOnly_UnitTestDefault);
 
-			CVarPartialWakePreIntegration = IConsoleManager::Get().FindConsoleVariable(TEXT("p.Chaos.Solver.Sleep.PartialWakePreIntegration"), false);
+			CVarPartialWakePreIntegration = IConsoleManager::Get().FindConsoleVariable(TEXT("p.Chaos.Solver.Sleep.PartialIslandSleep.PartialWakePreIntegration"), false);
 			check(CVarPartialWakePreIntegration && CVarPartialWakePreIntegration->IsVariableBool());
 			bPartialWakePreIntegration_EditorDefault = CVarPartialWakePreIntegration->GetBool();
 			CVarPartialWakePreIntegration->Set(bPartialWakePreIntegration_UnitTestDefault);
 
-			CVarLinearWakeThresholdMultiplier = IConsoleManager::Get().FindConsoleVariable(TEXT("p.Chaos.Solver.Sleep.Defaults.LinearWakeThresholdMultiplier"), false);
+			CVarLinearWakeThresholdMultiplier = IConsoleManager::Get().FindConsoleVariable(TEXT("p.Chaos.Solver.Sleep.PartialIslandSleep.LinearWakeThresholdMultiplier"), false);
 			check(CVarLinearWakeThresholdMultiplier && CVarLinearWakeThresholdMultiplier->IsVariableFloat());
 			LinearWakeThresholdMultiplier_EditorDefault = CVarLinearWakeThresholdMultiplier->GetFloat();
 			CVarLinearWakeThresholdMultiplier->Set(LinearWakeThresholdMultiplier_UnitTestDefault);
 
-			CVarAngularWakeThresholdMultiplier = IConsoleManager::Get().FindConsoleVariable(TEXT("p.Chaos.Solver.Sleep.Defaults.AngularWakeThresholdMultiplier"), false);
+			CVarAngularWakeThresholdMultiplier = IConsoleManager::Get().FindConsoleVariable(TEXT("p.Chaos.Solver.Sleep.PartialIslandSleep.AngularWakeThresholdMultiplier"), false);
 			check(CVarAngularWakeThresholdMultiplier && CVarAngularWakeThresholdMultiplier->IsVariableFloat());
 			AngularWakeThresholdMultiplier_EditorDefault = CVarAngularWakeThresholdMultiplier->GetFloat();
 			CVarAngularWakeThresholdMultiplier->Set(AngularWakeThresholdMultiplier_UnitTestDefault);
diff --git a/Engine/Source/Runtime/Experimental/Chaos/Private/Chaos/Island/IslandManager.cpp b/Engine/Source/Runtime/Experimental/Chaos/Private/Chaos/Island/IslandManager.cpp
--- a/Engine/Source/Runtime/Experimental/Chaos/Private/Chaos/Island/IslandManager.cpp
+++ b/Engine/Source/Runtime/Experimental/Chaos/Private/Chaos/Island/IslandManager.cpp
@@ -57,38 +57,46 @@ namespace Chaos::CVars
 
 	/** Cvar to enable/disable partial island waking for all islands */
 	bool bChaosSolverPartialIslandWake = true;
-	FAutoConsoleVariableRef CVarChaosSolverPartialIslandWake(TEXT("p.Chaos.Solver.Sleep.PartialIslandWake"), bChaosSolverPartialIslandWake, TEXT("[EXPERIMENTAL] New constraints and particle state changes will trigger partial wake events instead of waking up the entire island (partial island sleep only)."));
+	FAutoConsoleVariableRef CVarChaosSolverPartialIslandWake(TEXT("p.Chaos.Solver.Sleep.PartialIslandSleep.PartialIslandWake"), bChaosSolverPartialIslandWake, TEXT("[EXPERIMENTAL] Safeguard only: New constraints and particle state changes will trigger partial wake events instead of waking up the entire island.[def:true]"));
 	
 	/** Cvar to enable/disable forcing a partial island waking for all islands even if all particles are already awake.
 	This is just a safeguard to force resetting the particle sleep counter when a new wake event occurs. It should generally be set to false. */
 	bool bChaosSolverForcePartialIslandWake = false;
-	FAutoConsoleVariableRef CVarChaosSolverForcePartialIslandWake(TEXT("p.Chaos.Solver.Sleep.ForcePartialIslandWake"), bChaosSolverForcePartialIslandWake, TEXT("[EXPERIMENTAL] Safeguard only: We force partial wake events even if all particles in the island are already awake to guarantee resetting the particle sleep counter (partial island sleep only).[def:0]"));
+	FAutoConsoleVariableRef CVarChaosSolverForcePartialIslandWake(TEXT("p.Chaos.Solver.Sleep.PartialIslandSleep.ForcePartialIslandWake"), bChaosSolverForcePartialIslandWake, TEXT("[EXPERIMENTAL] Safeguard only: We force partial wake events even if all particles in the island are already awake to guarantee resetting the particle sleep counter.[def:false]"));
 
 	/** Cvar to enable/disable waking up the entire island if the state of any particle has changed (partial island sleeping only) */
 	bool bChaosSolverWakeOnParticleStateChanged = false;
-	FAutoConsoleVariableRef CVarChaosSolverWakeOnParticleStateChanged(TEXT("p.Chaos.Solver.Sleep.WakeOnParticleStateChanged"), bChaosSolverWakeOnParticleStateChanged, TEXT("[EXPERIMENTAL] State changes of a particle will trigger the island to wake up entirely (partial island sleep only)."));
+	FAutoConsoleVariableRef CVarChaosSolverWakeOnParticleStateChanged(TEXT("p.Chaos.Solver.Sleep.PartialIslandSleep.WakeOnParticleStateChanged"), bChaosSolverWakeOnParticleStateChanged, TEXT("[EXPERIMENTAL] Safeguard only: State changes of a particle will trigger the island to wake up entirely.[def:false]"));
 
 	/** Cvar to enable/disable triggering wake up for newly detected constraints with moving particles only (partial island sleeping only) */
 	bool bChaosSolverWakeOnMovingOnly = true;
-	FAutoConsoleVariableRef CVarChaosSolverWakeOnMovingOnly(TEXT("p.Chaos.Solver.Sleep.WakeOnMovingOnly"), bChaosSolverWakeOnMovingOnly, TEXT("[EXPERIMENTAL] Only newly detected constraints connected to a moving particle will trigger wake-up event (partial island sleep only)."));
+	FAutoConsoleVariableRef CVarChaosSolverWakeOnMovingOnly(TEXT("p.Chaos.Solver.Sleep.PartialIslandSleep.WakeOnMovingOnly"), bChaosSolverWakeOnMovingOnly, TEXT("[EXPERIMENTAL] Safeguard only: Only newly detected constraints connected to a moving particle will trigger wake-up event.[def:true]"));
 
 	/** Cvar to adjust post-step waking triggered by velocity solver impulses.
 	* The solver impulse is converted to a change in velocity by division by the particle mass.
-	* Ideal setting = 10.0f, i.e. change in velocity > 10 will trigger wake of the sleeping particle. */
-	FRealSingle ChaosSolverPostStepWakeThreshold = 0.0f;
-	FAutoConsoleVariableRef CVarChaosSolverPostStepWakeThreshold(TEXT("p.Chaos.Solver.Sleep.PostStepWakeThreshold"), ChaosSolverPostStepWakeThreshold, TEXT("[EXPERIMENTAL] Post-solver waking of sleeping particles: Minimum squared change of velocity induced by the velocity-based collision solver (VBD) or the joint solver (partial island sleeping only, 0 to disable).[def:10]"));
+	* Ideal setting = 20.0f, i.e. change in velocity > 20 will trigger wake of the sleeping particle. */
+	FRealSingle ChaosSolverPostStepWakeThreshold = 20.0f;
+	FAutoConsoleVariableRef CVarChaosSolverPostStepWakeThreshold(TEXT("p.Chaos.Solver.Sleep.PartialIslandSleep.PostStepWakeThreshold"), ChaosSolverPostStepWakeThreshold, TEXT("[EXPERIMENTAL] Post-solver waking of sleeping particles: Minimum squared change of velocity induced by the velocity-based collision solver (VBD, 0 to disable).[def:20]"));
 
 	/** Cvar to enable/disable momentum propagation through the constraint graph (partial island sleeping only) */
-	FRealSingle bChaosSolverMomentumPropagation = 0.0f;
-	FAutoConsoleVariableRef CVarChaosSolverMomentumPropagation(TEXT("p.Chaos.Solver.Sleep.MomentumPropagation"), bChaosSolverMomentumPropagation, TEXT("[EXPERIMENTAL] Ratio of the impact momentum transferred from one sleeping particle to another through a constraint during partial wake event (partial island sleeping only, 0 to disable, 1 to propagate with dissipation).[def:0.25]"));
+	FRealSingle ChaosSolverMomentumPropagation = 0.3f;
+	FAutoConsoleVariableRef CVarChaosSolverMomentumPropagation(TEXT("p.Chaos.Solver.Sleep.PartialIslandSleep.MomentumPropagation"), ChaosSolverMomentumPropagation, TEXT("[EXPERIMENTAL] Ratio of the impact momentum transferred from one sleeping particle to another through a constraint during partial wake event (0 to disable, 1 to propagate without dissipation).[def:0.3]"));
 
 	/** Cvar to enable/disable partial sleep for collision constraints only, i.e. the sleep state of particles connect by other constraint types must match (partial island sleeping only) */
-	bool bChaosSolverPartialSleepCollisionConstraintsOnly = false;
-	FAutoConsoleVariableRef CVarChaosSolverPartialSleepCollisionConstraintsOnly(TEXT("p.Chaos.Solver.Sleep.PartialSleepCollisionConstraintsOnly"), bChaosSolverPartialSleepCollisionConstraintsOnly, TEXT("[EXPERIMENTAL] Only collision constraints permit partial sleep, i.e. the sleep state of particles connect by other constraint types must match (partial island sleep only)."));
+	bool bChaosSolverPartialSleepCollisionConstraintsOnly = true;
+	FAutoConsoleVariableRef CVarChaosSolverPartialSleepCollisionConstraintsOnly(TEXT("p.Chaos.Solver.Sleep.PartialIslandSleep.CollisionConstraintsOnly"), bChaosSolverPartialSleepCollisionConstraintsOnly, TEXT("[EXPERIMENTAL] Safeguard only: Only collision constraints permit partial sleep, i.e. the sleep state of particles connect by other constraint types must match.[def:true]"));
 
 	/** Cvar to enable partial waking pre-integration (UpdateExplicitSleep). */
-	bool bChaosSolverPartialWakePreIntegration = false;
-	FAutoConsoleVariableRef CVarChaosSolverPartialWakePreIntegration(TEXT("p.Chaos.Solver.Sleep.PartialWakePreIntegration"), bChaosSolverPartialWakePreIntegration, TEXT("[EXPERIMENTAL] Partial wake events also happen pre-integration when the user explicilty sets a particle transform or applies external impulses or forces (partial island sleeping only)."));
+	bool bChaosSolverPartialWakePreIntegration = true;
+	FAutoConsoleVariableRef CVarChaosSolverPartialWakePreIntegration(TEXT("p.Chaos.Solver.Sleep.PartialIslandSleep.PartialWakePreIntegration"), bChaosSolverPartialWakePreIntegration, TEXT("[EXPERIMENTAL] Safeguard only: Partial wake events also happen pre-integration when the user explicilty sets a particle transform or applies external impulses or forces.[def:true]"));
+
+	/** Cvar to adjust the linear wake threshold for partial waking */
+	FRealSingle LinearPartialWakeThresholdMultiplier = 10.0f;
+	FAutoConsoleVariableRef CVarChaosSolverLinearWakeThresholdMultiplier(TEXT("p.Chaos.Solver.Sleep.PartialIslandSleep.LinearWakeThresholdMultiplier"), LinearPartialWakeThresholdMultiplier, TEXT("[EXPERIMENTAL] A multiplier applied to LinearSleepThreshold for particle waking.[def:10]"));
+
+	/** Cvar to adjust the angular wake threshold for partial waking */
+	FRealSingle AngularPartialWakeThresholdMultiplier = 10.0f;
+	FAutoConsoleVariableRef CVarChaosSolverAngularWakeThresholdMultiplier(TEXT("p.Chaos.Solver.Sleep.PartialIslandSleep.AngularWakeThresholdMultiplier"), AngularPartialWakeThresholdMultiplier, TEXT("[EXPERIMENTAL] A multiplier applied to AngularSleepThreshold for particle waking.[def:10]"));
 
 	/** Cvar to override the sleep counter threshold if necessary */
 	int32 ChaosSolverCollisionDefaultSleepCounterThreshold = 20;
@@ -109,14 +117,6 @@ namespace Chaos::CVars
 	FRealSingle ChaosSolverCollisionAngularSleepThresholdSize = 0;
 	FAutoConsoleVariableRef CVarChaosSolverCollisionAngularSleepThresholdSize(TEXT("p.Chaos.Solver.Sleep.AngularSleepThresholdSize"), ChaosSolverCollisionAngularSleepThresholdSize, TEXT("Scales the angular threshold based on size (0 to disable size based scaling)"));
 
-	/** Cvar to adjust the linear wake threshold for partial waking */
-	FRealSingle LinearPartialWakeThresholdMultiplier = 1.0f;
-	FAutoConsoleVariableRef CVarChaosSolverLinearWakeThresholdMultiplier(TEXT("p.Chaos.Solver.Sleep.Defaults.LinearWakeThresholdMultiplier"), LinearPartialWakeThresholdMultiplier, TEXT("[EXPERIMENTAL] A multiplier applied to SleepLinearThreshold for particle waking (partial island sleeping only)."));
-
-	/** Cvar to adjust the angular wake threshold for partial waking */
-	FRealSingle AngularPartialWakeThresholdMultiplier = 1.0f;
-	FAutoConsoleVariableRef CVarChaosSolverAngularWakeThresholdMultiplier(TEXT("p.Chaos.Solver.Sleep.Defaults.AngularWakeThresholdMultiplier"), AngularPartialWakeThresholdMultiplier, TEXT("[EXPERIMENTAL] A multiplier applied to SleepAngularThreshold for particle waking (partial island sleeping only)."));
-
 	/* Cvar to increase the sleep counter threshold for floating particles */
 	int32 IsolatedParticleSleepCounterThresholdMultiplier = 1;
 	FAutoConsoleVariableRef CVarChaosSolverIsolatedParticleSleepCounterThresholdMultiplier(TEXT("p.Chaos.Solver.Sleep.IsolatedParticle.CounterMultiplier"), IsolatedParticleSleepCounterThresholdMultiplier, TEXT("A multiplier applied to SleepCounterThreshold for floating particles"));
@@ -334,12 +334,12 @@ namespace Chaos::Private
 	}
 
 	// Validate that the number of elements in NodeQueue and MomentumQueue is compatible with the wake algorithm.
-	// a) bChaosSolverMomentumPropagation <= 0.0f: No momentum information is stored.
-	// b) bChaosSolverMomentumPropagation > 0.0f: Both queues must have the same number of elements.
+	// a) ChaosSolverMomentumPropagation <= 0.0f: No momentum information is stored.
+	// b) ChaosSolverMomentumPropagation > 0.0f: Both queues must have the same number of elements.
 	inline void ValidateWakeQueues(TArray<FPBDIslandParticle*>& NodeQueue, TArray<FPBDParticleMomentumSq>& MomentumQueue)
 	{
-		check((CVars::bChaosSolverMomentumPropagation > 0.0f && NodeQueue.Num() == MomentumQueue.Num())
-			|| (CVars::bChaosSolverMomentumPropagation <= 0.0f && MomentumQueue.Num() == 0));
+		check((CVars::ChaosSolverMomentumPropagation > 0.0f && NodeQueue.Num() == MomentumQueue.Num())
+			|| (CVars::ChaosSolverMomentumPropagation <= 0.0f && MomentumQueue.Num() == 0));
 	}
 
 	//////////////////////////////////////////////////////////////////////////////////////////////////
@@ -3013,7 +3013,7 @@ namespace Chaos::Private
 		TArray<FPBDIslandParticle*> NodeQueue;				// Queue of island nodes to be traversed by partial waking.
 		TArray<FPBDParticleMomentumSq> MomentumQueue;		// Linear and angular momentum of island nodes in NodeQueue.
 		NodeQueue.Reserve(Island->Nodes.Num());
-		if (CVars::bChaosSolverMomentumPropagation > 0.0f)
+		if (CVars::ChaosSolverMomentumPropagation > 0.0f)
 		{
 			MomentumQueue.Reserve(Island->Nodes.Num());
 		}
@@ -3148,7 +3148,7 @@ namespace Chaos::Private
 							ProcessParticleWake(OtherNode, NodeQueue[NodeIndex], NodeMomentum, NodeQueue, MomentumQueue, ParticlesToWakeUp, /*bAlwaysPropagate=*/ true);
 						}
 						// 3) We propagate the impact momentum of the trigger particle through the constraint graph.
-						else if (CVars::bChaosSolverMomentumPropagation > 0.0f)
+						else if (CVars::ChaosSolverMomentumPropagation > 0.0f)
 						{
 							check(!MomentumQueue.IsEmpty());
 							ProcessParticleWake(OtherNode, NodeQueue[NodeIndex], &MomentumQueue[NodeIndex], NodeQueue, MomentumQueue, ParticlesToWakeUp, /*bAlwaysPropagate=*/ false);
@@ -3282,7 +3282,7 @@ namespace Chaos::Private
 			// Add node to the queue if (a) above the motion wake threshold or (b) always triggering.
 			if (bAlwaysTrigger || IsMoving(Node))
 			{
-				if (CVars::bChaosSolverMomentumPropagation > 0.0f)
+				if (CVars::ChaosSolverMomentumPropagation > 0.0f)
 				{
 					const FRealSingle LinearMomentumSq = ComputeLinearMomentumSq(Node);
 					const FRealSingle AngularMomentumSq = ComputeAngularMomentumSq(Node);
@@ -3332,7 +3332,7 @@ namespace Chaos::Private
 		}
 
 		// No need for momentum computation, if momentum propagation disabled and we always propagate to the next particle.
-		if (CVars::bChaosSolverMomentumPropagation <= 0.0f && bAlwaysPropagate)
+		if (CVars::ChaosSolverMomentumPropagation <= 0.0f && bAlwaysPropagate)
 		{
 			NodeQueue.Add(Node);
 			Rigid->SetSleepCounter(int8(SleepCounterValue)); // Force reset in partial wake events to prevent immediate re-sleep
@@ -3353,7 +3353,7 @@ namespace Chaos::Private
 		const FRealSingle AngularMomentumSq = FMath::Square(FRealSingle(Rigid->I().Min())) * Node->WakeAngularThresholdSq;
 
 		// Other momentum needed to perform momentum propagation
-		if (CVars::bChaosSolverMomentumPropagation > 0.0f && OtherMomentumSq)
+		if (CVars::ChaosSolverMomentumPropagation > 0.0f && OtherMomentumSq)
 		{
 			const FRealSingle ReceivedLinearMomentumSq = OtherMomentumSq->Linear;
 			const FRealSingle ReceivedAngularMomentumSq = OtherMomentumSq->Angular;
@@ -3366,8 +3366,8 @@ namespace Chaos::Private
 				// @todo (Chaos): It would be great to use the direction of the momentum as opposed to the absolute value. 
 				// For contacts, if the received momentum acts in the opposing direction of the relative position of the particle,
 				// there is no need to propagate the wake through the particle.
-				const FRealSingle ResidualLinearMomentumSq = ReceivedLinearMomentumSq * FMath::Square(CVars::bChaosSolverMomentumPropagation);
-				const FRealSingle ResidualAngularMomentumSq = ReceivedAngularMomentumSq * FMath::Square(CVars::bChaosSolverMomentumPropagation);
+				const FRealSingle ResidualLinearMomentumSq = ReceivedLinearMomentumSq * FMath::Square(CVars::ChaosSolverMomentumPropagation);
+				const FRealSingle ResidualAngularMomentumSq = ReceivedAngularMomentumSq * FMath::Square(CVars::ChaosSolverMomentumPropagation);
 
 				MomentumQueue.Add(FPBDParticleMomentumSq(ResidualLinearMomentumSq, ResidualAngularMomentumSq));
 				NodeQueue.Add(Node);
