[Misc] Fix use after remove

#tests Ran client build

diff --git a/Engine/Source/Runtime/Experimental/IoStore/OnDemand/Private/OnDemandHttpIoDispacher.cpp b/Engine/Source/Runtime/Experimental/IoStore/OnDemand/Private/OnDemandHttpIoDispacher.cpp
--- a/Engine/Source/Runtime/Experimental/IoStore/OnDemand/Private/OnDemandHttpIoDispacher.cpp
+++ b/Engine/Source/Runtime/Experimental/IoStore/OnDemand/Private/OnDemandHttpIoDispacher.cpp
@@ -462,11 +462,11 @@ FHttpRequest* FHttpQueue::Dequeue()
 	}
 
 	//TODO: Return N request up to some range limit
-	if (FHttpRequestList* List = ByChunkKey.Find(Next->ChunkHash))
+	if (FHttpRequestList List; ByChunkKey.RemoveAndCopyValue(Next->ChunkHash, List))
 	{
-		ByChunkKey.Remove(Next->ChunkHash);
 		// The first request that was enqueued for the chunk is not necessarily the first request in the list
-		Next = List->PeekHead();
+		Next = List.PeekHead();
+
 		if (ByChunkKey.IsEmpty())
 		{
 			ByChunkKey.Empty(ContainerSlackSize);
@@ -1461,3 +1461,4 @@ TSharedPtr<IOnDemandHttpIoDispatcher> MakeOnDemanHttpIoDispatcher(TUniquePtr<cla
 #undef UE_ALLOW_CACHE_POISONING
 
 #undef UE_ALLOW_INVALID_URL_DEBUGGING
+
